Removed remaining conditions on node id copies. Promoted SetNodeId and SetNodeIds to FAnimationBaseContext so FPoseLinkBase::Initialize can update the current node id.
#rnx

diff --git a/Engine/Source/Runtime/Engine/Classes/Animation/AnimNodeBase.h b/Engine/Source/Runtime/Engine/Classes/Animation/AnimNodeBase.h
--- a/Engine/Source/Runtime/Engine/Classes/Animation/AnimNodeBase.h
+++ b/Engine/Source/Runtime/Engine/Classes/Animation/AnimNodeBase.h
@@ -240,6 +240,18 @@ public:
 		return *Message;
 	}
 	
+	void SetNodeId(int32 InNodeId)
+	{ 
+		PreviousNodeId = CurrentNodeId;
+		CurrentNodeId = InNodeId;
+	}
+
+	void SetNodeIds(const FAnimationBaseContext& InContext)
+	{ 
+		CurrentNodeId = InContext.CurrentNodeId;
+		PreviousNodeId = InContext.PreviousNodeId;
+	}
+
 	// Get the current node Id, set when we recurse into graph traversal functions from pose links
 	ENGINE_API int32 GetCurrentNodeId() const { return CurrentNodeId; }
 
@@ -318,10 +330,8 @@ public:
 		, RootMotionWeightModifier(Copy.RootMotionWeightModifier)
 		, DeltaTime(Copy.DeltaTime)
 	{
-#if ANIM_TRACE_ENABLED
 		CurrentNodeId = Copy.CurrentNodeId;
 		PreviousNodeId = Copy.PreviousNodeId;
-#endif
 	}
 
 public:
@@ -335,10 +345,8 @@ public:
 		FAnimationUpdateContext Result(*this);
 		Result.SharedContext = InSharedContext;
 
-#if ANIM_TRACE_ENABLED
 		// This is currently only used in the case of cached poses, where we dont want to preserve the previous node, so clear it here
 		Result.PreviousNodeId = INDEX_NONE;
-#endif
 
 		return Result;
 	}
@@ -381,8 +389,7 @@ public:
 	FAnimationUpdateContext WithNodeId(int32 InNodeId) const
 	{ 
 		FAnimationUpdateContext Result(*this);
-		Result.PreviousNodeId = CurrentNodeId;
-		Result.CurrentNodeId = InNodeId;
+		Result.SetNodeId(InNodeId);
 		return Result; 
 	}
 
@@ -435,18 +442,6 @@ public:
 		PreviousNodeId = SourceContext.PreviousNodeId;
 	}
 
-	void SetNodeId(int32 InNodeId)
-	{ 
-		PreviousNodeId = CurrentNodeId;
-		CurrentNodeId = InNodeId;
-	}
-
-	void SetNodeIds(const FAnimationBaseContext& InContext)
-	{ 
-		CurrentNodeId = InContext.GetCurrentNodeId();
-		PreviousNodeId = InContext.GetPreviousNodeId();
-	}
-
 	ENGINE_API void Initialize(FAnimInstanceProxy* InAnimInstanceProxy);
 
 	// Log evaluation message
@@ -529,18 +524,6 @@ public:
 		PreviousNodeId = SourceContext.PreviousNodeId;
 	}
 
-	void SetNodeId(int32 InNodeId)
-	{ 
-		PreviousNodeId = CurrentNodeId;
-		CurrentNodeId = InNodeId;
-	}
-
-	void SetNodeIds(const FAnimationBaseContext& InContext)
-	{ 
-		CurrentNodeId = InContext.GetCurrentNodeId();
-		PreviousNodeId = InContext.GetPreviousNodeId();
-	}
-
 	ENGINE_API void ResetToRefPose();
 
 	ENGINE_API bool ContainsNaN() const;
diff --git a/Engine/Source/Runtime/Engine/Private/Animation/AnimNodeBase.cpp b/Engine/Source/Runtime/Engine/Private/Animation/AnimNodeBase.cpp
--- a/Engine/Source/Runtime/Engine/Private/Animation/AnimNodeBase.cpp
+++ b/Engine/Source/Runtime/Engine/Private/Animation/AnimNodeBase.cpp
@@ -13,30 +13,24 @@
 FAnimationBaseContext::FAnimationBaseContext()
 	: AnimInstanceProxy(nullptr)
 	, SharedContext(nullptr)
-#if ANIM_TRACE_ENABLED
 	, CurrentNodeId(INDEX_NONE)
 	, PreviousNodeId(INDEX_NONE)
-#endif
 {
 }
 
 FAnimationBaseContext::FAnimationBaseContext(FAnimInstanceProxy* InAnimInstanceProxy, FAnimationUpdateSharedContext* InSharedContext)
 	: AnimInstanceProxy(InAnimInstanceProxy)
 	, SharedContext(InSharedContext)
-#if ANIM_TRACE_ENABLED
 	, CurrentNodeId(INDEX_NONE)
 	, PreviousNodeId(INDEX_NONE)
-#endif
 {
 }
 
 FAnimationBaseContext::FAnimationBaseContext(const FAnimationBaseContext& InContext)
 	: AnimInstanceProxy(InContext.AnimInstanceProxy)
 	, SharedContext(InContext.SharedContext)
-#if ANIM_TRACE_ENABLED
 	, CurrentNodeId(InContext.CurrentNodeId)
 	, PreviousNodeId(InContext.PreviousNodeId)
-#endif
 {
 }
 
@@ -192,7 +186,10 @@ void FPoseLinkBase::Initialize(const FAnimationInitializeContext& Context)
 	// Do standard initialization
 	if (LinkedNode != NULL)
 	{
-		LinkedNode->Initialize_AnyThread(Context);
+		FAnimationInitializeContext LinkContext(Context);
+		LinkContext.SetNodeId(LinkID);
+		TRACE_SCOPED_ANIM_NODE(LinkContext);
+		LinkedNode->Initialize_AnyThread(LinkContext);
 	}
 }
 
@@ -299,15 +296,9 @@ void FPoseLinkBase::Update(const FAnimationUpdateContext& Context)
 
 	if (LinkedNode != NULL)
 	{
-#if ANIM_TRACE_ENABLED
-		{
-			FAnimationUpdateContext LinkContext(Context.WithNodeId(LinkID));
-			TRACE_SCOPED_ANIM_NODE(LinkContext);
-			LinkedNode->Update_AnyThread(LinkContext);
-		}
-#else
-		LinkedNode->Update_AnyThread(Context);
-#endif
+		FAnimationUpdateContext LinkContext(Context.WithNodeId(LinkID));
+		TRACE_SCOPED_ANIM_NODE(LinkContext);
+		LinkedNode->Update_AnyThread(LinkContext);
 	}
 }
 
@@ -352,10 +343,8 @@ void FPoseLink::Evaluate(FPoseContext& Output)
 #endif
 
 		{
-#if ANIM_TRACE_ENABLED
 			Output.SetNodeId(LinkID);
 			TRACE_SCOPED_ANIM_NODE(Output);
-#endif
 			LinkedNode->Evaluate_AnyThread(Output);
 		}
 
@@ -438,10 +427,8 @@ void FComponentSpacePoseLink::EvaluateComponentSpace(FComponentSpacePoseContext&
 	if (LinkedNode != NULL)
 	{
 		{
-#if ANIM_TRACE_ENABLED
 			Output.SetNodeId(LinkID);
 			TRACE_SCOPED_ANIM_NODE(Output);
-#endif
 			LinkedNode->EvaluateComponentSpace_AnyThread(Output);
 		}
 
