Fixing GBuffer format with high-precision normals.

diff --git a/Engine/Shaders/Public/ShaderVersion.ush b/Engine/Shaders/Public/ShaderVersion.ush
--- a/Engine/Shaders/Public/ShaderVersion.ush
+++ b/Engine/Shaders/Public/ShaderVersion.ush
@@ -3,4 +3,4 @@
 // in Platform.ush (which should be included in any shader) it allows to invalidate the shader DDC.
 // 
 // If you are merging streams and there is a conflict with this GUID you should make a new GUID rather than taking one or the other.
-// GUID = ECD05EA2-D8F5-4874-B259-8109E2FEDEA9
+// GUID = 29C6A537-C090-4AEC-A4Af-BA4C7D219338          
\ No newline at end of file
diff --git a/Engine/Source/Runtime/Engine/Private/ShaderCompiler/ShaderGenerationUtil.cpp b/Engine/Source/Runtime/Engine/Private/ShaderCompiler/ShaderGenerationUtil.cpp
--- a/Engine/Source/Runtime/Engine/Private/ShaderCompiler/ShaderGenerationUtil.cpp
+++ b/Engine/Source/Runtime/Engine/Private/ShaderCompiler/ShaderGenerationUtil.cpp
@@ -248,7 +248,7 @@ void FShaderCompileUtilities::ApplyFetchEnvironment(FShaderCompilerDefines& SrcD
 
 
 // if we change the logic, increment this number to force a DDC key change
-static const int32 GBufferGeneratorVersion = 1;
+static const int32 GBufferGeneratorVersion = 2;
 
 static FShaderGlobalDefines FetchShaderGlobalDefines(EShaderPlatform TargetPlatform)
 {
@@ -532,6 +532,7 @@ static const FGBufferCompressionInfo GBufferCompressionInfo[] =
 	{ GBC_Packed_Normal_Octahedral_8_8, 3, 2, {  8,  8,  0,  0 }, false,  true, TEXT("CompressOctahedral")     , TEXT("DecompressOctahedral")   },
 	{ GBC_EncodeNormal_Normal_16_16_16, 3, 3, { 16, 16, 16,  0 }, false,  true, TEXT("EncodeNormalHelper")     , TEXT("DecodeNormalHelper")     },
 	{ GBC_EncodeNormal_Normal_10_10_10, 3, 3, { 10, 10, 10,  0 }, false,  true, TEXT("EncodeNormalHelper")     , TEXT("DecodeNormalHelper")     },
+	{ GBC_EncodeNormal_Normal_8_8_8,    3, 3, {  8,  8,  8,  0 }, false,  true, TEXT("EncodeNormalHelper")     , TEXT("DecodeNormalHelper")     },
 	{ GBC_Packed_Color_5_6_5,			3, 3, {  5,  6,  5,  0 },  true,  true, TEXT("EncodeQuantize565")      , TEXT("DecodeQuantize565")      },
 	{ GBC_Packed_Color_5_6_5_Sqrt,		3, 3, {  5,  6,  5,  0 },  true,  true, TEXT("EncodeQuantize565Sqrt")  , TEXT("DecodeQuantize565Sqrt")  },
 	{ GBC_Packed_Color_4_4_4,			3, 3, {  4,  4,  4,  0 },  true,  true, TEXT("EncodeQuantize444")      , TEXT("DecodeQuantize444")	    },
diff --git a/Engine/Source/Runtime/RenderCore/Private/GBufferInfo.cpp b/Engine/Source/Runtime/RenderCore/Private/GBufferInfo.cpp
--- a/Engine/Source/Runtime/RenderCore/Private/GBufferInfo.cpp
+++ b/Engine/Source/Runtime/RenderCore/Private/GBufferInfo.cpp
@@ -200,9 +200,6 @@ FGBufferInfo RENDERCORE_API FetchLegacyGBufferInfo(const FGBufferParams& Params)
 		Info.NumTargets = Params.bHasPrecShadowFactor ? 7 : 6;
 	}
 
-
-
-
 	// good to see the quality loss due to precision in the gbuffer
 	const bool bHighPrecisionGBuffers = (Params.LegacyFormatIndex >= EGBufferFormat_Force16BitsPerChannel);
 	// good to profile the impact of non 8 bit formats
@@ -214,12 +211,12 @@ FGBufferInfo RENDERCORE_API FetchLegacyGBufferInfo(const FGBufferParams& Params)
 	if (bEnforce8BitPerChannel)
 	{
 		NormalGBufferFormatTarget  = GBT_Unorm_8_8_8_8;
-		NormalGBufferFormatChannel = GBC_Raw_Unorm_8_8_8_8;
+		NormalGBufferFormatChannel = GBC_EncodeNormal_Normal_8_8_8;
 	}
 	else if (Params.LegacyFormatIndex == EGBufferFormat_HighPrecisionNormals)
 	{
 		NormalGBufferFormatTarget  = GBT_Float_16_16_16_16;
-		NormalGBufferFormatChannel = GBC_Raw_Float_16_16_16;
+		NormalGBufferFormatChannel = GBC_EncodeNormal_Normal_16_16_16;
 	}
 
 	const EGBufferType        DiffuseAndSpecularGBufferFormat  = bHighPrecisionGBuffers ? GBT_Float_16_16_16_16 : GBT_Unorm_8_8_8_8;
diff --git a/Engine/Source/Runtime/RenderCore/Public/GBufferInfo.h b/Engine/Source/Runtime/RenderCore/Public/GBufferInfo.h
--- a/Engine/Source/Runtime/RenderCore/Public/GBufferInfo.h
+++ b/Engine/Source/Runtime/RenderCore/Public/GBufferInfo.h
@@ -58,6 +58,7 @@ enum EGBufferCompression
 	GBC_Packed_Normal_Octahedral_8_8,
 	GBC_EncodeNormal_Normal_16_16_16,
 	GBC_EncodeNormal_Normal_10_10_10,
+	GBC_EncodeNormal_Normal_8_8_8,
 
 	GBC_Packed_Color_5_6_5, // a unorm value, packed to 565
 	GBC_Packed_Color_5_6_5_Sqrt, // a unorm value, packed to 565, with sqrt
