Exposed WorldPartition's DefaultHLODLayer

diff --git a/Engine/Source/Runtime/Engine/Public/WorldPartition/WorldPartition.h b/Engine/Source/Runtime/Engine/Public/WorldPartition/WorldPartition.h
--- a/Engine/Source/Runtime/Engine/Public/WorldPartition/WorldPartition.h
+++ b/Engine/Source/Runtime/Engine/Public/WorldPartition/WorldPartition.h
@@ -174,7 +174,7 @@ public:
 
 #if WITH_EDITORONLY_DATA
 	// Default HLOD layer
-	UPROPERTY()
+	UPROPERTY(EditAnywhere, Category=HLOD)
 	class UHLODLayer* DefaultHLODLayer;
 #endif
 
