WP HLOD - Fix HLOD instancing when used with ISMC

diff --git a/Engine/Source/Runtime/Engine/Private/WorldPartition/HLOD/HLODBuilder.cpp b/Engine/Source/Runtime/Engine/Private/WorldPartition/HLOD/HLODBuilder.cpp
--- a/Engine/Source/Runtime/Engine/Private/WorldPartition/HLOD/HLODBuilder.cpp
+++ b/Engine/Source/Runtime/Engine/Private/WorldPartition/HLOD/HLODBuilder.cpp
@@ -137,7 +137,20 @@ class FHLODBuilder_Instancing : public FHLODBuilder
 				// Add all instances
 				for (UPrimitiveComponent* SMC : EntryComponents)
 				{
-					Component->AddInstanceWorldSpace(SMC->GetComponentTransform());
+					// If we have an ISMC, retrieve all instances
+					if (UInstancedStaticMeshComponent* InstancedStaticMeshComponent = Cast<UInstancedStaticMeshComponent>(SMC))
+					{
+						for (int32 InstanceIdx = 0; InstanceIdx < InstancedStaticMeshComponent->GetInstanceCount(); InstanceIdx++)
+						{
+							FTransform InstanceTransform;
+							InstancedStaticMeshComponent->GetInstanceTransform(InstanceIdx, InstanceTransform, true);
+							Component->AddInstanceWorldSpace(InstanceTransform);
+						}
+					}
+					else
+					{
+						Component->AddInstanceWorldSpace(SMC->GetComponentTransform());
+					}
 				}
 
 				Components.Add(Component);
