External Actor : Fixed SpawnActor not dirtying the actor's package. Also removed calls to ClearDirtyFlag on the package of level to spawn in.

diff --git a/Engine/Source/Runtime/Engine/Private/LevelActor.cpp b/Engine/Source/Runtime/Engine/Private/LevelActor.cpp
--- a/Engine/Source/Runtime/Engine/Private/LevelActor.cpp
+++ b/Engine/Source/Runtime/Engine/Private/LevelActor.cpp
@@ -672,12 +672,6 @@ AActor* UWorld::SpawnActor( UClass* Class, FTransform const* UserTransformPtr, c
 
 	if ( GUndo )
 	{
-		// if we are spawning an external actor, clear the dirty flag without capturing in the transaction beforehand
-		// This allows the transaction to capture the package as not being dirty when capturing its current state, which is what we need for proper external actor behavior
-		if (ExternalPackage)
-		{
-			LevelToSpawnIn->GetPackage()->ClearDirtyFlag();
-		}
 		ModifyLevel( LevelToSpawnIn );
 	}
 	LevelToSpawnIn->Actors.Add( Actor );
@@ -715,11 +709,11 @@ AActor* UWorld::SpawnActor( UClass* Class, FTransform const* UserTransformPtr, c
 	OnActorPreSpawnInitialization.Broadcast(Actor);
 
 	Actor->PostSpawnInitialize(UserTransform, SpawnParameters.Owner, SpawnParameters.Instigator, SpawnParameters.IsRemoteOwned(), SpawnParameters.bNoFail, SpawnParameters.bDeferConstruction);
-
-	// if we are spawning an external actor, clear the dirty flag after post spawn initialize which might have dirtied the level package through running construction scripts
+	
+	// If we are spawning an external actor, mark this package dirty
 	if (ExternalPackage)
 	{
-		LevelToSpawnIn->GetPackage()->ClearDirtyFlag();
+		ExternalPackage->MarkPackageDirty();
 	}
 
 	if (Actor->IsPendingKill() && !SpawnParameters.bNoFail)
