Fixed GenerateHLOD command to not dirty the Level anymore.

diff --git a/Engine/Source/Runtime/Engine/Private/WorldPartition/WorldPartition.cpp b/Engine/Source/Runtime/Engine/Private/WorldPartition/WorldPartition.cpp
--- a/Engine/Source/Runtime/Engine/Private/WorldPartition/WorldPartition.cpp
+++ b/Engine/Source/Runtime/Engine/Private/WorldPartition/WorldPartition.cpp
@@ -56,7 +56,6 @@ static FAutoConsoleCommand GenerateHLODCmd(
 			{
 				if (UWorldPartition* WorldPartition = World->GetWorldPartition())
 				{
-					WorldPartition->Modify();
 					WorldPartition->GenerateHLOD();
 				}
 			}
