Chaning ShaderAutogen to make shader autogen code in FEngineLoop::PreInitPreStartupScreen, instead of when first shader is compiled.

diff --git a/Engine/Source/Editor/UnrealEd/Private/EditorEngine.cpp b/Engine/Source/Editor/UnrealEd/Private/EditorEngine.cpp
--- a/Engine/Source/Editor/UnrealEd/Private/EditorEngine.cpp
+++ b/Engine/Source/Editor/UnrealEd/Private/EditorEngine.cpp
@@ -7668,6 +7668,10 @@ void UEditorEngine::SetPreviewPlatform(const FPreviewPlatformInfo& NewPreviewPla
 		NewPreviewPlatform.bPreviewFeatureLevelActive != PreviewPlatform.bPreviewFeatureLevelActive;
 	const ERHIFeatureLevel::Type EffectiveFeatureLevel = NewPreviewPlatform.GetEffectivePreviewFeatureLevel();
 
+	{
+		// Force generation of the autogen files if they don't already exist. We'll ignore the returned data.
+		FGBufferInfo GBufferInfo = FShaderCompileUtilities::FetchGBufferInfoAndWriteAutogen(ShaderPlatform, MaxFeatureLevel);
+	}
 
 	// Record the new preview platform
 	PreviewPlatform = NewPreviewPlatform;
diff --git a/Engine/Source/Runtime/Engine/Private/ShaderCompiler/ShaderGenerationUtil.cpp b/Engine/Source/Runtime/Engine/Private/ShaderCompiler/ShaderGenerationUtil.cpp
--- a/Engine/Source/Runtime/Engine/Private/ShaderCompiler/ShaderGenerationUtil.cpp
+++ b/Engine/Source/Runtime/Engine/Private/ShaderCompiler/ShaderGenerationUtil.cpp
@@ -1981,11 +1981,27 @@ FGBufferInfo FShaderCompileUtilities::FetchGBufferInfoAndWriteAutogen(EShaderPla
 		OutputFileData += TEXT("#endif\n");
 		OutputFileData += TEXT("\n");
 
-		FPaths::MakeStandardFilename(AutogenHeaderFilenameTemp);
-		FFileHelper::SaveStringToFile(OutputFileData, AutogenHeaderFilenameTemp.GetCharArray().GetData(), FFileHelper::EEncodingOptions::ForceAnsi);
+		// If this file already exists, and the text file is the same as what we are planning to write, then don't write it. In that case,
+		// we don't need to update it (since it didn't change) and it leaves the timestamp unchanged in case you need to open the file manually
+		// for debugging.
+		bool bWriteNeeded = true;
+		{
+			FString PrevFileData;
+			bool bFileExists = FFileHelper::LoadFileToString(PrevFileData,*AutogenHeaderFilename);
+			if (bFileExists && PrevFileData == OutputFileData)
+			{
+				bWriteNeeded = false;
+			}
+		}
 
-		FPlatformFileManager::Get().GetPlatformFile().DeleteFile(*AutogenHeaderFilename);
-		FPlatformFileManager::Get().GetPlatformFile().MoveFile(*AutogenHeaderFilename, *AutogenHeaderFilenameTemp);
+		if (bWriteNeeded)
+		{
+			FPaths::MakeStandardFilename(AutogenHeaderFilenameTemp);
+			FFileHelper::SaveStringToFile(OutputFileData, *AutogenHeaderFilenameTemp, FFileHelper::EEncodingOptions::ForceAnsi);
+
+			FPlatformFileManager::Get().GetPlatformFile().DeleteFile(*AutogenHeaderFilename);
+			FPlatformFileManager::Get().GetPlatformFile().MoveFile(*AutogenHeaderFilename, *AutogenHeaderFilenameTemp);
+		}
 	}
 
 	return BufferInfo;
diff --git a/Engine/Source/Runtime/Engine/Public/ShaderCompiler.h b/Engine/Source/Runtime/Engine/Public/ShaderCompiler.h
--- a/Engine/Source/Runtime/Engine/Public/ShaderCompiler.h
+++ b/Engine/Source/Runtime/Engine/Public/ShaderCompiler.h
@@ -528,7 +528,7 @@ namespace FShaderCompileUtilities
 	void GenerateBrdfHeaders(const EShaderPlatform Platform);
 	void ApplyDerivedDefines(FShaderCompilerEnvironment& OutEnvironment, FShaderCompilerEnvironment * SharedEnvironment, const EShaderPlatform Platform);
 	void AppendGBufferDDCKeyString(const EShaderPlatform Platform, FString& KeyString);
-	FGBufferInfo FetchGBufferInfoAndWriteAutogen(EShaderPlatform TargetPlatform, ERHIFeatureLevel::Type FeatureLevel);
+	ENGINE_API FGBufferInfo FetchGBufferInfoAndWriteAutogen(EShaderPlatform TargetPlatform, ERHIFeatureLevel::Type FeatureLevel);
 
 	void ApplyFetchEnvironment(FShaderMaterialPropertyDefines& DefineData, FShaderCompilerEnvironment& OutEnvironment);
 	void ApplyFetchEnvironment(FShaderGlobalDefines& DefineData, FShaderCompilerEnvironment& OutEnvironment, const EShaderPlatform Platform);
diff --git a/Engine/Source/Runtime/Launch/Private/LaunchEngineLoop.cpp b/Engine/Source/Runtime/Launch/Private/LaunchEngineLoop.cpp
--- a/Engine/Source/Runtime/Launch/Private/LaunchEngineLoop.cpp
+++ b/Engine/Source/Runtime/Launch/Private/LaunchEngineLoop.cpp
@@ -91,6 +91,7 @@
 	#include "PIEPreviewDeviceProfileSelectorModule.h"
 	#include "Misc/QueuedThreadPoolWrapper.h"
 	#include "TextureCompiler.h"
+	#include "ShaderCompiler.h"
 
 	#if PLATFORM_WINDOWS
 		#include "Windows/AllowWindowsPlatformTypes.h"
@@ -2616,6 +2617,13 @@ PRAGMA_ENABLE_DEPRECATION_WARNINGS
 		RenderUtilsInit();
 	}
 
+#if WITH_EDITOR
+	{
+		// Force generation of the autogen files if they don't already exist. We'll ignore the returned data.
+		FGBufferInfo GBufferInfo = FShaderCompileUtilities::FetchGBufferInfoAndWriteAutogen(GMaxRHIShaderPlatform, GetMaxSupportedFeatureLevel(GMaxRHIShaderPlatform));
+	}
+#endif
+
 	{
 		bool bUseCodeLibrary = FPlatformProperties::RequiresCookedData() || GAllowCookedDataInEditorBuilds;
 		if (bUseCodeLibrary)
diff --git a/Engine/Source/Runtime/RenderCore/Private/GBufferInfo.cpp b/Engine/Source/Runtime/RenderCore/Private/GBufferInfo.cpp
--- a/Engine/Source/Runtime/RenderCore/Private/GBufferInfo.cpp
+++ b/Engine/Source/Runtime/RenderCore/Private/GBufferInfo.cpp
@@ -356,12 +356,6 @@ FGBufferInfo RENDERCORE_API FetchLegacyGBufferInfo(const FGBufferParams& Params)
 		Info.Slots[GBS_BaseColor].Packing[2] = FGBufferPacking(TargetGBufferC, 2, 0, 0, 6, 2);
 #endif
 
-		if (bStaticLighting)
-		{
-			Info.Slots[GBS_GenericAO] = FGBufferItem(GBS_GenericAO, GBC_Raw_Unorm_8, GBCH_Both);
-			Info.Slots[GBS_GenericAO].Packing[0] = FGBufferPacking(TargetGBufferC, 0, 3);
-		}
-		else
 		{
 			Info.Slots[GBS_GenericAO] = FGBufferItem(GBS_GenericAO, GBC_Raw_Unorm_8, GBCH_Both);
 			Info.Slots[GBS_GenericAO].Packing[0] = FGBufferPacking(TargetGBufferC, 0, 3);
