Finish async compilation when a texture is about to be modified

diff --git a/Engine/Source/Runtime/Engine/Classes/Engine/Texture.h b/Engine/Source/Runtime/Engine/Classes/Engine/Texture.h
--- a/Engine/Source/Runtime/Engine/Classes/Engine/Texture.h
+++ b/Engine/Source/Runtime/Engine/Classes/Engine/Texture.h
@@ -1046,6 +1046,7 @@ public:
 
 	//~ Begin UObject Interface.
 #if WITH_EDITOR
+	ENGINE_API virtual bool Modify(bool bAlwaysMarkDirty = true) override;
 	ENGINE_API virtual void PostEditChangeProperty(FPropertyChangedEvent& PropertyChangedEvent) override;
 	ENGINE_API virtual bool CanEditChange(const FProperty* InProperty) const override;
 #endif // WITH_EDITOR
diff --git a/Engine/Source/Runtime/Engine/Private/Texture.cpp b/Engine/Source/Runtime/Engine/Private/Texture.cpp
--- a/Engine/Source/Runtime/Engine/Private/Texture.cpp
+++ b/Engine/Source/Runtime/Engine/Private/Texture.cpp
@@ -27,6 +27,10 @@
 #include "Interfaces/ITextureFormat.h"
 #include "Interfaces/ITextureFormatModule.h"
 
+#if WITH_EDITOR
+#include "TextureCompiler.h"
+#endif
+
 #if WITH_EDITORONLY_DATA
 	#include "EditorFramework/AssetImportData.h"
 #endif
@@ -221,6 +225,18 @@ bool UTexture::IsDefaultTexture() const
 	return false;
 }
 
+bool UTexture::Modify(bool bAlwaysMarkDirty)
+{
+	// Before applying any modification to the texture
+	// make sure no compilation is still ongoing.
+	if (IsDefaultTexture())
+	{
+		FTextureCompilingManager::Get().FinishCompilation({this});
+	}
+
+	return Super::Modify(bAlwaysMarkDirty);
+}
+
 bool UTexture::CanEditChange(const FProperty* InProperty) const
 {
 	if (InProperty)
