Blueprints: Consider a connection via type promotion a Safe connection at the schema level. This will allow existing functions like ReplaceOldNodeWithNew to work correctly, which matters when making connections programmatically via the BP Upgrade tool

#rnx

diff --git a/Engine/Source/Runtime/Engine/Classes/EdGraph/EdGraphSchema.h b/Engine/Source/Runtime/Engine/Classes/EdGraph/EdGraphSchema.h
--- a/Engine/Source/Runtime/Engine/Classes/EdGraph/EdGraphSchema.h
+++ b/Engine/Source/Runtime/Engine/Classes/EdGraph/EdGraphSchema.h
@@ -401,7 +401,7 @@ public:
 	/** If a connection can be made without breaking existing connections */
 	bool CanSafeConnect() const
 	{
-		return (Response == CONNECT_RESPONSE_MAKE);
+		return (Response == CONNECT_RESPONSE_MAKE  || Response == CONNECT_RESPONSE_MAKE_WITH_PROMOTION);
 	}
 
 	bool IsFatal() const
