Do not stall when compiling special materials for a non-host platform.

- Global shaders will still stall, this is tracked as UE-102657.

diff --git a/Engine/Source/Runtime/Engine/Private/Materials/MaterialShared.cpp b/Engine/Source/Runtime/Engine/Private/Materials/MaterialShared.cpp
--- a/Engine/Source/Runtime/Engine/Private/Materials/MaterialShared.cpp
+++ b/Engine/Source/Runtime/Engine/Private/Materials/MaterialShared.cpp
@@ -2227,13 +2227,15 @@ bool FMaterial::BeginCompileShaderMap(
 			delete UniformBufferStruct;
 		}
 
+		// we can ignore requests for synch compilation if we are compiling for a different platform than we're running, or we're a commandlet that doesn't render (e.g. cooker)
+		const bool bCanIgnoreSynchronousRequirement = (TargetPlatform && !TargetPlatform->IsRunningPlatform()) || (IsRunningCommandlet() && !IsAllowCommandletRendering());
 		const bool bSkipCompilationForODSC = !IsDefaultMaterial() && GShaderCompilingManager->IsShaderCompilationSkipped();
 		if (bSkipCompilationForODSC)
 		{
 			// Force compilation off.
 			PrecompileMode = EMaterialShaderPrecompileMode::None;
 		}
-		else if (RequiresSynchronousCompilation())
+		else if (!bCanIgnoreSynchronousRequirement && RequiresSynchronousCompilation())
 		{
 			// Force sync compilation by material
 			PrecompileMode = EMaterialShaderPrecompileMode::Synchronous;
