Fix crash when cooking font face assets

diff --git a/Engine/Source/Runtime/Core/Public/UObject/UE5MainStreamObjectVersion.h b/Engine/Source/Runtime/Core/Public/UObject/UE5MainStreamObjectVersion.h
--- a/Engine/Source/Runtime/Core/Public/UObject/UE5MainStreamObjectVersion.h
+++ b/Engine/Source/Runtime/Core/Public/UObject/UE5MainStreamObjectVersion.h
@@ -45,6 +45,9 @@ struct CORE_API FUE5MainStreamObjectVersion
 
 		// Fixed Nanite Geometry Collection cooked data
 		GeometryCollectionNaniteCooked,
+			
+		// Added bCooked to UFontFace assets
+		AddedCookedBoolFontFaceAssets,
 
 		// -----<new versions can be added above this line>-------------------------------------------------
 		VersionPlusOne,
diff --git a/Engine/Source/Runtime/Engine/Private/FontFace.cpp b/Engine/Source/Runtime/Engine/Private/FontFace.cpp
--- a/Engine/Source/Runtime/Engine/Private/FontFace.cpp
+++ b/Engine/Source/Runtime/Engine/Private/FontFace.cpp
@@ -11,6 +11,7 @@
 #include "UObject/Package.h"
 #include "Misc/PackageName.h"
 #include "UObject/EditorObjectVersion.h"
+#include "UObject/UE5MainStreamObjectVersion.h"
 
 DEFINE_LOG_CATEGORY_STATIC(LogFontFace, Log, All);
 
@@ -23,23 +24,17 @@ void UFontFace::Serialize(FArchive& Ar)
 {
 	Ar.UsingCustomVersion(FEditorObjectVersion::GUID);
 
-	FString OriginalSourceFilename;
-	if (Ar.IsCooking() && !HasAnyFlags(RF_ClassDefaultObject))
-	{
-		OriginalSourceFilename = MoveTemp(SourceFilename);
-		SourceFilename = GetCookedFilename();
-	}
-
 	Super::Serialize(Ar);
 
-	if (Ar.IsCooking() && !HasAnyFlags(RF_ClassDefaultObject))
+	bool bCooked = Ar.IsCooking();
+	if (Ar.CustomVer(FEditorObjectVersion::GUID) >= FUE5MainStreamObjectVersion::AddedCookedBoolFontFaceAssets)
 	{
-		SourceFilename = MoveTemp(OriginalSourceFilename);
+		Ar << bCooked;
 	}
 
 	if (Ar.IsLoading())
 	{
-		if (FPlatformProperties::RequiresCookedData())
+		if (FPlatformProperties::RequiresCookedData() || bCooked)
 		{
 			SourceFilename = GetCookedFilename();
 		}
