Fix FontFace versions

diff --git a/Engine/Source/Runtime/Engine/Private/FontFace.cpp b/Engine/Source/Runtime/Engine/Private/FontFace.cpp
--- a/Engine/Source/Runtime/Engine/Private/FontFace.cpp
+++ b/Engine/Source/Runtime/Engine/Private/FontFace.cpp
@@ -23,11 +23,12 @@ UFontFace::UFontFace()
 void UFontFace::Serialize(FArchive& Ar)
 {
 	Ar.UsingCustomVersion(FEditorObjectVersion::GUID);
+	Ar.UsingCustomVersion(FUE5MainStreamObjectVersion::GUID);
 
 	Super::Serialize(Ar);
 
 	bool bCooked = Ar.IsCooking();
-	if (Ar.CustomVer(FEditorObjectVersion::GUID) >= FUE5MainStreamObjectVersion::AddedCookedBoolFontFaceAssets)
+	if (Ar.CustomVer(FUE5MainStreamObjectVersion::GUID) >= FUE5MainStreamObjectVersion::AddedCookedBoolFontFaceAssets)
 	{
 		Ar << bCooked;
 	}
