Make skylight dependant on a specific cubemap initialize faster under heavy asset compilation

diff --git a/Engine/Source/Runtime/Engine/Private/Components/SkyLightComponent.cpp b/Engine/Source/Runtime/Engine/Private/Components/SkyLightComponent.cpp
--- a/Engine/Source/Runtime/Engine/Private/Components/SkyLightComponent.cpp
+++ b/Engine/Source/Runtime/Engine/Private/Components/SkyLightComponent.cpp
@@ -732,6 +732,12 @@ void USkyLightComponent::UpdateSkyCaptureContentsArray(UWorld* WorldToUpdate, TA
 			CaptureComponent->SourceType == SLS_SpecifiedCubemap &&
 			CaptureComponent->Cubemap &&
 			CaptureComponent->Cubemap->IsDefaultTexture();
+
+		if (bIsCubemapCompiling)
+		{
+			// We should process this texture as soon as possible so we can have a proper skylight.
+			FTextureCompilingManager::Get().RequestPriorityChange(CaptureComponent->Cubemap, EQueuedWorkPriority::Highest);
+		}
 #endif
 
 		if (((!Owner || !Owner->GetLevel() || Owner->GetLevel()->bIsVisible) && CaptureComponent->GetWorld() == WorldToUpdate)
