DataLayers : Added option for DataLayers to optionally affect loaded actors in the editor

diff --git a/Engine/Source/Editor/DataLayerEditor/Private/DataLayer/DataLayerEditorSubsystem.cpp b/Engine/Source/Editor/DataLayerEditor/Private/DataLayer/DataLayerEditorSubsystem.cpp
--- a/Engine/Source/Editor/DataLayerEditor/Private/DataLayer/DataLayerEditorSubsystem.cpp
+++ b/Engine/Source/Editor/DataLayerEditor/Private/DataLayer/DataLayerEditorSubsystem.cpp
@@ -5,6 +5,8 @@
 #include "WorldPartition/DataLayer/DataLayer.h"
 #include "WorldPartition/DataLayer/WorldDataLayers.h"
 #include "WorldPartition/WorldPartitionSubsystem.h"
+#include "WorldPartition/WorldPartition.h"
+#include "WorldPartition/DataLayer/DataLayerEditorPerProjectUserSettings.h"
 #include "Modules/ModuleManager.h"
 #include "DataLayerEditorModule.h"
 #include "ActorEditorUtils.h"
@@ -171,6 +173,36 @@ void UDataLayerEditorSubsystem::Deinitialize()
 	DataLayersBroadcast->Deinitialize();
 }
 
+bool UDataLayerEditorSubsystem::RefreshWorldPartitionEditorCells()
+{
+	if (UWorldPartition* WorldPartition = GetWorld() ? GetWorld()->GetWorldPartition() : nullptr)
+	{
+		if (!WorldPartition->RefreshLoadedEditorCells())
+		{
+			return false;
+		}
+		UpdateDataLayerEditorPerProjectUserSettings();
+	}
+	return true;
+}
+
+void UDataLayerEditorSubsystem::UpdateDataLayerEditorPerProjectUserSettings()
+{
+	if (AWorldDataLayers* WorldDataLayers = GetWorldDataLayers())
+	{
+		TArray<FName> DataLayersNotLoadedInEditor;
+		WorldDataLayers->ForEachDataLayer([&DataLayersNotLoadedInEditor](UDataLayer* DataLayer)
+		{
+			if (!DataLayer->IsDynamicallyLoadedInEditor())
+			{
+				DataLayersNotLoadedInEditor.Add(DataLayer->GetFName());
+			}
+			return true;
+		});
+		GetMutableDefault<UDataLayerEditorPerProjectUserSettings>()->SetWorldDataLayersNotLoadedInEditor(GetWorld(), DataLayersNotLoadedInEditor);
+	}
+}
+
 void UDataLayerEditorSubsystem::EditorMapChange()
 {
 	if (UWorld * World = GetWorld())
@@ -190,6 +222,7 @@ void UDataLayerEditorSubsystem::PostUndoRedo()
 {
 	DataLayerChanged.Broadcast(EDataLayerAction::Reset, NULL, NAME_None);
 	UpdateAllActorsVisibility(true, true);
+	RefreshWorldPartitionEditorCells();
 }
 
 ///////////////////////////////////////////////////////////////////////////////////////////////////////////////////////
@@ -968,12 +1001,7 @@ void UDataLayerEditorSubsystem::SetDataLayersVisibility(const TArray<UDataLayer*
 void UDataLayerEditorSubsystem::ToggleDataLayerVisibility(UDataLayer* DataLayer)
 {
 	check(DataLayer);
-
-	DataLayer->Modify();
-	DataLayer->SetVisible(!DataLayer->IsVisible());
-
-	DataLayerChanged.Broadcast(EDataLayerAction::Modify, DataLayer, "bIsVisible");
-	UpdateAllActorsVisibility(true, true);
+	SetDataLayerVisibility(DataLayer, !DataLayer->IsVisible());
 }
 
 void UDataLayerEditorSubsystem::ToggleDataLayersVisibility(const TArray<UDataLayer*>& DataLayers)
@@ -1012,57 +1040,94 @@ void UDataLayerEditorSubsystem::MakeAllDataLayersVisible()
 	UpdateAllActorsVisibility(true, true);
 }
 
-void UDataLayerEditorSubsystem::SetDataLayerIsDynamicallyLoaded(UDataLayer* DataLayer, const bool bIsDynamicallyLoaded)
+bool UDataLayerEditorSubsystem::SetDataLayerIsDynamicallyLoadedInternal(UDataLayer* DataLayer, const bool bIsDynamicallyLoaded)
+{
+	check(DataLayer);
+	if (DataLayer->IsDynamicallyLoaded() != bIsDynamicallyLoaded)
+	{
+		DataLayer->Modify();
+		DataLayer->SetIsDynamicallyLoaded(bIsDynamicallyLoaded);
+		DataLayerChanged.Broadcast(EDataLayerAction::Modify, DataLayer, "bIsDynamicallyLoaded");
+		return true;
+	}
+	return false;
+}
+
+bool UDataLayerEditorSubsystem::SetDataLayerIsDynamicallyLoaded(UDataLayer* DataLayer, const bool bIsDynamicallyLoaded)
 {
-	SetDataLayersIsDynamicallyLoaded({ DataLayer }, bIsDynamicallyLoaded);
+	bool bRefreshNeeded = SetDataLayerIsDynamicallyLoadedInternal(DataLayer, bIsDynamicallyLoaded);
+	return bRefreshNeeded ? RefreshWorldPartitionEditorCells() : true;
 }
 
-void UDataLayerEditorSubsystem::SetDataLayersIsDynamicallyLoaded(const TArray<UDataLayer*>& DataLayers, const bool bIsDynamicallyLoaded)
+bool UDataLayerEditorSubsystem::SetDataLayersIsDynamicallyLoaded(const TArray<UDataLayer*>& DataLayers, const bool bIsDynamicallyLoaded)
 {
-	bool bChangeOccurred = false;
+	bool bRefreshNeeded = false;
 	for (UDataLayer* DataLayer : DataLayers)
 	{
-		if (DataLayer->IsDynamicallyLoaded() != bIsDynamicallyLoaded)
-		{
-			DataLayer->Modify();
-			DataLayer->SetIsDynamicallyLoaded(bIsDynamicallyLoaded);
-			DataLayerChanged.Broadcast(EDataLayerAction::Modify, DataLayer, "bIsDynamicallyLoaded");
-			bChangeOccurred = true;
-		}
-		}
+		bRefreshNeeded = bRefreshNeeded || SetDataLayerIsDynamicallyLoadedInternal(DataLayer, bIsDynamicallyLoaded);
+	}
+	return bRefreshNeeded ? RefreshWorldPartitionEditorCells() : true;
+}
 
-	if (bChangeOccurred)
+bool UDataLayerEditorSubsystem::ToggleDataLayerIsDynamicallyLoaded(UDataLayer* DataLayer)
+{
+	check(DataLayer);
+	return SetDataLayerIsDynamicallyLoaded(DataLayer, !DataLayer->IsDynamicallyLoaded());
+}
+
+bool UDataLayerEditorSubsystem::ToggleDataLayersIsDynamicallyLoaded(const TArray<UDataLayer*>& DataLayers)
+{
+	bool bRefreshNeeded = false;
+	for (UDataLayer* DataLayer : DataLayers)
 	{
-		GEditor->RedrawLevelEditingViewports();
+		bRefreshNeeded = bRefreshNeeded || SetDataLayerIsDynamicallyLoadedInternal(DataLayer, !DataLayer->IsDynamicallyLoaded());
 	}
+	return bRefreshNeeded ? RefreshWorldPartitionEditorCells() : true;
 }
 
-void UDataLayerEditorSubsystem::ToggleDataLayerIsDynamicallyLoaded(UDataLayer* DataLayer)
+bool UDataLayerEditorSubsystem::SetDataLayerIsDynamicallyLoadedInEditorInternal(UDataLayer* DataLayer, const bool bIsDynamicallyLoadedInEditor)
 {
 	check(DataLayer);
+	if (DataLayer->IsDynamicallyLoadedInEditor() != bIsDynamicallyLoadedInEditor)
+	{
+		DataLayer->Modify(false);
+		DataLayer->SetIsDynamicallyLoadedInEditor(bIsDynamicallyLoadedInEditor);
+		DataLayerChanged.Broadcast(EDataLayerAction::Modify, DataLayer, "bIsDynamicallyLoadedInEditor");
+		return true;
+	}
+	return false;
+}
 
-	DataLayer->Modify();
-	DataLayer->SetIsDynamicallyLoaded(!DataLayer->IsDynamicallyLoaded());
-
-	DataLayerChanged.Broadcast(EDataLayerAction::Modify, DataLayer, "bIsDynamicallyLoaded");
-	GEditor->RedrawLevelEditingViewports();
+bool UDataLayerEditorSubsystem::SetDataLayerIsDynamicallyLoadedInEditor(UDataLayer* DataLayer, const bool bIsDynamicallyLoadedInEditor)
+{
+	bool bRefreshNeeded = SetDataLayerIsDynamicallyLoadedInEditorInternal(DataLayer, bIsDynamicallyLoadedInEditor);
+	return bRefreshNeeded ? RefreshWorldPartitionEditorCells() : true;
 }
 
-void UDataLayerEditorSubsystem::ToggleDataLayersIsDynamicallyLoaded(const TArray<UDataLayer*>& DataLayers)
+bool UDataLayerEditorSubsystem::SetDataLayersIsDynamicallyLoadedInEditor(const TArray<UDataLayer*>& DataLayers, const bool bIsDynamicallyLoadedInEditor)
 {
-	if (DataLayers.Num() == 0)
+	bool bRefreshNeeded = false;
+	for (UDataLayer* DataLayer : DataLayers)
 	{
-		return;
+		bRefreshNeeded = bRefreshNeeded || SetDataLayerIsDynamicallyLoadedInEditorInternal(DataLayer, bIsDynamicallyLoadedInEditor);
 	}
+	return bRefreshNeeded ? RefreshWorldPartitionEditorCells() : true;
+}
 
+bool UDataLayerEditorSubsystem::ToggleDataLayerIsDynamicallyLoadedInEditor(UDataLayer* DataLayer)
+{
+	check(DataLayer);
+	return SetDataLayerIsDynamicallyLoadedInEditor(DataLayer, !DataLayer->IsDynamicallyLoadedInEditor());
+}
+
+bool UDataLayerEditorSubsystem::ToggleDataLayersIsDynamicallyLoadedInEditor(const TArray<UDataLayer*>& DataLayers)
+{
+	bool bRefreshNeeded = false;
 	for (UDataLayer* DataLayer : DataLayers)
 	{
-		DataLayer->Modify();
-		DataLayer->SetIsDynamicallyLoaded(!DataLayer->IsDynamicallyLoaded());
-		DataLayerChanged.Broadcast(EDataLayerAction::Modify, DataLayer, "bIsDynamicallyLoaded");
+		bRefreshNeeded = bRefreshNeeded || SetDataLayerIsDynamicallyLoadedInEditorInternal(DataLayer, !DataLayer->IsDynamicallyLoadedInEditor());
 	}
-
-	GEditor->RedrawLevelEditingViewports();
+	return bRefreshNeeded ? RefreshWorldPartitionEditorCells() : true;
 }
 
 const UDataLayer* UDataLayerEditorSubsystem::GetDataLayerFromName(const FName& DataLayerName) const
diff --git a/Engine/Source/Editor/DataLayerEditor/Private/DataLayer/DataLayerModeSettings.h b/Engine/Source/Editor/DataLayerEditor/Private/DataLayer/DataLayerModeSettings.h
deleted file mode 100644
diff --git a/Engine/Source/Editor/DataLayerEditor/Private/DataLayer/DataLayerOutlinerIsDynamicallyLoadedColumn.cpp b/Engine/Source/Editor/DataLayerEditor/Private/DataLayer/DataLayerOutlinerIsDynamicallyLoadedColumn.cpp
--- a/Engine/Source/Editor/DataLayerEditor/Private/DataLayer/DataLayerOutlinerIsDynamicallyLoadedColumn.cpp
+++ b/Engine/Source/Editor/DataLayerEditor/Private/DataLayer/DataLayerOutlinerIsDynamicallyLoadedColumn.cpp
@@ -5,6 +5,7 @@
 #include "DataLayerTreeItem.h"
 #include "Widgets/Images/SImage.h"
 #include "Widgets/Input/SButton.h"
+#include "Widgets/Input/SCheckBox.h"
 #include "Widgets/Views/STreeView.h"
 #include "ScopedTransaction.h"
 #include "EditorStyleSet.h"
@@ -21,7 +22,7 @@ FName FDataLayerOutlinerIsDynamicallyLoadedColumn::GetID()
 SHeaderRow::FColumn::FArguments FDataLayerOutlinerIsDynamicallyLoadedColumn::ConstructHeaderRowColumn()
 {
 	return SHeaderRow::Column(GetColumnID())
-		.FixedWidth(24.f)
+		.FixedWidth(40.f)
 		.HAlignHeader(HAlign_Center)
 		.VAlignHeader(VAlign_Center)
 		.HAlignCell(HAlign_Center)
@@ -38,6 +39,7 @@ const TSharedRef<SWidget> FDataLayerOutlinerIsDynamicallyLoadedColumn::Construct
 	{
 		return SNew(SHorizontalBox)
 			+ SHorizontalBox::Slot()
+			.Padding(2, 0, 0, 0)
 			.AutoWidth()
 			.VAlign(VAlign_Center)
 			[
@@ -49,6 +51,7 @@ const TSharedRef<SWidget> FDataLayerOutlinerIsDynamicallyLoadedColumn::Construct
 					FDataLayerTreeItem* DataLayerTreeItem = TreeItem->CastTo<FDataLayerTreeItem>();
 					if (UDataLayer* DataLayer = DataLayerTreeItem->GetDataLayer())
 					{
+						bool bSuccess = false;
 						UDataLayerEditorSubsystem* DataLayerEditorSubsystem = UDataLayerEditorSubsystem::Get();
 						const auto& Tree = WeakSceneOutliner.Pin()->GetTree();
 						if (Tree.IsItemSelected(TreeItem))
@@ -67,18 +70,23 @@ const TSharedRef<SWidget> FDataLayerOutlinerIsDynamicallyLoadedColumn::Construct
 								}
 							}
 
-							const FScopedTransaction Transaction(LOCTEXT("ToggleDataLayersIsDynamicallyLoaded", "Toggle DataLayers Actors Loading"));
-							DataLayerEditorSubsystem->ToggleDataLayersIsDynamicallyLoaded(AllSelectedDataLayers);
+							const FScopedTransaction Transaction(LOCTEXT("ToggleDataLayersIsDynamicallyLoaded", "Toggle DataLayers Runtime Dynamically Loaded Flag"));
+							bSuccess = DataLayerEditorSubsystem->ToggleDataLayersIsDynamicallyLoaded(AllSelectedDataLayers);
 						}
 						else
 						{
-							const FScopedTransaction Transaction(LOCTEXT("ToggleDataLayerIsDynamicallyLoaded", "Toggle DataLayer Actors Loading"));
-							DataLayerEditorSubsystem->ToggleDataLayerIsDynamicallyLoaded(DataLayer);
+							const FScopedTransaction Transaction(LOCTEXT("ToggleDataLayerIsDynamicallyLoaded", "Toggle DataLayer Runtime Dynamically Loaded Flag"));
+							bSuccess = DataLayerEditorSubsystem->ToggleDataLayerIsDynamicallyLoaded(DataLayer);
+						}
+						if (!bSuccess)
+						{
+							// Cancelled : Undo last transaction
+							GEditor->Trans->Undo();
 						}
 					}
 					return FReply::Handled();
 				})
-				.ToolTipText(LOCTEXT("IsDynamicallyLoadedButtonToolTip", "Toggle DataLayer Actors Loading"))
+				.ToolTipText(LOCTEXT("IsDynamicallyLoadedButtonToolTip", "Toggle DataLayer Runtime Dynamically Loaded Flag"))
 				.HAlign(HAlign_Center)
 				.VAlign(VAlign_Center)
 				.Content()
@@ -98,6 +106,67 @@ const TSharedRef<SWidget> FDataLayerOutlinerIsDynamicallyLoadedColumn::Construct
 						}
 					})
 				]
+			]
+			+ SHorizontalBox::Slot()
+			.Padding(2, 0, 0, 0)
+			.AutoWidth()
+			.VAlign(VAlign_Center)
+			[
+				SNew(SCheckBox)
+				.IsEnabled_Lambda([this, TreeItem]()
+				{
+					FDataLayerTreeItem* DataLayerTreeItem = TreeItem->CastTo<FDataLayerTreeItem>();
+					UDataLayer* DataLayer = DataLayerTreeItem->GetDataLayer();
+					return DataLayer && DataLayer->IsDynamicallyLoaded();
+				})
+				.IsChecked_Lambda([this, TreeItem]()
+				{
+					FDataLayerTreeItem* DataLayerTreeItem = TreeItem->CastTo<FDataLayerTreeItem>();
+					UDataLayer* DataLayer = DataLayerTreeItem->GetDataLayer();
+					return DataLayer && DataLayer->IsDynamicallyLoadedInEditor() ? ECheckBoxState::Checked : ECheckBoxState::Unchecked;
+				})
+				.OnCheckStateChanged_Lambda([this, TreeItem](ECheckBoxState NewState)
+				{
+					bool bNewState = (NewState == ECheckBoxState::Checked);
+					FDataLayerTreeItem* DataLayerTreeItem = TreeItem->CastTo<FDataLayerTreeItem>();
+					if (UDataLayer* DataLayer = DataLayerTreeItem->GetDataLayer())
+					{
+						bool bSuccess = false;
+						UDataLayerEditorSubsystem* DataLayerEditorSubsystem = UDataLayerEditorSubsystem::Get();
+						const auto& Tree = WeakSceneOutliner.Pin()->GetTree();
+						if (Tree.IsItemSelected(TreeItem))
+						{
+							// Toggle IsDynamicallyLoadedInEditor flag of selected DataLayers to the same state as the given DataLayer
+							bool bIsDynamicallyLoadedInEditor = DataLayer->IsDynamicallyLoadedInEditor();
+
+							TArray<UDataLayer*> AllSelectedDataLayers;
+							for (auto& SelectedItem : Tree.GetSelectedItems())
+							{
+								FDataLayerTreeItem* SelectedDataLayerTreeItem = SelectedItem->CastTo<FDataLayerTreeItem>();
+								UDataLayer* SelectedDataLayer = SelectedDataLayerTreeItem ? SelectedDataLayerTreeItem->GetDataLayer() : nullptr;
+								if (SelectedDataLayer && SelectedDataLayer->IsDynamicallyLoadedInEditor() == bIsDynamicallyLoadedInEditor)
+								{
+									AllSelectedDataLayers.Add(SelectedDataLayer);
+								}
+							}
+
+							const FScopedTransaction Transaction(LOCTEXT("ToggleDataLayersIsDynamicallyLoadedInEditor", "Toggle DataLayers Editor Dynamically Loaded Flag"));
+							bSuccess = DataLayerEditorSubsystem->ToggleDataLayersIsDynamicallyLoadedInEditor(AllSelectedDataLayers);
+						}
+						else
+						{
+							const FScopedTransaction Transaction(LOCTEXT("ToggleDataLayerIsDynamicallyLoadedInEditor", "Toggle DataLayer Editor Dynamically Loaded Flag"));
+							bSuccess = DataLayerEditorSubsystem->ToggleDataLayerIsDynamicallyLoadedInEditor(DataLayer);
+						}
+						if (!bSuccess)
+						{
+							// Cancelled : Undo last transaction
+							GEditor->Trans->Undo();
+						}
+					}
+				})
+				.ToolTipText(LOCTEXT("IsDynamicallyLoadedInEditorCheckBoxToolTip", "Toggle DataLayer Editor Dynamically Loaded Flag"))
+				.HAlign(HAlign_Center)
 			];
 	}
 	return SNullWidget::NullWidget;
diff --git a/Engine/Source/Editor/DataLayerEditor/Private/DataLayer/SDataLayerBrowser.cpp b/Engine/Source/Editor/DataLayerEditor/Private/DataLayer/SDataLayerBrowser.cpp
--- a/Engine/Source/Editor/DataLayerEditor/Private/DataLayer/SDataLayerBrowser.cpp
+++ b/Engine/Source/Editor/DataLayerEditor/Private/DataLayer/SDataLayerBrowser.cpp
@@ -7,7 +7,7 @@
 #include "DataLayerActorTreeItem.h"
 #include "DataLayerTreeItem.h"
 #include "DataLayerMode.h"
-#include "DataLayerModeSettings.h"
+#include "WorldPartition/DataLayer/DataLayerEditorPerProjectUserSettings.h"
 #include "DataLayerOutlinerIsDynamicallyLoadedColumn.h"
 #include "DataLayerOutlinerDeleteButtonColumn.h"
 #include "DataLayer/DataLayerEditorSubsystem.h"
@@ -19,8 +19,7 @@
 
 void SDataLayerBrowser::Construct(const FArguments& InArgs)
 {
-	const UDataLayerModeSettings* SharedSettings = GetDefault<UDataLayerModeSettings>();
-	Mode = SharedSettings->bShowDataLayerContent ? EDataLayerBrowserMode::DataLayerContents : EDataLayerBrowserMode::DataLayers;
+	Mode = GetDefault<UDataLayerEditorPerProjectUserSettings>()->GetShowDataLayerContent() ? EDataLayerBrowserMode::DataLayerContents : EDataLayerBrowserMode::DataLayers;
 
 	auto ToggleDataLayerContents = [this]()
 	{
@@ -149,9 +148,7 @@ void SDataLayerBrowser::SetupDataLayerMode(EDataLayerBrowserMode InNewMode)
 
 	Mode = InNewMode;
 	
-	UDataLayerModeSettings* SharedSettings = GetMutableDefault<UDataLayerModeSettings>();
-	SharedSettings->bShowDataLayerContent = Mode == EDataLayerBrowserMode::DataLayerContents;
-	SharedSettings->SaveConfig();
+	GetMutableDefault<UDataLayerEditorPerProjectUserSettings>()->SetShowDataLayerContent(Mode == EDataLayerBrowserMode::DataLayerContents);
 
 	ModeChanged.Broadcast(Mode);
 }
diff --git a/Engine/Source/Editor/DataLayerEditor/Public/DataLayer/DataLayerEditorSubsystem.h b/Engine/Source/Editor/DataLayerEditor/Public/DataLayer/DataLayerEditorSubsystem.h
--- a/Engine/Source/Editor/DataLayerEditor/Public/DataLayer/DataLayerEditorSubsystem.h
+++ b/Engine/Source/Editor/DataLayerEditor/Public/DataLayer/DataLayerEditorSubsystem.h
@@ -466,7 +466,7 @@ public:
 	 * @param	DataLayer				The DataLayer to affect.
 	 * @param	bIsDynamicallyLoaded	If true the DataLayer will affect runtime streaming and actors will only be loaded if the DataLayer is active.
 	 */
-	virtual void SetDataLayerIsDynamicallyLoaded(UDataLayer* DataLayer, const bool bIsDynamicallyLoaded);
+	virtual bool SetDataLayerIsDynamicallyLoaded(UDataLayer* DataLayer, const bool bIsDynamicallyLoaded);
 
 	/**
 	 * Changes the IsDynamicallyLoaded flag of the DataLayers to the provided state
@@ -474,21 +474,55 @@ public:
 	 * @param	DataLayers				The DataLayers to affect
 	 * @param	bIsDynamicallyLoaded	If true the DataLayer will affect runtime streaming and actors will only be loaded if the DataLayer is active.
 	 */
-	virtual void SetDataLayersIsDynamicallyLoaded(const TArray<UDataLayer*>& DataLayers, const bool bIsDynamicallyLoaded);
+	virtual bool SetDataLayersIsDynamicallyLoaded(const TArray<UDataLayer*>& DataLayers, const bool bIsDynamicallyLoaded);
 
 	/**
 	 * Toggles the DataLayer's IsDynamicallyLoaded flag
 	 *
 	 * @param DataLayer	The DataLayer to affect
 	 */
-	virtual void ToggleDataLayerIsDynamicallyLoaded(UDataLayer* DataLayer);
+	virtual bool ToggleDataLayerIsDynamicallyLoaded(UDataLayer* DataLayer);
 
 	/**
 	 * Toggles the IsDynamicallyLoaded flag of all of the DataLayers
 	 *
 	 * @param	DataLayers	The DataLayers to affect
 	 */
-	virtual void ToggleDataLayersIsDynamicallyLoaded(const TArray<UDataLayer*>& DataLayers);
+	virtual bool ToggleDataLayersIsDynamicallyLoaded(const TArray<UDataLayer*>& DataLayers);
+
+	/**
+	 * Changes the DataLayer's IsDynamicallyLoadedInEditor flag to the provided state
+	 *
+	 * @param	DataLayer						The DataLayer to affect.
+	 * @param	bIsDynamicallyLoadedInEditor	The new value of the flag IsDynamicallyLoadedInEditor.
+	 *											If the DataLayer is DynamicallyLoaded, the Editor loading will consider this DataLayer to load or not an Actor.
+	 *											An Actor will not be loaded in the Editor if all its DataLayers are DynamicallyLoaded and not DynamicallyLoadedInEditor.
+	 */
+	virtual bool SetDataLayerIsDynamicallyLoadedInEditor(UDataLayer* DataLayer, const bool bIsDynamicallyLoadedInEditor);
+
+	/**
+	 * Changes the IsDynamicallyLoadedInEditor flag of the DataLayers to the provided state
+	 *
+	 * @param	DataLayers						The DataLayers to affect
+	 * @param	bIsDynamicallyLoadedInEditor	The new value of the flag IsDynamicallyLoadedInEditor.
+	 *											If the DataLayer is DynamicallyLoaded, the Editor loading will consider this DataLayer to load or not an Actor.
+	 *											An Actor will not be loaded in the Editor if all its DataLayers are DynamicallyLoaded and not DynamicallyLoadedInEditor.
+	 */
+	virtual bool SetDataLayersIsDynamicallyLoadedInEditor(const TArray<UDataLayer*>& DataLayers, const bool bIsDynamicallyLoadedInEditor);
+
+	/**
+	 * Toggles the DataLayer's IsDynamicallyLoadedInEditor flag
+	 *
+	 * @param DataLayer	The DataLayer to affect
+	 */
+	virtual bool ToggleDataLayerIsDynamicallyLoadedInEditor(UDataLayer* DataLayer);
+
+	/**
+	 * Toggles the IsDynamicallyLoadedInEditor flag of all of the DataLayers
+	 *
+	 * @param	DataLayers	The DataLayers to affect
+	 */
+	virtual bool ToggleDataLayersIsDynamicallyLoadedInEditor(const TArray<UDataLayer*>& DataLayers);
 
 	/**
 	 * Set the visibility of all DataLayers to true
@@ -595,6 +629,11 @@ public:
 
 private:
 
+	bool SetDataLayerIsDynamicallyLoadedInternal(UDataLayer* DataLayer, const bool bIsDynamicallyLoaded);
+	bool SetDataLayerIsDynamicallyLoadedInEditorInternal(UDataLayer* DataLayer, const bool bIsDynamicallyLoadedInEditor);
+	bool RefreshWorldPartitionEditorCells();
+	void UpdateDataLayerEditorPerProjectUserSettings();
+
 	/** Fires whenever one or more DataLayer changes */
 	FOnDataLayerChanged DataLayerChanged;
 
diff --git a/Engine/Source/Runtime/Engine/Private/WorldPartition/DataLayer/DataLayer.cpp b/Engine/Source/Runtime/Engine/Private/WorldPartition/DataLayer/DataLayer.cpp
--- a/Engine/Source/Runtime/Engine/Private/WorldPartition/DataLayer/DataLayer.cpp
+++ b/Engine/Source/Runtime/Engine/Private/WorldPartition/DataLayer/DataLayer.cpp
@@ -12,6 +12,7 @@ UDataLayer::UDataLayer(const FObjectInitializer& ObjectInitializer)
 , DataLayerLabel(GetFName())
 , bIsVisible(true)
 , bIsDynamicallyLoaded(true)
+, bIsDynamicallyLoadedInEditor(true)
 #endif
 {
 }
@@ -46,4 +47,13 @@ void UDataLayer::SetIsDynamicallyLoaded(bool bInIsDynamicallyLoaded)
 	}
 }
 
+void UDataLayer::SetIsDynamicallyLoadedInEditor(bool bInIsDynamicallyLoadedInEditor)
+{
+	if (bIsDynamicallyLoadedInEditor != bInIsDynamicallyLoadedInEditor && (!bInIsDynamicallyLoadedInEditor || bIsDynamicallyLoaded))
+	{
+		Modify(false);
+		bIsDynamicallyLoadedInEditor = bInIsDynamicallyLoadedInEditor;
+	}
+}
+
 #endif
\ No newline at end of file
diff --git a/Engine/Source/Runtime/Engine/Private/WorldPartition/DataLayer/WorldDataLayers.cpp b/Engine/Source/Runtime/Engine/Private/WorldPartition/DataLayer/WorldDataLayers.cpp
--- a/Engine/Source/Runtime/Engine/Private/WorldPartition/DataLayer/WorldDataLayers.cpp
+++ b/Engine/Source/Runtime/Engine/Private/WorldPartition/DataLayer/WorldDataLayers.cpp
@@ -7,6 +7,9 @@
 #include "WorldPartition/DataLayer/WorldDataLayers.h"
 #include "WorldPartition/DataLayer/DataLayer.h"
 #include "EngineUtils.h"
+#if WITH_EDITOR
+#include "WorldPartition/DataLayer/DataLayerEditorPerProjectUserSettings.h"
+#endif
 
 #define LOCTEXT_NAMESPACE "WorldDataLayers"
 
@@ -190,7 +193,19 @@ void AWorldDataLayers::PostLoad()
 {
 	Super::PostLoad();
 
-#if !WITH_EDITOR
+#if WITH_EDITOR
+	// Initialize DataLayer's IsDynamicallyLoadedInEditor based on DataLayerEditorPerProjectUserSettings
+	if (const FDataLayerNames* SettingsDataLayersNotLoadedInEditor = GetDefault<UDataLayerEditorPerProjectUserSettings>()->GetWorldDataLayersNotLoadedInEditor(GetWorld()))
+	{
+		for (const FName& DataLayerName : SettingsDataLayersNotLoadedInEditor->DataLayers)
+		{
+			if (UDataLayer* DataLayer = const_cast<UDataLayer*>(GetDataLayerFromName(DataLayerName)))
+			{
+				DataLayer->SetIsDynamicallyLoadedInEditor(false);
+			}
+		}
+	}
+#else
 	// Build acceleration tables
 	for (const UDataLayer* DataLayer : WorldDataLayers)
 	{
diff --git a/Engine/Source/Runtime/Engine/Private/WorldPartition/WorldPartition.cpp b/Engine/Source/Runtime/Engine/Private/WorldPartition/WorldPartition.cpp
--- a/Engine/Source/Runtime/Engine/Private/WorldPartition/WorldPartition.cpp
+++ b/Engine/Source/Runtime/Engine/Private/WorldPartition/WorldPartition.cpp
@@ -554,81 +554,147 @@ void UWorldPartition::LoadEditorCells(const FBox& Box)
 	}
 }
 
-void UWorldPartition::UnloadEditorCells(const FBox& Box)
+bool UWorldPartition::RefreshLoadedEditorCells()
 {
-	TArray<UWorldPartitionEditorCell*> CellsToUnload;
-	if (EditorHash->GetIntersectingCells(Box, CellsToUnload))
+	auto GetCellsToRefresh = [this](TArray<UWorldPartitionEditorCell*>& CellsToRefresh)
 	{
-		FWorldPartionCellUpdateContext CellUpdateContext(this);
-		
-		TSet<UPackage*> ModifiedPackages;
-		TMap<FWorldPartitionActorDesc*, int32> UnloadCount;
-	
-		int32 NumActorsToUnload = Algo::TransformAccumulate(CellsToUnload, [&UnloadCount](UWorldPartitionEditorCell* Cell)
-		{ 
-			for (FWorldPartitionActorDesc* ActorDesc : Cell->LoadedActors)
+		CellsToRefresh.Add(EditorHash->GetAlwaysLoadedCell());
+		EditorHash->ForEachCell([&CellsToRefresh](UWorldPartitionEditorCell* Cell)
+		{
+			if (Cell->bLoaded)
 			{
-				UnloadCount.FindOrAdd(ActorDesc, 0)++;
+				CellsToRefresh.Add(Cell);
 			}
-						
-			return Cell->LoadedActors.Num(); 
-		}, 0);
-		
-		for (const TPair<FWorldPartitionActorDesc*, int32> Pair : UnloadCount)
+		});
+		return !CellsToRefresh.IsEmpty();
+	};
+
+	return UpdateEditorCells(GetCellsToRefresh, /*bIsCellShouldBeLoaded*/true);
+}
+
+void UWorldPartition::UnloadEditorCells(const FBox& Box)
+{
+	auto GetCellsToUnload = [this, Box](TArray<UWorldPartitionEditorCell*>& CellsToUnload)
+	{
+		return EditorHash->GetIntersectingCells(Box, CellsToUnload) > 0;
+	};
+
+	UpdateEditorCells(GetCellsToUnload, /*bIsCellShouldBeLoaded*/false);
+}
+
+bool UWorldPartition::UpdateEditorCells(TFunctionRef<bool(TArray<UWorldPartitionEditorCell*>&)> GetCellsToProcess, bool bIsCellShouldBeLoaded)
+{
+	FWorldPartionCellUpdateContext CellUpdateContext(this);
+
+	TArray<UWorldPartitionEditorCell*> CellsToProcess;
+	if (!GetCellsToProcess(CellsToProcess))
+	{
+		return true;
+	}
+
+	TSet<UPackage*> ModifiedPackages;
+	TMap<FWorldPartitionActorDesc*, int32> UnloadCount;
+
+	for (UWorldPartitionEditorCell* Cell : CellsToProcess)
+	{
+		for (FWorldPartitionActorDesc* ActorDesc : Cell->LoadedActors)
 		{
-			FWorldPartitionActorDesc* ActorDesc = Pair.Key;
-			// Only prompt if the actor will get unloaded by the unloading cells
-			if (ActorDesc->GetLoadedRefCount() == Pair.Value)
+			if (!bIsCellShouldBeLoaded || !ShouldActorBeLoaded(ActorDesc))
 			{
-				AActor* LoadedActor = ActorDesc->GetActor();
-				check(LoadedActor);
-				UPackage* ActorPackage = LoadedActor->GetExternalPackage();
-				if (ActorPackage && ActorPackage->IsDirty())
-				{
-					ModifiedPackages.Add(ActorPackage);
-				}
+				UnloadCount.FindOrAdd(ActorDesc, 0)++;
 			}
 		}
-		
-		// Make sure we save modified actor packages before unloading
-		FEditorFileUtils::EPromptReturnCode RetCode = FEditorFileUtils::PR_Success;
-		if (ModifiedPackages.Num())
+	}
+
+	bool bIsUpdateUnloadingActors = false;
+
+	for (const TPair<FWorldPartitionActorDesc*, int32> Pair : UnloadCount)
+	{
+		FWorldPartitionActorDesc* ActorDesc = Pair.Key;
+		// Only prompt if the actor will get unloaded by the unloading cells
+		if (ActorDesc->GetLoadedRefCount() == Pair.Value)
 		{
-			const bool bCheckDirty = false;
-			const bool bAlreadyCheckedOut = false;
-			const bool bCanBeDeclined = true;
-			const bool bPromptToSave = true;
-			const FText Title = LOCTEXT("SaveActorsTitle", "Save Actor(s)");
-			const FText Message = LOCTEXT("SaveActorsMessage", "Save Actor(s) before unloading them.");
-
-			RetCode = FEditorFileUtils::PromptForCheckoutAndSave(ModifiedPackages.Array(), bCheckDirty, bPromptToSave, Title, Message, nullptr, bAlreadyCheckedOut, bCanBeDeclined);
-			if (RetCode == FEditorFileUtils::PR_Cancelled)
+			AActor* LoadedActor = ActorDesc->GetActor();
+			check(LoadedActor);
+			UPackage* ActorPackage = LoadedActor->GetExternalPackage();
+			if (ActorPackage && ActorPackage->IsDirty())
 			{
-				return;
+				ModifiedPackages.Add(ActorPackage);
 			}
+			bIsUpdateUnloadingActors = true;
+		}
+	}
 
-			check(RetCode != FEditorFileUtils::PR_Failure);
+	// Make sure we save modified actor packages before unloading
+	FEditorFileUtils::EPromptReturnCode RetCode = FEditorFileUtils::PR_Success;
+	if (ModifiedPackages.Num())
+	{
+		const bool bCheckDirty = false;
+		const bool bAlreadyCheckedOut = false;
+		const bool bCanBeDeclined = true;
+		const bool bPromptToSave = true;
+		const FText Title = LOCTEXT("SaveActorsTitle", "Save Actor(s)");
+		const FText Message = LOCTEXT("SaveActorsMessage", "Save Actor(s) before unloading them.");
+
+		RetCode = FEditorFileUtils::PromptForCheckoutAndSave(ModifiedPackages.Array(), bCheckDirty, bPromptToSave, Title, Message, nullptr, bAlreadyCheckedOut, bCanBeDeclined);
+		if (RetCode == FEditorFileUtils::PR_Cancelled)
+		{
+			return false;
 		}
 
+		check(RetCode != FEditorFileUtils::PR_Failure);
+	}
+
+	if (bIsUpdateUnloadingActors)
+	{
 		GEditor->SelectNone(true, true);
+	}
+
+	// At this point, cells might have changed due to saving deleted actors
+	CellsToProcess.Empty(CellsToProcess.Num());
+	if (!GetCellsToProcess(CellsToProcess))
+	{
+		return true;
+	}
+
+	FScopedSlowTask SlowTask(CellsToProcess.Num(), LOCTEXT("UpdatingCells", "Updating cells..."));
+	SlowTask.MakeDialog();
+
+	for (UWorldPartitionEditorCell* Cell : CellsToProcess)
+	{
+		SlowTask.EnterProgressFrame(1);
+		UpdateLoadingEditorCell(Cell, bIsCellShouldBeLoaded);
+	}
 
-		// At this point, cells might have changed due to saving deleted actors
-		CellsToUnload.Empty(CellsToUnload.Num());
-		if (EditorHash->GetIntersectingCells(Box, CellsToUnload))
-		{	
-			FScopedSlowTask SlowTask(NumActorsToUnload, LOCTEXT("UnloadingCells", "Unloading cells..."));
-			SlowTask.MakeDialog();
+	if (bIsUpdateUnloadingActors)
+	{
+		GEditor->ResetTransaction(LOCTEXT("UnloadingEditorCellsResetTrans", "Unloading Cells"));
+	}
+
+	return true;
+}
 
-			for (UWorldPartitionEditorCell* Cell: CellsToUnload)
+bool UWorldPartition::ShouldActorBeLoaded(const FWorldPartitionActorDesc* ActorDesc) const
+{
+	if (const AWorldDataLayers* WorldDataLayers = AWorldDataLayers::Get(GetWorld()))
+	{
+		uint32 NumValidLayers = 0;
+		for (const FName& DataLayerName : ActorDesc->GetDataLayers())
+		{
+			if (const UDataLayer* DataLayer = WorldDataLayers->GetDataLayerFromName(DataLayerName))
 			{
-				SlowTask.EnterProgressFrame(Cell->LoadedActors.Num());
-				UpdateLoadingEditorCell(Cell, false);
+				if (DataLayer->IsDynamicallyLoadedInEditor())
+				{
+					return true;
+				}
+				NumValidLayers++;
 			}
-		
-			GEditor->ResetTransaction(LOCTEXT("UnloadingEditorCellsResetTrans", "Unloading Cells"));
 		}
+		return !NumValidLayers;
 	}
-}
+
+	return true;
+};
 
 // Loads actors in Editor cell
 void UWorldPartition::UpdateLoadingEditorCell(UWorldPartitionEditorCell* Cell, bool bShouldBeLoaded)
@@ -683,32 +749,40 @@ void UWorldPartition::UpdateLoadingEditorCell(UWorldPartitionEditorCell* Cell, b
 			// or directly referenced by a loaded cell.
 			check(Actor || !ActorDesc->GetLoadedRefCount());
 
-			bool bIsAlreadyInLoadedActors = false;
-			Cell->LoadedActors.Add(ActorDesc, &bIsAlreadyInLoadedActors);
-
-			if (bIsAlreadyInLoadedActors)
+			// Filter actor against DataLayers
+			if (ShouldActorBeLoaded(ActorDesc))
 			{
-				// We already hold a reference to this actor
-				check(Actor);
-				check(ActorDesc->GetLoadedRefCount());
-				UE_LOG(LogWorldPartition, Verbose, TEXT(" ==> Skipped already loaded actor %s"), *Actor->GetFullName());
-			}
-			else
-			{
-				const uint32 ActorRefCount = ActorDesc->AddLoadedRefCount();
+				bool bIsAlreadyInLoadedActors = false;
+				Cell->LoadedActors.Add(ActorDesc, &bIsAlreadyInLoadedActors);
 
-				if (ActorRefCount == 1)
+				if (bIsAlreadyInLoadedActors)
 				{
-					// Register actor
-					Actor = RegisterActor(ActorDesc);
+					// We already hold a reference to this actor
+					check(Actor);
+					check(ActorDesc->GetLoadedRefCount());
+					UE_LOG(LogWorldPartition, Verbose, TEXT(" ==> Skipped already loaded actor %s"), *Actor->GetFullName());
 				}
 				else
 				{
-					check(Actor);
-					UE_LOG(LogWorldPartition, Verbose, TEXT(" ==> Referenced unloaded actor %s(%d)"), *Actor->GetFullName(), ActorRefCount);
-				}
+					const uint32 ActorRefCount = ActorDesc->AddLoadedRefCount();
 
-				check(Actor && Actor->GetActorGuid() == ActorDesc->GetGuid());
+					if (ActorRefCount == 1)
+					{
+						// Register actor
+						Actor = RegisterActor(ActorDesc);
+					}
+					else
+					{
+						check(Actor);
+						UE_LOG(LogWorldPartition, Verbose, TEXT(" ==> Referenced unloaded actor %s(%d)"), *Actor->GetFullName(), ActorRefCount);
+					}
+
+					check(Actor && Actor->GetActorGuid() == ActorDesc->GetGuid());
+				}
+			}
+			else if (Cell->LoadedActors.Remove(ActorDesc))
+			{
+				UnloadActor(ActorDesc, Actor);
 			}
 		}
 	}
diff --git a/Engine/Source/Runtime/Engine/Public/WorldPartition/DataLayer/DataLayer.h b/Engine/Source/Runtime/Engine/Public/WorldPartition/DataLayer/DataLayer.h
--- a/Engine/Source/Runtime/Engine/Public/WorldPartition/DataLayer/DataLayer.h
+++ b/Engine/Source/Runtime/Engine/Public/WorldPartition/DataLayer/DataLayer.h
@@ -20,6 +20,8 @@ public:
 	void SetDataLayerLabel(FName InDataLayerLabel);
 	void SetVisible(bool bInIsVisible);
 	void SetIsDynamicallyLoaded(bool bInIsDynamicallyLoaded);
+	void SetIsDynamicallyLoadedInEditor(bool bInIsDynamicallyLoadedInEditor);
+	bool IsDynamicallyLoadedInEditor() const { return !IsDynamicallyLoaded() || bIsDynamicallyLoadedInEditor; }
 #endif
 	
 	FName GetDataLayerLabel() const  { return DataLayerLabel; }
@@ -35,7 +37,13 @@ private:
 	UPROPERTY()
 	uint32 bIsVisible : 1;
 
-	/** Whether the layer can affect actor loading or not when layer is activated */
+	/** Whether the DataLayer affects actor runtime loading */
 	UPROPERTY()
 	uint32 bIsDynamicallyLoaded : 1;
+
+#if WITH_EDITORONLY_DATA
+	/** Whether the DataLayer affects actor editor loading */
+	UPROPERTY(Transient)
+	uint32 bIsDynamicallyLoadedInEditor : 1;
+#endif
 };
\ No newline at end of file
diff --git a/Engine/Source/Runtime/Engine/Public/WorldPartition/DataLayer/DataLayerEditorPerProjectUserSettings.h b/Engine/Source/Runtime/Engine/Public/WorldPartition/DataLayer/DataLayerEditorPerProjectUserSettings.h
new file mode 100644
--- /dev/null
+++ b/Engine/Source/Runtime/Engine/Public/WorldPartition/DataLayer/DataLayerEditorPerProjectUserSettings.h
@@ -0,0 +1,65 @@
+// Copyright Epic Games, Inc. All Rights Reserved.
+
+#pragma once
+
+#include "CoreMinimal.h"
+#include "UObject/ObjectMacros.h"
+#include "UObject/Object.h"
+#include "DataLayerEditorPerProjectUserSettings.generated.h"
+
+USTRUCT()
+struct FDataLayerNames
+{
+	GENERATED_BODY();
+
+#if WITH_EDITORONLY_DATA
+	UPROPERTY()
+	TArray<FName> DataLayers;
+#endif
+};
+
+UCLASS(MinimalAPI, config = EditorPerProjectUserSettings)
+class UDataLayerEditorPerProjectUserSettings : public UObject
+{
+	GENERATED_BODY()
+
+public:
+#if WITH_EDITOR	
+	bool GetShowDataLayerContent() const { return bShowDataLayerContent; }
+	void SetShowDataLayerContent(bool bInShowDataLayerContent)
+	{
+		if (bShowDataLayerContent != bInShowDataLayerContent)
+		{
+			bShowDataLayerContent = bInShowDataLayerContent;
+			SaveConfig();
+		}
+	}
+
+	const FDataLayerNames* GetWorldDataLayersNotLoadedInEditor(UWorld* InWorld) const
+	{
+		return InWorld ? WorldDataLayersNotLoadedInEditor.Find(TSoftObjectPtr<UWorld>(InWorld)) : nullptr;
+	}
+
+	void SetWorldDataLayersNotLoadedInEditor(UWorld* InWorld, const TArray<FName>& InDataLayersNotLoadedInEditor)
+	{
+		if (InWorld)
+		{
+			FDataLayerNames& DataLayersNotLoadedInEditor = WorldDataLayersNotLoadedInEditor.FindOrAdd(TSoftObjectPtr<UWorld>(InWorld));
+			if (DataLayersNotLoadedInEditor.DataLayers != InDataLayersNotLoadedInEditor)
+			{
+				DataLayersNotLoadedInEditor.DataLayers = InDataLayersNotLoadedInEditor;
+				SaveConfig();
+			}
+		}
+	}
+#endif
+
+private:
+#if WITH_EDITORONLY_DATA
+	UPROPERTY(config)
+	uint32 bShowDataLayerContent : 1;
+
+	UPROPERTY(config)
+	TMap<TSoftObjectPtr<UWorld>, FDataLayerNames> WorldDataLayersNotLoadedInEditor;
+#endif
+};
\ No newline at end of file
diff --git a/Engine/Source/Runtime/Engine/Public/WorldPartition/WorldPartition.h b/Engine/Source/Runtime/Engine/Public/WorldPartition/WorldPartition.h
--- a/Engine/Source/Runtime/Engine/Public/WorldPartition/WorldPartition.h
+++ b/Engine/Source/Runtime/Engine/Public/WorldPartition/WorldPartition.h
@@ -116,6 +116,7 @@ public:
 
 	void LoadEditorCells(const FBox& Box);
 	void UnloadEditorCells(const FBox& Box);
+	bool RefreshLoadedEditorCells();
 
 	bool IsPreCooked() const { return bIsPreCooked; }
 	void SetIsPreCooked(bool bInIsPreCooked) { bIsPreCooked = bInIsPreCooked; }
@@ -206,6 +207,8 @@ private:
 	void UpdateLoadingEditorCell(UWorldPartitionEditorCell* Cell, bool bShouldBeLoaded);
 	void HashActorDesc(FWorldPartitionActorDesc* ActorDesc);
 	void UnhashActorDesc(FWorldPartitionActorDesc* ActorDesc);
+	bool ShouldActorBeLoaded(const FWorldPartitionActorDesc* ActorDesc) const;
+	bool UpdateEditorCells(TFunctionRef<bool(TArray<UWorldPartitionEditorCell*>&)> GetCellsToProcess, bool bIsCellShouldBeLoaded);
 #endif
 
 	UWorldPartitionStreamingPolicy* GetStreamingPolicy() const;
