Integrate of Cooking Analytics and StrudioAnalytics plug-in

diff --git a/Engine/Source/Runtime/Engine/Private/Analytics/StudioAnalytics.cpp b/Engine/Source/Runtime/Engine/Private/Analytics/StudioAnalytics.cpp
--- a/Engine/Source/Runtime/Engine/Private/Analytics/StudioAnalytics.cpp
+++ b/Engine/Source/Runtime/Engine/Private/Analytics/StudioAnalytics.cpp
@@ -23,6 +23,7 @@ bool FStudioAnalytics::bInitialized = false;
 volatile double FStudioAnalytics::TimeEstimation = 0;
 FThread FStudioAnalytics::TimerThread;
 TSharedPtr<IAnalyticsProviderET> FStudioAnalytics::Analytics;
+TArray<FAnalyticsEventAttribute> FStudioAnalytics::DefaultAttributes;
 
 void FStudioAnalytics::SetProvider(TSharedRef<IAnalyticsProviderET> InAnalytics)
 {
@@ -40,6 +41,34 @@ void FStudioAnalytics::SetProvider(TSharedRef<IAnalyticsProviderET> InAnalytics)
 	}
 }
 
+void FStudioAnalytics::ApplyDefaultEventAttributes()
+{
+	if (Analytics.IsValid())
+	{
+		// Get the current attributes
+		TArray<FAnalyticsEventAttribute> CurrentDefaultAttributes = Analytics->GetDefaultEventAttributesSafe();
+
+		// Append any new attributes to our current ones
+		CurrentDefaultAttributes += MoveTemp(DefaultAttributes);
+		DefaultAttributes.Reset();
+
+		// Set the new default attributes in the provider
+		Analytics->SetDefaultEventAttributes(MoveTemp(CurrentDefaultAttributes));
+	}	
+}
+
+void FStudioAnalytics::AddDefaultEventAttribute(const FAnalyticsEventAttribute& Attribute)
+{
+	// Append a single default attribute to the existing list
+	DefaultAttributes.Emplace(Attribute);
+}
+
+void FStudioAnalytics::AddDefaultEventAttributes(TArray<FAnalyticsEventAttribute>&& Attributes)
+{
+	// Append attributes list to the existing list
+	DefaultAttributes += MoveTemp(Attributes);
+}
+
 IAnalyticsProviderET& FStudioAnalytics::GetProvider()
 {
 	checkf(IsAvailable(), TEXT("FStudioAnalytics::GetProvider called outside of Initialize/Shutdown."));
diff --git a/Engine/Source/Runtime/Engine/Public/StudioAnalytics.h b/Engine/Source/Runtime/Engine/Public/StudioAnalytics.h
--- a/Engine/Source/Runtime/Engine/Public/StudioAnalytics.h
+++ b/Engine/Source/Runtime/Engine/Public/StudioAnalytics.h
@@ -17,6 +17,15 @@ class FStudioAnalytics : FNoncopyable
 public:
 	static ENGINE_API void SetProvider(TSharedRef<IAnalyticsProviderET> InAnalytics);
 
+	/** Add a list of attributes to the defaults */
+	static ENGINE_API void AddDefaultEventAttributes(TArray<FAnalyticsEventAttribute>&& Attributes);
+
+	/** Add a single attribute to the defaults */
+	static ENGINE_API void AddDefaultEventAttribute(const FAnalyticsEventAttribute& Attribute);
+
+	/** Applies the default attributes */
+	static ENGINE_API void ApplyDefaultEventAttributes();
+
 	/**
 	 * Return the provider instance. Not valid outside of Initialize/Shutdown calls.
 	 * Note: must check IsAvailable() first else this code will assert if the provider is not valid.
@@ -45,6 +54,7 @@ private:
 private:
 	static bool bInitialized;
 	static TSharedPtr<IAnalyticsProviderET> Analytics;
+	static TArray<FAnalyticsEventAttribute> DefaultAttributes;
 	static FThread TimerThread;
 	static volatile double TimeEstimation;
 };
