Allow async compilation to take place in Editor when starting in -game mode
  - Allow multi-threaded bulkdata reading even when EDL is turned on

diff --git a/Engine/Source/Runtime/CoreUObject/Private/Serialization/BulkData.cpp b/Engine/Source/Runtime/CoreUObject/Private/Serialization/BulkData.cpp
--- a/Engine/Source/Runtime/CoreUObject/Private/Serialization/BulkData.cpp
+++ b/Engine/Source/Runtime/CoreUObject/Private/Serialization/BulkData.cpp
@@ -715,7 +715,7 @@ bool FUntypedBulkData::LoadBulkDataWithFileReader()
 bool FUntypedBulkData::CanLoadBulkDataWithFileReader() const
 {
 #if WITH_EDITOR
-	return !Filename.IsEmpty() && GIsEditor && !GEventDrivenLoaderEnabled;
+	return !Filename.IsEmpty();
 #else
 	return false;
 #endif
diff --git a/Engine/Source/Runtime/Engine/Private/StaticMeshCompiler.cpp b/Engine/Source/Runtime/Engine/Private/StaticMeshCompiler.cpp
--- a/Engine/Source/Runtime/Engine/Private/StaticMeshCompiler.cpp
+++ b/Engine/Source/Runtime/Engine/Private/StaticMeshCompiler.cpp
@@ -372,14 +372,7 @@ void FStaticMeshCompilingManager::FinishStaticMeshCompilation(UStaticMesh* Stati
 
 bool FStaticMeshCompilingManager::IsAsyncCompilationAllowed(UStaticMesh* StaticMesh) const
 {
-	return 
-		// @todo Same requirement as FUntypedBulkData::LoadBulkDataWithFileReader() for now
-		// because if we can't load bulk data properly from staticmesh building thread,
-		// every staticmesh compilation will effectively be single-threaded anyway...
-		// -game mode is extremely slow when staticmesh compilation is required because
-		// of this limitation in the loader. Fix the loader and then remove this here!
-		GIsEditor && !GEventDrivenLoaderEnabled &&
-		IsAsyncStaticMeshCompilationEnabled();
+	return IsAsyncStaticMeshCompilationEnabled();
 }
 
 FStaticMeshCompilingManager& FStaticMeshCompilingManager::Get()
diff --git a/Engine/Source/Runtime/Engine/Private/TextureCompiler.cpp b/Engine/Source/Runtime/Engine/Private/TextureCompiler.cpp
--- a/Engine/Source/Runtime/Engine/Private/TextureCompiler.cpp
+++ b/Engine/Source/Runtime/Engine/Private/TextureCompiler.cpp
@@ -314,14 +314,7 @@ void FTextureCompilingManager::FinishTextureCompilation(UTexture* Texture)
 
 bool FTextureCompilingManager::IsAsyncCompilationAllowed(UTexture* Texture) const
 {
-	return 
-		// @todo Same requirement as FUntypedBulkData::LoadBulkDataWithFileReader() for now
-		// because if we can't load bulk data properly from texture building thread,
-		// every texture compilation will effectively be single-threaded anyway...
-		// -game mode is extremely slow when texture compilation is required because
-		// of this limitation in the loader. Fix the loader and then remove this here!
-		GIsEditor && !GEventDrivenLoaderEnabled &&
-		IsAsyncTextureCompilationEnabled();
+	return IsAsyncTextureCompilationEnabled();
 }
 
 FTextureCompilingManager& FTextureCompilingManager::Get()
