Allow strata to use local shared variables instead of inlining.
Fixes a few issues.

diff --git a/Engine/Source/Runtime/Engine/Private/Materials/HLSLMaterialTranslator.cpp b/Engine/Source/Runtime/Engine/Private/Materials/HLSLMaterialTranslator.cpp
--- a/Engine/Source/Runtime/Engine/Private/Materials/HLSLMaterialTranslator.cpp
+++ b/Engine/Source/Runtime/Engine/Private/Materials/HLSLMaterialTranslator.cpp
@@ -2024,7 +2024,7 @@ int32 FHLSLMaterialTranslator::AddCodeChunkInner(uint64 Hash, const TCHAR* Forma
 		return CodeIndex;
 	}
 	// Can only create temporaries for certain types
-	else if ((Type & (MCT_Float | MCT_VTPageTableResult)) || Type == MCT_ShadingModel || Type == MCT_MaterialAttributes)
+	else if ((Type & (MCT_Float | MCT_VTPageTableResult)) || Type == MCT_ShadingModel || Type == MCT_MaterialAttributes || Type == MCT_Strata)
 	{
 		// Check for existing
 		for (int32 i = 0; i < CurrentScopeChunks->Num(); ++i)
@@ -7236,7 +7236,7 @@ int32 FHLSLMaterialTranslator::FrontMaterial()
 
 int32 FHLSLMaterialTranslator::StrataDiffuseOrenNayarBSDF(int32 Albedo, int32 Roughness, int32 Normal)
 {
-	return AddInlinedCodeChunk(
+	return AddCodeChunk(
 		MCT_Strata, TEXT("GetStrataDiffuseOrenNayarBSDF(%s, %s, %s)"),
 		*GetParameterCode(Albedo),
 		*GetParameterCode(Roughness),
@@ -7246,7 +7246,7 @@ int32 FHLSLMaterialTranslator::StrataDiffuseOrenNayarBSDF(int32 Albedo, int32 Ro
 
 int32 FHLSLMaterialTranslator::StrataDiffuseChanBSDF(int32 Albedo, int32 Roughness, int32 Normal)
 {
-	return AddInlinedCodeChunk(
+	return AddCodeChunk(
 		MCT_Strata, TEXT("GetStrataDiffuseChanBSDF(%s, %s, %s)"),
 		*GetParameterCode(Albedo),
 		*GetParameterCode(Roughness),
@@ -7256,7 +7256,7 @@ int32 FHLSLMaterialTranslator::StrataDiffuseChanBSDF(int32 Albedo, int32 Roughne
 
 int32 FHLSLMaterialTranslator::StrataDielectricBSDF(int32 Roughness, int32 Reflectivity, int32 Tint, int32 Normal)
 {
-	return AddInlinedCodeChunk(
+	return AddCodeChunk(
 		MCT_Strata, TEXT("GetStrataDielectricBSDF(%s, %s, %s, %s)"),
 		*GetParameterCode(Roughness),
 		*GetParameterCode(Reflectivity),
@@ -7267,7 +7267,7 @@ int32 FHLSLMaterialTranslator::StrataDielectricBSDF(int32 Roughness, int32 Refle
 
 int32 FHLSLMaterialTranslator::StrataConductorBSDF(int32 Reflectivity, int32 EdgeColor, int32 Roughness, int32 Normal)
 {
-	return AddInlinedCodeChunk(
+	return AddCodeChunk(
 		MCT_Strata, TEXT("GetStrataConductorBSDF(%s, %s, %s, %s)"),
 		*GetParameterCode(Reflectivity),
 		*GetParameterCode(EdgeColor),
@@ -7278,7 +7278,7 @@ int32 FHLSLMaterialTranslator::StrataConductorBSDF(int32 Reflectivity, int32 Edg
 
 int32 FHLSLMaterialTranslator::StrataVolumeBSDF(int32 Albedo, int32 Extinction, int32 Anisotropy, int32 Thickness)
 {
-	return AddInlinedCodeChunk(
+	return AddCodeChunk(
 		MCT_Strata, TEXT("GetStrataVolumeBSDF(%s, %s, %s, %s, %s)"),
 		*GetParameterCode(Albedo),
 		*GetParameterCode(Extinction),
@@ -7294,7 +7294,7 @@ int32 FHLSLMaterialTranslator::StrataHorizontalMixing(int32 Foreground, int32 Ba
 	{
 		return INDEX_NONE;
 	}
-	return AddInlinedCodeChunk(
+	return AddCodeChunk(
 		MCT_Strata, TEXT("StrataHorizontalMixing(%s, %s, %s)"),
 		*GetParameterCode(Foreground),
 		*GetParameterCode(Background),
@@ -7308,7 +7308,7 @@ int32 FHLSLMaterialTranslator::StrataVerticalLayering(int32 Top, int32 Base)
 	{
 		return INDEX_NONE;
 	}
-	return AddInlinedCodeChunk(
+	return AddCodeChunk(
 		MCT_Strata, TEXT("StrataVerticalLayering(%s, %s)"),
 		*GetParameterCode(Top),
 		*GetParameterCode(Base)
@@ -7321,7 +7321,7 @@ int32 FHLSLMaterialTranslator::StrataAdd(int32 A, int32 B)
 	{
 		return INDEX_NONE;
 	}
-	return AddInlinedCodeChunk(
+	return AddCodeChunk(
 		MCT_Strata, TEXT("StrataAdd(%s, %s)"),
 		*GetParameterCode(A),
 		*GetParameterCode(B)
@@ -7334,7 +7334,7 @@ int32 FHLSLMaterialTranslator::StrataMultiply(int32 A, int32 Weight)
 	{
 		return INDEX_NONE;
 	}
-	return AddInlinedCodeChunk(
+	return AddCodeChunk(
 		MCT_Strata, TEXT("StrataMultiply(%s, %s)"),
 		*GetParameterCode(A),
 		*GetParameterCode(Weight)
@@ -7347,8 +7347,8 @@ int32 FHLSLMaterialTranslator::StrataArtisticIOR(int32 Reflectivity, int32 EdgeC
 	{
 		return INDEX_NONE;
 	}
-	return AddInlinedCodeChunk(
-		MCT_Strata, 
+	return AddCodeChunk(
+		MCT_Float3, 
 		OutputIndex == 0 ? TEXT("ComputeComplexIORFromF0AndEdgeTint_IOR(%s, %s)") : TEXT("ComputeComplexIORFromF0AndEdgeTint_Extinction(%s, %s)"),
 		*GetParameterCode(Reflectivity),
 		*GetParameterCode(EdgeColor)
@@ -7361,8 +7361,8 @@ int32 FHLSLMaterialTranslator::StrataPhysicalIOR(int32 IOR, int32 Extinction, in
 	{
 		return INDEX_NONE;
 	}
-	return AddInlinedCodeChunk(
-		MCT_Strata,
+	return AddCodeChunk(
+		MCT_Float3,
 		OutputIndex == 0 ? TEXT("ComputeF0AndEdgeTintFromComplexIOR_Reflectivity(%s, %s)") : TEXT("ComputeF0AndEdgeTintFromComplexIOR_EdgeTint(%s, %s)"),
 		*GetParameterCode(IOR),
 		*GetParameterCode(Extinction)
@@ -7375,7 +7375,7 @@ int32 FHLSLMaterialTranslator::StrataDielectricIORToReflectivity(int32 IOR, int3
 	{
 		return INDEX_NONE;
 	}
-	return AddInlinedCodeChunk(
+	return AddCodeChunk(
 		MCT_Float,
 		TEXT("DielectricIorToF0(%s)"),
 		*GetParameterCode(IOR));
