Downgrade warning to log for staticmesh that should be resaved

diff --git a/Engine/Source/Runtime/Engine/Private/StaticMesh.cpp b/Engine/Source/Runtime/Engine/Private/StaticMesh.cpp
--- a/Engine/Source/Runtime/Engine/Private/StaticMesh.cpp
+++ b/Engine/Source/Runtime/Engine/Private/StaticMesh.cpp
@@ -5430,9 +5430,9 @@ void UStaticMesh::BeginPostLoadInternal(FStaticMeshPostLoadContext& Context)
 		const int32 CustomVersion = GetLinkerCustomVersion(FReleaseObjectVersion::GUID);
 		if (GetLinkerUE4Version() < VER_UE4_STATIC_MESH_EXTENDED_BOUNDS || CustomVersion < FReleaseObjectVersion::StaticMeshExtendedBoundsFix)
 		{
-			// A stall is almost guaranteed during async build because mesh bounds are used extensively from a many different places.
+			// A stall is almost guaranteed during async build because mesh bounds are used extensively from many different places.
 			Context.bShouldComputeExtendedBounds = true;
-			UE_LOG(LogStaticMesh, Warning, TEXT("%s should be resaved to improve async compilation performance."), *GetFullName());
+			UE_LOG(LogStaticMesh, Log, TEXT("%s should be resaved to improve async compilation performance."), *GetFullName());
 		}
 		else
 		{
