Fixed non-editor builds.

diff --git a/Engine/Source/Runtime/Engine/Private/StaticMesh.cpp b/Engine/Source/Runtime/Engine/Private/StaticMesh.cpp
--- a/Engine/Source/Runtime/Engine/Private/StaticMesh.cpp
+++ b/Engine/Source/Runtime/Engine/Private/StaticMesh.cpp
@@ -5952,13 +5952,14 @@ UStaticMeshDescription* UStaticMesh::CreateStaticMeshDescription(UObject* Outer)
 
 UStaticMeshDescription* UStaticMesh::GetStaticMeshDescription(int32 LODIndex)
 {
+#if WITH_EDITOR
 	if (LODIndex < GetNumSourceModels())
 	{
 		UStaticMeshDescription* StaticMeshDescription = NewObject<UStaticMeshDescription>(this, NAME_None, RF_Transient);
 		StaticMeshDescription->SetMeshDescriptionRef(GetMeshDescription(LODIndex));
 		return StaticMeshDescription;
 	}
-
+#endif
 	return nullptr;
 }
 
