Fixed NonUnity compile

diff --git a/Engine/Source/Runtime/Engine/Private/ExternalActorsUtils.cpp b/Engine/Source/Runtime/Engine/Private/ExternalActorsUtils.cpp
--- a/Engine/Source/Runtime/Engine/Private/ExternalActorsUtils.cpp
+++ b/Engine/Source/Runtime/Engine/Private/ExternalActorsUtils.cpp
@@ -1,6 +1,8 @@
 // Copyright Epic Games, Inc. All Rights Reserved.
 
 #include "ExternalActorsUtils.h"
+#include "Serialization/ArchiveUObject.h"
+#include "GameFramework/Actor.h"
 
 class FArchiveGatherExternalActorRefs : public FArchiveUObject
 {
diff --git a/Engine/Source/Runtime/Engine/Public/ExternalActorsUtils.h b/Engine/Source/Runtime/Engine/Public/ExternalActorsUtils.h
--- a/Engine/Source/Runtime/Engine/Public/ExternalActorsUtils.h
+++ b/Engine/Source/Runtime/Engine/Public/ExternalActorsUtils.h
@@ -4,6 +4,8 @@
 
 #include "CoreMinimal.h"
 
+class AActor;
+
 namespace ExternalActorsUtils
 {
 	/**
