Disabled shader compiling log spam

diff --git a/Engine/Source/Runtime/Engine/Private/Materials/MaterialShader.cpp b/Engine/Source/Runtime/Engine/Private/Materials/MaterialShader.cpp
--- a/Engine/Source/Runtime/Engine/Private/Materials/MaterialShader.cpp
+++ b/Engine/Source/Runtime/Engine/Private/Materials/MaterialShader.cpp
@@ -1586,7 +1586,7 @@ void FMaterialShaderMap::SubmitCompileJobs(uint32 CompilingShaderMapId,
 		}
 	}
 
-	UE_LOG(LogShaders, Log, TEXT("		%u Shaders among %u VertexFactories"), NumShaders, NumVertexFactories);
+	UE_LOG(LogShaders, Verbose, TEXT("		%u Shaders among %u VertexFactories"), NumShaders, NumVertexFactories);
 
 	GShaderCompilingManager->SubmitJobs(CompileJobs, Material->GetBaseMaterialPathName(), GetDebugDescription());
 }
diff --git a/Engine/Source/Runtime/Engine/Private/ShaderCompiler/ShaderCompiler.cpp b/Engine/Source/Runtime/Engine/Private/ShaderCompiler/ShaderCompiler.cpp
--- a/Engine/Source/Runtime/Engine/Private/ShaderCompiler/ShaderCompiler.cpp
+++ b/Engine/Source/Runtime/Engine/Private/ShaderCompiler/ShaderCompiler.cpp
@@ -299,7 +299,7 @@ void FShaderCompileJobCollection::SubmitJobs(const TArray<FShaderCommonCompileJo
 		{
 			if (NumSubmittedJobs[PriorityIndex] > 0)
 			{
-				UE_LOG(LogShaderCompilers, Display, TEXT("Submitted %d shader compile jobs with '%s' priority"),
+				UE_LOG(LogShaderCompilers, Verbose, TEXT("Submitted %d shader compile jobs with '%s' priority"),
 					NumSubmittedJobs[PriorityIndex],
 					ShaderCompileJobPriorityToString((EShaderCompileJobPriority)PriorityIndex));
 			}
@@ -1520,7 +1520,7 @@ int32 FShaderCompileThreadRunnable::PullTasksFromQueue()
 					NumPendingJobs = Manager->AllJobs.GetNumPendingJobs((EShaderCompileJobPriority)PriorityIndex);
 					if (NumPendingJobs > 0)
 					{
-						UE_LOG(LogShaderCompilers, Display, TEXT("Worker (%d/%d): shaders left to compile %i"), WorkerIndex + 1, WorkerInfos.Num(), NumPendingJobs);
+						UE_LOG(LogShaderCompilers, Verbose, TEXT("Worker (%d/%d): shaders left to compile %i"), WorkerIndex + 1, WorkerInfos.Num(), NumPendingJobs);
 
 						int32 MaxNumJobs = 1;
 						//if (PriorityIndex < (int32)EShaderCompileJobPriority::ForceLocal)
@@ -1577,7 +1577,7 @@ int32 FShaderCompileThreadRunnable::PullTasksFromQueue()
 				COOK_STAT(ShaderCompilerCookStats::AsyncCompileTimeSec += ElapsedTime);
 
 				// Log if requested or if there was an exceptionally slow batch, to see the offender easily
-				if (Manager->bLogJobCompletionTimes || ElapsedTime > 30.0f)
+				if (Manager->bLogJobCompletionTimes || ElapsedTime > 60.0f)
 				{
 					FString JobNames;
 
@@ -1620,7 +1620,7 @@ int32 FShaderCompileThreadRunnable::PullTasksFromQueue()
 	{
 		if (NumJobsStarted[PriorityIndex] > 0)
 		{
-			UE_LOG(LogShaderCompilers, Display, TEXT("Started %d 'Local' shader compile jobs with '%s' priority"),
+			UE_LOG(LogShaderCompilers, Verbose, TEXT("Started %d 'Local' shader compile jobs with '%s' priority"),
 				NumJobsStarted[PriorityIndex],
 				ShaderCompileJobPriorityToString((EShaderCompileJobPriority)PriorityIndex));
 		}
