Make sure we have other active platform modules available before dispatching static mesh build task

diff --git a/Engine/Source/Runtime/Engine/Private/StaticMesh.cpp b/Engine/Source/Runtime/Engine/Private/StaticMesh.cpp
--- a/Engine/Source/Runtime/Engine/Private/StaticMesh.cpp
+++ b/Engine/Source/Runtime/Engine/Private/StaticMesh.cpp
@@ -5306,6 +5306,16 @@ void UStaticMesh::PostLoad()
 		check(RunningPlatform);
 		IMeshBuilderModule::GetForPlatform(RunningPlatform);
 
+		// Additionally cache derived data for any other platforms we care about.
+		const TArray<ITargetPlatform*>& TargetPlatforms = TargetPlatformManager.GetActiveTargetPlatforms();
+		for (ITargetPlatform* Platform : TargetPlatforms)
+		{
+			if (Platform != RunningPlatform)
+			{
+				IMeshBuilderModule::GetForPlatform(Platform);
+			}
+		}
+
 		FQueuedThreadPool* StaticMeshThreadPool = FStaticMeshCompilingManager::Get().GetThreadPool();
 		EQueuedWorkPriority BasePriority = FStaticMeshCompilingManager::Get().GetBasePriority(this);
 
