Improve Datasmith import time by avoiding stall while setting materials of a mesh being built asynchronously

diff --git a/Engine/Source/Runtime/Engine/Private/Components/MeshComponent.cpp b/Engine/Source/Runtime/Engine/Private/Components/MeshComponent.cpp
--- a/Engine/Source/Runtime/Engine/Private/Components/MeshComponent.cpp
+++ b/Engine/Source/Runtime/Engine/Private/Components/MeshComponent.cpp
@@ -78,10 +78,14 @@ void UMeshComponent::SetMaterial(int32 ElementIndex, UMaterialInterface* Materia
 			}
 
 #if WITH_EDITOR
-			FStaticLightingSystemInterface::OnPrimitiveComponentUnregistered.Broadcast(this);
-			if (HasValidSettingsForStaticLighting(false))
+			// Static Lighting is updated when compilation finishes
+			if (!IsCompiling())
 			{
-				FStaticLightingSystemInterface::OnPrimitiveComponentRegistered.Broadcast(this);
+				FStaticLightingSystemInterface::OnPrimitiveComponentUnregistered.Broadcast(this);
+				if (HasValidSettingsForStaticLighting(false))
+				{
+					FStaticLightingSystemInterface::OnPrimitiveComponentRegistered.Broadcast(this);
+				}
 			}
 #endif
 		}
