Improve async static mesh with huge datasmith scene imports
  - Reuse the same object cache when waiting for all mesh to finish compiling
  - Use object cache to resolve static mesh components to update after a build
  - Change the limit execution time heuristic of the per tick mesh update to converge faster when thousands of meshes are ready
    - Framerate after an import of 20 000 meshes was so slow it would have taken hours to add all the meshes to the scene, is now almost complete on first viewport draw.

diff --git a/Engine/Source/Runtime/Engine/Private/StaticMeshCompiler.cpp b/Engine/Source/Runtime/Engine/Private/StaticMeshCompiler.cpp
--- a/Engine/Source/Runtime/Engine/Private/StaticMeshCompiler.cpp
+++ b/Engine/Source/Runtime/Engine/Private/StaticMeshCompiler.cpp
@@ -339,17 +339,8 @@ void FStaticMeshCompilingManager::FinishStaticMeshCompilation(UStaticMesh* Stati
 
 			if (LocalAsyncTask->GetTask().BuildContext.IsValid())
 			{
-				TArray<UStaticMeshComponent*> AffectedComponents;
-				for (UStaticMeshComponent* Component : ObjectCacheScope.GetContext().GetStaticMeshComponents())
-				{
-					if (Component->GetStaticMesh() == StaticMesh)
-					{
-						AffectedComponents.Add(Component);
-					}
-				}
-
 				StaticMesh->FinishBuildInternal(
-					AffectedComponents,
+					ObjectCacheScope.GetContext().GetStaticMeshComponents(StaticMesh),
 					LocalAsyncTask->GetTask().BuildContext->bHasRenderDataChanged,
 					LocalAsyncTask->GetTask().BuildContext->bShouldComputeExtendedBounds
 				);
@@ -409,6 +400,7 @@ void FStaticMeshCompilingManager::FinishCompilation(const TArray<UStaticMesh*>&
 	using namespace StaticMeshCompilingManagerImpl;
 	check(IsInGameThread());
 
+	FObjectCacheContextScope ObjectCacheScope;
 	TSet<UStaticMesh*> PendingStaticMeshes;
 	PendingStaticMeshes.Reserve(InStaticMeshes.Num());
 
@@ -723,15 +715,15 @@ void FStaticMeshCompilingManager::Reschedule()
 	}
 }
 
-void FStaticMeshCompilingManager::ProcessStaticMeshes(int32 MinBatchSize)
+void FStaticMeshCompilingManager::ProcessStaticMeshes(bool bLimitExecutionTime, int32 MinBatchSize)
 {
 	using namespace StaticMeshCompilingManagerImpl;
 	TRACE_CPUPROFILER_EVENT_SCOPE(FStaticMeshCompilingManager::ProcessStaticMeshes);
-	const double MaxSecondsPerFrame = 0.016;
-	const uint32 MaxMeshUpdatesPerFrame = 16;
+	const int32 NumRemainingMeshes = GetNumRemainingMeshes();
+	// Spread out the load over multiple frames but if too many meshes, convergence is more important than frame time
+	const int32 MaxMeshUpdatesPerFrame = bLimitExecutionTime ? FMath::Max(64, NumRemainingMeshes / 10) : INT32_MAX;
 
 	FObjectCacheContextScope ObjectCacheScope;
-	const int32 NumRemainingMeshes = GetNumRemainingMeshes();
 	if (NumRemainingMeshes && NumRemainingMeshes >= MinBatchSize)
 	{
 		TSet<UStaticMesh*> StaticMeshesToProcess;
@@ -754,9 +746,8 @@ void FStaticMeshCompilingManager::ProcessStaticMeshes(int32 MinBatchSize)
 				TSet<TWeakObjectPtr<UStaticMesh>> StaticMeshesToPostpone;
 				for (UStaticMesh* StaticMesh : StaticMeshesToProcess)
 				{
-					const bool bHasTimeLeft = (FPlatformTime::Seconds() - TickStartTime) < MaxSecondsPerFrame;
 					const bool bHasMeshUpdateLeft = ProcessedStaticMeshes.Num() <= MaxMeshUpdatesPerFrame;
-					if ((bHasTimeLeft && bHasMeshUpdateLeft) && StaticMesh->IsAsyncTaskComplete())
+					if (bHasMeshUpdateLeft && StaticMesh->IsAsyncTaskComplete())
 					{
 						FinishStaticMeshCompilation(StaticMesh);
 						ProcessedStaticMeshes.Add(StaticMesh);
@@ -782,7 +773,7 @@ void FStaticMeshCompilingManager::ProcessAsyncTasks(bool bLimitExecutionTime)
 
 	Reschedule();
 
-	ProcessStaticMeshes();
+	ProcessStaticMeshes(bLimitExecutionTime);
 
 	UpdateCompilationNotification();
 }
diff --git a/Engine/Source/Runtime/Engine/Public/StaticMeshCompiler.h b/Engine/Source/Runtime/Engine/Public/StaticMeshCompiler.h
--- a/Engine/Source/Runtime/Engine/Public/StaticMeshCompiler.h
+++ b/Engine/Source/Runtime/Engine/Public/StaticMeshCompiler.h
@@ -76,7 +76,7 @@ private:
 	TSet<TWeakObjectPtr<UStaticMesh>> RegisteredStaticMesh;
 	void FinishCompilationsForGame();
 	void Reschedule();
-	void ProcessStaticMeshes(int32 MinBatchSize = 1);
+	void ProcessStaticMeshes(bool bLimitExecutionTime, int32 MinBatchSize = 1);
 	void UpdateCompilationNotification();
 	void PostStaticMeshesCompilation(const TSet<UStaticMesh*>& InStaticMeshes);
 	void FinishStaticMeshCompilation(UStaticMesh* StaticMesh);
