Add r.vt.ShowDecodeErrors console command.
This enables/disables the hot pink visualization for broken streaming virtual textures.

diff --git a/Engine/Source/Runtime/Engine/Private/VT/VirtualTextureTranscodeCache.cpp b/Engine/Source/Runtime/Engine/Private/VT/VirtualTextureTranscodeCache.cpp
--- a/Engine/Source/Runtime/Engine/Private/VT/VirtualTextureTranscodeCache.cpp
+++ b/Engine/Source/Runtime/Engine/Private/VT/VirtualTextureTranscodeCache.cpp
@@ -1,11 +1,14 @@
 // Copyright Epic Games, Inc. All Rights Reserved.
 
 #include "VirtualTextureTranscodeCache.h"
-#include "VirtualTextureUploadCache.h"
-#include "VirtualTextureChunkManager.h"
-#include "UploadingVirtualTexture.h"
-#include "CrunchCompression.h"
+
 #include "BlockCodingHelpers.h"
+#include "EngineModule.h"
+#include "CrunchCompression.h"
+#include "RendererInterface.h"
+#include "UploadingVirtualTexture.h"
+#include "VirtualTextureChunkManager.h"
+#include "VirtualTextureUploadCache.h"
 
 static int32 TranscodeRetireAge = 60; //1 second @ 60 fps
 static FAutoConsoleVariableRef CVarVTTranscodeRetireAge(
@@ -15,6 +18,27 @@ static FAutoConsoleVariableRef CVarVTTranscodeRetireAge(
 	, ECVF_Default
 );
 
+/** Highlight decoding errors in hot pink. Remove this after we fix the bug where the DDC for virtual textures is getting poisoned with bad data. */
+static bool GShowDecodeErrors = true;
+static FAutoConsoleCommand CVarVTShowDecodeErrors(
+	TEXT("r.VT.ShowDecodeErrors"),
+	TEXT("Highlight virtual textures with decode errors in hot pink."),
+	FConsoleCommandWithArgsDelegate::CreateStatic(
+		[](const TArray< FString >& Args)
+		{
+			if (Args.Num() != 1)
+			{
+				GShowDecodeErrors = !GShowDecodeErrors;
+			}
+			else
+			{
+				GShowDecodeErrors = TCString<TCHAR>::Atoi(*Args[0]) != 0;
+			}
+
+			GetRendererModule().FlushVirtualTextureCache();
+		})
+);
+
 namespace TextureBorderGenerator
 {
 	static int32 Enabled = 0;
@@ -184,7 +208,7 @@ struct FTranscodeTask
 				break;
 			}
 
-			if (!bDecodeResult)
+			if (!bDecodeResult && GShowDecodeErrors)
 			{
 				UniformColorPixels(StagingBufferForLayer, TilePixelSize, TilePixelSize, LayerFormat, ErrorColor);
 			}
