Renamed some strata pipeline function to make it clearer what is happening.

diff --git a/Engine/Source/Runtime/Engine/Private/Materials/HLSLMaterialTranslator.cpp b/Engine/Source/Runtime/Engine/Private/Materials/HLSLMaterialTranslator.cpp
--- a/Engine/Source/Runtime/Engine/Private/Materials/HLSLMaterialTranslator.cpp
+++ b/Engine/Source/Runtime/Engine/Private/Materials/HLSLMaterialTranslator.cpp
@@ -1110,7 +1110,7 @@ bool FHLSLMaterialTranslator::Translate()
 
 			StrataMaterialDescription += FString::Printf(TEXT("StrataCompilationInfo -\r\n"));
 			StrataMaterialDescription += FString::Printf(TEXT(" - TotalBSDFCount    = %i\r\n"), StrataCompilationInfo.TotalBSDFCount);
-			StrataMaterialDescription += FString::Printf(TEXT(" - SharedNormalCount = %i\r\n"), GetStrataSharedNormalCount());
+			StrataMaterialDescription += FString::Printf(TEXT(" - SharedNormalCount = %i\r\n"), StrataCompilationInfoGetSharedNormalCount());
 
 			for (uint32 LayerIt = 0; LayerIt < StrataCompilationInfo.LayerCount; ++LayerIt)
 			{
@@ -7225,13 +7225,13 @@ int32 FHLSLMaterialTranslator::ShadingModel(EMaterialShadingModel InSelectedShad
 	return AddInlinedCodeChunk(MCT_ShadingModel, TEXT("%d"), InSelectedShadingModel);
 }
 
-void FHLSLMaterialTranslator::AddStrataCodeChunk(int32 CodeChunk, FStrataMaterialCompilationInfo& StrataMaterialCompilationInfo)
+void FHLSLMaterialTranslator::StrataCompilationInfoRegisterCodeChunk(int32 CodeChunk, FStrataMaterialCompilationInfo& StrataMaterialCompilationInfo)
 {
 	check(CodeChunk != INDEX_NONE);
 	CodeChunkToStrataCompilationInfoMap.Add(CodeChunk, StrataMaterialCompilationInfo);
 }
 
-bool FHLSLMaterialTranslator::ContainsStrataCodeChunk(int32 CodeChunk)
+bool FHLSLMaterialTranslator::StrataCompilationInfoContainsCodeChunk(int32 CodeChunk)
 {
 	return CodeChunk != INDEX_NONE ? CodeChunkToStrataCompilationInfoMap.Contains(CodeChunk) : false;
 }
@@ -7242,7 +7242,7 @@ const FStrataMaterialCompilationInfo& FHLSLMaterialTranslator::GetStrataCompilat
 	return CodeChunkToStrataCompilationInfoMap[CodeChunk];
 }
 
-uint8 FHLSLMaterialTranslator::GetStrataSharedNormalIndex(int32 NormalCodeChunk)
+uint8 FHLSLMaterialTranslator::StrataCompilationInfoRegisterSharedNormalIndex(int32 NormalCodeChunk)
 {
 	check(NormalCodeChunk != INDEX_NONE);
 	check(NextFreeStrataShaderNormalIndex < 255);	// Out of shared normal slots
@@ -7261,7 +7261,7 @@ uint8 FHLSLMaterialTranslator::GetStrataSharedNormalIndex(int32 NormalCodeChunk)
 	return CodeChunkToStrataSharedNormal[NormalCodeChunkHash].SharedNormalIndex;
 }
 
-uint8 FHLSLMaterialTranslator::GetStrataSharedNormalCount()
+uint8 FHLSLMaterialTranslator::StrataCompilationInfoGetSharedNormalCount()
 {
 	return CodeChunkToStrataSharedNormal.Num();
 }
diff --git a/Engine/Source/Runtime/Engine/Private/Materials/HLSLMaterialTranslator.h b/Engine/Source/Runtime/Engine/Private/Materials/HLSLMaterialTranslator.h
--- a/Engine/Source/Runtime/Engine/Private/Materials/HLSLMaterialTranslator.h
+++ b/Engine/Source/Runtime/Engine/Private/Materials/HLSLMaterialTranslator.h
@@ -799,11 +799,11 @@ protected:
 	virtual int32 StrataArtisticIOR(int32 Reflectivity, int32 EdgeColor, int32 OutputIndex) override;
 	virtual int32 StrataPhysicalIOR(int32 IOR, int32 Extinction, int32 OutputIndex) override;
 
-	virtual void AddStrataCodeChunk(int32 CodeChunk, FStrataMaterialCompilationInfo& StrataMaterialCompilationInfo) override;
-	virtual bool ContainsStrataCodeChunk(int32 CodeChunk) override;
+	virtual void StrataCompilationInfoRegisterCodeChunk(int32 CodeChunk, FStrataMaterialCompilationInfo& StrataMaterialCompilationInfo) override;
+	virtual bool StrataCompilationInfoContainsCodeChunk(int32 CodeChunk) override;
 	virtual const FStrataMaterialCompilationInfo& GetStrataCompilationInfo(int32 CodeChunk) override;
-	virtual uint8 GetStrataSharedNormalIndex(int32 NormalCodeChunk) override;
-	virtual uint8 GetStrataSharedNormalCount() override;
+	virtual uint8 StrataCompilationInfoRegisterSharedNormalIndex(int32 NormalCodeChunk) override;
+	virtual uint8 StrataCompilationInfoGetSharedNormalCount() override;
 
 #if HANDLE_CUSTOM_OUTPUTS_AS_MATERIAL_ATTRIBUTES
 	/** Used to translate code for custom output attributes such as ClearCoatBottomNormal */
diff --git a/Engine/Source/Runtime/Engine/Private/Materials/MaterialExpressions.cpp b/Engine/Source/Runtime/Engine/Private/Materials/MaterialExpressions.cpp
--- a/Engine/Source/Runtime/Engine/Private/Materials/MaterialExpressions.cpp
+++ b/Engine/Source/Runtime/Engine/Private/Materials/MaterialExpressions.cpp
@@ -609,12 +609,12 @@ int32 CompileStrataBlendFunction(FMaterialCompiler* Compiler, const int32 A, con
 
 #if WITH_EDITOR
 	// Also register the horizontal mixing operation with the compiler.
-	if (!Compiler->ContainsStrataCodeChunk(A) || !Compiler->ContainsStrataCodeChunk(B))
+	if (!Compiler->StrataCompilationInfoContainsCodeChunk(A) || !Compiler->StrataCompilationInfoContainsCodeChunk(B))
 	{
 		return Compiler->Errorf(TEXT("Could not find A or B code chunk in CompileStrataBlendFunction"));
 	}
-	FStrataMaterialCompilationInfo StrataInfo = StrataHorizontalMixing(Compiler, Compiler->GetStrataCompilationInfo(B), Compiler->GetStrataCompilationInfo(A));
-	Compiler->AddStrataCodeChunk(OutputCodeChunk, StrataInfo);
+	FStrataMaterialCompilationInfo StrataInfo = StrataCompilationInfoHorizontalMixing(Compiler, Compiler->GetStrataCompilationInfo(B), Compiler->GetStrataCompilationInfo(A));
+	Compiler->StrataCompilationInfoRegisterCodeChunk(OutputCodeChunk, StrataInfo);
 #endif
 
 	return OutputCodeChunk;
@@ -19031,14 +19031,14 @@ UMaterialExpressionStrataDiffuseBSDF::UMaterialExpressionStrataDiffuseBSDF(const
 int32 UMaterialExpressionStrataDiffuseBSDF::Compile(class FMaterialCompiler* Compiler, int32 OutputIndex)
 {
 	int32 NormalCodeChunk = Normal.GetTracedInput().Expression ? Normal.Compile(Compiler) : Compiler->PixelNormalWS();
-	uint8 SharedNormalIndex = StrataCreateSharedNormal(Compiler, NormalCodeChunk);
+	uint8 SharedNormalIndex = StrataCompilationInfoCreateSharedNormal(Compiler, NormalCodeChunk);
 
 	int32 OutputCodeChunk = Compiler->StrataDiffuseOrenNayarBSDF(
 		Albedo.GetTracedInput().Expression		? Albedo.Compile(Compiler)		: Compiler->Constant3(0.18f, 0.18f, 0.18f),
 		Roughness.GetTracedInput().Expression	? Roughness.Compile(Compiler)	: Compiler->Constant(0.0f),
 		NormalCodeChunk,
 		SharedNormalIndex);
-	StrataCreateSingleBSDFMaterial(Compiler, OutputCodeChunk, SharedNormalIndex, STRATA_BSDF_TYPE_DIFFUSE_ON);
+	StrataCompilationInfoCreateSingleBSDFMaterial(Compiler, OutputCodeChunk, SharedNormalIndex, STRATA_BSDF_TYPE_DIFFUSE_ON);
 
 	return OutputCodeChunk;
 }
@@ -19095,14 +19095,14 @@ UMaterialExpressionStrataDiffuseChanBSDF::UMaterialExpressionStrataDiffuseChanBS
 int32 UMaterialExpressionStrataDiffuseChanBSDF::Compile(class FMaterialCompiler* Compiler, int32 OutputIndex)
 {
 	int32 NormalCodeChunk = Normal.GetTracedInput().Expression ? Normal.Compile(Compiler) : Compiler->PixelNormalWS();
-	uint8 SharedNormalIndex = StrataCreateSharedNormal(Compiler, NormalCodeChunk);
+	uint8 SharedNormalIndex = StrataCompilationInfoCreateSharedNormal(Compiler, NormalCodeChunk);
 
 	int32 OutputCodeChunk = Compiler->StrataDiffuseChanBSDF(
 		Albedo.GetTracedInput().Expression		? Albedo.Compile(Compiler)		: Compiler->Constant3(0.18f, 0.18f, 0.18f),
 		Roughness.GetTracedInput().Expression	? Roughness.Compile(Compiler)	: Compiler->Constant(0.0f),
 		NormalCodeChunk,
 		SharedNormalIndex);
-	StrataCreateSingleBSDFMaterial(Compiler, OutputCodeChunk, SharedNormalIndex, STRATA_BSDF_TYPE_DIFFUSE_CHAN);
+	StrataCompilationInfoCreateSingleBSDFMaterial(Compiler, OutputCodeChunk, SharedNormalIndex, STRATA_BSDF_TYPE_DIFFUSE_CHAN);
 
 	return OutputCodeChunk;
 }
@@ -19159,7 +19159,7 @@ UMaterialExpressionStrataDielectricBSDF::UMaterialExpressionStrataDielectricBSDF
 int32 UMaterialExpressionStrataDielectricBSDF::Compile(class FMaterialCompiler* Compiler, int32 OutputIndex)
 {
 	int32 NormalCodeChunk = Normal.GetTracedInput().Expression ? Normal.Compile(Compiler) : Compiler->PixelNormalWS();
-	uint8 SharedNormalIndex = StrataCreateSharedNormal(Compiler, NormalCodeChunk);
+	uint8 SharedNormalIndex = StrataCompilationInfoCreateSharedNormal(Compiler, NormalCodeChunk);
 
 	int32 OutputCodeChunk = Compiler->StrataDielectricBSDF(
 		Roughness.GetTracedInput().Expression	? Roughness.Compile(Compiler)	: Compiler->Constant2(0.0f, 0.0f),
@@ -19167,7 +19167,7 @@ int32 UMaterialExpressionStrataDielectricBSDF::Compile(class FMaterialCompiler*
 		Tint.GetTracedInput().Expression		? Tint.Compile(Compiler)		: Compiler->Constant3(1.0f, 1.0f, 1.0f),
 		NormalCodeChunk,
 		SharedNormalIndex);
-	StrataCreateSingleBSDFMaterial(Compiler, OutputCodeChunk, SharedNormalIndex, STRATA_BSDF_TYPE_DIELECTRIC);
+	StrataCompilationInfoCreateSingleBSDFMaterial(Compiler, OutputCodeChunk, SharedNormalIndex, STRATA_BSDF_TYPE_DIELECTRIC);
 
 	return OutputCodeChunk;
 }
@@ -19227,7 +19227,7 @@ UMaterialExpressionStrataConductorBSDF::UMaterialExpressionStrataConductorBSDF(c
 int32 UMaterialExpressionStrataConductorBSDF::Compile(class FMaterialCompiler* Compiler, int32 OutputIndex)
 {
 	int32 NormalCodeChunk = Normal.GetTracedInput().Expression ? Normal.Compile(Compiler) : Compiler->PixelNormalWS();
-	uint8 SharedNormalIndex = StrataCreateSharedNormal(Compiler, NormalCodeChunk);
+	uint8 SharedNormalIndex = StrataCompilationInfoCreateSharedNormal(Compiler, NormalCodeChunk);
 
 	int32 OutputCodeChunk = Compiler->StrataConductorBSDF(
 		Reflectivity.GetTracedInput().Expression	? Reflectivity.Compile(Compiler): Compiler->Constant3(0.947f, 0.776f, 0.371f),	// Default to Gold
@@ -19235,7 +19235,7 @@ int32 UMaterialExpressionStrataConductorBSDF::Compile(class FMaterialCompiler* C
 		Roughness.GetTracedInput().Expression		? Roughness.Compile(Compiler)	: Compiler->Constant2(0.0f, 0.0f),
 		NormalCodeChunk,
 		SharedNormalIndex);
-	StrataCreateSingleBSDFMaterial(Compiler, OutputCodeChunk, SharedNormalIndex, STRATA_BSDF_TYPE_CONDUCTOR);
+	StrataCompilationInfoCreateSingleBSDFMaterial(Compiler, OutputCodeChunk, SharedNormalIndex, STRATA_BSDF_TYPE_CONDUCTOR);
 
 	return OutputCodeChunk;
 }
@@ -19295,7 +19295,7 @@ UMaterialExpressionStrataVolumeBSDF::UMaterialExpressionStrataVolumeBSDF(const F
 int32 UMaterialExpressionStrataVolumeBSDF::Compile(class FMaterialCompiler* Compiler, int32 OutputIndex)
 {
 	int32 NormalCodeChunk = Compiler->VertexNormal();
-	uint8 SharedNormalIndex = StrataCreateSharedNormal(Compiler, NormalCodeChunk);
+	uint8 SharedNormalIndex = StrataCompilationInfoCreateSharedNormal(Compiler, NormalCodeChunk);
 
 	int32 OutputCodeChunk = Compiler->StrataVolumeBSDF(
 		Albedo.GetTracedInput().Expression		? Albedo.Compile(Compiler)			: Compiler->Constant3(0.0f, 0.0f, 0.0f),
@@ -19304,7 +19304,7 @@ int32 UMaterialExpressionStrataVolumeBSDF::Compile(class FMaterialCompiler* Comp
 		Thickness.GetTracedInput().Expression	? Thickness.Compile(Compiler)		: Compiler->Constant(0.001f), // default = 1mm
 		NormalCodeChunk,
 		SharedNormalIndex);
-	StrataCreateSingleBSDFMaterial(Compiler, OutputCodeChunk, SharedNormalIndex, STRATA_BSDF_TYPE_VOLUME);
+	StrataCompilationInfoCreateSingleBSDFMaterial(Compiler, OutputCodeChunk, SharedNormalIndex, STRATA_BSDF_TYPE_VOLUME);
 
 	return OutputCodeChunk;
 }
@@ -19378,12 +19378,12 @@ int32 UMaterialExpressionStrataHorizontalMixing::Compile(class FMaterialCompiler
 		BackgroundCodeChunk,
 		Mix.GetTracedInput().Expression ? Mix.Compile(Compiler) : Compiler->Constant(0.0f));
 
-	if (!Compiler->ContainsStrataCodeChunk(ForegroundCodeChunk) || !Compiler->ContainsStrataCodeChunk(BackgroundCodeChunk))
+	if (!Compiler->StrataCompilationInfoContainsCodeChunk(ForegroundCodeChunk) || !Compiler->StrataCompilationInfoContainsCodeChunk(BackgroundCodeChunk))
 	{
 		return Compiler->Errorf(TEXT("Could not find ForegroundCodeChunk or BackgroundCodeChunk to mix"));
 	}
-	FStrataMaterialCompilationInfo StrataInfo = StrataHorizontalMixing(Compiler, Compiler->GetStrataCompilationInfo(ForegroundCodeChunk), Compiler->GetStrataCompilationInfo(BackgroundCodeChunk));
-	Compiler->AddStrataCodeChunk(OutputCodeChunk, StrataInfo);
+	FStrataMaterialCompilationInfo StrataInfo = StrataCompilationInfoHorizontalMixing(Compiler, Compiler->GetStrataCompilationInfo(ForegroundCodeChunk), Compiler->GetStrataCompilationInfo(BackgroundCodeChunk));
+	Compiler->StrataCompilationInfoRegisterCodeChunk(OutputCodeChunk, StrataInfo);
 
 	return OutputCodeChunk;
 }
@@ -19437,12 +19437,12 @@ int32 UMaterialExpressionStrataVerticalLayering::Compile(class FMaterialCompiler
 
 	int32 OutputCodeChunk = Compiler->StrataVerticalLayering(TopCodeChunk, BaseCodeChunk);
 
-	if (!Compiler->ContainsStrataCodeChunk(TopCodeChunk) || !Compiler->ContainsStrataCodeChunk(BaseCodeChunk))
+	if (!Compiler->StrataCompilationInfoContainsCodeChunk(TopCodeChunk) || !Compiler->StrataCompilationInfoContainsCodeChunk(BaseCodeChunk))
 	{
 		return Compiler->Errorf(TEXT("Could not find TopCodeChunk or BaseCodeChunk to layer"));
 	}
-	FStrataMaterialCompilationInfo StrataInfo = StrataVerticalLayering(Compiler, Compiler->GetStrataCompilationInfo(TopCodeChunk), Compiler->GetStrataCompilationInfo(BaseCodeChunk));
-	Compiler->AddStrataCodeChunk(OutputCodeChunk, StrataInfo);
+	FStrataMaterialCompilationInfo StrataInfo = StrataCompilationInfoVerticalLayering(Compiler, Compiler->GetStrataCompilationInfo(TopCodeChunk), Compiler->GetStrataCompilationInfo(BaseCodeChunk));
+	Compiler->StrataCompilationInfoRegisterCodeChunk(OutputCodeChunk, StrataInfo);
 
 	return OutputCodeChunk;
 }
@@ -19496,12 +19496,12 @@ int32 UMaterialExpressionStrataAdd::Compile(class FMaterialCompiler* Compiler, i
 
 	int32 OutputCodeChunk = Compiler->StrataAdd(ACodeChunk, BCodeChunk);
 
-	if (!Compiler->ContainsStrataCodeChunk(ACodeChunk) || !Compiler->ContainsStrataCodeChunk(BCodeChunk))
+	if (!Compiler->StrataCompilationInfoContainsCodeChunk(ACodeChunk) || !Compiler->StrataCompilationInfoContainsCodeChunk(BCodeChunk))
 	{
 		return Compiler->Errorf(TEXT("Could not find ACodeChunk or BCodeChunk to add"));
 	}
-	FStrataMaterialCompilationInfo StrataInfo = StrataAdd(Compiler, Compiler->GetStrataCompilationInfo(ACodeChunk), Compiler->GetStrataCompilationInfo(BCodeChunk));
-	Compiler->AddStrataCodeChunk(OutputCodeChunk, StrataInfo);
+	FStrataMaterialCompilationInfo StrataInfo = StrataCompilationInfoAdd(Compiler, Compiler->GetStrataCompilationInfo(ACodeChunk), Compiler->GetStrataCompilationInfo(BCodeChunk));
+	Compiler->StrataCompilationInfoRegisterCodeChunk(OutputCodeChunk, StrataInfo);
 
 	return OutputCodeChunk;
 }
@@ -19555,12 +19555,12 @@ int32 UMaterialExpressionStrataMultiply::Compile(class FMaterialCompiler* Compil
 
 	int32 OutputCodeChunk = Compiler->StrataMultiply(ACodeChunk, WeightCodeChunk);
 
-	if (!Compiler->ContainsStrataCodeChunk(ACodeChunk))
+	if (!Compiler->StrataCompilationInfoContainsCodeChunk(ACodeChunk))
 	{
 		return Compiler->Errorf(TEXT("Could not find ACodeChunk to multiply"));
 	}
-	FStrataMaterialCompilationInfo StrataInfo = StrataMultiply(Compiler, Compiler->GetStrataCompilationInfo(ACodeChunk));
-	Compiler->AddStrataCodeChunk(OutputCodeChunk, StrataInfo);
+	FStrataMaterialCompilationInfo StrataInfo = StrataCompilationInfoMultiply(Compiler, Compiler->GetStrataCompilationInfo(ACodeChunk));
+	Compiler->StrataCompilationInfoRegisterCodeChunk(OutputCodeChunk, StrataInfo);
 
 	return OutputCodeChunk;
 }
diff --git a/Engine/Source/Runtime/Engine/Private/Materials/StrataMaterial.cpp b/Engine/Source/Runtime/Engine/Private/Materials/StrataMaterial.cpp
--- a/Engine/Source/Runtime/Engine/Private/Materials/StrataMaterial.cpp
+++ b/Engine/Source/Runtime/Engine/Private/Materials/StrataMaterial.cpp
@@ -39,13 +39,13 @@ static void UpdateTotalBSDFCount(FStrataMaterialCompilationInfo& StrataInfo)
 }
 
 
-uint8 StrataCreateSharedNormal(FMaterialCompiler* Compiler, int32 NormalCodeChunk)
+uint8 StrataCompilationInfoCreateSharedNormal(FMaterialCompiler* Compiler, int32 NormalCodeChunk)
 {
-	return Compiler->GetStrataSharedNormalIndex(NormalCodeChunk);
+	return Compiler->StrataCompilationInfoRegisterSharedNormalIndex(NormalCodeChunk);
 }
 
 
-void StrataCreateSingleBSDFMaterial(FMaterialCompiler* Compiler, int32 CodeChunk, uint8 SharedNormalIndex, uint8 BSDFType)
+void StrataCompilationInfoCreateSingleBSDFMaterial(FMaterialCompiler* Compiler, int32 CodeChunk, uint8 SharedNormalIndex, uint8 BSDFType)
 {
 	FStrataMaterialCompilationInfo StrataInfo;
 	StrataInfo.LayerCount = 1;
@@ -53,11 +53,11 @@ void StrataCreateSingleBSDFMaterial(FMaterialCompiler* Compiler, int32 CodeChunk
 	StrataInfo.Layers[0].BSDFs[0].Type = BSDFType;
 	StrataInfo.Layers[0].BSDFs[0].SharedNormalIndex = SharedNormalIndex;
 	UpdateTotalBSDFCount(StrataInfo);
-	Compiler->AddStrataCodeChunk(CodeChunk, StrataInfo);
+	Compiler->StrataCompilationInfoRegisterCodeChunk(CodeChunk, StrataInfo);
 }
 
 
-FStrataMaterialCompilationInfo StrataAdd(FMaterialCompiler* Compiler, const FStrataMaterialCompilationInfo& A, const FStrataMaterialCompilationInfo& B)
+FStrataMaterialCompilationInfo StrataCompilationInfoAdd(FMaterialCompiler* Compiler, const FStrataMaterialCompilationInfo& A, const FStrataMaterialCompilationInfo& B)
 {
 	if ((A.TotalBSDFCount + B.TotalBSDFCount) > STRATA_MAX_TOTAL_BSDF) // See StrataDefinitions.h
 	{
@@ -93,14 +93,14 @@ FStrataMaterialCompilationInfo StrataAdd(FMaterialCompiler* Compiler, const FStr
 }
 
 
-FStrataMaterialCompilationInfo StrataMultiply(FMaterialCompiler* Compiler, const FStrataMaterialCompilationInfo& A)
+FStrataMaterialCompilationInfo StrataCompilationInfoMultiply(FMaterialCompiler* Compiler, const FStrataMaterialCompilationInfo& A)
 {
 	FStrataMaterialCompilationInfo StrataInfo = A;
 	return StrataInfo;
 }
 
 
-FStrataMaterialCompilationInfo StrataHorizontalMixing(FMaterialCompiler* Compiler, const FStrataMaterialCompilationInfo& A, const FStrataMaterialCompilationInfo& B)
+FStrataMaterialCompilationInfo StrataCompilationInfoHorizontalMixing(FMaterialCompiler* Compiler, const FStrataMaterialCompilationInfo& A, const FStrataMaterialCompilationInfo& B)
 {
 	if ((A.TotalBSDFCount + B.TotalBSDFCount) > STRATA_MAX_TOTAL_BSDF) // See StrataDefinitions.h
 	{
@@ -108,11 +108,11 @@ FStrataMaterialCompilationInfo StrataHorizontalMixing(FMaterialCompiler* Compile
 		return A;
 	}
 
-	return StrataAdd(Compiler, A, B); // Mixing is a similar operation to Add when it comes to bsdf count
+	return StrataCompilationInfoAdd(Compiler, A, B); // Mixing is a similar operation to Add when it comes to bsdf count
 }
 
 
-FStrataMaterialCompilationInfo StrataVerticalLayering(FMaterialCompiler* Compiler, const FStrataMaterialCompilationInfo& Top, const FStrataMaterialCompilationInfo& Base)
+FStrataMaterialCompilationInfo StrataCompilationInfoVerticalLayering(FMaterialCompiler* Compiler, const FStrataMaterialCompilationInfo& Top, const FStrataMaterialCompilationInfo& Base)
 {
 	if ((Top.TotalBSDFCount + Base.TotalBSDFCount) > STRATA_MAX_TOTAL_BSDF) // See StrataDefinitions.h
 	{
diff --git a/Engine/Source/Runtime/Engine/Private/Materials/StrataMaterial.h b/Engine/Source/Runtime/Engine/Private/Materials/StrataMaterial.h
--- a/Engine/Source/Runtime/Engine/Private/Materials/StrataMaterial.h
+++ b/Engine/Source/Runtime/Engine/Private/Materials/StrataMaterial.h
@@ -41,12 +41,12 @@ struct FStrataMaterialCompilationInfo
 
 FString GetStrataBSDFName(uint8 BSDFType);
 
-uint8 StrataCreateSharedNormal(FMaterialCompiler* Compiler, int32 NormalCodeChunk);
-void StrataCreateSingleBSDFMaterial(FMaterialCompiler* Compiler, int32 CodeChunk, uint8 SharedNormalIndex, uint8 BSDFType);
+uint8 StrataCompilationInfoCreateSharedNormal(FMaterialCompiler* Compiler, int32 NormalCodeChunk);
+void  StrataCompilationInfoCreateSingleBSDFMaterial(FMaterialCompiler* Compiler, int32 CodeChunk, uint8 SharedNormalIndex, uint8 BSDFType);
 
-FStrataMaterialCompilationInfo StrataAdd(FMaterialCompiler* Compiler, const FStrataMaterialCompilationInfo& A, const FStrataMaterialCompilationInfo& B);
-FStrataMaterialCompilationInfo StrataMultiply(FMaterialCompiler* Compiler, const FStrataMaterialCompilationInfo& A);
-FStrataMaterialCompilationInfo StrataHorizontalMixing(FMaterialCompiler* Compiler, const FStrataMaterialCompilationInfo& A, const FStrataMaterialCompilationInfo& B);
-FStrataMaterialCompilationInfo StrataVerticalLayering(FMaterialCompiler* Compiler, const FStrataMaterialCompilationInfo& Top, const FStrataMaterialCompilationInfo& Base);
+FStrataMaterialCompilationInfo StrataCompilationInfoAdd(FMaterialCompiler* Compiler, const FStrataMaterialCompilationInfo& A, const FStrataMaterialCompilationInfo& B);
+FStrataMaterialCompilationInfo StrataCompilationInfoMultiply(FMaterialCompiler* Compiler, const FStrataMaterialCompilationInfo& A);
+FStrataMaterialCompilationInfo StrataCompilationInfoHorizontalMixing(FMaterialCompiler* Compiler, const FStrataMaterialCompilationInfo& A, const FStrataMaterialCompilationInfo& B);
+FStrataMaterialCompilationInfo StrataCompilationInfoVerticalLayering(FMaterialCompiler* Compiler, const FStrataMaterialCompilationInfo& Top, const FStrataMaterialCompilationInfo& Base);
 
 
diff --git a/Engine/Source/Runtime/Engine/Public/MaterialCompiler.h b/Engine/Source/Runtime/Engine/Public/MaterialCompiler.h
--- a/Engine/Source/Runtime/Engine/Public/MaterialCompiler.h
+++ b/Engine/Source/Runtime/Engine/Public/MaterialCompiler.h
@@ -403,11 +403,11 @@ public:
 	virtual int32 StrataArtisticIOR(int32 Reflectivity, int32 EdgeColor, int32 OutputIndex) = 0;
 	virtual int32 StrataPhysicalIOR(int32 IOR, int32 Extinction, int32 OutputIndex) = 0;
 
-	virtual void AddStrataCodeChunk(int32 CodeChunk, FStrataMaterialCompilationInfo& StrataMaterialCompilationInfo) = 0;
-	virtual bool ContainsStrataCodeChunk(int32 CodeChunk) = 0;
+	virtual void StrataCompilationInfoRegisterCodeChunk(int32 CodeChunk, FStrataMaterialCompilationInfo& StrataMaterialCompilationInfo) = 0;
+	virtual bool StrataCompilationInfoContainsCodeChunk(int32 CodeChunk) = 0;
 	virtual const FStrataMaterialCompilationInfo& GetStrataCompilationInfo(int32 CodeChunk) = 0;
-	virtual uint8 GetStrataSharedNormalIndex(int32 NormalCodeChunk) = 0;
-	virtual uint8 GetStrataSharedNormalCount() = 0;
+	virtual uint8 StrataCompilationInfoRegisterSharedNormalIndex(int32 NormalCodeChunk) = 0;
+	virtual uint8 StrataCompilationInfoGetSharedNormalCount() = 0;
 
 	// Water
 	virtual int32 SceneDepthWithoutWater(int32 Offset, int32 ViewportUV, bool bUseOffset, float FallbackDepth) = 0;
@@ -858,14 +858,14 @@ public:
 		return Compiler->StrataPhysicalIOR(IOR, Extinction, OutputIndex);
 	}
 
-	virtual void AddStrataCodeChunk(int32 CodeChunk, FStrataMaterialCompilationInfo& StrataMaterialCompilationInfo) override
+	virtual void StrataCompilationInfoRegisterCodeChunk(int32 CodeChunk, FStrataMaterialCompilationInfo& StrataMaterialCompilationInfo) override
 	{
-		Compiler->AddStrataCodeChunk(CodeChunk, StrataMaterialCompilationInfo);
+		Compiler->StrataCompilationInfoRegisterCodeChunk(CodeChunk, StrataMaterialCompilationInfo);
 	}
 
-	virtual bool ContainsStrataCodeChunk(int32 CodeChunk) override
+	virtual bool StrataCompilationInfoContainsCodeChunk(int32 CodeChunk) override
 	{
-		return Compiler->ContainsStrataCodeChunk(CodeChunk);
+		return Compiler->StrataCompilationInfoContainsCodeChunk(CodeChunk);
 	}
 	
 	virtual const FStrataMaterialCompilationInfo& GetStrataCompilationInfo(int32 CodeChunk) override
@@ -873,14 +873,14 @@ public:
 		return Compiler->GetStrataCompilationInfo(CodeChunk);
 	}
 
-	virtual uint8 GetStrataSharedNormalIndex(int32 NormalCodeChunk) override
+	virtual uint8 StrataCompilationInfoRegisterSharedNormalIndex(int32 NormalCodeChunk) override
 	{
-		return Compiler->GetStrataSharedNormalIndex(NormalCodeChunk);
+		return Compiler->StrataCompilationInfoRegisterSharedNormalIndex(NormalCodeChunk);
 	}
 
-	virtual uint8 GetStrataSharedNormalCount() override
+	virtual uint8 StrataCompilationInfoGetSharedNormalCount() override
 	{
-		return Compiler->GetStrataSharedNormalCount();
+		return Compiler->StrataCompilationInfoGetSharedNormalCount();
 	}
 
 protected:
