Packed Level Instance MVP (based on SPMeshPacker)

- supports packing source level MeshComponents (SMC,ISM,HISM) into a single actor APackedLevelInstance.
- supports recursive packing of source LevelInstance's
- ability to pack result into an actor blueprint
- build menu allows rebuilding all loaded packed level instances (no dependency tracking)

diff --git a/Engine/Source/Developer/AssetTools/Private/AssetTools.h b/Engine/Source/Developer/AssetTools/Private/AssetTools.h
--- a/Engine/Source/Developer/AssetTools/Private/AssetTools.h
+++ b/Engine/Source/Developer/AssetTools/Private/AssetTools.h
@@ -120,15 +120,13 @@ public:
 	virtual TSharedRef<FBlacklistPaths>& GetWritableFolderBlacklist() override;
 	virtual bool AllPassWritableFolderFilter(const TArray<FString>& InPaths) const override;
 	virtual void NotifyBlockedByWritableFolderFilter() const;
-
+	
+	virtual void SyncBrowserToAssets(const TArray<UObject*>& AssetsToSync) override;
+	virtual void SyncBrowserToAssets(const TArray<FAssetData>& AssetsToSync) override;
 public:
 	/** Gets the asset tools singleton as a FAssetTools for asset tools module use */
 	static UAssetToolsImpl& Get();
 
-	/** Syncs the primary content browser to the specified assets, whether or not it is locked. Most syncs that come from AssetTools -feel- like they came from the content browser, so this is okay. */
-	void SyncBrowserToAssets(const TArray<UObject*>& AssetsToSync);
-	void SyncBrowserToAssets(const TArray<FAssetData>& AssetsToSync);
-
 	/** The manager to handle renaming assets */
 	TSharedPtr<FAssetRenameManager> AssetRenameManager;
 
diff --git a/Engine/Source/Developer/AssetTools/Public/IAssetTools.h b/Engine/Source/Developer/AssetTools/Public/IAssetTools.h
--- a/Engine/Source/Developer/AssetTools/Public/IAssetTools.h
+++ b/Engine/Source/Developer/AssetTools/Public/IAssetTools.h
@@ -448,6 +448,10 @@ public:
 
 	/** Show notification that writable folder filter blocked an action */
 	virtual void NotifyBlockedByWritableFolderFilter() const = 0;
+
+	/** Syncs the primary content browser to the specified assets, whether or not it is locked. Most syncs that come from AssetTools -feel- like they came from the content browser, so this is okay. */
+	virtual void SyncBrowserToAssets(const TArray<UObject*>& AssetsToSync) = 0;
+	virtual void SyncBrowserToAssets(const TArray<FAssetData>& AssetsToSync) = 0;
 };
 
 UCLASS(transient)
diff --git a/Engine/Source/Editor/LevelInstanceEditor/LevelInstanceEditor.Build.cs b/Engine/Source/Editor/LevelInstanceEditor/LevelInstanceEditor.Build.cs
--- a/Engine/Source/Editor/LevelInstanceEditor/LevelInstanceEditor.Build.cs
+++ b/Engine/Source/Editor/LevelInstanceEditor/LevelInstanceEditor.Build.cs
@@ -31,7 +31,8 @@ public class LevelInstanceEditor : ModuleRules
 				"MainFrame",
 				"ContentBrowser",
 				"AssetTools",
-				"ClassViewer"
+				"ClassViewer",
+				"MessageLog"
 			}
 		);
 
diff --git a/Engine/Source/Editor/LevelInstanceEditor/Private/LevelInstanceEditorModule.cpp b/Engine/Source/Editor/LevelInstanceEditor/Private/LevelInstanceEditorModule.cpp
--- a/Engine/Source/Editor/LevelInstanceEditor/Private/LevelInstanceEditorModule.cpp
+++ b/Engine/Source/Editor/LevelInstanceEditor/Private/LevelInstanceEditorModule.cpp
@@ -3,6 +3,8 @@
 #include "LevelInstanceActorDetails.h"
 #include "LevelInstance/LevelInstanceSubsystem.h"
 #include "LevelInstance/LevelInstanceActor.h"
+#include "LevelInstance/Packed/PackedLevelInstanceBuilder.h"
+#include "LevelInstance/Packed/PackedLevelInstanceActor.h"
 #include "LevelInstanceEditorSettings.h"
 #include "ToolMenus.h"
 #include "Editor.h"
@@ -28,6 +30,7 @@
 #include "Misc/ScopeExit.h"
 #include "Widgets/Input/SNumericEntryBox.h"
 #include "Widgets/Input/SButton.h"
+#include "MessageLogModule.h"
 
 IMPLEMENT_MODULE( FLevelInstanceEditorModule, LevelInstanceEditor );
 
@@ -295,7 +298,7 @@ namespace LevelInstanceMenuUtils
 		}
 	}
 
-	void CreateLevelInstanceFromSelection(ULevelInstanceSubsystem* LevelInstanceSubsystem)
+	void CreateLevelInstanceFromSelection(ULevelInstanceSubsystem* LevelInstanceSubsystem, ELevelInstanceCreationType CreationType)
 	{
 		TArray<AActor*> ActorsToMove;
 		ActorsToMove.Reserve(GEditor->GetSelectedActorCount());
@@ -315,7 +318,7 @@ namespace LevelInstanceMenuUtils
 			UPackage* TemplatePackage = !TemplateMapPackage.IsEmpty() ? LoadPackage(nullptr, *TemplateMapPackage, LOAD_None) : nullptr;
 			UWorld* TemplateWorld = TemplatePackage ? UWorld::FindWorldInPackage(TemplatePackage) : nullptr;
 
-			if (!LevelInstanceSubsystem->CreateLevelInstanceFrom(ActorsToMove, TemplateWorld))
+			if (!LevelInstanceSubsystem->CreateLevelInstanceFrom(ActorsToMove, CreationType, TemplateWorld))
 			{
 				FText Title = LOCTEXT("CreateFromSelectionFailTitle", "Create from selection failed");
 				FMessageDialog::Open(EAppMsgType::Ok, LOCTEXT("CreateFromSelectionFailMsg", "Failed to create LevelInstance from selection. Check log for details."), &Title);
@@ -323,27 +326,39 @@ namespace LevelInstanceMenuUtils
 		}
 	}
 
+	void CreateLevelInstanceFromSelectionSubMenu(UToolMenu* ToolMenu, ULevelInstanceSubsystem* LevelInstanceSubsystem)
+	{
+		FToolMenuSection* Section = &ToolMenu->AddSection(FName("LevelInstanceCreateFromSelectionSub"));
+		UEnum* CreationTypeEnum = StaticEnum<ELevelInstanceCreationType>();
+		check(CreationTypeEnum);
+		for (int32 i = 0; i < CreationTypeEnum->NumEnums()-1; ++i)
+		{
+			if (!CreationTypeEnum->HasMetaData(TEXT("Hidden"), i))
+			{
+				ELevelInstanceCreationType CreationType = (ELevelInstanceCreationType)CreationTypeEnum->GetValueByIndex(i);
+				FToolUIAction CreateFromSelectionAction;
+				CreateFromSelectionAction.ExecuteAction.BindLambda([LevelInstanceSubsystem, CreationType](const FToolMenuContext&)
+				{
+					CreateLevelInstanceFromSelection(LevelInstanceSubsystem, CreationType);
+				});
+
+				Section->AddMenuEntry(NAME_None, CreationTypeEnum->GetDisplayNameTextByIndex(i), TAttribute<FText>(), FSlateIcon(), CreateFromSelectionAction);
+			}
+		}
+	}
+
 	void CreateCreateMenu(UToolMenu* Menu)
 	{
 		if (ULevelInstanceSubsystem* LevelInstanceSubsystem = GEditor->GetEditorWorldContext().World()->GetSubsystem<ULevelInstanceSubsystem>())
 		{
 			FToolMenuSection& Section = CreateLevelInstanceSection(Menu);
-			FToolUIAction UIAction;
-			UIAction.ExecuteAction.BindLambda([LevelInstanceSubsystem](const FToolMenuContext& MenuContext)
-			{
-				CreateLevelInstanceFromSelection(LevelInstanceSubsystem);
-			});
-			UIAction.CanExecuteAction.BindLambda([](const FToolMenuContext& MenuContext)
-			{
-				return GEditor->GetSelectedActorCount() > 0;
-			});
-
-			Section.AddMenuEntry(
+			
+			Section.AddSubMenu(
 				"CreateLevelInstanceFromSelection",
 				LOCTEXT("CreateLevelInstanceFromSelection", "Create from selection"),
 				TAttribute<FText>(),
-				TAttribute<FSlateIcon>(),
-				UIAction);
+				FNewToolMenuDelegate::CreateStatic(&CreateLevelInstanceFromSelectionSubMenu, LevelInstanceSubsystem)
+			);
 		}
 	}
 
@@ -477,6 +492,76 @@ namespace LevelInstanceMenuUtils
 		}
 	}
 
+	void CreatePackedBlueprintMenu(UToolMenu* Menu, AActor* ContextActor)
+	{
+		if (ULevelInstanceSubsystem* LevelInstanceSubsystem = ContextActor->GetWorld()->GetSubsystem<ULevelInstanceSubsystem>())
+		{
+			ALevelInstance* ContextLevelInstance = nullptr;
+
+			// Find the top level LevelInstance
+			LevelInstanceSubsystem->ForEachLevelInstanceAncestorsAndSelf(ContextActor, [LevelInstanceSubsystem, ContextActor, &ContextLevelInstance](ALevelInstance* Ancestor)
+			{
+				if (Ancestor->GetLevel() == ContextActor->GetWorld()->GetCurrentLevel())
+				{
+					ContextLevelInstance = Ancestor;
+					return false;
+				}
+				return true;
+			});
+						
+			if (ContextLevelInstance && !ContextLevelInstance->IsEditing())
+			{
+				FToolMenuSection& Section = CreateLevelInstanceSection(Menu);
+				TSoftObjectPtr<UBlueprint> BlueprintAsset;
+				if (APackedLevelInstance* PackedLevelInstance = Cast<APackedLevelInstance>(ContextLevelInstance))
+				{
+					BlueprintAsset = PackedLevelInstance->BlueprintAsset;
+				}
+
+				if (BlueprintAsset.IsNull())
+				{
+					FToolUIAction UIAction;
+					UIAction.ExecuteAction.BindLambda([ContextLevelInstance](const FToolMenuContext& MenuContext)
+					{
+						TSharedPtr<FPackedLevelInstanceBuilder> Builder = FPackedLevelInstanceBuilder::CreateDefaultBuilder();
+						Builder->CreateOrUpdateBlueprint(ContextLevelInstance, nullptr);
+					});
+					UIAction.CanExecuteAction.BindLambda([](const FToolMenuContext& MenuContext)
+					{
+						return GEditor->GetSelectedActorCount() > 0;
+					});
+
+					Section.AddMenuEntry(
+						"CreatePackedBlueprint",
+						LOCTEXT("CreatePackedBlueprint", "Create Packed Blueprint"),
+						TAttribute<FText>(),
+						TAttribute<FSlateIcon>(),
+						UIAction);
+				}
+				else
+				{
+					FToolUIAction UIAction;
+					UIAction.ExecuteAction.BindLambda([ContextLevelInstance, BlueprintAsset](const FToolMenuContext& MenuContext)
+					{
+						TSharedPtr<FPackedLevelInstanceBuilder> Builder = FPackedLevelInstanceBuilder::CreateDefaultBuilder();
+						Builder->CreateOrUpdateBlueprint(ContextLevelInstance->GetWorldAsset(), BlueprintAsset);
+					});
+					UIAction.CanExecuteAction.BindLambda([](const FToolMenuContext& MenuContext)
+					{
+						return GEditor->GetSelectedActorCount() > 0;
+					});
+
+					Section.AddMenuEntry(
+						"UpdatePackedBlueprint",
+						LOCTEXT("UpdatePackedBlueprint", "Update Packed Blueprint"),
+						TAttribute<FText>(),
+						TAttribute<FSlateIcon>(),
+						UIAction);
+				}
+			}
+		}
+	}
+
 	class FLevelInstanceClassFilter : public IClassViewerFilter
 	{
 	public:
@@ -567,6 +652,13 @@ void FLevelInstanceEditorModule::StartupModule()
 
 	FLevelEditorModule& LevelEditorModule = FModuleManager::LoadModuleChecked<FLevelEditorModule>("LevelEditor");
 	LevelEditorModule.OnMapChanged().AddRaw(this, &FLevelInstanceEditorModule::OnMapChanged);
+
+	FMessageLogModule& MessageLogModule = FModuleManager::LoadModuleChecked<FMessageLogModule>("MessageLog");
+	FMessageLogInitializationOptions InitOptions;
+	InitOptions.bShowFilters = true;
+	InitOptions.bShowPages = false;
+	InitOptions.bAllowClear = true;
+	MessageLogModule.RegisterLogListing("LevelInstance", LOCTEXT("LevelInstanceLog", "Level Instance Log"), InitOptions);
 }
 
 void FLevelInstanceEditorModule::ShutdownModule()
@@ -622,6 +714,32 @@ void FLevelInstanceEditorModule::CanMoveActorToLevel(const AActor* ActorToMove,
 
 void FLevelInstanceEditorModule::ExtendContextMenu()
 {
+	if (UToolMenu* BuildMenu = UToolMenus::Get()->ExtendMenu("LevelEditor.MainMenu.Build"))
+	{
+		FToolMenuSection& Section = BuildMenu->AddSection("LevelEditorLevelInstance", LOCTEXT("LevelInstanceHeading", "Level Instance"));
+		FUIAction PackAction(
+			FExecuteAction::CreateLambda([]() 
+			{
+				UWorld* World = GEditor->GetEditorWorldContext().World();
+				if (ULevelInstanceSubsystem* LevelInstanceSubsystem = World->GetSubsystem<ULevelInstanceSubsystem>())
+				{
+					LevelInstanceSubsystem->PackLevelInstances();
+				}
+			}), 
+			FCanExecuteAction::CreateLambda([]()
+			{
+				UWorld* World = GEditor->GetEditorWorldContext().World();
+				if (ULevelInstanceSubsystem* LevelInstanceSubsystem = World->GetSubsystem<ULevelInstanceSubsystem>())
+				{
+					return LevelInstanceSubsystem->CanPackLevelInstances();
+				}
+				return false;
+			}));
+
+		FToolMenuEntry& Entry = Section.AddMenuEntry(NAME_None, LOCTEXT("PackLevelInstancesTitle", "Pack Level Instances"),
+			LOCTEXT("PackLevelInstancesTooltip", "Update packed level instances and blueprints"), FSlateIcon(FEditorStyle::GetStyleSetName(), "ClassIcon.LevelInstance"), PackAction, EUserInterfaceActionType::Button);
+	}
+
 	if (UToolMenu* Menu = UToolMenus::Get()->ExtendMenu("LevelEditor.ActorContextMenu"))
 	{
 		FToolMenuSection& Section = Menu->AddDynamicSection("ActorLevelInstance", FNewToolMenuDelegate::CreateLambda([this](UToolMenu* ToolMenu)
@@ -642,6 +760,7 @@ void FLevelInstanceEditorModule::ExtendContextMenu()
 						LevelInstanceMenuUtils::CreateCommitMenu(ToolMenu, ContextActor);
 						LevelInstanceMenuUtils::CreateSaveAsMenu(ToolMenu, ContextActor);
 						LevelInstanceMenuUtils::CreateBreakMenu(ToolMenu, ContextActor);
+						LevelInstanceMenuUtils::CreatePackedBlueprintMenu(ToolMenu, ContextActor);
 					}
 
 					LevelInstanceMenuUtils::CreateSetCurrentMenu(ToolMenu, ContextActor);
diff --git a/Engine/Source/Editor/UnrealEd/Classes/Factories/BlueprintFactory.h b/Engine/Source/Editor/UnrealEd/Classes/Factories/BlueprintFactory.h
--- a/Engine/Source/Editor/UnrealEd/Classes/Factories/BlueprintFactory.h
+++ b/Engine/Source/Editor/UnrealEd/Classes/Factories/BlueprintFactory.h
@@ -23,6 +23,10 @@ class UNREALED_API UBlueprintFactory : public UFactory
 	UPROPERTY(EditAnywhere, Category=BlueprintFactory, meta=(AllowAbstract = "", BlueprintBaseOnly = ""))
 	TSubclassOf<class UObject> ParentClass;
 
+	// Skips the class choosing dialog and uses the ParentClass as the blueprint base class
+	UPROPERTY()
+	bool bSkipClassPicker = false;
+
 	// Delegate allows overriding the class viewer initialization options before displaying it
 	DECLARE_DELEGATE_OneParam(FOnConfigureProperties, FClassViewerInitializationOptions*)
 	FOnConfigureProperties OnConfigurePropertiesDelegate;
diff --git a/Engine/Source/Editor/UnrealEd/Private/Factories/EditorFactories.cpp b/Engine/Source/Editor/UnrealEd/Private/Factories/EditorFactories.cpp
--- a/Engine/Source/Editor/UnrealEd/Private/Factories/EditorFactories.cpp
+++ b/Engine/Source/Editor/UnrealEd/Private/Factories/EditorFactories.cpp
@@ -6625,6 +6625,12 @@ UBlueprintFactory::UBlueprintFactory(const FObjectInitializer& ObjectInitializer
 
 bool UBlueprintFactory::ConfigureProperties()
 {
+	if (bSkipClassPicker)
+	{
+		check(ParentClass);
+		return true;
+	}
+
 	// Null the parent class to ensure one is selected
 	ParentClass = nullptr;
 
diff --git a/Engine/Source/Runtime/Engine/Classes/GameFramework/Actor.h b/Engine/Source/Runtime/Engine/Classes/GameFramework/Actor.h
--- a/Engine/Source/Runtime/Engine/Classes/GameFramework/Actor.h
+++ b/Engine/Source/Runtime/Engine/Classes/GameFramework/Actor.h
@@ -2006,7 +2006,7 @@ public:
 
 	/** Returns true if this actor is hidden in the editor viewports, also checking temporary flags. */
 	UFUNCTION(BlueprintCallable, Category = "Editor Scripting | Actor Editing")
-	bool IsHiddenEd() const;
+	virtual bool IsHiddenEd() const;
 
 	/**
 	 * Explicitly sets whether or not this actor is hidden in the editor for the duration of the current editor session
diff --git a/Engine/Source/Runtime/Engine/Private/LevelInstance/LevelInstanceActor.cpp b/Engine/Source/Runtime/Engine/Private/LevelInstance/LevelInstanceActor.cpp
--- a/Engine/Source/Runtime/Engine/Private/LevelInstance/LevelInstanceActor.cpp
+++ b/Engine/Source/Runtime/Engine/Private/LevelInstance/LevelInstanceActor.cpp
@@ -107,11 +107,18 @@ void ALevelInstance::PostUnregisterAllComponents()
 
 }
 
-void ALevelInstance::LoadLevelInstance()
+bool ALevelInstance::SupportsLoading() const
 {
 #if WITH_EDITOR
-	if (!bGuardLoadUnload && !bIsEditorPreviewActor)
+	return !bGuardLoadUnload && !bIsEditorPreviewActor;
+#else
+	return true;
 #endif
+}
+
+void ALevelInstance::LoadLevelInstance()
+{
+	if(SupportsLoading())
 	{
 		if (ULevelInstanceSubsystem* LevelInstanceSubsystem = GetLevelInstanceSubsystem())
 		{
@@ -127,9 +134,7 @@ void ALevelInstance::LoadLevelInstance()
 
 void ALevelInstance::UnloadLevelInstance()
 {
-#if WITH_EDITOR
-	if (!bGuardLoadUnload && !bIsEditorPreviewActor)
-#endif
+	if(SupportsLoading())
 	{
 		if (ULevelInstanceSubsystem* LevelInstanceSubsystem = GetLevelInstanceSubsystem())
 		{
@@ -223,7 +228,7 @@ void ALevelInstance::PostEditUndo()
 
 	if (CachedWorldAsset != WorldAsset)
 	{
-		UpdateLevelInstance();
+		OnWorldAssetChanged();
 	}
 
 	if (bCachedIsTemporarilyHiddenInEditor != IsTemporarilyHiddenInEditor(false))
@@ -395,7 +400,7 @@ void ALevelInstance::PostEditChangeProperty(FPropertyChangedEvent& PropertyChang
 				}
 				else
 				{
-					UpdateLevelInstance();
+					OnWorldAssetChanged();
 				}
 				CachedWorldAsset.Reset();
 			}
@@ -491,7 +496,7 @@ void ALevelInstance::UpdateLevelInstance()
 	{
 		if (ULevelInstanceSubsystem* LevelInstanceSubsystem = GetLevelInstanceSubsystem())
 		{
-			if (IsLevelInstancePathValid())
+			if (IsLevelInstancePathValid() && SupportsLoading())
 			{
 				const bool bForceUpdate = true;
 				LevelInstanceSubsystem->RequestLoadLevelInstance(this, bForceUpdate);
diff --git a/Engine/Source/Runtime/Engine/Private/LevelInstance/LevelInstanceSubsystem.cpp b/Engine/Source/Runtime/Engine/Private/LevelInstance/LevelInstanceSubsystem.cpp
--- a/Engine/Source/Runtime/Engine/Private/LevelInstance/LevelInstanceSubsystem.cpp
+++ b/Engine/Source/Runtime/Engine/Private/LevelInstance/LevelInstanceSubsystem.cpp
@@ -18,6 +18,7 @@
 #include "LevelInstance/LevelInstanceEditorLevelStreaming.h"
 #include "WorldPartition/LevelInstance/LevelInstanceActorDescFactory.h"
 #include "Misc/ScopedSlowTask.h"
+#include "Misc/ITransaction.h"
 #include "AssetRegistryModule.h"
 #include "AssetData.h"
 #include "FileHelpers.h"
@@ -29,6 +30,10 @@
 #include "Engine/LevelBounds.h"
 #include "LevelInstance/LevelInstanceEditorInstanceActor.h"
 #include "LevelEditorViewport.h"
+#include "Modules/ModuleManager.h"
+#include "Engine/Blueprint.h"
+#include "LevelInstance/Packed/PackedLevelInstanceActor.h"
+#include "LevelInstance/Packed/PackedLevelInstanceBuilder.h"
 #endif
 
 #include "HAL/IConsoleManager.h"
@@ -266,7 +271,7 @@ void ULevelInstanceSubsystem::UnloadLevelInstance(const FLevelInstanceID& LevelI
 		LevelsToRemoveScope.Reset(new FLevelsToRemoveScope());
 	}
 #endif
-
+				
 	FLevelInstance LevelInstance;
 	if (LevelInstances.RemoveAndCopyValue(LevelInstanceID, LevelInstance))
 	{
@@ -276,6 +281,9 @@ void ULevelInstanceSubsystem::UnloadLevelInstance(const FLevelInstanceID& LevelI
 			{
 				if (ALevelInstance* LevelInstanceActor = Cast<ALevelInstance>(LevelActor))
 				{
+					// Make sure to remove from pending loads if we are unloading child can't be loaded
+					LevelInstancesToLoadOrUpdate.Remove(LevelInstanceActor);
+					
 					UnloadLevelInstance(LevelInstanceActor->GetLevelInstanceID());
 				}
 				return true;
@@ -375,6 +383,70 @@ void ULevelInstanceSubsystem::Tick()
 	}
 }
 
+bool ULevelInstanceSubsystem::CanPackLevelInstances() const
+{
+	return !LevelInstanceEdits.Num();
+}
+
+void ULevelInstanceSubsystem::PackLevelInstances()
+{
+	if (!CanPackLevelInstances())
+	{
+		return;
+	}
+
+	TArray<APackedLevelInstance*> PackedLevelInstancesToUpdate;
+	TSet<UBlueprint*> BlueprintsToUpdate;
+	for (TObjectIterator<UWorld> It(RF_ClassDefaultObject | RF_ArchetypeObject, true, EInternalObjectFlags::PendingKill); It; ++It)
+	{
+		UWorld* CurrentWorld = *It;
+		for (TActorIterator<ALevelInstance> LevelInstanceIt(GetWorld()); LevelInstanceIt; ++LevelInstanceIt)
+		{
+			if (APackedLevelInstance* PackedLevelInstance = Cast<APackedLevelInstance>(*LevelInstanceIt))
+			{
+				if (UBlueprint* Blueprint = Cast<UBlueprint>(PackedLevelInstance->GetClass()->ClassGeneratedBy))
+				{
+					BlueprintsToUpdate.Add(Blueprint);
+				}
+				else
+				{
+					PackedLevelInstancesToUpdate.Add(PackedLevelInstance);
+				}
+			}
+		}
+	}
+
+	int32 Count = BlueprintsToUpdate.Num() + PackedLevelInstancesToUpdate.Num();
+	if (!Count)
+	{
+		return;
+	}
+		
+	FScopedSlowTask SlowTask(Count, (LOCTEXT("LevelInstance_PackLevelInstances", "Packing Level Instances")));
+	SlowTask.MakeDialog();
+		
+	auto UpdateProgress = [&SlowTask]()
+	{
+		if (SlowTask.CompletedWork < SlowTask.TotalAmountOfWork)
+		{
+			SlowTask.EnterProgressFrame(1, FText::Format(LOCTEXT("LevelInstance_PackLevelInstancesProgress", "Packing Level Instance {0} of {1})"), FText::AsNumber(SlowTask.CompletedWork), FText::AsNumber(SlowTask.TotalAmountOfWork)));
+		}
+	};
+
+	for (APackedLevelInstance* PackedLevelInstance : PackedLevelInstancesToUpdate)
+	{
+		PackedLevelInstance->OnWorldAssetChanged();
+		UpdateProgress();
+	}
+
+	TSharedPtr<FPackedLevelInstanceBuilder> Builder = FPackedLevelInstanceBuilder::CreateDefaultBuilder();
+	for (UBlueprint* Blueprint : BlueprintsToUpdate)
+	{
+		Builder->UpdateBlueprint(Blueprint);
+		UpdateProgress();
+	}
+}
+
 bool ULevelInstanceSubsystem::GetLevelInstanceBounds(const ALevelInstance* LevelInstanceActor, FBox& OutBounds) const
 {
 	if (IsLoaded(LevelInstanceActor))
@@ -570,7 +642,7 @@ bool ULevelInstanceSubsystem::MoveActorsTo(ALevelInstance* LevelInstanceActor, c
 	return MoveActorsToLevel(ActorsToMove, LevelInstanceLevel);
 }
 
-ALevelInstance* ULevelInstanceSubsystem::CreateLevelInstanceFrom(const TArray<AActor*>& ActorsToMove, UWorld* TemplateWorld)
+ALevelInstance* ULevelInstanceSubsystem::CreateLevelInstanceFrom(const TArray<AActor*>& ActorsToMove, ELevelInstanceCreationType CreationType, UWorld* TemplateWorld)
 {
 	ULevel* CurrentLevel = GetWorld()->GetCurrentLevel();
 		
@@ -630,9 +702,42 @@ ALevelInstance* ULevelInstanceSubsystem::CreateLevelInstanceFrom(const TArray<AA
 
 	FActorSpawnParameters SpawnParams;
 	SpawnParams.OverrideLevel = CurrentLevel;
-	ALevelInstance* NewLevelInstanceActor = GetWorld()->SpawnActor<ALevelInstance>(ALevelInstance::StaticClass(), SpawnParams);
+	ALevelInstance* NewLevelInstanceActor = nullptr;
+	TSoftObjectPtr<UWorld> WorldPtr(LoadedLevel->GetTypedOuter<UWorld>());
+	
+	if (CreationType == ELevelInstanceCreationType::LevelInstance)
+	{
+		NewLevelInstanceActor = GetWorld()->SpawnActor<ALevelInstance>(ALevelInstance::StaticClass(), SpawnParams);
+	}
+	else if (CreationType == ELevelInstanceCreationType::PackedLevelInstance)
+	{
+		NewLevelInstanceActor = GetWorld()->SpawnActor<APackedLevelInstance>(APackedLevelInstance::StaticClass(), SpawnParams);
+	}
+	else if (CreationType == ELevelInstanceCreationType::PackedLevelInstanceBlueprint)
+	{
+		int32 LastSlashIndex = 0;
+		FString LongPackageName = WorldPtr.GetLongPackageName();
+		LongPackageName.FindLastChar('/', LastSlashIndex);
+
+		FString PackagePath = LongPackageName.Mid(0, LastSlashIndex == INDEX_NONE ? MAX_int32 : LastSlashIndex);
+		FString AssetName = WorldPtr.GetAssetName() + FPackedLevelInstanceBuilder::GetPackedBPSuffix();
+		const bool bCompile = true;
+		if (UBlueprint* NewBP = FPackedLevelInstanceBuilder::CreatePackedLevelInstanceBlueprint(PackagePath, AssetName, bCompile))
+		{
+			NewLevelInstanceActor = GetWorld()->SpawnActor<APackedLevelInstance>(NewBP->GeneratedClass, SpawnParams);
+			APackedLevelInstance* CDO = Cast<APackedLevelInstance>(NewBP->GeneratedClass->GetDefaultObject());
+			CDO->SetWorldAsset(WorldPtr);
+		}
+
+		if (!NewLevelInstanceActor)
+		{
+			UE_LOG(LogLevelInstance, Warning, TEXT("Failed to create packed level blueprint. Creating non blueprint packed level instance instead."));
+			NewLevelInstanceActor = GetWorld()->SpawnActor<APackedLevelInstance>(APackedLevelInstance::StaticClass(), SpawnParams);
+		}
+	}
+	
 	check(NewLevelInstanceActor);
-	NewLevelInstanceActor->SetWorldAsset(LoadedLevel->GetTypedOuter<UWorld>());
+	NewLevelInstanceActor->SetWorldAsset(WorldPtr);
 	NewLevelInstanceActor->SetActorLocation(LevelInstanceLocation);
 
 	// Actors were moved and kept their World positions so when saving we want their positions to actually be relative to the FounationActor/LevelTransform
@@ -644,25 +749,16 @@ ALevelInstance* ULevelInstanceSubsystem::CreateLevelInstanceFrom(const TArray<AA
 	GEditor->SelectNone(false, true);
 	GEditor->SelectActor(NewLevelInstanceActor, true, true);
 
+	NewLevelInstanceActor->OnEdit();
+
 	FLevelInstanceEdit& LevelInstanceEdit = LevelInstanceEdits.Add(FName(*NewLevelInstanceActor->GetWorldAssetPackage()));
 	LevelInstanceEdit.LevelStreaming = LevelStreaming;
 	LevelStreaming->LevelInstanceID = NewLevelInstanceActor->GetLevelInstanceID();
-
+		
 	GetWorld()->SetCurrentLevel(LoadedLevel);
-
-	if (!UEditorLoadingAndSavingUtils::SavePackages({ LoadedLevel->GetOutermost() }, true))
-	{
-		UE_LOG(LogLevelInstance, Error, TEXT("Failed to save Level Instance (%s)"), *NewLevelInstanceActor->GetWorldAssetPackage());
-	}
-	else
-	{
-		UE_LOG(LogLevelInstance, Log, TEXT("Level Instance create successfully (%s)"), *NewLevelInstanceActor->GetWorldAssetPackage());
-	}
-
-	// Exit edit
-	CommitLevelInstance(NewLevelInstanceActor);
-
-	return NewLevelInstanceActor;
+				
+	// Use LevelStreaming->GetLevelInstanceActor() because OnWorldAssetSaved could've reinstanced the LevelInstanceActor
+	return CommitLevelInstance(LevelStreaming->GetLevelInstanceActor());
 }
 
 bool ULevelInstanceSubsystem::BreakLevelInstance(ALevelInstance* LevelInstanceActor, uint32 Levels)
@@ -983,7 +1079,7 @@ bool ULevelInstanceSubsystem::CanEditLevelInstance(const ALevelInstance* LevelIn
 				{
 					if (OutReason)
 					{
-						*OutReason = LOCTEXT("CanEditLevelInstanceAncestorBeingEdited", "Ancestor Level Instance already being edited");
+						*OutReason = LOCTEXT("CanEditLevelInstanceAncestorBeingEdited", "Ancestor Level Instance is being edited and dirty. Commit changes first.");
 					}
 					bAncestorBeingEdited = true;
 					return false;
@@ -1017,7 +1113,7 @@ bool ULevelInstanceSubsystem::CanEditLevelInstance(const ALevelInstance* LevelIn
 						{
 							if (OutReason)
 							{
-								*OutReason = LOCTEXT("CanEditLevelInstanceOtherChildren", "Children Level Instance already being edited");
+								*OutReason = LOCTEXT("CanEditLevelInstanceOtherChildren", "Children Level Instance already being edited and dirty. Commit changes first.");
 							}
 							bChildBeingEdited = true;
 							return false;
@@ -1078,6 +1174,7 @@ bool ULevelInstanceSubsystem::CanCommitLevelInstance(const ALevelInstance* Level
 void ULevelInstanceSubsystem::EditLevelInstance(ALevelInstance* LevelInstanceActor, TWeakObjectPtr<AActor> ContextActorPtr)
 {
 	check(CanEditLevelInstance(LevelInstanceActor));
+	const bool bWasPendingEdit = PendingLevelInstanceToEdit != InvalidLevelInstanceID;
 	PendingLevelInstanceToEdit = InvalidLevelInstanceID;
 		
 	FScopedSlowTask SlowTask(0, LOCTEXT("BeginEditLevelInstance", "Loading Level Instance for edit..."), !GetWorld()->IsGameWorld());
@@ -1121,12 +1218,21 @@ void ULevelInstanceSubsystem::EditLevelInstance(ALevelInstance* LevelInstanceAct
 			return true;
 		};
 
-		ForEachLevelInstanceAncestors(LevelInstanceActor, GetLevelInstanceToCommit);
+		ForEachLevelInstanceAncestors(LevelInstanceActor, [&LevelInstanceToCommit, &GetLevelInstanceToCommit](ALevelInstance* AncestorLevelInstance)
+		{
+			if (!LevelInstanceToCommit)
+			{
+				GetLevelInstanceToCommit(AncestorLevelInstance);
+			}
+			AncestorLevelInstance->OnEditChild();
+			return true;
+		});
+
 		if (!LevelInstanceToCommit)
 		{
 			ForEachLevelInstanceChildren(LevelInstanceActor, true, GetLevelInstanceToCommit);
 		}
-
+		
 		if (LevelInstanceToCommit)
 		{
 			PendingLevelInstanceToEdit = LevelInstanceActor->GetLevelInstanceID();
@@ -1147,7 +1253,7 @@ void ULevelInstanceSubsystem::EditLevelInstance(ALevelInstance* LevelInstanceAct
 	ULevelStreamingLevelInstanceEditor* LevelStreaming = ULevelStreamingLevelInstanceEditor::Load(LevelInstanceActor);
 	FLevelInstanceEdit& LevelInstanceEdit = LevelInstanceEdits.Add(FName(*LevelInstanceActor->GetWorldAssetPackage()));
 	LevelInstanceEdit.LevelStreaming = LevelStreaming;
-
+		
 	// Try and select something meaningful
 	AActor* ActorToSelect = nullptr;
 	if (!ActorNameToSelect.IsEmpty())
@@ -1161,6 +1267,20 @@ void ULevelInstanceSubsystem::EditLevelInstance(ALevelInstance* LevelInstanceAct
 		ActorToSelect = LevelInstanceActor;
 	}
 	LevelInstanceActor->SetIsTemporarilyHiddenInEditor(false);
+
+	// Notify
+	LevelInstanceActor->OnEdit();
+
+	// Pending edit notification was already sent
+	if (!bWasPendingEdit)
+	{
+		ForEachLevelInstanceAncestors(LevelInstanceActor, [](ALevelInstance* AncestorLevelInstance)
+		{
+			AncestorLevelInstance->OnEditChild();
+			return true;
+		});
+	}
+
 	GEditor->SelectActor(ActorToSelect, true, true);
 
 	for (const auto& Actor : LevelStreaming->LoadedLevel->Actors)
@@ -1197,7 +1317,7 @@ void ULevelInstanceSubsystem::CommitChildrenLevelInstances(ALevelInstance* Level
 	});
 }
 
-void ULevelInstanceSubsystem::CommitLevelInstance(ALevelInstance* LevelInstanceActor, bool bDiscardEdits)
+ALevelInstance* ULevelInstanceSubsystem::CommitLevelInstance(ALevelInstance* LevelInstanceActor, bool bDiscardEdits)
 {
 	check(CanCommitLevelInstance(LevelInstanceActor));
 
@@ -1222,7 +1342,7 @@ void ULevelInstanceSubsystem::CommitLevelInstance(ALevelInstance* LevelInstanceA
 				return ShouldIgnoreDirtyPackage(DirtyPackage, EditingWorld);
 			}))
 		{
-			return;
+			return LevelInstanceActor;
 		}
 
 		// Validate that we indeed need to refresh instances (user can cancel the changes when prompted)
@@ -1233,6 +1353,11 @@ void ULevelInstanceSubsystem::CommitLevelInstance(ALevelInstance* LevelInstanceA
 			// Sync the AssetData so that the updated instances have the latest Actor Registry Data
 			IAssetRegistry& AssetRegistry = FModuleManager::LoadModuleChecked<FAssetRegistryModule>(TEXT("AssetRegistry")).Get();
 			AssetRegistry.ScanPathsSynchronous({ LevelInstanceActor->GetWorldAssetPackage() }, true);
+			// Notify 
+			LevelInstanceActor->OnWorldAssetSaved();
+
+			// Update pointer since BP Compilation might have invalidated LevelInstanceActor
+			LevelInstanceActor = LevelInstanceEdit->LevelStreaming->GetLevelInstanceActor();
 		}
 	}
 
@@ -1240,12 +1365,13 @@ void ULevelInstanceSubsystem::CommitLevelInstance(ALevelInstance* LevelInstanceA
 	SlowTask.MakeDialog();
 
 	GEditor->SelectNone(false, true);
-
+		
 	// End edit non dirty child edits
 	CommitChildrenLevelInstances(LevelInstanceActor);
-
+		
 	// Try to find proper LevelInstance to select
-	ALevelInstance* ActorToSelect = LevelInstanceActor;
+	TWeakObjectPtr<ALevelInstance> ActorToSelect = LevelInstanceActor;
+	
 	ForEachLevelInstanceAncestors(LevelInstanceActor, [&ActorToSelect](ALevelInstance* AncestorLevelInstance)
 	{
 		// if we find a parent editing LevelInstance this is what we want to select
@@ -1267,6 +1393,16 @@ void ULevelInstanceSubsystem::CommitLevelInstance(ALevelInstance* LevelInstanceA
 	ULevelStreamingLevelInstanceEditor::Unload(LevelInstanceEdit->LevelStreaming);
 	LevelInstanceEdits.Remove(FName(*EditPackage));
 
+	// Notify (Actor might get destroyed by this call if its a packed bp
+	LevelInstanceActor->OnCommit();
+
+	// Notify Ancestors
+	ForEachLevelInstanceAncestors(LevelInstanceActor, [](ALevelInstance* AncestorLevelInstance)
+	{
+		AncestorLevelInstance->OnCommitChild();
+		return true;
+	});
+				
 	// Propagate to other instances
 	for (TObjectIterator<UWorld> It(RF_ClassDefaultObject|RF_ArchetypeObject, true, EInternalObjectFlags::PendingKill); It; ++It)
 	{
@@ -1284,7 +1420,10 @@ void ULevelInstanceSubsystem::CommitLevelInstance(ALevelInstance* LevelInstanceA
 		}
 	}
 
-	GEditor->SelectActor(ActorToSelect, true, true);
+	if (ALevelInstance* Actor = ActorToSelect.Get())
+	{
+		GEditor->SelectActor(Actor, true, true);
+	}
 
 	const bool bHasLevelInstanceEdits = LevelInstanceEdits.Num() > 0;
 
@@ -1298,6 +1437,8 @@ void ULevelInstanceSubsystem::CommitLevelInstance(ALevelInstance* LevelInstanceA
 			}
 		}
 	}
+
+	return LevelInstanceActor;
 }
 
 void ULevelInstanceSubsystem::SaveLevelInstanceAs(ALevelInstance* LevelInstanceActor)
@@ -1338,6 +1479,21 @@ ALevelInstance* ULevelInstanceSubsystem::GetParentLevelInstance(const AActor* Ac
 	check(OwningLevel);
 	return GetOwningLevelInstance(OwningLevel);
 }
+
+void ULevelInstanceSubsystem::BlockLoadLevelInstance(ALevelInstance* LevelInstanceActor)
+{
+	check(!LevelInstanceActor->IsEditing());
+	RequestLoadLevelInstance(LevelInstanceActor, true);
+
+	// Make sure blocking loads can happen and are not part of transaction
+	TGuardValue<ITransaction*> TransactionGuard(GUndo, nullptr);
+
+	// Blocking until LevelInstance is loaded and all its child LevelInstances
+	while (LevelInstancesToLoadOrUpdate.Num())
+	{
+		UpdateStreamingState();
+	}
+}
 #endif
 
 #undef LOCTEXT_NAMESPACE
diff --git a/Engine/Source/Runtime/Engine/Private/LevelInstance/Packed/LevelInstanceISMPacker.cpp b/Engine/Source/Runtime/Engine/Private/LevelInstance/Packed/LevelInstanceISMPacker.cpp
new file mode 100644
--- /dev/null
+++ b/Engine/Source/Runtime/Engine/Private/LevelInstance/Packed/LevelInstanceISMPacker.cpp
@@ -0,0 +1,119 @@
+// Copyright Epic Games, Inc. All Rights Reserved.
+
+#include "LevelInstance/Packed/LevelInstanceISMPacker.h"
+
+#if WITH_EDITOR
+
+#include "LevelInstance/Packed/PackedLevelInstanceBuilder.h"
+#include "LevelInstance/Packed/PackedLevelInstanceActor.h"
+
+#include "Misc/Crc.h"
+
+#include "Engine/StaticMeshActor.h"
+#include "Engine/StaticMesh.h"
+#include "Components/StaticMeshComponent.h"
+#include "Components/InstancedStaticMeshComponent.h"
+
+FLevelInstancePackerID FLevelInstanceISMPacker::PackerID = 'ISMP';
+
+FLevelInstancePackerID FLevelInstanceISMPacker::GetID() const
+{
+	return PackerID;
+}
+
+void FLevelInstanceISMPacker::GetPackClusters(FPackedLevelInstanceBuilderContext& InContext, AActor* InActor) const
+{
+	TArray<UStaticMeshComponent*> StaticMeshComponents;
+	InActor->GetComponents<UStaticMeshComponent>(StaticMeshComponents);
+
+	for(UStaticMeshComponent* StaticMeshComponent : StaticMeshComponents)
+	{
+		FLevelInstancePackerClusterID ClusterID(MakeUnique<FLevelInstanceISMPackerCluster>(GetID(), StaticMeshComponent));
+
+		InContext.FindOrAddCluster(MoveTemp(ClusterID), StaticMeshComponent);
+	}
+}
+
+void FLevelInstanceISMPacker::PackActors(FPackedLevelInstanceBuilderContext& InContext, APackedLevelInstance* InPackingActor, const FLevelInstancePackerClusterID& InClusterID, const TArray<UActorComponent*>& InComponents) const
+{
+	check(InClusterID.GetPackerID() == GetID());
+	UInstancedStaticMeshComponent* PackComponent = InPackingActor->AddPackedComponent<UInstancedStaticMeshComponent>();
+	
+	PackComponent->SetComponentToWorld(InPackingActor->GetActorTransform());
+	PackComponent->AttachToComponent(InPackingActor->GetRootComponent(), FAttachmentTransformRules::KeepWorldTransform);
+	
+	FLevelInstanceISMPackerCluster* ISMCluster = (FLevelInstanceISMPackerCluster*)InClusterID.GetData();
+	check(ISMCluster);
+
+	PackComponent->bReceivesDecals = ISMCluster->bReceivesDecals;
+	PackComponent->CastShadow = ISMCluster->bCastShadow;
+	PackComponent->bVisibleInRayTracing = ISMCluster->bVisibleInRayTracing;
+	PackComponent->SetStaticMesh(ISMCluster->StaticMesh);
+	for (int32 MaterialIndex = 0; MaterialIndex < ISMCluster->Materials.Num(); ++MaterialIndex)
+	{
+		PackComponent->SetMaterial(MaterialIndex, ISMCluster->Materials[MaterialIndex]);
+	}
+
+	for (UActorComponent* Component : InComponents)
+	{
+		// If we have a ISM we need to add all instances
+		if (UInstancedStaticMeshComponent* ISMComponent = Cast<UInstancedStaticMeshComponent>(Component))
+		{
+			for (int32 InstanceIndex = 0; InstanceIndex < ISMComponent->GetInstanceCount(); ++InstanceIndex)
+			{
+				FTransform InstanceTransform;
+
+				if (ensure(ISMComponent->GetInstanceTransform(InstanceIndex, InstanceTransform, true)))
+				{
+					PackComponent->AddInstanceWorldSpace(InstanceTransform);
+				}
+			}
+		}
+		else // other subclasses are processed like regular UStaticMeshComponent
+		{
+			UStaticMeshComponent* StaticMeshComponent = CastChecked<UStaticMeshComponent>(Component);
+			PackComponent->AddInstanceWorldSpace(StaticMeshComponent->GetComponentTransform());
+		}
+	}
+	PackComponent->RegisterComponent();
+}
+
+FLevelInstanceISMPackerCluster::FLevelInstanceISMPackerCluster(FLevelInstancePackerID InPackerID, UStaticMeshComponent* InComponent)
+	: FLevelInstancePackerCluster(InPackerID)
+	, StaticMesh(InComponent->GetStaticMesh())
+	, Materials(InComponent->GetMaterials())
+	, bReceivesDecals(InComponent->bReceivesDecals)
+	, bCastShadow(InComponent->CastShadow)
+	, bVisibleInRayTracing(InComponent->bVisibleInRayTracing)
+{
+	
+}
+
+uint32 FLevelInstanceISMPackerCluster::ComputeHash() const
+{
+	uint32 ParentHash = FLevelInstancePackerCluster::ComputeHash();
+	// Hash pointers for now since we are not persisting the clusters
+	uint32 Hash = FCrc::TypeCrc32(StaticMesh, ParentHash);
+	for (UMaterialInterface* Material : Materials)
+	{
+		Hash = FCrc::TypeCrc32(Material, Hash);
+	}
+	Hash = FCrc::TypeCrc32(bReceivesDecals, Hash);
+	Hash = FCrc::TypeCrc32(bCastShadow, Hash);
+	Hash = FCrc::TypeCrc32(bVisibleInRayTracing, Hash);
+
+	return Hash;
+}
+
+bool FLevelInstanceISMPackerCluster::operator==(const FLevelInstancePackerCluster& InOther) const
+{
+	const FLevelInstanceISMPackerCluster& ISMOther = (const FLevelInstanceISMPackerCluster&)InOther;
+	return FLevelInstancePackerCluster::operator==(InOther) &&
+		StaticMesh == ISMOther.StaticMesh &&
+		Materials == ISMOther.Materials &&
+		bReceivesDecals == ISMOther.bReceivesDecals &&
+		bCastShadow == ISMOther.bCastShadow &&
+		bVisibleInRayTracing == ISMOther.bVisibleInRayTracing;
+}
+
+#endif
\ No newline at end of file
diff --git a/Engine/Source/Runtime/Engine/Private/LevelInstance/Packed/LevelInstanceRecursivePacker.cpp b/Engine/Source/Runtime/Engine/Private/LevelInstance/Packed/LevelInstanceRecursivePacker.cpp
new file mode 100644
--- /dev/null
+++ b/Engine/Source/Runtime/Engine/Private/LevelInstance/Packed/LevelInstanceRecursivePacker.cpp
@@ -0,0 +1,80 @@
+// Copyright Epic Games, Inc. All Rights Reserved.
+
+#include "LevelInstance/Packed/LevelInstanceRecursivePacker.h"
+
+#if WITH_EDITOR
+
+#include "LevelInstance/Packed/PackedLevelInstanceBuilder.h"
+#include "LevelInstance/Packed/PackedLevelInstanceActor.h"
+
+#include "LevelInstance/LevelInstanceActor.h"
+#include "LevelInstance/LevelInstanceSubsystem.h"
+
+#include "Misc/Crc.h"
+
+FLevelInstancePackerID FLevelInstanceRecursivePacker::PackerID = 'RECP';
+
+FLevelInstancePackerID FLevelInstanceRecursivePacker::GetID() const
+{
+	return PackerID;
+}
+
+void FLevelInstanceRecursivePacker::GetPackClusters(FPackedLevelInstanceBuilderContext& InContext, AActor* InActor) const
+{
+	if (ALevelInstance* LevelInstance = Cast<ALevelInstance>(InActor))
+	{
+		FLevelInstancePackerClusterID ClusterID(MakeUnique<FLevelInstanceRecursivePackerCluster>(GetID(), LevelInstance));
+		InContext.FindOrAddCluster(MoveTemp(ClusterID));
+
+		// This Actor can be safely discarded without warning because it is a container
+		InContext.DiscardActor(LevelInstance);
+
+		ULevelInstanceSubsystem* LevelInstanceSubsystem = LevelInstance->GetLevelInstanceSubsystem();
+		check(LevelInstanceSubsystem);
+		if (ULevel* Level = LevelInstanceSubsystem->GetLevelInstanceLevel(LevelInstance))
+		{
+			for (AActor* LevelActor : Level->Actors)
+			{
+				if (LevelActor)
+				{
+					InContext.ClusterLevelActor(LevelActor);
+				}
+			}
+		}
+	}
+}
+
+void FLevelInstanceRecursivePacker::PackActors(FPackedLevelInstanceBuilderContext& InContext, APackedLevelInstance* InPackingActor, const FLevelInstancePackerClusterID& InClusterID, const TArray<UActorComponent*>& InComponents) const
+{
+	check(InClusterID.GetPackerID() == GetID());
+	FLevelInstanceRecursivePackerCluster* LevelInstanceCluster = (FLevelInstanceRecursivePackerCluster*)InClusterID.GetData();
+	check(LevelInstanceCluster);
+
+	if (LevelInstanceCluster->LevelInstance->IsLevelInstancePathValid())
+	{
+		InPackingActor->PackDependencies.AddUnique(LevelInstanceCluster->LevelInstance->GetWorldAsset());
+	}
+}
+
+FLevelInstanceRecursivePackerCluster::FLevelInstanceRecursivePackerCluster(FLevelInstancePackerID InPackerID, ALevelInstance* InLevelInstance)
+	: FLevelInstancePackerCluster(InPackerID), LevelInstance(InLevelInstance)
+{
+}
+
+bool FLevelInstanceRecursivePackerCluster::operator==(const FLevelInstancePackerCluster& InOther) const
+{
+	if (!FLevelInstancePackerCluster::operator==(InOther))
+	{
+		return false;
+	}
+
+	const FLevelInstanceRecursivePackerCluster& LevelInstanceCluster = (const FLevelInstanceRecursivePackerCluster&)InOther;
+	return LevelInstance == LevelInstanceCluster.LevelInstance;
+}
+	
+uint32 FLevelInstanceRecursivePackerCluster::ComputeHash() const
+{
+	return FCrc::TypeCrc32(LevelInstance->GetLevelInstanceID(), FLevelInstancePackerCluster::ComputeHash());
+}
+
+#endif
diff --git a/Engine/Source/Runtime/Engine/Private/LevelInstance/Packed/PackedLevelInstanceActor.cpp b/Engine/Source/Runtime/Engine/Private/LevelInstance/Packed/PackedLevelInstanceActor.cpp
new file mode 100644
--- /dev/null
+++ b/Engine/Source/Runtime/Engine/Private/LevelInstance/Packed/PackedLevelInstanceActor.cpp
@@ -0,0 +1,154 @@
+// Copyright Epic Games, Inc. All Rights Reserved.
+
+#include "LevelInstance/Packed/PackedLevelInstanceActor.h"
+#include "LevelInstance/Packed/PackedLevelInstanceBuilder.h"
+#include "LevelInstance/LevelInstanceSubsystem.h"
+
+APackedLevelInstance::APackedLevelInstance(const FObjectInitializer& ObjectInitializer)
+	: Super(ObjectInitializer)
+{
+#if WITH_EDITORONLY_DATA
+	ISMComponentClass = UInstancedStaticMeshComponent::StaticClass();
+	bEditing = false;
+#endif
+}
+
+bool APackedLevelInstance::SupportsLoading() const
+{
+#if WITH_EDITOR
+	return ChildEditing > 0 || IsLoaded();
+#else
+	return false;
+#endif
+}
+
+#if WITH_EDITOR
+FName APackedLevelInstance::GetPackedComponentTag()
+{
+	static FName PackedComponentTag("PackedComponent");
+	return PackedComponentTag;
+}
+
+void APackedLevelInstance::OnWorldAssetChanged()
+{
+	if (IsLevelInstancePathValid())
+	{
+		TSharedPtr<FPackedLevelInstanceBuilder> Builder = FPackedLevelInstanceBuilder::CreateDefaultBuilder();
+		Builder->PackActor(this, GetWorldAsset());
+	}
+	else
+	{
+		DestroyPackedComponents();
+	}
+}
+
+void APackedLevelInstance::OnWorldAssetSaved()
+{
+	TSharedPtr<FPackedLevelInstanceBuilder> Builder = FPackedLevelInstanceBuilder::CreateDefaultBuilder();
+	
+	if (UBlueprint* GeneratedBy = Cast<UBlueprint>(GetClass()->ClassGeneratedBy))
+	{
+		check(GeneratedBy == BlueprintAsset.Get());
+		Builder->UpdateBlueprint(GeneratedBy);
+	}
+	else
+	{
+		Builder->PackActor(this);
+	}
+}
+
+void APackedLevelInstance::OnEditChild()
+{
+	check(GetLevelInstanceSubsystem()->GetLevelInstanceLevel(this) != nullptr);
+	ChildEditing++;
+	MarkComponentsRenderStateDirty();
+}
+
+void APackedLevelInstance::OnCommitChild()
+{
+	check(GetLevelInstanceSubsystem()->GetLevelInstanceLevel(this) != nullptr);
+	check(ChildEditing > 0);
+	ChildEditing--;
+	if (!ChildEditing)
+	{
+		UnloadLevelInstance();
+
+		// Reflect child changes
+		TSharedPtr<FPackedLevelInstanceBuilder> Builder = FPackedLevelInstanceBuilder::CreateDefaultBuilder();
+		
+		if (UBlueprint* GeneratedBy = Cast<UBlueprint>(GetClass()->ClassGeneratedBy))
+		{
+			check(GeneratedBy == BlueprintAsset.Get());
+			Builder->UpdateBlueprint(GeneratedBy);
+		}
+		else
+		{
+			Builder->PackActor(this, GetWorldAsset());
+		}
+
+		MarkComponentsRenderStateDirty();
+	}
+}
+
+void APackedLevelInstance::OnEdit()
+{
+	check(!bEditing);
+	
+	bEditing = true;
+	MarkComponentsRenderStateDirty();
+}
+
+void APackedLevelInstance::OnCommit()
+{
+	check(bEditing);
+	bEditing = false;
+	MarkComponentsRenderStateDirty();
+}
+
+bool APackedLevelInstance::IsHiddenEd() const
+{
+	return Super::IsHiddenEd() || bEditing || (ChildEditing > 0);
+}
+
+bool APackedLevelInstance::CanEditChange(const FProperty* InProperty) const
+{
+	if (!Super::CanEditChange(InProperty))
+	{
+		return false;
+	}
+
+	// Disallow editing of the World if we are a BP instance
+	if (InProperty->GetFName() == GET_MEMBER_NAME_CHECKED(APackedLevelInstance, WorldAsset))
+	{
+		return GetClass()->IsNative();
+	}
+
+	return true;
+}
+
+void APackedLevelInstance::GetPackedComponents(TArray<UActorComponent*>& OutPackedComponents) const
+{
+	const TSet<UActorComponent*>& Components = GetComponents();
+	OutPackedComponents.Reserve(Components.Num());
+		
+	for (UActorComponent* Component : Components)
+	{
+		if (Component && Component->ComponentHasTag(GetPackedComponentTag()))
+		{
+			OutPackedComponents.Add(Component);
+		}
+	}
+}
+
+void APackedLevelInstance::DestroyPackedComponents()
+{
+	Modify();
+	TArray<UActorComponent*> PackedComponents;
+	GetPackedComponents(PackedComponents);
+	for (UActorComponent* PackedComponent : PackedComponents)
+	{
+		PackedComponent->Modify();
+		PackedComponent->DestroyComponent();
+	}
+}
+#endif
\ No newline at end of file
diff --git a/Engine/Source/Runtime/Engine/Private/LevelInstance/Packed/PackedLevelInstanceBuilder.cpp b/Engine/Source/Runtime/Engine/Private/LevelInstance/Packed/PackedLevelInstanceBuilder.cpp
new file mode 100644
--- /dev/null
+++ b/Engine/Source/Runtime/Engine/Private/LevelInstance/Packed/PackedLevelInstanceBuilder.cpp
@@ -0,0 +1,391 @@
+// Copyright Epic Games, Inc. All Rights Reserved.
+
+#include "LevelInstance/Packed/PackedLevelInstanceBuilder.h"
+
+#if WITH_EDITOR
+
+#include "LevelInstance/Packed/PackedLevelInstanceActor.h"
+#include "LevelInstance/Packed/ILevelInstancePacker.h"
+
+#include "LevelInstance/Packed/LevelInstanceISMPacker.h"
+#include "LevelInstance/Packed/LevelInstanceRecursivePacker.h"
+
+#include "LevelInstance/LevelInstanceSubsystem.h"
+
+#include "Kismet2/KismetEditorUtilities.h"
+#include "Engine/SimpleConstructionScript.h"
+#include "Engine/SCS_Node.h"
+
+#include "Misc/ScopeExit.h"
+#include "Logging/LogMacros.h"
+#include "Logging/MessageLog.h"
+
+#include "Components/SceneComponent.h"
+#include "GameFramework/WorldSettings.h"
+#include "Engine/LevelBounds.h"
+#include "Engine/Brush.h"
+#include "FileHelpers.h"
+#include "Editor.h"
+#include "AssetToolsModule.h"
+#include "Factories/BlueprintFactory.h"
+
+#define LOCTEXT_NAMESPACE "FPackedLevelInstanceBuilder"
+
+void FPackedLevelInstanceBuilderContext::ClusterLevelActor(AActor* InActor)
+{
+	PerActorClusteredComponents.FindOrAdd(InActor);
+
+	for (const auto& Pair : Packers)
+	{
+		Pair.Value->GetPackClusters(*this, InActor);
+	}
+}
+
+void FPackedLevelInstanceBuilderContext::FindOrAddCluster(FLevelInstancePackerClusterID&& InClusterID, UActorComponent* InComponent)
+{
+	TArray<UActorComponent*>& ClusterComponents = Clusters.FindOrAdd(MoveTemp(InClusterID));
+	if (InComponent)
+	{
+		ClusterComponents.Add(InComponent);
+		PerActorClusteredComponents.FindChecked(InComponent->GetOwner()).Add(InComponent);
+	}
+}
+
+void FPackedLevelInstanceBuilderContext::DiscardActor(AActor* InActor)
+{
+	ActorDiscards.Add(InActor);
+}
+
+FPackedLevelInstanceBuilder::FPackedLevelInstanceBuilder()
+{
+}
+
+const FString& FPackedLevelInstanceBuilder::GetPackedBPSuffix()
+{
+	static FString BPSuffix = "_BP_Packed";
+	return BPSuffix;
+}
+
+UBlueprint* FPackedLevelInstanceBuilder::CreatePackedLevelInstanceBlueprint(const FString& InPackagePath, const FString& InAssetName, bool bInCompile)
+{
+	IAssetTools& AssetTools = FAssetToolsModule::GetModule().Get();
+
+	UBlueprintFactory* BlueprintFactory = NewObject<UBlueprintFactory>();
+	BlueprintFactory->ParentClass = APackedLevelInstance::StaticClass();
+	BlueprintFactory->bSkipClassPicker = true;
+
+	if (UBlueprint* NewBP = Cast<UBlueprint>(AssetTools.CreateAssetWithDialog(InAssetName, InPackagePath, UBlueprint::StaticClass(), BlueprintFactory, FName("Create LevelInstance Blueprint"))))
+	{
+		APackedLevelInstance* CDO = CastChecked<APackedLevelInstance>(NewBP->GeneratedClass->GetDefaultObject());
+		CDO->BlueprintAsset = NewBP;
+
+		if (bInCompile)
+		{
+			FKismetEditorUtilities::CompileBlueprint(NewBP, EBlueprintCompileOptions::SkipGarbageCollection);
+		}
+
+		AssetTools.SyncBrowserToAssets(TArray<UObject*>{ NewBP });
+
+		return NewBP;
+	}
+
+	return nullptr;
+}
+
+TSharedPtr<FPackedLevelInstanceBuilder> FPackedLevelInstanceBuilder::CreateDefaultBuilder()
+{
+	TSharedPtr<FPackedLevelInstanceBuilder> Builder = MakeShared<FPackedLevelInstanceBuilder>();
+
+	// Class Discards are used to validate the packing result.
+	// Components or Actor classes in this set will not generate warnings
+	Builder->ClassDiscards.Add(ALevelBounds::StaticClass());
+	// Avoid dependency (find class)
+	UClass* ChaosDebugClass = FindObject<UClass>(ANY_PACKAGE, TEXT("ChaosDebugDrawComponent"));
+	if (ChaosDebugClass)
+	{
+		Builder->ClassDiscards.Add(ChaosDebugClass);
+	}
+
+	Builder->ClassDiscards.Add(AWorldSettings::StaticClass());
+	
+	// Root Components that are SceneComponents (not child class of)
+	Builder->ClassDiscards.Add(USceneComponent::StaticClass());
+	
+	Builder->Packers.Add(FLevelInstanceRecursivePacker::PackerID, MakeUnique<FLevelInstanceRecursivePacker>());
+	Builder->Packers.Add(FLevelInstanceISMPacker::PackerID, MakeUnique<FLevelInstanceISMPacker>());
+
+	return Builder;
+}
+
+void FPackedLevelInstanceBuilder::PackActor(APackedLevelInstance* InPackedLevelInstance)
+{
+	PackActor(InPackedLevelInstance, InPackedLevelInstance);
+}
+
+void FPackedLevelInstanceBuilder::PackActor(APackedLevelInstance* InPackedLevelInstance, ALevelInstance* InLevelInstanceToPack)
+{
+	FMessageLog LevelInstanceLog("LevelInstance");
+	LevelInstanceLog.Info(FText::Format(LOCTEXT("PackingStarted", "Packing of '{0}' started..."), FText::FromString(InPackedLevelInstance->GetWorldAssetPackage())));
+	
+	FPackedLevelInstanceBuilderContext Context(*this, InPackedLevelInstance);
+
+	InPackedLevelInstance->DestroyPackedComponents();
+
+	ULevelInstanceSubsystem* LevelInstanceSubystem = InPackedLevelInstance->GetLevelInstanceSubsystem();
+	check(LevelInstanceSubystem);
+	
+	ULevel* SourceLevel = LevelInstanceSubystem->GetLevelInstanceLevel(InLevelInstanceToPack);
+		
+	if (AActor* DefaultBrush = SourceLevel->GetDefaultBrush())
+	{
+		Context.DiscardActor(DefaultBrush);
+	}
+		
+	for (AActor* LevelActor : SourceLevel->Actors)
+	{
+		if (LevelActor)
+		{
+			Context.ClusterLevelActor(LevelActor);
+		}
+	}
+
+	for (const auto& Pair : Context.GetClusters())
+	{
+		TUniquePtr<ILevelInstancePacker>& Packer = Packers.FindChecked(Pair.Key.GetPackerID());
+		Packer->PackActors(Context, InPackedLevelInstance, Pair.Key, Pair.Value);
+	}
+			
+	Context.Report(LevelInstanceLog);
+}
+
+void FPackedLevelInstanceBuilderContext::Report(FMessageLog& LevelInstanceLog) const
+{
+	TSet<UActorComponent*> NotClusteredComponents;
+	uint32 TotalWarningCount = 0;
+
+	for (const auto& Pair : PerActorClusteredComponents)
+	{
+		AActor* Actor = Pair.Key;
+		const TSet<UActorComponent*>& ClusteredComponents = Pair.Value;
+
+		if (ActorDiscards.Contains(Actor))
+		{
+			LevelInstanceLog.Info(FText::Format(LOCTEXT("ActorDiscard", "Actor '{0}' ignored (Actor Discard)"), FText::FromString(Actor->GetPathName())));
+			continue;
+		}
+
+		if (Actor->GetClass()->HasAnyClassFlags(CLASS_Transient))
+		{
+			LevelInstanceLog.Info(FText::Format(LOCTEXT("ActorTransientClassDiscard", "Actor '{0}' of type '{1}' ignored (Transient Class Discard)"), FText::FromString(Actor->GetPathName()), FText::FromString(Actor->GetClass()->GetPathName())));
+			continue;
+		}
+
+		// Class must match (not a child)
+		if (ClassDiscards.Contains(Actor->GetClass()))
+		{
+			LevelInstanceLog.Info(FText::Format(LOCTEXT("ActorClassDiscard", "Actor '{0}' of type '{1}' ignored (Class Discard)"), FText::FromString(Actor->GetPathName()), FText::FromString(Actor->GetClass()->GetPathName())));
+			continue;
+		}
+
+		NotClusteredComponents = Actor->GetComponents().Difference(ClusteredComponents);
+		uint32 WarningCount = 0;
+		for (UActorComponent* Component : NotClusteredComponents)
+		{
+			if (Component == nullptr)
+			{
+				continue;
+			}
+
+			if (ClassDiscards.Contains(Component->GetClass()))
+			{
+				LevelInstanceLog.Info(FText::Format(LOCTEXT("ComponentClassDiscard", "Component '{0}' of type '{1}' ignored (Class Discard)"), FText::FromString(Component->GetPathName()), FText::FromString(Component->GetClass()->GetPathName())));
+				continue;
+			}
+
+			if (Actor->GetClass()->HasAnyClassFlags(CLASS_Transient))
+			{
+				LevelInstanceLog.Info(FText::Format(LOCTEXT("ComopnentTransientClassDiscard", "Component '{0}' of type '{1}' ignored (Transient Class Discard)"), FText::FromString(Component->GetPathName()), FText::FromString(Component->GetClass()->GetPathName())));
+				continue;
+			}
+
+			WarningCount++;
+
+			LevelInstanceLog.Warning(FText::Format(LOCTEXT("ComponentNotPacked", "Component '{0}' was not packed"), FText::FromString(Component->GetPathName())));
+		}
+
+		if (WarningCount)
+		{
+			LevelInstanceLog.Warning(FText::Format(LOCTEXT("ActorNotPacked", "Actor '{0}' was not packed completely ({1} warning(s))"), FText::FromString(Actor->GetPathName()), FText::AsNumber(WarningCount)));
+		}
+		else
+		{
+			LevelInstanceLog.Info(FText::Format(LOCTEXT("ActorPacked", "Actor '{0}' packed successfully"), FText::FromString(Actor->GetPathName())));
+		}
+		TotalWarningCount += WarningCount;
+	}
+
+	if (TotalWarningCount)
+	{
+		LevelInstanceLog.Warning(LOCTEXT("WarningsReported", "Warnings have been reported. Consider using a regular ALevelInstance instead."));
+		LevelInstanceLog.Open();
+	}
+	LevelInstanceLog.Info(FText::Format(LOCTEXT("PackCompleted", "Packing '{0}' completed with {1} warning(s)"), FText::FromString(PackedLevelInstance->GetWorldAssetPackage()), FText::AsNumber(TotalWarningCount)));
+}
+
+ALevelInstance* FPackedLevelInstanceBuilder::CreateTransientLevelInstanceForPacking(TSoftObjectPtr<UWorld> InWorldAsset, const FVector& InLocation, const FRotator& InRotator)
+{
+	// Create Temp Actor for Packing
+	FActorSpawnParameters SpawnParams;
+	SpawnParams.bCreateActorPackage = false;
+	SpawnParams.bHideFromSceneOutliner = true;
+	SpawnParams.bNoFail = true;
+	SpawnParams.ObjectFlags |= RF_Transient;
+
+	UWorld* World = GEditor->GetEditorWorldContext().World();
+	check(World);
+	SpawnParams.OverrideLevel = World->PersistentLevel;
+	ALevelInstance* LevelInstance = World->SpawnActor<ALevelInstance>(InLocation, InRotator, SpawnParams);
+	LevelInstance->SetWorldAsset(InWorldAsset);
+
+	// Wait for load
+	LevelInstance->GetLevelInstanceSubsystem()->BlockLoadLevelInstance(LevelInstance);
+
+	return LevelInstance;
+}
+
+void FPackedLevelInstanceBuilder::PackActor(APackedLevelInstance* InActor, TSoftObjectPtr<UWorld> InWorldAsset)
+{
+	ALevelInstance* TransientLevelInstance = CreateTransientLevelInstanceForPacking(InWorldAsset, InActor->GetActorLocation(), InActor->GetActorRotation());
+
+	PackActor(InActor, TransientLevelInstance);
+
+	TransientLevelInstance->GetWorld()->DestroyActor(TransientLevelInstance);
+}
+
+void FPackedLevelInstanceBuilder::UpdateBlueprint(UBlueprint* Blueprint)
+{
+	APackedLevelInstance* CDO = CastChecked<APackedLevelInstance>(Blueprint->GeneratedClass->GetDefaultObject());
+	check(CDO);
+
+	CreateOrUpdateBlueprint(CDO->GetWorldAsset(), Blueprint);
+}
+
+bool FPackedLevelInstanceBuilder::CreateOrUpdateBlueprint(TSoftObjectPtr<UWorld> InWorldAsset, TSoftObjectPtr<UBlueprint> InBlueprintAsset)
+{
+	bool bResult = true;
+	
+	ALevelInstance* TransientLevelInstance = CreateTransientLevelInstanceForPacking(InWorldAsset, FVector::ZeroVector, FRotator::ZeroRotator);
+	
+	bResult = CreateOrUpdateBlueprintFromUnpacked(TransientLevelInstance, InBlueprintAsset);
+
+	TransientLevelInstance->GetWorld()->DestroyActor(TransientLevelInstance);
+
+	return bResult;
+}
+
+bool FPackedLevelInstanceBuilder::CreateOrUpdateBlueprint(ALevelInstance* InLevelInstance, TSoftObjectPtr<UBlueprint> InBlueprintAsset)
+{
+	if (APackedLevelInstance* PackedLevelInstance = Cast<APackedLevelInstance>(InLevelInstance))
+	{
+		return CreateOrUpdateBlueprintFromPacked(PackedLevelInstance, InBlueprintAsset);
+	}
+	
+	return CreateOrUpdateBlueprintFromUnpacked(InLevelInstance, InBlueprintAsset);
+}
+
+bool FPackedLevelInstanceBuilder::CreateOrUpdateBlueprintFromUnpacked(ALevelInstance* InActor, TSoftObjectPtr<UBlueprint> InBlueprintAsset)
+{
+	bool bResult = true;
+	
+	// Create Temp Actor for Packing
+	FActorSpawnParameters SpawnParams;
+	SpawnParams.bCreateActorPackage = false;
+	SpawnParams.bHideFromSceneOutliner = true;
+	SpawnParams.bNoFail = true;
+	SpawnParams.ObjectFlags |= RF_Transient;
+
+	UWorld* World = InActor->GetWorld();
+	SpawnParams.OverrideLevel = World->PersistentLevel;
+
+	APackedLevelInstance* PackedLevelInstance = World->SpawnActor<APackedLevelInstance>(InActor->GetActorLocation(), InActor->GetActorRotation(), SpawnParams);
+	PackedLevelInstance->SetWorldAsset(InActor->GetWorldAsset());
+
+	PackActor(PackedLevelInstance, InActor);
+
+	PackedLevelInstance->BlueprintAsset = InBlueprintAsset;
+	bResult &= CreateOrUpdateBlueprintFromPacked(PackedLevelInstance, PackedLevelInstance->BlueprintAsset);
+
+	InActor->GetWorld()->DestroyActor(PackedLevelInstance);
+
+	return bResult;
+}
+
+bool FPackedLevelInstanceBuilder::CreateOrUpdateBlueprintFromPacked(APackedLevelInstance* InActor, TSoftObjectPtr<UBlueprint> InBlueprintAsset)
+{
+	UBlueprint* BP = nullptr;
+	if (!InBlueprintAsset.IsNull())
+	{
+		BP = InBlueprintAsset.LoadSynchronous();
+		if (BP && BP->SimpleConstructionScript)
+		{
+			TArray<USCS_Node*> AllNodes = BP->SimpleConstructionScript->GetAllNodes();
+			for (USCS_Node* Node : AllNodes)
+			{
+				BP->SimpleConstructionScript->RemoveNodeAndPromoteChildren(Node);
+			}
+		}
+	}
+
+	if (!BP)
+	{
+		int32 LastSlashIndex = 0;
+		FString LongPackageName = InActor->GetWorldAsset().GetLongPackageName();
+		LongPackageName.FindLastChar('/', LastSlashIndex);
+
+		FString PackagePath = LongPackageName.Mid(0, LastSlashIndex == INDEX_NONE ? MAX_int32 : LastSlashIndex);
+		FString AssetName = InActor->GetWorldAsset().GetAssetName() + GetPackedBPSuffix();
+		const bool bCompile = false;
+		BP = CreatePackedLevelInstanceBlueprint(PackagePath, AssetName, bCompile);
+	}
+		
+	if (BP == nullptr)
+	{
+		return false;
+	}
+
+	BP->Modify();
+	APackedLevelInstance* CDO = CastChecked<APackedLevelInstance>(BP->GeneratedClass->GetDefaultObject());
+	CDO->SetWorldAsset(InActor->GetWorldAsset());
+	
+	// match root component mobility to source actor
+	USceneComponent* Root = CDO->GetRootComponent();
+	Root->SetMobility(InActor->GetRootComponent()->Mobility);
+		
+	// Prep AddComponentsToBlueprintParam
+	FKismetEditorUtilities::FAddComponentsToBlueprintParams AddCompToBPParams;
+	AddCompToBPParams.HarvestMode = FKismetEditorUtilities::EAddComponentToBPHarvestMode::None;
+	AddCompToBPParams.bKeepMobility = true;
+
+	// Add Components
+	TArray<UActorComponent*> PackedComponents;
+	InActor->GetPackedComponents(PackedComponents);
+	FKismetEditorUtilities::AddComponentsToBlueprint(BP, PackedComponents, AddCompToBPParams);
+	// If we are packing the actors BP then destroy packed components as they are now part of the BPs construction script
+	if (UBlueprint* GeneratedBy = Cast<UBlueprint>(InActor->GetClass()->ClassGeneratedBy))
+	{
+		InActor->DestroyPackedComponents();
+	}
+
+	// Synchronous compile
+	FKismetEditorUtilities::CompileBlueprint(BP, EBlueprintCompileOptions::SkipGarbageCollection);
+
+	const bool bCheckDirty = false;
+	const bool bPromptToSave = true;
+	TArray<UPackage*> OutFailedPackages;
+	FEditorFileUtils::PromptForCheckoutAndSave({ BP->GetPackage() }, bCheckDirty, bPromptToSave, &OutFailedPackages);
+		
+	return !OutFailedPackages.Num();
+}
+#endif
+
+#undef LOCTEXT_NAMESPACE
\ No newline at end of file
diff --git a/Engine/Source/Runtime/Engine/Public/LevelInstance/LevelInstanceActor.h b/Engine/Source/Runtime/Engine/Public/LevelInstance/LevelInstanceActor.h
--- a/Engine/Source/Runtime/Engine/Public/LevelInstance/LevelInstanceActor.h
+++ b/Engine/Source/Runtime/Engine/Public/LevelInstance/LevelInstanceActor.h
@@ -22,6 +22,7 @@ public:
 
 	void LoadLevelInstance();
 	void UnloadLevelInstance();
+	virtual bool SupportsLoading() const;
 	bool IsLevelInstancePathValid() const;
 	const TSoftObjectPtr<UWorld>& GetWorldAsset() const { return WorldAsset; }
 	class ULevelInstanceSubsystem* GetLevelInstanceSubsystem() const;
@@ -75,6 +76,13 @@ public:
 	bool CheckForLoop(TSoftObjectPtr<UWorld> WorldAsset, TArray<TPair<FText, TSoftObjectPtr<UWorld>>>* LoopInfo = nullptr, const ALevelInstance** LoopStart = nullptr) const;
 	AActor* FindEditorInstanceActor() const;
 
+	virtual void OnWorldAssetChanged() { UpdateLevelInstance(); }
+	virtual void OnWorldAssetSaved() {}
+	virtual void OnEdit() {}
+	virtual void OnEditChild() {}
+	virtual void OnCommit() {}
+	virtual void OnCommitChild() {}
+		
 	DECLARE_MULTICAST_DELEGATE_OneParam(FOnLevelInstanceActorPostLoad, ALevelInstance*);
 	static FOnLevelInstanceActorPostLoad OnLevelInstanceActorPostLoad;
 #endif
diff --git a/Engine/Source/Runtime/Engine/Public/LevelInstance/LevelInstanceSubsystem.h b/Engine/Source/Runtime/Engine/Public/LevelInstance/LevelInstanceSubsystem.h
--- a/Engine/Source/Runtime/Engine/Public/LevelInstance/LevelInstanceSubsystem.h
+++ b/Engine/Source/Runtime/Engine/Public/LevelInstance/LevelInstanceSubsystem.h
@@ -17,6 +17,15 @@ class ALevelInstance;
 class ULevelStreamingLevelInstance;
 class ULevelStreamingLevelInstanceEditor;
 class UWorldPartitionSubsystem;
+class UBlueprint;
+
+UENUM()
+enum class ELevelInstanceCreationType : uint8
+{
+	LevelInstance,
+	PackedLevelInstance,
+	PackedLevelInstanceBlueprint
+};
 
 /**
  * ULevelInstanceSubsystem
@@ -50,11 +59,13 @@ public:
 
 #if WITH_EDITOR
 	void Tick();
+	void PackLevelInstances();
+	bool CanPackLevelInstances() const;
 
 	bool CanEditLevelInstance(const ALevelInstance* LevelInstanceActor, FText* OutReason = nullptr) const;
 	bool CanCommitLevelInstance(const ALevelInstance* LevelInstanceActor, FText* OutReason = nullptr) const;
 	void EditLevelInstance(ALevelInstance* LevelInstanceActor, TWeakObjectPtr<AActor> ContextActorPtr = nullptr);
-	void CommitLevelInstance(ALevelInstance* LevelInstanceActor, bool bDiscardEdits = false);
+	ALevelInstance* CommitLevelInstance(ALevelInstance* LevelInstanceActor, bool bDiscardEdits = false);
 	void SaveLevelInstanceAs(ALevelInstance* LevelInstanceActor);
 	bool IsEditingLevelInstanceDirty(const ALevelInstance* LevelInstanceActor) const;
 	bool IsEditingLevelInstance(const ALevelInstance* LevelInstanceActor) const { return GetLevelInstanceEdit(LevelInstanceActor) != nullptr; }
@@ -75,7 +86,7 @@ public:
 
 	bool SetCurrent(ALevelInstance* LevelInstanceActor) const;
 	bool IsCurrent(const ALevelInstance* LevelInstanceActor) const;
-	ALevelInstance* CreateLevelInstanceFrom(const TArray<AActor*>& ActorsToMove, UWorld* TemplateWorld = nullptr);
+	ALevelInstance* CreateLevelInstanceFrom(const TArray<AActor*>& ActorsToMove, ELevelInstanceCreationType CreationType = ELevelInstanceCreationType::LevelInstance, UWorld* TemplateWorld = nullptr);
 	bool MoveActorsToLevel(const TArray<AActor*>& ActorsToRemove, ULevel* DestinationLevel) const;
 	bool MoveActorsTo(ALevelInstance* LevelInstanceActor, const TArray<AActor*>& ActorsToMove);
 	bool BreakLevelInstance(ALevelInstance* LevelInstanceActor, uint32 Levels = 1);
@@ -86,10 +97,10 @@ public:
 	ULevel* GetLevelInstanceLevel(const ALevelInstance* LevelInstanceActor) const;
 
 	ALevelInstance* GetParentLevelInstance(const AActor* Actor) const;
+	void BlockLoadLevelInstance(ALevelInstance* LevelInstanceActor);
 #endif
 
 private:
-
 	void LoadLevelInstance(ALevelInstance* LevelInstanceActor);
 	void UnloadLevelInstance(const FLevelInstanceID& LevelInstanceID);
 	void ForEachActorInLevel(ULevel* Level, TFunctionRef<bool(AActor * LevelActor)> Operation) const;
diff --git a/Engine/Source/Runtime/Engine/Public/LevelInstance/Packed/ILevelInstancePacker.h b/Engine/Source/Runtime/Engine/Public/LevelInstance/Packed/ILevelInstancePacker.h
new file mode 100644
--- /dev/null
+++ b/Engine/Source/Runtime/Engine/Public/LevelInstance/Packed/ILevelInstancePacker.h
@@ -0,0 +1,30 @@
+// Copyright Epic Games, Inc. All Rights Reserved.
+
+#pragma once
+
+#if WITH_EDITOR
+
+#include "CoreMinimal.h"
+#include "UObject/ObjectMacros.h"
+#include "LevelInstance/Packed/PackedLevelInstanceTypes.h"
+
+class APackedLevelInstance;
+class AActor;
+class UActorComponent;
+class FPackedLevelInstanceBuilder;
+class FPackedLevelInstanceBuilderContext;
+
+class ENGINE_API ILevelInstancePacker
+{
+public:
+	ILevelInstancePacker() {}
+	virtual ~ILevelInstancePacker() {}
+	virtual FLevelInstancePackerID GetID() const = 0;
+	virtual void GetPackClusters(FPackedLevelInstanceBuilderContext& InContext, AActor* InActor) const = 0;
+	virtual void PackActors(FPackedLevelInstanceBuilderContext& InBuilder, APackedLevelInstance* InPackingActor, const FLevelInstancePackerClusterID& InClusterID, const TArray<UActorComponent*>& InComponents) const = 0;
+private:
+	ILevelInstancePacker(const ILevelInstancePacker&) = delete;
+	ILevelInstancePacker& operator=(const ILevelInstancePacker&) = delete;
+};
+
+#endif
\ No newline at end of file
diff --git a/Engine/Source/Runtime/Engine/Public/LevelInstance/Packed/LevelInstanceISMPacker.h b/Engine/Source/Runtime/Engine/Public/LevelInstance/Packed/LevelInstanceISMPacker.h
new file mode 100644
--- /dev/null
+++ b/Engine/Source/Runtime/Engine/Public/LevelInstance/Packed/LevelInstanceISMPacker.h
@@ -0,0 +1,44 @@
+// Copyright Epic Games, Inc. All Rights Reserved.
+
+#pragma once
+
+#if WITH_EDITOR
+
+#include "LevelInstance/Packed/PackedLevelInstanceTypes.h"
+#include "LevelInstance/Packed/ILevelInstancePacker.h"
+
+class AActor;
+class APackedLevelInstance;
+class UActorComponent;
+class FPackedLevelInstanceBuilderContext;
+
+class FLevelInstanceISMPacker : public ILevelInstancePacker
+{
+public:
+	static FLevelInstancePackerID PackerID;
+
+	virtual FLevelInstancePackerID GetID() const override;
+	virtual void GetPackClusters(FPackedLevelInstanceBuilderContext& InContext, AActor* InActor) const override;
+	virtual void PackActors(FPackedLevelInstanceBuilderContext& InContext, APackedLevelInstance* InPackingActor, const FLevelInstancePackerClusterID& InClusterID, const TArray<UActorComponent*>& InComponents) const override;
+};
+
+class UStaticMesh;
+class UMaterialInterface;
+
+class FLevelInstanceISMPackerCluster : public FLevelInstancePackerCluster
+{
+public:
+	FLevelInstanceISMPackerCluster(FLevelInstancePackerID InPackerID, UStaticMeshComponent* InComponent);
+
+	virtual bool operator==(const FLevelInstancePackerCluster& InOther) const override;
+	virtual uint32 ComputeHash() const override;
+
+	// UStaticMeshComponent that contribute to the hash. More can be added if needed.
+	UStaticMesh* StaticMesh = nullptr;
+	TArray<UMaterialInterface*> Materials;
+	bool bReceivesDecals;
+	bool bCastShadow;
+	bool bVisibleInRayTracing;
+};
+
+#endif
\ No newline at end of file
diff --git a/Engine/Source/Runtime/Engine/Public/LevelInstance/Packed/LevelInstanceRecursivePacker.h b/Engine/Source/Runtime/Engine/Public/LevelInstance/Packed/LevelInstanceRecursivePacker.h
new file mode 100644
--- /dev/null
+++ b/Engine/Source/Runtime/Engine/Public/LevelInstance/Packed/LevelInstanceRecursivePacker.h
@@ -0,0 +1,37 @@
+// Copyright Epic Games, Inc. All Rights Reserved.
+
+#pragma once
+
+#if WITH_EDITOR
+
+#include "LevelInstance/Packed/PackedLevelInstanceTypes.h"
+#include "LevelInstance/Packed/ILevelInstancePacker.h"
+
+class AActor;
+class ALevelInstance;
+class APackedLevelInstance;
+class UActorComponent;
+class FPackedLevelInstanceBuilderContext;
+
+class FLevelInstanceRecursivePacker : public ILevelInstancePacker
+{
+public:
+	static FLevelInstancePackerID PackerID;
+
+	virtual FLevelInstancePackerID GetID() const override;
+	virtual void GetPackClusters(FPackedLevelInstanceBuilderContext& InContext, AActor* InActor) const override;
+	virtual void PackActors(FPackedLevelInstanceBuilderContext& InContext, APackedLevelInstance* InPackingActor, const FLevelInstancePackerClusterID& InClusterID, const TArray<UActorComponent*>& InComponents) const override;
+};
+
+class FLevelInstanceRecursivePackerCluster : public FLevelInstancePackerCluster
+{
+public:
+	FLevelInstanceRecursivePackerCluster(FLevelInstancePackerID InPackerID, ALevelInstance* InLevelInstance);
+
+	virtual bool operator==(const FLevelInstancePackerCluster& InOther) const override;
+	virtual uint32 ComputeHash() const override;
+
+	ALevelInstance* LevelInstance = nullptr;
+};
+
+#endif
\ No newline at end of file
diff --git a/Engine/Source/Runtime/Engine/Public/LevelInstance/Packed/PackedLevelInstanceActor.h b/Engine/Source/Runtime/Engine/Public/LevelInstance/Packed/PackedLevelInstanceActor.h
new file mode 100644
--- /dev/null
+++ b/Engine/Source/Runtime/Engine/Public/LevelInstance/Packed/PackedLevelInstanceActor.h
@@ -0,0 +1,81 @@
+// Copyright Epic Games, Inc. All Rights Reserved.
+
+#pragma once
+
+#include "CoreMinimal.h"
+#include "UObject/ObjectMacros.h"
+#include "LevelInstance/LevelInstanceActor.h"
+
+#if WITH_EDITOR
+#include "Kismet2/ComponentEditorUtils.h"
+#endif
+
+#include "PackedLevelInstanceActor.generated.h"
+
+class UInstancedStaticMeshComponent;
+class UBlueprint;
+
+/**
+ * APackedLevelInstance is the result of packing the source level (WorldAsset base class property) into a single actor. See FPackedLevelInstanceBuilder.
+ * 
+ * Currently supported source components:
+ * Packer FLevelInstanceISMPacker : UStaticMeshComponent/UInstancedStaticMeshComponent/UHierarchicalInstancedStaticMeshComponent
+ * Packer FRecursiveLevelInstancePacker : Allows packing recursive LevelInstances
+ * 
+ * Other components are unsupported and will result in an incomplete APackedLevelInstance. In this case using a regular ALevelInstance is recommended.
+ */
+UCLASS()
+class ENGINE_API APackedLevelInstance : public ALevelInstance
+{
+	GENERATED_UCLASS_BODY()
+
+public:
+	virtual bool SupportsLoading() const override;
+
+#if WITH_EDITOR
+	static FName GetPackedComponentTag();
+
+	virtual void OnWorldAssetChanged() override;
+	virtual void OnWorldAssetSaved() override;
+	virtual void OnCommit() override;
+	virtual void OnCommitChild() override;
+	virtual void OnEdit() override;
+	virtual void OnEditChild() override;
+
+	virtual bool CanEditChange(const FProperty* InProperty) const override;
+
+	virtual bool IsHiddenEd() const override;
+
+	void DestroyPackedComponents();
+	void GetPackedComponents(TArray<UActorComponent*>& OutPackedComponents) const;
+
+	template<class T>
+	T* AddPackedComponent()
+	{
+		Modify();
+		FName NewComponentName = *FComponentEditorUtils::GenerateValidVariableName(T::StaticClass(), this);
+		T* NewComponent = NewObject<T>(this, NewComponentName, RF_Transactional);
+		AddInstanceComponent(NewComponent);
+		NewComponent->ComponentTags.Add(GetPackedComponentTag());
+		return NewComponent;
+	}
+#endif
+
+#if WITH_EDITORONLY_DATA
+	UPROPERTY(EditAnywhere, Category = Pack, meta = (DisplayName = "ISM Component Class"))
+	TSubclassOf<UInstancedStaticMeshComponent> ISMComponentClass;
+
+	UPROPERTY(VisibleAnywhere, Category = Pack)
+	TSoftObjectPtr<UBlueprint> BlueprintAsset;
+
+	UPROPERTY()
+	TArray<TSoftObjectPtr<UWorld>> PackDependencies;
+
+private:
+	UPROPERTY(NonTransactional)
+	bool bEditing = false;
+
+	UPROPERTY(NonTransactional)
+	int32 ChildEditing = 0;
+#endif
+};
\ No newline at end of file
diff --git a/Engine/Source/Runtime/Engine/Public/LevelInstance/Packed/PackedLevelInstanceBuilder.h b/Engine/Source/Runtime/Engine/Public/LevelInstance/Packed/PackedLevelInstanceBuilder.h
new file mode 100644
--- /dev/null
+++ b/Engine/Source/Runtime/Engine/Public/LevelInstance/Packed/PackedLevelInstanceBuilder.h
@@ -0,0 +1,95 @@
+// Copyright Epic Games, Inc. All Rights Reserved.
+
+#pragma once
+
+#if WITH_EDITOR
+
+#include "CoreMinimal.h"
+#include "UObject/ObjectMacros.h"
+#include "LevelInstance/Packed/PackedLevelInstanceTypes.h"
+#include "LevelInstance/Packed/ILevelInstancePacker.h"
+#include "UObject/SoftObjectPtr.h"
+#include "Containers/Set.h"
+
+class APackedLevelInstance;
+class ALevelInstance;
+class AActor;
+class UActorComponent;
+class UBlueprint;
+class FMessageLog;
+
+/**
+ * FPackedLevelInstanceBuiler handles packing of ALevelInstance actors into APackedLevelInstance actors and Blueprints.
+ */
+class ENGINE_API FPackedLevelInstanceBuilder
+{
+public:
+	FPackedLevelInstanceBuilder();
+	
+	static TSharedPtr<FPackedLevelInstanceBuilder> CreateDefaultBuilder();
+	
+	/* Packs InPackedLevelInstance using InLevelInstanceToPack as its source level actor */
+	void PackActor(APackedLevelInstance* InPackedLevelInstance, ALevelInstance* InLevelInstanceToPack);
+	/* Packs InPackedLevelInstance using itself as the source level actor */
+	void PackActor(APackedLevelInstance* InPackedLevelInstance);
+	/* Packs InPackedLevelInstance using InWorldAsset as the source level */
+	void PackActor(APackedLevelInstance* InPackedLevelInstance, TSoftObjectPtr<UWorld> InWorldAsset);
+
+	/* Creates/Updates a APackedLevelInstance Blueprint from InLevelInstance (will overwrite existing asset or show a dialog if InBlueprintAsset is null) */
+	bool CreateOrUpdateBlueprint(ALevelInstance* InLevelInstance, TSoftObjectPtr<UBlueprint> InBlueprintAsset);
+	/* Creates/Updates a APackedLeveInstance Blueprint from InWorldAsset (will overwrite existing asset or show a dialog if InBlueprintAsset is null) */
+	bool CreateOrUpdateBlueprint(TSoftObjectPtr<UWorld> InWorldAsset, TSoftObjectPtr<UBlueprint> InBlueprintAsset);
+	/* Update existing Blueprint */
+	void UpdateBlueprint(UBlueprint* Blueprint);
+
+	/* Creates a new APackedLevelInstance Blueprint using InPackagePath/InAssetName as hint for path */
+	static const FString& GetPackedBPSuffix();
+	static UBlueprint* CreatePackedLevelInstanceBlueprint(const FString& InPackagePath, const FString& InAssetName, bool bInCompile);
+	
+private:
+	/* Create/Updates a APackedLevelInstance Blueprint from InPackedActor (will overwrite existing asset or show a dialog if InBlueprintAsset is null) */
+	bool CreateOrUpdateBlueprintFromPacked(APackedLevelInstance* InPackedActor, TSoftObjectPtr<UBlueprint> InBlueprintAsset);
+	/* Creates/Updates a APackedLevelInstance Blueprint from InLevelInstance (will overwrite existing asset or show a dialog if InBlueprintAsset is null) */
+	bool CreateOrUpdateBlueprintFromUnpacked(ALevelInstance* InPackedActor, TSoftObjectPtr<UBlueprint> InBlueprintAsset);
+	/* Creates and loads a ALevelInstance so it can be used for packing */
+	ALevelInstance* CreateTransientLevelInstanceForPacking(TSoftObjectPtr<UWorld> InWorldAsset, const FVector& InLocation, const FRotator& InRotator);
+
+	FPackedLevelInstanceBuilder(const FPackedLevelInstanceBuilder&) = delete;
+	FPackedLevelInstanceBuilder& operator=(const FPackedLevelInstanceBuilder&) = delete;
+
+public:
+	friend class FPackedLevelInstanceBuilderContext;
+		
+private:
+	
+	TSet<UClass*> ClassDiscards;
+	TMap<FLevelInstancePackerID, TUniquePtr<ILevelInstancePacker>> Packers;
+};
+
+
+class FPackedLevelInstanceBuilderContext
+{
+public:
+	FPackedLevelInstanceBuilderContext(const FPackedLevelInstanceBuilder& InBuilder, APackedLevelInstance* InPackedLevelInstance) : Packers(InBuilder.Packers), ClassDiscards(InBuilder.ClassDiscards), PackedLevelInstance(InPackedLevelInstance) {}
+
+	/* Interface for ILevelInstancePacker's to use */
+	void ClusterLevelActor(AActor* InLevelActor);
+	void FindOrAddCluster(FLevelInstancePackerClusterID&& InClusterID, UActorComponent* InComponent = nullptr);
+	void DiscardActor(AActor* InActor);
+	void Report(FMessageLog& LevelInstanceLog) const;
+
+	const TMap<FLevelInstancePackerClusterID, TArray<UActorComponent*>>& GetClusters() const { return Clusters; }
+
+private:
+	const TMap<FLevelInstancePackerID, TUniquePtr<ILevelInstancePacker>>& Packers;
+	const TSet<UClass*>& ClassDiscards;
+
+	APackedLevelInstance* PackedLevelInstance;
+	
+	TMap<FLevelInstancePackerClusterID, TArray<UActorComponent*>> Clusters;
+
+	TMap<AActor*, TSet<UActorComponent*>> PerActorClusteredComponents;
+	TSet<AActor*> ActorDiscards;
+};
+
+#endif
\ No newline at end of file
diff --git a/Engine/Source/Runtime/Engine/Public/LevelInstance/Packed/PackedLevelInstanceTypes.h b/Engine/Source/Runtime/Engine/Public/LevelInstance/Packed/PackedLevelInstanceTypes.h
new file mode 100644
--- /dev/null
+++ b/Engine/Source/Runtime/Engine/Public/LevelInstance/Packed/PackedLevelInstanceTypes.h
@@ -0,0 +1,92 @@
+// Copyright Epic Games, Inc. All Rights Reserved.
+
+#pragma once
+
+#if WITH_EDITOR
+
+#include "CoreMinimal.h"
+#include "UObject/ObjectMacros.h"
+#include "Misc/Crc.h"
+#include "Logging/LogMacros.h"
+
+using FLevelInstancePackerID = uint32;
+
+class ENGINE_API FLevelInstancePackerCluster
+{
+	FLevelInstancePackerID PackerID;
+
+public:
+	FLevelInstancePackerCluster(FLevelInstancePackerID InPackerID)
+		: PackerID(InPackerID) {}
+	virtual ~FLevelInstancePackerCluster() {}
+
+	FLevelInstancePackerID GetPackerID() const { return PackerID; }
+
+	virtual uint32 ComputeHash() const
+	{
+		return FCrc::TypeCrc32(PackerID);
+	}
+
+	virtual bool operator==(const FLevelInstancePackerCluster& Other) const
+	{
+		return PackerID == Other.PackerID;
+	}
+
+private:
+	FLevelInstancePackerCluster(const FLevelInstancePackerCluster&) = delete;
+	FLevelInstancePackerCluster& operator=(const FLevelInstancePackerCluster&) = delete;
+};
+
+class ENGINE_API FLevelInstancePackerClusterID
+{
+public:
+	static FLevelInstancePackerClusterID Invalid;
+
+	FLevelInstancePackerClusterID() : Hash(0)
+	{
+	}
+
+	FLevelInstancePackerClusterID(FLevelInstancePackerClusterID&& Other)
+		: Hash(Other.Hash), Data(MoveTemp(Other.Data))
+	{
+	}
+
+	FLevelInstancePackerClusterID(TUniquePtr<FLevelInstancePackerCluster>&& InData)
+		: Data(MoveTemp(InData))
+	{
+		Hash = Data->ComputeHash();
+	}
+		
+	bool operator==(const FLevelInstancePackerClusterID& Other) const
+	{
+		return (Hash == Other.Hash) && (*Data == *Other.Data);
+	}
+
+	bool operator!=(const FLevelInstancePackerClusterID& Other) const
+	{
+		return !(*this == Other);
+	}
+
+	uint32 GetHash() const { return Hash; }
+
+	FLevelInstancePackerID GetPackerID() const { return Data->GetPackerID(); }
+		
+	friend uint32 GetTypeHash(const FLevelInstancePackerClusterID& ID)
+	{
+		return ID.GetHash();
+	}
+
+	FLevelInstancePackerCluster* GetData() const
+	{
+		return Data.Get();
+	}
+
+private:
+	FLevelInstancePackerClusterID(const FLevelInstancePackerClusterID&) = delete;
+	FLevelInstancePackerClusterID& operator=(const FLevelInstancePackerClusterID&) = delete;
+
+	uint32 Hash;
+	TUniquePtr<FLevelInstancePackerCluster> Data;
+};
+
+#endif
\ No newline at end of file
