Fixed NonUnity compile

diff --git a/Engine/Source/Runtime/Engine/Public/WorldPartition/DataLayer/DataLayerEditorPerProjectUserSettings.h b/Engine/Source/Runtime/Engine/Public/WorldPartition/DataLayer/DataLayerEditorPerProjectUserSettings.h
--- a/Engine/Source/Runtime/Engine/Public/WorldPartition/DataLayer/DataLayerEditorPerProjectUserSettings.h
+++ b/Engine/Source/Runtime/Engine/Public/WorldPartition/DataLayer/DataLayerEditorPerProjectUserSettings.h
@@ -5,6 +5,8 @@
 #include "CoreMinimal.h"
 #include "UObject/ObjectMacros.h"
 #include "UObject/Object.h"
+#include "UObject/SoftObjectPtr.h"
+#include "Engine/World.h"
 #include "DataLayerEditorPerProjectUserSettings.generated.h"
 
 USTRUCT()
