Proper fix for check that was failing in Burley SSS due to roundoff error.

diff --git a/Engine/Source/Runtime/Engine/Private/Rendering/BurleyNormalizedSSS.cpp b/Engine/Source/Runtime/Engine/Private/Rendering/BurleyNormalizedSSS.cpp
--- a/Engine/Source/Runtime/Engine/Private/Rendering/BurleyNormalizedSSS.cpp
+++ b/Engine/Source/Runtime/Engine/Private/Rendering/BurleyNormalizedSSS.cpp
@@ -138,6 +138,8 @@ void ComputeMirroredBSSSKernel(FLinearColor* TargetBuffer, uint32 TargetBufferSi
 			float sign = o < 0.0f ? -1.0f : 1.0f;
 			kernel[i].A = Range * sign * FMath::Abs(FMath::Pow(o, Exponent)) / FMath::Pow(Range, Exponent);
 		}
+		// Center sample should always be zero, but might not be due to potential roundoff error.
+		kernel[nTotalSamples / 2].A = 0.0f;
 
 		//Scale the profile sampling radius. This scale enables the sampling between [-3*SpaceScale,+3*SpaceScale] instead of 
 		//the default [-3,3] range when fetching kernel parameters.
@@ -208,9 +210,6 @@ void ComputeMirroredBSSSKernel(FLinearColor* TargetBuffer, uint32 TargetBufferSi
 
 	// generate output (remove negative samples)
 	{
-		kernel[0].A = 0.0f;
-		//check(kernel[0].A == 0.0f);
-
 		// center sample
 		TargetBuffer[0] = kernel[0];
 
