Fixing missing include in DistanceFieldDownsampling.h.

diff --git a/Engine/Source/Runtime/Engine/Public/DistanceFieldDownsampling.h b/Engine/Source/Runtime/Engine/Public/DistanceFieldDownsampling.h
--- a/Engine/Source/Runtime/Engine/Public/DistanceFieldDownsampling.h
+++ b/Engine/Source/Runtime/Engine/Public/DistanceFieldDownsampling.h
@@ -9,6 +9,7 @@
 #include "RHI.h"
 
 class FRHICommandListImmediate;
+class FRDGBuilder;
 
 struct FDistanceFieldDownsamplingDataTask
 {
