Skip per instance post load on HISM to avoid reentrance since HISM is fully rebuilt after compilation
  - Reentrance path (OnPostLoadPerInstanceData - > InitPerInstanceRenderData -> GetRenderData -> FinishCompilation -> RebuildTreeIfOutdated -> InitPerInstanceRenderData)

diff --git a/Engine/Source/Runtime/Engine/Private/HierarchicalInstancedStaticMesh.cpp b/Engine/Source/Runtime/Engine/Private/HierarchicalInstancedStaticMesh.cpp
--- a/Engine/Source/Runtime/Engine/Private/HierarchicalInstancedStaticMesh.cpp
+++ b/Engine/Source/Runtime/Engine/Private/HierarchicalInstancedStaticMesh.cpp
@@ -3081,7 +3081,8 @@ void UHierarchicalInstancedStaticMeshComponent::OnPostLoadPerInstanceData()
 {
 	SCOPE_CYCLE_COUNTER(STAT_FoliagePostLoad);
 
-	if (!HasAnyFlags(RF_ClassDefaultObject | RF_ArchetypeObject))
+	// Tree will be fully rebuilt when staticmesh has finished compiling
+	if (!HasAnyFlags(RF_ClassDefaultObject | RF_ArchetypeObject) && !IsCompiling())
 	{
 		bool bForceTreeBuild = false;
 
