World Partition HLODs
* Only fallback to the default HLODLayer for the first level of HLOD
* Fix HLOD grid naming that could include dots
* Use FindObjectFast() to look for existing an existing HLOD
* Removed RequiresMapSaving() override for the HLOD builder, not needed anymore

diff --git a/Engine/Source/Editor/UnrealEd/Public/WorldPartition/WorldPartitionHLODsBuilder.h b/Engine/Source/Editor/UnrealEd/Public/WorldPartition/WorldPartitionHLODsBuilder.h
--- a/Engine/Source/Editor/UnrealEd/Public/WorldPartition/WorldPartitionHLODsBuilder.h
+++ b/Engine/Source/Editor/UnrealEd/Public/WorldPartition/WorldPartitionHLODsBuilder.h
@@ -13,7 +13,6 @@ class UWorldPartitionHLODsBuilder : public UWorldPartitionBuilder
 public:
 	// UWorldPartitionBuilder interface begin
 	virtual bool RequiresCommandletRendering() const override;
-	virtual bool RequiresMapSaving() const override { return true; }
 	virtual bool Run(UWorld* World, FPackageSourceControlHelper& PackageHelper) override;
 	// UWorldPartitionBuilder interface end
 };
\ No newline at end of file
diff --git a/Engine/Source/Runtime/Engine/Private/WorldPartition/HLOD/HLODBuilder.cpp b/Engine/Source/Runtime/Engine/Private/WorldPartition/HLOD/HLODBuilder.cpp
--- a/Engine/Source/Runtime/Engine/Private/WorldPartition/HLOD/HLODBuilder.cpp
+++ b/Engine/Source/Runtime/Engine/Private/WorldPartition/HLOD/HLODBuilder.cpp
@@ -34,9 +34,7 @@ public:
 	void SpawnHLODActor(const TCHAR* InName, const TArray<UPrimitiveComponent*>& InSubComponents, FCreateComponentsFunction InCreateComponentsFunc)
 	{
 		FString HLODActorName = FString::Printf(TEXT("%s_%s_%s"), *HLODLayer->GetName(), *CellName.ToString(), InName);
-		FString HLODActorPath = FString::Printf(TEXT("%s.%s"), *World->PersistentLevel->GetPathName(), *HLODActorName);
-
-		AWorldPartitionHLOD* HLODActor = FindObject<AWorldPartitionHLOD>(ANY_PACKAGE, *HLODActorPath);
+		AWorldPartitionHLOD* HLODActor = FindObjectFast<AWorldPartitionHLOD>(World->PersistentLevel, *HLODActorName);
 		if (!HLODActor)
 		{
 			FActorSpawnParameters SpawnParams;
diff --git a/Engine/Source/Runtime/Engine/Private/WorldPartition/HLOD/HLODLayer.cpp b/Engine/Source/Runtime/Engine/Private/WorldPartition/HLOD/HLODLayer.cpp
--- a/Engine/Source/Runtime/Engine/Private/WorldPartition/HLOD/HLODLayer.cpp
+++ b/Engine/Source/Runtime/Engine/Private/WorldPartition/HLOD/HLODLayer.cpp
@@ -10,7 +10,6 @@
 #if WITH_EDITOR
 #include "UObject/ConstructorHelpers.h"
 #include "Materials/MaterialInterface.h"
-#include "Math/UnitConversion.h"
 
 #include "WorldPartition/WorldPartition.h"
 #include "WorldPartition/WorldPartitionActorDesc.h"
@@ -90,9 +89,15 @@ UHLODLayer* UHLODLayer::GetHLODLayer(const AActor* InActor)
 	{
 		return HLODLayer;
 	}
-	if (UWorldPartition* WorldPartition = InActor->GetWorld()->GetWorldPartition())
+
+	// Only fallback to the default HLODLayer for the first level of HLOD
+	bool bIsHLOD0 = !InActor->IsA<AWorldPartitionHLOD>();
+	if (!bIsHLOD0) 
 	{
-		return WorldPartition->DefaultHLODLayer;
+		if (UWorldPartition* WorldPartition = InActor->GetWorld()->GetWorldPartition())
+		{
+			return WorldPartition->DefaultHLODLayer;
+		}
 	}
 	return nullptr;
 }
@@ -103,15 +108,7 @@ UHLODLayer* UHLODLayer::GetHLODLayer(const AActor* InActor)
 
 FName UHLODLayer::GetRuntimeGridName(uint32 InLODLevel, int32 InCellSize, float InLoadingRange)
 {
-	FNumericUnit<float> CellSize(InCellSize, EUnit::Centimeters);
-	FNumericUnit<float> CellSizeQuantized = CellSize.QuantizeUnitsToBestFit();
-	FString CellSizeString = FString::SanitizeFloat(CellSizeQuantized.Value, 0) + FUnitConversion::GetUnitDisplayString(CellSizeQuantized.Units);
-
-	FNumericUnit<float> LoadingRange(InLoadingRange, EUnit::Centimeters);
-	FNumericUnit<float> LoadingRangeQuantized = LoadingRange.QuantizeUnitsToBestFit();
-	FString LoadingRangeString = FString::SanitizeFloat(LoadingRangeQuantized.Value, 0) + FUnitConversion::GetUnitDisplayString(LoadingRangeQuantized.Units);
-
-	return *FString::Format(TEXT("HLOD{0}_{1}_{2}"), { InLODLevel, CellSizeString, LoadingRangeString });
+	return *FString::Format(TEXT("HLOD{0}_{1}m_{2}m"), { InLODLevel, int32(InCellSize * 0.01f), int32(InLoadingRange * 0.01f)});
 }
 
 FName UHLODLayer::GetRuntimeGrid(uint32 InHLODLevel) const
