Removed superfluous parent component assignment since it is now always set prior to components registration after 14586066.

diff --git a/Engine/Source/Runtime/Engine/Private/Components/ChildActorComponent.cpp b/Engine/Source/Runtime/Engine/Private/Components/ChildActorComponent.cpp
--- a/Engine/Source/Runtime/Engine/Private/Components/ChildActorComponent.cpp
+++ b/Engine/Source/Runtime/Engine/Private/Components/ChildActorComponent.cpp
@@ -649,9 +649,6 @@ void UChildActorComponent::CreateChildActor()
 				{
 					ChildActorName = ChildActor->GetFName();
 
-					// Remember which component spawned it (for selection in editor etc)
-					FActorParentComponentSetter::Set(ChildActor, this);
-
 					// Parts that we deferred from SpawnActor
 					const FComponentInstanceDataCache* ComponentInstanceData = (CachedInstanceData ? CachedInstanceData->ComponentInstanceData.Get() : nullptr);
 					ChildActor->FinishSpawning(GetComponentTransform(), false, ComponentInstanceData);
