Made StrataMultiply only take a float as input to be conservative on the BSDF size later.

diff --git a/Engine/Shaders/Private/Strata/Strata.ush b/Engine/Shaders/Private/Strata/Strata.ush
--- a/Engine/Shaders/Private/Strata/Strata.ush
+++ b/Engine/Shaders/Private/Strata/Strata.ush
@@ -704,7 +704,7 @@ FStrataData StrataAdd(FStrataData A, FStrataData B)
 	return StrataData;
 }
 
-FStrataData StrataMultiply(FStrataData A, float3 Weight)
+FStrataData StrataMultiply(FStrataData A, float Weight)
 {
 	FStrataData StrataData = A;
 
diff --git a/Engine/Source/Runtime/Engine/Private/Materials/MaterialExpressions.cpp b/Engine/Source/Runtime/Engine/Private/Materials/MaterialExpressions.cpp
--- a/Engine/Source/Runtime/Engine/Private/Materials/MaterialExpressions.cpp
+++ b/Engine/Source/Runtime/Engine/Private/Materials/MaterialExpressions.cpp
@@ -19529,7 +19529,7 @@ uint32 UMaterialExpressionStrataMultiply::GetInputType(int32 InputIndex)
 	{
 		return MCT_Strata;
 	}
-	return MCT_Float3;
+	return MCT_Float;
 }
 #endif // WITH_EDITOR
 
