Blueprints: Create a MakeLiteral function for Int64 types

#rnx

diff --git a/Engine/Source/Runtime/Engine/Classes/Kismet/KismetSystemLibrary.h b/Engine/Source/Runtime/Engine/Classes/Kismet/KismetSystemLibrary.h
--- a/Engine/Source/Runtime/Engine/Classes/Kismet/KismetSystemLibrary.h
+++ b/Engine/Source/Runtime/Engine/Classes/Kismet/KismetSystemLibrary.h
@@ -315,6 +315,14 @@ class ENGINE_API UKismetSystemLibrary : public UBlueprintFunctionLibrary
 	UFUNCTION(BlueprintPure, Category="Math|Integer", meta=(BlueprintThreadSafe))
 	static int32 MakeLiteralInt(int32 Value);
 
+	/**
+	 * Creates a literal 64-bit integer
+	 * @param	Value	value to set the 64-bit integer to
+	 * @return	The literal 64-bit integer
+	 */
+	UFUNCTION(BlueprintPure, Category = "Math|Integer", meta = (BlueprintThreadSafe))
+	static int64 MakeLiteralInt64(int64 Value);
+
 	/**
 	 * Creates a literal float
 	 * @param	Value	value to set the float to
@@ -2008,6 +2016,11 @@ FORCEINLINE int32 UKismetSystemLibrary::MakeLiteralInt(int32 Value)
 	return Value;
 }
 
+FORCEINLINE int64 UKismetSystemLibrary::MakeLiteralInt64(int64 Value)
+{
+	return Value;
+}
+
 FORCEINLINE float UKismetSystemLibrary::MakeLiteralFloat(float Value)
 {
 	return Value;
