Blueprints: Add an int64 ToString function

#rnx

diff --git a/Engine/Source/Runtime/Engine/Classes/Kismet/KismetStringLibrary.h b/Engine/Source/Runtime/Engine/Classes/Kismet/KismetStringLibrary.h
--- a/Engine/Source/Runtime/Engine/Classes/Kismet/KismetStringLibrary.h
+++ b/Engine/Source/Runtime/Engine/Classes/Kismet/KismetStringLibrary.h
@@ -24,6 +24,10 @@ class ENGINE_API UKismetStringLibrary : public UBlueprintFunctionLibrary
 	UFUNCTION(BlueprintPure, meta=(DisplayName = "ToString (integer)", CompactNodeTitle = "->", BlueprintAutocast), Category="Utilities|String")
 	static FString Conv_IntToString(int32 InInt);
 
+	/** Converts an 64-bit integer value to a string */
+	UFUNCTION(BlueprintPure, meta = (DisplayName = "ToString (Int64)", CompactNodeTitle = "->", BlueprintAutocast), Category = "Utilities|String")
+	static FString Conv_Int64ToString(int64 InInt);
+	
 	/** Converts a byte value to a string */
 	UFUNCTION(BlueprintPure, meta=(DisplayName = "ToString (byte)", CompactNodeTitle = "->", BlueprintAutocast), Category="Utilities|String")
 	static FString Conv_ByteToString(uint8 InByte);
diff --git a/Engine/Source/Runtime/Engine/Private/KismetStringLibrary.cpp b/Engine/Source/Runtime/Engine/Private/KismetStringLibrary.cpp
--- a/Engine/Source/Runtime/Engine/Private/KismetStringLibrary.cpp
+++ b/Engine/Source/Runtime/Engine/Private/KismetStringLibrary.cpp
@@ -66,6 +66,11 @@ FString UKismetStringLibrary::Conv_IntToString(int32 InInt)
 	return FString::Printf(TEXT("%d"), InInt);	
 }
 
+FString UKismetStringLibrary::Conv_Int64ToString(int64 InInt)
+{
+	return FString::Printf(TEXT("%ld"), InInt);
+}
+
 FString UKismetStringLibrary::Conv_ByteToString(uint8 InByte)
 {
 	return FString::Printf(TEXT("%d"), InByte);	
