#characterMovement - Allow TickCharacterPose to be overridden.

diff --git a/Engine/Source/Runtime/Engine/Classes/GameFramework/CharacterMovementComponent.h b/Engine/Source/Runtime/Engine/Classes/GameFramework/CharacterMovementComponent.h
--- a/Engine/Source/Runtime/Engine/Classes/GameFramework/CharacterMovementComponent.h
+++ b/Engine/Source/Runtime/Engine/Classes/GameFramework/CharacterMovementComponent.h
@@ -2231,7 +2231,7 @@ protected:
 	virtual float GetClientNetSendDeltaTime(const APlayerController* PC, const FNetworkPredictionData_Client_Character* ClientData, const FSavedMovePtr& NewMove) const;
 
 	/** Ticks the characters pose and accumulates root motion */
-	void TickCharacterPose(float DeltaTime);
+	virtual void TickCharacterPose(float DeltaTime);
 
 	/** On the server if we know we are having our replication rate throttled, this method checks if important replicated properties have changed that should cause us to return to the normal replication rate. */
 	virtual bool ShouldCancelAdaptiveReplication() const;
