Fix "Preparing StaticMesh" notification not disappearing when weakptr on staticmesh to process have been invalidated

diff --git a/Engine/Source/Runtime/Engine/Private/StaticMeshCompiler.cpp b/Engine/Source/Runtime/Engine/Private/StaticMeshCompiler.cpp
--- a/Engine/Source/Runtime/Engine/Private/StaticMeshCompiler.cpp
+++ b/Engine/Source/Runtime/Engine/Private/StaticMeshCompiler.cpp
@@ -735,15 +735,15 @@ void FStaticMeshCompilingManager::ProcessStaticMeshes(bool bLimitExecutionTime,
 			}
 		}
 
-		TSet<UStaticMesh*> ProcessedStaticMeshes;
 		{
 			TRACE_CPUPROFILER_EVENT_SCOPE(ProcessFinishedStaticMeshes);
 
 			const double TickStartTime = FPlatformTime::Seconds();
 
+			TSet<TWeakObjectPtr<UStaticMesh>> StaticMeshesToPostpone;
+			TSet<UStaticMesh*> ProcessedStaticMeshes;
 			if (StaticMeshesToProcess.Num())
 			{
-				TSet<TWeakObjectPtr<UStaticMesh>> StaticMeshesToPostpone;
 				for (UStaticMesh* StaticMesh : StaticMeshesToProcess)
 				{
 					const bool bHasMeshUpdateLeft = ProcessedStaticMeshes.Num() <= MaxMeshUpdatesPerFrame;
@@ -757,10 +757,10 @@ void FStaticMeshCompilingManager::ProcessStaticMeshes(bool bLimitExecutionTime,
 						StaticMeshesToPostpone.Emplace(StaticMesh);
 					}
 				}
-
-				RegisteredStaticMesh = MoveTemp(StaticMeshesToPostpone);
 			}
 
+			RegisteredStaticMesh = MoveTemp(StaticMeshesToPostpone);
+
 			PostStaticMeshesCompilation(ProcessedStaticMeshes);
 		}
 	}
