Fix non unity compile error.

diff --git a/Engine/Source/Runtime/Engine/Private/ComputeFramework/ComputeGraphComponent.cpp b/Engine/Source/Runtime/Engine/Private/ComputeFramework/ComputeGraphComponent.cpp
--- a/Engine/Source/Runtime/Engine/Private/ComputeFramework/ComputeGraphComponent.cpp
+++ b/Engine/Source/Runtime/Engine/Private/ComputeFramework/ComputeGraphComponent.cpp
@@ -3,6 +3,7 @@
 #include "ComputeFramework/ComputeGraphComponent.h"
 #include "ComputeFramework/ComputeGraph.h"
 #include "ComputeFramework/ComputeFramework.h"
+#include "SceneInterface.h"
 
 
 UComputeGraphComponent::UComputeGraphComponent()
