Blueprints: Add a byte to double conversion function and a Int to FVector conversion function

#rnx

diff --git a/Engine/Source/Runtime/Engine/Classes/Kismet/KismetMathLibrary.h b/Engine/Source/Runtime/Engine/Classes/Kismet/KismetMathLibrary.h
--- a/Engine/Source/Runtime/Engine/Classes/Kismet/KismetMathLibrary.h
+++ b/Engine/Source/Runtime/Engine/Classes/Kismet/KismetMathLibrary.h
@@ -3427,6 +3427,10 @@ class ENGINE_API UKismetMathLibrary : public UBlueprintFunctionLibrary
 	UFUNCTION(BlueprintPure, meta=(DisplayName = "ToFloat (byte)", CompactNodeTitle = "->", Keywords="cast convert", BlueprintAutocast), Category="Math|Conversions")
 	static float Conv_ByteToFloat(uint8 InByte);
 
+	/** Converts a byte to a double */
+	UFUNCTION(BlueprintPure, meta = (DisplayName = "ToDouble (byte)", CompactNodeTitle = "->", Keywords = "cast convert", BlueprintAutocast), Category = "Math|Conversions")
+	static double Conv_ByteToDouble(uint8 InByte);
+
 	/** Converts an integer to a float */
 	UFUNCTION(BlueprintPure, meta=(DisplayName = "ToFloat (integer)", CompactNodeTitle = "->", Keywords="cast convert", BlueprintAutocast), Category="Math|Conversions")
 	static float Conv_IntToFloat(int32 InInt);
@@ -3459,10 +3463,22 @@ class ENGINE_API UKismetMathLibrary : public UBlueprintFunctionLibrary
 	UFUNCTION(BlueprintPure, meta=(DisplayName = "ToByte (Int64)", CompactNodeTitle = "->", Keywords="cast convert", BlueprintAutocast), Category="Math|Conversions")
 	static uint8 Conv_Int64ToByte(int64 InInt);
 
+	/** Converts a 64 bit floating point to a 64 bit integer */
+	UFUNCTION(BlueprintPure, meta = (DisplayName = "ToInt64 (Double)", CompactNodeTitle = "->", Keywords = "cast convert", BlueprintAutocast), Category = "Math|Conversions")
+	static int64 Conv_DoubleToInt64(double InDouble);
+
+	/** Converts a 64 bit integer to a 64 bit floating point */
+	UFUNCTION(BlueprintPure, meta = (DisplayName = "ToDouble (Int64)", CompactNodeTitle = "->", Keywords = "cast convert", BlueprintAutocast), Category = "Math|Conversions")
+	static double Conv_Int64ToDouble(int64 InInt);
+
 	/** Converts an integer to an IntVector*/
 	UFUNCTION(BlueprintPure, meta = (DisplayName = "ToIntVector (integer)", CompactNodeTitle = "->", Keywords = "cast convert", BlueprintAutocast), Category = "Math|Conversions")
 	static FIntVector Conv_IntToIntVector(int32 InInt);
 
+	/** Converts an integer to a FVector */
+	UFUNCTION(BlueprintPure, meta = (DisplayName = "ToVector (integer)", CompactNodeTitle = "->", Keywords = "cast convert", BlueprintAutocast), Category = "Math|Conversions")
+	static FVector Conv_IntToVector(int32 InInt);
+
 	/** Converts a int to a bool*/
 	UFUNCTION(BlueprintPure, meta=(DisplayName = "ToBool (integer)", CompactNodeTitle = "->", Keywords="cast convert", BlueprintAutocast), Category="Math|Conversions")
 	static bool Conv_IntToBool(int32 InInt);
@@ -3483,6 +3499,10 @@ class ENGINE_API UKismetMathLibrary : public UBlueprintFunctionLibrary
 	UFUNCTION(BlueprintPure, meta=(DisplayName = "ToInt (byte)", CompactNodeTitle = "->", Keywords="cast convert", BlueprintAutocast), Category="Math|Conversions")
 	static int32 Conv_ByteToInt(uint8 InByte);
 
+	/** Converts a byte to an integer */
+	UFUNCTION(BlueprintPure, meta = (DisplayName = "ToInt64 (byte)", CompactNodeTitle = "->", Keywords = "cast convert", BlueprintAutocast), Category = "Math|Conversions")
+	static int64 Conv_ByteToInt64(uint8 InByte);
+
 	/** Converts a color to LinearColor */
 	UFUNCTION(BlueprintPure, meta = (DisplayName = "ToLinearColor (Color)", CompactNodeTitle = "->", Keywords = "cast convert", BlueprintAutocast), Category = "Math|Conversions")
 	static FLinearColor Conv_ColorToLinearColor(FColor InColor);
diff --git a/Engine/Source/Runtime/Engine/Classes/Kismet/KismetMathLibrary.inl b/Engine/Source/Runtime/Engine/Classes/Kismet/KismetMathLibrary.inl
--- a/Engine/Source/Runtime/Engine/Classes/Kismet/KismetMathLibrary.inl
+++ b/Engine/Source/Runtime/Engine/Classes/Kismet/KismetMathLibrary.inl
@@ -2643,6 +2643,12 @@ float UKismetMathLibrary::Conv_ByteToFloat(uint8 InByte)
 	return (float)InByte;
 }
 
+KISMET_MATH_FORCEINLINE
+double UKismetMathLibrary::Conv_ByteToDouble(uint8 InByte)
+{
+	return (double)InByte;
+}
+
 KISMET_MATH_FORCEINLINE
 float UKismetMathLibrary::Conv_IntToFloat(int32 InInt)
 {
@@ -2691,12 +2697,30 @@ uint8 UKismetMathLibrary::Conv_Int64ToByte(int64 InInt)
 	return (uint8)InInt;
 }
 
+KISMET_MATH_FORCEINLINE
+int64 UKismetMathLibrary::Conv_DoubleToInt64(double InDouble)
+{
+	return (int64)InDouble;
+}
+
+KISMET_MATH_FORCEINLINE
+double UKismetMathLibrary::Conv_Int64ToDouble(int64 InInt)
+{
+	return (double)InInt;
+}
+
 KISMET_MATH_FORCEINLINE
 FIntVector UKismetMathLibrary::Conv_IntToIntVector(int32 InInt)
 {
 	return FIntVector(InInt, InInt, InInt);
 }
 
+KISMET_MATH_FORCEINLINE
+FVector UKismetMathLibrary::Conv_IntToVector(int32 InInt)
+{
+	return FVector((float)InInt);
+}
+
 KISMET_MATH_FORCEINLINE
 bool UKismetMathLibrary::Conv_IntToBool(int32 InInt)
 {
@@ -2727,6 +2751,12 @@ int32 UKismetMathLibrary::Conv_ByteToInt(uint8 InByte)
 	return (int32)InByte;
 }
 
+KISMET_MATH_FORCEINLINE
+int64 UKismetMathLibrary::Conv_ByteToInt64(uint8 InByte)
+{
+	return (int64)InByte;
+}
+
 KISMET_MATH_FORCEINLINE
 FTransform UKismetMathLibrary::Conv_RotatorToTransform(const FRotator& InRotator)
 {
