Add missing texture bulkdata locks.

diff --git a/Engine/Source/Runtime/Engine/Private/Texture.cpp b/Engine/Source/Runtime/Engine/Private/Texture.cpp
--- a/Engine/Source/Runtime/Engine/Private/Texture.cpp
+++ b/Engine/Source/Runtime/Engine/Private/Texture.cpp
@@ -436,6 +436,9 @@ void UTexture::Serialize(FArchive& Ar)
 #if WITH_EDITORONLY_DATA
 	if (!StripFlags.IsEditorDataStripped())
 	{
+#if WITH_EDITOR
+		FWriteScopeLock BulkDataExclusiveScope(Source.BulkDataLock.Get());
+#endif
 		Source.BulkData.Serialize(Ar, this);
 	}
 
@@ -1075,6 +1078,9 @@ void FTextureSource::Compress()
 {
 	if (CanPNGCompress())
 	{
+#if WITH_EDITOR
+		FWriteScopeLock BulkDataExclusiveScope(BulkDataLock.Get());
+#endif
 		uint8* BulkDataPtr = (uint8*)BulkData.Lock(LOCK_READ_WRITE);
 		IImageWrapperModule& ImageWrapperModule = FModuleManager::LoadModuleChecked<IImageWrapperModule>( FName("ImageWrapper") );
 		TSharedPtr<IImageWrapper> ImageWrapper = ImageWrapperModule.CreateImageWrapper( EImageFormat::PNG );
