Add cancellation support for DistanceField and MeshCards async tasks
Cancel async operations on StaticMesh when they are unloaded during world change
Avoid potential crash when closing the editor by properly cancelling async tasks
Use cancellation instead of BlockUntilComplete whenever possible to shorten wait time

diff --git a/Engine/Source/Editor/UnrealEd/Private/EditorServer.cpp b/Engine/Source/Editor/UnrealEd/Private/EditorServer.cpp
--- a/Engine/Source/Editor/UnrealEd/Private/EditorServer.cpp
+++ b/Engine/Source/Editor/UnrealEd/Private/EditorServer.cpp
@@ -138,7 +138,8 @@
 #include "Misc/UObjectToken.h"
 #include "Misc/MapErrors.h"
 #include "Misc/ScopedSlowTask.h"
-
+#include "DistanceFieldAtlas.h"
+#include "MeshCardRepresentation.h"
 
 #include "ComponentReregisterContext.h"
 #include "Engine/DocumentationActor.h"
@@ -2125,6 +2126,32 @@ void UEditorEngine::EditorDestroyWorld( FWorldContext & Context, const FText& Cl
 		}
 	}
 
+	// Prevent the GC from not being able to garbage collect the packages we're tying to unload due to async tasks
+	// for static mesh that are embedded in the world.
+	if (GDistanceFieldAsyncQueue || GCardRepresentationAsyncQueue)
+	{
+		UPackage* NewWorldPackage = NewWorld ? NewWorld->GetOutermost() : nullptr;
+		if (NewWorldPackage && WorldPackage != NewWorldPackage)
+		{
+			for (TObjectIterator<UStaticMesh> It; It; ++It)
+			{
+				UStaticMesh* StaticMesh = *It;
+				if (WorldPackage == StaticMesh->GetPackage())
+				{
+					if (GDistanceFieldAsyncQueue)
+					{
+						GDistanceFieldAsyncQueue->CancelBuild(StaticMesh);
+					}
+
+					if (GCardRepresentationAsyncQueue)
+					{
+						GCardRepresentationAsyncQueue->CancelBuild(StaticMesh);
+					}
+				}
+			}
+		}
+	}
+
 	ContextWorld->DestroyWorld( true, NewWorld );
 	Context.SetCurrentWorld(NULL);
 
diff --git a/Engine/Source/Editor/UnrealEd/Private/ObjectTools.cpp b/Engine/Source/Editor/UnrealEd/Private/ObjectTools.cpp
--- a/Engine/Source/Editor/UnrealEd/Private/ObjectTools.cpp
+++ b/Engine/Source/Editor/UnrealEd/Private/ObjectTools.cpp
@@ -2604,8 +2604,15 @@ namespace ObjectTools
 			UStaticMesh* StaticMesh = Cast<UStaticMesh>(ObjectToDelete);
 			if (StaticMesh != nullptr)
 			{
-				GDistanceFieldAsyncQueue->BlockUntilBuildComplete(StaticMesh, true);
-				GCardRepresentationAsyncQueue->BlockUntilBuildComplete(StaticMesh, true);
+				if (GDistanceFieldAsyncQueue)
+				{
+					GDistanceFieldAsyncQueue->CancelBuild(StaticMesh);
+				}
+				
+				if (GCardRepresentationAsyncQueue)
+				{
+					GCardRepresentationAsyncQueue->CancelBuild(StaticMesh);
+				}
 			}
 		}
 
diff --git a/Engine/Source/Runtime/Engine/Private/DistanceFieldAtlas.cpp b/Engine/Source/Runtime/Engine/Private/DistanceFieldAtlas.cpp
--- a/Engine/Source/Runtime/Engine/Private/DistanceFieldAtlas.cpp
+++ b/Engine/Source/Runtime/Engine/Private/DistanceFieldAtlas.cpp
@@ -1016,6 +1016,42 @@ FDistanceFieldAsyncQueue::~FDistanceFieldAsyncQueue()
 {
 }
 
+void FAsyncDistanceFieldTaskWorker::DoWork()
+{
+	// Put on background thread to avoid interfering with game-thread bound tasks
+	FQueuedThreadPoolTaskGraphWrapper TaskGraphWrapper(ENamedThreads::AnyBackgroundThreadNormalTask);
+	GDistanceFieldAsyncQueue->Build(&Task, TaskGraphWrapper);
+}
+
+void FDistanceFieldAsyncQueue::CancelBackgroundTask(TArray<FAsyncDistanceFieldTask*> Tasks)
+{
+	// Do all the cancellation first to make sure none of these tasks
+	// get scheduled as we're waiting for completion.
+	for (FAsyncDistanceFieldTask* Task : Tasks)
+	{
+		if (Task->AsyncTask)
+		{
+			Task->AsyncTask->Cancel();
+		}
+	}
+
+	for (FAsyncDistanceFieldTask* Task : Tasks)
+	{
+		if (Task->AsyncTask)
+		{
+			Task->AsyncTask->EnsureCompletion();
+			Task->AsyncTask.Reset();
+		}
+	}
+}
+
+void FDistanceFieldAsyncQueue::StartBackgroundTask(FAsyncDistanceFieldTask* Task)
+{
+	check(Task->AsyncTask == nullptr);
+	Task->AsyncTask = MakeUnique<FAsyncTask<FAsyncDistanceFieldTaskWorker>>(*Task);
+	Task->AsyncTask->StartBackgroundTask(ThreadPool.Get(), EQueuedWorkPriority::Lowest);
+}
+
 void FDistanceFieldAsyncQueue::ProcessPendingTasks()
 {
 	FScopeLock Lock(&CriticalSection);
@@ -1028,16 +1064,7 @@ void FDistanceFieldAsyncQueue::ProcessPendingTasks()
 		}
 		else
 		{
-			// To avoid deadlocks, we must queue the inner build tasks on another thread pool, so use the task graph.
-			AsyncPool(
-				*ThreadPool,
-				[this, Task]()
-				{
-					// Put on background thread to avoid interfering with game-thread bound tasks
-					FQueuedThreadPoolTaskGraphWrapper TaskGraphWrapper(ENamedThreads::AnyBackgroundThreadNormalTask);
-					Build(Task, TaskGraphWrapper);
-				}
-			);
+			StartBackgroundTask(Task);
 		}
 	}
 }
@@ -1065,23 +1092,17 @@ void FDistanceFieldAsyncQueue::AddTask(FAsyncDistanceFieldTask* Task)
 	}
 	else
 	{
-		auto BuildLambda = 
-			[this, Task]()
-			{
-				// To avoid deadlocks, we must queue the inner build tasks on another thread pool, so use the task graph.
-				// Put on background thread to avoid interfering with game-thread bound tasks
-				FQueuedThreadPoolTaskGraphWrapper TaskGraphWrapper(ENamedThreads::AnyBackgroundThreadNormalTask); 
-				Build(Task, TaskGraphWrapper);
-			};
-
 		// If we're already in worker threads, there is no need to launch an async task.
 		if (GUseAsyncDistanceFieldBuildQueue || !IsInGameThread())
 		{
-			AsyncPool(*ThreadPool, MoveTemp(BuildLambda));
+			StartBackgroundTask(Task);
 		}
 		else
 		{
-			BuildLambda();
+			// To avoid deadlocks, we must queue the inner build tasks on another thread pool, so use the task graph.
+			// Put on background thread to avoid interfering with game-thread bound tasks
+			FQueuedThreadPoolTaskGraphWrapper TaskGraphWrapper(ENamedThreads::AnyBackgroundThreadNormalTask);
+			Build(Task, TaskGraphWrapper);
 		}
 	}
 #else
@@ -1089,6 +1110,74 @@ void FDistanceFieldAsyncQueue::AddTask(FAsyncDistanceFieldTask* Task)
 #endif
 }
 
+void FDistanceFieldAsyncQueue::CancelBuild(UStaticMesh* StaticMesh)
+{
+	TRACE_CPUPROFILER_EVENT_SCOPE(FDistanceFieldAsyncQueue::CancelBuild)
+
+	TArray<FAsyncDistanceFieldTask*> TasksToCancel;
+	{
+		FScopeLock Lock(&CriticalSection);
+		TArray<FAsyncDistanceFieldTask*> Tasks = MoveTemp(PendingTasks);
+		PendingTasks.Reserve(Tasks.Num());
+		for (FAsyncDistanceFieldTask* Task : Tasks)
+		{
+			if (Task->GenerateSource != StaticMesh && Task->StaticMesh != StaticMesh)
+			{
+				PendingTasks.Add(Task);
+			}
+		}
+
+		Tasks = MoveTemp(ReferencedTasks);
+		ReferencedTasks.Reserve(Tasks.Num());
+		for (FAsyncDistanceFieldTask* Task : Tasks)
+		{
+			if (Task->GenerateSource != StaticMesh && Task->StaticMesh != StaticMesh)
+			{
+				ReferencedTasks.Add(Task);
+			}
+			else
+			{
+				TasksToCancel.Add(Task);
+			}
+		}
+	}
+
+	CancelBackgroundTask(TasksToCancel);
+	for (FAsyncDistanceFieldTask* Task : TasksToCancel)
+	{
+		delete Task;
+	}
+}
+
+void FDistanceFieldAsyncQueue::CancelAllOutstandingBuilds()
+{
+	TRACE_CPUPROFILER_EVENT_SCOPE(FDistanceFieldAsyncQueue::CancelAllOutstandingBuilds)
+
+	TArray<FAsyncDistanceFieldTask*> OutstandingTasks;
+	{
+		FScopeLock Lock(&CriticalSection);
+		PendingTasks.Empty();
+		OutstandingTasks = MoveTemp(ReferencedTasks);
+	}
+
+	CancelBackgroundTask(OutstandingTasks);
+	for (FAsyncDistanceFieldTask* Task : OutstandingTasks)
+	{
+		delete Task;
+	}
+}
+
+void FDistanceFieldAsyncQueue::RescheduleBackgroundTask(FAsyncDistanceFieldTask* InTask, EQueuedWorkPriority InPriority)
+{
+	if (InTask->AsyncTask)
+	{
+		if (InTask->AsyncTask->GetPriority() != InPriority)
+		{
+			InTask->AsyncTask->Reschedule(GThreadPool, InPriority);
+		}
+	}
+}
+
 void FDistanceFieldAsyncQueue::BlockUntilBuildComplete(UStaticMesh* StaticMesh, bool bWarnIfBlocked)
 {
 	TRACE_CPUPROFILER_EVENT_SCOPE(FDistanceFieldAsyncQueue::BlockUntilBuildComplete)
@@ -1111,12 +1200,17 @@ void FDistanceFieldAsyncQueue::BlockUntilBuildComplete(UStaticMesh* StaticMesh,
 
 		bReferenced = false;
 
+		// Reschedule the tasks we're waiting on as highest prio
 		{
 			FScopeLock Lock(&CriticalSection);
 			for (int TaskIndex = 0; TaskIndex < ReferencedTasks.Num(); TaskIndex++)
 			{
-				bReferenced = bReferenced || ReferencedTasks[TaskIndex]->StaticMesh == StaticMesh;
-				bReferenced = bReferenced || ReferencedTasks[TaskIndex]->GenerateSource == StaticMesh;
+				if (ReferencedTasks[TaskIndex]->StaticMesh == StaticMesh ||
+					ReferencedTasks[TaskIndex]->GenerateSource == StaticMesh)
+				{
+					bReferenced = true;
+					RescheduleBackgroundTask(ReferencedTasks[TaskIndex], EQueuedWorkPriority::Highest);
+				}
 			}
 		}
 
@@ -1154,6 +1248,15 @@ void FDistanceFieldAsyncQueue::BlockUntilAllBuildsComplete()
 #if WITH_EDITOR
 		FStaticMeshCompilingManager::Get().FinishAllCompilation();
 #endif
+		// Reschedule as highest prio since we're explicitly waiting on them
+		{
+			FScopeLock Lock(&CriticalSection);
+			for (int TaskIndex = 0; TaskIndex < ReferencedTasks.Num(); TaskIndex++)
+			{
+				RescheduleBackgroundTask(ReferencedTasks[TaskIndex], EQueuedWorkPriority::Highest);
+			}
+		}
+
 		ProcessAsyncTasks();
 		FPlatformProcess::Sleep(.01f);
 	} 
@@ -1166,10 +1269,9 @@ void FDistanceFieldAsyncQueue::Build(FAsyncDistanceFieldTask* Task, FQueuedThrea
 	// Editor 'force delete' can null any UObject pointers which are seen by reference collecting (eg FProperty or serialized)
 	if (Task->StaticMesh && Task->GenerateSource)
 	{
-		TRACE_CPUPROFILER_EVENT_SCOPE(FDistanceFieldAsyncQueue::Build)
+		TRACE_CPUPROFILER_EVENT_SCOPE(FDistanceFieldAsyncQueue::Build);
 
 		const FStaticMeshLODResources& LODModel = Task->GenerateSource->GetRenderData()->LODResources[0];
-
 		MeshUtilities->GenerateSignedDistanceFieldVolumeData(
 			Task->StaticMesh->GetName(),
 			Task->SourceMeshData,
@@ -1183,7 +1285,6 @@ void FDistanceFieldAsyncQueue::Build(FAsyncDistanceFieldTask* Task, FQueuedThrea
 	}
 
     CompletedTasks.Push(Task);
-
 #endif
 }
 
@@ -1224,9 +1325,21 @@ void FDistanceFieldAsyncQueue::ProcessAsyncTasks(bool bLimitExecutionTime)
 		// We want to count each resource built from a DDC miss, so count each iteration of the loop separately.
 		COOK_STAT(auto Timer = DistanceFieldCookStats::UsageStats.TimeSyncWork());
 
+		bool bWasCancelled = false;
 		{
 			FScopeLock Lock(&CriticalSection);
-			ReferencedTasks.Remove(Task);
+			bWasCancelled = ReferencedTasks.Remove(Task) == 0;
+		}
+
+		if (bWasCancelled)
+		{
+			continue;
+		}
+
+		if (Task->AsyncTask)
+		{
+			Task->AsyncTask->EnsureCompletion();
+			Task->AsyncTask.Reset();
 		}
 
 		// Editor 'force delete' can null any UObject pointers which are seen by reference collecting (eg FProperty or serialized)
@@ -1275,6 +1388,8 @@ void FDistanceFieldAsyncQueue::ProcessAsyncTasks(bool bLimitExecutionTime)
 
 void FDistanceFieldAsyncQueue::Shutdown()
 {
+	CancelAllOutstandingBuilds();
+
 	UE_LOG(LogStaticMesh, Log, TEXT("Abandoning remaining async distance field tasks for shutdown"));
 	ThreadPool->Destroy();
 }
diff --git a/Engine/Source/Runtime/Engine/Private/MeshCardRepresentation.cpp b/Engine/Source/Runtime/Engine/Private/MeshCardRepresentation.cpp
--- a/Engine/Source/Runtime/Engine/Private/MeshCardRepresentation.cpp
+++ b/Engine/Source/Runtime/Engine/Private/MeshCardRepresentation.cpp
@@ -171,6 +171,42 @@ FCardRepresentationAsyncQueue::~FCardRepresentationAsyncQueue()
 {
 }
 
+void FAsyncCardRepresentationTaskWorker::DoWork()
+{
+	// Put on background thread to avoid interfering with game-thread bound tasks
+	FQueuedThreadPoolTaskGraphWrapper TaskGraphWrapper(ENamedThreads::AnyBackgroundThreadNormalTask);
+	GCardRepresentationAsyncQueue->Build(&Task, TaskGraphWrapper);
+}
+
+void FCardRepresentationAsyncQueue::CancelBackgroundTask(TArray<FAsyncCardRepresentationTask*> Tasks)
+{
+	// Do all the cancellation first to make sure none of these tasks
+	// get scheduled as we're waiting for completion.
+	for (FAsyncCardRepresentationTask* Task : Tasks)
+	{
+		if (Task->AsyncTask)
+		{
+			Task->AsyncTask->Cancel();
+		}
+	}
+
+	for (FAsyncCardRepresentationTask* Task : Tasks)
+	{
+		if (Task->AsyncTask)
+		{
+			Task->AsyncTask->EnsureCompletion();
+			Task->AsyncTask.Reset();
+		}
+	}
+}
+
+void FCardRepresentationAsyncQueue::StartBackgroundTask(FAsyncCardRepresentationTask* Task)
+{
+	check(Task->AsyncTask == nullptr);
+	Task->AsyncTask = MakeUnique<FAsyncTask<FAsyncCardRepresentationTaskWorker>>(*Task);
+	Task->AsyncTask->StartBackgroundTask(ThreadPool.Get(), EQueuedWorkPriority::Lowest);
+}
+
 void FCardRepresentationAsyncQueue::ProcessPendingTasks()
 {
 	FScopeLock Lock(&CriticalSection);
@@ -183,16 +219,7 @@ void FCardRepresentationAsyncQueue::ProcessPendingTasks()
 		}
 		else
 		{
-			// To avoid deadlocks, we must queue the inner build tasks on another thread pool, so use the task graph.
-			AsyncPool(
-				*ThreadPool,
-				[this, Task]()
-				{
-					// Put on background thread to avoid interfering with game-thread bound tasks
-					FQueuedThreadPoolTaskGraphWrapper TaskGraphWrapper(ENamedThreads::AnyBackgroundThreadNormalTask);
-					Build(Task, TaskGraphWrapper);
-				}
-			);
+			StartBackgroundTask(Task);
 		}
 	}
 }
@@ -220,23 +247,17 @@ void FCardRepresentationAsyncQueue::AddTask(FAsyncCardRepresentationTask* Task)
 	}
 	else
 	{
-		auto BuildLambda =
-			[this, Task]()
-			{
-				// To avoid deadlocks, we must queue the inner build tasks on another thread pool, so use the task graph.
-				// Put on background thread to avoid interfering with game-thread bound tasks
-				FQueuedThreadPoolTaskGraphWrapper TaskGraphWrapper(ENamedThreads::AnyBackgroundThreadNormalTask);
-				Build(Task, TaskGraphWrapper);
-			};
-      
 		// If we're already in worker threads there is no need to launch an async task.
 		if (GUseAsyncCardRepresentationBuildQueue || !IsInGameThread())
 		{
-			AsyncPool(*ThreadPool, MoveTemp(BuildLambda));
+			StartBackgroundTask(Task);
 		}
 		else
 		{
-			BuildLambda();
+			// To avoid deadlocks, we must queue the inner build tasks on another thread pool, so use the task graph.
+			// Put on background thread to avoid interfering with game-thread bound tasks
+			FQueuedThreadPoolTaskGraphWrapper TaskGraphWrapper(ENamedThreads::AnyBackgroundThreadNormalTask);
+			Build(Task, TaskGraphWrapper);
 		}
 	}
 #else
@@ -244,6 +265,74 @@ void FCardRepresentationAsyncQueue::AddTask(FAsyncCardRepresentationTask* Task)
 #endif
 }
 
+void FCardRepresentationAsyncQueue::CancelBuild(UStaticMesh* StaticMesh)
+{
+	TRACE_CPUPROFILER_EVENT_SCOPE(FCardRepresentationAsyncQueue::CancelBuild);
+
+	TArray<FAsyncCardRepresentationTask*> TasksToCancel;
+	{
+		FScopeLock Lock(&CriticalSection);
+		TArray<FAsyncCardRepresentationTask*> Tasks = MoveTemp(PendingTasks);
+		PendingTasks.Reserve(Tasks.Num());
+		for (FAsyncCardRepresentationTask* Task : Tasks)
+		{
+			if (Task->GenerateSource != StaticMesh && Task->StaticMesh != StaticMesh)
+			{
+				PendingTasks.Add(Task);
+			}
+		}
+
+		Tasks = MoveTemp(ReferencedTasks);
+		ReferencedTasks.Reserve(Tasks.Num());
+		for (FAsyncCardRepresentationTask* Task : Tasks)
+		{
+			if (Task->GenerateSource != StaticMesh && Task->StaticMesh != StaticMesh)
+			{
+				ReferencedTasks.Add(Task);
+			}
+			else
+			{
+				TasksToCancel.Add(Task);
+			}
+		}
+	}
+
+	CancelBackgroundTask(TasksToCancel);
+	for (FAsyncCardRepresentationTask* Task : TasksToCancel)
+	{
+		delete Task;
+	}
+}
+
+void FCardRepresentationAsyncQueue::CancelAllOutstandingBuilds()
+{
+	TRACE_CPUPROFILER_EVENT_SCOPE(FCardRepresentationAsyncQueue::CancelAllOutstandingBuilds);
+
+	TArray<FAsyncCardRepresentationTask*> OutstandingTasks;
+	{
+		FScopeLock Lock(&CriticalSection);
+		PendingTasks.Empty();
+		OutstandingTasks = MoveTemp(ReferencedTasks);
+	}
+
+	CancelBackgroundTask(OutstandingTasks);
+	for (FAsyncCardRepresentationTask* Task : OutstandingTasks)
+	{
+		delete Task;
+	}
+}
+
+void FCardRepresentationAsyncQueue::RescheduleBackgroundTask(FAsyncCardRepresentationTask* InTask, EQueuedWorkPriority InPriority)
+{
+	if (InTask->AsyncTask)
+	{
+		if (InTask->AsyncTask->GetPriority() != InPriority)
+		{
+			InTask->AsyncTask->Reschedule(GThreadPool, InPriority);
+		}
+	}
+}
+
 void FCardRepresentationAsyncQueue::BlockUntilBuildComplete(UStaticMesh* StaticMesh, bool bWarnIfBlocked)
 {
 	// We will track the wait time here, but only the cycles used.
@@ -269,8 +358,12 @@ void FCardRepresentationAsyncQueue::BlockUntilBuildComplete(UStaticMesh* StaticM
 			FScopeLock Lock(&CriticalSection);
 			for (int TaskIndex = 0; TaskIndex < ReferencedTasks.Num(); TaskIndex++)
 			{
-				bReferenced = bReferenced || ReferencedTasks[TaskIndex]->StaticMesh == StaticMesh;
-				bReferenced = bReferenced || ReferencedTasks[TaskIndex]->GenerateSource == StaticMesh;
+				if (ReferencedTasks[TaskIndex]->StaticMesh == StaticMesh ||
+					ReferencedTasks[TaskIndex]->GenerateSource == StaticMesh)
+				{
+					bReferenced = true;
+					RescheduleBackgroundTask(ReferencedTasks[TaskIndex], EQueuedWorkPriority::Highest);
+				}
 			}
 		}
 
@@ -308,6 +401,15 @@ void FCardRepresentationAsyncQueue::BlockUntilAllBuildsComplete()
 #if WITH_EDITOR
 		FStaticMeshCompilingManager::Get().FinishAllCompilation();
 #endif
+		// Reschedule as highest prio since we're explicitly waiting on them
+		{
+			FScopeLock Lock(&CriticalSection);
+			for (int TaskIndex = 0; TaskIndex < ReferencedTasks.Num(); TaskIndex++)
+			{
+				RescheduleBackgroundTask(ReferencedTasks[TaskIndex], EQueuedWorkPriority::Highest);
+			}
+		}
+
 		ProcessAsyncTasks();
 		FPlatformProcess::Sleep(.01f);
 	} 
@@ -321,6 +423,8 @@ void FCardRepresentationAsyncQueue::Build(FAsyncCardRepresentationTask* Task, FQ
 	// Editor 'force delete' can null any UObject pointers which are seen by reference collecting (eg UProperty or serialized)
 	if (Task->StaticMesh && Task->GenerateSource)
 	{
+		TRACE_CPUPROFILER_EVENT_SCOPE(FCardRepresentationAsyncQueue::Build);
+
 		const FStaticMeshLODResources& LODModel = Task->GenerateSource->GetRenderData()->LODResources[0];
 
 		Task->bSuccess = MeshUtilities->GenerateCardRepresentationData(
@@ -371,13 +475,25 @@ void FCardRepresentationAsyncQueue::ProcessAsyncTasks(bool bLimitExecutionTime)
 		{
 			break;
 		}
-	
+
 		// We want to count each resource built from a DDC miss, so count each iteration of the loop separately.
 		COOK_STAT(auto Timer = CardRepresentationCookStats::UsageStats.TimeSyncWork());
 
+		bool bWasCancelled = false;
 		{
 			FScopeLock Lock(&CriticalSection);
-			ReferencedTasks.Remove(Task);
+			bWasCancelled = ReferencedTasks.Remove(Task) == 0;
+		}
+
+		if (bWasCancelled)
+		{
+			continue;
+		}
+
+		if (Task->AsyncTask)
+		{
+			Task->AsyncTask->EnsureCompletion();
+			Task->AsyncTask.Reset();
 		}
 
 		// Editor 'force delete' can null any UObject pointers which are seen by reference collecting (eg UProperty or serialized)
@@ -420,6 +536,8 @@ void FCardRepresentationAsyncQueue::ProcessAsyncTasks(bool bLimitExecutionTime)
 
 void FCardRepresentationAsyncQueue::Shutdown()
 {
+	CancelAllOutstandingBuilds();
+
 	UE_LOG(LogStaticMesh, Log, TEXT("Abandoning remaining async card representation tasks for shutdown"));
 	ThreadPool->Destroy();
 }
\ No newline at end of file
diff --git a/Engine/Source/Runtime/Engine/Private/StaticMesh.cpp b/Engine/Source/Runtime/Engine/Private/StaticMesh.cpp
--- a/Engine/Source/Runtime/Engine/Private/StaticMesh.cpp
+++ b/Engine/Source/Runtime/Engine/Private/StaticMesh.cpp
@@ -4880,16 +4880,16 @@ void UStaticMesh::CacheDerivedData()
 		// on the main thread when calling CacheDerivedData() from another thread.
 		if (IsInGameThread())
 		{
-			// Finish any previous async builds before modifying RenderData
+			// Cancel any previous async builds before modifying RenderData
 			// This can happen during import as the mesh is rebuilt redundantly
 			if (GDistanceFieldAsyncQueue)
 			{
-				GDistanceFieldAsyncQueue->BlockUntilBuildComplete(this, true);
+				GDistanceFieldAsyncQueue->CancelBuild(this);
 			}
 
 			if (GCardRepresentationAsyncQueue)
 			{
-				GCardRepresentationAsyncQueue->BlockUntilBuildComplete(this, true);
+				GCardRepresentationAsyncQueue->CancelBuild(this);
 			}
 		}
 
diff --git a/Engine/Source/Runtime/Engine/Private/StaticMeshBuild.cpp b/Engine/Source/Runtime/Engine/Private/StaticMeshBuild.cpp
--- a/Engine/Source/Runtime/Engine/Private/StaticMeshBuild.cpp
+++ b/Engine/Source/Runtime/Engine/Private/StaticMeshBuild.cpp
@@ -173,17 +173,16 @@ void UStaticMesh::BatchBuild(const TArray<UStaticMesh*>& InStaticMeshes, bool bI
 		{
 			if (StaticMesh->GetRenderData())
 			{
-				// #TODO Editor-Perf Those should be canceled instead of blocking on them
-				// Finish any previous async builds before modifying RenderData
+				// Cancel any previous async builds before modifying RenderData
 				// This can happen during import as the mesh is rebuilt redundantly
 				if (GDistanceFieldAsyncQueue)
 				{
-					GDistanceFieldAsyncQueue->BlockUntilBuildComplete(StaticMesh, true);
+					GDistanceFieldAsyncQueue->CancelBuild(StaticMesh);
 				}
 
 				if (GCardRepresentationAsyncQueue)
 				{
-					GCardRepresentationAsyncQueue->BlockUntilBuildComplete(StaticMesh, true);
+					GCardRepresentationAsyncQueue->CancelBuild(StaticMesh);
 				}
 			}
 		}
diff --git a/Engine/Source/Runtime/Engine/Public/DistanceFieldAtlas.h b/Engine/Source/Runtime/Engine/Public/DistanceFieldAtlas.h
--- a/Engine/Source/Runtime/Engine/Public/DistanceFieldAtlas.h
+++ b/Engine/Source/Runtime/Engine/Public/DistanceFieldAtlas.h
@@ -16,6 +16,7 @@
 #include "TextureLayout3d.h"
 #include "Templates/UniquePtr.h"
 #include "DerivedMeshDataTaskUtils.h"
+#include "Async/AsyncWork.h"
 
 class FDistanceFieldVolumeData;
 class UStaticMesh;
@@ -380,6 +381,26 @@ public:
 	}
 };
 
+class FAsyncDistanceFieldTask;
+class FAsyncDistanceFieldTaskWorker : public FNonAbandonableTask
+{
+public:
+	FAsyncDistanceFieldTaskWorker(FAsyncDistanceFieldTask& InTask)
+		: Task(InTask)
+	{
+	}
+
+	FORCEINLINE TStatId GetStatId() const
+	{
+		RETURN_QUICK_DECLARE_CYCLE_STAT(FAsyncDistanceFieldTaskWorker, STATGROUP_ThreadPoolAsyncTasks);
+	}
+
+	void DoWork();
+
+private:
+	FAsyncDistanceFieldTask& Task;
+};
+
 /** A task to build a distance field for a single mesh */
 class FAsyncDistanceFieldTask
 {
@@ -395,6 +416,7 @@ public:
 	const ITargetPlatform* TargetPlatform;
 	FString DDCKey;
 	FDistanceFieldVolumeData* GeneratedVolumeData;
+	TUniquePtr<FAsyncTask<FAsyncDistanceFieldTaskWorker>> AsyncTask = nullptr;
 };
 
 /** Class that manages asynchronous building of mesh distance fields. */
@@ -409,6 +431,12 @@ public:
 	/** Adds a new build task. (Thread-Safe) */
 	ENGINE_API void AddTask(FAsyncDistanceFieldTask* Task);
 
+	/** Cancel the build on this specific static mesh or block until it is completed if already started. */
+	ENGINE_API void CancelBuild(UStaticMesh* StaticMesh);
+
+	/** Blocks the main thread until the async build are either cancelled or completed. */
+	ENGINE_API void CancelAllOutstandingBuilds();
+
 	/** Blocks the main thread until the async build of the specified mesh is complete. */
 	ENGINE_API void BlockUntilBuildComplete(UStaticMesh* StaticMesh, bool bWarnIfBlocked);
 
@@ -434,6 +462,7 @@ public:
 	}
 
 private:
+	friend FAsyncDistanceFieldTaskWorker;
 	void ProcessPendingTasks();
 
 	TUniquePtr<FQueuedThreadPool> ThreadPool;
@@ -441,6 +470,15 @@ private:
 	/** Builds a single task with the given threadpool.  Called from the worker thread. */
 	void Build(FAsyncDistanceFieldTask* Task, class FQueuedThreadPool& ThreadPool);
 
+	/** Change the priority of the background task. */
+	void RescheduleBackgroundTask(FAsyncDistanceFieldTask* InTask, EQueuedWorkPriority InPriority);
+
+	/** Task will be sent to a background worker. */
+	void StartBackgroundTask(FAsyncDistanceFieldTask* Task);
+
+	/** Cancel or finish any background work for the given task. */
+	void CancelBackgroundTask(TArray<FAsyncDistanceFieldTask*> Tasks);
+
 	/** Game-thread managed list of tasks in the async system. */
 	TArray<FAsyncDistanceFieldTask*> ReferencedTasks;
 
diff --git a/Engine/Source/Runtime/Engine/Public/MeshCardRepresentation.h b/Engine/Source/Runtime/Engine/Public/MeshCardRepresentation.h
--- a/Engine/Source/Runtime/Engine/Public/MeshCardRepresentation.h
+++ b/Engine/Source/Runtime/Engine/Public/MeshCardRepresentation.h
@@ -15,6 +15,7 @@
 #include "RenderingThread.h"
 #include "Templates/UniquePtr.h"
 #include "DerivedMeshDataTaskUtils.h"
+#include "Async/AsyncWork.h"
 
 template <class T> class TLockFreePointerListLIFO;
 
@@ -155,6 +156,27 @@ public:
 	}
 };
 
+
+class FAsyncCardRepresentationTask;
+class FAsyncCardRepresentationTaskWorker : public FNonAbandonableTask
+{
+public:
+	FAsyncCardRepresentationTaskWorker(FAsyncCardRepresentationTask& InTask)
+		: Task(InTask)
+	{
+	}
+
+	FORCEINLINE TStatId GetStatId() const
+	{
+		RETURN_QUICK_DECLARE_CYCLE_STAT(FAsyncCardRepresentationTaskWorker, STATGROUP_ThreadPoolAsyncTasks);
+	}
+
+	void DoWork();
+
+private:
+	FAsyncCardRepresentationTask& Task;
+};
+
 class FAsyncCardRepresentationTask
 {
 public:
@@ -167,6 +189,7 @@ public:
 	UStaticMesh* GenerateSource;
 	FString DDCKey;
 	FCardRepresentationData* GeneratedCardRepresentation;
+	TUniquePtr<FAsyncTask<FAsyncCardRepresentationTaskWorker>> AsyncTask = nullptr;
 };
 
 /** Class that manages asynchronous building of mesh distance fields. */
@@ -181,6 +204,12 @@ public:
 	/** Adds a new build task. */
 	ENGINE_API void AddTask(FAsyncCardRepresentationTask* Task);
 
+	/** Cancel the build on this specific static mesh or block until it is completed if already started. */
+	ENGINE_API void CancelBuild(UStaticMesh* StaticMesh);
+
+	/** Blocks the main thread until the async build are either cancelled or completed. */
+	ENGINE_API void CancelAllOutstandingBuilds();
+
 	/** Blocks the main thread until the async build of the specified mesh is complete. */
 	ENGINE_API void BlockUntilBuildComplete(UStaticMesh* StaticMesh, bool bWarnIfBlocked);
 
@@ -205,6 +234,7 @@ public:
 	}
 
 private:
+	friend FAsyncCardRepresentationTaskWorker;
 	void ProcessPendingTasks();
 
 	TUniquePtr<FQueuedThreadPool> ThreadPool;
@@ -212,6 +242,15 @@ private:
 	/** Builds a single task with the given threadpool.  Called from the worker thread. */
 	void Build(FAsyncCardRepresentationTask* Task, class FQueuedThreadPool& ThreadPool);
 
+	/** Change the priority of the background task. */
+	void RescheduleBackgroundTask(FAsyncCardRepresentationTask* InTask, EQueuedWorkPriority InPriority);
+
+	/** Task will be sent to a background worker. */
+	void StartBackgroundTask(FAsyncCardRepresentationTask* Task);
+
+	/** Cancel or finish any background work for the given task. */
+	void CancelBackgroundTask(TArray<FAsyncCardRepresentationTask*> Tasks);
+
 	/** Game-thread managed list of tasks in the async system. */
 	TArray<FAsyncCardRepresentationTask*> ReferencedTasks;
 
