Force UTexture releases and deletes resource on render thread.
This prevents accessing deleted data when releasing resource in batched release.

Prevent FCanvasTileItem from referencing proxy resources (valid while the textures are being build).
This prevents accessing deleted data ub FBatchedElements::Draw().

diff --git a/Engine/Source/Runtime/Engine/Private/Texture.cpp b/Engine/Source/Runtime/Engine/Private/Texture.cpp
--- a/Engine/Source/Runtime/Engine/Private/Texture.cpp
+++ b/Engine/Source/Runtime/Engine/Private/Texture.cpp
@@ -168,10 +168,10 @@ void UTexture::ReleaseResource()
 
 		FTextureResource* ToDelete = PrivateResource;
 		// Free the resource.
-		BeginReleaseResource(ToDelete);
 		SetResource(nullptr);
 		ENQUEUE_RENDER_COMMAND(DeleteResource)([ToDelete](FRHICommandListImmediate& RHICmdList)
 		{
+			ToDelete->ReleaseResource();
 			delete ToDelete;
 		});
 	}
diff --git a/Engine/Source/Runtime/Engine/Private/UserInterface/Canvas.cpp b/Engine/Source/Runtime/Engine/Private/UserInterface/Canvas.cpp
--- a/Engine/Source/Runtime/Engine/Private/UserInterface/Canvas.cpp
+++ b/Engine/Source/Runtime/Engine/Private/UserInterface/Canvas.cpp
@@ -1966,12 +1966,26 @@ void UCanvas::K2_DrawLine(FVector2D ScreenPositionA, FVector2D ScreenPositionB,
 	}
 }
 
+FTexture* GetTextureForCanvasItem(UTexture* RenderTexture)
+{
+	// Only assign a valid texture for the canvas item. Proxy textures, from the UTexture implementation, are not valid here.
+	// This is because the proxy textures get destroyed on the render thread as soon as the build completes, which could happen
+	// before the canvas item renders.
+	if (RenderTexture && RenderTexture->Resource && !RenderTexture->Resource->IsProxy())
+	{
+		return RenderTexture->Resource;
+	}
+	else
+	{
+		return GWhiteTexture;
+	}
+}
+
 void UCanvas::K2_DrawTexture(UTexture* RenderTexture, FVector2D ScreenPosition, FVector2D ScreenSize, FVector2D CoordinatePosition, FVector2D CoordinateSize, FLinearColor RenderColor, EBlendMode BlendMode, float Rotation, FVector2D PivotPoint)
 {
 	if (ScreenSize.X > 0.0f && ScreenSize.Y > 0.0f && Canvas)
 	{
-		FTexture* RenderTextureResource = (RenderTexture) ? RenderTexture->Resource : GWhiteTexture;
-		FCanvasTileItem TileItem(ScreenPosition, RenderTextureResource, ScreenSize, CoordinatePosition, CoordinatePosition + CoordinateSize, RenderColor);
+		FCanvasTileItem TileItem(ScreenPosition, GetTextureForCanvasItem(RenderTexture), ScreenSize, CoordinatePosition, CoordinatePosition + CoordinateSize, RenderColor);
 		TileItem.Rotation = FRotator(0, Rotation, 0);
 		TileItem.PivotPoint = PivotPoint;
 		TileItem.BlendMode = FCanvas::BlendToSimpleElementBlend(BlendMode);
@@ -2044,7 +2058,7 @@ void UCanvas::K2_DrawTriangle(UTexture* RenderTexture, TArray<FCanvasUVTri> Tria
 {
 	if (Triangles.Num() > 0 && Canvas)
 	{
-		FCanvasTriangleItem TriangleItem(FVector2D::ZeroVector, FVector2D::ZeroVector, FVector2D::ZeroVector, (RenderTexture) ? RenderTexture->Resource : GWhiteTexture);
+		FCanvasTriangleItem TriangleItem(FVector2D::ZeroVector, FVector2D::ZeroVector, FVector2D::ZeroVector, GetTextureForCanvasItem(RenderTexture));
 		TriangleItem.TriangleList = MoveTemp(Triangles);
 		DrawItem(TriangleItem);
 	}
@@ -2065,7 +2079,7 @@ void UCanvas::K2_DrawPolygon(UTexture* RenderTexture, FVector2D ScreenPosition,
 {
 	if (Radius.X > 0.0f && Radius.Y > 0.0f && NumberOfSides >= 3 && Canvas)
 	{
-		FCanvasNGonItem NGonItem(ScreenPosition, Radius, NumberOfSides, (RenderTexture) ? RenderTexture->Resource : GWhiteTexture, RenderColor);
+		FCanvasNGonItem NGonItem(ScreenPosition, Radius, NumberOfSides, GetTextureForCanvasItem(RenderTexture), RenderColor);
 		DrawItem(NGonItem);
 	}
 }
