Strata - do not use PixelNormalWS by default because it comes from the material node that my not be fed. Use VertexNormal as the correct fallback.
That also fixes the current shared normal expected allocation.

diff --git a/Engine/Source/Runtime/Engine/Private/Materials/MaterialExpressions.cpp b/Engine/Source/Runtime/Engine/Private/Materials/MaterialExpressions.cpp
--- a/Engine/Source/Runtime/Engine/Private/Materials/MaterialExpressions.cpp
+++ b/Engine/Source/Runtime/Engine/Private/Materials/MaterialExpressions.cpp
@@ -19030,7 +19030,7 @@ UMaterialExpressionStrataDiffuseBSDF::UMaterialExpressionStrataDiffuseBSDF(const
 #if WITH_EDITOR
 int32 UMaterialExpressionStrataDiffuseBSDF::Compile(class FMaterialCompiler* Compiler, int32 OutputIndex)
 {
-	int32 NormalCodeChunk = Normal.GetTracedInput().Expression ? Normal.Compile(Compiler) : Compiler->PixelNormalWS();
+	int32 NormalCodeChunk = Normal.GetTracedInput().Expression ? Normal.Compile(Compiler) : Compiler->VertexNormal();
 	uint8 SharedNormalIndex = StrataCompilationInfoCreateSharedNormal(Compiler, NormalCodeChunk);
 
 	int32 OutputCodeChunk = Compiler->StrataDiffuseBSDF(
@@ -19094,7 +19094,7 @@ UMaterialExpressionStrataDielectricBSDF::UMaterialExpressionStrataDielectricBSDF
 #if WITH_EDITOR
 int32 UMaterialExpressionStrataDielectricBSDF::Compile(class FMaterialCompiler* Compiler, int32 OutputIndex)
 {
-	int32 NormalCodeChunk = Normal.GetTracedInput().Expression ? Normal.Compile(Compiler) : Compiler->PixelNormalWS();
+	int32 NormalCodeChunk = Normal.GetTracedInput().Expression ? Normal.Compile(Compiler) : Compiler->VertexNormal();
 	int32 TangentCodeChunk = Tangent.GetTracedInput().Expression ? Tangent.Compile(Compiler) : Compiler->VertexTangent();
 	uint8 SharedNormalIndex = StrataCompilationInfoCreateSharedNormal(Compiler, NormalCodeChunk);
 
@@ -19167,7 +19167,7 @@ UMaterialExpressionStrataConductorBSDF::UMaterialExpressionStrataConductorBSDF(c
 #if WITH_EDITOR
 int32 UMaterialExpressionStrataConductorBSDF::Compile(class FMaterialCompiler* Compiler, int32 OutputIndex)
 {
-	int32 NormalCodeChunk = Normal.GetTracedInput().Expression ? Normal.Compile(Compiler) : Compiler->PixelNormalWS();
+	int32 NormalCodeChunk = Normal.GetTracedInput().Expression ? Normal.Compile(Compiler) : Compiler->VertexNormal();
 	int32 TangentCodeChunk = Tangent.GetTracedInput().Expression ? Tangent.Compile(Compiler) : Compiler->VertexTangent();
 	uint8 SharedNormalIndex = StrataCompilationInfoCreateSharedNormal(Compiler, NormalCodeChunk);
 
