Separated roughness X and Y inputs of conductor and dielectric.

diff --git a/Engine/Source/Runtime/Engine/Classes/Materials/MaterialExpressionStrata.h b/Engine/Source/Runtime/Engine/Classes/Materials/MaterialExpressionStrata.h
--- a/Engine/Source/Runtime/Engine/Classes/Materials/MaterialExpressionStrata.h
+++ b/Engine/Source/Runtime/Engine/Classes/Materials/MaterialExpressionStrata.h
@@ -86,10 +86,16 @@ class UMaterialExpressionStrataDielectricBSDF : public UMaterialExpression
 	FExpressionInput Tint;
 		
 	/**
-	 * Roughness (type = float2, unit = unitless)
+	 * RoughnessX (type = float, unit = unitless)
 	 */
 	UPROPERTY()
-	FExpressionInput Roughness;
+	FExpressionInput RoughnessX;
+		
+	/**
+	 * RoughnessY (type = float, unit = unitless), If not plugged in, RoughnessY is set to RoughnessX to get an isotropic behavior.
+	 */
+	UPROPERTY()
+	FExpressionInput RoughnessY;
 
 	/**
 	 * Normal (type = float3, unit = unitless)
@@ -129,12 +135,18 @@ class UMaterialExpressionStrataConductorBSDF : public UMaterialExpression
 	 */
 	UPROPERTY()
 	FExpressionInput EdgeColor;
-
+		
 	/**
-	 * Roughness (type = float2, unit = unitless)
+	 * RoughnessX (type = float, unit = unitless)
 	 */
 	UPROPERTY()
-	FExpressionInput Roughness;
+	FExpressionInput RoughnessX;
+		
+	/**
+	 * RoughnessY (type = float, unit = unitless), If not plugged in, RoughnessY is set to RoughnessX to get an isotropic behavior.
+	 */
+	UPROPERTY()
+	FExpressionInput RoughnessY;
 
 	/**
 	 * Normal (type = float3, unit = unitless)
diff --git a/Engine/Source/Runtime/Engine/Private/Materials/HLSLMaterialTranslator.cpp b/Engine/Source/Runtime/Engine/Private/Materials/HLSLMaterialTranslator.cpp
--- a/Engine/Source/Runtime/Engine/Private/Materials/HLSLMaterialTranslator.cpp
+++ b/Engine/Source/Runtime/Engine/Private/Materials/HLSLMaterialTranslator.cpp
@@ -7282,11 +7282,12 @@ int32 FHLSLMaterialTranslator::StrataDiffuseBSDF(int32 Albedo, int32 Roughness,
 	);
 }
 
-int32 FHLSLMaterialTranslator::StrataDielectricBSDF(int32 Roughness, int32 IOR, int32 Tint, int32 Normal, int32 Tangent, uint8 SharedNormalIndex)
+int32 FHLSLMaterialTranslator::StrataDielectricBSDF(int32 RoughnessX, int32 RoughnessY, int32 IOR, int32 Tint, int32 Normal, int32 Tangent, uint8 SharedNormalIndex)
 {
 	return AddCodeChunk(
-		MCT_Strata, TEXT("GetStrataDielectricBSDF(%s, %s, %s, %s, %u) /* Normal:%s Tangent:%s */"),
-		*GetParameterCode(Roughness),
+		MCT_Strata, TEXT("GetStrataDielectricBSDF(float2(%s, %s), %s, %s, %s, %u) /* Normal:%s Tangent:%s */"),
+		*GetParameterCode(RoughnessX),
+		*GetParameterCode(RoughnessY),
 		*GetParameterCode(IOR),
 		*GetParameterCode(Tint),
 		*GetParameterCode(Tangent),
@@ -7296,13 +7297,14 @@ int32 FHLSLMaterialTranslator::StrataDielectricBSDF(int32 Roughness, int32 IOR,
 	);
 }
 
-int32 FHLSLMaterialTranslator::StrataConductorBSDF(int32 Reflectivity, int32 EdgeColor, int32 Roughness, int32 Normal, int32 Tangent, uint8 SharedNormalIndex)
+int32 FHLSLMaterialTranslator::StrataConductorBSDF(int32 Reflectivity, int32 EdgeColor, int32 RoughnessX, int32 RoughnessY, int32 Normal, int32 Tangent, uint8 SharedNormalIndex)
 {
 	return AddCodeChunk(
-		MCT_Strata, TEXT("GetStrataConductorBSDF(%s, %s, %s, %s, %u) /* Normal:%s Tangent:%s */"),
+		MCT_Strata, TEXT("GetStrataConductorBSDF(%s, %s, float2(%s, %s), %s, %u) /* Normal:%s Tangent:%s */"),
 		*GetParameterCode(Reflectivity),
 		*GetParameterCode(EdgeColor),
-		*GetParameterCode(Roughness),
+		*GetParameterCode(RoughnessX),
+		*GetParameterCode(RoughnessY),
 		*GetParameterCode(Tangent),
 		SharedNormalIndex,
 		*GetParameterCode(Normal),
diff --git a/Engine/Source/Runtime/Engine/Private/Materials/HLSLMaterialTranslator.h b/Engine/Source/Runtime/Engine/Private/Materials/HLSLMaterialTranslator.h
--- a/Engine/Source/Runtime/Engine/Private/Materials/HLSLMaterialTranslator.h
+++ b/Engine/Source/Runtime/Engine/Private/Materials/HLSLMaterialTranslator.h
@@ -788,8 +788,8 @@ protected:
 	// Strata
 	virtual int32 FrontMaterial() override;
 	virtual int32 StrataDiffuseBSDF(int32 Albedo, int32 Roughness, int32 Normal, uint8 SharedNormalIndex) override;
-	virtual int32 StrataDielectricBSDF(int32 Roughness, int32 IOR, int32 Tint, int32 Normal, int32 Tangent, uint8 SharedNormalIndex) override;
-	virtual int32 StrataConductorBSDF(int32 Reflectivity, int32 EdgeColor, int32 Roughness, int32 Normal, int32 Tangent, uint8 SharedNormalIndex) override;
+	virtual int32 StrataDielectricBSDF(int32 RoughnessX, int32 RoughnessY, int32 IOR, int32 Tint, int32 Normal, int32 Tangent, uint8 SharedNormalIndex) override;
+	virtual int32 StrataConductorBSDF(int32 Reflectivity, int32 EdgeColor, int32 RoughnessX, int32 RoughnessY, int32 Normal, int32 Tangent, uint8 SharedNormalIndex) override;
 	virtual int32 StrataVolumeBSDF(int32 Albedo, int32 Extinction, int32 Anisotropy, int32 Thickness, int32 Normal, uint8 SharedNormalIndex) override;
 	virtual int32 StrataHorizontalMixing(int32 Foreground, int32 Background, int32 Mix) override;
 	virtual int32 StrataVerticalLayering(int32 Top, int32 Base) override;
diff --git a/Engine/Source/Runtime/Engine/Private/Materials/MaterialExpressions.cpp b/Engine/Source/Runtime/Engine/Private/Materials/MaterialExpressions.cpp
--- a/Engine/Source/Runtime/Engine/Private/Materials/MaterialExpressions.cpp
+++ b/Engine/Source/Runtime/Engine/Private/Materials/MaterialExpressions.cpp
@@ -19098,8 +19098,13 @@ int32 UMaterialExpressionStrataDielectricBSDF::Compile(class FMaterialCompiler*
 	int32 TangentCodeChunk = Tangent.GetTracedInput().Expression ? Tangent.Compile(Compiler) : Compiler->VertexTangent();
 	uint8 SharedNormalIndex = StrataCompilationInfoCreateSharedNormal(Compiler, NormalCodeChunk);
 
+	int32 RoughnessXCodeChunk = RoughnessX.GetTracedInput().Expression ? RoughnessX.Compile(Compiler) : Compiler->Constant(0.0f);
+	// If not plugged in, RoughnessYCodeChunk is set to RoughnessXCodeChunk to get an isotropic behavior
+	int32 RoughnessYCodeChunk = RoughnessY.GetTracedInput().Expression ? RoughnessY.Compile(Compiler) : RoughnessXCodeChunk;
+
 	int32 OutputCodeChunk = Compiler->StrataDielectricBSDF(
-		Roughness.GetTracedInput().Expression	? Roughness.Compile(Compiler)	: Compiler->Constant2(0.0f, 0.0f),
+		RoughnessXCodeChunk,
+		RoughnessYCodeChunk,
 		IOR.GetTracedInput().Expression			? IOR.Compile(Compiler)			: Compiler->Constant(1.5f),// Default to Glass
 		Tint.GetTracedInput().Expression		? Tint.Compile(Compiler)		: Compiler->Constant3(1.0f, 1.0f, 1.0f),
 		NormalCodeChunk,
@@ -19131,14 +19136,17 @@ uint32 UMaterialExpressionStrataDielectricBSDF::GetInputType(int32 InputIndex)
 		return MCT_Float3;
 		break;
 	case 2:
-		return MCT_Float2;
+		return MCT_Float;
 		break;
 	case 3:
-		return MCT_Float3;
+		return MCT_Float;
 		break;
 	case 4:
 		return MCT_Float3;
 		break;
+	case 5:
+		return MCT_Float3;
+		break;
 	}
 
 	check(false);
@@ -19171,10 +19179,15 @@ int32 UMaterialExpressionStrataConductorBSDF::Compile(class FMaterialCompiler* C
 	int32 TangentCodeChunk = Tangent.GetTracedInput().Expression ? Tangent.Compile(Compiler) : Compiler->VertexTangent();
 	uint8 SharedNormalIndex = StrataCompilationInfoCreateSharedNormal(Compiler, NormalCodeChunk);
 
+	int32 RoughnessXCodeChunk = RoughnessX.GetTracedInput().Expression ? RoughnessX.Compile(Compiler) : Compiler->Constant(0.0f);
+	// If not plugged in, RoughnessYCodeChunk is set to RoughnessXCodeChunk to get an isotropic behavior
+	int32 RoughnessYCodeChunk = RoughnessY.GetTracedInput().Expression ? RoughnessY.Compile(Compiler) : RoughnessXCodeChunk;
+
 	int32 OutputCodeChunk = Compiler->StrataConductorBSDF(
 		Reflectivity.GetTracedInput().Expression	? Reflectivity.Compile(Compiler): Compiler->Constant3(0.947f, 0.776f, 0.371f),	// Default to Gold
 		EdgeColor.GetTracedInput().Expression		? EdgeColor.Compile(Compiler)	: Compiler->Constant3(1.000f, 0.982f, 0.753f),	// Default to Gold
-		Roughness.GetTracedInput().Expression		? Roughness.Compile(Compiler)	: Compiler->Constant2(0.0f, 0.0f),
+		RoughnessXCodeChunk,
+		RoughnessYCodeChunk,
 		NormalCodeChunk,
 		TangentCodeChunk,
 		SharedNormalIndex);
@@ -19204,14 +19217,17 @@ uint32 UMaterialExpressionStrataConductorBSDF::GetInputType(int32 InputIndex)
 		return MCT_Float3;
 		break;
 	case 2:
-		return MCT_Float2;
+		return MCT_Float;
 		break;
 	case 3:
-		return MCT_Float3;
+		return MCT_Float;
 		break;
 	case 4:
 		return MCT_Float3;
 		break;
+	case 5:
+		return MCT_Float3;
+		break;
 	}
 
 	check(false);
diff --git a/Engine/Source/Runtime/Engine/Public/MaterialCompiler.h b/Engine/Source/Runtime/Engine/Public/MaterialCompiler.h
--- a/Engine/Source/Runtime/Engine/Public/MaterialCompiler.h
+++ b/Engine/Source/Runtime/Engine/Public/MaterialCompiler.h
@@ -392,8 +392,8 @@ public:
 	// Strata
 	virtual int32 FrontMaterial() = 0;
 	virtual int32 StrataDiffuseBSDF(int32 Albedo, int32 Roughness, int32 Normal, uint8 SharedNormalIndex) = 0;
-	virtual int32 StrataDielectricBSDF(int32 Roughness, int32 IOR, int32 Tint, int32 Normal, int32 Tangent, uint8 SharedNormalIndex) = 0;
-	virtual int32 StrataConductorBSDF(int32 Reflectivity, int32 EdgeColor, int32 Roughness, int32 Normal, int32 Tangent, uint8 SharedNormalIndex) = 0;
+	virtual int32 StrataDielectricBSDF(int32 RoughnessX, int32 RoughnessY, int32 IOR, int32 Tint, int32 Normal, int32 Tangent, uint8 SharedNormalIndex) = 0;
+	virtual int32 StrataConductorBSDF(int32 Reflectivity, int32 EdgeColor, int32 RoughnessX, int32 RoughnessY, int32 Normal, int32 Tangent, uint8 SharedNormalIndex) = 0;
 	virtual int32 StrataVolumeBSDF(int32 Albedo, int32 Extinction, int32 Anisotropy, int32 Thickness, int32 Normal, uint8 SharedNormalIndex) = 0;
 	virtual int32 StrataHorizontalMixing(int32 Foreground, int32 Background, int32 Mix) = 0;
 	virtual int32 StrataVerticalLayering(int32 Top, int32 Base) = 0;
@@ -808,14 +808,14 @@ public:
 		return Compiler->StrataDiffuseBSDF(Albedo, Roughness, Normal, SharedNormalIndex);
 	}
 
-	virtual int32 StrataDielectricBSDF(int32 Roughness, int32 IOR, int32 Tint, int32 Normal, int32 Tangent, uint8 SharedNormalIndex) override
+	virtual int32 StrataDielectricBSDF(int32 RoughnessX, int32 RoughnessY, int32 IOR, int32 Tint, int32 Normal, int32 Tangent, uint8 SharedNormalIndex) override
 	{
-		return Compiler->StrataDielectricBSDF(Roughness, IOR, Tint, Normal, Tangent, SharedNormalIndex);
+		return Compiler->StrataDielectricBSDF(RoughnessX, RoughnessY, IOR, Tint, Normal, Tangent, SharedNormalIndex);
 	}
 
-	virtual int32 StrataConductorBSDF(int32 Reflectivity, int32 EdgeColor, int32 Roughness, int32 Normal, int32 Tangent, uint8 SharedNormalIndex) override
+	virtual int32 StrataConductorBSDF(int32 Reflectivity, int32 EdgeColor, int32 RoughnessX, int32 RoughnessY, int32 Normal, int32 Tangent, uint8 SharedNormalIndex) override
 	{
-		return Compiler->StrataConductorBSDF(Reflectivity, EdgeColor, Roughness, Normal, Tangent, SharedNormalIndex);
+		return Compiler->StrataConductorBSDF(Reflectivity, EdgeColor, RoughnessX, RoughnessY, Normal, Tangent, SharedNormalIndex);
 	}
 
 	virtual int32 StrataVolumeBSDF(int32 Albedo, int32 Extinction, int32 Anisotropy, int32 Thickness, int32 Normal, uint8 SharedNormalIndex) override
