Fix crash in RepLayout that can result from replicating arrays of structs with no replicated properties.

diff --git a/Engine/Source/Runtime/Engine/Private/RepLayout.cpp b/Engine/Source/Runtime/Engine/Private/RepLayout.cpp
--- a/Engine/Source/Runtime/Engine/Private/RepLayout.cpp
+++ b/Engine/Source/Runtime/Engine/Private/RepLayout.cpp
@@ -1519,6 +1519,10 @@ static void CompareProperties_Array_r(
 	const uint16 Handle)
 {
 	const FRepLayoutCmd& Cmd = SharedParams.Cmds[CmdIndex];
+	if (EnumHasAnyFlags(Cmd.Flags, ERepLayoutCmdFlags::IsEmptyArrayStruct))
+	{
+		return;
+	}
 
 	FScriptArray* ShadowArray = (FScriptArray*)StackParams.ShadowData.Data;
 	FScriptArray* Array = (FScriptArray*)StackParams.Data.Data;
@@ -2919,13 +2923,19 @@ void FRepLayout::SendAllProperties_BackwardsCompatible_r(
 
 		check(Cmd.Type != ERepLayoutCmdType::Return);
 
-		PackageMapClient->TrackNetFieldExport(NetFieldExportGroup, CmdIndex);
+		const bool bIsArray = Cmd.Type == ERepLayoutCmdType::DynamicArray;
+		if (bIsArray && EnumHasAnyFlags(Cmd.Flags, ERepLayoutCmdFlags::IsEmptyArrayStruct))
+		{
+			CmdIndex = Cmd.EndCmd - 1;
+			continue;
+		}
 
+		PackageMapClient->TrackNetFieldExport(NetFieldExportGroup, CmdIndex);
 		WritePropertyHandle_BackwardsCompatible(Writer, CmdIndex + 1, bDoChecksum);
 
 		FConstRepObjectDataBuffer Data = SourceData + Cmd;
 
-		if (Cmd.Type == ERepLayoutCmdType::DynamicArray)
+		if (bIsArray)
 		{			
 			const FScriptArray* Array = (FScriptArray *)Data.Data;
 			const FConstRepObjectDataBuffer ArrayData(Array->GetData());
@@ -4472,6 +4482,10 @@ static void ValidateWithChecksum_DynamicArray_r(
 	FBitArchive& Ar)
 {
 	const FRepLayoutCmd& Cmd = *CmdIt;
+	if (EnumHasAllFlags(Cmd.Flags, ERepLayoutCmdFlags::IsEmptyArrayStruct))
+	{
+		return;
+	}
 
 	// -2 because the current index will be the Owner Array Properties Cmd Index (+1)
 	// and EndCmd will be the Cmd Index just *after* the Return Command (+1) 
@@ -5322,7 +5336,18 @@ static int32 InitFromProperty_r(
 
 		AddReturnCmd(SharedParams.Cmds);
 
-		SharedParams.Cmds[CmdStart].EndCmd = SharedParams.Cmds.Num();		// Patch in the offset to jump over our array inner elements
+		const int32 CmdEnd = SharedParams.Cmds.Num();
+		SharedParams.Cmds[CmdStart].EndCmd = CmdEnd;		// Patch in the offset to jump over our array inner elements
+
+		// Array commands will have their array property, the layout of the inner property, and a terminator.
+		// That means if we only have 2 commands, the array's inner propertry had no replicated properties of
+		// its own.
+		if (CmdEnd - CmdStart <= 2)
+		{
+			SharedParams.Cmds[CmdStart].Flags |= ERepLayoutCmdFlags::IsEmptyArrayStruct;
+			UE_LOG(LogRep, Warning, TEXT("InitFromProperty_r: Array property has empty inner struct: Outer=%s, Array=%s, Inner=%2"),
+				*ArrayProp->Owner.GetName(), *ArrayProp->GetName(), *ArrayProp->Inner->GetName());
+		}
 	}
 	else if (FStructProperty* StructProp = CastField<FStructProperty>(StackParams.Property))
 	{
@@ -6148,6 +6173,10 @@ void FRepLayout::SerializeProperties_DynamicArray_r(
 	FNetTraceCollector* Collector) const
 {
 	const FRepLayoutCmd& Cmd = Cmds[ CmdIndex ];
+	if (EnumHasAnyFlags(Cmd.Flags, ERepLayoutCmdFlags::IsEmptyArrayStruct))
+	{
+		return;
+	}
 
 	FScriptArray* Array = (FScriptArray*)Data.Data;
 
@@ -6416,7 +6445,11 @@ void FRepLayout::BuildSharedSerializationForRPC_DynamicArray_r(
 	int32 ArrayDepth,
 	FRepSerializationSharedInfo& SharedInfo)
 {
-	const FRepLayoutCmd& Cmd = Cmds[ CmdIndex ];
+	const FRepLayoutCmd& Cmd = Cmds[CmdIndex];
+	if (EnumHasAnyFlags(Cmd.Flags, ERepLayoutCmdFlags::IsEmptyArrayStruct))
+	{
+		return;
+	}
 
 	FScriptArray* Array = (FScriptArray *)Data.Data;	
 	const int32 ArrayNum = Array->Num();
diff --git a/Engine/Source/Runtime/Engine/Public/Net/RepLayout.h b/Engine/Source/Runtime/Engine/Public/Net/RepLayout.h
--- a/Engine/Source/Runtime/Engine/Public/Net/RepLayout.h
+++ b/Engine/Source/Runtime/Engine/Public/Net/RepLayout.h
@@ -796,7 +796,8 @@ enum class ERepLayoutCmdFlags : uint8
 {
 	None					= 0,		//! No flags.
 	IsSharedSerialization	= (1 << 0),	//! Indicates the property is eligible for shared serialization.
-	IsStruct				= (1 << 1)	//! This is a struct property.
+	IsStruct				= (1 << 1),	//! This is a struct property.
+	IsEmptyArrayStruct		= (1 << 2),	//! This is an ArrayProperty whose InnerProperty has no replicated properties.
 };
 
 ENUM_CLASS_FLAGS(ERepLayoutCmdFlags)
