When using ODSC, ensure special engine materials are fully compiled (along with default materials)
Otherwise non-default special ensure materials will fail a completeness check later

diff --git a/Engine/Source/Runtime/Engine/Private/Materials/MaterialShared.cpp b/Engine/Source/Runtime/Engine/Private/Materials/MaterialShared.cpp
--- a/Engine/Source/Runtime/Engine/Private/Materials/MaterialShared.cpp
+++ b/Engine/Source/Runtime/Engine/Private/Materials/MaterialShared.cpp
@@ -2229,7 +2229,7 @@ bool FMaterial::BeginCompileShaderMap(
 
 		// we can ignore requests for synch compilation if we are compiling for a different platform than we're running, or we're a commandlet that doesn't render (e.g. cooker)
 		const bool bCanIgnoreSynchronousRequirement = (TargetPlatform && !TargetPlatform->IsRunningPlatform()) || (IsRunningCommandlet() && !IsAllowCommandletRendering());
-		const bool bSkipCompilationForODSC = !IsDefaultMaterial() && GShaderCompilingManager->IsShaderCompilationSkipped();
+		const bool bSkipCompilationForODSC = !IsSpecialEngineMaterial() && GShaderCompilingManager->IsShaderCompilationSkipped();
 		if (bSkipCompilationForODSC)
 		{
 			// Force compilation off.
