Tweaks and fixes to anim node messages API

Fixes issues when scopes are added to anim instance proxies that do not have an anim class (such as editor preview instances).
Added optional scopes for use in cases where generic code needs to skip a scope push based around some context (sync is an example of this usage).

diff --git a/Engine/Source/Runtime/Engine/Private/Animation/AnimNodeMessages.cpp b/Engine/Source/Runtime/Engine/Private/Animation/AnimNodeMessages.cpp
--- a/Engine/Source/Runtime/Engine/Private/Animation/AnimNodeMessages.cpp
+++ b/Engine/Source/Runtime/Engine/Private/Animation/AnimNodeMessages.cpp
@@ -15,10 +15,19 @@ FScopedGraphTag::FScopedGraphTag(const FAnimationBaseContext& InContext, FName I
 	check(SharedContext);
 
 	int32 NodeId = InContext.GetCurrentNodeId();
-	FStructProperty* NodeProperty = InContext.AnimInstanceProxy->GetAnimClassInterface()->GetAnimNodeProperties()[NodeId];
-	FAnimNode_Base* AnimNode = NodeProperty->ContainerPtrToValuePtr<FAnimNode_Base>(InContext.AnimInstanceProxy->GetAnimInstanceObject());
+	UScriptStruct* Struct = nullptr;
+	FAnimNode_Base* AnimNode = nullptr;
+	if(NodeId != INDEX_NONE)
+	{
+		if(IAnimClassInterface* Interface = InContext.AnimInstanceProxy->GetAnimClassInterface())
+		{
+			FStructProperty* NodeProperty = Interface->GetAnimNodeProperties()[NodeId];
+			Struct = NodeProperty->Struct;
+			AnimNode = NodeProperty->ContainerPtrToValuePtr<FAnimNode_Base>(InContext.AnimInstanceProxy->GetAnimInstanceObject());
+		}
+	}
 
-	SharedContext->MessageStack.PushTag(AnimNode, NodeProperty->Struct, Tag);
+	SharedContext->MessageStack.PushTag(AnimNode, Struct, Tag);
 }
 
 FScopedGraphTag::~FScopedGraphTag()
@@ -39,10 +48,19 @@ void FScopedGraphMessage::PushMessage(const FAnimationBaseContext& InContext, TS
 	MessageType = InMessageType;
 
 	int32 NodeId = InContext.GetCurrentNodeId();
-	FStructProperty* NodeProperty = InContext.AnimInstanceProxy->GetAnimClassInterface()->GetAnimNodeProperties()[NodeId];
-	FAnimNode_Base* AnimNode = NodeProperty->ContainerPtrToValuePtr<FAnimNode_Base>(InContext.AnimInstanceProxy->GetAnimInstanceObject());
+	UScriptStruct* Struct = nullptr;
+	FAnimNode_Base* AnimNode = nullptr;
+	if(NodeId != INDEX_NONE)
+	{
+		if(IAnimClassInterface* Interface = InContext.AnimInstanceProxy->GetAnimClassInterface())
+		{
+			FStructProperty* NodeProperty = Interface->GetAnimNodeProperties()[NodeId];
+			Struct = NodeProperty->Struct;
+			AnimNode = NodeProperty->ContainerPtrToValuePtr<FAnimNode_Base>(InContext.AnimInstanceProxy->GetAnimInstanceObject());
+		}
+	}
 
-	SharedContext->MessageStack.PushMessage(MessageType, AnimNode, NodeProperty->Struct, InMessage);
+	SharedContext->MessageStack.PushMessage(MessageType, AnimNode, Struct, InMessage);
 }
 
 void FScopedGraphMessage::PopMessage()
diff --git a/Engine/Source/Runtime/Engine/Public/Animation/AnimNodeMessages.h b/Engine/Source/Runtime/Engine/Public/Animation/AnimNodeMessages.h
--- a/Engine/Source/Runtime/Engine/Public/Animation/AnimNodeMessages.h
+++ b/Engine/Source/Runtime/Engine/Public/Animation/AnimNodeMessages.h
@@ -115,6 +115,38 @@ struct TScopedGraphMessage : FScopedGraphMessage
 	}
 };
 
+// Optionally pushes a message onto the shared context stack
+template<typename TGraphMessageType>
+struct TOptionalScopedGraphMessage : FScopedGraphMessage
+{
+	template<typename... TArgs>
+	TOptionalScopedGraphMessage(bool bInCondition, const FAnimationBaseContext& InContext, TArgs&&... Args)
+		: FScopedGraphMessage(InContext)
+		, bCondition(bInCondition)
+	{
+		static_assert(TIsDerivedFrom<TGraphMessageType, IGraphMessage>::IsDerived, "Argument TGraphMessageType must derive from IGraphMessage");
+		static_assert(!TIsSame<TGraphMessageType, IGraphMessage>::Value, "Argument TGraphMessageType must not be IGraphMessage");
+
+		if(bInCondition)
+		{
+			TSharedRef<TGraphMessageType> Message = MakeShared<TGraphMessageType>(Forward<TArgs>(Args)...);
+			PushMessage(InContext, Message, Message->GetTypeName());
+		}
+	}
+
+	~TOptionalScopedGraphMessage()
+	{
+		if(bCondition)
+		{
+			PopMessage();
+		}
+	}
+
+private:
+	// Record of condition
+	bool bCondition;
+};
+
 // Stack of tags & events used to track context during graph execution
 struct ENGINE_API FMessageStack
 {
