U5 - Fix some Concurrency issues in Parallel Scene Proxy creation.

diff --git a/Engine/Source/Runtime/Engine/Private/PrimitiveSceneProxy.cpp b/Engine/Source/Runtime/Engine/Private/PrimitiveSceneProxy.cpp
--- a/Engine/Source/Runtime/Engine/Private/PrimitiveSceneProxy.cpp
+++ b/Engine/Source/Runtime/Engine/Private/PrimitiveSceneProxy.cpp
@@ -295,10 +295,10 @@ FPrimitiveSceneProxy::FPrimitiveSceneProxy(const UPrimitiveComponent* InComponen
 		{
 			if (MaterialInterface)
 			{
-				UMaterial* Material = MaterialInterface->GetMaterial();
+				const UMaterial* Material = MaterialInterface->GetMaterial_Concurrent();
 				if (const FMaterialResource* MaterialResource = Material->GetMaterialResource(FeatureLevel))
 				{
-					if (IsInGameThread())
+					if (IsInParallelGameThread() || IsInGameThread())
 					{
 						bAlwaysHasVelocity = MaterialResource->MaterialModifiesMeshPosition_GameThread();
 					}
