Fixes for ODSC when using COTF...
* Make sure we fully load/compile shader maps for both default and special engine materials
* Fix compile requests for material instances with static permutation resources

diff --git a/Engine/Source/Runtime/Engine/Private/Materials/MaterialShader.cpp b/Engine/Source/Runtime/Engine/Private/Materials/MaterialShader.cpp
--- a/Engine/Source/Runtime/Engine/Private/Materials/MaterialShader.cpp
+++ b/Engine/Source/Runtime/Engine/Private/Materials/MaterialShader.cpp
@@ -1217,7 +1217,7 @@ void FMaterialShaderMap::SaveForRemoteRecompile(FArchive& Ar, const TMap<FString
 	}
 }
 
-void FMaterialShaderMap::LoadForRemoteRecompile(FArchive& Ar, EShaderPlatform ShaderPlatform, const TArray<FString>& MaterialsForShaderMaps)
+void FMaterialShaderMap::LoadForRemoteRecompile(FArchive& Ar, EShaderPlatform ShaderPlatform, TArray<UMaterialInterface*>& OutLoadedMaterials)
 {
 	int32 MapSize;
 	Ar << MapSize;
@@ -1228,6 +1228,7 @@ void FMaterialShaderMap::LoadForRemoteRecompile(FArchive& Ar, EShaderPlatform Sh
 		Ar << MaterialName;
 
 		UMaterialInterface* MatchingMaterial = FindObjectChecked<UMaterialInterface>(NULL, *MaterialName);
+		OutLoadedMaterials.Add(MatchingMaterial);
 
 		int32 NumShaderMaps = 0;
 		Ar << NumShaderMaps;
@@ -1260,7 +1261,7 @@ void FMaterialShaderMap::LoadForRemoteRecompile(FArchive& Ar, EShaderPlatform Sh
 			for (int32 ShaderMapIndex = 0; ShaderMapIndex < LoadedShaderMaps.Num(); ShaderMapIndex++)
 			{
 				FMaterialShaderMap* LoadedShaderMap = LoadedShaderMaps[ShaderMapIndex];
-
+	
 				if (LoadedShaderMap->GetShaderPlatform() == ShaderPlatform 
 					&& LoadedShaderMap->GetShaderMapId().FeatureLevel == GetMaxSupportedFeatureLevel(ShaderPlatform))
 				{
@@ -1274,9 +1275,8 @@ void FMaterialShaderMap::LoadForRemoteRecompile(FArchive& Ar, EShaderPlatform Sh
 							|| (PassIndex == 1 && QualityLevelIndex == LoadedQualityLevel))
 						{
 							FMaterialResource* MaterialResource = MatchingMaterial->GetMaterialResource(GetMaxSupportedFeatureLevel(ShaderPlatform), (EMaterialQualityLevel::Type)QualityLevelIndex);
-
+							checkf(LoadedShaderMap->IsComplete(MaterialResource, false), TEXT("Shader map %s loaded for remote recompile is not complete"), LoadedShaderMap->GetFriendlyName());
 							MaterialResource->SetGameThreadShaderMap(LoadedShaderMap);
-							MaterialResource->RegisterInlineShaderMap(false);
 						}
 					}
 				}
diff --git a/Engine/Source/Runtime/Engine/Private/Materials/MaterialShared.cpp b/Engine/Source/Runtime/Engine/Private/Materials/MaterialShared.cpp
--- a/Engine/Source/Runtime/Engine/Private/Materials/MaterialShared.cpp
+++ b/Engine/Source/Runtime/Engine/Private/Materials/MaterialShared.cpp
@@ -590,6 +590,11 @@ ERefractionMode FMaterial::GetRefractionMode() const
 	return RM_IndexOfRefraction; 
 }
 
+bool FMaterial::IsRequiredComplete() const
+{
+	return IsDefaultMaterial() || IsSpecialEngineMaterial();
+}
+
 #if WITH_EDITOR
 void FMaterial::GetShaderMapIDsWithUnfinishedCompilation(TArray<int32>& ShaderMapIds)
 {
@@ -1660,6 +1665,10 @@ bool FMaterialResource::ShouldInlineShaderCode() const
 
 FString FMaterialResource::GetFullPath() const
 {
+	if (MaterialInstance)
+	{
+		return MaterialInstance->GetPathName();
+	}
 	if (Material)
 	{
 		return Material->GetPathName();
@@ -2084,10 +2093,10 @@ bool FMaterial::CacheShaders(const FMaterialShaderMapId& ShaderMapId, EShaderPla
 
 	UMaterialInterface* MaterialInterface = GetMaterialInterface();
 	const bool bMaterialInstance = MaterialInterface && MaterialInterface->IsA(UMaterialInstance::StaticClass());
-	const bool bSpecialEngineMaterial = !bMaterialInstance && (IsSpecialEngineMaterial() || IsDefaultMaterial());
+	const bool bRequiredComplete = !bMaterialInstance && IsRequiredComplete();
 
 	bool bShaderMapValid = (bool)GameThreadShaderMap;
-	if (bShaderMapValid && bSpecialEngineMaterial)
+	if (bShaderMapValid && bRequiredComplete)
 	{
 		// Special engine materials (default materials) are required to be complete
 		bShaderMapValid = GameThreadShaderMap->IsComplete(this, false);
@@ -2097,7 +2106,7 @@ bool FMaterial::CacheShaders(const FMaterialShaderMapId& ShaderMapId, EShaderPla
 	{
 		if (bContainsInlineShaders || FPlatformProperties::RequiresCookedData())
 		{
-			if (bSpecialEngineMaterial)
+			if (bRequiredComplete)
 			{
 				UMaterialInterface* Interface = GetMaterialInterface();
 				FString Instance;
@@ -2229,7 +2238,7 @@ bool FMaterial::BeginCompileShaderMap(
 
 		// we can ignore requests for synch compilation if we are compiling for a different platform than we're running, or we're a commandlet that doesn't render (e.g. cooker)
 		const bool bCanIgnoreSynchronousRequirement = (TargetPlatform && !TargetPlatform->IsRunningPlatform()) || (IsRunningCommandlet() && !IsAllowCommandletRendering());
-		const bool bSkipCompilationForODSC = !IsSpecialEngineMaterial() && GShaderCompilingManager->IsShaderCompilationSkipped();
+		const bool bSkipCompilationForODSC = !IsRequiredComplete() && GShaderCompilingManager->IsShaderCompilationSkipped();
 		if (bSkipCompilationForODSC)
 		{
 			// Force compilation off.
diff --git a/Engine/Source/Runtime/Engine/Private/ShaderCompiler/ShaderCompiler.cpp b/Engine/Source/Runtime/Engine/Private/ShaderCompiler/ShaderCompiler.cpp
--- a/Engine/Source/Runtime/Engine/Private/ShaderCompiler/ShaderCompiler.cpp
+++ b/Engine/Source/Runtime/Engine/Private/ShaderCompiler/ShaderCompiler.cpp
@@ -4557,19 +4557,16 @@ void ProcessCookOnTheFlyShaders(bool bReloadGlobalShaders, const TArray<uint8>&
 		// parse the shaders
 		FMemoryReader MemoryReader(MeshMaterialMaps, true);
 		FNameAsStringProxyArchive Ar(MemoryReader);
-		FMaterialShaderMap::LoadForRemoteRecompile(Ar, GMaxRHIShaderPlatform, MaterialsToLoad);
+
+		TArray<UMaterialInterface*> LoadedMaterials;
+		FMaterialShaderMap::LoadForRemoteRecompile(Ar, GMaxRHIShaderPlatform, LoadedMaterials);
 
 		// gather the shader maps to reattach
-		for (TObjectIterator<UMaterial> It; It; ++It)
+		for (UMaterialInterface* Material : LoadedMaterials)
 		{
-			UpdateContext.AddMaterial(*It);
+			Material->RecacheUniformExpressions(true);
+			UpdateContext.AddMaterialInterface(Material);
 		}
-
-		// fixup uniform expressions
-		UMaterialInterface::RecacheAllMaterialUniformExpressions(true);
-
-		// Need to recache all cached mesh draw commands, as they store pointers to material uniform buffers which we just invalidated.
-		GetRendererModule().UpdateStaticDrawLists();
 	}
 }
 
diff --git a/Engine/Source/Runtime/Engine/Public/MaterialShared.h b/Engine/Source/Runtime/Engine/Public/MaterialShared.h
--- a/Engine/Source/Runtime/Engine/Public/MaterialShared.h
+++ b/Engine/Source/Runtime/Engine/Public/MaterialShared.h
@@ -1174,7 +1174,7 @@ public:
 
 	/** Serializes a shader map to an archive (used with recompiling shaders for a remote console) */
 	ENGINE_API static void SaveForRemoteRecompile(FArchive& Ar, const TMap<FString, TArray<TRefCountPtr<FMaterialShaderMap> > >& CompiledShaderMaps);
-	ENGINE_API static void LoadForRemoteRecompile(FArchive& Ar, EShaderPlatform ShaderPlatform, const TArray<FString>& MaterialsForShaderMaps);
+	ENGINE_API static void LoadForRemoteRecompile(FArchive& Ar, EShaderPlatform ShaderPlatform, TArray<UMaterialInterface*>& OutLoadedMaterials);
 
 #if WITH_EDITOR
 	/** Returns the maximum number of texture samplers used by any shader in this shader map. */
@@ -1748,6 +1748,9 @@ public:
 	virtual bool IsPersistent() const = 0;
 	virtual UMaterialInterface* GetMaterialInterface() const { return NULL; }
 
+	/** Is the material required to be complete? Default materials, special engine materials, etc */
+	ENGINE_API bool IsRequiredComplete() const;
+
 #if WITH_EDITOR
 	/**
 	* Called when compilation of an FMaterial finishes, after the GameThreadShaderMap is set and the render command to set the RenderThreadShaderMap is queued
