Updated DrawOnscreenDebugMessages to apply inverse DPI scale when determining the Y position of the next message. The canvas expects unscaled positions so it can apply it's own scaling, but MessageTextItem.DrawnSize reflects the actual size. As a result, we need to inverse the DPI scale before adding to the Y position so that we aren't applying the scaling twice

diff --git a/Engine/Source/Runtime/Engine/Private/UnrealEngine.cpp b/Engine/Source/Runtime/Engine/Private/UnrealEngine.cpp
--- a/Engine/Source/Runtime/Engine/Private/UnrealEngine.cpp
+++ b/Engine/Source/Runtime/Engine/Private/UnrealEngine.cpp
@@ -10352,6 +10352,7 @@ float UEngine::DrawOnscreenDebugMessages(UWorld* World, FViewport* Viewport, FCa
 	static TFrameValue<bool> HasUpdatedScreenDebugMessages;
 
 	int32 YPos = MessageY;
+	const float InverseDPI = 1.f / Canvas->GetDPIScale();
 	const int32 MaxYPos = CanvasObject ? CanvasObject->SizeY : 700;
 	{
 		// Because some components add their message in concurrent work, we need a CS here.
@@ -10370,7 +10371,7 @@ float UEngine::DrawOnscreenDebugMessages(UWorld* World, FViewport* Viewport, FCa
 				MessageTextItem.SetColor(Message.DisplayColor);
 				MessageTextItem.Scale = Message.TextScale;
 				Canvas->DrawItem(MessageTextItem, FVector2D(MessageX, YPos));
-				YPos += MessageTextItem.DrawnSize.Y * 1.15f;
+				YPos += MessageTextItem.DrawnSize.Y * InverseDPI * 1.15f;
 			}
 			if (!HasUpdatedScreenDebugMessages.IsSet())
 			{
@@ -10396,7 +10397,7 @@ float UEngine::DrawOnscreenDebugMessages(UWorld* World, FViewport* Viewport, FCa
 				MessageTextItem.SetColor(Message.DisplayColor);
 				MessageTextItem.Scale = Message.TextScale;
 				Canvas->DrawItem(MessageTextItem, FVector2D(MessageX, YPos));
-				YPos += MessageTextItem.DrawnSize.Y * 1.15f;
+				YPos += MessageTextItem.DrawnSize.Y * InverseDPI * 1.15f;
 			}
 			if (!HasUpdatedScreenDebugMessages.IsSet())
 			{
