Fxed depth bias in FPrimitiveSceneProxy::DrawArrowHead & FPrimitiveSceneProxy::DrawArc

diff --git a/Engine/Source/Runtime/Engine/Private/PrimitiveSceneProxy.cpp b/Engine/Source/Runtime/Engine/Private/PrimitiveSceneProxy.cpp
--- a/Engine/Source/Runtime/Engine/Private/PrimitiveSceneProxy.cpp
+++ b/Engine/Source/Runtime/Engine/Private/PrimitiveSceneProxy.cpp
@@ -933,6 +933,7 @@ void FPrimitiveSceneProxy::DrawArc(FPrimitiveDrawInterface* PDI, const FVector&
 	}
 
 	const float ARC_PTS_SCALE = 1.0f / (float)Segments;
+	const float DepthBias = 0.0f;
 
 	const float X0 = Start.X;
 	const float Y0 = Start.Y;
@@ -949,7 +950,7 @@ void FPrimitiveSceneProxy::DrawArc(FPrimitiveDrawInterface* PDI, const FVector&
 		const float Y = Y0 + Dy * U;
 		const float Z = Z0 + Dz * U + (Length*Height) * (1-(U*2-1)*(U*2-1));
 
-		PDI->DrawLine( FVector(Px, Py, Pz), FVector(X, Y, Z), Color, SDPG_World, Thickness, bScreenSpace);
+		PDI->DrawLine( FVector(Px, Py, Pz), FVector(X, Y, Z), Color, SDPG_World, Thickness, DepthBias, bScreenSpace);
 
 		Px = X; Py = Y; Pz = Z;
 	}
@@ -960,8 +961,9 @@ void FPrimitiveSceneProxy::DrawArrowHead(FPrimitiveDrawInterface* PDI, const FVe
 	const FVector Forward = (Origin - Tip).GetUnsafeNormal();
 	const FVector Right = FVector::CrossProduct(Forward, FVector::UpVector);
 	const float HalfWidth = Size / 3.0f;
-	PDI->DrawLine(Tip, Tip + Forward * Size + Right * HalfWidth, Color, DepthPriorityGroup, Thickness, bScreenSpace);
-	PDI->DrawLine(Tip, Tip + Forward * Size - Right * HalfWidth, Color, DepthPriorityGroup, Thickness, bScreenSpace);
+	const float DepthBias = 0.0f;
+	PDI->DrawLine(Tip, Tip + Forward * Size + Right * HalfWidth, Color, DepthPriorityGroup, Thickness, DepthBias, bScreenSpace);
+	PDI->DrawLine(Tip, Tip + Forward * Size - Right * HalfWidth, Color, DepthPriorityGroup, Thickness, DepthBias, bScreenSpace);
 }
 
 
