Workaround a race condition in ENQUEUE_RENDER_COMMAND by avoiding it in SyncUVChannelData during async mesh compilation

diff --git a/Engine/Source/Runtime/Engine/Private/StaticMesh.cpp b/Engine/Source/Runtime/Engine/Private/StaticMesh.cpp
--- a/Engine/Source/Runtime/Engine/Private/StaticMesh.cpp
+++ b/Engine/Source/Runtime/Engine/Private/StaticMesh.cpp
@@ -1979,10 +1979,22 @@ void FStaticMeshRenderData::SyncUVChannelData(const TArray<FStaticMaterial>& Obj
 		UpdateData->Add(StaticMaterial.UVChannelData);
 	}
 
-	ENQUEUE_RENDER_COMMAND(SyncUVChannelData)([this, UpdateData = MoveTemp(UpdateData)](FRHICommandListImmediate& RHICmdList)
+	// SyncUVChannelData can be called from any thread during async mesh compilation. 
+	// There is currently multiple race conditions in ENQUEUE_RENDER_COMMAND making it unsafe to be called from
+	// any other thread than rendering or game because of the render thread suspension mecanism.
+	// We sidestep the issue here by avoiding a call to ENQUEUE_RENDER_COMMAND if the resource has not been initialized and is still unknown
+	// to the render thread.
+	if (bIsInitialized)
+	{
+		ENQUEUE_RENDER_COMMAND(SyncUVChannelData)([this, UpdateData = MoveTemp(UpdateData)](FRHICommandListImmediate& RHICmdList)
+		{
+			FMemory::Memswap(&UVChannelDataPerMaterial, UpdateData.Get(), sizeof(UVChannelDataPerMaterial));
+		});
+	}
+	else
 	{
 		FMemory::Memswap(&UVChannelDataPerMaterial, UpdateData.Get(), sizeof(UVChannelDataPerMaterial));
-	});
+	}
 }
 
 /*------------------------------------------------------------------------------
