Fix external actor refrences gathering archive to properly walk the outer chain when encountering subobject references.

diff --git a/Engine/Source/Runtime/Engine/Private/ExternalActorsUtils.cpp b/Engine/Source/Runtime/Engine/Private/ExternalActorsUtils.cpp
--- a/Engine/Source/Runtime/Engine/Private/ExternalActorsUtils.cpp
+++ b/Engine/Source/Runtime/Engine/Private/ExternalActorsUtils.cpp
@@ -7,7 +7,7 @@
 class FArchiveGatherExternalActorRefs : public FArchiveUObject
 {
 public:
-	FArchiveGatherExternalActorRefs(UObject* InRoot, TArray<AActor*>& InActorReferences)
+	FArchiveGatherExternalActorRefs(UObject* InRoot, TSet<AActor*>& InActorReferences)
 		: Root(InRoot)
 		, ActorReferences(InActorReferences)
 	{
@@ -17,12 +17,14 @@ public:
 		ArIsObjectReferenceCollector = true;
 		ArShouldSkipBulkData = true;
 
+		SubObjects.Add(Root);
+
 		Root->Serialize(*this);
 	}
 
 	virtual FArchive& operator<<(UObject*& Obj) override
 	{
-		if (Obj && (Obj != Root) && !Obj->IsTemplate() && !Obj->HasAnyFlags(RF_Transient))
+		if (Obj && !Obj->IsTemplate() && !Obj->HasAnyFlags(RF_Transient))
 		{
 			bool bWasAlreadyInSet;
 			SubObjects.Add(Obj, &bWasAlreadyInSet);
@@ -43,9 +45,15 @@ public:
 private:
 	void HandleObjectReference(UObject* Obj)
 	{
-		if(Obj->IsA<AActor>())
+		AActor* Actor = Cast<AActor>(Obj);
+
+		if (!Actor)
+		{
+			Actor = Obj->GetTypedOuter<AActor>();
+		}
+
+		if(Actor)
 		{
-			AActor* Actor = (AActor*)Obj;
 			AActor* TopParentActor = Actor;
 			while(TopParentActor->GetParentActor())
 			{
@@ -54,7 +62,7 @@ private:
 
 			check(TopParentActor);
 
-			if (TopParentActor->IsPackageExternal())
+			if (TopParentActor->IsPackageExternal() && (TopParentActor != Root))
 			{
 				ActorReferences.Add(TopParentActor);
 			}
@@ -62,13 +70,13 @@ private:
 	}
 
 	UObject* Root;
-	TArray<AActor*>& ActorReferences;
+	TSet<AActor*>& ActorReferences;
 	TSet<UObject*> SubObjects;
 };
 
 TArray<AActor*> ExternalActorsUtils::GetExternalActorReferences(UObject* Root)
 {
-	TArray<AActor*> Result;
+	TSet<AActor*> Result;
 	FArchiveGatherExternalActorRefs Ar(Root, Result);
-	return MoveTemp(Result);
+	return Result.Array();
 }
