Strata save a DWORD when volume has not any scattering pluged in.
Simple bool for now, can be extended to a bit field of properties later.

diff --git a/Engine/Source/Runtime/Engine/Private/Materials/MaterialExpressions.cpp b/Engine/Source/Runtime/Engine/Private/Materials/MaterialExpressions.cpp
--- a/Engine/Source/Runtime/Engine/Private/Materials/MaterialExpressions.cpp
+++ b/Engine/Source/Runtime/Engine/Private/Materials/MaterialExpressions.cpp
@@ -19264,16 +19264,16 @@ int32 UMaterialExpressionStrataVolumeBSDF::Compile(class FMaterialCompiler* Comp
 	int32 NormalCodeChunk = Compiler->VertexNormal();
 	uint8 SharedNormalIndex = StrataCompilationInfoCreateSharedNormal(Compiler, NormalCodeChunk);
 
-	const bool bHasScattering = Albedo.GetTracedInput().Expression != nullptr;	// STRATA_TODO use when we switch to byte per pixel
+	const bool bHasScattering = Albedo.GetTracedInput().Expression != nullptr;
 
 	int32 OutputCodeChunk = Compiler->StrataVolumeBSDF(
-		bHasScattering							? Albedo.Compile(Compiler)			: Compiler->Constant3(0.0f, 0.0f, 0.0f),
+		bHasScattering							? Albedo.Compile(Compiler)			: Compiler->Constant3(0.0f, 0.0f, 0.0f), // Default must be 0 to remove memory used by Albedo when it is not pluged (bHasScattering==false)
 		Extinction.GetTracedInput().Expression	? Extinction.Compile(Compiler)		: Compiler->Constant3(0.0f, 0.0f, 0.0f),
 		Anisotropy.GetTracedInput().Expression	? Anisotropy.Compile(Compiler)		: Compiler->Constant(0.0f),
 		Thickness.GetTracedInput().Expression	? Thickness.Compile(Compiler)		: Compiler->Constant(0.001f), // default = 1mm
 		NormalCodeChunk,
 		SharedNormalIndex);
-	StrataCompilationInfoCreateSingleBSDFMaterial(Compiler, OutputCodeChunk, SharedNormalIndex, STRATA_BSDF_TYPE_VOLUME);
+	StrataCompilationInfoCreateSingleBSDFMaterial(Compiler, OutputCodeChunk, SharedNormalIndex, STRATA_BSDF_TYPE_VOLUME, bHasScattering);
 
 	return OutputCodeChunk;
 }
diff --git a/Engine/Source/Runtime/Engine/Private/Materials/StrataMaterial.cpp b/Engine/Source/Runtime/Engine/Private/Materials/StrataMaterial.cpp
--- a/Engine/Source/Runtime/Engine/Private/Materials/StrataMaterial.cpp
+++ b/Engine/Source/Runtime/Engine/Private/Materials/StrataMaterial.cpp
@@ -49,13 +49,14 @@ uint8 StrataCompilationInfoCreateSharedNormal(FMaterialCompiler* Compiler, int32
 	return Compiler->StrataCompilationInfoRegisterSharedNormalIndex(NormalCodeChunk, TangentCodeChunk);
 }
 
-void StrataCompilationInfoCreateSingleBSDFMaterial(FMaterialCompiler* Compiler, int32 CodeChunk, uint8 SharedNormalIndex, uint8 BSDFType)
+void StrataCompilationInfoCreateSingleBSDFMaterial(FMaterialCompiler* Compiler, int32 CodeChunk, uint8 SharedNormalIndex, uint8 BSDFType, bool bVolumeHasScattering)
 {
 	FStrataMaterialCompilationInfo StrataInfo;
 	StrataInfo.LayerCount = 1;
 	StrataInfo.Layers[0].BSDFCount = 1;
 	StrataInfo.Layers[0].BSDFs[0].Type = BSDFType;
 	StrataInfo.Layers[0].BSDFs[0].SharedNormalIndex = SharedNormalIndex;
+	StrataInfo.Layers[0].BSDFs[0].bVolumeHasScattering = bVolumeHasScattering;
 	UpdateTotalBSDFCount(StrataInfo);
 	Compiler->StrataCompilationInfoRegisterCodeChunk(CodeChunk, StrataInfo);
 }
@@ -194,8 +195,7 @@ FStrataMaterialAnalysisResult StrataCompilationInfoMaterialAnalysis(FMaterialCom
 			{
 				Result.RequestedByteCount += UintByteSize;
 
-				const bool bMayHaveScattering = true;
-				if (bMayHaveScattering)
+				if (BSDF.bVolumeHasScattering)
 				{
 					Result.RequestedByteCount += UintByteSize;
 				}
diff --git a/Engine/Source/Runtime/Engine/Private/Materials/StrataMaterial.h b/Engine/Source/Runtime/Engine/Private/Materials/StrataMaterial.h
--- a/Engine/Source/Runtime/Engine/Private/Materials/StrataMaterial.h
+++ b/Engine/Source/Runtime/Engine/Private/Materials/StrataMaterial.h
@@ -18,6 +18,7 @@ struct FStrataMaterialCompilationInfo
 	{
 		uint8 Type;
 		uint8 SharedNormalIndex;
+		bool  bVolumeHasScattering;
 	};
 
 	struct FLayer
@@ -43,7 +44,7 @@ FString GetStrataBSDFName(uint8 BSDFType);
 
 uint8 StrataCompilationInfoCreateSharedNormal(FMaterialCompiler* Compiler, int32 NormalCodeChunk);
 uint8 StrataCompilationInfoCreateSharedNormal(FMaterialCompiler* Compiler, int32 NormalCodeChunk, int32 TangentCodeChunk);
-void  StrataCompilationInfoCreateSingleBSDFMaterial(FMaterialCompiler* Compiler, int32 CodeChunk, uint8 SharedNormalIndex, uint8 BSDFType);
+void  StrataCompilationInfoCreateSingleBSDFMaterial(FMaterialCompiler* Compiler, int32 CodeChunk, uint8 SharedNormalIndex, uint8 BSDFType, bool bVolumeHasScattering = false);
 
 FStrataMaterialCompilationInfo StrataCompilationInfoAdd(FMaterialCompiler* Compiler, const FStrataMaterialCompilationInfo& A, const FStrataMaterialCompilationInfo& B);
 FStrataMaterialCompilationInfo StrataCompilationInfoMultiply(FMaterialCompiler* Compiler, const FStrataMaterialCompilationInfo& A);
