Made StrataMultiply have a default weight of 1 when not plugedin.

diff --git a/Engine/Source/Runtime/Engine/Classes/Materials/MaterialExpressionStrata.h b/Engine/Source/Runtime/Engine/Classes/Materials/MaterialExpressionStrata.h
--- a/Engine/Source/Runtime/Engine/Classes/Materials/MaterialExpressionStrata.h
+++ b/Engine/Source/Runtime/Engine/Classes/Materials/MaterialExpressionStrata.h
@@ -336,7 +336,7 @@ class UMaterialExpressionStrataAdd : public UMaterialExpression
 	//~ End UMaterialExpression Interface
 };
 
-UCLASS(MinimalAPI, collapsecategories, hidecategories = Object)
+UCLASS(MinimalAPI, DisplayName = "Strata Multiply Weight", collapsecategories, hidecategories = Object)
 class UMaterialExpressionStrataMultiply : public UMaterialExpression
 {
 	GENERATED_UCLASS_BODY()
diff --git a/Engine/Source/Runtime/Engine/Private/Materials/MaterialExpressions.cpp b/Engine/Source/Runtime/Engine/Private/Materials/MaterialExpressions.cpp
--- a/Engine/Source/Runtime/Engine/Private/Materials/MaterialExpressions.cpp
+++ b/Engine/Source/Runtime/Engine/Private/Materials/MaterialExpressions.cpp
@@ -19580,15 +19580,11 @@ int32 UMaterialExpressionStrataMultiply::Compile(class FMaterialCompiler* Compil
 {
 	if (!A.GetTracedInput().Expression)
 	{
-		return Compiler->Errorf(TEXT("Missing Top input"));
-	}
-	if (!Weight.GetTracedInput().Expression)
-	{
-		return Compiler->Errorf(TEXT("Missing Base input"));
+		return Compiler->Errorf(TEXT("Missing A input"));
 	}
 
 	int32 ACodeChunk = A.Compile(Compiler);
-	int32 WeightCodeChunk = Weight.Compile(Compiler);
+	int32 WeightCodeChunk = Weight.GetTracedInput().Expression ? Weight.Compile(Compiler) : Compiler->Constant(1.0f);
 
 	int32 OutputCodeChunk = Compiler->StrataMultiply(ACodeChunk, WeightCodeChunk);
 
