Prevent FCanvasTileItem from referencing proxy resources (valid while the textures are being build).
This prevents accessing deleted data in FBatchedElements::Draw().

diff --git a/Engine/Source/Runtime/Engine/Private/UserInterface/Canvas.cpp b/Engine/Source/Runtime/Engine/Private/UserInterface/Canvas.cpp
--- a/Engine/Source/Runtime/Engine/Private/UserInterface/Canvas.cpp
+++ b/Engine/Source/Runtime/Engine/Private/UserInterface/Canvas.cpp
@@ -47,6 +47,21 @@ FCanvasWordWrapper::FCanvasWordWrapper()
 {
 }
 
+FTexture* GetTextureForCanvasItem(UTexture* RenderTexture)
+{
+	// Only assign a valid texture for the canvas item. Proxy textures, from the UTexture implementation, are not valid here.
+	// This is because the proxy textures get destroyed on the render thread as soon as the build completes, which could happen
+	// before the canvas item renders.
+	if (RenderTexture && RenderTexture->Resource && !RenderTexture->Resource->IsProxy())
+	{
+		return RenderTexture->Resource;
+	}
+	else
+	{
+		return GWhiteTexture;
+	}
+}
+
 void FCanvasWordWrapper::Execute(const TCHAR* const InString, const FTextSizingParameters& InParameters, TArray<FWrappedStringElement>& OutStrings, FWrappedLineData* const OutWrappedLineData)
 {
 	SCOPE_CYCLE_COUNTER(STAT_Canvas_WordWrappingTime);
@@ -1534,7 +1549,7 @@ void UCanvas::DrawTile( UTexture* Tex, float X, float Y, float XL, float YL, flo
 		float TexSurfaceWidth= Tex->GetSurfaceWidth();
 		float TexSurfaceHeight = Tex->GetSurfaceHeight();
 
-		FCanvasTileItem TileItem( FVector2D( X, Y ), Tex->Resource,  FVector2D( w, h ),  
+		FCanvasTileItem TileItem( FVector2D( X, Y ), GetTextureForCanvasItem(Tex),  FVector2D( w, h ),  
 			FVector2D(U / TexSurfaceWidth, V / TexSurfaceHeight),
 			FVector2D(U / TexSurfaceWidth + UL / TexSurfaceWidth * w / XL, V / TexSurfaceHeight + VL / TexSurfaceHeight * h / YL),
 			DrawColor );
@@ -1966,21 +1981,6 @@ void UCanvas::K2_DrawLine(FVector2D ScreenPositionA, FVector2D ScreenPositionB,
 	}
 }
 
-FTexture* GetTextureForCanvasItem(UTexture* RenderTexture)
-{
-	// Only assign a valid texture for the canvas item. Proxy textures, from the UTexture implementation, are not valid here.
-	// This is because the proxy textures get destroyed on the render thread as soon as the build completes, which could happen
-	// before the canvas item renders.
-	if (RenderTexture && RenderTexture->Resource && !RenderTexture->Resource->IsProxy())
-	{
-		return RenderTexture->Resource;
-	}
-	else
-	{
-		return GWhiteTexture;
-	}
-}
-
 void UCanvas::K2_DrawTexture(UTexture* RenderTexture, FVector2D ScreenPosition, FVector2D ScreenSize, FVector2D CoordinatePosition, FVector2D CoordinateSize, FLinearColor RenderColor, EBlendMode BlendMode, float Rotation, FVector2D PivotPoint)
 {
 	if (ScreenSize.X > 0.0f && ScreenSize.Y > 0.0f && Canvas)
@@ -2150,4 +2150,4 @@ void FDisplayDebugManager::AddColumnIfNeeded()
 			CurrentPos.X = NextColumXPos + YStep * 2.f;
 		}
 	}
-}
+}
\ No newline at end of file
