Make UMaterial::CachedExpressionData stored by pointer, rather than by value.
 No longer a UPROPERTY(), instead only serialized based on an addition bool property.
Take away is this cached data is no longer serialized by the editor (it's always created on load in editor anyway, so no reason to serialize)
This keeps the UPROPERTIES stored in the cached data from appearing in the reference viewer (in the editor at least), which is important for artists/content workflows
This may seem like a convoluted way to accomplish this goal, but the way the reference viewer displays references isn't very customizable....plus this change should improve editor load times a bit, and reduce editor material disk footprint

diff --git a/Engine/Source/Runtime/Engine/Classes/Materials/Material.h b/Engine/Source/Runtime/Engine/Classes/Materials/Material.h
--- a/Engine/Source/Runtime/Engine/Classes/Materials/Material.h
+++ b/Engine/Source/Runtime/Engine/Classes/Materials/Material.h
@@ -980,7 +980,9 @@ private:
 	FThreadSafeBool ReleasedByRT;
 
 	UPROPERTY()
-	FMaterialCachedExpressionData CachedExpressionData;
+	bool bSavedCachedExpressionData;
+
+	FMaterialCachedExpressionData* CachedExpressionData;
 
 #if WITH_EDITORONLY_DATA
 	UPROPERTY()
@@ -989,7 +991,7 @@ private:
 #endif // WITH_EDITORONLY_DATA
 public:
 
-	const FMaterialCachedExpressionData& GetCachedExpressionData() const { return CachedExpressionData; }
+	const FMaterialCachedExpressionData& GetCachedExpressionData() const { check(CachedExpressionData); return *CachedExpressionData; }
 
 	//~ Begin UMaterialInterface Interface.
 	ENGINE_API virtual UMaterial* GetMaterial() override;
@@ -1712,7 +1714,7 @@ public:
 #endif // WITH_EDITOR
 
 	/** Appends textures referenced by expressions, including nested functions. */
-	ENGINE_API virtual TArrayView<UObject* const> GetReferencedTextures() const override final { return CachedExpressionData.ReferencedTextures; }
+	ENGINE_API virtual TArrayView<UObject* const> GetReferencedTextures() const override final { return GetCachedExpressionData().ReferencedTextures; }
 
 protected:
 
diff --git a/Engine/Source/Runtime/Engine/Classes/Materials/MaterialInstance.h b/Engine/Source/Runtime/Engine/Classes/Materials/MaterialInstance.h
--- a/Engine/Source/Runtime/Engine/Classes/Materials/MaterialInstance.h
+++ b/Engine/Source/Runtime/Engine/Classes/Materials/MaterialInstance.h
@@ -442,14 +442,9 @@ private:
 	FStaticParameterSet StaticParameters;
 
 	UPROPERTY()
-	FMaterialCachedParameters CachedLayerParameters;
+	bool bSavedCachedData;
 
-	/**
-	 * Cached texture references from all expressions in the material (including nested functions).
-	 * This is used to link uniform texture expressions which were stored in the DDC with the UTextures that they reference.
-	 */
-	UPROPERTY()
-	TArray<UObject*> CachedReferencedTextures;
+	FMaterialInstanceCachedData* CachedData;
 
 #ifdef WITH_EDITOR
 	mutable TOptional<FStaticParameterSet> CachedStaticParameterValues;
@@ -670,7 +665,7 @@ public:
 #endif // WITH_EDITOR
 
 	/** Appends textures referenced by expressions, including nested functions. */
-	ENGINE_API virtual TArrayView<UObject* const> GetReferencedTextures() const override final { return CachedReferencedTextures; }
+	ENGINE_API virtual TArrayView<UObject* const> GetReferencedTextures() const override final { check(CachedData); return CachedData->ReferencedTextures; }
 
 #if WITH_EDITOR
 	/** Add to the set any texture referenced by expressions, including nested functions, as well as any overrides from parameters. */
diff --git a/Engine/Source/Runtime/Engine/Private/Materials/Material.cpp b/Engine/Source/Runtime/Engine/Private/Materials/Material.cpp
--- a/Engine/Source/Runtime/Engine/Private/Materials/Material.cpp
+++ b/Engine/Source/Runtime/Engine/Private/Materials/Material.cpp
@@ -1739,27 +1739,27 @@ void UMaterial::FixupMaterialUsageAfterLoad()
 
 void UMaterial::GetAllScalarParameterInfo(TArray<FMaterialParameterInfo>& OutParameterInfo, TArray<FGuid>& OutParameterIds) const
 {
-	CachedExpressionData.Parameters.GetAllParameterInfoOfType(EMaterialParameterType::Scalar, true, OutParameterInfo, OutParameterIds);
+	GetCachedExpressionData().Parameters.GetAllParameterInfoOfType(EMaterialParameterType::Scalar, true, OutParameterInfo, OutParameterIds);
 }
 
 void UMaterial::GetAllVectorParameterInfo(TArray<FMaterialParameterInfo>& OutParameterInfo, TArray<FGuid>& OutParameterIds) const
 {
-	CachedExpressionData.Parameters.GetAllParameterInfoOfType(EMaterialParameterType::Vector, true, OutParameterInfo, OutParameterIds);
+	GetCachedExpressionData().Parameters.GetAllParameterInfoOfType(EMaterialParameterType::Vector, true, OutParameterInfo, OutParameterIds);
 }
 
 void UMaterial::GetAllTextureParameterInfo(TArray<FMaterialParameterInfo>& OutParameterInfo, TArray<FGuid>& OutParameterIds) const
 {
-	CachedExpressionData.Parameters.GetAllParameterInfoOfType(EMaterialParameterType::Texture, true, OutParameterInfo, OutParameterIds);
+	GetCachedExpressionData().Parameters.GetAllParameterInfoOfType(EMaterialParameterType::Texture, true, OutParameterInfo, OutParameterIds);
 }
 
 void UMaterial::GetAllFontParameterInfo(TArray<FMaterialParameterInfo>& OutParameterInfo, TArray<FGuid>& OutParameterIds) const
 {
-	CachedExpressionData.Parameters.GetAllParameterInfoOfType(EMaterialParameterType::Font, true, OutParameterInfo, OutParameterIds);
+	GetCachedExpressionData().Parameters.GetAllParameterInfoOfType(EMaterialParameterType::Font, true, OutParameterInfo, OutParameterIds);
 }
 
 void UMaterial::GetAllRuntimeVirtualTextureParameterInfo(TArray<FMaterialParameterInfo>& OutParameterInfo, TArray<FGuid>& OutParameterIds) const
 {
-	CachedExpressionData.Parameters.GetAllParameterInfoOfType(EMaterialParameterType::RuntimeVirtualTexture, true, OutParameterInfo, OutParameterIds);
+	GetCachedExpressionData().Parameters.GetAllParameterInfoOfType(EMaterialParameterType::RuntimeVirtualTexture, true, OutParameterInfo, OutParameterIds);
 }
 
 #if WITH_EDITORONLY_DATA
@@ -2106,24 +2106,20 @@ bool UMaterial::GetParameterDesc(const FHashedMaterialParameterInfo& ParameterIn
 }
 void UMaterial::UpdateCachedExpressionData()
 {
-	if (!GIsClient)
+	if (bSavedCachedExpressionData)
 	{
-		// Expressions may not be loaded on server, so only rebuild the data if we're the client
+		// Don't need to rebuild cached data if it was serialized
 		return;
 	}
 
+	if (!CachedExpressionData)
 	{
-		FMaterialCachedExpressionData UpdatedCachedExpressionData;
-		UpdatedCachedExpressionData.Reset();
-		FMaterialCachedExpressionContext Context(nullptr); // UMaterial have no parent
-		if (UpdatedCachedExpressionData.UpdateForExpressions(Context, Expressions, EMaterialParameterAssociation::GlobalParameter, -1))
-		{
-			// Only update our cached data if the update succeeded
-			// It's possible we could have some nullptr UMaterialExpressions, if we're loading cooked data here
-			// In that case we simply keep existing cached data
-			CachedExpressionData = MoveTemp(UpdatedCachedExpressionData);
-		}
+		CachedExpressionData = new FMaterialCachedExpressionData();
 	}
+
+	CachedExpressionData->Reset();
+	FMaterialCachedExpressionContext Context(nullptr); // UMaterial have no parent
+	CachedExpressionData->UpdateForExpressions(Context, Expressions, EMaterialParameterAssociation::GlobalParameter, -1);
 }
 #endif // WITH_EDITOR
 
@@ -2146,13 +2142,13 @@ bool UMaterial::GetScalarParameterValue(const FHashedMaterialParameterInfo& Para
 
 bool UMaterial::GetScalarParameterValue_New(const FHashedMaterialParameterInfo& ParameterInfo, float& OutValue, bool bOveriddenOnly) const
 {
-	const int32 Index = CachedExpressionData.Parameters.FindParameterIndex(EMaterialParameterType::Scalar, ParameterInfo, bOveriddenOnly);
+	const int32 Index = CachedExpressionData->Parameters.FindParameterIndex(EMaterialParameterType::Scalar, ParameterInfo, bOveriddenOnly);
 	if (Index == INDEX_NONE)
 	{
 		return false;
 	}
 
-	OutValue = CachedExpressionData.Parameters.ScalarValues[Index];
+	OutValue = CachedExpressionData->Parameters.ScalarValues[Index];
 	return true;
 }
 
@@ -2231,14 +2227,14 @@ bool UMaterial::GetScalarParameterValue_Legacy(const FHashedMaterialParameterInf
 
 bool UMaterial::IsScalarParameterUsedAsAtlasPosition(const FHashedMaterialParameterInfo& ParameterInfo, bool& OutValue, TSoftObjectPtr<UCurveLinearColor>& OutCurve, TSoftObjectPtr<UCurveLinearColorAtlas>& OutAtlas) const
 {
-	const int32 Index = CachedExpressionData.Parameters.FindParameterIndex(EMaterialParameterType::Scalar, ParameterInfo);
+	const int32 Index = CachedExpressionData->Parameters.FindParameterIndex(EMaterialParameterType::Scalar, ParameterInfo);
 	if (Index == INDEX_NONE)
 	{
 		return false;
 	}
 
-	UCurveLinearColor* Curve = CachedExpressionData.Parameters.ScalarCurveValues[Index];
-	UCurveLinearColorAtlas* Atlas = CachedExpressionData.Parameters.ScalarCurveAtlasValues[Index];
+	UCurveLinearColor* Curve = CachedExpressionData->Parameters.ScalarCurveValues[Index];
+	UCurveLinearColorAtlas* Atlas = CachedExpressionData->Parameters.ScalarCurveAtlasValues[Index];
 	if (Curve && Atlas)
 	{
 		OutCurve = TSoftObjectPtr<UCurveLinearColor>(FSoftObjectPath(Curve->GetPathName()));
@@ -2255,13 +2251,13 @@ bool UMaterial::IsScalarParameterUsedAsAtlasPosition(const FHashedMaterialParame
 
 bool UMaterial::GetScalarParameterSliderMinMax(const FHashedMaterialParameterInfo& ParameterInfo, float& OutSliderMin, float& OutSliderMax) const
 {
-	const int32 Index = CachedExpressionData.Parameters.FindParameterIndex(EMaterialParameterType::Scalar, ParameterInfo);
+	const int32 Index = CachedExpressionData->Parameters.FindParameterIndex(EMaterialParameterType::Scalar, ParameterInfo);
 	if (Index == INDEX_NONE)
 	{
 		return false;
 	}
-	OutSliderMin = CachedExpressionData.Parameters.ScalarMinMaxValues[Index].X;
-	OutSliderMax = CachedExpressionData.Parameters.ScalarMinMaxValues[Index].Y;
+	OutSliderMin = CachedExpressionData->Parameters.ScalarMinMaxValues[Index].X;
+	OutSliderMax = CachedExpressionData->Parameters.ScalarMinMaxValues[Index].Y;
 	return true;
 }
 #endif // WITH_EDITOR
@@ -2285,12 +2281,12 @@ bool UMaterial::GetVectorParameterValue(const FHashedMaterialParameterInfo& Para
 
 bool UMaterial::GetVectorParameterValue_New(const FHashedMaterialParameterInfo& ParameterInfo, FLinearColor& OutValue, bool bOveriddenOnly) const
 {
-	const int32 Index = CachedExpressionData.Parameters.FindParameterIndex(EMaterialParameterType::Vector, ParameterInfo, bOveriddenOnly);
+	const int32 Index = CachedExpressionData->Parameters.FindParameterIndex(EMaterialParameterType::Vector, ParameterInfo, bOveriddenOnly);
 	if (Index == INDEX_NONE)
 	{
 		return false;
 	}
-	OutValue = CachedExpressionData.Parameters.VectorValues[Index];
+	OutValue = CachedExpressionData->Parameters.VectorValues[Index];
 	return true;
 }
 
@@ -2366,23 +2362,23 @@ bool UMaterial::GetVectorParameterValue_Legacy(const FHashedMaterialParameterInf
 
 bool UMaterial::IsVectorParameterUsedAsChannelMask(const FHashedMaterialParameterInfo& ParameterInfo, bool& OutValue) const
 {
-	const int32 Index = CachedExpressionData.Parameters.FindParameterIndex(EMaterialParameterType::Vector, ParameterInfo);
+	const int32 Index = CachedExpressionData->Parameters.FindParameterIndex(EMaterialParameterType::Vector, ParameterInfo);
 	if (Index == INDEX_NONE)
 	{
 		return false;
 	}
-	OutValue = CachedExpressionData.Parameters.VectorUsedAsChannelMaskValues[Index];
+	OutValue = CachedExpressionData->Parameters.VectorUsedAsChannelMaskValues[Index];
 	return true;
 }
 
 bool UMaterial::GetVectorParameterChannelNames(const FHashedMaterialParameterInfo& ParameterInfo, FParameterChannelNames& OutValue) const
 {
-	const int32 Index = CachedExpressionData.Parameters.FindParameterIndex(EMaterialParameterType::Vector, ParameterInfo);
+	const int32 Index = CachedExpressionData->Parameters.FindParameterIndex(EMaterialParameterType::Vector, ParameterInfo);
 	if (Index == INDEX_NONE)
 	{
 		return false;
 	}
-	OutValue = CachedExpressionData.Parameters.VectorChannelNameValues[Index];
+	OutValue = CachedExpressionData->Parameters.VectorChannelNameValues[Index];
 	return true;
 }
 
@@ -2410,13 +2406,13 @@ bool UMaterial::GetTextureParameterValue(const FHashedMaterialParameterInfo& Par
 
 bool UMaterial::GetTextureParameterValue_New(const FHashedMaterialParameterInfo& ParameterInfo, UTexture*& OutValue, bool bOveriddenOnly) const
 {
-	const int32 Index = CachedExpressionData.Parameters.FindParameterIndex(EMaterialParameterType::Texture, ParameterInfo, bOveriddenOnly);
+	const int32 Index = CachedExpressionData->Parameters.FindParameterIndex(EMaterialParameterType::Texture, ParameterInfo, bOveriddenOnly);
 	if (Index == INDEX_NONE)
 	{
 		OutValue = nullptr;
 		return false;
 	}
-	OutValue = CachedExpressionData.Parameters.TextureValues[Index];
+	OutValue = CachedExpressionData->Parameters.TextureValues[Index];
 	return true;
 }
 
@@ -2515,13 +2511,13 @@ bool UMaterial::GetRuntimeVirtualTextureParameterValue(const FHashedMaterialPara
 
 bool UMaterial::GetRuntimeVirtualTextureParameterValue_New(const FHashedMaterialParameterInfo& ParameterInfo, URuntimeVirtualTexture*& OutValue, bool bOveriddenOnly) const
 {
-	const int32 Index = CachedExpressionData.Parameters.FindParameterIndex(EMaterialParameterType::RuntimeVirtualTexture, ParameterInfo, bOveriddenOnly);
+	const int32 Index = CachedExpressionData->Parameters.FindParameterIndex(EMaterialParameterType::RuntimeVirtualTexture, ParameterInfo, bOveriddenOnly);
 	if (Index == INDEX_NONE)
 	{
 		OutValue = nullptr;
 		return false;
 	}
-	OutValue = CachedExpressionData.Parameters.RuntimeVirtualTextureValues[Index];
+	OutValue = CachedExpressionData->Parameters.RuntimeVirtualTextureValues[Index];
 	return true;
 }
 
@@ -2602,12 +2598,12 @@ bool UMaterial::GetRuntimeVirtualTextureParameterValue_Legacy(const FHashedMater
 #if WITH_EDITOR
 bool UMaterial::GetTextureParameterChannelNames(const FHashedMaterialParameterInfo& ParameterInfo, FParameterChannelNames& OutValue) const
 {
-	const int32 Index = CachedExpressionData.Parameters.FindParameterIndex(EMaterialParameterType::Texture, ParameterInfo);
+	const int32 Index = CachedExpressionData->Parameters.FindParameterIndex(EMaterialParameterType::Texture, ParameterInfo);
 	if (Index == INDEX_NONE)
 	{
 		return false;
 	}
-	OutValue = CachedExpressionData.Parameters.TextureChannelNameValues[Index];
+	OutValue = CachedExpressionData->Parameters.TextureChannelNameValues[Index];
 	return true;
 }
 #endif // WITH_EDITOR
@@ -2634,15 +2630,15 @@ bool UMaterial::GetFontParameterValue(const FHashedMaterialParameterInfo& Parame
 
 bool UMaterial::GetFontParameterValue_New(const FHashedMaterialParameterInfo& ParameterInfo, UFont*& OutFontValue, int32& OutFontPage, bool bOveriddenOnly) const
 {
-	const int32 Index = CachedExpressionData.Parameters.FindParameterIndex(EMaterialParameterType::Font, ParameterInfo, bOveriddenOnly);
+	const int32 Index = CachedExpressionData->Parameters.FindParameterIndex(EMaterialParameterType::Font, ParameterInfo, bOveriddenOnly);
 	if (Index == INDEX_NONE)
 	{
 		OutFontValue = nullptr;
 		OutFontPage = INDEX_NONE;
 		return false;
 	}
-	OutFontValue = CachedExpressionData.Parameters.FontValues[Index];
-	OutFontPage = CachedExpressionData.Parameters.FontPageValues[Index];
+	OutFontValue = CachedExpressionData->Parameters.FontValues[Index];
+	OutFontPage = CachedExpressionData->Parameters.FontPageValues[Index];
 	return true;
 }
 
@@ -3163,6 +3159,9 @@ void UMaterial::CacheResourceShadersForCooking(EShaderPlatform ShaderPlatform, T
 
 	CacheShadersForResources(ShaderPlatform, NewResourcesToCache, EMaterialShaderPrecompileMode::Background, TargetPlatform);
 
+	// Make sure cached data is saved in cooked builds
+	bSavedCachedExpressionData = true;
+
 	OutCachedMaterialResources.Append(NewResourcesToCache);
 }
 
@@ -3291,6 +3290,19 @@ void UMaterial::Serialize(FArchive& Ar)
 #endif
 	}
 
+#if !WITH_EDITORONLY_DATA
+	checkf(bSavedCachedExpressionData, TEXT("Material %s must have saved cached expression data, if editor-only data is not present"), *GetName());
+#endif
+	if (bSavedCachedExpressionData)
+	{
+		if (Ar.IsLoading())
+		{
+			CachedExpressionData = new FMaterialCachedExpressionData();
+		}
+		UScriptStruct* Struct = FMaterialCachedExpressionData::StaticStruct();
+		Struct->SerializeTaggedProperties(Ar, (uint8*)CachedExpressionData, Struct, nullptr);
+	}
+
 #if WITH_EDITOR
 	if (Ar.UE4Ver() < VER_UE4_FLIP_MATERIAL_COORDS)
 	{
@@ -3499,7 +3511,7 @@ void UMaterial::BackwardsCompatibilityVirtualTextureOutputConversion()
 
 void UMaterial::GetQualityLevelUsage(TArray<bool, TInlineAllocator<EMaterialQualityLevel::Num> >& OutQualityLevelsUsed, EShaderPlatform ShaderPlatform, bool bCooking)
 {
-	OutQualityLevelsUsed = CachedExpressionData.QualityLevelsUsed;
+	OutQualityLevelsUsed = CachedExpressionData->QualityLevelsUsed;
 	if (OutQualityLevelsUsed.Num() == 0)
 	{
 		OutQualityLevelsUsed.AddDefaulted(EMaterialQualityLevel::Num);
@@ -3615,14 +3627,6 @@ void UMaterial::PostLoad()
 	}
 #endif // WITH_EDITORONLY_DATA
 
-	for (int32 CollectionIndex = 0; CollectionIndex < CachedExpressionData.ParameterCollectionInfos.Num(); CollectionIndex++)
-	{
-		if (CachedExpressionData.ParameterCollectionInfos[CollectionIndex].ParameterCollection)
-		{
-			CachedExpressionData.ParameterCollectionInfos[CollectionIndex].ParameterCollection->ConditionalPostLoad();
-		}
-	}
-
 	// Fixup for legacy materials which didn't recreate the lighting guid properly on duplication
 	if (GetLinker() && GetLinker()->UE4Ver() < VER_UE4_BUMPED_MATERIAL_EXPORT_GUIDS)
 	{
@@ -3720,6 +3724,14 @@ void UMaterial::PostLoad()
 	UpdateCachedExpressionData();
 #endif // WITH_EDITOR
 
+	for (int32 CollectionIndex = 0; CollectionIndex < CachedExpressionData->ParameterCollectionInfos.Num(); CollectionIndex++)
+	{
+		if (CachedExpressionData->ParameterCollectionInfos[CollectionIndex].ParameterCollection)
+		{
+			CachedExpressionData->ParameterCollectionInfos[CollectionIndex].ParameterCollection->ConditionalPostLoad();
+		}
+	}
+
 	STAT(double MaterialLoadTime = 0);
 	{
 		SCOPE_SECONDS_COUNTER(MaterialLoadTime);
@@ -3742,7 +3754,7 @@ void UMaterial::PostLoad()
 		{
 			// Before caching shader resources we have to make sure all referenced textures have been post loaded
 			// as we depend on their resources being valid.
-			for (UObject* Texture : CachedExpressionData.ReferencedTextures)
+			for (UObject* Texture : CachedExpressionData->ReferencedTextures)
 			{
 				if (Texture)
 				{
@@ -4660,6 +4672,11 @@ void UMaterial::ReleaseResources()
 		DefaultMaterialInstance->GameThread_Destroy();
 		DefaultMaterialInstance = nullptr;
 	}
+	if (CachedExpressionData)
+	{
+		delete CachedExpressionData;
+		CachedExpressionData = nullptr;
+	}
 }
 
 void UMaterial::FinishDestroy()
@@ -4688,6 +4705,11 @@ void UMaterial::AddReferencedObjects(UObject* InThis, FReferenceCollector& Colle
 {
 	UMaterial* This = CastChecked<UMaterial>(InThis);
 
+	if (This->CachedExpressionData)
+	{
+		This->CachedExpressionData->AddReferencedObjects(Collector);
+	}
+
 	for (FMaterialResource* CurrentResource : This->MaterialResources)
 	{
 		CurrentResource->AddReferencedObjects(Collector);
@@ -4706,7 +4728,9 @@ bool UMaterial::CanBeClusterRoot() const
 
 void UMaterial::GetAssetRegistryTags(TArray<FAssetRegistryTag>& OutTags) const
 {
-	OutTags.Add(FAssetRegistryTag("HasSceneColor", CachedExpressionData.bHasSceneColor ? TEXT("True") : TEXT("False"), FAssetRegistryTag::TT_Alphabetical));
+	// CachedExpressionData may be NULL here for default objects
+	const bool bHasSceneColor = CachedExpressionData ? CachedExpressionData->bHasSceneColor : false;
+	OutTags.Add(FAssetRegistryTag("HasSceneColor", bHasSceneColor ? TEXT("True") : TEXT("False"), FAssetRegistryTag::TT_Alphabetical));
 	Super::GetAssetRegistryTags(OutTags);
 }
 
@@ -4984,8 +5008,8 @@ int32 UMaterial::GetLayerParameterIndex(EMaterialParameterAssociation Associatio
 	int32 Index = INDEX_NONE;
 	switch (Association)
 	{
-	case BlendParameter: Index = CachedExpressionData.DefaultLayerBlends.Find(LayerFunction); break;
-	case LayerParameter: Index = CachedExpressionData.DefaultLayers.Find(LayerFunction); break;
+	case BlendParameter: Index = CachedExpressionData->DefaultLayerBlends.Find(LayerFunction); break;
+	case LayerParameter: Index = CachedExpressionData->DefaultLayers.Find(LayerFunction); break;
 	default: checkNoEntry(); break;
 	}
 	return Index;
@@ -4995,7 +5019,7 @@ int32 UMaterial::GetLayerParameterIndex(EMaterialParameterAssociation Associatio
 
 void UMaterial::GetReferencedTexturesAndOverrides(TSet<const UTexture*>& InOutTextures) const
 {
-	for (UObject* UsedObject : CachedExpressionData.ReferencedTextures)
+	for (UObject* UsedObject : CachedExpressionData->ReferencedTextures)
 	{
 		if (const UTexture* UsedTexture = Cast<UTexture>(UsedObject))
 		{
@@ -5547,17 +5571,17 @@ void UMaterial::RecursiveUpdateRealtimePreview( UMaterialExpression* InExpressio
 
 void UMaterial::AppendReferencedFunctionIdsTo(TArray<FGuid>& Ids) const
 {
-	for (int32 FunctionIndex = 0; FunctionIndex < CachedExpressionData.FunctionInfos.Num(); FunctionIndex++)
+	for (int32 FunctionIndex = 0; FunctionIndex < CachedExpressionData->FunctionInfos.Num(); FunctionIndex++)
 	{
-		Ids.AddUnique(CachedExpressionData.FunctionInfos[FunctionIndex].StateId);
+		Ids.AddUnique(CachedExpressionData->FunctionInfos[FunctionIndex].StateId);
 	}
 }
 
 void UMaterial::AppendReferencedParameterCollectionIdsTo(TArray<FGuid>& Ids) const
 {
-	for (int32 CollectionIndex = 0; CollectionIndex < CachedExpressionData.ParameterCollectionInfos.Num(); CollectionIndex++)
+	for (int32 CollectionIndex = 0; CollectionIndex < CachedExpressionData->ParameterCollectionInfos.Num(); CollectionIndex++)
 	{
-		Ids.AddUnique(CachedExpressionData.ParameterCollectionInfos[CollectionIndex].StateId);
+		Ids.AddUnique(CachedExpressionData->ParameterCollectionInfos[CollectionIndex].StateId);
 	}
 }
 
diff --git a/Engine/Source/Runtime/Engine/Private/Materials/MaterialCachedData.cpp b/Engine/Source/Runtime/Engine/Private/Materials/MaterialCachedData.cpp
--- a/Engine/Source/Runtime/Engine/Private/Materials/MaterialCachedData.cpp
+++ b/Engine/Source/Runtime/Engine/Private/Materials/MaterialCachedData.cpp
@@ -53,6 +53,40 @@ void FMaterialCachedExpressionData::Reset()
 	bHasSceneColor = false;
 }
 
+void FMaterialCachedExpressionData::AddReferencedObjects(FReferenceCollector& Collector)
+{
+	Parameters.AddReferencedObjects(Collector);
+	Collector.AddReferencedObjects(ReferencedTextures);
+	Collector.AddReferencedObjects(DefaultLayers);
+	Collector.AddReferencedObjects(DefaultLayerBlends);
+	Collector.AddReferencedObjects(GrassTypes);
+	for (FMaterialFunctionInfo& FunctionInfo : FunctionInfos)
+	{
+		Collector.AddReferencedObject(FunctionInfo.Function);
+	}
+	for (FMaterialParameterCollectionInfo& ParameterCollectionInfo : ParameterCollectionInfos)
+	{
+		Collector.AddReferencedObject(ParameterCollectionInfo.ParameterCollection);
+	}
+}
+
+void FMaterialInstanceCachedData::AddReferencedObjects(FReferenceCollector& Collector)
+{
+	Parameters.AddReferencedObjects(Collector);
+	Collector.AddReferencedObjects(ReferencedTextures);
+}
+
+void FMaterialCachedParameters::AddReferencedObjects(FReferenceCollector& Collector)
+{
+	Collector.AddReferencedObjects(TextureValues);
+	Collector.AddReferencedObjects(RuntimeVirtualTextureValues);
+	Collector.AddReferencedObjects(FontValues);
+#if WITH_EDITORONLY_DATA
+	Collector.AddReferencedObjects(ScalarCurveValues);
+	Collector.AddReferencedObjects(ScalarCurveAtlasValues);
+#endif
+}
+
 #if WITH_EDITOR
 static int32 TryAddParameter(FMaterialCachedParameters& CachedParameters, EMaterialParameterType Type, const FMaterialParameterInfo& ParameterInfo, const FGuid& ExpressionGuid, bool bOverride = false)
 {
@@ -716,3 +750,8 @@ void FMaterialCachedParameters::GetAllGlobalParameterInfoOfType(EMaterialParamet
 	}
 }
 
+void FMaterialInstanceCachedData::Initialize(FMaterialCachedExpressionData&& InCachedExpressionData)
+{
+	Parameters = MoveTemp(InCachedExpressionData.Parameters);
+	ReferencedTextures = MoveTemp(InCachedExpressionData.ReferencedTextures);
+}
diff --git a/Engine/Source/Runtime/Engine/Private/Materials/MaterialInstance.cpp b/Engine/Source/Runtime/Engine/Private/Materials/MaterialInstance.cpp
--- a/Engine/Source/Runtime/Engine/Private/Materials/MaterialInstance.cpp
+++ b/Engine/Source/Runtime/Engine/Private/Materials/MaterialInstance.cpp
@@ -745,11 +745,11 @@ bool UMaterialInstance::GetScalarParameterSliderMinMax(const FHashedMaterialPara
 
 	if (ParameterInfo.Association != EMaterialParameterAssociation::GlobalParameter)
 	{
-		const int32 ParameterIndex = CachedLayerParameters.FindParameterIndex(EMaterialParameterType::Scalar, ParameterInfo);
+		const int32 ParameterIndex = CachedData->Parameters.FindParameterIndex(EMaterialParameterType::Scalar, ParameterInfo);
 		if (ParameterIndex != INDEX_NONE)
 		{
-			OutSliderMin = CachedLayerParameters.ScalarMinMaxValues[ParameterIndex].X;
-			OutSliderMax = CachedLayerParameters.ScalarMinMaxValues[ParameterIndex].Y;
+			OutSliderMin = CachedData->Parameters.ScalarMinMaxValues[ParameterIndex].X;
+			OutSliderMax = CachedData->Parameters.ScalarMinMaxValues[ParameterIndex].Y;
 			return true;
 		}
 	}
@@ -784,11 +784,11 @@ bool UMaterialInstance::GetScalarParameterValue(const FHashedMaterialParameterIn
 	// Instance-included default
 	if (ParameterInfo.Association != EMaterialParameterAssociation::GlobalParameter)
 	{
-		const int32 ParameterIndex = CachedLayerParameters.FindParameterIndex(EMaterialParameterType::Scalar, ParameterInfo);
+		const int32 ParameterIndex = CachedData->Parameters.FindParameterIndex(EMaterialParameterType::Scalar, ParameterInfo);
 		if (ParameterIndex != INDEX_NONE)
 		{
-			const bool bResult = CachedLayerParameters.IsParameterValid(EMaterialParameterType::Scalar, ParameterIndex, bOveriddenOnly);
-			OutValue = CachedLayerParameters.ScalarValues[ParameterIndex];
+			const bool bResult = CachedData->Parameters.IsParameterValid(EMaterialParameterType::Scalar, ParameterIndex, bOveriddenOnly);
+			OutValue = CachedData->Parameters.ScalarValues[ParameterIndex];
 			return bResult;
 		}
 	}
@@ -828,11 +828,11 @@ bool UMaterialInstance::IsScalarParameterUsedAsAtlasPosition(const FHashedMateri
 	// Instance-included default
 	if (ParameterInfo.Association != EMaterialParameterAssociation::GlobalParameter)
 	{
-		const int32 ParameterIndex = CachedLayerParameters.FindParameterIndex(EMaterialParameterType::Scalar, ParameterInfo);
+		const int32 ParameterIndex = CachedData->Parameters.FindParameterIndex(EMaterialParameterType::Scalar, ParameterInfo);
 		if (ParameterIndex != INDEX_NONE)
 		{
-			UCurveLinearColor* Curve = CachedLayerParameters.ScalarCurveValues[ParameterIndex];
-			UCurveLinearColorAtlas* Atlas = CachedLayerParameters.ScalarCurveAtlasValues[ParameterIndex];
+			UCurveLinearColor* Curve = CachedData->Parameters.ScalarCurveValues[ParameterIndex];
+			UCurveLinearColorAtlas* Atlas = CachedData->Parameters.ScalarCurveAtlasValues[ParameterIndex];
 			if (Curve && Atlas)
 			{
 				OutCurve = TSoftObjectPtr<UCurveLinearColor>(FSoftObjectPath(Curve->GetPathName()));
@@ -878,11 +878,11 @@ bool UMaterialInstance::GetVectorParameterValue(const FHashedMaterialParameterIn
 	// Instance-included default
 	if (ParameterInfo.Association != EMaterialParameterAssociation::GlobalParameter)
 	{
-		const int32 ParameterIndex = CachedLayerParameters.FindParameterIndex(EMaterialParameterType::Vector, ParameterInfo);
+		const int32 ParameterIndex = CachedData->Parameters.FindParameterIndex(EMaterialParameterType::Vector, ParameterInfo);
 		if (ParameterIndex != INDEX_NONE)
 		{
-			const bool bResult = CachedLayerParameters.IsParameterValid(EMaterialParameterType::Vector, ParameterIndex, bOveriddenOnly);
-			OutValue = CachedLayerParameters.VectorValues[ParameterIndex];
+			const bool bResult = CachedData->Parameters.IsParameterValid(EMaterialParameterType::Vector, ParameterIndex, bOveriddenOnly);
+			OutValue = CachedData->Parameters.VectorValues[ParameterIndex];
 			return bResult;
 		}
 	}
@@ -910,10 +910,10 @@ bool UMaterialInstance::IsVectorParameterUsedAsChannelMask(const FHashedMaterial
 	// Instance-included default
 	if (ParameterInfo.Association != EMaterialParameterAssociation::GlobalParameter)
 	{
-		const int32 ParameterIndex = CachedLayerParameters.FindParameterIndex(EMaterialParameterType::Vector, ParameterInfo);
+		const int32 ParameterIndex = CachedData->Parameters.FindParameterIndex(EMaterialParameterType::Vector, ParameterInfo);
 		if (ParameterIndex != INDEX_NONE)
 		{
-			OutValue = CachedLayerParameters.VectorUsedAsChannelMaskValues[ParameterIndex];
+			OutValue = CachedData->Parameters.VectorUsedAsChannelMaskValues[ParameterIndex];
 			return true;
 		}
 	}
@@ -940,10 +940,10 @@ bool  UMaterialInstance::GetVectorParameterChannelNames(const FHashedMaterialPar
 	// Instance-included default
 	if (ParameterInfo.Association != EMaterialParameterAssociation::GlobalParameter)
 	{
-		const int32 ParameterIndex = CachedLayerParameters.FindParameterIndex(EMaterialParameterType::Vector, ParameterInfo);
+		const int32 ParameterIndex = CachedData->Parameters.FindParameterIndex(EMaterialParameterType::Vector, ParameterInfo);
 		if (ParameterIndex != INDEX_NONE)
 		{
-			OutValue = CachedLayerParameters.VectorChannelNameValues[ParameterIndex];
+			OutValue = CachedData->Parameters.VectorChannelNameValues[ParameterIndex];
 			return true;
 		}
 	}
@@ -980,11 +980,11 @@ bool UMaterialInstance::GetTextureParameterValue(const FHashedMaterialParameterI
 	// Instance-included default
 	if (ParameterInfo.Association != EMaterialParameterAssociation::GlobalParameter)
 	{
-		const int32 ParameterIndex = CachedLayerParameters.FindParameterIndex(EMaterialParameterType::Texture, ParameterInfo);
+		const int32 ParameterIndex = CachedData->Parameters.FindParameterIndex(EMaterialParameterType::Texture, ParameterInfo);
 		if (ParameterIndex != INDEX_NONE)
 		{
-			const bool bResult = CachedLayerParameters.IsParameterValid(EMaterialParameterType::Texture, ParameterIndex, bOveriddenOnly);
-			OutValue = CachedLayerParameters.TextureValues[ParameterIndex];
+			const bool bResult = CachedData->Parameters.IsParameterValid(EMaterialParameterType::Texture, ParameterIndex, bOveriddenOnly);
+			OutValue = CachedData->Parameters.TextureValues[ParameterIndex];
 			return bResult;
 		}
 	}
@@ -1021,11 +1021,11 @@ bool UMaterialInstance::GetRuntimeVirtualTextureParameterValue(const FHashedMate
 	// Instance-included default
 	if (ParameterInfo.Association != EMaterialParameterAssociation::GlobalParameter)
 	{
-		const int32 ParameterIndex = CachedLayerParameters.FindParameterIndex(EMaterialParameterType::RuntimeVirtualTexture, ParameterInfo);
+		const int32 ParameterIndex = CachedData->Parameters.FindParameterIndex(EMaterialParameterType::RuntimeVirtualTexture, ParameterInfo);
 		if (ParameterIndex != INDEX_NONE)
 		{
-			const bool bResult = CachedLayerParameters.IsParameterValid(EMaterialParameterType::RuntimeVirtualTexture, ParameterIndex, bOveriddenOnly);
-			OutValue = CachedLayerParameters.RuntimeVirtualTextureValues[ParameterIndex];
+			const bool bResult = CachedData->Parameters.IsParameterValid(EMaterialParameterType::RuntimeVirtualTexture, ParameterIndex, bOveriddenOnly);
+			OutValue = CachedData->Parameters.RuntimeVirtualTextureValues[ParameterIndex];
 			return bResult;
 		}
 	}
@@ -1055,10 +1055,10 @@ bool  UMaterialInstance::GetTextureParameterChannelNames(const FHashedMaterialPa
 	// Instance-included default
 	if (ParameterInfo.Association != EMaterialParameterAssociation::GlobalParameter)
 	{
-		const int32 ParameterIndex = CachedLayerParameters.FindParameterIndex(EMaterialParameterType::Texture, ParameterInfo);
+		const int32 ParameterIndex = CachedData->Parameters.FindParameterIndex(EMaterialParameterType::Texture, ParameterInfo);
 		if (ParameterIndex != INDEX_NONE)
 		{
-			OutValue = CachedLayerParameters.TextureChannelNameValues[ParameterIndex];
+			OutValue = CachedData->Parameters.TextureChannelNameValues[ParameterIndex];
 			return true;
 		}
 	}
@@ -1097,12 +1097,12 @@ bool UMaterialInstance::GetFontParameterValue(const FHashedMaterialParameterInfo
 	// Instance-included default
 	if (ParameterInfo.Association != EMaterialParameterAssociation::GlobalParameter)
 	{
-		const int32 ParameterIndex = CachedLayerParameters.FindParameterIndex(EMaterialParameterType::Font, ParameterInfo);
+		const int32 ParameterIndex = CachedData->Parameters.FindParameterIndex(EMaterialParameterType::Font, ParameterInfo);
 		if (ParameterIndex != INDEX_NONE)
 		{
-			const bool bResult = CachedLayerParameters.IsParameterValid(EMaterialParameterType::Font, ParameterIndex, bOveriddenOnly);
-			OutFontValue = CachedLayerParameters.FontValues[ParameterIndex];
-			OutFontPage = CachedLayerParameters.FontPageValues[ParameterIndex];
+			const bool bResult = CachedData->Parameters.IsParameterValid(EMaterialParameterType::Font, ParameterIndex, bOveriddenOnly);
+			OutFontValue = CachedData->Parameters.FontValues[ParameterIndex];
+			OutFontPage = CachedData->Parameters.FontPageValues[ParameterIndex];
 			return bResult;
 		}
 	}
@@ -2146,7 +2146,7 @@ void UMaterialInstance::GetStaticParameterValues(FStaticParameterSet& OutStaticP
 void UMaterialInstance::GetAllParametersOfType(EMaterialParameterType Type, TArray<FMaterialParameterInfo>& OutParameterInfo, TArray<FGuid>& OutParameterIds) const
 {
 	const UMaterial* Material = GetMaterial_Concurrent();
-	int32 NumParameters = CachedLayerParameters.GetNumParameters(Type);
+	int32 NumParameters = CachedData->Parameters.GetNumParameters(Type);
 	if (Material)
 	{
 		NumParameters += Material->GetCachedExpressionData().Parameters.GetNumParameters(Type);
@@ -2154,7 +2154,7 @@ void UMaterialInstance::GetAllParametersOfType(EMaterialParameterType Type, TArr
 
 	OutParameterInfo.Empty(NumParameters);
 	OutParameterIds.Empty(NumParameters);
-	CachedLayerParameters.GetAllParameterInfoOfType(Type, false, OutParameterInfo, OutParameterIds);
+	CachedData->Parameters.GetAllParameterInfoOfType(Type, false, OutParameterInfo, OutParameterIds);
 	if (Material)
 	{
 		Material->GetCachedExpressionData().Parameters.GetAllGlobalParameterInfoOfType(Type, false, OutParameterInfo, OutParameterIds);
@@ -2296,10 +2296,10 @@ bool UMaterialInstance::GetScalarParameterDefaultValue(const FHashedMaterialPara
 	// first matching expression found, not necessarily the one that's used for rendering
 	if (ParameterInfo.Association != EMaterialParameterAssociation::GlobalParameter)
 	{
-		const int32 ParameterIndex = CachedLayerParameters.FindParameterIndex(EMaterialParameterType::Scalar, ParameterInfo);
+		const int32 ParameterIndex = CachedData->Parameters.FindParameterIndex(EMaterialParameterType::Scalar, ParameterInfo);
 		if (ParameterIndex != INDEX_NONE)
 		{
-			OutValue = CachedLayerParameters.ScalarValues[ParameterIndex];
+			OutValue = CachedData->Parameters.ScalarValues[ParameterIndex];
 			return true;
 		}
 	}
@@ -2341,10 +2341,10 @@ bool UMaterialInstance::GetVectorParameterDefaultValue(const FHashedMaterialPara
 
 	if (ParameterInfo.Association != EMaterialParameterAssociation::GlobalParameter)
 	{
-		const int32 ParameterIndex = CachedLayerParameters.FindParameterIndex(EMaterialParameterType::Vector, ParameterInfo);
+		const int32 ParameterIndex = CachedData->Parameters.FindParameterIndex(EMaterialParameterType::Vector, ParameterInfo);
 		if (ParameterIndex != INDEX_NONE)
 		{
-			OutValue = CachedLayerParameters.VectorValues[ParameterIndex];
+			OutValue = CachedData->Parameters.VectorValues[ParameterIndex];
 			return true;
 		}
 	}
@@ -2386,10 +2386,10 @@ bool UMaterialInstance::GetTextureParameterDefaultValue(const FHashedMaterialPar
 
 	if (ParameterInfo.Association != EMaterialParameterAssociation::GlobalParameter)
 	{
-		const int32 ParameterIndex = CachedLayerParameters.FindParameterIndex(EMaterialParameterType::Texture, ParameterInfo);
+		const int32 ParameterIndex = CachedData->Parameters.FindParameterIndex(EMaterialParameterType::Texture, ParameterInfo);
 		if (ParameterIndex != INDEX_NONE)
 		{
-			OutValue = CachedLayerParameters.TextureValues[ParameterIndex];
+			OutValue = CachedData->Parameters.TextureValues[ParameterIndex];
 			return true;
 		}
 	}
@@ -2431,10 +2431,10 @@ bool UMaterialInstance::GetRuntimeVirtualTextureParameterDefaultValue(const FHas
 
 	if (ParameterInfo.Association != EMaterialParameterAssociation::GlobalParameter)
 	{
-		const int32 ParameterIndex = CachedLayerParameters.FindParameterIndex(EMaterialParameterType::RuntimeVirtualTexture, ParameterInfo);
+		const int32 ParameterIndex = CachedData->Parameters.FindParameterIndex(EMaterialParameterType::RuntimeVirtualTexture, ParameterInfo);
 		if (ParameterIndex != INDEX_NONE)
 		{
-			OutValue = CachedLayerParameters.RuntimeVirtualTextureValues[ParameterIndex];
+			OutValue = CachedData->Parameters.RuntimeVirtualTextureValues[ParameterIndex];
 			return true;
 		}
 	}
@@ -2477,11 +2477,11 @@ bool UMaterialInstance::GetFontParameterDefaultValue(const FHashedMaterialParame
 
 	if (ParameterInfo.Association != EMaterialParameterAssociation::GlobalParameter)
 	{
-		const int32 ParameterIndex = CachedLayerParameters.FindParameterIndex(EMaterialParameterType::Font, ParameterInfo);
+		const int32 ParameterIndex = CachedData->Parameters.FindParameterIndex(EMaterialParameterType::Font, ParameterInfo);
 		if (ParameterIndex != INDEX_NONE)
 		{
-			OutFontValue = CachedLayerParameters.FontValues[ParameterIndex];
-			OutFontPage = CachedLayerParameters.FontPageValues[ParameterIndex];
+			OutFontValue = CachedData->Parameters.FontValues[ParameterIndex];
+			OutFontPage = CachedData->Parameters.FontPageValues[ParameterIndex];
 			return true;
 		}
 	}
@@ -2875,6 +2875,9 @@ void UMaterialInstance::CacheResourceShadersForRendering(EMaterialShaderPrecompi
 
 void UMaterialInstance::CacheResourceShadersForCooking(EShaderPlatform ShaderPlatform, TArray<FMaterialResource*>& OutCachedMaterialResources, EMaterialShaderPrecompileMode PrecompileMode, const ITargetPlatform* TargetPlatform)
 {
+	// Make sure cached data is saved in cooked builds
+	bSavedCachedData = true;
+
 	if (bHasStaticPermutationResource)
 	{
 		UMaterial* BaseMaterial = GetMaterial();
@@ -3330,6 +3333,19 @@ void UMaterialInstance::Serialize(FArchive& Ar)
 	}
 #endif // WITH_EDITOR
 
+#if !WITH_EDITORONLY_DATA
+	checkf(bSavedCachedData, TEXT("MaterialInstance %s must have saved cached data, if editor-only data is not present"), *GetName());
+#endif
+	if (bSavedCachedData)
+	{
+		if (Ar.IsLoading())
+		{
+			CachedData = new FMaterialInstanceCachedData();
+		}
+		UScriptStruct* Struct = FMaterialInstanceCachedData::StaticStruct();
+		Struct->SerializeTaggedProperties(Ar, (uint8*)CachedData, Struct, nullptr);
+	}
+
 	// Only serialize the static permutation resource if one exists
 	if (bHasStaticPermutationResource)
 	{
@@ -3658,6 +3674,12 @@ void UMaterialInstance::FinishDestroy()
 		ClearAllCachedCookedPlatformData();
 	}
 #endif
+	if (CachedData)
+	{
+		delete CachedData;
+		CachedData = nullptr;
+	}
+
 	Super::FinishDestroy();
 }
 
@@ -4022,7 +4044,7 @@ void UMaterialInstance::UpdateStaticPermutation(const FStaticParameterSet& NewPa
 
 void UMaterialInstance::GetReferencedTexturesAndOverrides(TSet<const UTexture*>& InOutTextures) const
 {
-	for (UObject* UsedObject : CachedReferencedTextures)
+	for (UObject* UsedObject : CachedData->ReferencedTextures)
 	{
 		if (const UTexture* UsedTexture = Cast<UTexture>(UsedObject))
 		{
@@ -4050,11 +4072,12 @@ extern void FMaterialCachedParameters_UpdateForLayerParameters(FMaterialCachedPa
 
 void UMaterialInstance::UpdateCachedLayerParameters()
 {
-	if (!GIsClient)
+	if (bSavedCachedData)
 	{
-		// Expressions may not be loaded on server, so only rebuild the data if we're the client
+		// Don't need to rebuild cached data if it was serialized
 		return;
 	}
+
 	UMaterialInstance* ParentInstance = nullptr;
 	FMaterialCachedExpressionData CachedExpressionData;
 	CachedExpressionData.Reset();
@@ -4079,11 +4102,11 @@ void UMaterialInstance::UpdateCachedLayerParameters()
 		}
 	}
 
-	if (bCachedDataValid)
+	if (!CachedData)
 	{
-		CachedLayerParameters = MoveTemp(CachedExpressionData.Parameters);
-		CachedReferencedTextures = MoveTemp(CachedExpressionData.ReferencedTextures);
+		CachedData = new FMaterialInstanceCachedData();
 	}
+	CachedData->Initialize(MoveTemp(CachedExpressionData));
 }
 
 void UMaterialInstance::UpdateStaticPermutation(const FStaticParameterSet& NewParameters, FMaterialUpdateContext* MaterialUpdateContext)
diff --git a/Engine/Source/Runtime/Engine/Public/MaterialCachedData.h b/Engine/Source/Runtime/Engine/Public/MaterialCachedData.h
--- a/Engine/Source/Runtime/Engine/Public/MaterialCachedData.h
+++ b/Engine/Source/Runtime/Engine/Public/MaterialCachedData.h
@@ -154,6 +154,8 @@ struct FMaterialCachedParameters
 	void GetAllGlobalParameterInfoOfType(EMaterialParameterType Type, bool bEmptyOutput, TArray<FMaterialParameterInfo>& OutParameterInfo, TArray<FGuid>& OutParameterIds) const;
 	void Reset();
 
+	void AddReferencedObjects(FReferenceCollector& Collector);
+
 	UPROPERTY()
 	FMaterialCachedParameterEntry RuntimeEntries[NumMaterialRuntimeParameterTypes];
 
@@ -227,6 +229,8 @@ struct FMaterialCachedExpressionData
 
 	void Reset();
 
+	void AddReferencedObjects(FReferenceCollector& Collector);
+
 	UPROPERTY()
 	FMaterialCachedParameters Parameters;
 
@@ -263,3 +267,18 @@ struct FMaterialCachedExpressionData
 	UPROPERTY()
 	uint32 bHasSceneColor : 1;
 };
+
+USTRUCT()
+struct FMaterialInstanceCachedData
+{
+	GENERATED_USTRUCT_BODY()
+
+	void Initialize(FMaterialCachedExpressionData&& InCachedExpressionData);
+	void AddReferencedObjects(FReferenceCollector& Collector);
+
+	UPROPERTY()
+	FMaterialCachedParameters Parameters;
+
+	UPROPERTY()
+	TArray<UObject*> ReferencedTextures;
+};
