Multiple material editor fixes related to Strata
- Fixed debug text when input does not match on root node
- MaterialAttribute breakout node FrontMaterial output can now be used
- Strata material can now be an output of a material function
- Strata material output from mateiral function can now be plug into he root mateiral node directly
- fixed normal/tangent code chunk to always be valid for deduplcaited when inside a material function.
- Strata material can now be previewed in material functions too

diff --git a/Engine/Source/Editor/UnrealEd/Private/MaterialGraphNode_Root.cpp b/Engine/Source/Editor/UnrealEd/Private/MaterialGraphNode_Root.cpp
--- a/Engine/Source/Editor/UnrealEd/Private/MaterialGraphNode_Root.cpp
+++ b/Engine/Source/Editor/UnrealEd/Private/MaterialGraphNode_Root.cpp
@@ -77,6 +77,10 @@ uint32 UMaterialGraphNode_Root::GetInputType(const UEdGraphPin* InputPin) const
 	{
 		return MCT_MaterialAttributes;
 	}
+	else if (Property == MP_FrontMaterial)
+	{
+		return MCT_Strata;
+	}
 	else
 	{
 		return FMaterialAttributeDefinitionMap::GetValueType(Property);
diff --git a/Engine/Source/Runtime/Engine/Classes/Materials/MaterialExpression.h b/Engine/Source/Runtime/Engine/Classes/Materials/MaterialExpression.h
--- a/Engine/Source/Runtime/Engine/Classes/Materials/MaterialExpression.h
+++ b/Engine/Source/Runtime/Engine/Classes/Materials/MaterialExpression.h
@@ -326,6 +326,11 @@ class ENGINE_API UMaterialExpression : public UObject
 	 */
 	virtual bool IsResultMaterialAttributes(int32 OutputIndex) { return false; }
 
+	/**
+	 * Marks certain expression types as outputting Strata material. Allows the material functions to directly return a Strata material as output pin.
+	 */
+	virtual bool IsResultStrataMaterial(int32 OutputIndex) { return false; }
+
 	/**
 	 * If true, discards the output index when caching this expression which allows more cases to re-use the output instead of adding a separate instruction
 	 */
diff --git a/Engine/Source/Runtime/Engine/Classes/Materials/MaterialExpressionFunctionOutput.h b/Engine/Source/Runtime/Engine/Classes/Materials/MaterialExpressionFunctionOutput.h
--- a/Engine/Source/Runtime/Engine/Classes/Materials/MaterialExpressionFunctionOutput.h
+++ b/Engine/Source/Runtime/Engine/Classes/Materials/MaterialExpressionFunctionOutput.h
@@ -63,6 +63,7 @@ class UMaterialExpressionFunctionOutput : public UMaterialExpression
 	virtual void GetExpressionToolTip(TArray<FString>& OutToolTip) override;
 	virtual uint32 GetInputType(int32 InputIndex) override;
 	virtual bool IsResultMaterialAttributes(int32 OutputIndex) override;
+	virtual bool IsResultStrataMaterial(int32 OutputIndex) override;
 #endif // WITH_EDITOR
 	//~ End UMaterialExpression Interface
 
diff --git a/Engine/Source/Runtime/Engine/Classes/Materials/MaterialExpressionMaterialFunctionCall.h b/Engine/Source/Runtime/Engine/Classes/Materials/MaterialExpressionMaterialFunctionCall.h
--- a/Engine/Source/Runtime/Engine/Classes/Materials/MaterialExpressionMaterialFunctionCall.h
+++ b/Engine/Source/Runtime/Engine/Classes/Materials/MaterialExpressionMaterialFunctionCall.h
@@ -128,6 +128,7 @@ class UMaterialExpressionMaterialFunctionCall : public UMaterialExpression
 	virtual void GetExpressionToolTip(TArray<FString>& OutToolTip) override;
 	virtual bool MatchesSearchQuery( const TCHAR* SearchQuery ) override;
 	virtual bool IsResultMaterialAttributes(int32 OutputIndex) override;
+	virtual bool IsResultStrataMaterial(int32 OutputIndex) override;
 	virtual uint32 GetInputType(int32 InputIndex) override;
 	//~ End UMaterialExpression Interface
 
diff --git a/Engine/Source/Runtime/Engine/Classes/Materials/MaterialExpressionStrata.h b/Engine/Source/Runtime/Engine/Classes/Materials/MaterialExpressionStrata.h
--- a/Engine/Source/Runtime/Engine/Classes/Materials/MaterialExpressionStrata.h
+++ b/Engine/Source/Runtime/Engine/Classes/Materials/MaterialExpressionStrata.h
@@ -64,6 +64,7 @@ class UMaterialExpressionStrataDiffuseBSDF : public UMaterialExpression // STRAT
 	virtual void GetCaption(TArray<FString>& OutCaptions) const override;
 	virtual uint32 GetOutputType(int32 OutputIndex) override;
 	virtual uint32 GetInputType(int32 InputIndex) override;
+	virtual bool IsResultStrataMaterial(int32 OutputIndex) override;
 #endif
 	//~ End UMaterialExpression Interface
 };
@@ -115,6 +116,7 @@ class UMaterialExpressionStrataDielectricBSDF : public UMaterialExpression
 	virtual void GetCaption(TArray<FString>& OutCaptions) const override;
 	virtual uint32 GetOutputType(int32 OutputIndex) override;
 	virtual uint32 GetInputType(int32 InputIndex) override;
+	virtual bool IsResultStrataMaterial(int32 OutputIndex) override;
 #endif
 	//~ End UMaterialExpression Interface
 };
@@ -166,6 +168,7 @@ class UMaterialExpressionStrataConductorBSDF : public UMaterialExpression
 	virtual void GetCaption(TArray<FString>& OutCaptions) const override;
 	virtual uint32 GetOutputType(int32 OutputIndex) override;
 	virtual uint32 GetInputType(int32 InputIndex) override;
+	virtual bool IsResultStrataMaterial(int32 OutputIndex) override;
 #endif
 	//~ End UMaterialExpression Interface
 };
@@ -205,6 +208,7 @@ class UMaterialExpressionStrataVolumeBSDF : public UMaterialExpression
 	virtual void GetCaption(TArray<FString>& OutCaptions) const override;
 	virtual uint32 GetOutputType(int32 OutputIndex) override;
 	virtual uint32 GetInputType(int32 InputIndex) override;
+	virtual bool IsResultStrataMaterial(int32 OutputIndex) override;
 #endif
 	//~ End UMaterialExpression Interface
 };
@@ -238,6 +242,7 @@ class UMaterialExpressionStrataSheenBSDF : public UMaterialExpression
 	virtual void GetCaption(TArray<FString>& OutCaptions) const override;
 	virtual uint32 GetOutputType(int32 OutputIndex) override;
 	virtual uint32 GetInputType(int32 InputIndex) override;
+	virtual bool IsResultStrataMaterial(int32 OutputIndex) override;
 #endif
 	//~ End UMaterialExpression Interface
 };
@@ -278,6 +283,7 @@ class UMaterialExpressionStrataHorizontalMixing : public UMaterialExpression
 	virtual void GetCaption(TArray<FString>& OutCaptions) const override;
 	virtual uint32 GetOutputType(int32 OutputIndex) override;
 	virtual uint32 GetInputType(int32 InputIndex) override;
+	virtual bool IsResultStrataMaterial(int32 OutputIndex) override;
 #endif
 	//~ End UMaterialExpression Interface
 };
@@ -305,6 +311,7 @@ class UMaterialExpressionStrataVerticalLayering : public UMaterialExpression
 	virtual void GetCaption(TArray<FString>& OutCaptions) const override;
 	virtual uint32 GetOutputType(int32 OutputIndex) override;
 	virtual uint32 GetInputType(int32 InputIndex) override;
+	virtual bool IsResultStrataMaterial(int32 OutputIndex) override;
 #endif
 	//~ End UMaterialExpression Interface
 };
@@ -332,6 +339,7 @@ class UMaterialExpressionStrataAdd : public UMaterialExpression
 	virtual void GetCaption(TArray<FString>& OutCaptions) const override;
 	virtual uint32 GetOutputType(int32 OutputIndex) override;
 	virtual uint32 GetInputType(int32 InputIndex) override;
+	virtual bool IsResultStrataMaterial(int32 OutputIndex) override;
 #endif
 	//~ End UMaterialExpression Interface
 };
@@ -359,6 +367,7 @@ class UMaterialExpressionStrataMultiply : public UMaterialExpression
 	virtual void GetCaption(TArray<FString>& OutCaptions) const override;
 	virtual uint32 GetOutputType(int32 OutputIndex) override;
 	virtual uint32 GetInputType(int32 InputIndex) override;
+	virtual bool IsResultStrataMaterial(int32 OutputIndex) override;
 #endif
 	//~ End UMaterialExpression Interface
 };
diff --git a/Engine/Source/Runtime/Engine/Private/Materials/MaterialExpressions.cpp b/Engine/Source/Runtime/Engine/Private/Materials/MaterialExpressions.cpp
--- a/Engine/Source/Runtime/Engine/Private/Materials/MaterialExpressions.cpp
+++ b/Engine/Source/Runtime/Engine/Private/Materials/MaterialExpressions.cpp
@@ -398,6 +398,8 @@ void GetMaterialValueTypeDescriptions(const uint32 MaterialValueType, TArray<FTe
 		OutDescriptions.Add(LOCTEXT("MaterialAttributes", "Material Attributes"));
 	if (MaterialValueType & MCT_ShadingModel)
 		OutDescriptions.Add(LOCTEXT("ShadingModel", "Shading Model"));
+	if (MaterialValueType & MCT_Strata)
+		OutDescriptions.Add(LOCTEXT("Strata", "Strata Material"));
 	if (MaterialValueType & MCT_Unknown)
 		OutDescriptions.Add(LOCTEXT("Unknown", "Unknown"));
 }
@@ -1073,8 +1075,12 @@ uint32 UMaterialExpression::GetInputType(int32 InputIndex)
 
 uint32 UMaterialExpression::GetOutputType(int32 OutputIndex)
 {
-	// different outputs should be defined by sub classed expressions
-	if (IsResultMaterialAttributes(OutputIndex))
+	// different outputs should be defined by sub classed expressions 
+	if (IsResultStrataMaterial(OutputIndex))
+	{
+		return MCT_Strata;
+	}
+	else if (IsResultMaterialAttributes(OutputIndex))
 	{
 		return MCT_MaterialAttributes;
 	}
@@ -1282,9 +1288,15 @@ void UMaterialExpression::ConnectToPreviewMaterial(UMaterial* InMaterial, int32
 {
 	if (InMaterial && OutputIndex >= 0 && OutputIndex < Outputs.Num())
 	{
-		bool bUseMaterialAttributes = IsResultMaterialAttributes(0);
-
-		if( bUseMaterialAttributes )
+		if (IsResultStrataMaterial(0))
+		{
+			InMaterial->SetShadingModel(MSM_DefaultLit);
+			InMaterial->bUseMaterialAttributes = false;
+			FExpressionInput* MaterialInput = InMaterial->GetExpressionInputForProperty(MP_FrontMaterial);
+			check(MaterialInput);
+			ConnectExpression(MaterialInput, OutputIndex);
+		}
+		else if(IsResultMaterialAttributes(0))
 		{
 			InMaterial->SetShadingModel(MSM_DefaultLit);
 			InMaterial->bUseMaterialAttributes = true;
@@ -5960,6 +5972,10 @@ uint32 UMaterialExpressionBreakMaterialAttributes::GetOutputType(int32 OutputInd
 	{
 		return MCT_ShadingModel;
 	}
+	else if (Property && *Property == EMaterialProperty::MP_FrontMaterial)
+	{
+		return MCT_Strata;
+	}
 	else
 	{
 		return UMaterialExpression::GetOutputType(OutputIndex);
@@ -14243,7 +14259,7 @@ bool UMaterialExpressionMaterialFunctionCall::MatchesSearchQuery( const TCHAR* S
 	return Super::MatchesSearchQuery(SearchQuery);
 }
 
-bool UMaterialExpressionMaterialFunctionCall::IsResultMaterialAttributes(int32 OutputIndex)
+bool UMaterialExpressionMaterialFunctionCall::IsResultMaterialAttributes(int32 OutputIndex) 
 {
 	if( OutputIndex >= 0 && OutputIndex < FunctionOutputs.Num() && FunctionOutputs[OutputIndex].ExpressionOutput)
 	{
@@ -14255,6 +14271,18 @@ bool UMaterialExpressionMaterialFunctionCall::IsResultMaterialAttributes(int32 O
 	}
 }
 
+bool UMaterialExpressionMaterialFunctionCall::IsResultStrataMaterial(int32 OutputIndex)
+{
+	if( OutputIndex >= 0 && OutputIndex < FunctionOutputs.Num() && FunctionOutputs[OutputIndex].ExpressionOutput)
+	{
+		return FunctionOutputs[OutputIndex].ExpressionOutput->IsResultStrataMaterial(0);
+	}
+	else
+	{
+		return false;
+	}
+}
+
 uint32 UMaterialExpressionMaterialFunctionCall::GetInputType(int32 InputIndex)
 {
 	if (InputIndex < FunctionInputs.Num())
@@ -14695,7 +14723,7 @@ void UMaterialExpressionFunctionOutput::GetExpressionToolTip(TArray<FString>& Ou
 uint32 UMaterialExpressionFunctionOutput::GetInputType(int32 InputIndex)
 {
 	// Acceptable types for material function outputs
-	return MCT_Float | MCT_MaterialAttributes;
+	return MCT_Float | MCT_MaterialAttributes | MCT_Strata;
 }
 
 int32 UMaterialExpressionFunctionOutput::Compile(class FMaterialCompiler* Compiler, int32 OutputIndex)
@@ -14765,6 +14793,20 @@ bool UMaterialExpressionFunctionOutput::IsResultMaterialAttributes(int32 OutputI
 		return false;
 	}
 }
+
+bool UMaterialExpressionFunctionOutput::IsResultStrataMaterial(int32 OutputIndex)
+{
+	// If there is a loop anywhere in this expression's inputs then we can't risk checking them
+	if( A.GetTracedInput().Expression && !A.Expression->ContainsInputLoop() )
+	{
+		return A.Expression->IsResultStrataMaterial(A.OutputIndex);
+	}
+	else
+	{
+		return false;
+	}
+}
+
 #endif // WITH_EDITOR
 
 ///////////////////////////////////////////////////////////////////////////////
@@ -19012,6 +19054,14 @@ FName UMaterialExpressionSceneDepthWithoutWater::GetInputName(int32 InputIndex)
 // Strata
 
 
+// The compilation of an expression can sometimes lead to a INDEX_NONE code chunk when editing material graphs 
+// or when the node is inside a material function, linked to an input pin of the material function and that input is not plugged in to anything.
+// But for normals or tangents, Strata absolutely need a valid code chunk to de-duplicate when stored in memory. So we provide those helper functions.
+static int32 CompileWithDefaultCodeChunk(class FMaterialCompiler* Compiler, FExpressionInput& Input, int DefaultCodeChunk)
+{
+	int32 CodeChunk = Input.GetTracedInput().Expression ? Input.Compile(Compiler) : DefaultCodeChunk;
+	return CodeChunk == -1 ? DefaultCodeChunk : CodeChunk;
+}
 
 UMaterialExpressionStrataDiffuseBSDF::UMaterialExpressionStrataDiffuseBSDF(const FObjectInitializer& ObjectInitializer)
 	: Super(ObjectInitializer)
@@ -19030,7 +19080,7 @@ UMaterialExpressionStrataDiffuseBSDF::UMaterialExpressionStrataDiffuseBSDF(const
 #if WITH_EDITOR
 int32 UMaterialExpressionStrataDiffuseBSDF::Compile(class FMaterialCompiler* Compiler, int32 OutputIndex)
 {
-	int32 NormalCodeChunk = Normal.GetTracedInput().Expression ? Normal.Compile(Compiler) : Compiler->VertexNormal();
+	int32 NormalCodeChunk = CompileWithDefaultCodeChunk(Compiler, Normal, Compiler->VertexNormal());
 	uint8 SharedNormalIndex = StrataCompilationInfoCreateSharedNormal(Compiler, NormalCodeChunk);
 
 	int32 OutputCodeChunk = Compiler->StrataDiffuseBSDF(
@@ -19071,9 +19121,12 @@ uint32 UMaterialExpressionStrataDiffuseBSDF::GetInputType(int32 InputIndex)
 	check(false);
 	return MCT_Float1;
 }
-#endif // WITH_EDITOR
-
 
+bool UMaterialExpressionStrataDiffuseBSDF::IsResultStrataMaterial(int32 OutputIndex)
+{
+	return true;
+}
+#endif // WITH_EDITOR
 
 
 
@@ -19094,8 +19147,8 @@ UMaterialExpressionStrataDielectricBSDF::UMaterialExpressionStrataDielectricBSDF
 #if WITH_EDITOR
 int32 UMaterialExpressionStrataDielectricBSDF::Compile(class FMaterialCompiler* Compiler, int32 OutputIndex)
 {
-	int32 NormalCodeChunk = Normal.GetTracedInput().Expression ? Normal.Compile(Compiler) : Compiler->VertexNormal();
-	int32 TangentCodeChunk = Tangent.GetTracedInput().Expression ? Tangent.Compile(Compiler) : Compiler->VertexTangent();
+	int32 NormalCodeChunk = CompileWithDefaultCodeChunk(Compiler, Normal, Compiler->VertexNormal());
+	int32 TangentCodeChunk = CompileWithDefaultCodeChunk(Compiler, Tangent, Compiler->VertexTangent());
 	uint8 SharedNormalIndex = StrataCompilationInfoCreateSharedNormal(Compiler, NormalCodeChunk, TangentCodeChunk);
 
 	int32 RoughnessXCodeChunk = RoughnessX.GetTracedInput().Expression ? RoughnessX.Compile(Compiler) : Compiler->Constant(0.0f);
@@ -19105,8 +19158,8 @@ int32 UMaterialExpressionStrataDielectricBSDF::Compile(class FMaterialCompiler*
 	int32 OutputCodeChunk = Compiler->StrataDielectricBSDF(
 		RoughnessXCodeChunk,
 		RoughnessYCodeChunk,
-		IOR.GetTracedInput().Expression			? IOR.Compile(Compiler)			: Compiler->Constant(1.5f),// Default to Glass
-		Tint.GetTracedInput().Expression		? Tint.Compile(Compiler)		: Compiler->Constant3(1.0f, 1.0f, 1.0f),
+		IOR.GetTracedInput().Expression			? IOR.Compile(Compiler)		: Compiler->Constant(1.5f),// Default to Glass
+		Tint.GetTracedInput().Expression		? Tint.Compile(Compiler)	: Compiler->Constant3(1.0f, 1.0f, 1.0f),
 		NormalCodeChunk,
 		TangentCodeChunk,
 		SharedNormalIndex);
@@ -19152,9 +19205,12 @@ uint32 UMaterialExpressionStrataDielectricBSDF::GetInputType(int32 InputIndex)
 	check(false);
 	return MCT_Float1;
 }
-#endif // WITH_EDITOR
-
 
+bool UMaterialExpressionStrataDielectricBSDF::IsResultStrataMaterial(int32 OutputIndex)
+{
+	return true;
+}
+#endif // WITH_EDITOR
 
 
 
@@ -19175,13 +19231,8 @@ UMaterialExpressionStrataConductorBSDF::UMaterialExpressionStrataConductorBSDF(c
 #if WITH_EDITOR
 int32 UMaterialExpressionStrataConductorBSDF::Compile(class FMaterialCompiler* Compiler, int32 OutputIndex)
 {
-	int32 NormalCodeChunk = Normal.GetTracedInput().Expression ? Normal.Compile(Compiler) : Compiler->VertexNormal();
-	int32 TangentCodeChunk = Tangent.GetTracedInput().Expression ? Tangent.Compile(Compiler) : Compiler->VertexTangent();
-
-	// For some reason, when editing shader, Normal.Compile() / Tangent.Compile() might return -1. In this case we fallback to vertex normal/tangent
-	if (NormalCodeChunk == -1)  { NormalCodeChunk = Compiler->VertexNormal(); }
-	if (TangentCodeChunk == -1)	{ TangentCodeChunk = Compiler->VertexTangent();	}
-
+	int32 NormalCodeChunk = CompileWithDefaultCodeChunk(Compiler, Normal, Compiler->VertexNormal());
+	int32 TangentCodeChunk = CompileWithDefaultCodeChunk(Compiler, Tangent, Compiler->VertexTangent());
 	uint8 SharedNormalIndex = StrataCompilationInfoCreateSharedNormal(Compiler, NormalCodeChunk, TangentCodeChunk);
 
 	int32 RoughnessXCodeChunk = RoughnessX.GetTracedInput().Expression ? RoughnessX.Compile(Compiler) : Compiler->Constant(0.0f);
@@ -19238,9 +19289,12 @@ uint32 UMaterialExpressionStrataConductorBSDF::GetInputType(int32 InputIndex)
 	check(false);
 	return MCT_Float1;
 }
-#endif // WITH_EDITOR
-
 
+bool UMaterialExpressionStrataConductorBSDF::IsResultStrataMaterial(int32 OutputIndex)
+{
+	return true;
+}
+#endif // WITH_EDITOR
 
 
 
@@ -19309,6 +19363,11 @@ uint32 UMaterialExpressionStrataVolumeBSDF::GetInputType(int32 InputIndex)
 	check(false);
 	return MCT_Float1;
 }
+
+bool UMaterialExpressionStrataVolumeBSDF::IsResultStrataMaterial(int32 OutputIndex)
+{
+	return true;
+}
 #endif // WITH_EDITOR
 
 
@@ -19330,7 +19389,7 @@ UMaterialExpressionStrataSheenBSDF::UMaterialExpressionStrataSheenBSDF(const FOb
 #if WITH_EDITOR
 int32 UMaterialExpressionStrataSheenBSDF::Compile(class FMaterialCompiler* Compiler, int32 OutputIndex)
 {
-	int32 NormalCodeChunk = Normal.GetTracedInput().Expression ? Normal.Compile(Compiler) : Compiler->VertexNormal();
+	int32 NormalCodeChunk = CompileWithDefaultCodeChunk(Compiler, Normal, Compiler->VertexNormal());
 
 	// For some reason, when editing shader, Normal.Compile() / Tangent.Compile() might return -1. In this case we fallback to vertex normal/tangent
 	if (NormalCodeChunk == -1) { NormalCodeChunk = Compiler->VertexNormal(); }
@@ -19377,6 +19436,11 @@ uint32 UMaterialExpressionStrataSheenBSDF::GetInputType(int32 InputIndex)
 	check(false);
 	return MCT_Float1;
 }
+
+bool UMaterialExpressionStrataSheenBSDF::IsResultStrataMaterial(int32 OutputIndex)
+{
+	return true;
+}
 #endif // WITH_EDITOR
 
 
@@ -19439,6 +19503,11 @@ uint32 UMaterialExpressionStrataHorizontalMixing::GetInputType(int32 InputIndex)
 {
 	return InputIndex == 2 ? MCT_Float1 : MCT_Strata;
 }
+
+bool UMaterialExpressionStrataHorizontalMixing::IsResultStrataMaterial(int32 OutputIndex)
+{
+	return true;
+}
 #endif // WITH_EDITOR
 
 
@@ -19498,6 +19567,11 @@ uint32 UMaterialExpressionStrataVerticalLayering::GetInputType(int32 InputIndex)
 {
 	return MCT_Strata;
 }
+
+bool UMaterialExpressionStrataVerticalLayering::IsResultStrataMaterial(int32 OutputIndex)
+{
+	return true;
+}
 #endif // WITH_EDITOR
 
 
@@ -19557,6 +19631,11 @@ uint32 UMaterialExpressionStrataAdd::GetInputType(int32 InputIndex)
 {
 	return MCT_Strata;
 }
+
+bool UMaterialExpressionStrataAdd::IsResultStrataMaterial(int32 OutputIndex)
+{
+	return true;
+}
 #endif // WITH_EDITOR
 
 
@@ -19616,6 +19695,11 @@ uint32 UMaterialExpressionStrataMultiply::GetInputType(int32 InputIndex)
 	}
 	return MCT_Float;
 }
+
+bool UMaterialExpressionStrataMultiply::IsResultStrataMaterial(int32 OutputIndex)
+{
+	return true;
+}
 #endif // WITH_EDITOR
 
 
