Missing edigrate:
Add ability to initialize global array of AudioBuses at end of AudioDevice initialization

diff --git a/Engine/Source/Runtime/AudioMixer/Private/AudioMixerDevice.cpp b/Engine/Source/Runtime/AudioMixer/Private/AudioMixerDevice.cpp
--- a/Engine/Source/Runtime/AudioMixer/Private/AudioMixerDevice.cpp
+++ b/Engine/Source/Runtime/AudioMixer/Private/AudioMixerDevice.cpp
@@ -15,6 +15,7 @@
 #include "Sound/SoundSubmixSend.h"
 #include "SubmixEffects/AudioMixerSubmixEffectEQ.h"
 #include "SubmixEffects/AudioMixerSubmixEffectDynamicsProcessor.h"
+#include "UObject/StrongObjectPtr.h"
 #include "UObject/UObjectHash.h"
 #include "UObject/UObjectIterator.h"
 #include "Runtime/HeadMountedDisplay/Public/IHeadMountedDisplayModule.h"
diff --git a/Engine/Source/Runtime/AudioMixer/Public/AudioMixerDevice.h b/Engine/Source/Runtime/AudioMixer/Public/AudioMixerDevice.h
--- a/Engine/Source/Runtime/AudioMixer/Public/AudioMixerDevice.h
+++ b/Engine/Source/Runtime/AudioMixer/Public/AudioMixerDevice.h
@@ -9,6 +9,7 @@
 #include "DSP/BufferVectorOperations.h"
 #include "DSP/MultithreadedPatching.h"
 #include "Quartz/AudioMixerClockManager.h"
+#include "UObject/StrongObjectPtr.h"
 
 // Forward Declarations
 class FOnSubmixEnvelopeBP;
@@ -96,6 +97,8 @@ namespace Audio
 		virtual void UnregisterSoundSubmix(const USoundSubmixBase* SoundSubmix) override;
 
 		virtual void InitSoundEffectPresets() override;
+		virtual void InitDefaultAudioBuses() override;
+		virtual void ShutdownDefaultAudioBuses() override;
 		virtual int32 GetNumActiveSources() const override;
 
 		// Updates the source effect chain (using unique object id). 
@@ -299,6 +302,8 @@ namespace Audio
 		TArray<USoundSubmix*> MasterSubmixes;
 		TArray<FMixerSubmixPtr> MasterSubmixInstances;
 
+		TArray<TStrongObjectPtr<UAudioBus>> DefaultAudioBuses;
+
 		// The active audio bus list accessible on the game thread
 		TArray<int32> ActiveAudioBuses_GameThread;
 
diff --git a/Engine/Source/Runtime/Engine/Classes/Sound/AudioSettings.h b/Engine/Source/Runtime/Engine/Classes/Sound/AudioSettings.h
--- a/Engine/Source/Runtime/Engine/Classes/Sound/AudioSettings.h
+++ b/Engine/Source/Runtime/Engine/Classes/Sound/AudioSettings.h
@@ -64,6 +64,16 @@ struct ENGINE_API FAudioQualitySettings
 	}
 };
 
+USTRUCT()
+struct ENGINE_API FDefaultAudioBusSettings
+{
+	GENERATED_BODY()
+
+	/** The SoundConcurrency assigned to newly created sounds */
+	UPROPERTY(EditAnywhere, Category = "Mix", meta = (AllowedClasses = "AudioBus"))
+	FSoftObjectPath AudioBus;
+};
+
 /**
  * Audio settings.
  */
@@ -178,6 +188,11 @@ class ENGINE_API UAudioSettings : public UDeveloperSettings
 	UPROPERTY(config, EditAnywhere, Category="Dialogue")
 	FString DialogueFilenameFormat;
 
+	/** Array of AudioBuses that are automatically initialized when the AudioEngine is initialized */
+	UPROPERTY(config, EditAnywhere, Category="Mix")
+	TArray<FDefaultAudioBusSettings> DefaultAudioBuses;
+
+
 #if WITH_EDITOR
 	FAudioSettingsChanged AudioSettingsChanged;
 #endif // WITH_EDITOR
diff --git a/Engine/Source/Runtime/Engine/Private/AudioBus.cpp b/Engine/Source/Runtime/Engine/Private/AudioBus.cpp
--- a/Engine/Source/Runtime/Engine/Private/AudioBus.cpp
+++ b/Engine/Source/Runtime/Engine/Private/AudioBus.cpp
@@ -45,6 +45,61 @@ void UAudioBus::BeginDestroy()
 #if WITH_EDITOR
 void UAudioBus::PostEditChangeProperty(FPropertyChangedEvent& PropertyChangedEvent)
 {
-}
-#endif
+	if (!PropertyChangedEvent.Property)
+	{
+		return;
+	}
+
+	FName PropertyName = PropertyChangedEvent.Property->GetFName();
+	if (PropertyName != GET_MEMBER_NAME_CHECKED(UAudioBus, AudioBusChannels))
+	{
+		return;
+	}
+
+	const UAudioSettings* AudioSettings = GetDefault<UAudioSettings>();
+	if (!AudioSettings)
+	{
+		return;
+	}
+
+	for (const FDefaultAudioBusSettings& DefaultBusSettings : AudioSettings->DefaultAudioBuses)
+	{
+		if (DefaultBusSettings.AudioBus.ResolveObject() != this)
+		{
+			continue;
+		}
+
+		FAudioDeviceManager* DeviceManager = FAudioDeviceManager::Get();
+		if (!DeviceManager)
+		{
+			continue;
+		}
+
+		// Restart bus with new channel count
+		UObject* AudioBus = DefaultBusSettings.AudioBus.ResolveObject();
+		if (this != AudioBus)
+		{
+			continue;
+		}
+
+		DeviceManager->IterateOverAllDevices([BusId = AudioBus->GetUniqueID(), NumChannels = AudioBusChannels](Audio::FDeviceId, FAudioDevice* InDevice)
+		{
+			if (Audio::FMixerDevice* MixerDevice = static_cast<Audio::FMixerDevice*>(InDevice))
+			{
+				MixerDevice->StopAudioBus(BusId);
+				switch(NumChannels)
+				{
+					case EAudioBusChannels::Stereo:
+						MixerDevice->StartAudioBus(BusId, 2, false /* bInIsAutomatic */);
+					break;
 
+					case EAudioBusChannels::Mono:
+					default:
+						MixerDevice->StartAudioBus(BusId, 1, false /* bInIsAutomatic */);
+					break;
+				}
+			}
+		});
+	}
+}
+#endif // WITH_EDITOR
diff --git a/Engine/Source/Runtime/Engine/Private/AudioDevice.cpp b/Engine/Source/Runtime/Engine/Private/AudioDevice.cpp
--- a/Engine/Source/Runtime/Engine/Private/AudioDevice.cpp
+++ b/Engine/Source/Runtime/Engine/Private/AudioDevice.cpp
@@ -536,13 +536,15 @@ bool FAudioDevice::Init(Audio::FDeviceId InDeviceID, int32 InMaxSources)
 		PrecacheStartupSounds();
 	}
 
-	UE_LOG(LogInit, Log, TEXT("FAudioDevice initialized."));
-
 	bIsInitialized = true;
 
 	FCoreUObjectDelegates::GetPreGarbageCollectDelegate().AddRaw(this, &FAudioDevice::OnPreGarbageCollect);
 	FCoreUObjectDelegates::PreGarbageCollectConditionalBeginDestroy.AddRaw(this, &FAudioDevice::OnPreGarbageCollect);
 
+	InitDefaultAudioBuses();
+
+	UE_LOG(LogInit, Log, TEXT("FAudioDevice initialized."));
+
 	return true;
 }
 
@@ -670,6 +672,8 @@ TRange<float> FAudioDevice::GetGlobalPitchRange() const
 
 void FAudioDevice::Teardown()
 {
+	ShutdownDefaultAudioBuses();
+
 	// Make sure we process any pending game thread tasks before tearing down the audio device.
 	FTaskGraphInterface::Get().ProcessThreadUntilIdle(ENamedThreads::GameThread);
 
diff --git a/Engine/Source/Runtime/Engine/Private/AudioSettings.cpp b/Engine/Source/Runtime/Engine/Private/AudioSettings.cpp
--- a/Engine/Source/Runtime/Engine/Private/AudioSettings.cpp
+++ b/Engine/Source/Runtime/Engine/Private/AudioSettings.cpp
@@ -1,15 +1,14 @@
 // Copyright Epic Games, Inc. All Rights Reserved.
-
-/*=============================================================================
-	AudioSettings.cpp: Unreal audio settings
-=============================================================================*/
-
 #include "Sound/AudioSettings.h"
+
+#include "AudioDevice.h"
+#include "AudioDeviceManager.h"
+#include "AudioMixerDevice.h"
 #include "Misc/ConfigCacheIni.h"
 #include "Misc/Paths.h"
-#include "Sound/SoundSubmix.h"
 #include "Sound/SoundBase.h"
 #include "Sound/SoundNodeQualityLevel.h"
+#include "Sound/SoundSubmix.h"
 #include "UObject/UObjectHash.h"
 #include "UObject/UObjectIterator.h"
 
@@ -126,6 +125,17 @@ void UAudioSettings::PostEditChangeChainProperty(FPropertyChangedChainEvent& Pro
 		{
 			bReconcileQualityNodes = true;
 		}
+		else if (PropertyName == GET_MEMBER_NAME_CHECKED(UAudioSettings, DefaultAudioBuses)
+			|| PropertyName == GET_MEMBER_NAME_CHECKED(FDefaultAudioBusSettings, AudioBus))
+		{
+			if (FAudioDeviceManager* DeviceManager = FAudioDeviceManager::Get())
+			{
+				DeviceManager->IterateOverAllDevices([this](Audio::FDeviceId, FAudioDevice* InDevice)
+				{
+					InDevice->InitDefaultAudioBuses();
+				});
+			}
+		}
 
 		if (bReconcileQualityNodes)
 		{
diff --git a/Engine/Source/Runtime/Engine/Public/AudioDevice.h b/Engine/Source/Runtime/Engine/Public/AudioDevice.h
--- a/Engine/Source/Runtime/Engine/Public/AudioDevice.h
+++ b/Engine/Source/Runtime/Engine/Public/AudioDevice.h
@@ -821,6 +821,7 @@ protected:
 	 */
 	void InitSoundSources();
 
+
 public:
 	/**
 	 * Registers a sound class with the audio device
@@ -835,6 +836,12 @@ public:
 	*/
 	void UnregisterSoundClass(USoundClass* SoundClass);
 
+	/* Initialized audio buses marked as default that are to be enabled for the duration of the application. */
+	virtual void InitDefaultAudioBuses() {}
+
+	/* Shutdown all audio buses marked as default. */
+	virtual void ShutdownDefaultAudioBuses() {}
+
 	/** Initializes sound submixes. */
 	virtual void InitSoundSubmixes() {}
 
