Fixed non editor build

diff --git a/Engine/Source/Runtime/Engine/Private/Materials/MaterialExpressions.cpp b/Engine/Source/Runtime/Engine/Private/Materials/MaterialExpressions.cpp
--- a/Engine/Source/Runtime/Engine/Private/Materials/MaterialExpressions.cpp
+++ b/Engine/Source/Runtime/Engine/Private/Materials/MaterialExpressions.cpp
@@ -19054,6 +19054,8 @@ FName UMaterialExpressionSceneDepthWithoutWater::GetInputName(int32 InputIndex)
 // Strata
 
 
+#if WITH_EDITOR
+
 // The compilation of an expression can sometimes lead to a INDEX_NONE code chunk when editing material graphs 
 // or when the node is inside a material function, linked to an input pin of the material function and that input is not plugged in to anything.
 // But for normals or tangents, Strata absolutely need a valid code chunk to de-duplicate when stored in memory. So we provide those helper functions.
@@ -19063,6 +19065,8 @@ static int32 CompileWithDefaultCodeChunk(class FMaterialCompiler* Compiler, FExp
 	return CodeChunk == -1 ? DefaultCodeChunk : CodeChunk;
 }
 
+#endif // WITH_EDITOR
+
 UMaterialExpressionStrataDiffuseBSDF::UMaterialExpressionStrataDiffuseBSDF(const FObjectInitializer& ObjectInitializer)
 	: Super(ObjectInitializer)
 {
