Add CVAR to optionally validate VT compression before saving to DDC, and after loading
- Errors on load can be recovered by recreating the corrupt data, but this will increase load times

diff --git a/Engine/Source/Runtime/Engine/Private/Texture2D.cpp b/Engine/Source/Runtime/Engine/Private/Texture2D.cpp
--- a/Engine/Source/Runtime/Engine/Private/Texture2D.cpp
+++ b/Engine/Source/Runtime/Engine/Private/Texture2D.cpp
@@ -158,7 +158,7 @@ void FTexture2DMipMap::Serialize(FArchive& Ar, UObject* Owner, int32 MipIdx)
 }
 
 #if WITH_EDITORONLY_DATA
-uint32 FTexture2DMipMap::StoreInDerivedDataCache(const FString& InDerivedDataKey, const FStringView& TextureName)
+uint32 FTexture2DMipMap::StoreInDerivedDataCache(const FString& InDerivedDataKey, const FStringView& TextureName, bool bReplaceExistingDDC)
 {
 	int32 BulkDataSizeInBytes = BulkData.GetBulkDataSize();
 	check(BulkDataSizeInBytes > 0);
@@ -172,7 +172,7 @@ uint32 FTexture2DMipMap::StoreInDerivedDataCache(const FString& InDerivedDataKey
 		BulkData.Unlock();
 	}
 	const uint32 Result = DerivedData.Num();
-	GetDerivedDataCacheRef().Put(*InDerivedDataKey, DerivedData, TextureName);
+	GetDerivedDataCacheRef().Put(*InDerivedDataKey, DerivedData, TextureName, bReplaceExistingDDC);
 	DerivedDataKey = InDerivedDataKey;
 	BulkData.RemoveBulkData();
 	return Result;
diff --git a/Engine/Source/Runtime/Engine/Private/TextureDerivedData.cpp b/Engine/Source/Runtime/Engine/Private/TextureDerivedData.cpp
--- a/Engine/Source/Runtime/Engine/Private/TextureDerivedData.cpp
+++ b/Engine/Source/Runtime/Engine/Private/TextureDerivedData.cpp
@@ -591,7 +591,7 @@ static void GetBuildSettingsPerFormat(const UTexture& Texture, const FTextureBui
  * @param bForceAllMipsToBeInlined - Whether to store all mips in the main DDC. Relates to how the texture resources get initialized (not supporting streaming).
  * @return number of bytes put to the DDC (total, including all mips)
  */
-uint32 PutDerivedDataInCache(FTexturePlatformData* DerivedData, const FString& DerivedDataKeySuffix, const FStringView& TextureName, bool bForceAllMipsToBeInlined)
+uint32 PutDerivedDataInCache(FTexturePlatformData* DerivedData, const FString& DerivedDataKeySuffix, const FStringView& TextureName, bool bForceAllMipsToBeInlined, bool bReplaceExistingDDC)
 {
 	TArray<uint8> RawDerivedData;
 	FString DerivedDataKey;
@@ -639,7 +639,7 @@ uint32 PutDerivedDataInCache(FTexturePlatformData* DerivedData, const FString& D
 		if (!bInline)
 		{
 			// store in the DDC, also drop the bulk data storage.
-			TotalBytesPut += Mip.StoreInDerivedDataCache(MipDerivedDataKey, TextureName);
+			TotalBytesPut += Mip.StoreInDerivedDataCache(MipDerivedDataKey, TextureName, bReplaceExistingDDC);
 		}
 	}
 
@@ -654,7 +654,7 @@ uint32 PutDerivedDataInCache(FTexturePlatformData* DerivedData, const FString& D
 				*FString::Printf(TEXT("%s_VTCHUNK%u"), *DerivedDataKeySuffix, ChunkIndex));
 
 			FVirtualTextureDataChunk& Chunk = DerivedData->VTData->Chunks[ChunkIndex];
-			TotalBytesPut += Chunk.StoreInDerivedDataCache(ChunkDerivedDataKey, TextureName);
+			TotalBytesPut += Chunk.StoreInDerivedDataCache(ChunkDerivedDataKey, TextureName, bReplaceExistingDDC);
 		}
 	}
 
diff --git a/Engine/Source/Runtime/Engine/Private/TextureDerivedDataTask.cpp b/Engine/Source/Runtime/Engine/Private/TextureDerivedDataTask.cpp
--- a/Engine/Source/Runtime/Engine/Private/TextureDerivedDataTask.cpp
+++ b/Engine/Source/Runtime/Engine/Private/TextureDerivedDataTask.cpp
@@ -35,6 +35,20 @@
 #include "ProfilingDebugging/CookStats.h"
 #include "VT/VirtualTextureDataBuilder.h"
 
+static TAutoConsoleVariable<int32> CVarVTValidateCompressionOnLoad(
+	TEXT("r.VT.ValidateCompressionOnLoad"),
+	0,
+	TEXT("Validates that VT data contains no compression errors when loading from DDC")
+	TEXT("This is slow, but allows debugging corrupt VT data (and allows recovering from bad DDC)")
+);
+
+static TAutoConsoleVariable<int32> CVarVTValidateCompressionOnSave(
+	TEXT("r.VT.ValidateCompressionOnSave"),
+	1,
+	TEXT("Validates that VT data contains no compression errors before saving to DDC")
+	TEXT("This is slow, but allows debugging corrupt VT data")
+);
+
 class FTextureStatusMessageContext : public FScopedSlowTask
 {
 public:
@@ -197,7 +211,7 @@ void FTextureSourceData::GetAsyncSourceMips(IImageWrapperModule* InImageWrapper)
 	}
 }
 
-void FTextureCacheDerivedDataWorker::BuildTexture()
+void FTextureCacheDerivedDataWorker::BuildTexture(bool bReplaceExistingDDC)
 {
 	TRACE_CPUPROFILER_EVENT_SCOPE(FTextureCacheDerivedDataWorker::BuildTexture);
 
@@ -233,20 +247,29 @@ void FTextureCacheDerivedDataWorker::BuildTexture()
 		DerivedData->PixelFormat = DerivedData->VTData->LayerTypes[0];
 		DerivedData->SetNumSlices(1);
 
-		// Store it in the cache.
-		// @todo: This will remove the streaming bulk data, which we immediately reload below!
-		// Should ideally avoid this redundant work, but it only happens when we actually have 
-		// to build the texture, which should only ever be once.
-		this->BytesCached = PutDerivedDataInCache(DerivedData, KeySuffix, Texture.GetPathName(), BuildSettingsPerLayer[0].bCubemap || BuildSettingsPerLayer[0].bVolume || BuildSettingsPerLayer[0].bTextureArray);
-
-		if (DerivedData->VTData->Chunks.Num())
+		bool bCompressionValid = true;
+		if (CVarVTValidateCompressionOnSave.GetValueOnAnyThread())
 		{
-			const bool bInlineMips = (CacheFlags & ETextureCacheFlags::InlineMips) != 0;
-			bSucceeded = !bInlineMips || DerivedData->TryInlineMipData(BuildSettingsPerLayer[0].LODBiasWithCinematicMips, &Texture);
+			bCompressionValid = DerivedData->VTData->ValidateCompression();
 		}
-		else
+
+		if (ensureMsgf(bCompressionValid, TEXT("Corrupt Virtual Texture compression for %s, can't store to DDC"), *Texture.GetPathName()))
 		{
-			UE_LOG(LogTexture, Warning, TEXT("Failed to build %s derived data for %s"), *BuildSettingsPerLayer[0].TextureFormatName.GetPlainNameString(), *Texture.GetPathName());
+			// Store it in the cache.
+			// @todo: This will remove the streaming bulk data, which we immediately reload below!
+			// Should ideally avoid this redundant work, but it only happens when we actually have 
+			// to build the texture, which should only ever be once.
+			this->BytesCached = PutDerivedDataInCache(DerivedData, KeySuffix, Texture.GetPathName(), BuildSettingsPerLayer[0].bCubemap || BuildSettingsPerLayer[0].bVolume || BuildSettingsPerLayer[0].bTextureArray, bReplaceExistingDDC);
+
+			if (DerivedData->VTData->Chunks.Num())
+			{
+				const bool bInlineMips = (CacheFlags & ETextureCacheFlags::InlineMips) != 0;
+				bSucceeded = !bInlineMips || DerivedData->TryInlineMipData(BuildSettingsPerLayer[0].LODBiasWithCinematicMips, &Texture);
+			}
+			else
+			{
+				UE_LOG(LogTexture, Warning, TEXT("Failed to build %s derived data for %s"), *BuildSettingsPerLayer[0].TextureFormatName.GetPlainNameString(), *Texture.GetPathName());
+			}
 		}
 	}
 	else if (bHasValidMip0)
@@ -321,7 +344,7 @@ void FTextureCacheDerivedDataWorker::BuildTexture()
 			// @todo: This will remove the streaming bulk data, which we immediately reload below!
 			// Should ideally avoid this redundant work, but it only happens when we actually have 
 			// to build the texture, which should only ever be once.
-			this->BytesCached = PutDerivedDataInCache(DerivedData, KeySuffix, Texture.GetPathName(), BuildSettingsPerLayer[0].bCubemap || (BuildSettingsPerLayer[0].bVolume && !GSupportsVolumeTextureStreaming) || (BuildSettingsPerLayer[0].bTextureArray && !GSupportsTexture2DArrayStreaming));
+			this->BytesCached = PutDerivedDataInCache(DerivedData, KeySuffix, Texture.GetPathName(), BuildSettingsPerLayer[0].bCubemap || (BuildSettingsPerLayer[0].bVolume && !GSupportsVolumeTextureStreaming) || (BuildSettingsPerLayer[0].bTextureArray && !GSupportsTexture2DArrayStreaming), bReplaceExistingDDC);
 		}
 
 		if (DerivedData->Mips.Num())
@@ -406,6 +429,7 @@ void FTextureCacheDerivedDataWorker::DoWork()
 	const bool bAllowAsyncBuild = (CacheFlags & ETextureCacheFlags::AllowAsyncBuild) != 0;
 	const bool bAllowAsyncLoading = (CacheFlags & ETextureCacheFlags::AllowAsyncLoading) != 0;
 	const bool bForVirtualTextureStreamingBuild = (CacheFlags & ETextureCacheFlags::ForVirtualTextureStreamingBuild) != 0;
+	bool bInvalidVirtualTextureCompression = false;
 
 	TArray<uint8> RawDerivedData;
 
@@ -464,6 +488,16 @@ void FTextureCacheDerivedDataWorker::DoWork()
 		}
 		bLoadedFromDDC = true;
 
+		if (bSucceeded && bForVirtualTextureStreamingBuild && CVarVTValidateCompressionOnLoad.GetValueOnAnyThread())
+		{
+			bSucceeded = DerivedData->VTData->ValidateCompression();
+			if (!bSucceeded)
+			{
+				UE_LOG(LogTexture, Error, TEXT("Texture %s has corrupt Virtual Texture compression. The texture will be rebuild."), *Texture.GetFullName());
+				bInvalidVirtualTextureCompression = true;
+			}
+		}
+
 		// Reset everything derived data so that we can do a clean load from the source data
 		if (!bSucceeded)
 		{
@@ -504,7 +538,19 @@ void FTextureCacheDerivedDataWorker::DoWork()
 		if (TextureData.Blocks.Num() && TextureData.Blocks[0].MipsPerLayer.Num() && TextureData.Blocks[0].MipsPerLayer[0].Num() && 
 			(!CompositeTextureData.IsValid() || (CompositeTextureData.Blocks.Num() && CompositeTextureData.Blocks[0].MipsPerLayer.Num() && CompositeTextureData.Blocks[0].MipsPerLayer[0].Num())))
 		{
-			BuildTexture();
+			// Replace any existing DDC data, if corrupt compression was detected
+			const bool bReplaceExistingDDC = bInvalidVirtualTextureCompression;
+			BuildTexture(bReplaceExistingDDC);
+			if (bInvalidVirtualTextureCompression && DerivedData->VTData)
+			{
+				// If we loaded data that turned out to be corrupt, flag it here so we can also recreate the VT data cached to local /DerivedDataCache/VT/ directory
+				for (FVirtualTextureDataChunk& Chunk : DerivedData->VTData->Chunks)
+				{
+					Chunk.bCorruptDataLoadedFromDDC = true;
+				}
+
+			}
+
 			bSucceeded = true;
 		}
 		else
diff --git a/Engine/Source/Runtime/Engine/Private/TextureDerivedDataTask.h b/Engine/Source/Runtime/Engine/Private/TextureDerivedDataTask.h
--- a/Engine/Source/Runtime/Engine/Private/TextureDerivedDataTask.h
+++ b/Engine/Source/Runtime/Engine/Private/TextureDerivedDataTask.h
@@ -27,7 +27,7 @@ enum
 #if WITH_EDITOR
 
 void GetTextureDerivedDataKeySuffix(const UTexture& Texture, const FTextureBuildSettings* BuildSettingsPerLayer, FString& OutKeySuffix);
-uint32 PutDerivedDataInCache(FTexturePlatformData* DerivedData, const FString& DerivedDataKeySuffix, const FStringView& TextureName, bool bForceAllMipsToBeInlined);
+uint32 PutDerivedDataInCache(FTexturePlatformData* DerivedData, const FString& DerivedDataKeySuffix, const FStringView& TextureName, bool bForceAllMipsToBeInlined, bool bReplaceExistingDDC);
 
 namespace ETextureCacheFlags
 {
@@ -132,7 +132,7 @@ class FTextureCacheDerivedDataWorker : public FNonAbandonableTask
 	bool bLoadedFromDDC = false;
 
 	/** Build the texture. This function is safe to call from any thread. */
-	void BuildTexture();
+	void BuildTexture(bool bReplaceExistingDDC = false);
 
 public:
 
diff --git a/Engine/Source/Runtime/Engine/Private/VT/VirtualTextureBuiltData.cpp b/Engine/Source/Runtime/Engine/Private/VT/VirtualTextureBuiltData.cpp
--- a/Engine/Source/Runtime/Engine/Private/VT/VirtualTextureBuiltData.cpp
+++ b/Engine/Source/Runtime/Engine/Private/VT/VirtualTextureBuiltData.cpp
@@ -180,6 +180,75 @@ void FVirtualTextureBuiltData::Serialize(FArchive& Ar, UObject* Owner, int32 Fir
 	}
 }
 
+bool FVirtualTextureBuiltData::ValidateCompression() const
+{
+	const uint32 TilePixelSize = GetPhysicalTileSize();
+	TArray<uint8> UncompressedResult;
+	TArray<uint8> ChunkDataDDC;
+
+	bool bResult = true;
+	for (int32 ChunkIndex = 0; bResult && ChunkIndex < Chunks.Num(); ++ChunkIndex)
+	{
+		const FVirtualTextureDataChunk& Chunk = Chunks[ChunkIndex];
+
+		const uint8* ChunkData = nullptr;
+		bool bNeedToUnlockBulkData = false;
+		if (Chunk.BulkData.GetBulkDataSize() > 0)
+		{
+			ChunkData = (uint8*)Chunk.BulkData.LockReadOnly();
+			bNeedToUnlockBulkData = true;
+		}
+		else
+		{
+			ChunkDataDDC.Reset();
+			const bool bDDCResult = GetDerivedDataCacheRef().GetSynchronous(*Chunk.DerivedDataKey, ChunkDataDDC);
+			check(bDDCResult);
+			ChunkData = ChunkDataDDC.GetData() + 4;
+		}
+
+		uint32 TileIndex = TileIndexPerChunk[ChunkIndex];
+		while (bResult && TileIndex < TileIndexPerChunk[ChunkIndex + 1])
+		{
+			for (uint32 LayerIndex = 0u; LayerIndex < GetNumLayers(); ++LayerIndex)
+			{
+				const EVirtualTextureCodec VTCodec = Chunk.CodecType[LayerIndex];
+				const EPixelFormat LayerFormat = LayerTypes[LayerIndex];
+				const uint32 TileWidthInBlocks = FMath::DivideAndRoundUp(TilePixelSize, (uint32)GPixelFormats[LayerFormat].BlockSizeX);
+				const uint32 TileHeightInBlocks = FMath::DivideAndRoundUp(TilePixelSize, (uint32)GPixelFormats[LayerFormat].BlockSizeY);
+				const uint32 PackedStride = TileWidthInBlocks * GPixelFormats[LayerFormat].BlockBytes;
+				const size_t PackedOutputSize = PackedStride * TileHeightInBlocks;
+
+				if (VTCodec == EVirtualTextureCodec::ZippedGPU)
+				{
+					const uint32 TileOffset = GetTileOffset(ChunkIndex, TileIndex);
+					const uint32 NextTileOffset = GetTileOffset(ChunkIndex, TileIndex + 1);
+					check(NextTileOffset >= TileOffset);
+					if (NextTileOffset > TileOffset)
+					{
+						const uint32 CompressedTileSize = NextTileOffset - TileOffset;
+
+						UncompressedResult.SetNumUninitialized(PackedOutputSize, false);
+						const bool bUncompressResult = FCompression::UncompressMemory(NAME_Zlib, UncompressedResult.GetData(), PackedOutputSize, &ChunkData[TileOffset], CompressedTileSize);
+						if (!bUncompressResult)
+						{
+							bResult = false;
+							break;
+						}
+					}
+				}
+				++TileIndex;
+			}
+		}
+
+		if (bNeedToUnlockBulkData)
+		{
+			Chunk.BulkData.Unlock();
+		}
+	}
+
+	return bResult;
+}
+
 #if WITH_EDITORONLY_DATA
 
 bool FVirtualTextureDataChunk::ShortenKey(const FString& CacheKey, FString& Result)
@@ -214,7 +283,7 @@ bool FVirtualTextureDataChunk::ShortenKey(const FString& CacheKey, FString& Resu
 	return true;
 }
 
-uint32 FVirtualTextureDataChunk::StoreInDerivedDataCache(const FString& InDerivedDataKey, const FStringView& TextureName)
+uint32 FVirtualTextureDataChunk::StoreInDerivedDataCache(const FString& InDerivedDataKey, const FStringView& TextureName, bool bReplaceExistingDDC)
 {
 	int32 BulkDataSizeInBytes = BulkData.GetBulkDataSize();
 	check(BulkDataSizeInBytes > 0);
@@ -228,7 +297,7 @@ uint32 FVirtualTextureDataChunk::StoreInDerivedDataCache(const FString& InDerive
 		BulkData.Unlock();
 	}
 	const uint32 Result = DerivedData.Num();
-	GetDerivedDataCacheRef().Put(*InDerivedDataKey, DerivedData, TextureName);
+	GetDerivedDataCacheRef().Put(*InDerivedDataKey, DerivedData, TextureName, bReplaceExistingDDC);
 	DerivedDataKey = InDerivedDataKey;
 	ShortenKey(DerivedDataKey, ShortDerivedDataKey);
 
diff --git a/Engine/Source/Runtime/Engine/Private/VT/VirtualTextureBuiltData.h b/Engine/Source/Runtime/Engine/Private/VT/VirtualTextureBuiltData.h
--- a/Engine/Source/Runtime/Engine/Private/VT/VirtualTextureBuiltData.h
+++ b/Engine/Source/Runtime/Engine/Private/VT/VirtualTextureBuiltData.h
@@ -59,8 +59,9 @@ struct FVirtualTextureDataChunk
 	FString ShortDerivedDataKey;
 	bool ShortenKey(const FString& CacheKey, FString& Result);
 	FThreadSafeBool bFileAvailableInVTDDCDache;
+	bool bCorruptDataLoadedFromDDC = false;
 
-	uint32 StoreInDerivedDataCache(const FString& InDerivedDataKey, const FStringView& TextureName);
+	uint32 StoreInDerivedDataCache(const FString& InDerivedDataKey, const FStringView& TextureName, bool bReplaceExistingDDC);
 #endif // WITH_EDITORONLY_DATA
 };
 
@@ -188,4 +189,10 @@ struct FVirtualTextureBuiltData
 		// This allows us to determine size of region by asking for start/end offsets
 		return Chunks[ChunkIndex].SizeInBytes;
 	}
+
+	/**
+	* Attempts to decompress every VT tile, to verify data is valid
+	* This will be very slow
+	*/
+	bool ValidateCompression() const;
 };
diff --git a/Engine/Source/Runtime/Engine/Private/VT/VirtualTextureChunkDDCCache.cpp b/Engine/Source/Runtime/Engine/Private/VT/VirtualTextureChunkDDCCache.cpp
--- a/Engine/Source/Runtime/Engine/Private/VT/VirtualTextureChunkDDCCache.cpp
+++ b/Engine/Source/Runtime/Engine/Private/VT/VirtualTextureChunkDDCCache.cpp
@@ -225,6 +225,14 @@ bool FVirtualTextureChunkDDCCache::MakeChunkAvailable(struct FVirtualTextureData
 	const FString CachedFilePath = AbsoluteCachePath / Chunk->ShortDerivedDataKey;
 	const FString TempFilePath = AbsoluteCachePath / FGuid::NewGuid().ToString() + ".tmp";
 
+	if (Chunk->bCorruptDataLoadedFromDDC)
+	{
+		// We determined data loaded from DDC was corrupt...this means any file saved to VT DDC cache is also corrupt and can no longer be used
+		FPlatformFileManager::Get().GetPlatformFile().DeleteFile(*CachedFilePath);
+		Chunk->bCorruptDataLoadedFromDDC = false;
+		Chunk->bFileAvailableInVTDDCDache = false;
+	}
+
 	// File already available? 
 	if (Chunk->bFileAvailableInVTDDCDache)
 	{
diff --git a/Engine/Source/Runtime/Engine/Public/TextureResource.h b/Engine/Source/Runtime/Engine/Public/TextureResource.h
--- a/Engine/Source/Runtime/Engine/Public/TextureResource.h
+++ b/Engine/Source/Runtime/Engine/Public/TextureResource.h
@@ -70,7 +70,7 @@ struct FTexture2DMipMap
 	 * Place mip-map data in the derived data cache associated with the provided
 	 * key.
 	 */
-	uint32 StoreInDerivedDataCache(const FString& InDerivedDataKey, const FStringView& TextureName);
+	uint32 StoreInDerivedDataCache(const FString& InDerivedDataKey, const FStringView& TextureName, bool bReplaceExistingDDC);
 #endif // #if WITH_EDITORONLY_DATA
 };
 
