Strata: make sure we do not crash and use default values when a input pin from a material function is not plugged in and if it is does not have use default value set to true.
(Input.Compile can return INDEX_NONE in this case)

diff --git a/Engine/Source/Runtime/Engine/Private/Materials/MaterialExpressions.cpp b/Engine/Source/Runtime/Engine/Private/Materials/MaterialExpressions.cpp
--- a/Engine/Source/Runtime/Engine/Private/Materials/MaterialExpressions.cpp
+++ b/Engine/Source/Runtime/Engine/Private/Materials/MaterialExpressions.cpp
@@ -19058,12 +19058,47 @@ FName UMaterialExpressionSceneDepthWithoutWater::GetInputName(int32 InputIndex)
 
 // The compilation of an expression can sometimes lead to a INDEX_NONE code chunk when editing material graphs 
 // or when the node is inside a material function, linked to an input pin of the material function and that input is not plugged in to anything.
-// But for normals or tangents, Strata absolutely need a valid code chunk to de-duplicate when stored in memory. So we provide those helper functions.
-static int32 CompileWithDefaultCodeChunk(class FMaterialCompiler* Compiler, FExpressionInput& Input, int DefaultCodeChunk)
+// But for normals or tangents, Strata absolutely need a valid code chunk to de-duplicate when stored in memory. 
+// Also, we want all our nodes to have default, as that is needed when creating BSDF, when registering code chunk representing material topology.
+static int32 CompileWithDefaultCodeChunk(class FMaterialCompiler* Compiler, FExpressionInput& Input, int DefaultCodeChunk, bool* bDefaultIsUsed = nullptr)
 {
+	if (bDefaultIsUsed)
+	{
+		*bDefaultIsUsed = Input.GetTracedInput().Expression == nullptr;
+	}
 	int32 CodeChunk = Input.GetTracedInput().Expression ? Input.Compile(Compiler) : DefaultCodeChunk;
+	if (bDefaultIsUsed)
+	{
+		*bDefaultIsUsed |= CodeChunk == -1;
+	}
 	return CodeChunk == -1 ? DefaultCodeChunk : CodeChunk;
 }
+static int32 CompileWithDefaultFloat1(class FMaterialCompiler* Compiler, FExpressionInput& Input, float X, bool* bDefaultIsUsed = nullptr)
+{
+	if (bDefaultIsUsed)
+	{
+		*bDefaultIsUsed = Input.GetTracedInput().Expression == nullptr;
+	}
+	int32 CodeChunk = Input.GetTracedInput().Expression ? Input.Compile(Compiler) : Compiler->Constant(X);
+	if (bDefaultIsUsed)
+	{
+		*bDefaultIsUsed |= CodeChunk == -1;
+	}
+	return CodeChunk == -1 ? Compiler->Constant(X) : CodeChunk;
+}
+static int32 CompileWithDefaultFloat3(class FMaterialCompiler* Compiler, FExpressionInput& Input, float X, float Y, float Z, bool* bDefaultIsUsed = nullptr)
+{
+	if (bDefaultIsUsed)
+	{
+		*bDefaultIsUsed = Input.GetTracedInput().Expression == nullptr;
+	}
+	int32 CodeChunk = Input.GetTracedInput().Expression ? Input.Compile(Compiler) : Compiler->Constant3(X, Y, Z);
+	if (bDefaultIsUsed)
+	{
+		*bDefaultIsUsed |= CodeChunk == -1;
+	}
+	return CodeChunk == -1 ? Compiler->Constant3(X, Y, Z) : CodeChunk;
+}
 
 #endif // WITH_EDITOR
 
@@ -19088,8 +19123,8 @@ int32 UMaterialExpressionStrataDiffuseBSDF::Compile(class FMaterialCompiler* Com
 	uint8 SharedNormalIndex = StrataCompilationInfoCreateSharedNormal(Compiler, NormalCodeChunk);
 
 	int32 OutputCodeChunk = Compiler->StrataDiffuseBSDF(
-		Albedo.GetTracedInput().Expression		? Albedo.Compile(Compiler)		: Compiler->Constant3(0.18f, 0.18f, 0.18f),
-		Roughness.GetTracedInput().Expression	? Roughness.Compile(Compiler)	: Compiler->Constant(0.0f),
+		CompileWithDefaultFloat3(Compiler, Albedo,		0.18f, 0.18f, 0.18f),
+		CompileWithDefaultFloat1(Compiler, Roughness,	0.0f),
 		NormalCodeChunk,
 		SharedNormalIndex);
 	StrataCompilationInfoCreateSingleBSDFMaterial(Compiler, OutputCodeChunk, SharedNormalIndex, STRATA_BSDF_TYPE_DIFFUSE);
@@ -19155,15 +19190,15 @@ int32 UMaterialExpressionStrataDielectricBSDF::Compile(class FMaterialCompiler*
 	int32 TangentCodeChunk = CompileWithDefaultCodeChunk(Compiler, Tangent, Compiler->VertexTangent());
 	uint8 SharedNormalIndex = StrataCompilationInfoCreateSharedNormal(Compiler, NormalCodeChunk, TangentCodeChunk);
 
-	int32 RoughnessXCodeChunk = RoughnessX.GetTracedInput().Expression ? RoughnessX.Compile(Compiler) : Compiler->Constant(0.0f);
+	int32 RoughnessXCodeChunk = CompileWithDefaultFloat1(Compiler, RoughnessX, 0.0f);
 	// If not plugged in, RoughnessYCodeChunk is set to RoughnessXCodeChunk to get an isotropic behavior
-	int32 RoughnessYCodeChunk = RoughnessY.GetTracedInput().Expression ? RoughnessY.Compile(Compiler) : RoughnessXCodeChunk;
+	int32 RoughnessYCodeChunk = CompileWithDefaultCodeChunk(Compiler, RoughnessY, RoughnessXCodeChunk);
 
 	int32 OutputCodeChunk = Compiler->StrataDielectricBSDF(
 		RoughnessXCodeChunk,
 		RoughnessYCodeChunk,
-		IOR.GetTracedInput().Expression			? IOR.Compile(Compiler)		: Compiler->Constant(1.5f),// Default to Glass
-		Tint.GetTracedInput().Expression		? Tint.Compile(Compiler)	: Compiler->Constant3(1.0f, 1.0f, 1.0f),
+		CompileWithDefaultFloat1(Compiler, IOR,	1.5f),	// Default to Glass
+		CompileWithDefaultFloat3(Compiler, Tint,1.0f, 1.0f, 1.0f),
 		NormalCodeChunk,
 		TangentCodeChunk,
 		SharedNormalIndex);
@@ -19239,13 +19274,13 @@ int32 UMaterialExpressionStrataConductorBSDF::Compile(class FMaterialCompiler* C
 	int32 TangentCodeChunk = CompileWithDefaultCodeChunk(Compiler, Tangent, Compiler->VertexTangent());
 	uint8 SharedNormalIndex = StrataCompilationInfoCreateSharedNormal(Compiler, NormalCodeChunk, TangentCodeChunk);
 
-	int32 RoughnessXCodeChunk = RoughnessX.GetTracedInput().Expression ? RoughnessX.Compile(Compiler) : Compiler->Constant(0.0f);
+	int32 RoughnessXCodeChunk = CompileWithDefaultFloat1(Compiler, RoughnessX, 0.0f);
 	// If not plugged in, RoughnessYCodeChunk is set to RoughnessXCodeChunk to get an isotropic behavior
-	int32 RoughnessYCodeChunk = RoughnessY.GetTracedInput().Expression ? RoughnessY.Compile(Compiler) : RoughnessXCodeChunk;
+	int32 RoughnessYCodeChunk = CompileWithDefaultCodeChunk(Compiler, RoughnessY, RoughnessXCodeChunk);
 
 	int32 OutputCodeChunk = Compiler->StrataConductorBSDF(
-		Reflectivity.GetTracedInput().Expression	? Reflectivity.Compile(Compiler): Compiler->Constant3(0.947f, 0.776f, 0.371f),	// Default to Gold
-		EdgeColor.GetTracedInput().Expression		? EdgeColor.Compile(Compiler)	: Compiler->Constant3(1.000f, 0.982f, 0.753f),	// Default to Gold
+		CompileWithDefaultFloat3(Compiler, Reflectivity,	0.947f, 0.776f, 0.371f),	// Default to Gold
+		CompileWithDefaultFloat3(Compiler, EdgeColor,		1.000f, 0.982f, 0.753f),	// Default to Gold
 		RoughnessXCodeChunk,
 		RoughnessYCodeChunk,
 		NormalCodeChunk,
@@ -19322,15 +19357,16 @@ int32 UMaterialExpressionStrataVolumeBSDF::Compile(class FMaterialCompiler* Comp
 	int32 NormalCodeChunk = Compiler->VertexNormal();
 	uint8 SharedNormalIndex = StrataCompilationInfoCreateSharedNormal(Compiler, NormalCodeChunk);
 
-	const bool bHasScattering = Albedo.GetTracedInput().Expression != nullptr;
-
+	bool DefaultAlbedoIsUsed = false;
 	int32 OutputCodeChunk = Compiler->StrataVolumeBSDF(
-		bHasScattering							? Albedo.Compile(Compiler)			: Compiler->Constant3(0.0f, 0.0f, 0.0f), // Default must be 0 to remove memory used by Albedo when it is not pluged (bHasScattering==false)
-		Extinction.GetTracedInput().Expression	? Extinction.Compile(Compiler)		: Compiler->Constant3(0.0f, 0.0f, 0.0f),
-		Anisotropy.GetTracedInput().Expression	? Anisotropy.Compile(Compiler)		: Compiler->Constant(0.0f),
-		Thickness.GetTracedInput().Expression	? Thickness.Compile(Compiler)		: Compiler->Constant(0.001f), // default = 1mm
+		CompileWithDefaultFloat3(Compiler, Albedo,		0.0f, 0.0f, 0.0f, &DefaultAlbedoIsUsed),  // Default must be 0 to remove memory used by Albedo when it is not pluged (bHasScattering==false)
+		CompileWithDefaultFloat3(Compiler, Extinction,	0.0f, 0.0f, 0.0f),
+		CompileWithDefaultFloat1(Compiler, Anisotropy,	0.0f),
+		CompileWithDefaultFloat1(Compiler, Thickness,	0.001f),	// default = 1mm
 		NormalCodeChunk,
 		SharedNormalIndex);
+
+	const bool bHasScattering = !DefaultAlbedoIsUsed;
 	StrataCompilationInfoCreateSingleBSDFMaterial(Compiler, OutputCodeChunk, SharedNormalIndex, STRATA_BSDF_TYPE_VOLUME, bHasScattering);
 
 	return OutputCodeChunk;
@@ -19394,17 +19430,11 @@ UMaterialExpressionStrataSheenBSDF::UMaterialExpressionStrataSheenBSDF(const FOb
 int32 UMaterialExpressionStrataSheenBSDF::Compile(class FMaterialCompiler* Compiler, int32 OutputIndex)
 {
 	int32 NormalCodeChunk = CompileWithDefaultCodeChunk(Compiler, Normal, Compiler->VertexNormal());
-
-	// For some reason, when editing shader, Normal.Compile() / Tangent.Compile() might return -1. In this case we fallback to vertex normal/tangent
-	if (NormalCodeChunk == -1) { NormalCodeChunk = Compiler->VertexNormal(); }
-
 	uint8 SharedNormalIndex = StrataCompilationInfoCreateSharedNormal(Compiler, NormalCodeChunk);
 
-	int32 RoughnessCodeChunk = Roughness.GetTracedInput().Expression ? Roughness.Compile(Compiler) : Compiler->Constant(0.0f);
-
 	int32 OutputCodeChunk = Compiler->StrataSheenBSDF(
-		Albedo.GetTracedInput().Expression ? Albedo.Compile(Compiler) : Compiler->Constant3(0.0f, 0.0f, 0.0f),
-		RoughnessCodeChunk,
+		CompileWithDefaultFloat3(Compiler, Albedo,		0.0f, 0.0f, 0.0f),
+		CompileWithDefaultFloat1(Compiler, Roughness,	0.0f),
 		NormalCodeChunk,
 		SharedNormalIndex);
 	StrataCompilationInfoCreateSingleBSDFMaterial(Compiler, OutputCodeChunk, SharedNormalIndex, STRATA_BSDF_TYPE_SHEEN);
