Fix for one of the current compile errors for monolithics.

diff --git a/Engine/Source/Runtime/Engine/Private/VT/VirtualTextureBuiltData.cpp b/Engine/Source/Runtime/Engine/Private/VT/VirtualTextureBuiltData.cpp
--- a/Engine/Source/Runtime/Engine/Private/VT/VirtualTextureBuiltData.cpp
+++ b/Engine/Source/Runtime/Engine/Private/VT/VirtualTextureBuiltData.cpp
@@ -198,6 +198,7 @@ bool FVirtualTextureBuiltData::ValidateCompression() const
 			ChunkData = (uint8*)Chunk.BulkData.LockReadOnly();
 			bNeedToUnlockBulkData = true;
 		}
+#if WITH_EDITORONLY_DATA
 		else
 		{
 			ChunkDataDDC.Reset();
@@ -205,6 +206,7 @@ bool FVirtualTextureBuiltData::ValidateCompression() const
 			check(bDDCResult);
 			ChunkData = ChunkDataDDC.GetData() + 4;
 		}
+#endif // WITH_EDITORONLY_DATA
 
 		uint32 TileIndex = TileIndexPerChunk[ChunkIndex];
 		while (bResult && TileIndex < TileIndexPerChunk[ChunkIndex + 1])
