Use non-deprecated version of GetSynchronous()

diff --git a/Engine/Source/Runtime/Engine/Private/TextureDerivedDataTask.cpp b/Engine/Source/Runtime/Engine/Private/TextureDerivedDataTask.cpp
--- a/Engine/Source/Runtime/Engine/Private/TextureDerivedDataTask.cpp
+++ b/Engine/Source/Runtime/Engine/Private/TextureDerivedDataTask.cpp
@@ -250,7 +250,7 @@ void FTextureCacheDerivedDataWorker::BuildTexture(bool bReplaceExistingDDC)
 		bool bCompressionValid = true;
 		if (CVarVTValidateCompressionOnSave.GetValueOnAnyThread())
 		{
-			bCompressionValid = DerivedData->VTData->ValidateCompression();
+			bCompressionValid = DerivedData->VTData->ValidateCompression(Texture.GetPathName());
 		}
 
 		if (ensureMsgf(bCompressionValid, TEXT("Corrupt Virtual Texture compression for %s, can't store to DDC"), *Texture.GetPathName()))
@@ -490,7 +490,7 @@ void FTextureCacheDerivedDataWorker::DoWork()
 
 		if (bSucceeded && bForVirtualTextureStreamingBuild && CVarVTValidateCompressionOnLoad.GetValueOnAnyThread())
 		{
-			bSucceeded = DerivedData->VTData->ValidateCompression();
+			bSucceeded = DerivedData->VTData->ValidateCompression(Texture.GetPathName());
 			if (!bSucceeded)
 			{
 				UE_LOG(LogTexture, Error, TEXT("Texture %s has corrupt Virtual Texture compression. The texture will be rebuild."), *Texture.GetFullName());
diff --git a/Engine/Source/Runtime/Engine/Private/VT/VirtualTextureBuiltData.cpp b/Engine/Source/Runtime/Engine/Private/VT/VirtualTextureBuiltData.cpp
--- a/Engine/Source/Runtime/Engine/Private/VT/VirtualTextureBuiltData.cpp
+++ b/Engine/Source/Runtime/Engine/Private/VT/VirtualTextureBuiltData.cpp
@@ -180,7 +180,7 @@ void FVirtualTextureBuiltData::Serialize(FArchive& Ar, UObject* Owner, int32 Fir
 	}
 }
 
-bool FVirtualTextureBuiltData::ValidateCompression() const
+bool FVirtualTextureBuiltData::ValidateCompression(FStringView const& InDDCDebugContext) const
 {
 	const uint32 TilePixelSize = GetPhysicalTileSize();
 	TArray<uint8> UncompressedResult;
@@ -202,7 +202,7 @@ bool FVirtualTextureBuiltData::ValidateCompression() const
 		else
 		{
 			ChunkDataDDC.Reset();
-			const bool bDDCResult = GetDerivedDataCacheRef().GetSynchronous(*Chunk.DerivedDataKey, ChunkDataDDC);
+			const bool bDDCResult = GetDerivedDataCacheRef().GetSynchronous(*Chunk.DerivedDataKey, ChunkDataDDC, InDDCDebugContext);
 			check(bDDCResult);
 			ChunkData = ChunkDataDDC.GetData() + 4;
 		}
diff --git a/Engine/Source/Runtime/Engine/Private/VT/VirtualTextureBuiltData.h b/Engine/Source/Runtime/Engine/Private/VT/VirtualTextureBuiltData.h
--- a/Engine/Source/Runtime/Engine/Private/VT/VirtualTextureBuiltData.h
+++ b/Engine/Source/Runtime/Engine/Private/VT/VirtualTextureBuiltData.h
@@ -194,5 +194,5 @@ struct FVirtualTextureBuiltData
 	* Attempts to decompress every VT tile, to verify data is valid
 	* This will be very slow
 	*/
-	bool ValidateCompression() const;
+	bool ValidateCompression(FStringView const& InDDCDebugContext) const;
 };
