To protect the skeletalmesh when using async operations like build or re-import we deprecate all public member variables of USkeletalMesh class and create the necessary getter and setter. Those member

diff --git a/Engine/Source/Runtime/Engine/Classes/Engine/SkeletalMesh.h b/Engine/Source/Runtime/Engine/Classes/Engine/SkeletalMesh.h
--- a/Engine/Source/Runtime/Engine/Classes/Engine/SkeletalMesh.h
+++ b/Engine/Source/Runtime/Engine/Classes/Engine/SkeletalMesh.h
@@ -535,6 +535,7 @@ public:
 	 * There is only one editor data asset possible per skeletalmesh.
 	 * The reason we store the editor data in a separate asset is because the size of it can be very big and affect the editor performance. (undo/redo transactions)
 	 */
+	UE_DEPRECATED(4.27, "Please do not access this member directly; use the public ImportData API.")
 	UPROPERTY()
 	mutable USkeletalMeshEditorData* MeshEditorDataObject;
 
@@ -544,6 +545,13 @@ private:
 	 */
 	USkeletalMeshEditorData& GetMeshEditorData() const;
 
+	bool IsMeshEditorDataValid() const
+	{
+		PRAGMA_DISABLE_DEPRECATION_WARNINGS
+		return MeshEditorDataObject != nullptr;
+		PRAGMA_ENABLE_DEPRECATION_WARNINGS
+	}
+
 #endif //WITH_EDITORONLY_DATA
 
 public:
@@ -618,9 +626,39 @@ public:
 	FORCEINLINE FSkeletalMeshRenderData* GetResourceForRendering() const { return SkeletalMeshRenderData.Get(); }
 
 	/** Skeleton of this skeletal mesh **/
-	UPROPERTY(Category=Mesh, AssetRegistrySearchable, VisibleAnywhere, BlueprintReadOnly)
+	UE_DEPRECATED(4.27, "Please do not access this member directly; use USkeletalMesh::GetSkeleton() or USkeletalMesh::SetSkeleton().")
+	UPROPERTY(Category=Mesh, AssetRegistrySearchable, VisibleAnywhere, BlueprintGetter = GetSkeleton)
 	USkeleton* Skeleton;
 
+	static FName GetSkeletonMemberName()
+	{
+		PRAGMA_DISABLE_DEPRECATION_WARNINGS
+		return GET_MEMBER_NAME_CHECKED(USkeletalMesh, Skeleton);
+		PRAGMA_ENABLE_DEPRECATION_WARNINGS
+	}
+
+	USkeleton* GetSkeleton()
+	{
+		PRAGMA_DISABLE_DEPRECATION_WARNINGS
+		return Skeleton;
+		PRAGMA_ENABLE_DEPRECATION_WARNINGS
+	}
+
+	UFUNCTION(BlueprintGetter)
+	const USkeleton* GetSkeleton() const
+	{
+		PRAGMA_DISABLE_DEPRECATION_WARNINGS
+		return Skeleton;
+		PRAGMA_ENABLE_DEPRECATION_WARNINGS
+	}
+
+	void SetSkeleton(USkeleton* InSkeleton)
+	{
+		PRAGMA_DISABLE_DEPRECATION_WARNINGS
+		Skeleton = InSkeleton;
+		PRAGMA_ENABLE_DEPRECATION_WARNINGS
+	}
+
 private:
 	/** Original imported mesh bounds */
 	UPROPERTY(transient, duplicatetransient)
@@ -694,13 +732,73 @@ public:
 #endif
 
 	/** List of materials applied to this mesh. */
-	UPROPERTY(EditAnywhere, BlueprintReadWrite, transient, duplicatetransient, Category=SkeletalMesh)
+	UE_DEPRECATED(4.27, "Please do not access this member directly; use USkeletalMesh::GetMaterials() or USkeletalMesh::SetMaterials().")
+	UPROPERTY(EditAnywhere, BlueprintGetter = GetMaterials, BlueprintSetter = SetMaterials, transient, duplicatetransient, Category=SkeletalMesh)
 	TArray<FSkeletalMaterial> Materials;
+	
+	static FName GetMaterialsMemberName()
+	{
+		PRAGMA_DISABLE_DEPRECATION_WARNINGS
+		return GET_MEMBER_NAME_CHECKED(USkeletalMesh, Materials);
+		PRAGMA_ENABLE_DEPRECATION_WARNINGS
+	}
+
+	TArray<FSkeletalMaterial>& GetMaterials()
+	{
+		PRAGMA_DISABLE_DEPRECATION_WARNINGS
+		return Materials;
+		PRAGMA_ENABLE_DEPRECATION_WARNINGS
+	}
+
+	UFUNCTION(BlueprintGetter)
+	const TArray<FSkeletalMaterial>& GetMaterials() const
+	{
+		PRAGMA_DISABLE_DEPRECATION_WARNINGS
+		return Materials;
+		PRAGMA_ENABLE_DEPRECATION_WARNINGS
+	}
+
+	UFUNCTION(BlueprintSetter)
+	void SetMaterials(const TArray<FSkeletalMaterial>& InMaterials)
+	{
+		PRAGMA_DISABLE_DEPRECATION_WARNINGS
+		Materials = InMaterials;
+		PRAGMA_ENABLE_DEPRECATION_WARNINGS
+	}
 
 	/** List of bones that should be mirrored. */
+	UE_DEPRECATED(4.27, "Please do not access this member directly; use USkeletalMesh::GetSkelMirrorTable() or USkeletalMesh::SetSkelMirrorTable().")
 	UPROPERTY(EditAnywhere, editfixedsize, Category=Mirroring)
 	TArray<struct FBoneMirrorInfo> SkelMirrorTable;
 
+	static FName GetSkelMirrorTableMemberName()
+	{
+		PRAGMA_DISABLE_DEPRECATION_WARNINGS
+		return GET_MEMBER_NAME_CHECKED(USkeletalMesh, SkelMirrorTable);
+		PRAGMA_ENABLE_DEPRECATION_WARNINGS
+	}
+
+	TArray<struct FBoneMirrorInfo>& GetSkelMirrorTable()
+	{
+		PRAGMA_DISABLE_DEPRECATION_WARNINGS
+		return SkelMirrorTable;
+		PRAGMA_ENABLE_DEPRECATION_WARNINGS
+	}
+
+	const TArray<struct FBoneMirrorInfo>& GetSkelMirrorTable() const
+	{
+		PRAGMA_DISABLE_DEPRECATION_WARNINGS
+		return SkelMirrorTable;
+		PRAGMA_ENABLE_DEPRECATION_WARNINGS
+	}
+
+	void SetSkelMirrorTable(const TArray<struct FBoneMirrorInfo>& InSkelMirrorTable)
+	{
+		PRAGMA_DISABLE_DEPRECATION_WARNINGS
+		SkelMirrorTable = InSkelMirrorTable;
+		PRAGMA_ENABLE_DEPRECATION_WARNINGS
+	}
+
 private:
 	/** Struct containing information for each LOD level, such as materials to use, and when use the LOD. */
 	UPROPERTY(EditAnywhere, EditFixedSize, Category=LevelOfDetail)
@@ -724,38 +822,211 @@ private:
 
 public:
 	/** Minimum LOD to render. Can be overridden per component as well as set here for all mesh instances here */
+	UE_DEPRECATED(4.27, "Please do not access this member directly; use USkeletalMesh::GetMinLod() or USkeletalMesh::SetMinLod().")
 	UPROPERTY(EditAnywhere, Category = LODSettings, meta = (DisplayName = "Minimum LOD"))
 	FPerPlatformInt MinLod;
 
+	static FName GetMinLodMemberName()
+	{
+		PRAGMA_DISABLE_DEPRECATION_WARNINGS
+		return GET_MEMBER_NAME_CHECKED(USkeletalMesh, MinLod);
+		PRAGMA_ENABLE_DEPRECATION_WARNINGS
+	}
+
+	const FPerPlatformInt& GetMinLod() const
+	{
+		PRAGMA_DISABLE_DEPRECATION_WARNINGS
+		return MinLod;
+		PRAGMA_ENABLE_DEPRECATION_WARNINGS
+	}
+
+	void SetMinLod(FPerPlatformInt InMinLod)
+	{
+		PRAGMA_DISABLE_DEPRECATION_WARNINGS
+		MinLod = MoveTemp(InMinLod);
+		PRAGMA_ENABLE_DEPRECATION_WARNINGS
+	}
+
 	/** when true all lods below minlod will still be cooked */
+	UE_DEPRECATED(4.27, "Please do not access this member directly; use USkeletalMesh::GetDisableBelowMinLodStripping() or USkeletalMesh::SetDisableBelowMinLodStripping().")
 	UPROPERTY(EditAnywhere, Category = LODSettings)
 	FPerPlatformBool DisableBelowMinLodStripping;
 
+	static FName GetDisableBelowMinLodStrippingMemberName()
+	{
+		PRAGMA_DISABLE_DEPRECATION_WARNINGS
+		return GET_MEMBER_NAME_CHECKED(USkeletalMesh, DisableBelowMinLodStripping);
+		PRAGMA_ENABLE_DEPRECATION_WARNINGS
+	}
+
+	const FPerPlatformBool& GetDisableBelowMinLodStripping() const
+	{
+		PRAGMA_DISABLE_DEPRECATION_WARNINGS
+		return DisableBelowMinLodStripping;
+		PRAGMA_ENABLE_DEPRECATION_WARNINGS
+	}
+
+	void SetDisableBelowMinLodStripping(FPerPlatformBool InDisableBelowMinLodStripping)
+	{
+		PRAGMA_DISABLE_DEPRECATION_WARNINGS
+		DisableBelowMinLodStripping = MoveTemp(InDisableBelowMinLodStripping);
+		PRAGMA_ENABLE_DEPRECATION_WARNINGS
+	}
+
 #if WITH_EDITORONLY_DATA
 	/** Whether this skeletal mesh overrides default LOD streaming settings. */
+	UE_DEPRECATED(4.27, "Please do not access this member directly; use USkeletalMesh::GetOverrideLODStreamingSettings() or USkeletalMesh::SetOverrideLODStreamingSettings().")
 	UPROPERTY(EditAnywhere, Category=LODSettings)
 	bool bOverrideLODStreamingSettings;
+	static FName GetOverrideLODStreamingSettingsMemberName()
+	{
+		PRAGMA_DISABLE_DEPRECATION_WARNINGS
+		return GET_MEMBER_NAME_CHECKED(USkeletalMesh, bOverrideLODStreamingSettings);
+		PRAGMA_ENABLE_DEPRECATION_WARNINGS
+	}
+
+	bool GetOverrideLODStreamingSettings() const
+	{
+		PRAGMA_DISABLE_DEPRECATION_WARNINGS
+		return bOverrideLODStreamingSettings;
+		PRAGMA_ENABLE_DEPRECATION_WARNINGS
+	}
+
+	void SetOverrideLODStreamingSettings(bool bInOverrideLODStreamingSettings)
+	{
+		PRAGMA_DISABLE_DEPRECATION_WARNINGS
+		bOverrideLODStreamingSettings = bInOverrideLODStreamingSettings;
+		PRAGMA_ENABLE_DEPRECATION_WARNINGS
+	}
 
 	/** Whether we can stream the LODs of this mesh */
+	UE_DEPRECATED(4.27, "Please do not access this member directly; use USkeletalMesh::GetSupportLODStreaming() or USkeletalMesh::SetSupportLODStreaming().")
 	UPROPERTY(EditAnywhere, Category=LODSettings, meta=(DisplayName="Stream LODs", EditCondition="bOverrideLODStreamingSettings"))
 	FPerPlatformBool bSupportLODStreaming;
+	static FName GetSupportLODStreamingMemberName()
+	{
+		PRAGMA_DISABLE_DEPRECATION_WARNINGS
+		return GET_MEMBER_NAME_CHECKED(USkeletalMesh, bSupportLODStreaming);
+		PRAGMA_ENABLE_DEPRECATION_WARNINGS
+	}
+
+	const FPerPlatformBool& GetSupportLODStreaming() const
+	{
+		PRAGMA_DISABLE_DEPRECATION_WARNINGS
+		return bSupportLODStreaming;
+		PRAGMA_ENABLE_DEPRECATION_WARNINGS
+	}
+
+	void SetSupportLODStreaming(FPerPlatformBool bInSupportLODStreaming)
+	{
+		PRAGMA_DISABLE_DEPRECATION_WARNINGS
+		bSupportLODStreaming = MoveTemp(bInSupportLODStreaming);
+		PRAGMA_ENABLE_DEPRECATION_WARNINGS
+	}
 
 	/** Maximum number of LODs that can be streamed */
+	UE_DEPRECATED(4.27, "Please do not access this member directly; use USkeletalMesh::GetMaxNumStreamedLODs() or USkeletalMesh::SetMaxNumStreamedLODs().")
 	UPROPERTY(EditAnywhere, Category=LODSettings, meta=(EditCondition="bOverrideLODStreamingSettings"))
 	FPerPlatformInt MaxNumStreamedLODs;
+	static FName GetMaxNumStreamedLODsMemberName()
+	{
+		PRAGMA_DISABLE_DEPRECATION_WARNINGS
+		return GET_MEMBER_NAME_CHECKED(USkeletalMesh, MaxNumStreamedLODs);
+		PRAGMA_ENABLE_DEPRECATION_WARNINGS
+	}
+
+	const FPerPlatformInt& GetMaxNumStreamedLODs() const
+	{
+		PRAGMA_DISABLE_DEPRECATION_WARNINGS
+		return MaxNumStreamedLODs;
+		PRAGMA_ENABLE_DEPRECATION_WARNINGS
+	}
+
+	void SetMaxNumStreamedLODs(FPerPlatformInt InMaxNumStreamedLODs)
+	{
+		PRAGMA_DISABLE_DEPRECATION_WARNINGS
+		MaxNumStreamedLODs = MoveTemp(InMaxNumStreamedLODs);
+		PRAGMA_ENABLE_DEPRECATION_WARNINGS
+	}
 
 	/** Maximum number of LODs below min LOD level that can be saved to optional pak (currently, need to be either 0 or > num of LODs below MinLod) */
+	UE_DEPRECATED(4.27, "Please do not access this member directly; use USkeletalMesh::GetMaxNumOptionalLODs() or USkeletalMesh::SetMaxNumOptionalLODs().")
 	UPROPERTY(EditAnywhere, Category=LODSettings, meta=(EditCondition="bOverrideLODStreamingSettings"))
 	FPerPlatformInt MaxNumOptionalLODs;
+	static FName GetMaxNumOptionalLODsMemberName()
+	{
+		PRAGMA_DISABLE_DEPRECATION_WARNINGS
+		return GET_MEMBER_NAME_CHECKED(USkeletalMesh, MaxNumOptionalLODs);
+		PRAGMA_ENABLE_DEPRECATION_WARNINGS
+	}
+
+	const FPerPlatformInt& GetMaxNumOptionalLODs() const
+	{
+		PRAGMA_DISABLE_DEPRECATION_WARNINGS
+		return MaxNumOptionalLODs;
+		PRAGMA_ENABLE_DEPRECATION_WARNINGS
+	}
 
-	UPROPERTY(EditAnywhere, BlueprintReadWrite, AssetRegistrySearchable, BlueprintSetter = SetLODSettings, Category = LODSettings)
+	void SetMaxNumOptionalLODs(FPerPlatformInt InMaxNumOptionalLODs)
+	{
+		PRAGMA_DISABLE_DEPRECATION_WARNINGS
+		MaxNumOptionalLODs = MoveTemp(InMaxNumOptionalLODs);
+		PRAGMA_ENABLE_DEPRECATION_WARNINGS
+	}
+
+	UE_DEPRECATED(4.27, "Please do not access this member directly; use USkeletalMesh::GetLODSettings() or USkeletalMesh::SetLODSettings().")
+	UPROPERTY(EditAnywhere, AssetRegistrySearchable, BlueprintGetter = GetLODSettings, BlueprintSetter = SetLODSettings, Category = LODSettings)
 	USkeletalMeshLODSettings* LODSettings;
 
+	static FName GetLODSettingsMemberName()
+	{
+		PRAGMA_DISABLE_DEPRECATION_WARNINGS
+		return GET_MEMBER_NAME_CHECKED(USkeletalMesh, LODSettings);
+		PRAGMA_ENABLE_DEPRECATION_WARNINGS
+	}
+
 	/** The Default Control Rig To Animate with when used in Sequnecer. */
+	UE_DEPRECATED(4.27, "Please do not access this member directly; use USkeletalMesh::GetDefaultAnimatingRig() or USkeletalMesh::SetDefaultAnimatingRig().")
 	UPROPERTY(EditAnywhere, Category = AnimationRig, BlueprintGetter = GetDefaultAnimatingRig, BlueprintSetter = SetDefaultAnimatingRig, meta = (AllowedClasses = "ControlRigBlueprint"))
 	TSoftObjectPtr<UObject> DefaultAnimatingRig;
+	static FName GetDefaultAnimatingRigMemberName()
+	{
+		PRAGMA_DISABLE_DEPRECATION_WARNINGS
+		return GET_MEMBER_NAME_CHECKED(USkeletalMesh, DefaultAnimatingRig);
+		PRAGMA_ENABLE_DEPRECATION_WARNINGS
+	}
+	
+#endif // WITH_EDITORONLY_DATA
 
+	USkeletalMeshLODSettings* GetLODSettings()
+	{
+		PRAGMA_DISABLE_DEPRECATION_WARNINGS
+#if WITH_EDITORONLY_DATA
+		return LODSettings;
+#else
+		const bool bCallOutsideOf_WithEditorOnlyData = false;
+		ensure(bCallOutsideOf_WithEditorOnlyData);
+		return nullptr;
 #endif // WITH_EDITORONLY_DATA
+		PRAGMA_ENABLE_DEPRECATION_WARNINGS
+	}
+
+	UFUNCTION(BlueprintGetter)
+	const USkeletalMeshLODSettings* GetLODSettings() const
+	{
+		PRAGMA_DISABLE_DEPRECATION_WARNINGS
+#if WITH_EDITORONLY_DATA
+		return LODSettings;
+#else
+		const bool bCallOutsideOf_WithEditorOnlyData = false;
+		ensure(bCallOutsideOf_WithEditorOnlyData);
+		return nullptr;
+#endif // WITH_EDITORONLY_DATA
+		PRAGMA_ENABLE_DEPRECATION_WARNINGS
+	}
+
+	UFUNCTION(BlueprintSetter)
+	void SetLODSettings(USkeletalMeshLODSettings* InLODSettings);
 
 #if WITH_EDITOR
 	/** Get whether this mesh use LOD streaming. Do not use bSupportLODStreaming directly. Call this method instead. */
@@ -768,20 +1039,59 @@ public:
 	int32 GetMaxNumOptionalLODs(const class ITargetPlatform* TargetPlatform) const;
 #endif
 
-	UFUNCTION(BlueprintSetter)
-	void SetLODSettings(USkeletalMeshLODSettings* InLODSettings);
 
 	UFUNCTION(BlueprintSetter)
 	void SetDefaultAnimatingRig(TSoftObjectPtr<UObject> InAnimatingRig);
 
 	UFUNCTION(BlueprintGetter)
-	TSoftObjectPtr<UObject> GetDefaultAnimatingRig();
+	TSoftObjectPtr<UObject> GetDefaultAnimatingRig() const;
 
 	UPROPERTY(EditAnywhere, Category=Mirroring)
 	TEnumAsByte<EAxis::Type> SkelMirrorAxis;
+	static FName GetSkelMirrorAxisMemberName()
+	{
+		PRAGMA_DISABLE_DEPRECATION_WARNINGS
+		return GET_MEMBER_NAME_CHECKED(USkeletalMesh, SkelMirrorAxis);
+		PRAGMA_ENABLE_DEPRECATION_WARNINGS
+	}
+
+	TEnumAsByte<EAxis::Type> GetSkelMirrorAxis() const
+	{
+		PRAGMA_DISABLE_DEPRECATION_WARNINGS
+		return SkelMirrorAxis;
+		PRAGMA_ENABLE_DEPRECATION_WARNINGS
+	}
 
+	void SetSkelMirrorAxis(TEnumAsByte<EAxis::Type> InSkelMirrorAxis)
+	{
+		PRAGMA_DISABLE_DEPRECATION_WARNINGS
+		SkelMirrorAxis = InSkelMirrorAxis;
+		PRAGMA_ENABLE_DEPRECATION_WARNINGS
+	}
+
+	UE_DEPRECATED(4.27, "Please do not access this member directly; use USkeletalMesh::GetSkelMirrorFlipAxis() or USkeletalMesh::SetSkelMirrorFlipAxis().")
 	UPROPERTY(EditAnywhere, Category=Mirroring)
 	TEnumAsByte<EAxis::Type> SkelMirrorFlipAxis;
+	static FName GetSkelMirrorFlipAxisMemberName()
+	{
+		PRAGMA_DISABLE_DEPRECATION_WARNINGS
+		return GET_MEMBER_NAME_CHECKED(USkeletalMesh, SkelMirrorFlipAxis);
+		PRAGMA_ENABLE_DEPRECATION_WARNINGS
+	}
+
+	TEnumAsByte<EAxis::Type> GetSkelMirrorFlipAxis() const
+	{
+		PRAGMA_DISABLE_DEPRECATION_WARNINGS
+		return SkelMirrorFlipAxis;
+		PRAGMA_ENABLE_DEPRECATION_WARNINGS
+	}
+
+	void SetSkelMirrorFlipAxis(TEnumAsByte<EAxis::Type> InSkelMirrorFlipAxis)
+	{
+		PRAGMA_DISABLE_DEPRECATION_WARNINGS
+		SkelMirrorFlipAxis = InSkelMirrorFlipAxis;
+		PRAGMA_ENABLE_DEPRECATION_WARNINGS
+	}
 
 	/** If true, use 32 bit UVs. If false, use 16 bit UVs to save memory */
 	UPROPERTY()
@@ -792,47 +1102,255 @@ public:
 	uint8 bUseHighPrecisionTangentBasis_DEPRECATED : 1;
 
 	/** true if this mesh has ever been simplified with Simplygon. */
+	UE_DEPRECATED(4.27, "Please do not access this member directly; use USkeletalMesh::GetHasBeenSimplified() or USkeletalMesh::SetHasBeenSimplified().")
 	UPROPERTY()
 	uint8 bHasBeenSimplified:1;
+	static FName GetHasBeenSimplifiedMemberName()
+	{
+		PRAGMA_DISABLE_DEPRECATION_WARNINGS
+		return GET_MEMBER_NAME_CHECKED(USkeletalMesh, bHasBeenSimplified);
+		PRAGMA_ENABLE_DEPRECATION_WARNINGS
+	}
+
+	bool GetHasBeenSimplified() const
+	{
+		PRAGMA_DISABLE_DEPRECATION_WARNINGS
+		return bHasBeenSimplified;
+		PRAGMA_ENABLE_DEPRECATION_WARNINGS
+	}
+
+	void SetHasBeenSimplified(bool bInHasBeenSimplified)
+	{
+		PRAGMA_DISABLE_DEPRECATION_WARNINGS
+		bHasBeenSimplified = bInHasBeenSimplified;
+		PRAGMA_ENABLE_DEPRECATION_WARNINGS
+	}
 
 	/** Whether or not the mesh has vertex colors */
+	UE_DEPRECATED(4.27, "Please do not access this member directly; use USkeletalMesh::GetHasVertexColors() or USkeletalMesh::SetHasVertexColors().")
 	UPROPERTY()
 	uint8 bHasVertexColors:1;
 
+	static FName GetbHasVertexColorsMemberName()
+	{
+		PRAGMA_DISABLE_DEPRECATION_WARNINGS
+		return GET_MEMBER_NAME_CHECKED(USkeletalMesh, bHasVertexColors);
+		PRAGMA_ENABLE_DEPRECATION_WARNINGS
+	}
+
+	bool GetHasVertexColors() const
+	{
+		PRAGMA_DISABLE_DEPRECATION_WARNINGS
+		return bHasVertexColors != 0;
+		PRAGMA_ENABLE_DEPRECATION_WARNINGS
+	}
+
+	void SetHasVertexColors(bool InbHasVertexColors)
+	{
+		PRAGMA_DISABLE_DEPRECATION_WARNINGS
+		bHasVertexColors = InbHasVertexColors;
+		PRAGMA_ENABLE_DEPRECATION_WARNINGS
+	}
+
 	//caching optimization to avoid recalculating in non-editor builds
-	uint8 bHasActiveClothingAssets:1;
+	UE_DEPRECATED(4.27, "Please do not access this member directly; use USkeletalMesh::HasActiveClothingAssets() or USkeletalMesh::SetHasActiveClothingAssets().")
+	uint8 bHasActiveClothingAssets : 1;
+
+	static FName GetHasActiveClothingAssetsMemberName()
+	{
+		PRAGMA_DISABLE_DEPRECATION_WARNINGS
+		return GET_MEMBER_NAME_CHECKED(USkeletalMesh, bHasActiveClothingAssets);
+		PRAGMA_ENABLE_DEPRECATION_WARNINGS
+	}
+
+	void SetHasActiveClothingAssets(const bool InbHasActiveClothingAssets)
+	{
+		PRAGMA_DISABLE_DEPRECATION_WARNINGS
+		bHasActiveClothingAssets = InbHasActiveClothingAssets;
+		PRAGMA_ENABLE_DEPRECATION_WARNINGS
+	}
 
 	/** Uses skinned data for collision data. Per poly collision cannot be used for simulation, in most cases you are better off using the physics asset */
+	UE_DEPRECATED(4.27, "Please do not access this member directly; use USkeletalMesh::GetEnablePerPolyCollision() or USkeletalMesh::SetEnablePerPolyCollision().")
 	UPROPERTY(EditAnywhere, Category = Physics)
 	uint8 bEnablePerPolyCollision : 1;
+
+	static FName GetEnablePerPolyCollisionMemberName()
+	{
+		PRAGMA_DISABLE_DEPRECATION_WARNINGS
+		return GET_MEMBER_NAME_CHECKED(USkeletalMesh, bEnablePerPolyCollision);
+		PRAGMA_ENABLE_DEPRECATION_WARNINGS
+	}
+
+	bool GetEnablePerPolyCollision() const
+	{
+		PRAGMA_DISABLE_DEPRECATION_WARNINGS
+		return bEnablePerPolyCollision != 0;
+		PRAGMA_ENABLE_DEPRECATION_WARNINGS
+	}
+
+	void SetEnablePerPolyCollision(bool bInEnablePerPolyCollision)
+	{
+		PRAGMA_DISABLE_DEPRECATION_WARNINGS
+		bEnablePerPolyCollision = bInEnablePerPolyCollision;
+		PRAGMA_ENABLE_DEPRECATION_WARNINGS
+	}
 	
 #if WITH_EDITORONLY_DATA
 	/** The guid to compute the ddc key, it must be dirty when we change the vertex color. */
+	UE_DEPRECATED(4.27, "Please do not access this member directly; use USkeletalMesh::GetVertexColorGuid() or USkeletalMesh::SetVertexColorGuid().")
 	UPROPERTY()
 	FGuid VertexColorGuid;
+
+	static FName GetVertexColorGuidMemberName()
+	{
+		PRAGMA_DISABLE_DEPRECATION_WARNINGS
+		return GET_MEMBER_NAME_CHECKED(USkeletalMesh, VertexColorGuid);
+		PRAGMA_ENABLE_DEPRECATION_WARNINGS
+	}
+	
+	FGuid GetVertexColorGuid() const
+	{
+		PRAGMA_DISABLE_DEPRECATION_WARNINGS
+		return VertexColorGuid;
+		PRAGMA_ENABLE_DEPRECATION_WARNINGS
+	}
+
+	void SetVertexColorGuid(FGuid InVertexColorGuid)
+	{
+		PRAGMA_DISABLE_DEPRECATION_WARNINGS
+		VertexColorGuid = InVertexColorGuid;
+		PRAGMA_ENABLE_DEPRECATION_WARNINGS
+	}
+
 #endif
 
 	// Physics data for the per poly collision case. In 99% of cases you will not need this and are better off using simple ragdoll collision (physics asset)
+	UE_DEPRECATED(4.27, "Please do not access this member directly; use USkeletalMesh::GetBodySetup() or USkeletalMesh::SetBodySetup().")
 	UPROPERTY(transient)
 	class UBodySetup* BodySetup;
+	static FName GetBodySetupMemberName()
+	{
+		PRAGMA_DISABLE_DEPRECATION_WARNINGS
+		return GET_MEMBER_NAME_CHECKED(USkeletalMesh, BodySetup);
+		PRAGMA_ENABLE_DEPRECATION_WARNINGS
+	}
+
+	class UBodySetup* GetBodySetup() const
+	{
+		PRAGMA_DISABLE_DEPRECATION_WARNINGS
+		return BodySetup;
+		PRAGMA_ENABLE_DEPRECATION_WARNINGS
+	}
+
+	UE_DEPRECATED(4.27, "Please do not use this non const function; use the combination of USkeletalMesh::CreateBodySetup() and USkeletalMesh::GetBodySetup() const. Cast the skeletal mesh caller to const to force the compiler to use the USkeletalMesh::GetBodySetup() const function and avoid the deprecation warning")
+	class UBodySetup* GetBodySetup()
+	{
+		PRAGMA_DISABLE_DEPRECATION_WARNINGS
+		CreateBodySetup();
+		return BodySetup;
+		PRAGMA_ENABLE_DEPRECATION_WARNINGS
+	}
+
+	void SetBodySetup(class UBodySetup* InBodySetup)
+	{
+		PRAGMA_DISABLE_DEPRECATION_WARNINGS
+		BodySetup = InBodySetup;
+		PRAGMA_ENABLE_DEPRECATION_WARNINGS
+	}
 
 	/**
 	 *	Physics and collision information used for this USkeletalMesh, set up in Physics Asset Editor.
 	 *	This is used for per-bone hit detection, accurate bounding box calculation and ragdoll physics for example.
 	 */
-	UPROPERTY(EditAnywhere, AssetRegistrySearchable, BlueprintReadOnly, Category=Physics)
+	UE_DEPRECATED(4.27, "Please do not access this member directly; use USkeletalMesh::GetPhysicsAsset() or USkeletalMesh::SetPhysicsAsset().")
+	UPROPERTY(EditAnywhere, AssetRegistrySearchable, BlueprintGetter = GetPhysicsAsset, Category=Physics)
 	class UPhysicsAsset* PhysicsAsset;
+	
+	static FName GetPhysicsAssetMemberName()
+	{
+		PRAGMA_DISABLE_DEPRECATION_WARNINGS
+		return GET_MEMBER_NAME_CHECKED(USkeletalMesh, PhysicsAsset);
+		PRAGMA_ENABLE_DEPRECATION_WARNINGS
+	}
+
+	UFUNCTION(BlueprintGetter)
+	class UPhysicsAsset* GetPhysicsAsset() const
+	{
+		PRAGMA_DISABLE_DEPRECATION_WARNINGS
+		return PhysicsAsset;
+		PRAGMA_ENABLE_DEPRECATION_WARNINGS
+	}
+
+	void SetPhysicsAsset(class UPhysicsAsset* InPhysicsAsset)
+	{
+		PRAGMA_DISABLE_DEPRECATION_WARNINGS
+		PhysicsAsset = InPhysicsAsset;
+		PRAGMA_ENABLE_DEPRECATION_WARNINGS
+	}
 
 	/**
 	 * Physics asset whose shapes will be used for shadowing when components have bCastCharacterCapsuleDirectShadow or bCastCharacterCapsuleIndirectShadow enabled.
 	 * Only spheres and sphyl shapes in the physics asset can be supported.  The more shapes used, the higher the cost of the capsule shadows will be.
 	 */
-	UPROPERTY(EditAnywhere, AssetRegistrySearchable, BlueprintReadOnly, Category=Lighting)
+	UE_DEPRECATED(4.27, "Please do not access this member directly; use USkeletalMesh::GetShadowPhysicsAsset() or USkeletalMesh::SetShadowPhysicsAsset().")
+	UPROPERTY(EditAnywhere, AssetRegistrySearchable, BlueprintGetter = GetShadowPhysicsAsset, Category=Lighting)
 	class UPhysicsAsset* ShadowPhysicsAsset;
 
+	static FName GetShadowPhysicsAssetMemberName()
+	{
+		PRAGMA_DISABLE_DEPRECATION_WARNINGS
+		return GET_MEMBER_NAME_CHECKED(USkeletalMesh, ShadowPhysicsAsset);
+		PRAGMA_ENABLE_DEPRECATION_WARNINGS
+	}
+
+	UFUNCTION(BlueprintGetter)
+	class UPhysicsAsset* GetShadowPhysicsAsset() const
+	{
+		PRAGMA_DISABLE_DEPRECATION_WARNINGS
+		return ShadowPhysicsAsset;
+		PRAGMA_ENABLE_DEPRECATION_WARNINGS
+	}
+
+	void SetShadowPhysicsAsset(class UPhysicsAsset* InShadowPhysicsAsset)
+	{
+		PRAGMA_DISABLE_DEPRECATION_WARNINGS
+		ShadowPhysicsAsset = InShadowPhysicsAsset;
+		PRAGMA_ENABLE_DEPRECATION_WARNINGS
+	}
+
 	/** Mapping data that is saved */
-	UPROPERTY(EditAnywhere, editfixedsize, BlueprintReadOnly, Category=Animation)
+	UE_DEPRECATED(4.27, "Please do not access this member directly; use USkeletalMesh::GetNodeMappingData() or USkeletalMesh::SetNodeMappingData().")
+	UPROPERTY(EditAnywhere, editfixedsize, BlueprintGetter = GetNodeMappingData, Category=Animation)
 	TArray<class UNodeMappingContainer*> NodeMappingData;
+	static FName GetNodeMappingDataMemberName()
+	{
+		PRAGMA_DISABLE_DEPRECATION_WARNINGS
+		return GET_MEMBER_NAME_CHECKED(USkeletalMesh, NodeMappingData);
+		PRAGMA_ENABLE_DEPRECATION_WARNINGS
+	}
+
+	TArray<class UNodeMappingContainer*>& GetNodeMappingData()
+	{
+		PRAGMA_DISABLE_DEPRECATION_WARNINGS
+		return NodeMappingData;
+		PRAGMA_ENABLE_DEPRECATION_WARNINGS
+	}
+
+	UFUNCTION(BlueprintGetter)
+	const TArray<class UNodeMappingContainer*>& GetNodeMappingData() const
+	{
+		PRAGMA_DISABLE_DEPRECATION_WARNINGS
+		return NodeMappingData;
+		PRAGMA_ENABLE_DEPRECATION_WARNINGS
+	}
+
+	void SetNodeMappingData(const TArray<class UNodeMappingContainer*>& InNodeMappingData)
+	{
+		PRAGMA_DISABLE_DEPRECATION_WARNINGS
+		NodeMappingData = InNodeMappingData;
+		PRAGMA_ENABLE_DEPRECATION_WARNINGS
+	}
 
 	UFUNCTION(BlueprintCallable, Category = "Animation")
 	class UNodeMappingContainer* GetNodeMappingContainer(class UBlueprint* SourceAsset) const;
@@ -840,9 +1358,31 @@ public:
 #if WITH_EDITORONLY_DATA
 
 	/** Importing data and options used for this mesh */
+	UE_DEPRECATED(4.27, "Please do not access this member directly; use USkeletalMesh::GetAssetImportData() or USkeletalMesh::SetAssetImportData().")
 	UPROPERTY(EditAnywhere, Instanced, Category=ImportSettings)
 	class UAssetImportData* AssetImportData;
 
+	static FName GetAssetImportDataMemberName()
+	{
+		PRAGMA_DISABLE_DEPRECATION_WARNINGS
+		return GET_MEMBER_NAME_CHECKED(USkeletalMesh, AssetImportData);
+		PRAGMA_ENABLE_DEPRECATION_WARNINGS
+	}
+
+	class UAssetImportData* GetAssetImportData() const
+	{
+		PRAGMA_DISABLE_DEPRECATION_WARNINGS
+		return AssetImportData;
+		PRAGMA_ENABLE_DEPRECATION_WARNINGS
+	}
+
+	void SetAssetImportData(class UAssetImportData* InAssetImportData)
+	{
+		PRAGMA_DISABLE_DEPRECATION_WARNINGS
+		AssetImportData = InAssetImportData;
+		PRAGMA_ENABLE_DEPRECATION_WARNINGS
+	}
+
 	static FText GetSourceFileLabelFromIndex(int32 SourceFileIndex);
 
 	/** Path to the resource used to construct this skeletal mesh */
@@ -854,47 +1394,266 @@ public:
 	FString SourceFileTimestamp_DEPRECATED;
 
 	/** Information for thumbnail rendering */
+	UE_DEPRECATED(4.27, "Please do not access this member directly; use USkeletalMesh::GetThumbnailInfo() or USkeletalMesh::SetThumbnailInfo().")
 	UPROPERTY(VisibleAnywhere, Instanced, AdvancedDisplay, Category = Thumbnail)
 	class UThumbnailInfo* ThumbnailInfo;
+	
+	static FName GetThumbnailInfoMemberName()
+	{
+		PRAGMA_DISABLE_DEPRECATION_WARNINGS
+		return GET_MEMBER_NAME_CHECKED(USkeletalMesh, ThumbnailInfo);
+		PRAGMA_ENABLE_DEPRECATION_WARNINGS
+	}
+
+	class UThumbnailInfo* GetThumbnailInfo() const
+	{
+		PRAGMA_DISABLE_DEPRECATION_WARNINGS
+		return ThumbnailInfo;
+		PRAGMA_ENABLE_DEPRECATION_WARNINGS
+	}
+
+	void SetThumbnailInfo(class UThumbnailInfo* InThumbnailInfo)
+	{
+		PRAGMA_DISABLE_DEPRECATION_WARNINGS
+		ThumbnailInfo = InThumbnailInfo;
+		PRAGMA_ENABLE_DEPRECATION_WARNINGS
+	}
 
 	/** Should we use a custom camera transform when viewing this mesh in the tools */
+	UE_DEPRECATED(4.27, "Please do not access this member directly; use USkeletalMesh::GetHasCustomDefaultEditorCamera() or USkeletalMesh::SetHasCustomDefaultEditorCamera().")
 	UPROPERTY()
 	bool bHasCustomDefaultEditorCamera;
+
+	static FName GetHasCustomDefaultEditorCameraMemberName()
+	{
+		PRAGMA_DISABLE_DEPRECATION_WARNINGS
+		return GET_MEMBER_NAME_CHECKED(USkeletalMesh, bHasCustomDefaultEditorCamera);
+		PRAGMA_ENABLE_DEPRECATION_WARNINGS
+	}
+
+	bool GetHasCustomDefaultEditorCamera() const
+	{
+		PRAGMA_DISABLE_DEPRECATION_WARNINGS
+		return bHasCustomDefaultEditorCamera;
+		PRAGMA_ENABLE_DEPRECATION_WARNINGS
+	}
+
+	void SetHasCustomDefaultEditorCamera(bool bInHasCustomDefaultEditorCamera)
+	{
+		PRAGMA_DISABLE_DEPRECATION_WARNINGS
+		bHasCustomDefaultEditorCamera = bInHasCustomDefaultEditorCamera;
+		PRAGMA_ENABLE_DEPRECATION_WARNINGS
+	}
+
 	/** Default camera location */
+	UE_DEPRECATED(4.27, "Please do not access this member directly; use USkeletalMesh::GetDefaultEditorCameraLocation() or USkeletalMesh::SetDefaultEditorCameraLocation().")
 	UPROPERTY()
 	FVector DefaultEditorCameraLocation;
+
+	static FName GetDefaultEditorCameraLocationMemberName()
+	{
+		PRAGMA_DISABLE_DEPRECATION_WARNINGS
+		return GET_MEMBER_NAME_CHECKED(USkeletalMesh, DefaultEditorCameraLocation);
+		PRAGMA_ENABLE_DEPRECATION_WARNINGS
+	}
+
+	const FVector& GetDefaultEditorCameraLocation() const
+	{
+		PRAGMA_DISABLE_DEPRECATION_WARNINGS
+		return DefaultEditorCameraLocation;
+		PRAGMA_ENABLE_DEPRECATION_WARNINGS
+	}
+
+	void SetDefaultEditorCameraLocation(FVector InDefaultEditorCameraLocation)
+	{
+		PRAGMA_DISABLE_DEPRECATION_WARNINGS
+		DefaultEditorCameraLocation = InDefaultEditorCameraLocation;
+		PRAGMA_ENABLE_DEPRECATION_WARNINGS
+	}
+
 	/** Default camera rotation */
+	UE_DEPRECATED(4.27, "Please do not access this member directly; use USkeletalMesh::GetDefaultEditorCameraRotation() or USkeletalMesh::SetDefaultEditorCameraRotation().")
 	UPROPERTY()
 	FRotator DefaultEditorCameraRotation;
+
+	static FName GetDefaultEditorCameraRotationMemberName()
+	{
+		PRAGMA_DISABLE_DEPRECATION_WARNINGS
+		return GET_MEMBER_NAME_CHECKED(USkeletalMesh, DefaultEditorCameraRotation);
+		PRAGMA_ENABLE_DEPRECATION_WARNINGS
+	}
+
+	const FRotator& GetDefaultEditorCameraRotation() const
+	{
+		PRAGMA_DISABLE_DEPRECATION_WARNINGS
+		return DefaultEditorCameraRotation;
+		PRAGMA_ENABLE_DEPRECATION_WARNINGS
+	}
+
+	void SetDefaultEditorCameraRotation(FRotator InDefaultEditorCameraRotation)
+	{
+		PRAGMA_DISABLE_DEPRECATION_WARNINGS
+		DefaultEditorCameraRotation = InDefaultEditorCameraRotation;
+		PRAGMA_ENABLE_DEPRECATION_WARNINGS
+	}
+
 	/** Default camera look at */
+	UE_DEPRECATED(4.27, "Please do not access this member directly; use USkeletalMesh::GetDefaultEditorCameraLookAt() or USkeletalMesh::SetDefaultEditorCameraLookAt().")
 	UPROPERTY()
 	FVector DefaultEditorCameraLookAt;
+
+	static FName GetDefaultEditorCameraLookAtMemberName()
+	{
+		PRAGMA_DISABLE_DEPRECATION_WARNINGS
+		return GET_MEMBER_NAME_CHECKED(USkeletalMesh, DefaultEditorCameraLookAt);
+		PRAGMA_ENABLE_DEPRECATION_WARNINGS
+	}
+
+	const FVector& GetDefaultEditorCameraLookAt() const
+	{
+		PRAGMA_DISABLE_DEPRECATION_WARNINGS
+		return DefaultEditorCameraLookAt;
+		PRAGMA_ENABLE_DEPRECATION_WARNINGS
+	}
+
+	void SetDefaultEditorCameraLookAt(FVector InDefaultEditorCameraLookAt)
+	{
+		PRAGMA_DISABLE_DEPRECATION_WARNINGS
+		DefaultEditorCameraLookAt = InDefaultEditorCameraLookAt;
+		PRAGMA_ENABLE_DEPRECATION_WARNINGS
+	}
+
 	/** Default camera ortho zoom */
+	UE_DEPRECATED(4.27, "Please do not access this member directly; use USkeletalMesh::GetDefaultEditorCameraOrthoZoom() or USkeletalMesh::SetDefaultEditorCameraOrthoZoom().")
 	UPROPERTY()
 	float DefaultEditorCameraOrthoZoom;
 
+	static FName GetDefaultEditorCameraOrthoZoomMemberName()
+	{
+		PRAGMA_DISABLE_DEPRECATION_WARNINGS
+		return GET_MEMBER_NAME_CHECKED(USkeletalMesh, DefaultEditorCameraOrthoZoom);
+		PRAGMA_ENABLE_DEPRECATION_WARNINGS
+	}
+
+	float GetDefaultEditorCameraOrthoZoom() const
+	{
+		PRAGMA_DISABLE_DEPRECATION_WARNINGS
+		return DefaultEditorCameraOrthoZoom;
+		PRAGMA_ENABLE_DEPRECATION_WARNINGS
+	}
+
+	void SetDefaultEditorCameraOrthoZoom(float InDefaultEditorCameraOrthoZoom)
+	{
+		PRAGMA_DISABLE_DEPRECATION_WARNINGS
+		DefaultEditorCameraOrthoZoom = InDefaultEditorCameraOrthoZoom;
+		PRAGMA_ENABLE_DEPRECATION_WARNINGS
+	}
+
  
 	/* Attached assets component for this mesh */
+	UE_DEPRECATED(4.27, "Please do not access this member directly; use USkeletalMesh::GetPreviewAttachedAssetContainer() or USkeletalMesh::SetPreviewAttachedAssetContainer().")
 	UPROPERTY()
 	FPreviewAssetAttachContainer PreviewAttachedAssetContainer;
 
+	FPreviewAssetAttachContainer& GetPreviewAttachedAssetContainer()
+	{
+		PRAGMA_DISABLE_DEPRECATION_WARNINGS
+		return PreviewAttachedAssetContainer;
+		PRAGMA_ENABLE_DEPRECATION_WARNINGS
+	}
+
+	const FPreviewAssetAttachContainer& GetPreviewAttachedAssetContainer() const
+	{
+		PRAGMA_DISABLE_DEPRECATION_WARNINGS
+		return PreviewAttachedAssetContainer;
+		PRAGMA_ENABLE_DEPRECATION_WARNINGS
+	}
+
+	void SetPreviewAttachedAssetContainer(const FPreviewAssetAttachContainer& InPreviewAttachedAssetContainer)
+	{
+		PRAGMA_DISABLE_DEPRECATION_WARNINGS
+		PreviewAttachedAssetContainer = InPreviewAttachedAssetContainer;
+		PRAGMA_ENABLE_DEPRECATION_WARNINGS
+	}
+
 	/**
 	 * If true on post load we need to calculate resolution independent Display Factors from the
 	 * loaded LOD screen sizes.
 	 */
+	UE_DEPRECATED(4.27, "Please do not access this member directly; use USkeletalMesh::GetRequiresLODScreenSizeConversion() or USkeletalMesh::SetRequiresLODScreenSizeConversion().")
 	uint32 bRequiresLODScreenSizeConversion : 1;
 
+	bool GetRequiresLODScreenSizeConversion() const
+	{
+		PRAGMA_DISABLE_DEPRECATION_WARNINGS
+		return bRequiresLODScreenSizeConversion;
+		PRAGMA_ENABLE_DEPRECATION_WARNINGS
+	}
+
+	void SetRequiresLODScreenSizeConversion(bool bInRequiresLODScreenSizeConversion)
+	{
+		PRAGMA_DISABLE_DEPRECATION_WARNINGS
+		bRequiresLODScreenSizeConversion = bInRequiresLODScreenSizeConversion;
+		PRAGMA_ENABLE_DEPRECATION_WARNINGS
+	}
+
 	/**
 	 * If true on post load we need to calculate resolution independent LOD hysteresis from the
 	 * loaded LOD hysteresis.
 	 */
+	UE_DEPRECATED(4.27, "Please do not access this member directly; use USkeletalMesh::GetRequiresLODHysteresisConversion() or USkeletalMesh::SetRequiresLODHysteresisConversion().")
 	uint32 bRequiresLODHysteresisConversion : 1;
 
+	bool GetRequiresLODHysteresisConversion() const
+	{
+		PRAGMA_DISABLE_DEPRECATION_WARNINGS
+		return bRequiresLODHysteresisConversion;
+		PRAGMA_ENABLE_DEPRECATION_WARNINGS
+	}
+
+	void SetRequiresLODHysteresisConversion(bool bInRequiresLODHysteresisConversion)
+	{
+		PRAGMA_DISABLE_DEPRECATION_WARNINGS
+		bRequiresLODHysteresisConversion = bInRequiresLODHysteresisConversion;
+		PRAGMA_ENABLE_DEPRECATION_WARNINGS
+	}
+
 #endif // WITH_EDITORONLY_DATA
 
-	UPROPERTY(Category=Mesh, BlueprintReadWrite)
+	UE_DEPRECATED(4.27, "Please do not access this member directly; use USkeletalMesh::GetMorphTargets() or USkeletalMesh::SetMorphTargets().")
+	UPROPERTY(BlueprintGetter = GetMorphTargets, BlueprintSetter = SetMorphTargets, Category = Mesh)
 	TArray<UMorphTarget*> MorphTargets;
 
+	static FName GetMorphTargetsMemberName()
+	{
+		PRAGMA_DISABLE_DEPRECATION_WARNINGS
+		return GET_MEMBER_NAME_CHECKED(USkeletalMesh, MorphTargets);
+		PRAGMA_ENABLE_DEPRECATION_WARNINGS
+	}
+
+	TArray<UMorphTarget*>& GetMorphTargets()
+	{
+		PRAGMA_DISABLE_DEPRECATION_WARNINGS
+		return MorphTargets;
+		PRAGMA_ENABLE_DEPRECATION_WARNINGS
+	}
+
+	UFUNCTION(BlueprintGetter)
+	const TArray<UMorphTarget*>& GetMorphTargets() const
+	{
+		PRAGMA_DISABLE_DEPRECATION_WARNINGS
+		return MorphTargets;
+		PRAGMA_ENABLE_DEPRECATION_WARNINGS
+	}
+
+	UFUNCTION(BlueprintSetter)
+	void SetMorphTargets(const TArray<UMorphTarget*>& InMorphTargets)
+	{
+		PRAGMA_DISABLE_DEPRECATION_WARNINGS
+		MorphTargets = InMorphTargets;
+		PRAGMA_ENABLE_DEPRECATION_WARNINGS
+	}
+
 	/**
 	 *	Returns the list of all morph targets of this skeletal mesh
 	 *  @return	The list of morph targets
@@ -906,24 +1665,162 @@ public:
 	FRenderCommandFence ReleaseResourcesFence;
 
 	/** New Reference skeleton type **/
+	UE_DEPRECATED(4.27, "Please do not access this member directly; use USkeletalMesh::GetRefSkeleton() or USkeletalMesh::SetRefSkeleton().")
 	FReferenceSkeleton RefSkeleton;
 
+	static FName GetRefSkeletonMemberName()
+	{
+		PRAGMA_DISABLE_DEPRECATION_WARNINGS
+		return GET_MEMBER_NAME_CHECKED(USkeletalMesh, RefSkeleton);
+		PRAGMA_ENABLE_DEPRECATION_WARNINGS
+	}
+
+	FReferenceSkeleton& GetRefSkeleton()
+	{
+		PRAGMA_DISABLE_DEPRECATION_WARNINGS
+		return RefSkeleton;
+		PRAGMA_ENABLE_DEPRECATION_WARNINGS
+	}
+
+	const FReferenceSkeleton& GetRefSkeleton() const
+	{
+		PRAGMA_DISABLE_DEPRECATION_WARNINGS
+		return RefSkeleton;
+		PRAGMA_ENABLE_DEPRECATION_WARNINGS
+	}
+
+	void SetRefSkeleton(const FReferenceSkeleton& InRefSkeleton)
+	{
+		PRAGMA_DISABLE_DEPRECATION_WARNINGS
+		RefSkeleton = InRefSkeleton;
+		PRAGMA_ENABLE_DEPRECATION_WARNINGS
+	}
+
 	/** Map of morph target name to index into USkeletalMesh::MorphTargets**/
+	UE_DEPRECATED(4.27, "Please do not access this member directly; use USkeletalMesh::GetMorphTargetIndexMap() or USkeletalMesh::SetMorphTargetIndexMap().")
 	TMap<FName, int32> MorphTargetIndexMap;
 
+	static FName GetMorphTargetIndexMapMemberName()
+	{
+		PRAGMA_DISABLE_DEPRECATION_WARNINGS
+		return GET_MEMBER_NAME_CHECKED(USkeletalMesh, MorphTargetIndexMap);
+		PRAGMA_ENABLE_DEPRECATION_WARNINGS
+	}
+
+	TMap<FName, int32>& GetMorphTargetIndexMap()
+	{
+		PRAGMA_DISABLE_DEPRECATION_WARNINGS
+		return MorphTargetIndexMap;
+		PRAGMA_ENABLE_DEPRECATION_WARNINGS
+	}
+
+	const TMap<FName, int32>& GetMorphTargetIndexMap() const
+	{
+		PRAGMA_DISABLE_DEPRECATION_WARNINGS
+		return MorphTargetIndexMap;
+		PRAGMA_ENABLE_DEPRECATION_WARNINGS
+	}
+
+	void SetMorphTargetIndexMap(const TMap<FName, int32>& InMorphTargetIndexMap)
+	{
+		PRAGMA_DISABLE_DEPRECATION_WARNINGS
+		MorphTargetIndexMap = InMorphTargetIndexMap;
+		PRAGMA_ENABLE_DEPRECATION_WARNINGS
+	}
+
 	/** Reference skeleton precomputed bases. */
-	TArray<FMatrix> RefBasesInvMatrix;    
+	UE_DEPRECATED(4.27, "Please do not access this member directly; use USkeletalMesh::GetRefBasesInvMatrix() or USkeletalMesh::SetRefBasesInvMatrix().")
+	TArray<FMatrix> RefBasesInvMatrix;
+
+	static FName GetRefBasesInvMatrixMemberName()
+	{
+		PRAGMA_DISABLE_DEPRECATION_WARNINGS
+		return GET_MEMBER_NAME_CHECKED(USkeletalMesh, RefBasesInvMatrix);
+		PRAGMA_ENABLE_DEPRECATION_WARNINGS
+	}
+
+	TArray<FMatrix>& GetRefBasesInvMatrix()
+	{
+		PRAGMA_DISABLE_DEPRECATION_WARNINGS
+			return RefBasesInvMatrix;
+		PRAGMA_ENABLE_DEPRECATION_WARNINGS
+	}
+
+	const TArray<FMatrix>& GetRefBasesInvMatrix() const
+	{
+		PRAGMA_DISABLE_DEPRECATION_WARNINGS
+		return RefBasesInvMatrix;
+		PRAGMA_ENABLE_DEPRECATION_WARNINGS
+	}
+
+	void SetRefBasesInvMatrix(const TArray<FMatrix>& InRefBasesInvMatrix)
+	{
+		PRAGMA_DISABLE_DEPRECATION_WARNINGS
+		RefBasesInvMatrix = InRefBasesInvMatrix;
+		PRAGMA_ENABLE_DEPRECATION_WARNINGS
+	}
 
 #if WITH_EDITORONLY_DATA
 
 	/** Height offset for the floor mesh in the editor */
+	UE_DEPRECATED(4.27, "Please do not access this member directly; use USkeletalMesh::GetFloorOffset() or USkeletalMesh::SetFloorOffset().")
 	UPROPERTY()
 	float FloorOffset;
 
+	static FName GetFloorOffsetMemberName()
+	{
+		PRAGMA_DISABLE_DEPRECATION_WARNINGS
+		return GET_MEMBER_NAME_CHECKED(USkeletalMesh, FloorOffset);
+		PRAGMA_ENABLE_DEPRECATION_WARNINGS
+	}
+
+	float GetFloorOffset() const
+	{
+		PRAGMA_DISABLE_DEPRECATION_WARNINGS
+		return FloorOffset;
+		PRAGMA_ENABLE_DEPRECATION_WARNINGS
+	}
+
+	void SetFloorOffset(float InFloorOffset)
+	{
+		PRAGMA_DISABLE_DEPRECATION_WARNINGS
+		FloorOffset = InFloorOffset;
+		PRAGMA_ENABLE_DEPRECATION_WARNINGS
+	}
+
 	/** This is buffer that saves pose that is used by retargeting*/
+	UE_DEPRECATED(4.27, "Please do not access this member directly; use USkeletalMesh::GetRetargetBasePose() or USkeletalMesh::SetRetargetBasePose().")
 	UPROPERTY()
 	TArray<FTransform> RetargetBasePose;
 
+	static FName GetRetargetBasePoseMemberName()
+	{
+		PRAGMA_DISABLE_DEPRECATION_WARNINGS
+		return GET_MEMBER_NAME_CHECKED(USkeletalMesh, RetargetBasePose);
+		PRAGMA_ENABLE_DEPRECATION_WARNINGS
+	}
+
+	TArray<FTransform>& GetRetargetBasePose()
+	{
+		PRAGMA_DISABLE_DEPRECATION_WARNINGS
+		return RetargetBasePose;
+		PRAGMA_ENABLE_DEPRECATION_WARNINGS
+	}
+
+	const TArray<FTransform>& GetRetargetBasePose() const
+	{
+		PRAGMA_DISABLE_DEPRECATION_WARNINGS
+		return RetargetBasePose;
+		PRAGMA_ENABLE_DEPRECATION_WARNINGS
+	}
+
+	void SetRetargetBasePose(const TArray<FTransform>& InRetargetBasePose)
+	{
+		PRAGMA_DISABLE_DEPRECATION_WARNINGS
+		RetargetBasePose = InRetargetBasePose;
+		PRAGMA_ENABLE_DEPRECATION_WARNINGS
+	}
+
 	/** Legacy clothing asset data, will be converted to new assets after loading */
 	UPROPERTY()
 	TArray<FClothingAssetData_Legacy>		ClothingAssets_DEPRECATED;
@@ -933,9 +1830,31 @@ public:
 	 *  This blueprint will be ran before physics, but after the main
 	 *  anim instance for any skeletal mesh component using this mesh.
 	 */
+	UE_DEPRECATED(4.27, "Please do not access this member directly; use USkeletalMesh::GetPostProcessAnimBlueprint() or USkeletalMesh::SetPostProcessAnimBlueprint().")
 	UPROPERTY(EditAnywhere, BlueprintReadOnly, Category = SkeletalMesh)
 	TSubclassOf<UAnimInstance> PostProcessAnimBlueprint;
 
+	static FName GetPostProcessAnimBlueprintMemberName()
+	{
+		PRAGMA_DISABLE_DEPRECATION_WARNINGS
+		return GET_MEMBER_NAME_CHECKED(USkeletalMesh, PostProcessAnimBlueprint);
+		PRAGMA_ENABLE_DEPRECATION_WARNINGS
+	}
+
+	TSubclassOf<UAnimInstance> GetPostProcessAnimBlueprint() const
+	{
+		PRAGMA_DISABLE_DEPRECATION_WARNINGS
+		return PostProcessAnimBlueprint;
+		PRAGMA_ENABLE_DEPRECATION_WARNINGS
+	}
+
+	void SetPostProcessAnimBlueprint(TSubclassOf<UAnimInstance> InPostProcessAnimBlueprint)
+	{
+		PRAGMA_DISABLE_DEPRECATION_WARNINGS
+		PostProcessAnimBlueprint = InPostProcessAnimBlueprint;
+		PRAGMA_ENABLE_DEPRECATION_WARNINGS
+	}
+
 #if WITH_EDITOR && WITH_APEX_CLOTHING
 	/** 
 	 * Take clothing assets that were imported using APEX files before we moved away from the APEX simulation
@@ -966,14 +1885,47 @@ public:
 	UClothingAssetBase* GetSectionClothingAsset(int32 InLodIndex, int32 InSectionIndex);
 	const UClothingAssetBase* GetSectionClothingAsset(int32 InLodIndex, int32 InSectionIndex) const;
 
-	/** 
+	/**
 	 * Clothing assets imported to this mesh. May or may not be in use currently on the mesh.
 	 * Ordering not guaranteed, use the provided getters to access elements in this array
 	 * whenever possible
 	 */
-	UPROPERTY(EditAnywhere, editfixedsize, BlueprintReadOnly, Category = Clothing)
+	UE_DEPRECATED(4.27, "Please do not access this member directly; use USkeletalMesh::GetMeshClothingAssets() or USkeletalMesh::SetMeshClothingAssets().")
+	UPROPERTY(EditAnywhere, editfixedsize, BlueprintGetter = GetMeshClothingAssets, BlueprintSetter = SetMeshClothingAssets, Category = Clothing)
 	TArray<UClothingAssetBase*> MeshClothingAssets;
 
+	static FName GetMeshClothingAssetsMemberName()
+	{
+		PRAGMA_DISABLE_DEPRECATION_WARNINGS
+		return GET_MEMBER_NAME_CHECKED(USkeletalMesh, MeshClothingAssets);
+		PRAGMA_ENABLE_DEPRECATION_WARNINGS
+	}
+
+	TArray<UClothingAssetBase*>& GetMeshClothingAssets()
+	{
+		PRAGMA_DISABLE_DEPRECATION_WARNINGS
+		return MeshClothingAssets;
+		PRAGMA_ENABLE_DEPRECATION_WARNINGS
+	}
+
+	UFUNCTION(BlueprintGetter)
+	const TArray<UClothingAssetBase*>& GetMeshClothingAssets() const
+	{
+		PRAGMA_DISABLE_DEPRECATION_WARNINGS
+		return MeshClothingAssets;
+		PRAGMA_ENABLE_DEPRECATION_WARNINGS
+	}
+
+	UFUNCTION(BlueprintSetter)
+	void SetMeshClothingAssets(const TArray<UClothingAssetBase*>& InMeshClothingAssets)
+	{
+		PRAGMA_DISABLE_DEPRECATION_WARNINGS
+		MeshClothingAssets = InMeshClothingAssets;
+		PRAGMA_ENABLE_DEPRECATION_WARNINGS
+	}
+
+
+
 	/** Get a clothing asset from its associated GUID (returns nullptr if no match is found) */
 	UClothingAssetBase* GetClothingAsset(const FGuid& InAssetGuid) const;
 
@@ -1325,7 +2277,6 @@ public:
 	bool IsSectionUsingCloth(int32 InSectionIndex, bool bCheckCorrespondingSections = true) const;
 
 	void CreateBodySetup();
-	UBodySetup* GetBodySetup();
 
 #if WITH_EDITOR
 	/** Trigger a physics build to ensure per poly collision is created */
@@ -1523,7 +2474,30 @@ public:
 
 #if WITH_EDITORONLY_DATA
 	/*Transient data use when we postload an old asset to use legacy ddc key, it is turn off so if the user change the asset it go back to the latest ddc code*/
+	UE_DEPRECATED(4.27, "Please do not access this member directly; use USkeletalMesh::GetUseLegacyMeshDerivedDataKey() or USkeletalMesh::SetUseLegacyMeshDerivedDataKey().")
 	bool UseLegacyMeshDerivedDataKey = false;
+
+	static FName GetUseLegacyMeshDerivedDataKeyMemberName()
+	{
+		PRAGMA_DISABLE_DEPRECATION_WARNINGS
+		return GET_MEMBER_NAME_CHECKED(USkeletalMesh, UseLegacyMeshDerivedDataKey);
+		PRAGMA_ENABLE_DEPRECATION_WARNINGS
+	}
+
+	bool GetUseLegacyMeshDerivedDataKey() const
+	{
+		PRAGMA_DISABLE_DEPRECATION_WARNINGS
+		return UseLegacyMeshDerivedDataKey;
+		PRAGMA_ENABLE_DEPRECATION_WARNINGS
+	}
+
+	void SetUseLegacyMeshDerivedDataKey(const bool InUseLegacyMeshDerivedDataKey)
+	{
+		PRAGMA_DISABLE_DEPRECATION_WARNINGS
+		UseLegacyMeshDerivedDataKey = InUseLegacyMeshDerivedDataKey;
+		PRAGMA_ENABLE_DEPRECATION_WARNINGS
+	}
+
 #endif
 
 protected:
diff --git a/Engine/Source/Runtime/Engine/Private/AnimBlueprint.cpp b/Engine/Source/Runtime/Engine/Private/AnimBlueprint.cpp
--- a/Engine/Source/Runtime/Engine/Private/AnimBlueprint.cpp
+++ b/Engine/Source/Runtime/Engine/Private/AnimBlueprint.cpp
@@ -209,7 +209,7 @@ USkeletalMesh* UAnimBlueprint::GetPreviewMesh(bool bFindIfNotSet/*=false*/)
 #if WITH_EDITORONLY_DATA
 	USkeletalMesh* PreviewMesh = PreviewSkeletalMesh.LoadSynchronous();
 	// if somehow skeleton changes, just nullify it. 
-	if (PreviewMesh && PreviewMesh->Skeleton != TargetSkeleton)
+	if (PreviewMesh && PreviewMesh->GetSkeleton() != TargetSkeleton)
 	{
 		PreviewMesh = nullptr;
 		SetPreviewMesh(nullptr);
diff --git a/Engine/Source/Runtime/Engine/Private/Animation/AnimInstance.cpp b/Engine/Source/Runtime/Engine/Private/Animation/AnimInstance.cpp
--- a/Engine/Source/Runtime/Engine/Private/Animation/AnimInstance.cpp
+++ b/Engine/Source/Runtime/Engine/Private/Animation/AnimInstance.cpp
@@ -187,7 +187,7 @@ void UAnimInstance::InitializeAnimation(bool bInDeferRootNodeInitialization)
 	USkeletalMeshComponent* OwnerComponent = GetSkelMeshComponent();
 	if (OwnerComponent->SkeletalMesh != NULL)
 	{
-		CurrentSkeleton = OwnerComponent->SkeletalMesh->Skeleton;
+		CurrentSkeleton = OwnerComponent->SkeletalMesh->GetSkeleton();
 	}
 	else
 	{
@@ -1188,7 +1188,7 @@ void UAnimInstance::RecalcRequiredBones()
 	USkeletalMeshComponent* SkelMeshComp = GetSkelMeshComponent();
 	check( SkelMeshComp )
 
-	if( SkelMeshComp->SkeletalMesh && SkelMeshComp->SkeletalMesh->Skeleton )
+	if( SkelMeshComp->SkeletalMesh && SkelMeshComp->SkeletalMesh->GetSkeleton() )
 	{
 		GetProxyOnGameThread<FAnimInstanceProxy>().RecalcRequiredBones(SkelMeshComp, SkelMeshComp->SkeletalMesh);
 	}
@@ -1562,9 +1562,9 @@ void UAnimInstance::GetActiveCurveNames(EAnimCurveType CurveType, TArray<FName>&
 void UAnimInstance::GetAllCurveNames(TArray<FName>& OutNames) const
 {
 	USkeletalMeshComponent* SkelMeshComp = GetOwningComponent();
-	if (SkelMeshComp && SkelMeshComp->SkeletalMesh && SkelMeshComp->SkeletalMesh->Skeleton)
+	if (SkelMeshComp && SkelMeshComp->SkeletalMesh && SkelMeshComp->SkeletalMesh->GetSkeleton())
 	{
-		const USkeleton* CurSkeleton = SkelMeshComp->SkeletalMesh->Skeleton;
+		const USkeleton* CurSkeleton = SkelMeshComp->SkeletalMesh->GetSkeleton();
 
 		const FSmartNameMapping* Mapping = CurSkeleton->GetSmartNameContainer(USkeleton::AnimCurveMappingName);
 		if (Mapping)
diff --git a/Engine/Source/Runtime/Engine/Private/Animation/AnimInstanceProxy.cpp b/Engine/Source/Runtime/Engine/Private/Animation/AnimInstanceProxy.cpp
--- a/Engine/Source/Runtime/Engine/Private/Animation/AnimInstanceProxy.cpp
+++ b/Engine/Source/Runtime/Engine/Private/Animation/AnimInstanceProxy.cpp
@@ -561,7 +561,7 @@ void FAnimInstanceProxy::InitializeObjects(UAnimInstance* InAnimInstance)
 
 	if (SkeletalMeshComponent->SkeletalMesh != nullptr)
 	{
-		Skeleton = SkeletalMeshComponent->SkeletalMesh->Skeleton;
+		Skeleton = SkeletalMeshComponent->SkeletalMesh->GetSkeleton();
 	}
 	else
 	{
diff --git a/Engine/Source/Runtime/Engine/Private/Animation/AnimSequence.cpp b/Engine/Source/Runtime/Engine/Private/Animation/AnimSequence.cpp
--- a/Engine/Source/Runtime/Engine/Private/Animation/AnimSequence.cpp
+++ b/Engine/Source/Runtime/Engine/Private/Animation/AnimSequence.cpp
@@ -4706,7 +4706,7 @@ bool UAnimSequence::CreateAnimation(USkeletalMesh* Mesh)
 	{
 		ResetAnimation();
 
-		const FReferenceSkeleton& RefSkeleton = Mesh->RefSkeleton;
+		const FReferenceSkeleton& RefSkeleton = Mesh->GetRefSkeleton();
 		SetSequenceLength(MINIMUM_ANIMATION_LENGTH);
 		NumberOfKeys = 1;
 
@@ -4748,7 +4748,7 @@ bool UAnimSequence::CreateAnimation(USkeletalMeshComponent* MeshComponent)
 
 		ResetAnimation();
 
-		const FReferenceSkeleton& RefSkeleton = Mesh->RefSkeleton;
+		const FReferenceSkeleton& RefSkeleton = Mesh->GetRefSkeleton();
 		SetSequenceLength(MINIMUM_ANIMATION_LENGTH);
 		NumberOfKeys = 1;
 
diff --git a/Engine/Source/Runtime/Engine/Private/Animation/AnimSet.cpp b/Engine/Source/Runtime/Engine/Private/Animation/AnimSet.cpp
--- a/Engine/Source/Runtime/Engine/Private/Animation/AnimSet.cpp
+++ b/Engine/Source/Runtime/Engine/Private/Animation/AnimSet.cpp
@@ -58,7 +58,7 @@ float UAnimSet::GetSkeletalMeshMatchRatio(USkeletalMesh* SkelMesh) const
 	int32 TracksMatched = 0;
 	for(int32 i=0; i<TrackBoneNames.Num() ; i++)
 	{
-		const int32 BoneIndex = SkelMesh->RefSkeleton.FindBoneIndex( TrackBoneNames[i] );
+		const int32 BoneIndex = SkelMesh->GetRefSkeleton().FindBoneIndex( TrackBoneNames[i] );
 		if( BoneIndex != INDEX_NONE )
 		{
 			++TracksMatched;
@@ -216,7 +216,7 @@ void UAnimSet::ClearAllAnimSetLinkupCaches()
 
 void FAnimSetMeshLinkup::BuildLinkup(USkeletalMesh* InSkelMesh, UAnimSet* InAnimSet)
 {
-	int32 const NumBones = InSkelMesh->RefSkeleton.GetNum();
+	int32 const NumBones = InSkelMesh->GetRefSkeleton().GetNum();
 
 	// Bone to Track mapping.
 	BoneToTrackTable.Empty(NumBones);
@@ -225,7 +225,7 @@ void FAnimSetMeshLinkup::BuildLinkup(USkeletalMesh* InSkelMesh, UAnimSet* InAnim
 	// For each bone in skeletal mesh, find which track to pull from in the AnimSet.
 	for (int32 i=0; i<NumBones; i++)
 	{
-		FName const BoneName = InSkelMesh->RefSkeleton.GetBoneName(i);
+		FName const BoneName = InSkelMesh->GetRefSkeleton().GetBoneName(i);
 
 		// FindTrackWithName will return INDEX_NONE if no track exists.
 		BoneToTrackTable[i] = InAnimSet->FindTrackWithName(BoneName);
@@ -274,7 +274,7 @@ void FAnimSetMeshLinkup::BuildLinkup(USkeletalMesh* InSkelMesh, UAnimSet* InAnim
 				const int32 DupeTrackIndex = AnimLinkup.BoneToTrackTable[DupeBoneIndex];
 				if( DupeTrackIndex == TrackIndex )
 				{
-					UE_LOG(LogAnimation, Warning, TEXT(" BoneIndex: %i, BoneName: %s, TrackIndex: %i, TrackBoneName: %s"), DupeBoneIndex, *InSkelMesh->RefSkeleton.GetBoneName(DupeBoneIndex).ToString(), DupeTrackIndex, *InAnimSet->TrackBoneNames[DupeTrackIndex].ToString());	
+					UE_LOG(LogAnimation, Warning, TEXT(" BoneIndex: %i, BoneName: %s, TrackIndex: %i, TrackBoneName: %s"), DupeBoneIndex, *InSkelMesh->GetRefSkeleton().GetBoneName(DupeBoneIndex).ToString(), DupeTrackIndex, *InAnimSet->TrackBoneNames[DupeTrackIndex].ToString());
 				}
 			}
 		}
diff --git a/Engine/Source/Runtime/Engine/Private/Animation/AnimSingleNodeInstance.cpp b/Engine/Source/Runtime/Engine/Private/Animation/AnimSingleNodeInstance.cpp
--- a/Engine/Source/Runtime/Engine/Private/Animation/AnimSingleNodeInstance.cpp
+++ b/Engine/Source/Runtime/Engine/Private/Animation/AnimSingleNodeInstance.cpp
@@ -48,7 +48,7 @@ void UAnimSingleNodeInstance::SetAnimationAsset(class UAnimationAsset* NewAsset,
 		else if (CurrentAsset != nullptr)
 		{
 			// if we have an asset, make sure their skeleton matches, otherwise, null it
-			if (MeshComponent->SkeletalMesh->Skeleton != CurrentAsset->GetSkeleton())
+			if (MeshComponent->SkeletalMesh->GetSkeleton() != CurrentAsset->GetSkeleton())
 			{
 				// clear asset since we do not have matching skeleton
 				CurrentAsset = nullptr;
diff --git a/Engine/Source/Runtime/Engine/Private/Animation/AnimTrace.cpp b/Engine/Source/Runtime/Engine/Private/Animation/AnimTrace.cpp
--- a/Engine/Source/Runtime/Engine/Private/Animation/AnimTrace.cpp
+++ b/Engine/Source/Runtime/Engine/Private/Animation/AnimTrace.cpp
@@ -422,14 +422,14 @@ void FAnimTrace::OutputSkeletalMesh(const USkeletalMesh* InMesh)
 
 	TRACE_OBJECT(InMesh);
 
-	uint32 BoneCount = (uint32)InMesh->RefSkeleton.GetNum();
+	uint32 BoneCount = (uint32)InMesh->GetRefSkeleton().GetNum();
 
 	TArray<int32>& ParentIndices = FAnimTraceScratchBuffers::Get().ParentIndices;
 	ParentIndices.Reset();
 	ParentIndices.SetNumUninitialized(BoneCount);
 
 	int32 BoneIndex = 0;
-	for(const FMeshBoneInfo& BoneInfo : InMesh->RefSkeleton.GetRefBoneInfo())
+	for(const FMeshBoneInfo& BoneInfo : InMesh->GetRefSkeleton().GetRefBoneInfo())
 	{
 		ParentIndices[BoneIndex++] = BoneInfo.ParentIndex;
 	}
diff --git a/Engine/Source/Runtime/Engine/Private/Animation/AnimationRuntime.cpp b/Engine/Source/Runtime/Engine/Private/Animation/AnimationRuntime.cpp
--- a/Engine/Source/Runtime/Engine/Private/Animation/AnimationRuntime.cpp
+++ b/Engine/Source/Runtime/Engine/Private/Animation/AnimationRuntime.cpp
@@ -1269,9 +1269,9 @@ void FAnimationRuntime::FillWithRetargetBaseRefPose(FCompactPose& OutPose, const
 		for (FCompactPoseBoneIndex BoneIndex : OutPose.ForEachBoneIndex())
 		{
 			int32 PoseIndex = OutPose.GetBoneContainer().MakeMeshPoseIndex(BoneIndex).GetInt();
-			if (Mesh->RetargetBasePose.IsValidIndex(PoseIndex))
+			if (Mesh->GetRetargetBasePose().IsValidIndex(PoseIndex))
 			{
-				OutPose[BoneIndex] = Mesh->RetargetBasePose[PoseIndex];
+				OutPose[BoneIndex] = Mesh->GetRetargetBasePose()[PoseIndex];
 			}
 		}
 	}
@@ -2089,7 +2089,7 @@ void FAnimationRuntime::MakeSkeletonRefPoseFromMesh(const USkeletalMesh* InMesh,
 {
 	check(InMesh && InSkeleton);
 
-	const TArray<FTransform>& MeshRefPose = InMesh->RefSkeleton.GetRefBonePose();
+	const TArray<FTransform>& MeshRefPose = InMesh->GetRefSkeleton().GetRefBonePose();
 	const TArray<FTransform>& SkeletonRefPose = InSkeleton->GetReferenceSkeleton().GetRefBonePose();
 	const TArray<FMeshBoneInfo> & SkeletonBoneInfo = InSkeleton->GetReferenceSkeleton().GetRefBoneInfo();
 
@@ -2099,7 +2099,7 @@ void FAnimationRuntime::MakeSkeletonRefPoseFromMesh(const USkeletalMesh* InMesh,
 	for (int32 SkeletonBoneIndex = 0; SkeletonBoneIndex < SkeletonRefPose.Num(); ++SkeletonBoneIndex)
 	{
 		FName SkeletonBoneName = SkeletonBoneInfo[SkeletonBoneIndex].Name;
-		int32 MeshBoneIndex = InMesh->RefSkeleton.FindBoneIndex(SkeletonBoneName);
+		int32 MeshBoneIndex = InMesh->GetRefSkeleton().FindBoneIndex(SkeletonBoneName);
 		if (MeshBoneIndex != INDEX_NONE)
 		{
 			OutBoneBuffer[SkeletonBoneIndex] = MeshRefPose[MeshBoneIndex];
@@ -2124,8 +2124,8 @@ void FAnimationRuntime::FillUpComponentSpaceTransformsRetargetBasePose(const USk
 {
 	if (Mesh)
 	{
-		const TArray<FTransform>& ReferencePose = Mesh->RetargetBasePose;
-		const FReferenceSkeleton& RefSkeleton = Mesh->RefSkeleton;
+		const TArray<FTransform>& ReferencePose = Mesh->GetRetargetBasePose();
+		const FReferenceSkeleton& RefSkeleton = Mesh->GetRefSkeleton();
 		FillUpComponentSpaceTransforms(RefSkeleton, ReferencePose, ComponentSpaceTransforms);
 	}
 }
@@ -2178,7 +2178,7 @@ void FAnimationRuntime::AppendActiveMorphTargets(const USkeletalMesh* InSkeletal
 
 	if(MorphCurveAnims.Num() > 0)
 	{
-		const int32 NumMorphTargets = InSkeletalMesh->MorphTargets.Num();
+		const int32 NumMorphTargets = InSkeletalMesh->GetMorphTargets().Num();
 		InOutMorphTargetWeights.SetNumZeroed(NumMorphTargets);
 
 		if(NumMorphTargets > 0)
diff --git a/Engine/Source/Runtime/Engine/Private/Animation/PoseAsset.cpp b/Engine/Source/Runtime/Engine/Private/Animation/PoseAsset.cpp
--- a/Engine/Source/Runtime/Engine/Private/Animation/PoseAsset.cpp
+++ b/Engine/Source/Runtime/Engine/Private/Animation/PoseAsset.cpp
@@ -998,7 +998,7 @@ void UPoseAsset::AddOrUpdatePose(const FSmartName& PoseName, USkeletalMeshCompon
 		TArray<FName> TrackNames;
 		// note this ignores root motion
 		TArray<FTransform> BoneTransform = MeshComponent->GetComponentSpaceTransforms();
-		const FReferenceSkeleton& RefSkeleton = MeshComponent->SkeletalMesh->RefSkeleton;
+		const FReferenceSkeleton& RefSkeleton = MeshComponent->SkeletalMesh->GetRefSkeleton();
 		for (int32 BoneIndex = 0; BoneIndex < RefSkeleton.GetNum(); ++BoneIndex)
 		{
 			TrackNames.Add(RefSkeleton.GetBoneName(BoneIndex));
@@ -1014,7 +1014,7 @@ void UPoseAsset::AddOrUpdatePose(const FSmartName& PoseName, USkeletalMeshCompon
 			}
 		}
 
-		const USkeleton* MeshSkeleton = MeshComponent->SkeletalMesh->Skeleton;
+		const USkeleton* MeshSkeleton = MeshComponent->SkeletalMesh->GetSkeleton();
 		const FSmartNameMapping* Mapping = MeshSkeleton->GetSmartNameContainer(USkeleton::AnimCurveMappingName);
 		
 		TArray<float> NewCurveValues;
diff --git a/Engine/Source/Runtime/Engine/Private/Animation/ReferenceSkeleton.cpp b/Engine/Source/Runtime/Engine/Private/Animation/ReferenceSkeleton.cpp
--- a/Engine/Source/Runtime/Engine/Private/Animation/ReferenceSkeleton.cpp
+++ b/Engine/Source/Runtime/Engine/Private/Animation/ReferenceSkeleton.cpp
@@ -193,7 +193,7 @@ void FReferenceSkeleton::RemoveDuplicateBones(const UObject* Requester, TArray<F
 		{
 			if (const USkeletalMesh* Mesh = Cast<USkeletalMesh>(Requester))
 			{
-				Skeleton = Mesh->Skeleton;
+				Skeleton = Mesh->GetSkeleton();
 			}
 			else
 			{
diff --git a/Engine/Source/Runtime/Engine/Private/Animation/SkeletalMeshActor.cpp b/Engine/Source/Runtime/Engine/Private/Animation/SkeletalMeshActor.cpp
--- a/Engine/Source/Runtime/Engine/Private/Animation/SkeletalMeshActor.cpp
+++ b/Engine/Source/Runtime/Engine/Private/Animation/SkeletalMeshActor.cpp
@@ -215,8 +215,8 @@ void ASkeletalMeshActor::BeginAnimControl(UInterpGroup* InInterpGroup)
 
 bool ASkeletalMeshActor::CanPlayAnimation(class UAnimSequenceBase* AnimAssetBase/*=NULL*/) const
 {
-	return (SkeletalMeshComponent->SkeletalMesh && SkeletalMeshComponent->SkeletalMesh->Skeleton &&
-		(!AnimAssetBase || SkeletalMeshComponent->SkeletalMesh->Skeleton->IsCompatible(AnimAssetBase->GetSkeleton())));
+	return (SkeletalMeshComponent->SkeletalMesh && SkeletalMeshComponent->SkeletalMesh->GetSkeleton() &&
+		(!AnimAssetBase || SkeletalMeshComponent->SkeletalMesh->GetSkeleton()->IsCompatible(AnimAssetBase->GetSkeleton())));
 }
 
 void ASkeletalMeshActor::SetAnimPosition(FName SlotName, int32 ChannelIndex, UAnimSequence* InAnimSequence, float InPosition, bool bFireNotifies, bool bLooping)
diff --git a/Engine/Source/Runtime/Engine/Private/Animation/Skeleton.cpp b/Engine/Source/Runtime/Engine/Private/Animation/Skeleton.cpp
--- a/Engine/Source/Runtime/Engine/Private/Animation/Skeleton.cpp
+++ b/Engine/Source/Runtime/Engine/Private/Animation/Skeleton.cpp
@@ -332,7 +332,7 @@ void USkeleton::ConvertToFReferenceSkeleton()
 bool USkeleton::DoesParentChainMatch(int32 StartBoneIndex, const USkeletalMesh* InSkelMesh) const
 {
 	const FReferenceSkeleton& SkeletonRefSkel = ReferenceSkeleton;
-	const FReferenceSkeleton& MeshRefSkel = InSkelMesh->RefSkeleton;
+	const FReferenceSkeleton& MeshRefSkel = InSkelMesh->GetRefSkeleton();
 
 	// if start is root bone
 	if ( StartBoneIndex == 0 )
@@ -380,7 +380,7 @@ bool USkeleton::IsCompatibleMesh(const USkeletalMesh* InSkelMesh) const
 	int32 NumOfBoneMatches = 0;
 
 	const FReferenceSkeleton& SkeletonRefSkel = ReferenceSkeleton;
-	const FReferenceSkeleton& MeshRefSkel = InSkelMesh->RefSkeleton;
+	const FReferenceSkeleton& MeshRefSkel = InSkelMesh->GetRefSkeleton();
 	const int32 NumBones = MeshRefSkel.GetRawBoneNum();
 
 	// first ensure the parent exists for each bone
@@ -484,7 +484,7 @@ void USkeleton::RemoveLinkup(const USkeletalMesh* InSkelMesh)
 int32 USkeleton::BuildLinkup(const USkeletalMesh* InSkelMesh)
 {
 	const FReferenceSkeleton& SkeletonRefSkel = ReferenceSkeleton;
-	const FReferenceSkeleton& MeshRefSkel = InSkelMesh->RefSkeleton;
+	const FReferenceSkeleton& MeshRefSkel = InSkelMesh->GetRefSkeleton();
 
 	// @todoanim : need to refresh NULL SkeletalMeshes from Cache
 	// since now they're autoweak pointer, they will go away if not used
@@ -581,10 +581,10 @@ void USkeleton::UpdateReferencePoseFromMesh(const USkeletalMesh* InSkelMesh)
 	for (int32 BoneIndex = 0; BoneIndex < ReferenceSkeleton.GetRawBoneNum(); BoneIndex++)
 	{
 		// find index from ref pose array
-		const int32 MeshBoneIndex = InSkelMesh->RefSkeleton.FindRawBoneIndex(ReferenceSkeleton.GetBoneName(BoneIndex));
+		const int32 MeshBoneIndex = InSkelMesh->GetRefSkeleton().FindRawBoneIndex(ReferenceSkeleton.GetBoneName(BoneIndex));
 		if (MeshBoneIndex != INDEX_NONE)
 		{
-			RefSkelModifier.UpdateRefPoseTransform(BoneIndex, InSkelMesh->RefSkeleton.GetRefBonePose()[MeshBoneIndex]);
+			RefSkelModifier.UpdateRefPoseTransform(BoneIndex, InSkelMesh->GetRefSkeleton().GetRefBonePose()[MeshBoneIndex]);
 		}
 	}
 
@@ -613,9 +613,9 @@ bool USkeleton::MergeAllBonesToBoneTree(const USkeletalMesh* InSkelMesh)
 		TArray<int32> RequiredBoneIndices;
 
 		// for now add all in this case. 
-		RequiredBoneIndices.AddUninitialized(InSkelMesh->RefSkeleton.GetRawBoneNum());
+		RequiredBoneIndices.AddUninitialized(InSkelMesh->GetRefSkeleton().GetRawBoneNum());
 		// gather bone list
-		for (int32 I=0; I<InSkelMesh->RefSkeleton.GetRawBoneNum(); ++I)
+		for (int32 I=0; I<InSkelMesh->GetRefSkeleton().GetRawBoneNum(); ++I)
 		{
 			RequiredBoneIndices[I] = I;
 		}
@@ -634,7 +634,7 @@ bool USkeleton::CreateReferenceSkeletonFromMesh(const USkeletalMesh* InSkeletalM
 {
 	// Filter list, we only want bones that have their parents present in this array.
 	TArray<int32> FilteredRequiredBones; 
-	FAnimationRuntime::ExcludeBonesWithNoParents(RequiredRefBones, InSkeletalMesh->RefSkeleton, FilteredRequiredBones);
+	FAnimationRuntime::ExcludeBonesWithNoParents(RequiredRefBones, InSkeletalMesh->GetRefSkeleton(), FilteredRequiredBones);
 
 	FReferenceSkeletonModifier RefSkelModifier(ReferenceSkeleton, this);
 
@@ -649,7 +649,7 @@ bool USkeleton::CreateReferenceSkeletonFromMesh(const USkeletalMesh* InSkeletalM
 		{
 			const int32& BoneIndex = FilteredRequiredBones[Index];
 
-			FMeshBoneInfo NewMeshBoneInfo = InSkeletalMesh->RefSkeleton.GetRefBoneInfo()[BoneIndex];
+			FMeshBoneInfo NewMeshBoneInfo = InSkeletalMesh->GetRefSkeleton().GetRefBoneInfo()[BoneIndex];
 			// Fix up ParentIndex for our new Skeleton.
 			if( BoneIndex == 0 )
 			{
@@ -657,11 +657,11 @@ bool USkeleton::CreateReferenceSkeletonFromMesh(const USkeletalMesh* InSkeletalM
 			}
 			else
 			{
-				const int32 ParentIndex = InSkeletalMesh->RefSkeleton.GetParentIndex(BoneIndex);
-				const FName ParentName = InSkeletalMesh->RefSkeleton.GetBoneName(ParentIndex);
+				const int32 ParentIndex = InSkeletalMesh->GetRefSkeleton().GetParentIndex(BoneIndex);
+				const FName ParentName = InSkeletalMesh->GetRefSkeleton().GetBoneName(ParentIndex);
 				NewMeshBoneInfo.ParentIndex = ReferenceSkeleton.FindRawBoneIndex(ParentName);
 			}
-			RefSkelModifier.Add(NewMeshBoneInfo, InSkeletalMesh->RefSkeleton.GetRefBonePose()[BoneIndex]);
+			RefSkelModifier.Add(NewMeshBoneInfo, InSkeletalMesh->GetRefSkeleton().GetRefBonePose()[BoneIndex]);
 		}
 
 		return true;
@@ -692,19 +692,19 @@ bool USkeleton::MergeBonesToBoneTree(const USkeletalMesh* InSkeletalMesh, const
 		{
 			// Exclude bones who do not have a parent.
 			TArray<int32> FilteredRequiredBones;
-			FAnimationRuntime::ExcludeBonesWithNoParents(RequiredRefBones, InSkeletalMesh->RefSkeleton, FilteredRequiredBones);
+			FAnimationRuntime::ExcludeBonesWithNoParents(RequiredRefBones, InSkeletalMesh->GetRefSkeleton(), FilteredRequiredBones);
 
 			FReferenceSkeletonModifier RefSkelModifier(ReferenceSkeleton, this);
 
 			for (int32 Index=0; Index<FilteredRequiredBones.Num(); Index++)
 			{
 				const int32& MeshBoneIndex = FilteredRequiredBones[Index];
-				const int32& SkeletonBoneIndex = ReferenceSkeleton.FindRawBoneIndex(InSkeletalMesh->RefSkeleton.GetBoneName(MeshBoneIndex));
+				const int32& SkeletonBoneIndex = ReferenceSkeleton.FindRawBoneIndex(InSkeletalMesh->GetRefSkeleton().GetBoneName(MeshBoneIndex));
 				
 				// Bone doesn't already exist. Add it.
 				if( SkeletonBoneIndex == INDEX_NONE )
 				{
-					FMeshBoneInfo NewMeshBoneInfo = InSkeletalMesh->RefSkeleton.GetRefBoneInfo()[MeshBoneIndex];
+					FMeshBoneInfo NewMeshBoneInfo = InSkeletalMesh->GetRefSkeleton().GetRefBoneInfo()[MeshBoneIndex];
 					// Fix up ParentIndex for our new Skeleton.
 					if( ReferenceSkeleton.GetRawBoneNum() == 0 )
 					{
@@ -712,10 +712,10 @@ bool USkeleton::MergeBonesToBoneTree(const USkeletalMesh* InSkeletalMesh, const
 					}
 					else
 					{
-						NewMeshBoneInfo.ParentIndex = ReferenceSkeleton.FindRawBoneIndex(InSkeletalMesh->RefSkeleton.GetBoneName(InSkeletalMesh->RefSkeleton.GetParentIndex(MeshBoneIndex)));
+						NewMeshBoneInfo.ParentIndex = ReferenceSkeleton.FindRawBoneIndex(InSkeletalMesh->GetRefSkeleton().GetBoneName(InSkeletalMesh->GetRefSkeleton().GetParentIndex(MeshBoneIndex)));
 					}
 
-					RefSkelModifier.Add(NewMeshBoneInfo, InSkeletalMesh->RefSkeleton.GetRefBonePose()[MeshBoneIndex]);
+					RefSkelModifier.Add(NewMeshBoneInfo, InSkeletalMesh->GetRefSkeleton().GetRefBonePose()[MeshBoneIndex]);
 					BoneTree.AddZeroed(1);
 					bShouldHandleHierarchyChange = true;
 				}
@@ -811,7 +811,7 @@ USkeletalMesh* USkeleton::GetPreviewMesh(bool bFindIfNotSet/*=false*/)
 #if WITH_EDITORONLY_DATA
 	USkeletalMesh* PreviewMesh = PreviewSkeletalMesh.LoadSynchronous();
 
-	if(PreviewMesh && PreviewMesh->Skeleton != this) // fix mismatched skeleton
+	if(PreviewMesh && PreviewMesh->GetSkeleton() != this) // fix mismatched skeleton
 	{
 		PreviewSkeletalMesh.Reset();
 		PreviewMesh = nullptr;
@@ -964,7 +964,8 @@ USkeletalMesh* USkeleton::FindCompatibleMesh() const
 	Filter.ClassNames.Add(USkeletalMesh::StaticClass()->GetFName());
 
 	FString SkeletonString = FAssetData(this).GetExportTextName();
-	Filter.TagsAndValues.Add(GET_MEMBER_NAME_CHECKED(USkeletalMesh, Skeleton), SkeletonString);
+	
+	Filter.TagsAndValues.Add(USkeletalMesh::GetSkeletonMemberName(), SkeletonString);
 
 	TArray<FAssetData> AssetList;
 	FAssetRegistryModule& AssetRegistryModule = FModuleManager::LoadModuleChecked<FAssetRegistryModule>(TEXT("AssetRegistry"));
@@ -1710,10 +1711,10 @@ void USkeleton::HandleVirtualBoneChanges()
 	for (TObjectIterator<USkeletalMesh> ItMesh; ItMesh; ++ItMesh)
 	{
 		USkeletalMesh* SkelMesh = *ItMesh;
-		if (SkelMesh->Skeleton == this)
+		if (SkelMesh->GetSkeleton() == this)
 		{
 			// also have to update retarget base pose
-			SkelMesh->RefSkeleton.RebuildRefSkeleton(this, bRebuildNameMap);
+			SkelMesh->GetRefSkeleton().RebuildRefSkeleton(this, bRebuildNameMap);
 			RebuildLinkup(SkelMesh);
 #if WITH_EDITOR
 			// whole bone count has changed, so it has to recalculate retarget base pose
@@ -1734,7 +1735,7 @@ void USkeleton::HandleVirtualBoneChanges()
 		USkinnedMeshComponent* MeshComponent = *It;
 		if (MeshComponent &&
 			MeshComponent->SkeletalMesh &&
-			MeshComponent->SkeletalMesh->Skeleton == this &&
+			MeshComponent->SkeletalMesh->GetSkeleton() == this &&
 			!MeshComponent->IsTemplate())
 		{
 			FComponentReregisterContext Context(MeshComponent);
diff --git a/Engine/Source/Runtime/Engine/Private/AnimationAsset.cpp b/Engine/Source/Runtime/Engine/Private/AnimationAsset.cpp
--- a/Engine/Source/Runtime/Engine/Private/AnimationAsset.cpp
+++ b/Engine/Source/Runtime/Engine/Private/AnimationAsset.cpp
@@ -290,7 +290,7 @@ USkeletalMesh* UAnimationAsset::GetPreviewMesh(bool bFindIfNotSet)
 #if WITH_EDITORONLY_DATA
 	USkeletalMesh* PreviewMesh = PreviewSkeletalMesh.LoadSynchronous();
 	// if somehow skeleton changes, just nullify it. 
-	if (PreviewMesh && PreviewMesh->Skeleton != Skeleton)
+	if (PreviewMesh && PreviewMesh->GetSkeleton() != Skeleton)
 	{
 		PreviewMesh = nullptr;
 		SetPreviewMesh(nullptr);
diff --git a/Engine/Source/Runtime/Engine/Private/BoneContainer.cpp b/Engine/Source/Runtime/Engine/Private/BoneContainer.cpp
--- a/Engine/Source/Runtime/Engine/Private/BoneContainer.cpp
+++ b/Engine/Source/Runtime/Engine/Private/BoneContainer.cpp
@@ -61,8 +61,8 @@ void FBoneContainer::Initialize(const FCurveEvaluationOption& CurveEvalOption)
 
 	if (AssetSkeletalMeshObj)
 	{
-		RefSkeleton = &AssetSkeletalMeshObj->RefSkeleton;
-		AssetSkeletonObj = AssetSkeletalMeshObj->Skeleton;
+		RefSkeleton = &AssetSkeletalMeshObj->GetRefSkeleton();
+		AssetSkeletonObj = AssetSkeletalMeshObj->GetSkeleton();
 	}
 	else
 	{
diff --git a/Engine/Source/Runtime/Engine/Private/Components/PoseableMeshComponent.cpp b/Engine/Source/Runtime/Engine/Private/Components/PoseableMeshComponent.cpp
--- a/Engine/Source/Runtime/Engine/Private/Components/PoseableMeshComponent.cpp
+++ b/Engine/Source/Runtime/Engine/Private/Components/PoseableMeshComponent.cpp
@@ -21,9 +21,9 @@ bool UPoseableMeshComponent::AllocateTransformData()
 	// Allocate transforms if not present.
 	if ( Super::AllocateTransformData() )
 	{
-		if(BoneSpaceTransforms.Num() != SkeletalMesh->RefSkeleton.GetNum() )
+		if(BoneSpaceTransforms.Num() != SkeletalMesh->GetRefSkeleton().GetNum() )
 		{
-			BoneSpaceTransforms = SkeletalMesh->RefSkeleton.GetRefBonePose();
+			BoneSpaceTransforms = SkeletalMesh->GetRefSkeleton().GetRefBonePose();
 
 			TArray<FBoneIndexType> RequiredBoneIndexArray;
 			RequiredBoneIndexArray.AddUninitialized(BoneSpaceTransforms.Num());
@@ -84,9 +84,9 @@ void UPoseableMeshComponent::FillComponentSpaceTransforms()
 	}
 
 	// right now all this does is to convert to SpaceBases
-	check( SkeletalMesh->RefSkeleton.GetNum() == BoneSpaceTransforms.Num() );
-	check( SkeletalMesh->RefSkeleton.GetNum() == GetNumComponentSpaceTransforms());
-	check( SkeletalMesh->RefSkeleton.GetNum() == GetBoneVisibilityStates().Num() );
+	check( SkeletalMesh->GetRefSkeleton().GetNum() == BoneSpaceTransforms.Num() );
+	check( SkeletalMesh->GetRefSkeleton().GetNum() == GetNumComponentSpaceTransforms());
+	check( SkeletalMesh->GetRefSkeleton().GetNum() == GetBoneVisibilityStates().Num() );
 
 	const int32 NumBones = BoneSpaceTransforms.Num();
 
@@ -113,7 +113,7 @@ void UPoseableMeshComponent::FillComponentSpaceTransforms()
 		BoneProcessed[BoneIndex] = 1;
 #endif
 		// For all bones below the root, final component-space transform is relative transform * component-space transform of parent.
-		const int32 ParentIndex = SkeletalMesh->RefSkeleton.GetParentIndex(BoneIndex);
+		const int32 ParentIndex = SkeletalMesh->GetRefSkeleton().GetParentIndex(BoneIndex);
 		FPlatformMisc::Prefetch(SpaceBasesData + ParentIndex);
 
 #if DO_GUARD_SLOW
@@ -269,7 +269,7 @@ void UPoseableMeshComponent::ResetBoneTransformByName(FName BoneName)
 	const int32 BoneIndex = GetBoneIndex(BoneName);
 	if( BoneIndex != INDEX_NONE )
 	{
-		BoneSpaceTransforms[BoneIndex] = SkeletalMesh->RefSkeleton.GetRefBonePose()[BoneIndex];
+		BoneSpaceTransforms[BoneIndex] = SkeletalMesh->GetRefSkeleton().GetRefBonePose()[BoneIndex];
 	}
 	else
 	{
@@ -294,10 +294,10 @@ void UPoseableMeshComponent::CopyPoseFromSkeletalComponent(USkeletalMeshComponen
 			// The meshes don't match, search bone-by-bone (slow path)
 
 			// first set the localatoms to ref pose from our current mesh
-			BoneSpaceTransforms = SkeletalMesh->RefSkeleton.GetRefBonePose();
+			BoneSpaceTransforms = SkeletalMesh->GetRefSkeleton().GetRefBonePose();
 
 			// Now overwrite any matching bones
-			const int32 NumSourceBones = InComponentToCopy->SkeletalMesh->RefSkeleton.GetNum();
+			const int32 NumSourceBones = InComponentToCopy->SkeletalMesh->GetRefSkeleton().GetNum();
 
 			for(int32 SourceBoneIndex = 0 ; SourceBoneIndex < NumSourceBones ; ++SourceBoneIndex)
 			{
diff --git a/Engine/Source/Runtime/Engine/Private/Components/SkeletalMeshComponent.cpp b/Engine/Source/Runtime/Engine/Private/Components/SkeletalMeshComponent.cpp
--- a/Engine/Source/Runtime/Engine/Private/Components/SkeletalMeshComponent.cpp
+++ b/Engine/Source/Runtime/Engine/Private/Components/SkeletalMeshComponent.cpp
@@ -457,7 +457,7 @@ void USkeletalMeshComponent::UpdateClothTickRegisteredState()
 
 void USkeletalMeshComponent::FinalizePoseEvaluationResult(const USkeletalMesh* InMesh, TArray<FTransform>& OutBoneSpaceTransforms, FVector& OutRootBoneTranslation, FCompactPose& InFinalPose) const
 {
-	OutBoneSpaceTransforms = InMesh->RefSkeleton.GetRefBonePose();
+	OutBoneSpaceTransforms = InMesh->GetRefSkeleton().GetRefBonePose();
 
 	if(InFinalPose.IsValid() && InFinalPose.GetNumBones() > 0)
 	{
@@ -471,17 +471,17 @@ void USkeletalMeshComponent::FinalizePoseEvaluationResult(const USkeletalMesh* I
 	}
 	else
 	{
-		OutBoneSpaceTransforms = InMesh->RefSkeleton.GetRefBonePose();
+		OutBoneSpaceTransforms = InMesh->GetRefSkeleton().GetRefBonePose();
 	}
 
-	OutRootBoneTranslation = OutBoneSpaceTransforms[0].GetTranslation() - InMesh->RefSkeleton.GetRefBonePose()[0].GetTranslation();
+	OutRootBoneTranslation = OutBoneSpaceTransforms[0].GetTranslation() - InMesh->GetRefSkeleton().GetRefBonePose()[0].GetTranslation();
 }
 
 bool USkeletalMeshComponent::NeedToSpawnAnimScriptInstance() const
 {
 	IAnimClassInterface* AnimClassInterface = IAnimClassInterface::GetFromClass(AnimClass);
 	const USkeleton* AnimSkeleton = (AnimClassInterface) ? AnimClassInterface->GetTargetSkeleton() : nullptr;
-	const bool bAnimSkelValid = !AnimClassInterface || (AnimSkeleton && SkeletalMesh && SkeletalMesh->Skeleton->IsCompatible(AnimSkeleton) && AnimSkeleton->IsCompatibleMesh(SkeletalMesh));
+	const bool bAnimSkelValid = !AnimClassInterface || (AnimSkeleton && SkeletalMesh && SkeletalMesh->GetSkeleton()->IsCompatible(AnimSkeleton) && AnimSkeleton->IsCompatibleMesh(SkeletalMesh));
 
 	if (AnimationMode == EAnimationMode::AnimationBlueprint && AnimClass && bAnimSkelValid)
 	{
@@ -503,7 +503,7 @@ bool USkeletalMeshComponent::NeedToSpawnPostPhysicsInstance(bool bForceReinit) c
 	if(SkeletalMesh)
 	{
 		const UClass* MainInstanceClass = *AnimClass;
-		const UClass* ClassToUse = *SkeletalMesh->PostProcessAnimBlueprint;
+		const UClass* ClassToUse = *SkeletalMesh->GetPostProcessAnimBlueprint();
 		const UClass* CurrentClass = PostProcessAnimInstance ? PostProcessAnimInstance->GetClass() : nullptr;
 
 		// We need to have an instance, and we have the wrong class (different or null)
@@ -557,7 +557,7 @@ void USkeletalMeshComponent::OnRegister()
 		bool bSupportsAllAssets = true;
 
 		UClothingSimulationFactory* const DefaultObject = ClothingSimulationFactory->GetDefaultObject<UClothingSimulationFactory>();
-		for (UClothingAssetBase* const ClothingAsset : SkeletalMesh->MeshClothingAssets)
+		for (UClothingAssetBase* const ClothingAsset : SkeletalMesh->GetMeshClothingAssets())
 		{
 			if (ClothingAsset && !DefaultObject->SupportsAsset(ClothingAsset))
 			{
@@ -590,7 +590,7 @@ void USkeletalMeshComponent::OnRegister()
 						int NumAssets = 0;
 						int SupportedNumAssets = 0;
 						UClothingSimulationFactory* const NewDefaultObject = NewClothingSimulationFactory->GetDefaultObject<UClothingSimulationFactory>();
-						for (UClothingAssetBase* const ClothingAsset : SkeletalMesh->MeshClothingAssets)
+						for (UClothingAssetBase* const ClothingAsset : SkeletalMesh->GetMeshClothingAssets())
 						{
 							if (ClothingAsset)
 							{
@@ -712,7 +712,7 @@ void USkeletalMeshComponent::InitAnim(bool bForceReinit)
 		CachedAnimCurveUidVersion = 0;
 
 		// we still need this in case users doesn't call tick, but sent to renderer
-		MorphTargetWeights.SetNumZeroed(SkeletalMesh->MorphTargets.Num());
+		MorphTargetWeights.SetNumZeroed(SkeletalMesh->GetMorphTargets().Num());
 
 		// We may be doing parallel evaluation on the current anim instance
 		// Calling this here with true will block this init till that thread completes
@@ -727,7 +727,7 @@ void USkeletalMeshComponent::InitAnim(bool bForceReinit)
 		const USkeleton* AnimSkeleton = (AnimScriptInstance)? AnimScriptInstance->CurrentSkeleton : nullptr;
 
 		const bool bClearAnimInstance = AnimScriptInstance && !AnimSkeleton;
-		const bool bSkeletonMismatch = AnimSkeleton && (AnimScriptInstance->CurrentSkeleton!=SkeletalMesh->Skeleton);
+		const bool bSkeletonMismatch = AnimSkeleton && (AnimScriptInstance->CurrentSkeleton!=SkeletalMesh->GetSkeleton());
 		const bool bSkeletonNotCompatible = AnimSkeleton && !bSkeletonMismatch && (AnimSkeleton->IsCompatibleMesh(SkeletalMesh) == false);
 
 		LastPoseTickFrame = 0;
@@ -762,7 +762,7 @@ void USkeletalMeshComponent::InitAnim(bool bForceReinit)
 				else
 				{
 					PRAGMA_DISABLE_DEPRECATION_WARNINGS
-					BoneSpaceTransforms = SkeletalMesh->RefSkeleton.GetRefBonePose();
+					BoneSpaceTransforms = SkeletalMesh->GetRefSkeleton().GetRefBonePose();
 					//Mini RefreshBoneTransforms (the bit we actually care about)
 					FillComponentSpaceTransforms(SkeletalMesh, BoneSpaceTransforms, GetEditableComponentSpaceTransforms());
 					PRAGMA_ENABLE_DEPRECATION_WARNINGS
@@ -806,7 +806,7 @@ bool USkeletalMeshComponent::InitializeAnimScriptInstance(bool bForceReinit, boo
 		}
 		else 
 		{
-			bool bShouldSpawnSingleNodeInstance = SkeletalMesh && SkeletalMesh->Skeleton && AnimationMode == EAnimationMode::AnimationSingleNode;
+			bool bShouldSpawnSingleNodeInstance = SkeletalMesh && SkeletalMesh->GetSkeleton() && AnimationMode == EAnimationMode::AnimationSingleNode;
 			if (bShouldSpawnSingleNodeInstance)
 			{
 				SCOPE_CYCLE_COUNTER(STAT_AnimSpawnTime);
@@ -846,7 +846,7 @@ bool USkeletalMeshComponent::InitializeAnimScriptInstance(bool bForceReinit, boo
 		}
 
 		// May need to clear out the post physics instance
-		UClass* NewMeshInstanceClass = *SkeletalMesh->PostProcessAnimBlueprint;
+		UClass* NewMeshInstanceClass = *SkeletalMesh->GetPostProcessAnimBlueprint();
 		if(!NewMeshInstanceClass || NewMeshInstanceClass == *AnimClass)
 		{
 			PostProcessAnimInstance = nullptr;
@@ -854,7 +854,7 @@ bool USkeletalMeshComponent::InitializeAnimScriptInstance(bool bForceReinit, boo
 
 		if(NeedToSpawnPostPhysicsInstance(bForceReinit))
 		{
-			PostProcessAnimInstance = NewObject<UAnimInstance>(this, *SkeletalMesh->PostProcessAnimBlueprint);
+			PostProcessAnimInstance = NewObject<UAnimInstance>(this, *SkeletalMesh->GetPostProcessAnimBlueprint());
 
 			if(PostProcessAnimInstance)
 			{
@@ -868,7 +868,7 @@ bool USkeletalMeshComponent::InitializeAnimScriptInstance(bool bForceReinit, boo
 				bInitializedPostInstance = true;
 			}
 		}
-		else if (!SkeletalMesh->PostProcessAnimBlueprint.Get())
+		else if (!SkeletalMesh->GetPostProcessAnimBlueprint().Get())
 		{
 			PostProcessAnimInstance = nullptr;
 		}
@@ -991,12 +991,12 @@ void USkeletalMeshComponent::PostEditChangeProperty(FPropertyChangedEvent& Prope
 			if(PostProcessAnimInstance)
 			{
 				UClass* CurrentClass = PostProcessAnimInstance->GetClass();
-				UClass* MeshClass = SkeletalMesh ? *SkeletalMesh->PostProcessAnimBlueprint : nullptr;
+				UClass* MeshClass = SkeletalMesh ? *SkeletalMesh->GetPostProcessAnimBlueprint() : nullptr;
 				if(CurrentClass != MeshClass)
 				{
 					if(MeshClass)
 					{
-						PostProcessAnimInstance = NewObject<UAnimInstance>(this, *SkeletalMesh->PostProcessAnimBlueprint);
+						PostProcessAnimInstance = NewObject<UAnimInstance>(this, *SkeletalMesh->GetPostProcessAnimBlueprint());
 						PostProcessAnimInstance->InitializeAnimation();
 					}
 					else
@@ -1030,7 +1030,7 @@ void USkeletalMeshComponent::PostEditChangeProperty(FPropertyChangedEvent& Prope
 		if ( PropertyThatChanged->GetFName() == GET_MEMBER_NAME_CHECKED( FSingleAnimationPlayData, AnimToPlay ))
 		{
 			// make sure the animation skeleton matches the current skeletalmesh
-			if (AnimationData.AnimToPlay != nullptr && SkeletalMesh && AnimationData.AnimToPlay->GetSkeleton() != SkeletalMesh->Skeleton)
+			if (AnimationData.AnimToPlay != nullptr && SkeletalMesh && AnimationData.AnimToPlay->GetSkeleton() != SkeletalMesh->GetSkeleton())
 			{
 				UE_LOG(LogAnimation, Warning, TEXT("Invalid animation"));
 				AnimationData.AnimToPlay = nullptr;
@@ -1322,7 +1322,7 @@ void USkeletalMeshComponent::ResetMorphTargetCurves()
 
 	if (SkeletalMesh)
 	{
-		MorphTargetWeights.SetNum(SkeletalMesh->MorphTargets.Num());
+		MorphTargetWeights.SetNum(SkeletalMesh->GetMorphTargets().Num());
 
 		// we need this code to ensure the buffer gets cleared whether or not you have morphtarget curve set
 		// the case, where you had morphtargets weight on, and when you clear the weight, you want to make sure 
@@ -1481,8 +1481,8 @@ void USkeletalMeshComponent::FillComponentSpaceTransforms(const USkeletalMesh* I
 	}
 
 	// right now all this does is populate DestSpaceBases
-	check( InSkeletalMesh->RefSkeleton.GetNum() == InBoneSpaceTransforms.Num());
-	check( InSkeletalMesh->RefSkeleton.GetNum() == OutComponentSpaceTransforms.Num());
+	check( InSkeletalMesh->GetRefSkeleton().GetNum() == InBoneSpaceTransforms.Num());
+	check( InSkeletalMesh->GetRefSkeleton().GetNum() == OutComponentSpaceTransforms.Num());
 
 	const int32 NumBones = InBoneSpaceTransforms.Num();
 
@@ -1513,7 +1513,7 @@ void USkeletalMeshComponent::FillComponentSpaceTransforms(const USkeletalMesh* I
 			(ispc::FTransform*)&ComponentSpaceData[0],
 			(ispc::FTransform*)&LocalTransformsData[0],
 			FillComponentSpaceTransformsRequiredBones.GetData(),
-			(const uint8*)InSkeletalMesh->RefSkeleton.GetRefBoneInfo().GetData(),
+			(const uint8*)InSkeletalMesh->GetRefSkeleton().GetRefBoneInfo().GetData(),
 			sizeof(FMeshBoneInfo),
 			offsetof(FMeshBoneInfo, ParentIndex),
 			FillComponentSpaceTransformsRequiredBones.Num());
@@ -1533,7 +1533,7 @@ void USkeletalMeshComponent::FillComponentSpaceTransforms(const USkeletalMesh* I
 			BoneProcessed[BoneIndex] = 1;
 #endif
 			// For all bones below the root, final component-space transform is relative transform * component-space transform of parent.
-			const int32 ParentIndex = InSkeletalMesh->RefSkeleton.GetParentIndex(BoneIndex);
+			const int32 ParentIndex = InSkeletalMesh->GetRefSkeleton().GetParentIndex(BoneIndex);
 			FTransform* ParentSpaceBase = ComponentSpaceData + ParentIndex;
 			FPlatformMisc::Prefetch(ParentSpaceBase);
 
@@ -1614,9 +1614,9 @@ void USkeletalMeshComponent::RecalcRequiredCurves()
 		return;
 	}
 
-	if (SkeletalMesh->Skeleton)
+	if (SkeletalMesh->GetSkeleton())
 	{
-		CachedCurveUIDList = SkeletalMesh->Skeleton->GetDefaultCurveUIDList();
+		CachedCurveUIDList = SkeletalMesh->GetSkeleton()->GetDefaultCurveUIDList();
 	}
 
 	const FCurveEvaluationOption CurveEvalOption(bAllowAnimCurveEvaluation, &DisallowedAnimCurves, PredictedLODLevel);
@@ -1675,7 +1675,7 @@ void USkeletalMeshComponent::ComputeRequiredBones(TArray<FBoneIndexType>& OutReq
 	OutRequiredBones = LODData.RequiredBones;
 
 	// Add virtual bones
-	MergeInBoneIndexArrays(OutRequiredBones, SkeletalMesh->RefSkeleton.GetRequiredVirtualBones());
+	MergeInBoneIndexArrays(OutRequiredBones, SkeletalMesh->GetRefSkeleton().GetRequiredVirtualBones());
 
 	const UPhysicsAsset* const PhysicsAsset = GetPhysicsAsset();
 	// If we have a PhysicsAsset, we also need to make sure that all the bones used by it are always updated, as its used
@@ -1689,7 +1689,7 @@ void USkeletalMeshComponent::ComputeRequiredBones(TArray<FBoneIndexType>& OutReq
 			{
 				continue;
 			}
-			int32 PhysBoneIndex = SkeletalMesh->RefSkeleton.FindBoneIndex(PhysicsAsset->SkeletalBodySetups[i]->BoneName);
+			int32 PhysBoneIndex = SkeletalMesh->GetRefSkeleton().FindBoneIndex(PhysicsAsset->SkeletalBodySetups[i]->BoneName);
 			if (PhysBoneIndex != INDEX_NONE)
 			{
 				PhysAssetBones.Add(PhysBoneIndex);
@@ -1749,10 +1749,11 @@ void USkeletalMeshComponent::ComputeRequiredBones(TArray<FBoneIndexType>& OutReq
 		TArray<FBoneIndexType> MirroredDesiredBones;
 		MirroredDesiredBones.AddUninitialized(RequiredBones.Num());
 
+		const TArray<struct FBoneMirrorInfo>& SkelMirrorTable = SkeletalMesh->GetSkelMirrorTable();
 		// Look up each bone in the mirroring table.
 		for (int32 i = 0; i<OutRequiredBones.Num(); i++)
 		{
-			MirroredDesiredBones[i] = SkeletalMesh->SkelMirrorTable[OutRequiredBones[i]].SourceIndex;
+			MirroredDesiredBones[i] = SkelMirrorTable[OutRequiredBones[i]].SourceIndex;
 		}
 
 		// Sort to ensure strictly increasing order.
@@ -1768,7 +1769,7 @@ void USkeletalMeshComponent::ComputeRequiredBones(TArray<FBoneIndexType>& OutReq
 
 		for (const USkeletalMeshSocket* Socket : SkeletalMesh->GetActiveSocketList())
 		{
-			int32 BoneIndex = SkeletalMesh->RefSkeleton.FindBoneIndex(Socket->BoneName);
+			int32 BoneIndex = SkeletalMesh->GetRefSkeleton().FindBoneIndex(Socket->BoneName);
 			if (BoneIndex != INDEX_NONE)
 			{
 				if (Socket->bForceAlwaysAnimated)
@@ -1802,14 +1803,14 @@ void USkeletalMeshComponent::ComputeRequiredBones(TArray<FBoneIndexType>& OutReq
 	}
 
 	// Ensure that we have a complete hierarchy down to those bones.
-	FAnimationRuntime::EnsureParentsPresent(OutRequiredBones, SkeletalMesh->RefSkeleton);
+	FAnimationRuntime::EnsureParentsPresent(OutRequiredBones, SkeletalMesh->GetRefSkeleton());
 
 	OutFillComponentSpaceTransformsRequiredBones.Reset(OutRequiredBones.Num() + NeededBonesForFillComponentSpaceTransforms.Num());
 	OutFillComponentSpaceTransformsRequiredBones = OutRequiredBones;
 
 	NeededBonesForFillComponentSpaceTransforms.Sort();
 	MergeInBoneIndexArrays(OutFillComponentSpaceTransformsRequiredBones, NeededBonesForFillComponentSpaceTransforms);
-	FAnimationRuntime::EnsureParentsPresent(OutFillComponentSpaceTransformsRequiredBones, SkeletalMesh->RefSkeleton);
+	FAnimationRuntime::EnsureParentsPresent(OutFillComponentSpaceTransformsRequiredBones, SkeletalMesh->GetRefSkeleton());
 }
 
 void USkeletalMeshComponent::RecalcRequiredBones(int32 LODIndex)
@@ -1821,7 +1822,7 @@ void USkeletalMeshComponent::RecalcRequiredBones(int32 LODIndex)
 
 	ComputeRequiredBones(RequiredBones, FillComponentSpaceTransformsRequiredBones, LODIndex, /*bIgnorePhysicsAsset=*/ false);
 	PRAGMA_DISABLE_DEPRECATION_WARNINGS
-	BoneSpaceTransforms = SkeletalMesh->RefSkeleton.GetRefBonePose();
+	BoneSpaceTransforms = SkeletalMesh->GetRefSkeleton().GetRefBonePose();
 	PRAGMA_ENABLE_DEPRECATION_WARNINGS
 	// make sure animation requiredBone to mark as dirty
 	if (AnimScriptInstance)
@@ -1850,15 +1851,15 @@ void USkeletalMeshComponent::RecalcRequiredBones(int32 LODIndex)
 
 void USkeletalMeshComponent::MarkRequiredCurveUpToDate()
 {
-	if (SkeletalMesh && SkeletalMesh->Skeleton)
+	if (SkeletalMesh && SkeletalMesh->GetSkeleton())
 	{
-		CachedAnimCurveUidVersion = SkeletalMesh->Skeleton->GetAnimCurveUidVersion();
+		CachedAnimCurveUidVersion = SkeletalMesh->GetSkeleton()->GetAnimCurveUidVersion();
 	}
 }
 
 bool USkeletalMeshComponent::AreRequiredCurvesUpToDate() const
 {
-	return (!SkeletalMesh || !SkeletalMesh->Skeleton || CachedAnimCurveUidVersion == SkeletalMesh->Skeleton->GetAnimCurveUidVersion());
+	return (!SkeletalMesh || !SkeletalMesh->GetSkeleton() || CachedAnimCurveUidVersion == SkeletalMesh->GetSkeleton()->GetAnimCurveUidVersion());
 }
 
 void USkeletalMeshComponent::EvaluateAnimation(const USkeletalMesh* InSkeletalMesh, UAnimInstance* InAnimInstance, FVector& OutRootBoneTranslation, FBlendedHeapCurve& OutCurve, FCompactPose& OutPose, FHeapCustomAttributes& OutAttributes) const
@@ -1871,7 +1872,7 @@ void USkeletalMeshComponent::EvaluateAnimation(const USkeletalMesh* InSkeletalMe
 	}
 
 	// We can only evaluate animation if RequiredBones is properly setup for the right mesh!
-	if( InSkeletalMesh->Skeleton && 
+	if( InSkeletalMesh->GetSkeleton() && 
 		InAnimInstance &&
 		InAnimInstance->ParallelCanEvaluate(InSkeletalMesh))
 	{
@@ -1903,7 +1904,7 @@ void USkeletalMeshComponent::UpdateSlaveComponent()
 		// now propagate BP-driven curves from the master SMC...
 		if (SkeletalMesh)
 		{
-			check(MorphTargetWeights.Num() == SkeletalMesh->MorphTargets.Num());
+			check(MorphTargetWeights.Num() == SkeletalMesh->GetMorphTargets().Num());
 			if (MasterSMC->MorphTargetCurves.Num() > 0)
 			{
 				FAnimationRuntime::AppendActiveMorphTargets(SkeletalMesh, MasterSMC->MorphTargetCurves, ActiveMorphTargets, MorphTargetWeights);
@@ -2165,7 +2166,7 @@ void USkeletalMeshComponent::RefreshBoneTransforms(FActorComponentTickFunction*
 	//Dont mark cache as invalid if we aren't performing optimization anyway
 	PRAGMA_DISABLE_DEPRECATION_WARNINGS
 	const bool bInvalidCachedBones = bDoEvaluationRateOptimization &&
-									 ((BoneSpaceTransforms.Num() != SkeletalMesh->RefSkeleton.GetNum())
+									 ((BoneSpaceTransforms.Num() != SkeletalMesh->GetRefSkeleton().GetNum())
 									 || (BoneSpaceTransforms.Num() != CachedBoneSpaceTransforms.Num())
 									 || (GetNumComponentSpaceTransforms() != CachedComponentSpaceTransforms.Num()));
 
@@ -2819,9 +2820,9 @@ bool USkeletalMeshComponent::AllocateTransformData()
 	if ( Super::AllocateTransformData() )
 	{
 		PRAGMA_DISABLE_DEPRECATION_WARNINGS
-		if(BoneSpaceTransforms.Num() != SkeletalMesh->RefSkeleton.GetNum() )
+		if(BoneSpaceTransforms.Num() != SkeletalMesh->GetRefSkeleton().GetNum() )
 		{
-			BoneSpaceTransforms = SkeletalMesh->RefSkeleton.GetRefBonePose();
+			BoneSpaceTransforms = SkeletalMesh->GetRefSkeleton().GetRefBonePose();
 		}
 		PRAGMA_ENABLE_DEPRECATION_WARNINGS
 
@@ -3067,7 +3068,7 @@ void USkeletalMeshComponent::HideBone( int32 BoneIndex, EPhysBodyOp PhysBodyOpti
 
 		if (PhysBodyOption != PBO_None)
 		{
-			FName HideBoneName = SkeletalMesh->RefSkeleton.GetBoneName(BoneIndex);
+			FName HideBoneName = SkeletalMesh->GetRefSkeleton().GetBoneName(BoneIndex);
 			if (PhysBodyOption == PBO_Term)
 			{
 				TermBodiesBelow(HideBoneName);
@@ -3098,7 +3099,7 @@ void USkeletalMeshComponent::UnHideBone( int32 BoneIndex )
 	{
 		bRequiredBonesUpToDate = false;
 
-		//FName HideBoneName = SkeletalMesh->RefSkeleton.GetBoneName(BoneIndex);
+		//FName HideBoneName = SkeletalMesh->GetRefSkeleton().GetBoneName(BoneIndex);
 		// It's okay to turn this on for terminated bodies
 		// It won't do any if BodyData isn't found
 		// @JTODO
@@ -3595,7 +3596,7 @@ void USkeletalMeshComponent::UnregisterOnSkeletalMeshPropertyChanged( FDelegateH
 
 void USkeletalMeshComponent::ValidateAnimation()
 {
-	if (SkeletalMesh && SkeletalMesh->Skeleton == nullptr)
+	if (SkeletalMesh && SkeletalMesh->GetSkeleton() == nullptr)
 	{
 		UE_LOG(LogAnimation, Warning, TEXT("SkeletalMesh %s has no skeleton. This needs to fixed before an animation can be set"), *SkeletalMesh->GetName());
 		if (AnimationMode == EAnimationMode::AnimationSingleNode)
@@ -3616,11 +3617,11 @@ void USkeletalMeshComponent::ValidateAnimation()
 
 	if(AnimationMode == EAnimationMode::AnimationSingleNode)
 	{
-		if(AnimationData.AnimToPlay && SkeletalMesh && AnimationData.AnimToPlay->GetSkeleton() != SkeletalMesh->Skeleton)
+		if(AnimationData.AnimToPlay && SkeletalMesh && AnimationData.AnimToPlay->GetSkeleton() != SkeletalMesh->GetSkeleton())
 		{
-			if (SkeletalMesh->Skeleton)
+			if (SkeletalMesh->GetSkeleton())
 			{
-				UE_LOG(LogAnimation, Warning, TEXT("Animation %s is incompatible with skeleton %s, removing animation from actor."), *AnimationData.AnimToPlay->GetName(), *SkeletalMesh->Skeleton->GetName());
+				UE_LOG(LogAnimation, Warning, TEXT("Animation %s is incompatible with skeleton %s, removing animation from actor."), *AnimationData.AnimToPlay->GetName(), *SkeletalMesh->GetSkeleton()->GetName());
 			}
 			else
 			{
@@ -3633,11 +3634,11 @@ void USkeletalMeshComponent::ValidateAnimation()
 	else if (AnimationMode == EAnimationMode::AnimationBlueprint)
 	{
 		IAnimClassInterface* AnimClassInterface = IAnimClassInterface::GetFromClass(AnimClass);
-		if (AnimClassInterface && SkeletalMesh && AnimClassInterface->GetTargetSkeleton() != SkeletalMesh->Skeleton)
+		if (AnimClassInterface && SkeletalMesh && AnimClassInterface->GetTargetSkeleton() != SkeletalMesh->GetSkeleton())
 		{
-			if(SkeletalMesh->Skeleton)
+			if(SkeletalMesh->GetSkeleton())
 			{
-				UE_LOG(LogAnimation, Warning, TEXT("AnimBP %s is incompatible with skeleton %s, removing AnimBP from actor."), *AnimClass->GetName(), *SkeletalMesh->Skeleton->GetName());
+				UE_LOG(LogAnimation, Warning, TEXT("AnimBP %s is incompatible with skeleton %s, removing AnimBP from actor."), *AnimClass->GetName(), *SkeletalMesh->GetSkeleton()->GetName());
 			}
 			else
 			{
@@ -4010,7 +4011,7 @@ void USkeletalMeshComponent::GetCurrentRefToLocalMatrices(TArray<FMatrix>& OutRe
 		}
 		else
 		{
-			const FReferenceSkeleton& RefSkeleton = SkeletalMesh->RefSkeleton;
+			const FReferenceSkeleton& RefSkeleton = SkeletalMesh->GetRefSkeleton();
 			OutRefToLocals.AddUninitialized(RefSkeleton.GetNum());
 			for (int32 Index = 0; Index < OutRefToLocals.Num(); ++Index)
 			{
@@ -4127,7 +4128,7 @@ void USkeletalMeshComponent::SnapshotPose(FPoseSnapshot& Snapshot)
 	if (ensureAsRuntimeWarning(SkeletalMesh != nullptr))
 	{
 		const TArray<FTransform>& ComponentSpaceTMs = GetComponentSpaceTransforms();
-		const FReferenceSkeleton& RefSkeleton = SkeletalMesh->RefSkeleton;
+		const FReferenceSkeleton& RefSkeleton = SkeletalMesh->GetRefSkeleton();
 		const TArray<FTransform>& RefPoseSpaceBaseTMs = RefSkeleton.GetRefBonePose();
 
 		Snapshot.SkeletalMeshName = SkeletalMesh->GetFName();
@@ -4280,7 +4281,7 @@ void USkeletalMeshComponent::SetAllowedAnimCurvesEvaluation(const TArray<FName>&
 
 		if (SkeletalMesh)
 		{
-			USkeleton* Skeleton = SkeletalMesh->Skeleton;
+			USkeleton* Skeleton = SkeletalMesh->GetSkeleton();
 			if (Skeleton)
 			{
 				const FSmartNameMapping* Mapping = Skeleton->GetSmartNameContainer(USkeleton::AnimCurveMappingName);
@@ -4347,7 +4348,7 @@ bool USkeletalMeshComponent::GetBoneAttribute(const FName& BoneName, const FName
 	if (SkeletalMesh)
 	{
 		const FHeapCustomAttributes& Attributes = GetCustomAttributes();
-		const int32 BoneIndex = SkeletalMesh->RefSkeleton.FindBoneIndex(BoneName);
+		const int32 BoneIndex = SkeletalMesh->GetRefSkeleton().FindBoneIndex(BoneName);
 
 		bFound = Attributes.GetBoneAttribute(FCompactPoseBoneIndex(BoneIndex), AttributeName, OutValue);
 
@@ -4355,7 +4356,7 @@ bool USkeletalMeshComponent::GetBoneAttribute(const FName& BoneName, const FName
 		{
 			if (LookupType == ECustomBoneAttributeLookup::ImmediateParent)
 			{
-				const int32 ParentIndex = SkeletalMesh->RefSkeleton.GetParentIndex(BoneIndex);
+				const int32 ParentIndex = SkeletalMesh->GetRefSkeleton().GetParentIndex(BoneIndex);
 				if (ParentIndex != INDEX_NONE)
 				{
 					bFound = Attributes.GetBoneAttribute(FCompactPoseBoneIndex(ParentIndex), AttributeName, OutValue);
@@ -4364,7 +4365,7 @@ bool USkeletalMeshComponent::GetBoneAttribute(const FName& BoneName, const FName
 			else if (LookupType == ECustomBoneAttributeLookup::ParentHierarchy)
 			{
 				int32 SearchBoneIndex = BoneIndex;
-				int32 ParentIndex = SkeletalMesh->RefSkeleton.GetParentIndex(SearchBoneIndex);
+				int32 ParentIndex = SkeletalMesh->GetRefSkeleton().GetParentIndex(SearchBoneIndex);
 
 				while (ParentIndex != INDEX_NONE)
 				{
@@ -4375,7 +4376,7 @@ bool USkeletalMeshComponent::GetBoneAttribute(const FName& BoneName, const FName
 					}
 
 					SearchBoneIndex = ParentIndex;
-					ParentIndex = SkeletalMesh->RefSkeleton.GetParentIndex(SearchBoneIndex);
+					ParentIndex = SkeletalMesh->GetRefSkeleton().GetParentIndex(SearchBoneIndex);
 				}
 			}
 		}
