To protect the skeletalmesh when using async operations like build or re-import we deprecate all public member variables of USkeletalMesh class and create the necessary getter and setter. Those member

diff --git a/Engine/Source/Runtime/Engine/Private/Components/SkinnedMeshComponent.cpp b/Engine/Source/Runtime/Engine/Private/Components/SkinnedMeshComponent.cpp
--- a/Engine/Source/Runtime/Engine/Private/Components/SkinnedMeshComponent.cpp
+++ b/Engine/Source/Runtime/Engine/Private/Components/SkinnedMeshComponent.cpp
@@ -381,10 +381,10 @@ void USkinnedMeshComponent::UpdateMorphMaterialUsageOnProxy()
 	// update morph material usage
 	if (SceneProxy)
 	{
-		if (ActiveMorphTargets.Num() > 0 && SkeletalMesh->MorphTargets.Num() > 0)
+		if (ActiveMorphTargets.Num() > 0 && SkeletalMesh->GetMorphTargets().Num() > 0)
 		{
 			TArray<UMaterialInterface*> MaterialUsingMorphTarget;
-			for (UMorphTarget* MorphTarget : SkeletalMesh->MorphTargets)
+			for (UMorphTarget* MorphTarget : SkeletalMesh->GetMorphTargets())
 			{
 				if (!MorphTarget)
 				{
@@ -871,7 +871,7 @@ int32 USkinnedMeshComponent::GetNumMaterials() const
 {
 	if (SkeletalMesh)
 	{
-		return SkeletalMesh->Materials.Num();
+		return SkeletalMesh->GetMaterials().Num();
 	}
 
 	return 0;
@@ -883,9 +883,9 @@ UMaterialInterface* USkinnedMeshComponent::GetMaterial(int32 MaterialIndex) cons
 	{
 		return OverrideMaterials[MaterialIndex];
 	}
-	else if (SkeletalMesh && SkeletalMesh->Materials.IsValidIndex(MaterialIndex) && SkeletalMesh->Materials[MaterialIndex].MaterialInterface)
+	else if (SkeletalMesh && SkeletalMesh->GetMaterials().IsValidIndex(MaterialIndex) && SkeletalMesh->GetMaterials()[MaterialIndex].MaterialInterface)
 	{
-		return SkeletalMesh->Materials[MaterialIndex].MaterialInterface;
+		return SkeletalMesh->GetMaterials()[MaterialIndex].MaterialInterface;
 	}
 
 	return nullptr;
@@ -895,9 +895,10 @@ int32 USkinnedMeshComponent::GetMaterialIndex(FName MaterialSlotName) const
 {
 	if (SkeletalMesh != nullptr)
 	{
-		for (int32 MaterialIndex = 0; MaterialIndex < SkeletalMesh->Materials.Num(); ++MaterialIndex)
+		const TArray<FSkeletalMaterial>& SkeletalMeshMaterials = SkeletalMesh->GetMaterials();
+		for (int32 MaterialIndex = 0; MaterialIndex < SkeletalMeshMaterials.Num(); ++MaterialIndex)
 		{
-			const FSkeletalMaterial &SkeletalMaterial = SkeletalMesh->Materials[MaterialIndex];
+			const FSkeletalMaterial &SkeletalMaterial = SkeletalMeshMaterials[MaterialIndex];
 			if (SkeletalMaterial.MaterialSlotName == MaterialSlotName)
 			{
 				return MaterialIndex;
@@ -912,9 +913,10 @@ TArray<FName> USkinnedMeshComponent::GetMaterialSlotNames() const
 	TArray<FName> MaterialNames;
 	if (SkeletalMesh != nullptr)
 	{
-		for (int32 MaterialIndex = 0; MaterialIndex < SkeletalMesh->Materials.Num(); ++MaterialIndex)
+		const TArray<FSkeletalMaterial>& SkeletalMeshMaterials = SkeletalMesh->GetMaterials();
+		for (int32 MaterialIndex = 0; MaterialIndex < SkeletalMeshMaterials.Num(); ++MaterialIndex)
 		{
-			const FSkeletalMaterial &SkeletalMaterial = SkeletalMesh->Materials[MaterialIndex];
+			const FSkeletalMaterial &SkeletalMaterial = SkeletalMeshMaterials[MaterialIndex];
 			MaterialNames.Add(SkeletalMaterial.MaterialSlotName);
 		}
 	}
@@ -961,7 +963,7 @@ void USkinnedMeshComponent::SetCPUSkinningEnabled(bool bEnable, bool bRecreateRe
 
 		SkeletalMesh->WaitForPendingInitOrStreaming();
 
-		check(SkeletalMesh->GetResourceForRendering()->CurrentFirstLODIdx <= SkeletalMesh->MinLod.Default);
+		check(SkeletalMesh->GetResourceForRendering()->CurrentFirstLODIdx <= SkeletalMesh->GetMinLod().Default);
 
 		SkeletalMesh->UnlinkStreaming();
 		SkeletalMesh->bForceMiplevelsToBeResident = bOriginalForcedFullyLoad;
@@ -1028,9 +1030,9 @@ void USkinnedMeshComponent::RebuildVisibilityArray()
 
 		// The following code relies on a complete hierarchy sorted from parent to children
 		TArray<uint8>& EditableBoneVisibilityStates = GetEditableBoneVisibilityStates();
-		if (EditableBoneVisibilityStates.Num() != SkeletalMesh->RefSkeleton.GetNum())
+		if (EditableBoneVisibilityStates.Num() != SkeletalMesh->GetRefSkeleton().GetNum())
 		{
-			UE_LOG(LogSkinnedMeshComp, Warning, TEXT("RebuildVisibilityArray() failed because EditableBoneVisibilityStates size: %d not equal to RefSkeleton bone count: %d."), EditableBoneVisibilityStates.Num(), SkeletalMesh->RefSkeleton.GetNum());
+			UE_LOG(LogSkinnedMeshComp, Warning, TEXT("RebuildVisibilityArray() failed because EditableBoneVisibilityStates size: %d not equal to RefSkeleton bone count: %d."), EditableBoneVisibilityStates.Num(), SkeletalMesh->GetRefSkeleton().GetNum());
 			return;
 		}
 
@@ -1042,7 +1044,7 @@ void USkinnedMeshComponent::RebuildVisibilityArray()
 			if (VisState != BVS_ExplicitlyHidden)
 			{
 				// Check direct parent (only need to do one deep, since we have already processed the parent and written to BoneVisibilityStates previously)
-				const int32 ParentIndex = SkeletalMesh->RefSkeleton.GetParentIndex(BoneId);
+				const int32 ParentIndex = SkeletalMesh->GetRefSkeleton().GetParentIndex(BoneId);
 				if ((ParentIndex == -1) || (EditableBoneVisibilityStates[ParentIndex] == BVS_Visible))
 				{
 					EditableBoneVisibilityStates[BoneId] = BVS_Visible;
@@ -1073,12 +1075,12 @@ class UPhysicsAsset* USkinnedMeshComponent::GetPhysicsAsset() const
 		return PhysicsAssetOverride;
 	}
 
-	if (SkeletalMesh && SkeletalMesh->PhysicsAsset)
+	if (SkeletalMesh && SkeletalMesh->GetPhysicsAsset())
 	{
-		return SkeletalMesh->PhysicsAsset;
+		return SkeletalMesh->GetPhysicsAsset();
 	}
 
-	return NULL;
+	return nullptr;
 }
 
 FBoxSphereBounds USkinnedMeshComponent::CalcMeshBound(const FVector& RootOffset, bool UsePhysicsAsset, const FTransform& LocalToWorld) const
@@ -1096,7 +1098,7 @@ FBoxSphereBounds USkinnedMeshComponent::CalcMeshBound(const FVector& RootOffset,
 	// Can only use the PhysicsAsset to calculate the bounding box if we are not non-uniformly scaling the mesh.
 	const bool bCanUsePhysicsAsset = DrawScale.IsUniform() && (SkeletalMesh != NULL)
 		// either space base exists or child component
-		&& ( (GetNumComponentSpaceTransforms() == SkeletalMesh->RefSkeleton.GetNum()) || (MasterPhysicsAsset) );
+		&& ( (GetNumComponentSpaceTransforms() == SkeletalMesh->GetRefSkeleton().GetNum()) || (MasterPhysicsAsset) );
 
 	const bool bDetailModeAllowsRendering = (DetailMode <= GetCachedScalabilityCVars().DetailMode);
 	const bool bIsVisible = ( bDetailModeAllowsRendering && (ShouldRender() || bCastHiddenShadow));
@@ -1309,7 +1311,7 @@ FTransform USkinnedMeshComponent::GetBoneTransform(int32 BoneIdx, const FTransfo
 
 int32 USkinnedMeshComponent::GetNumBones()const
 {
-	return SkeletalMesh ? SkeletalMesh->RefSkeleton.GetNum() : 0;
+	return SkeletalMesh ? SkeletalMesh->GetRefSkeleton().GetNum() : 0;
 }
 
 int32 USkinnedMeshComponent::GetBoneIndex( FName BoneName) const
@@ -1317,7 +1319,7 @@ int32 USkinnedMeshComponent::GetBoneIndex( FName BoneName) const
 	int32 BoneIndex = INDEX_NONE;
 	if ( BoneName != NAME_None && SkeletalMesh )
 	{
-		BoneIndex = SkeletalMesh->RefSkeleton.FindBoneIndex( BoneName );
+		BoneIndex = SkeletalMesh->GetRefSkeleton().FindBoneIndex( BoneName );
 	}
 
 	return BoneIndex;
@@ -1326,7 +1328,7 @@ int32 USkinnedMeshComponent::GetBoneIndex( FName BoneName) const
 
 FName USkinnedMeshComponent::GetBoneName(int32 BoneIndex) const
 {
-	return (SkeletalMesh != NULL && SkeletalMesh->RefSkeleton.IsValidIndex(BoneIndex)) ? SkeletalMesh->RefSkeleton.GetBoneName(BoneIndex) : NAME_None;
+	return (SkeletalMesh != NULL && SkeletalMesh->GetRefSkeleton().IsValidIndex(BoneIndex)) ? SkeletalMesh->GetRefSkeleton().GetBoneName(BoneIndex) : NAME_None;
 }
 
 
@@ -1337,7 +1339,7 @@ FName USkinnedMeshComponent::GetParentBone( FName BoneName ) const
 	int32 BoneIndex = GetBoneIndex(BoneName);
 	if ((BoneIndex != INDEX_NONE) && (BoneIndex > 0)) // This checks that this bone is not the root (ie no parent), and that BoneIndex != INDEX_NONE (ie bone name was found)
 	{
-		Result = SkeletalMesh->RefSkeleton.GetBoneName(SkeletalMesh->RefSkeleton.GetParentIndex(BoneIndex));
+		Result = SkeletalMesh->GetRefSkeleton().GetBoneName(SkeletalMesh->GetRefSkeleton().GetParentIndex(BoneIndex));
 	}
 	return Result;
 }
@@ -1346,7 +1348,7 @@ FTransform USkinnedMeshComponent::GetDeltaTransformFromRefPose(FName BoneName, F
 {
 	if (SkeletalMesh)
 	{
-		const FReferenceSkeleton& RefSkeleton = SkeletalMesh->RefSkeleton;
+		const FReferenceSkeleton& RefSkeleton = SkeletalMesh->GetRefSkeleton();
 		const int32 BoneIndex = GetBoneIndex(BoneName);
 		if (BoneIndex != INDEX_NONE)
 		{
@@ -1374,7 +1376,7 @@ FTransform USkinnedMeshComponent::GetDeltaTransformFromRefPose(FName BoneName, F
 
 bool USkinnedMeshComponent::GetTwistAndSwingAngleOfDeltaRotationFromRefPose(FName BoneName, float& OutTwistAngle, float& OutSwingAngle) const
 {
-	const FReferenceSkeleton& RefSkeleton = SkeletalMesh->RefSkeleton;
+	const FReferenceSkeleton& RefSkeleton = SkeletalMesh->GetRefSkeleton();
 	const int32 BoneIndex = GetBoneIndex(BoneName);
 	const TArray<FTransform>& Transforms = GetComponentSpaceTransforms();
 
@@ -1460,11 +1462,11 @@ void USkinnedMeshComponent::GetBoneNames(TArray<FName>& BoneNames)
 	else
 	{
 		// pre-size the array to avoid unnecessary reallocation
-		BoneNames.Empty(SkeletalMesh->RefSkeleton.GetNum());
-		BoneNames.AddUninitialized(SkeletalMesh->RefSkeleton.GetNum());
-		for (int32 i = 0; i < SkeletalMesh->RefSkeleton.GetNum(); i++)
+		BoneNames.Empty(SkeletalMesh->GetRefSkeleton().GetNum());
+		BoneNames.AddUninitialized(SkeletalMesh->GetRefSkeleton().GetNum());
+		for (int32 i = 0; i < SkeletalMesh->GetRefSkeleton().GetNum(); i++)
 		{
-			BoneNames[i] = SkeletalMesh->RefSkeleton.GetBoneName(i);
+			BoneNames[i] = SkeletalMesh->GetRefSkeleton().GetBoneName(i);
 		}
 	}
 }
@@ -1476,21 +1478,21 @@ bool USkinnedMeshComponent::BoneIsChildOf(FName BoneName, FName ParentBoneName)
 
 	if( SkeletalMesh )
 	{
-		const int32 BoneIndex = SkeletalMesh->RefSkeleton.FindBoneIndex(BoneName);
+		const int32 BoneIndex = SkeletalMesh->GetRefSkeleton().FindBoneIndex(BoneName);
 		if(BoneIndex == INDEX_NONE)
 		{
 			UE_LOG(LogSkinnedMeshComp, Log, TEXT("execBoneIsChildOf: BoneName '%s' not found in SkeletalMesh '%s'"), *BoneName.ToString(), *SkeletalMesh->GetName());
 			return bResult;
 		}
 
-		const int32 ParentBoneIndex = SkeletalMesh->RefSkeleton.FindBoneIndex(ParentBoneName);
+		const int32 ParentBoneIndex = SkeletalMesh->GetRefSkeleton().FindBoneIndex(ParentBoneName);
 		if(ParentBoneIndex == INDEX_NONE)
 		{
 			UE_LOG(LogSkinnedMeshComp, Log, TEXT("execBoneIsChildOf: ParentBoneName '%s' not found in SkeletalMesh '%s'"), *ParentBoneName.ToString(), *SkeletalMesh->GetName());
 			return bResult;
 		}
 
-		bResult = SkeletalMesh->RefSkeleton.BoneIsChildOf(BoneIndex, ParentBoneIndex);
+		bResult = SkeletalMesh->GetRefSkeleton().BoneIsChildOf(BoneIndex, ParentBoneIndex);
 	}
 
 	return bResult;
@@ -1499,9 +1501,9 @@ bool USkinnedMeshComponent::BoneIsChildOf(FName BoneName, FName ParentBoneName)
 
 FVector USkinnedMeshComponent::GetRefPosePosition(int32 BoneIndex)
 {
-	if(SkeletalMesh && (BoneIndex >= 0) && (BoneIndex < SkeletalMesh->RefSkeleton.GetNum()))
+	if(SkeletalMesh && (BoneIndex >= 0) && (BoneIndex < SkeletalMesh->GetRefSkeleton().GetNum()))
 	{
-		return SkeletalMesh->RefSkeleton.GetRefBonePose()[BoneIndex].GetTranslation();
+		return SkeletalMesh->GetRefSkeleton().GetRefBonePose()[BoneIndex].GetTranslation();
 	}
 	else
 	{
@@ -1586,23 +1588,23 @@ bool USkinnedMeshComponent::AllocateTransformData()
 	// Allocate transforms if not present.
 	if ( SkeletalMesh != NULL && MasterPoseComponent == NULL )
 	{
-		if(GetNumComponentSpaceTransforms() != SkeletalMesh->RefSkeleton.GetNum() )
+		if(GetNumComponentSpaceTransforms() != SkeletalMesh->GetRefSkeleton().GetNum() )
 		{
 			for (int32 BaseIndex = 0; BaseIndex < 2; ++BaseIndex)
 			{
-				ComponentSpaceTransformsArray[BaseIndex].Empty(SkeletalMesh->RefSkeleton.GetNum());
-				ComponentSpaceTransformsArray[BaseIndex].AddUninitialized(SkeletalMesh->RefSkeleton.GetNum());
+				ComponentSpaceTransformsArray[BaseIndex].Empty(SkeletalMesh->GetRefSkeleton().GetNum());
+				ComponentSpaceTransformsArray[BaseIndex].AddUninitialized(SkeletalMesh->GetRefSkeleton().GetNum());
 
-				for (int32 I = 0; I < SkeletalMesh->RefSkeleton.GetNum(); ++I)
+				for (int32 I = 0; I < SkeletalMesh->GetRefSkeleton().GetNum(); ++I)
 				{
 					ComponentSpaceTransformsArray[BaseIndex][I].SetIdentity();
 				}
 
-				BoneVisibilityStates[BaseIndex].Empty( SkeletalMesh->RefSkeleton.GetNum() );
-				if( SkeletalMesh->RefSkeleton.GetNum() )
+				BoneVisibilityStates[BaseIndex].Empty( SkeletalMesh->GetRefSkeleton().GetNum() );
+				if( SkeletalMesh->GetRefSkeleton().GetNum() )
 				{
-					BoneVisibilityStates[BaseIndex].AddUninitialized( SkeletalMesh->RefSkeleton.GetNum() );
-					for (int32 BoneIndex = 0; BoneIndex < SkeletalMesh->RefSkeleton.GetNum(); BoneIndex++)
+					BoneVisibilityStates[BaseIndex].AddUninitialized( SkeletalMesh->GetRefSkeleton().GetNum() );
+					for (int32 BoneIndex = 0; BoneIndex < SkeletalMesh->GetRefSkeleton().GetNum(); BoneIndex++)
 					{
 						BoneVisibilityStates[BaseIndex][BoneIndex] = BVS_Visible;
 					}
@@ -1879,7 +1881,7 @@ UMorphTarget* USkinnedMeshComponent::FindMorphTarget( FName MorphTargetName ) co
 
 bool USkinnedMeshComponent::GetMissingMasterBoneRelativeTransform(int32 InBoneIndex, FMissingMasterBoneCacheEntry& OutInfo) const
 {
-	const FReferenceSkeleton& SlaveRefSkeleton = SkeletalMesh->RefSkeleton;
+	const FReferenceSkeleton& SlaveRefSkeleton = SkeletalMesh->GetRefSkeleton();
 	check(SlaveRefSkeleton.IsValidIndex(InBoneIndex));
 	const TArray<FTransform>& BoneSpaceRefPoseTransforms = SlaveRefSkeleton.GetRefBonePose();
 
@@ -1921,8 +1923,8 @@ void USkinnedMeshComponent::UpdateMasterBoneMap()
 		{
 			if (USkeletalMesh* MasterMesh = MasterPoseComponentPtr->SkeletalMesh)
 			{
-				const FReferenceSkeleton& SlaveRefSkeleton = SkeletalMesh->RefSkeleton;
-				const FReferenceSkeleton& MasterRefSkeleton = MasterMesh->RefSkeleton;
+				const FReferenceSkeleton& SlaveRefSkeleton = SkeletalMesh->GetRefSkeleton();
+				const FReferenceSkeleton& MasterRefSkeleton = MasterMesh->GetRefSkeleton();
 
 				MasterBoneMap.AddUninitialized(SlaveRefSkeleton.GetNum());
 				if (SkeletalMesh == MasterMesh)
@@ -2272,7 +2274,7 @@ bool USkinnedMeshComponent::HasAnySockets() const
 #if WITH_EDITOR
 		(SkeletalMesh->GetActiveSocketList().Num() > 0) ||
 #endif
-		(SkeletalMesh->RefSkeleton.GetNum() > 0));
+		(SkeletalMesh->GetRefSkeleton().GetNum() > 0));
 }
 
 void USkinnedMeshComponent::QuerySupportedSockets(TArray<FComponentSocketDescription>& OutSockets) const
@@ -2291,9 +2293,9 @@ void USkinnedMeshComponent::QuerySupportedSockets(TArray<FComponentSocketDescrip
 		}
 
 		// Now grab the bones, which can behave exactly like sockets
-		for (int32 BoneIdx = 0; BoneIdx < SkeletalMesh->RefSkeleton.GetNum(); ++BoneIdx)
+		for (int32 BoneIdx = 0; BoneIdx < SkeletalMesh->GetRefSkeleton().GetNum(); ++BoneIdx)
 		{
-			const FName BoneName = SkeletalMesh->RefSkeleton.GetBoneName(BoneIdx);
+			const FName BoneName = SkeletalMesh->GetRefSkeleton().GetBoneName(BoneIdx);
 			new (OutSockets) FComponentSocketDescription(BoneName, EComponentSocketType::Bone);
 		}
 	}
@@ -2377,7 +2379,7 @@ FName USkinnedMeshComponent::FindClosestBone(FVector TestLocation, FVector* Bone
 			bool bPassPACheck = !bRequirePhysicsAsset;
 			if (bRequirePhysicsAsset)
 			{
-				FName BoneName = SkeletalMesh->RefSkeleton.GetBoneName(i);
+				FName BoneName = SkeletalMesh->GetRefSkeleton().GetBoneName(i);
 				bPassPACheck = (PhysAsset->BodySetupIndexMap.Find(BoneName) != nullptr);
 			}
 
@@ -2407,7 +2409,7 @@ FName USkinnedMeshComponent::FindClosestBone(FVector TestLocation, FVector* Bone
 			{
 				*BoneLocation = (CompSpaceTransforms[BestIndex] * GetComponentTransform()).GetLocation();
 			}
-			return SkeletalMesh->RefSkeleton.GetBoneName(BestIndex);
+			return SkeletalMesh->GetRefSkeleton().GetBoneName(BestIndex);
 		}
 	}
 }
@@ -2434,11 +2436,11 @@ void USkinnedMeshComponent::ShowMaterialSection(int32 MaterialID, int32 SectionI
 		TArray<bool>& HiddenMaterials = SkelCompLODInfo.HiddenMaterials;
 	
 		// allocate if not allocated yet
-		if ( HiddenMaterials.Num() != SkeletalMesh->Materials.Num() )
+		if ( HiddenMaterials.Num() != SkeletalMesh->GetMaterials().Num() )
 		{
-			// Using skeletalmesh component because Materials.Num() should be <= SkeletalMesh->Materials.Num()		
-			HiddenMaterials.Empty(SkeletalMesh->Materials.Num());
-			HiddenMaterials.AddZeroed(SkeletalMesh->Materials.Num());
+			// Using skeletalmesh component because Materials.Num() should be <= SkeletalMesh->GetMaterials().Num()		
+			HiddenMaterials.Empty(SkeletalMesh->GetMaterials().Num());
+			HiddenMaterials.AddZeroed(SkeletalMesh->GetMaterials().Num());
 		}
 		// If we have a valid LODInfo LODMaterialMap, route material index through it.
 		int32 UseMaterialIndex = MaterialID;			
@@ -2517,7 +2519,7 @@ void USkinnedMeshComponent::GetUsedMaterials( TArray<UMaterialInterface*>& OutMa
 	if( SkeletalMesh )
 	{
 		// The max number of materials used is the max of the materials on the skeletal mesh and the materials on the mesh component
-		const int32 NumMaterials = FMath::Max( SkeletalMesh->Materials.Num(), OverrideMaterials.Num() );
+		const int32 NumMaterials = FMath::Max(SkeletalMesh->GetMaterials().Num(), OverrideMaterials.Num() );
 		for( int32 MatIdx = 0; MatIdx < NumMaterials; ++MatIdx )
 		{
 			// GetMaterial will determine the correct material to use for this index.  
@@ -2605,7 +2607,7 @@ void USkinnedMeshComponent::SetRefPoseOverride(const TArray<FTransform>& NewRefP
 		return;
 	}
 
-	const int32 NumRealBones = SkeletalMesh->RefSkeleton.GetRawBoneNum();
+	const int32 NumRealBones = SkeletalMesh->GetRefSkeleton().GetRawBoneNum();
 
 	if (NumRealBones != NewRefPoseTransforms.Num())
 	{
@@ -2648,7 +2650,7 @@ void USkinnedMeshComponent::SetRefPoseOverride(const TArray<FTransform>& NewRefP
 		// Construct mesh-space skeletal hierarchy.
 		if (BoneIndex > 0)
 		{
-			int32 ParentIndex = SkeletalMesh->RefSkeleton.GetRawParentIndex(BoneIndex);
+			int32 ParentIndex = SkeletalMesh->GetRefSkeleton().GetRawParentIndex(BoneIndex);
 			CachedComposedRefPoseMatrices[BoneIndex] = CachedComposedRefPoseMatrices[BoneIndex] * CachedComposedRefPoseMatrices[ParentIndex];
 		}
 
@@ -2660,7 +2662,7 @@ void USkinnedMeshComponent::SetRefPoseOverride(const TArray<FTransform>& NewRefP
 			ZAxis.IsNearlyZero(SMALL_NUMBER))
 		{
 			// this is not allowed, warn them 
-			UE_LOG(LogSkeletalMesh, Warning, TEXT("Reference Pose for asset %s for joint (%s) includes NIL matrix. Zero scale isn't allowed on ref pose. "), *SkeletalMesh->GetPathName(), *SkeletalMesh->RefSkeleton.GetBoneName(BoneIndex).ToString());
+			UE_LOG(LogSkeletalMesh, Warning, TEXT("Reference Pose for asset %s for joint (%s) includes NIL matrix. Zero scale isn't allowed on ref pose. "), *SkeletalMesh->GetPathName(), *SkeletalMesh->GetRefSkeleton().GetBoneName(BoneIndex).ToString());
 		}
 
 		// Precompute inverse so we can use from-refpose-skin vertices.
@@ -2681,7 +2683,7 @@ void USkinnedMeshComponent::ClearRefPoseOverride()
 void USkinnedMeshComponent::CacheRefToLocalMatrices(TArray<FMatrix>& OutRefToLocal)const
 {
 	const USkinnedMeshComponent* BaseComponent = GetBaseComponent();
-	OutRefToLocal.SetNumUninitialized(SkeletalMesh->RefBasesInvMatrix.Num());
+	OutRefToLocal.SetNumUninitialized(SkeletalMesh->GetRefBasesInvMatrix().Num());
 	const TArray<FTransform>& CompSpaceTransforms = BaseComponent->GetComponentSpaceTransforms();
 	if(CompSpaceTransforms.Num())
 	{
@@ -2689,7 +2691,7 @@ void USkinnedMeshComponent::CacheRefToLocalMatrices(TArray<FMatrix>& OutRefToLoc
 
 		for (int32 MatrixIdx = 0; MatrixIdx < OutRefToLocal.Num(); ++MatrixIdx)
 		{
-			OutRefToLocal[MatrixIdx] = SkeletalMesh->RefBasesInvMatrix[MatrixIdx] * CompSpaceTransforms[MatrixIdx].ToMatrixWithScale();
+			OutRefToLocal[MatrixIdx] = SkeletalMesh->GetRefBasesInvMatrix()[MatrixIdx] * CompSpaceTransforms[MatrixIdx].ToMatrixWithScale();
 		}
 	}
 	else
@@ -2697,7 +2699,7 @@ void USkinnedMeshComponent::CacheRefToLocalMatrices(TArray<FMatrix>& OutRefToLoc
 		//Possible in some cases to request this before the component space transforms are prepared (undo/redo)
 		for (int32 MatrixIdx = 0; MatrixIdx < OutRefToLocal.Num(); ++MatrixIdx)
 		{
-			OutRefToLocal[MatrixIdx] = SkeletalMesh->RefBasesInvMatrix[MatrixIdx];
+			OutRefToLocal[MatrixIdx] = SkeletalMesh->GetRefBasesInvMatrix()[MatrixIdx];
 		}
 	}
 }
@@ -2915,7 +2917,7 @@ void USkinnedMeshComponent::SetMinLOD(int32 InNewMinLOD)
 
 int32 USkinnedMeshComponent::ComputeMinLOD() const
 {
-	int32 MinLodIndex = bOverrideMinLod ? MinLodModel : SkeletalMesh->MinLod.GetValue();
+	int32 MinLodIndex = bOverrideMinLod ? MinLodModel : SkeletalMesh->GetMinLod().GetValue();
 	int32 NumLODs = GetNumLODs();
 	// want to make sure MinLOD stays within the valid range
 	MinLodIndex = FMath::Min(MinLodIndex, NumLODs - 1);
@@ -3608,7 +3610,7 @@ void USkinnedMeshComponent::SetSkinWeightOverride(int32 LODIndex, const TArray<F
 			Info.OverrideSkinWeights->SetMaxBoneInfluences(NumBoneInfluences);
 			Info.OverrideSkinWeights->SetUse16BitBoneIndex(bUse16BitBoneIndex);
 
-			const FReferenceSkeleton& RefSkel = SkeletalMesh->RefSkeleton;
+			const FReferenceSkeleton& RefSkel = SkeletalMesh->GetRefSkeleton();
 			TArray<FSkinWeightInfo> GPUWeights;
 			CreateSkinWeightsArray(SkinWeights, LODData, GPUWeights, RefSkel);
 			*(Info.OverrideSkinWeights) = GPUWeights;
@@ -4145,7 +4147,7 @@ FVector GetTypedSkinnedVertexPosition(
 		if (MasterPoseComponentInst)
 		{
 			const TArray<int32>& MasterBoneMap = SkinnedComp->GetMasterBoneMap();
-			check(MasterBoneMap.Num() == SkinnedComp->SkeletalMesh->RefSkeleton.GetNum());
+			check(MasterBoneMap.Num() == SkinnedComp->SkeletalMesh->GetRefSkeleton().GetNum());
 			TransformBoneIndex = MasterBoneMap[MeshBoneIndex];
 		}
 
@@ -4159,7 +4161,7 @@ FVector GetTypedSkinnedVertexPosition(
 			else
 			{
 				const FMatrix BoneTransformMatrix = (TransformBoneIndex != INDEX_NONE) ? BaseComponent->GetComponentSpaceTransforms()[TransformBoneIndex].ToMatrixWithScale() : FMatrix::Identity;
-				const FMatrix RefToLocal = SkinnedComp->SkeletalMesh->RefBasesInvMatrix[MeshBoneIndex] * BoneTransformMatrix;
+				const FMatrix RefToLocal = SkinnedComp->SkeletalMesh->GetRefBasesInvMatrix()[MeshBoneIndex] * BoneTransformMatrix;
 				SkinnedPos += RefToLocal.TransformPosition(PositionVertexBuffer.VertexPosition(BufferVertIndex)) * Weight;
 			}
 		}
diff --git a/Engine/Source/Runtime/Engine/Private/Particles/ParticleModules_Location.cpp b/Engine/Source/Runtime/Engine/Private/Particles/ParticleModules_Location.cpp
--- a/Engine/Source/Runtime/Engine/Private/Particles/ParticleModules_Location.cpp
+++ b/Engine/Source/Runtime/Engine/Private/Particles/ParticleModules_Location.cpp
@@ -1870,14 +1870,14 @@ bool UParticleModuleLocationBoneSocket::PerformCustomMenuEntry(int32 InEntryInde
 				else //BONESOCKETSOURCE_Bones
 				{
 					// Retrieve all the bones
-					if (EditorSkelMesh->RefSkeleton.GetNum() > 0)
+					if (EditorSkelMesh->GetRefSkeleton().GetNum() > 0)
 					{
 						SourceLocations.Empty();
-						for (int32 BoneIdx = 0; BoneIdx < EditorSkelMesh->RefSkeleton.GetNum(); BoneIdx++)
+						for (int32 BoneIdx = 0; BoneIdx < EditorSkelMesh->GetRefSkeleton().GetNum(); BoneIdx++)
 						{
 							int32 NewItemIdx = SourceLocations.AddZeroed();
 							FLocationBoneSocketInfo& Info = SourceLocations[NewItemIdx];
-							Info.BoneSocketName = EditorSkelMesh->RefSkeleton.GetBoneName(BoneIdx);
+							Info.BoneSocketName = EditorSkelMesh->GetRefSkeleton().GetBoneName(BoneIdx);
 						}
 						return true;
 					}
@@ -2583,7 +2583,7 @@ void UParticleModuleLocationSkelVertSurface::UpdateBoneIndicesList(FParticleEmit
 		int32 InsertionIndex = 0;
 		for (int32 FindBoneIdx = 0; FindBoneIdx < ValidAssociatedBones.Num(); FindBoneIdx++)
 		{
-			const int32 BoneIdx = SkelMeshComp->SkeletalMesh->RefSkeleton.FindBoneIndex(ValidAssociatedBones[FindBoneIdx]);
+			const int32 BoneIdx = SkelMeshComp->SkeletalMesh->GetRefSkeleton().FindBoneIndex(ValidAssociatedBones[FindBoneIdx]);
 			if (BoneIdx != INDEX_NONE && ValidAssociatedBones.Num() > InsertionIndex)
 			{
 				InstancePayload->ValidAssociatedBoneIndices[InsertionIndex++] = BoneIdx;
@@ -2668,13 +2668,13 @@ bool UParticleModuleLocationSkelVertSurface::PerformCustomMenuEntry(int32 InEntr
 			if (EditorSkelMesh != NULL)
 			{
 				// Retrieve all the bones
-				if (EditorSkelMesh->RefSkeleton.GetRawBoneNum() > 0)
+				if (EditorSkelMesh->GetRefSkeleton().GetRawBoneNum() > 0)
 				{
 					ValidAssociatedBones.Empty();
-					for (int32 BoneIdx = 0; BoneIdx < EditorSkelMesh->RefSkeleton.GetRawBoneNum(); BoneIdx++)
+					for (int32 BoneIdx = 0; BoneIdx < EditorSkelMesh->GetRefSkeleton().GetRawBoneNum(); BoneIdx++)
 					{
 						int32 NewItemIdx = ValidAssociatedBones.AddZeroed();
-						ValidAssociatedBones[NewItemIdx] = EditorSkelMesh->RefSkeleton.GetBoneName(BoneIdx);
+						ValidAssociatedBones[NewItemIdx] = EditorSkelMesh->GetRefSkeleton().GetBoneName(BoneIdx);
 					}
 				}
 				else
@@ -2749,7 +2749,7 @@ void UParticleModuleLocationSkelVertSurface::GetSkeletalMeshComponentSource(FPar
 		int32 MinLOD = INDEX_NONE;
 		if (NewSkelMeshComp && NewSkelMeshComp->GetScene() && NewSkelMeshComp->SkeletalMesh)
 		{
-			MinLOD = NewSkelMeshComp->SkeletalMesh->MinLod.GetValue();
+			MinLOD = NewSkelMeshComp->SkeletalMesh->GetMinLod().GetValue();
 			FSkeletalMeshRenderData* SkelMeshResource = NewSkelMeshComp->GetSkeletalMeshRenderData();
 			FSkinWeightVertexBuffer* SkinWeightBuffer = NewSkelMeshComp->GetSkinWeightBuffer(MinLOD);
 			FSkeletalMeshLODRenderData& LODData = SkelMeshResource->LODRenderData[MinLOD];
@@ -2960,7 +2960,7 @@ bool UParticleModuleLocationSkelVertSurface::VertInfluencedByActiveBoneTyped(
 			int32 BoneIndex = Section.BoneMap[WeightBuffer->GetBoneIndex(Section.GetVertexBufferIndex() + VertIndex, InfluenceIndex)];
 			if (InSkelMeshComponent->MasterPoseComponent.IsValid())
 			{
-				check(MasterBoneMap.Num() == InSkelMeshComponent->SkeletalMesh->RefSkeleton.GetNum());
+				check(MasterBoneMap.Num() == InSkelMeshComponent->SkeletalMesh->GetRefSkeleton().GetNum());
 				BoneIndex = MasterBoneMap[BoneIndex];
 			}
 
diff --git a/Engine/Source/Runtime/Engine/Private/PhysicsEngine/PhysAnim.cpp b/Engine/Source/Runtime/Engine/Private/PhysicsEngine/PhysAnim.cpp
--- a/Engine/Source/Runtime/Engine/Private/PhysicsEngine/PhysAnim.cpp
+++ b/Engine/Source/Runtime/Engine/Private/PhysicsEngine/PhysAnim.cpp
@@ -135,7 +135,7 @@ void UpdateWorldBoneTM(TAssetWorldBoneTMArray& WorldBoneTMs, const TArray<FTrans
 	else
 	{
 		// If not root, use our cached world-space bone transforms.
-		int32 ParentIndex = SkelComp->SkeletalMesh->RefSkeleton.GetParentIndex(BoneIndex);
+		int32 ParentIndex = SkelComp->SkeletalMesh->GetRefSkeleton().GetParentIndex(BoneIndex);
 		UpdateWorldBoneTM(WorldBoneTMs, InBoneSpaceTransforms, ParentIndex, SkelComp, LocalToWorldTM, Scale3D);
 		ParentTM = WorldBoneTMs[ParentIndex].TM;
 	}
@@ -204,7 +204,7 @@ void USkeletalMeshComponent::PerformBlendPhysicsBones(const TArray<FBoneIndexTyp
 			int32 BoneIndex = InRequiredBones[i];
 
 			// See if this is a physics bone..
-			int32 BodyIndex = PhysicsAsset->FindBodyIndex(SkeletalMesh->RefSkeleton.GetBoneName(BoneIndex));
+			int32 BodyIndex = PhysicsAsset->FindBodyIndex(SkeletalMesh->GetRefSkeleton().GetBoneName(BoneIndex));
 			// need to update back to physX so that physX knows where it was after blending
 			FBodyInstance* PhysicsAssetBodyInstance = nullptr;
 
@@ -253,7 +253,7 @@ void USkeletalMeshComponent::PerformBlendPhysicsBones(const TArray<FBoneIndexTyp
 						else
 						{
 							// If not root, get parent TM from cache (making sure its up-to-date).
-							int32 ParentIndex = SkeletalMesh->RefSkeleton.GetParentIndex(BoneIndex);
+							int32 ParentIndex = SkeletalMesh->GetRefSkeleton().GetParentIndex(BoneIndex);
 							UpdateWorldBoneTM(WorldBoneTMs, InOutBoneSpaceTransforms, ParentIndex, this, LocalToWorldTM, TotalScale3D);
 							ParentWorldTM = WorldBoneTMs[ParentIndex].TM;
 						}
@@ -303,7 +303,7 @@ void USkeletalMeshComponent::PerformBlendPhysicsBones(const TArray<FBoneIndexTyp
 					{
 						continue;
 					}
-					const int32 ParentIndex = SkeletalMesh->RefSkeleton.GetParentIndex(BoneIndex);
+					const int32 ParentIndex = SkeletalMesh->GetRefSkeleton().GetParentIndex(BoneIndex);
 					InOutComponentSpaceTransforms[BoneIndex] = InOutBoneSpaceTransforms[BoneIndex] * InOutComponentSpaceTransforms[ParentIndex];
 
 					/**
@@ -524,13 +524,13 @@ void USkeletalMeshComponent::UpdateKinematicBonesToAnim(const TArray<FTransform>
 
 #if !(UE_BUILD_SHIPPING || UE_BUILD_TEST)
 	// If desired, draw the skeleton at the point where we pass it to the physics.
-	if (bShowPrePhysBones && SkeletalMesh && InSpaceBases.Num() == SkeletalMesh->RefSkeleton.GetNum())
+	if (bShowPrePhysBones && SkeletalMesh && InSpaceBases.Num() == SkeletalMesh->GetRefSkeleton().GetNum())
 	{
 		for (int32 i = 1; i<InSpaceBases.Num(); i++)
 		{
 			FVector ThisPos = CurrentLocalToWorld.TransformPosition(InSpaceBases[i].GetLocation());
 
-			int32 ParentIndex = SkeletalMesh->RefSkeleton.GetParentIndex(i);
+			int32 ParentIndex = SkeletalMesh->GetRefSkeleton().GetParentIndex(i);
 			FVector ParentPos = CurrentLocalToWorld.TransformPosition(InSpaceBases[ParentIndex].GetLocation());
 
 			World->LineBatcher->DrawLine(ThisPos, ParentPos, AnimSkelDrawColor, SDPG_Foreground);
@@ -737,7 +737,7 @@ void USkeletalMeshComponent::UpdateRBJointMotors()
 			FConstraintInstance* CI = Constraints[i];
 
 			FName JointName = CS->DefaultInstance.JointName;
-			int32 BoneIndex = SkeletalMesh->RefSkeleton.FindBoneIndex(JointName);
+			int32 BoneIndex = SkeletalMesh->GetRefSkeleton().FindBoneIndex(JointName);
 
 			// If we found this bone, and a visible bone that is not the root, and its joint is motorised in some way..
 			if( (BoneIndex != INDEX_NONE) && (BoneIndex != 0) &&
@@ -756,8 +756,8 @@ void USkeletalMeshComponent::UpdateRBJointMotors()
 				// We need this to compensate for welding, where graphics and physics parents may not be the same.
 				FMatrix ControlBodyToParentBoneTM = FMatrix::Identity;
 
-				int32 TestBoneIndex = SkeletalMesh->RefSkeleton.GetParentIndex(BoneIndex); // This give the 'graphics' parent of this bone
-				bool bFoundControlBody = (SkeletalMesh->RefSkeleton.GetBoneName(TestBoneIndex) == CS->DefaultInstance.ConstraintBone2); // ConstraintBone2 is the 'physics' parent of this joint.
+				int32 TestBoneIndex = SkeletalMesh->GetRefSkeleton().GetParentIndex(BoneIndex); // This give the 'graphics' parent of this bone
+				bool bFoundControlBody = (SkeletalMesh->GetRefSkeleton().GetBoneName(TestBoneIndex) == CS->DefaultInstance.ConstraintBone2); // ConstraintBone2 is the 'physics' parent of this joint.
 
 				while(!bFoundControlBody)
 				{
@@ -779,7 +779,7 @@ void USkeletalMeshComponent::UpdateRBJointMotors()
 					ControlBodyToParentBoneTM = ControlBodyToParentBoneTM * RelTM;
 
 					// Move on to parent
-					TestBoneIndex = SkeletalMesh->RefSkeleton.GetParentIndex(TestBoneIndex);
+					TestBoneIndex = SkeletalMesh->GetRefSkeleton().GetParentIndex(TestBoneIndex);
 
 					// If we are at the root - bail out.
 					if(TestBoneIndex == 0)
@@ -788,7 +788,7 @@ void USkeletalMeshComponent::UpdateRBJointMotors()
 					}
 
 					// See if this is the controlling body
-					bFoundControlBody = (SkeletalMesh->RefSkeleton.GetBoneName(TestBoneIndex) == CS->DefaultInstance.ConstraintBone2);
+					bFoundControlBody = (SkeletalMesh->GetRefSkeleton().GetBoneName(TestBoneIndex) == CS->DefaultInstance.ConstraintBone2);
 				}
 
 				// If after that we didn't find a parent body, we can' do this, so skip.
diff --git a/Engine/Source/Runtime/Engine/Private/PhysicsEngine/PhysDrawing.cpp b/Engine/Source/Runtime/Engine/Private/PhysicsEngine/PhysDrawing.cpp
--- a/Engine/Source/Runtime/Engine/Private/PhysicsEngine/PhysDrawing.cpp
+++ b/Engine/Source/Runtime/Engine/Private/PhysicsEngine/PhysDrawing.cpp
@@ -1062,7 +1062,7 @@ void UPhysicsAsset::GetCollisionMesh(int32 ViewIndex, FMeshElementCollector& Col
 {
 	for( int32 i=0; i<SkeletalBodySetups.Num(); i++)
 	{
-		int32 BoneIndex = SkelMesh->RefSkeleton.FindBoneIndex( SkeletalBodySetups[i]->BoneName );
+		int32 BoneIndex = SkelMesh->GetRefSkeleton().FindBoneIndex( SkeletalBodySetups[i]->BoneName );
 		
 		FColor* BoneColor = (FColor*)( &SkeletalBodySetups[i] );
 
@@ -1093,7 +1093,7 @@ void UPhysicsAsset::DrawConstraints(int32 ViewIndex, FMeshElementCollector& Coll
 
 		// Get each constraint frame in world space.
 		FTransform Con1Frame = FTransform::Identity;
-		int32 Bone1Index = SkelMesh->RefSkeleton.FindBoneIndex(Instance.ConstraintBone1);
+		int32 Bone1Index = SkelMesh->GetRefSkeleton().FindBoneIndex(Instance.ConstraintBone1);
 		if (Bone1Index != INDEX_NONE)
 		{
 			FTransform Body1TM = GetSkelBoneTransform(Bone1Index, SpaceBases, LocalToWorld);
@@ -1102,7 +1102,7 @@ void UPhysicsAsset::DrawConstraints(int32 ViewIndex, FMeshElementCollector& Coll
 		}
 
 		FTransform Con2Frame = FTransform::Identity;
-		int32 Bone2Index = SkelMesh->RefSkeleton.FindBoneIndex(Instance.ConstraintBone2);
+		int32 Bone2Index = SkelMesh->GetRefSkeleton().FindBoneIndex(Instance.ConstraintBone2);
 		if (Bone2Index != INDEX_NONE)
 		{
 			FTransform Body2TM = GetSkelBoneTransform(Bone2Index, SpaceBases, LocalToWorld);
diff --git a/Engine/Source/Runtime/Engine/Private/PhysicsEngine/PhysicalAnimationComponent.cpp b/Engine/Source/Runtime/Engine/Private/PhysicsEngine/PhysicalAnimationComponent.cpp
--- a/Engine/Source/Runtime/Engine/Private/PhysicsEngine/PhysicalAnimationComponent.cpp
+++ b/Engine/Source/Runtime/Engine/Private/PhysicsEngine/PhysicalAnimationComponent.cpp
@@ -229,7 +229,7 @@ FTransform ComputeWorldSpaceTargetTM(const USkeletalMeshComponent& SkeletalMeshC
 
 FTransform ComputeLocalSpaceTargetTM(const USkeletalMeshComponent& SkeletalMeshComponent, const UPhysicsAsset& PhysAsset, const TArray<FTransform>& LocalTransforms, int32 BoneIndex)
 {
-	const FReferenceSkeleton& RefSkeleton = SkeletalMeshComponent.SkeletalMesh->RefSkeleton;
+	const FReferenceSkeleton& RefSkeleton = SkeletalMeshComponent.SkeletalMesh->GetRefSkeleton();
 	FTransform AccumulatedDelta = LocalTransforms[BoneIndex];
 	int32 CurBoneIdx = BoneIndex;
 	while ((CurBoneIdx = RefSkeleton.GetParentIndex(CurBoneIdx)) != INDEX_NONE)
@@ -273,7 +273,7 @@ void UPhysicalAnimationComponent::UpdateTargetActors(ETeleportType TeleportType)
 	UPhysicsAsset* PhysAsset = SkeletalMeshComponent ? SkeletalMeshComponent->GetPhysicsAsset() : nullptr;
 	if (PhysAsset && SkeletalMeshComponent->SkeletalMesh)
 	{
-		const FReferenceSkeleton& RefSkeleton = SkeletalMeshComponent->SkeletalMesh->RefSkeleton;
+		const FReferenceSkeleton& RefSkeleton = SkeletalMeshComponent->SkeletalMesh->GetRefSkeleton();
 
 		// Note we use GetEditableComponentSpaceTransforms because we need to update target actors in the midst of the 
 		// various anim ticks, before buffers are flipped (which happens in the skel mesh component's post-physics tick)
@@ -377,7 +377,7 @@ void UPhysicalAnimationComponent::UpdatePhysicsEngineImp()
 		// Note we use GetEditableComponentSpaceTransforms because we need to update target actors in the midst of the 
 		// various anim ticks, before buffers are flipped (which happens in the skel mesh component's post-physics tick)
 		const TArray<FTransform>& SpaceBases = SkeletalMeshComponent->GetEditableComponentSpaceTransforms();
-		const FReferenceSkeleton& RefSkeleton = SkeletalMeshComponent->SkeletalMesh->RefSkeleton;
+		const FReferenceSkeleton& RefSkeleton = SkeletalMeshComponent->SkeletalMesh->GetRefSkeleton();
 
 #if WITH_PHYSX
 		FPhysicsCommand::ExecuteWrite(SkeletalMeshComponent, [&]()
diff --git a/Engine/Source/Runtime/Engine/Private/PhysicsEngine/PhysicsAsset.cpp b/Engine/Source/Runtime/Engine/Private/PhysicsEngine/PhysicsAsset.cpp
--- a/Engine/Source/Runtime/Engine/Private/PhysicsEngine/PhysicsAsset.cpp
+++ b/Engine/Source/Runtime/Engine/Private/PhysicsEngine/PhysicsAsset.cpp
@@ -432,13 +432,13 @@ int32	UPhysicsAsset::FindControllingBodyIndex(class USkeletalMesh* skelMesh, int
 	int32 BoneIndex = StartBoneIndex;
 	while(BoneIndex!=INDEX_NONE)
 	{
-		FName BoneName = skelMesh->RefSkeleton.GetBoneName(BoneIndex);
+		FName BoneName = skelMesh->GetRefSkeleton().GetBoneName(BoneIndex);
 		int32 BodyIndex = FindBodyIndex(BoneName);
 
 		if(BodyIndex != INDEX_NONE)
 			return BodyIndex;
 
-		int32 ParentBoneIndex = skelMesh->RefSkeleton.GetParentIndex(BoneIndex);
+		int32 ParentBoneIndex = skelMesh->GetRefSkeleton().GetParentIndex(BoneIndex);
 
 		if(ParentBoneIndex == BoneIndex)
 			return INDEX_NONE;
@@ -452,9 +452,9 @@ int32	UPhysicsAsset::FindControllingBodyIndex(class USkeletalMesh* skelMesh, int
 int32 UPhysicsAsset::FindParentBodyIndex(class USkeletalMesh * skelMesh, int32 StartBoneIndex) const
 {
 	int32 BoneIndex = StartBoneIndex;
-	while ((BoneIndex = skelMesh->RefSkeleton.GetParentIndex(BoneIndex)) != INDEX_NONE)
+	while ((BoneIndex = skelMesh->GetRefSkeleton().GetParentIndex(BoneIndex)) != INDEX_NONE)
 	{
-		FName BoneName = skelMesh->RefSkeleton.GetBoneName(BoneIndex);
+		FName BoneName = skelMesh->GetRefSkeleton().GetBoneName(BoneIndex);
 		int32 BodyIndex = FindBodyIndex(BoneName);
 
 		if (StartBoneIndex == BoneIndex)
@@ -509,7 +509,7 @@ int32 UPhysicsAsset::FindMirroredBone(class USkeletalMesh* skelMesh, int32 BoneI
 		return INDEX_NONE;
 	}
 	//we try to find the mirroed bone using several approaches. The first is to look for the same name but with _R instead of _L or vise versa
-	FName BoneName = skelMesh->RefSkeleton.GetBoneName(BoneIndex);
+	FName BoneName = skelMesh->GetRefSkeleton().GetBoneName(BoneIndex);
 	FString BoneNameString = BoneName.ToString();
 
 	bool bIsLeft = BoneNameString.Find("_L", ESearchCase::IgnoreCase, ESearchDir::FromEnd) == (BoneNameString.Len() - 2);	//has _L at the end
@@ -523,7 +523,7 @@ int32 UPhysicsAsset::FindMirroredBone(class USkeletalMesh* skelMesh, int32 BoneI
 		FString BoneNameMirrored = BoneNameString.LeftChop(2);
 		BoneNameMirrored.Append(bIsLeft ? "_R" : "_L");
 
-		BoneIndexMirrored = skelMesh->RefSkeleton.FindBoneIndex(FName(*BoneNameMirrored));
+		BoneIndexMirrored = skelMesh->GetRefSkeleton().FindBoneIndex(FName(*BoneNameMirrored));
 	}
 
 	return BoneIndexMirrored;
@@ -531,7 +531,7 @@ int32 UPhysicsAsset::FindMirroredBone(class USkeletalMesh* skelMesh, int32 BoneI
 
 void UPhysicsAsset::GetBodyIndicesBelow(TArray<int32>& OutBodyIndices, FName InBoneName, USkeletalMesh* SkelMesh, bool bIncludeParent /*= true*/)
 {
-	int32 BaseIndex = SkelMesh->RefSkeleton.FindBoneIndex(InBoneName);
+	int32 BaseIndex = SkelMesh->GetRefSkeleton().FindBoneIndex(InBoneName);
 
 	// Iterate over all other bodies, looking for 'children' of this one
 	for(int32 i=0; i<SkeletalBodySetups.Num(); i++)
@@ -542,9 +542,9 @@ void UPhysicsAsset::GetBodyIndicesBelow(TArray<int32>& OutBodyIndices, FName InB
 			continue;
 		}
 		FName TestName = BS->BoneName;
-		int32 TestIndex = SkelMesh->RefSkeleton.FindBoneIndex(TestName);
+		int32 TestIndex = SkelMesh->GetRefSkeleton().FindBoneIndex(TestName);
 
-		if( (bIncludeParent && TestIndex == BaseIndex) || SkelMesh->RefSkeleton.BoneIsChildOf(TestIndex, BaseIndex))
+		if( (bIncludeParent && TestIndex == BaseIndex) || SkelMesh->GetRefSkeleton().BoneIsChildOf(TestIndex, BaseIndex))
 		{
 			OutBodyIndices.Add(i);
 		}
@@ -853,7 +853,7 @@ void UPhysicsAsset::SetPreviewMesh(USkeletalMesh* PreviewMesh, bool bMarkAsDirty
 		for (int32 i = 0; i < SkeletalBodySetups.Num(); ++i)
 		{
 			FName BodyName = SkeletalBodySetups[i]->BoneName;
-			int32 BoneIndex = PreviewMesh->RefSkeleton.FindBoneIndex(BodyName);
+			int32 BoneIndex = PreviewMesh->GetRefSkeleton().FindBoneIndex(BodyName);
 			if (BoneIndex == INDEX_NONE)
 			{
 				FMessageDialog::Open(EAppMsgType::Ok,
diff --git a/Engine/Source/Runtime/Engine/Private/Rendering/SkeletalMeshLODImporterData.cpp b/Engine/Source/Runtime/Engine/Private/Rendering/SkeletalMeshLODImporterData.cpp
--- a/Engine/Source/Runtime/Engine/Private/Rendering/SkeletalMeshLODImporterData.cpp
+++ b/Engine/Source/Runtime/Engine/Private/Rendering/SkeletalMeshLODImporterData.cpp
@@ -160,12 +160,13 @@ bool FSkeletalMeshImportData::ReplaceSkeletalMeshGeometryImportData(const USkele
 	//Material is a special case since we cannot serialize the UMaterialInstance when saving the RawSkeletalMeshBulkData
 	//So it has to be reconstructed.
 	ImportData->MaxMaterialIndex = 0;
-	for (int32 MaterialIndex = 0; MaterialIndex < SkeletalMesh->Materials.Num(); ++MaterialIndex)
+	const TArray<FSkeletalMaterial>& SkeletalMeshMaterials = SkeletalMesh->GetMaterials();
+	for (int32 MaterialIndex = 0; MaterialIndex < SkeletalMeshMaterials.Num(); ++MaterialIndex)
 	{
 		SkeletalMeshImportData::FMaterial NewMaterial;
 
-		NewMaterial.MaterialImportName = SkeletalMesh->Materials[MaterialIndex].ImportedMaterialSlotName.ToString();
-		NewMaterial.Material = SkeletalMesh->Materials[MaterialIndex].MaterialInterface;
+		NewMaterial.MaterialImportName = SkeletalMeshMaterials[MaterialIndex].ImportedMaterialSlotName.ToString();
+		NewMaterial.Material = SkeletalMeshMaterials[MaterialIndex].MaterialInterface;
 		// Add an entry for each unique material
 		ImportData->MaxMaterialIndex = FMath::Max(ImportData->MaxMaterialIndex, (uint32)(ImportData->Materials.Add(NewMaterial)));
 	}
