To protect the skeletalmesh when using async operations like build or re-import we deprecate all public member variables of USkeletalMesh class and create the necessary getter and setter. Those member

diff --git a/Engine/Source/Runtime/Engine/Private/SkeletalMesh.cpp b/Engine/Source/Runtime/Engine/Private/SkeletalMesh.cpp
--- a/Engine/Source/Runtime/Engine/Private/SkeletalMesh.cpp
+++ b/Engine/Source/Runtime/Engine/Private/SkeletalMesh.cpp
@@ -415,18 +415,18 @@ void FScopedSkeletalMeshPostEditChange::SetSkeletalMesh(USkeletalMesh* InSkeleta
 USkeletalMesh::USkeletalMesh(const FObjectInitializer& ObjectInitializer)
 	: Super(ObjectInitializer)
 {
-	SkelMirrorAxis = EAxis::X;
-	SkelMirrorFlipAxis = EAxis::Z;
+	SetSkelMirrorAxis(EAxis::X);
+	SetSkelMirrorFlipAxis(EAxis::Z);
 #if WITH_EDITORONLY_DATA
 	ImportedModel = MakeShareable(new FSkeletalMeshModel());
-	VertexColorGuid = FGuid();
-	bSupportLODStreaming.Default = false;
-	MaxNumStreamedLODs.Default = 0;
+	SetVertexColorGuid(FGuid());
+	SetSupportLODStreaming(FPerPlatformBool(false));
+	SetMaxNumStreamedLODs(FPerPlatformInt(0));
 	// TODO: support saving some but not all optional LODs
-	MaxNumOptionalLODs.Default = 0;
+	SetMaxNumOptionalLODs(FPerPlatformInt(0));
 #endif
-	MinLod.Default = 0;
-	DisableBelowMinLodStripping.Default = false;
+	SetMinLod(FPerPlatformInt(0));
+	SetDisableBelowMinLodStripping(FPerPlatformBool(false));
 }
 
 USkeletalMesh::USkeletalMesh(FVTableHelper& Helper)
@@ -441,7 +441,7 @@ void USkeletalMesh::PostInitProperties()
 #if WITH_EDITORONLY_DATA
 	if (!HasAnyFlags(RF_ClassDefaultObject))
 	{
-		AssetImportData = NewObject<UAssetImportData>(this, TEXT("AssetImportData"));
+		SetAssetImportData(NewObject<UAssetImportData>(this, TEXT("AssetImportData")));
 	}
 #endif
 	Super::PostInitProperties();
@@ -570,7 +570,7 @@ void USkeletalMesh::AddClothingAsset(UClothingAssetBase* InNewAsset)
 	if(InNewAsset && InNewAsset->GetOuter() == this)
 	{
 		// Ok this should be a correctly created asset, we can add it
-		MeshClothingAssets.AddUnique(InNewAsset);
+		GetMeshClothingAssets().AddUnique(InNewAsset);
 
 		// Consolidate the shared cloth configs
 		InNewAsset->PostUpdateAllAssets();
@@ -587,7 +587,7 @@ void USkeletalMesh::RemoveClothingAsset(int32 InLodIndex, int32 InSectionIndex)
 	if(UClothingAssetBase* Asset = GetSectionClothingAsset(InLodIndex, InSectionIndex))
 	{
 		Asset->UnbindFromSkeletalMesh(this, InLodIndex);
-		MeshClothingAssets.Remove(Asset);
+		GetMeshClothingAssets().Remove(Asset);
 		OnClothingChange.Broadcast();
 	}
 }
@@ -608,7 +608,7 @@ UClothingAssetBase* USkeletalMesh::GetSectionClothingAsset(int32 InLodIndex, int
 
 				if(ClothingAssetGuid.IsValid())
 				{
-					UClothingAssetBase** FoundAsset = MeshClothingAssets.FindByPredicate([&](UClothingAssetBase* InAsset)
+					UClothingAssetBase** FoundAsset = GetMeshClothingAssets().FindByPredicate([&](UClothingAssetBase* InAsset)
 					{
 						return InAsset && InAsset->GetAssetGuid() == ClothingAssetGuid;
 					});
@@ -637,7 +637,7 @@ const UClothingAssetBase* USkeletalMesh::GetSectionClothingAsset(int32 InLodInde
 
 				if (ClothingAssetGuid.IsValid())
 				{
-					UClothingAssetBase* const* FoundAsset = MeshClothingAssets.FindByPredicate([&](UClothingAssetBase* InAsset)
+					UClothingAssetBase* const* FoundAsset = GetMeshClothingAssets().FindByPredicate([&](UClothingAssetBase* InAsset)
 					{
 						return InAsset && InAsset->GetAssetGuid() == ClothingAssetGuid;
 					});
@@ -658,7 +658,7 @@ UClothingAssetBase* USkeletalMesh::GetClothingAsset(const FGuid& InAssetGuid) co
 		return nullptr;
 	}
 
-	UClothingAssetBase* const* FoundAsset = MeshClothingAssets.FindByPredicate([&](UClothingAssetBase* CurrAsset)
+	UClothingAssetBase* const* FoundAsset = GetMeshClothingAssets().FindByPredicate([&](UClothingAssetBase* CurrAsset)
 	{
 		return CurrAsset && CurrAsset->GetAssetGuid() == InAssetGuid;
 	});
@@ -673,11 +673,12 @@ int32 USkeletalMesh::GetClothingAssetIndex(UClothingAssetBase* InAsset) const
 
 int32 USkeletalMesh::GetClothingAssetIndex(const FGuid& InAssetGuid) const
 {
-	const int32 NumAssets = MeshClothingAssets.Num();
+	const TArray<UClothingAssetBase*>& CachedMeshClothingAssets = GetMeshClothingAssets();
+	const int32 NumAssets = CachedMeshClothingAssets.Num();
 	for(int32 SearchIndex = 0; SearchIndex < NumAssets; ++SearchIndex)
 	{
-		if(MeshClothingAssets[SearchIndex] && 
-		   MeshClothingAssets[SearchIndex]->GetAssetGuid() == InAssetGuid)
+		if(CachedMeshClothingAssets[SearchIndex] &&
+		   CachedMeshClothingAssets[SearchIndex]->GetAssetGuid() == InAssetGuid)
 		{
 			return SearchIndex;
 		}
@@ -690,7 +691,9 @@ bool USkeletalMesh::HasActiveClothingAssets() const
 #if WITH_EDITOR
 	return ComputeActiveClothingAssets();
 #else
+	PRAGMA_DISABLE_DEPRECATION_WARNINGS
 	return bHasActiveClothingAssets;
+	PRAGMA_ENABLE_DEPRECATION_WARNINGS
 #endif
 }
 
@@ -840,7 +843,7 @@ void USkeletalMesh::InitResources()
 
 		{
 			const int32 NumLODs = SkelMeshRenderData->LODRenderData.Num();
-			const int32 MinFirstLOD = MinLod.GetValue();
+			const int32 MinFirstLOD = GetMinLod().GetValue();
 
 			CachedSRRState.NumNonStreamingLODs = SkelMeshRenderData->NumInlinedLODs;
 			CachedSRRState.NumNonOptionalLODs = SkelMeshRenderData->NumNonOptionalLODs;
@@ -858,7 +861,7 @@ void USkeletalMesh::InitResources()
 
 		// TODO : Update RenderData->CurrentFirstLODIdx based on whether IStreamingManager::Get().IsRenderAssetStreamingEnabled(EStreamableRenderAssetType::SkeletalMesh).
 
-		SkelMeshRenderData->InitResources(bHasVertexColors, MorphTargets, this);
+		SkelMeshRenderData->InitResources(GetHasVertexColors(), GetMorphTargets(), this);
 		CachedSRRState.bHasPendingInitHint = true;
 	}
 
@@ -945,9 +948,10 @@ void USkeletalMesh::UpdateUVChannelData(bool bRebuildAll)
 	FSkeletalMeshRenderData* Resource = GetResourceForRendering();
 	if (FPlatformProperties::HasEditorOnlyData() && Resource)
 	{
-		for (int32 MaterialIndex = 0; MaterialIndex < Materials.Num(); ++MaterialIndex)
+		TArray<FSkeletalMaterial>& MeshMaterials = GetMaterials();
+		for (int32 MaterialIndex = 0; MaterialIndex < MeshMaterials.Num(); ++MaterialIndex)
 		{
-			FMeshUVChannelInfo& UVChannelData = Materials[MaterialIndex].UVChannelData;
+			FMeshUVChannelInfo& UVChannelData = MeshMaterials[MaterialIndex].UVChannelData;
 
 			// Skip it if we want to keep it.
 			if (UVChannelData.IsInitialized() && (!bRebuildAll || UVChannelData.bOverrideDensities))
@@ -975,17 +979,17 @@ void USkeletalMesh::UpdateUVChannelData(bool bRebuildAll)
 			}
 		}
 
-		Resource->SyncUVChannelData(Materials);
+		Resource->SyncUVChannelData(GetMaterials());
 	}
 #endif
 }
 
 const FMeshUVChannelInfo* USkeletalMesh::GetUVChannelData(int32 MaterialIndex) const
 {
-	if (Materials.IsValidIndex(MaterialIndex))
+	if (GetMaterials().IsValidIndex(MaterialIndex))
 	{
-		ensure(Materials[MaterialIndex].UVChannelData.bInitialized);
-		return &Materials[MaterialIndex].UVChannelData;
+		ensure(GetMaterials()[MaterialIndex].UVChannelData.bInitialized);
+		return &GetMaterials()[MaterialIndex].UVChannelData;
 	}
 
 	return nullptr;
@@ -1001,8 +1005,8 @@ void USkeletalMesh::GetResourceSizeEx(FResourceSizeEx& CumulativeResourceSize)
 		SkeletalMeshRenderData->GetResourceSizeEx(CumulativeResourceSize);
 	}
 
-	CumulativeResourceSize.AddDedicatedSystemMemoryBytes(RefBasesInvMatrix.GetAllocatedSize());
-	CumulativeResourceSize.AddDedicatedSystemMemoryBytes(RefSkeleton.GetDataSize());
+	CumulativeResourceSize.AddDedicatedSystemMemoryBytes(GetRefBasesInvMatrix().GetAllocatedSize());
+	CumulativeResourceSize.AddDedicatedSystemMemoryBytes(GetRefSkeleton().GetDataSize());
 }
 
 int32 USkeletalMesh::CalcCumulativeLODSize(int32 NumLODs) const
@@ -1250,7 +1254,7 @@ void USkeletalMesh::PostEditChangeProperty(FPropertyChangedEvent& PropertyChange
 
 	if ( GIsEditor &&
 		 PropertyThatChanged &&
-		 PropertyThatChanged->GetFName() == FName(TEXT("bEnablePerPolyCollision"))
+		 PropertyThatChanged->GetFName() == GetEnablePerPolyCollisionMemberName()
 		)
 	{
 		BuildPhysicsData();
@@ -1268,7 +1272,7 @@ void USkeletalMesh::PostEditChangeProperty(FPropertyChangedEvent& PropertyChange
 		}
 	}
 		
-	if (PropertyThatChanged && PropertyThatChanged->GetFName() == GET_MEMBER_NAME_CHECKED(USkeletalMesh, PostProcessAnimBlueprint))
+	if (PropertyThatChanged && PropertyThatChanged->GetFName() == USkeletalMesh::GetPostProcessAnimBlueprintMemberName())
 	{
 		bHasToReregisterComponent = true;
 	}
@@ -1343,7 +1347,7 @@ void USkeletalMesh::PostEditUndo()
 		}
 	}
 
-	if(MorphTargets.Num() > MorphTargetIndexMap.Num())
+	if(GetMorphTargets().Num() > GetMorphTargetIndexMap().Num())
 	{
 		// A morph target remove has been undone, reinitialise
 		InitMorphTargets();
@@ -1378,9 +1382,9 @@ void USkeletalMesh::BeginDestroy()
 	}
 
 	// remove the cache of link up
-	if (Skeleton)
+	if (GetSkeleton())
 	{
-		Skeleton->RemoveLinkup(this);
+		GetSkeleton()->RemoveLinkup(this);
 	}
 
 #if WITH_EDITORONLY_DATA
@@ -1436,14 +1440,15 @@ void USkeletalMesh::Serialize( FArchive& Ar )
 	FStripDataFlags StripFlags( Ar );
 
 	Ar << ImportedBounds;
-	Ar << Materials;
 
-	Ar << RefSkeleton;
+	Ar << GetMaterials();
+
+	Ar << GetRefSkeleton();
 
 	if(Ar.IsLoading())
 	{
 		const bool bRebuildNameMap = false;
-		RefSkeleton.RebuildRefSkeleton(Skeleton, bRebuildNameMap);
+		GetRefSkeleton().RebuildRefSkeleton(GetSkeleton(), bRebuildNameMap);
 	}
 
 #if WITH_EDITORONLY_DATA
@@ -1504,7 +1509,7 @@ void USkeletalMesh::Serialize( FArchive& Ar )
 	// make sure we're counting properly
 	if (!Ar.IsLoading() && !Ar.IsSaving())
 	{
-		Ar << RefBasesInvMatrix;
+		Ar << GetRefBasesInvMatrix();
 	}
 
 	if( Ar.UE4Ver() < VER_UE4_REFERENCE_SKELETON_REFACTOR )
@@ -1541,19 +1546,19 @@ void USkeletalMesh::Serialize( FArchive& Ar )
 		}
 	}
 
-	if ( Ar.IsLoading() && Ar.UE4Ver() < VER_UE4_ASSET_IMPORT_DATA_AS_JSON && !AssetImportData)
+	if ( Ar.IsLoading() && Ar.UE4Ver() < VER_UE4_ASSET_IMPORT_DATA_AS_JSON && !GetAssetImportData())
 	{
 		// AssetImportData should always be valid
-		AssetImportData = NewObject<UAssetImportData>(this, TEXT("AssetImportData"));
+		SetAssetImportData(NewObject<UAssetImportData>(this, TEXT("AssetImportData")));
 	}
 	
 	// SourceFilePath and SourceFileTimestamp were moved into a subobject
-	if ( Ar.IsLoading() && Ar.UE4Ver() < VER_UE4_ADDED_FBX_ASSET_IMPORT_DATA && AssetImportData)
+	if ( Ar.IsLoading() && Ar.UE4Ver() < VER_UE4_ADDED_FBX_ASSET_IMPORT_DATA && GetAssetImportData())
 	{
 		// AssetImportData should always have been set up in the constructor where this is relevant
 		FAssetImportInfo Info;
 		Info.Insert(FAssetImportInfo::FSourceFile(SourceFilePath_DEPRECATED));
-		AssetImportData->SourceData = MoveTemp(Info);
+		GetAssetImportData()->SourceData = MoveTemp(Info);
 		
 		SourceFilePath_DEPRECATED = TEXT("");
 		SourceFileTimestamp_DEPRECATED = TEXT("");
@@ -1585,13 +1590,16 @@ void USkeletalMesh::Serialize( FArchive& Ar )
 
 	if (Ar.UE4Ver() < VER_UE4_SKELETON_ASSET_PROPERTY_TYPE_CHANGE)
 	{
-		PreviewAttachedAssetContainer.SaveAttachedObjectsFromDeprecatedProperties();
+		GetPreviewAttachedAssetContainer().SaveAttachedObjectsFromDeprecatedProperties();
 	}
 #endif
 
-	if (bEnablePerPolyCollision)
+	if (GetEnablePerPolyCollision())
 	{
-		Ar << BodySetup;
+		const USkeletalMesh* ConstThis = this;
+		UBodySetup* LocalBodySetup = ConstThis->GetBodySetup();
+		Ar << LocalBodySetup;
+		SetBodySetup(LocalBodySetup);
 	}
 
 #if WITH_EDITORONLY_DATA
@@ -1602,8 +1610,8 @@ void USkeletalMesh::Serialize( FArchive& Ar )
 #endif
 
 #if WITH_EDITORONLY_DATA
-	bRequiresLODScreenSizeConversion = Ar.CustomVer(FFrameworkObjectVersion::GUID) < FFrameworkObjectVersion::LODsUseResolutionIndependentScreenSize;
-	bRequiresLODHysteresisConversion = Ar.CustomVer(FFrameworkObjectVersion::GUID) < FFrameworkObjectVersion::LODHysteresisUseResolutionIndependentScreenSize;
+	SetRequiresLODScreenSizeConversion(Ar.CustomVer(FFrameworkObjectVersion::GUID) < FFrameworkObjectVersion::LODsUseResolutionIndependentScreenSize);
+	SetRequiresLODHysteresisConversion(Ar.CustomVer(FFrameworkObjectVersion::GUID) < FFrameworkObjectVersion::LODHysteresisUseResolutionIndependentScreenSize);
 #endif
 
 	if (Ar.CustomVer(FFortniteMainBranchObjectVersion::GUID) < FFortniteMainBranchObjectVersion::ConvertReductionSettingOptions)
@@ -1647,7 +1655,7 @@ void USkeletalMesh::Serialize( FArchive& Ar )
 void USkeletalMesh::GetPreloadDependencies(TArray<UObject*>& OutDeps)
 {
 	Super::GetPreloadDependencies(OutDeps);
-	OutDeps.Add(Skeleton);
+	OutDeps.Add(GetSkeleton());
 }
 
 void USkeletalMesh::PostDuplicate(bool bDuplicateForPIE)
@@ -1655,17 +1663,18 @@ void USkeletalMesh::PostDuplicate(bool bDuplicateForPIE)
 	Super::PostDuplicate(bDuplicateForPIE);
 #if WITH_EDITORONLY_DATA
 	//We must not use the same asset for the imported data and duplicate it in the destination skeletalmesh package
-	if (MeshEditorDataObject)
+	if (IsMeshEditorDataValid())
 	{
+		PRAGMA_DISABLE_DEPRECATION_WARNINGS
 		USkeletalMeshEditorData* SourceMeshEditorData = MeshEditorDataObject;
-		check(SourceMeshEditorData);
 		MeshEditorDataObject = nullptr;
+		PRAGMA_ENABLE_DEPRECATION_WARNINGS
 		FSkeletalMeshModel& ImportedModels = *GetImportedModel();
 		//allocate the necessary data before going multi thread
 		ReserveLODImportData(ImportedModels.LODModels.Num() - 1);
 		USkeletalMeshEditorData& DestMeshEditorData = GetMeshEditorData();
 		//We should have a brand new created asset
-		check(MeshEditorDataObject);
+		check(IsMeshEditorDataValid());
 		//Lets duplicate the imported data
 		const int32 NumLODModels = ImportedModels.LODModels.Num();
 		for (int32 LODIndex = 0; LODIndex < NumLODModels; ++LODIndex)
@@ -1700,8 +1709,12 @@ void USkeletalMesh::PostRename(UObject* OldOuter, const FName OldName)
 	Super::PostRename(OldOuter, OldName);
 #if WITH_EDITORONLY_DATA
 	//We must not use the same asset for the imported data and duplicate it in the destination skeletalmesh package
-	if (MeshEditorDataObject)
+	if (IsMeshEditorDataValid())
 	{
+		PRAGMA_DISABLE_DEPRECATION_WARNINGS
+		USkeletalMeshEditorData* SourceMeshEditorData = MeshEditorDataObject;
+		PRAGMA_ENABLE_DEPRECATION_WARNINGS
+
 		FString MeshEditorDataString = GetName() + TEXT("_USkeletalMeshEditorData");
 		UPackage* SkeletalMeshPackage = GetOutermost();
 
@@ -1710,14 +1723,14 @@ void USkeletalMesh::PostRename(UObject* OldOuter, const FName OldName)
 		if(!ExistingObject)
 		{
 			//Do a soft rename avoid: dirty, redirector, transaction and reset of the loaders
-			MeshEditorDataObject->Rename(*MeshEditorDataString, GetOutermost(), (REN_DoNotDirty | REN_DontCreateRedirectors | REN_NonTransactional | REN_ForceNoResetLoaders));
+			SourceMeshEditorData->Rename(*MeshEditorDataString, GetOutermost(), (REN_DoNotDirty | REN_DontCreateRedirectors | REN_NonTransactional | REN_ForceNoResetLoaders));
 		}
 		else
 		{
-			if (ExistingObject != MeshEditorDataObject)
+			if (ExistingObject != SourceMeshEditorData)
 			{
 				//Log that the MeshEditorDataObject is not renamed properly.
-				UE_LOG(LogSkeletalMesh, Warning, TEXT("Renaming of skeletal mesh %s failed to rename the sub UObject MeshEditorDataObject (The imported data), because there is already an object with this name!/nMeshEditorDataObject name before renaming %s./nName of the existing object preventing the renaming %s"), *GetFullName(), *MeshEditorDataObject->GetFullName(), *ExistingObject->GetFullName());
+				UE_LOG(LogSkeletalMesh, Warning, TEXT("Renaming of skeletal mesh %s failed to rename the sub UObject MeshEditorDataObject (The imported data), because there is already an object with this name!/nMeshEditorDataObject name before renaming %s./nName of the existing object preventing the renaming %s"), *GetFullName(), *SourceMeshEditorData->GetFullName(), *ExistingObject->GetFullName());
 			}
 		}
 	}
@@ -1739,7 +1752,7 @@ void USkeletalMesh::FlushRenderState()
 uint32 USkeletalMesh::GetVertexBufferFlags() const
 {
 	uint32 VertexFlags = ESkeletalMeshVertexFlags::None;
-	if (bHasVertexColors)
+	if (GetHasVertexColors())
 	{
 		VertexFlags |= ESkeletalMeshVertexFlags::HasVertexColors;
 	}
@@ -1779,7 +1792,7 @@ void USkeletalMesh::Build()
 void USkeletalMesh::PreSave(const class ITargetPlatform* TargetPlatform)
 {
 	// check the parent index of the root bone is invalid
-	check((RefSkeleton.GetNum() == 0) || (RefSkeleton.GetRefBoneInfo()[0].ParentIndex == INDEX_NONE));
+	check((GetRefSkeleton().GetNum() == 0) || (GetRefSkeleton().GetRefBoneInfo()[0].ParentIndex == INDEX_NONE));
 
 	Super::PreSave(TargetPlatform);
 }
@@ -1787,12 +1800,12 @@ void USkeletalMesh::PreSave(const class ITargetPlatform* TargetPlatform)
 // Pre-calculate refpose-to-local transforms
 void USkeletalMesh::CalculateInvRefMatrices()
 {
-	const int32 NumRealBones = RefSkeleton.GetRawBoneNum();
+	const int32 NumRealBones = GetRefSkeleton().GetRawBoneNum();
 
-	if( RefBasesInvMatrix.Num() != NumRealBones)
+	if(GetRefBasesInvMatrix().Num() != NumRealBones)
 	{
-		RefBasesInvMatrix.Empty(NumRealBones);
-		RefBasesInvMatrix.AddUninitialized(NumRealBones);
+		GetRefBasesInvMatrix().Empty(NumRealBones);
+		GetRefBasesInvMatrix().AddUninitialized(NumRealBones);
 
 		// Reset cached mesh-space ref pose
 		CachedComposedRefPoseMatrices.Empty(NumRealBones);
@@ -1807,7 +1820,7 @@ void USkeletalMesh::CalculateInvRefMatrices()
 			// Construct mesh-space skeletal hierarchy.
 			if( b>0 )
 			{
-				int32 Parent = RefSkeleton.GetRawParentIndex(b);
+				int32 Parent = GetRefSkeleton().GetRawParentIndex(b);
 				CachedComposedRefPoseMatrices[b] = CachedComposedRefPoseMatrices[b] * CachedComposedRefPoseMatrices[Parent];
 			}
 
@@ -1819,17 +1832,17 @@ void USkeletalMesh::CalculateInvRefMatrices()
 				ZAxis.IsNearlyZero(SMALL_NUMBER))
 			{
 				// this is not allowed, warn them 
-				UE_LOG(LogSkeletalMesh, Warning, TEXT("Reference Pose for asset %s for joint (%s) includes NIL matrix. Zero scale isn't allowed on ref pose. "), *GetPathName(), *RefSkeleton.GetBoneName(b).ToString());
+				UE_LOG(LogSkeletalMesh, Warning, TEXT("Reference Pose for asset %s for joint (%s) includes NIL matrix. Zero scale isn't allowed on ref pose. "), *GetPathName(), *GetRefSkeleton().GetBoneName(b).ToString());
 			}
 
 			// Precompute inverse so we can use from-refpose-skin vertices.
-			RefBasesInvMatrix[b] = CachedComposedRefPoseMatrices[b].Inverse(); 
+			GetRefBasesInvMatrix()[b] = CachedComposedRefPoseMatrices[b].Inverse();
 		}
 
 #if WITH_EDITORONLY_DATA
-		if(RetargetBasePose.Num() == 0)
+		if(GetRetargetBasePose().Num() == 0)
 		{
-			RetargetBasePose = RefSkeleton.GetRefBonePose();
+			SetRetargetBasePose(GetRefSkeleton().GetRefBonePose());
 		}
 #endif // WITH_EDITORONLY_DATA
 	}
@@ -1841,33 +1854,33 @@ void USkeletalMesh::ReallocateRetargetBasePose()
 	// if you're adding other things here, please note that this function is called during postLoad
 	// fix up retarget base pose if VB has changed
 	// if we have virtual joints, we make sure Retarget Base Pose matches
-	const int32 RawNum = RefSkeleton.GetRawBoneNum();
-	const int32 VBNum = RefSkeleton.GetVirtualBoneRefData().Num();
-	const int32 BoneNum = RefSkeleton.GetNum();
+	const int32 RawNum = GetRefSkeleton().GetRawBoneNum();
+	const int32 VBNum = GetRefSkeleton().GetVirtualBoneRefData().Num();
+	const int32 BoneNum = GetRefSkeleton().GetNum();
 	check(RawNum + VBNum == BoneNum);
 
-	const int32 OldRetargetBasePoseNum = RetargetBasePose.Num();
+	const int32 OldRetargetBasePoseNum = GetRetargetBasePose().Num();
 	// we want to make sure retarget base pose contains raw numbers PREVIOUSLY
 	// otherwise, we may override wrong transform
 	if (OldRetargetBasePoseNum >= RawNum)
 	{
 		// we have to do this in case buffer size changes (shrink for example)
-		RetargetBasePose.SetNum(BoneNum);
+		GetRetargetBasePose().SetNum(BoneNum);
 
 		// if we have VB, we should override them
 		// they're not editable, so it's fine to override them from raw bones
 		if (VBNum > 0)
 		{
-			const TArray<FTransform>& BonePose = RefSkeleton.GetRefBonePose();
-			check(RetargetBasePose.GetTypeSize() == BonePose.GetTypeSize());
-			const int32 ElementSize = RetargetBasePose.GetTypeSize();
-			FMemory::Memcpy(RetargetBasePose.GetData() + RawNum, BonePose.GetData() + RawNum, ElementSize*VBNum);
+			const TArray<FTransform>& BonePose = GetRefSkeleton().GetRefBonePose();
+			check(GetRetargetBasePose().GetTypeSize() == BonePose.GetTypeSize());
+			const int32 ElementSize = GetRetargetBasePose().GetTypeSize();
+			FMemory::Memcpy(GetRetargetBasePose().GetData() + RawNum, BonePose.GetData() + RawNum, ElementSize*VBNum);
 		}
 	}
 	else
 	{
 		// else we think, something has changed, we just override retarget base pose to current pose
-		RetargetBasePose = RefSkeleton.GetRefBonePose();
+		GetRetargetBasePose() = GetRefSkeleton().GetRefBonePose();
 	}
 }
 
@@ -1955,7 +1968,7 @@ void USkeletalMesh::UpgradeOldClothingAssets()
 				// Pull the path across so reimports work as expected
 				NewAsset->ImportedFilePath = OldAssetData.ApexFileName;
 
-				MeshClothingAssets.Add(NewAsset);
+				GetMeshClothingAssets().Add(NewAsset);
 			}
 		}
 
@@ -1968,9 +1981,9 @@ void USkeletalMesh::UpgradeOldClothingAssets()
 
 		check(OldLodMappings.Num() == OldSectionMappings.Num());
 
-		for (int32 NewAssetIdx = 0; NewAssetIdx < MeshClothingAssets.Num(); ++NewAssetIdx)
+		for (int32 NewAssetIdx = 0; NewAssetIdx < GetMeshClothingAssets().Num(); ++NewAssetIdx)
 		{
-			UClothingAssetBase* CurrAsset = MeshClothingAssets[NewAssetIdx];
+			UClothingAssetBase* CurrAsset = GetMeshClothingAssets()[NewAssetIdx];
 
 			for (int32 MappedLodIdx = 0; MappedLodIdx < OldLodMappings[NewAssetIdx].Num(); ++MappedLodIdx)
 			{
@@ -2035,7 +2048,7 @@ void USkeletalMesh::RemoveLegacyClothingSections()
 							Section.ClothMappingData = DuplicatedSection.ClothMappingData;
 						}
 
-						Section.CorrespondClothAssetIndex = MeshClothingAssets.IndexOfByPredicate([&Section](const UClothingAssetBase* CurrAsset) 
+						Section.CorrespondClothAssetIndex = GetMeshClothingAssets().IndexOfByPredicate([&Section](const UClothingAssetBase* CurrAsset)
 						{
 							return CurrAsset && CurrAsset->GetAssetGuid() == Section.ClothingData.AssetGuid;
 						});
@@ -2084,7 +2097,8 @@ void USkeletalMesh::RemoveLegacyClothingSections()
 
 USkeletalMeshEditorData& USkeletalMesh::GetMeshEditorData() const
 {
-	if (MeshEditorDataObject == nullptr)
+	PRAGMA_DISABLE_DEPRECATION_WARNINGS
+	if (!IsMeshEditorDataValid())
 	{
 		FString MeshEditorDataString = GetName() + TEXT("_USkeletalMeshEditorData");
 		FName MeshEditorDataName = *MeshEditorDataString;
@@ -2094,7 +2108,7 @@ USkeletalMeshEditorData& USkeletalMesh::GetMeshEditorData() const
 		{
 			MeshEditorDataObject = Cast<USkeletalMeshEditorData>(ExistingObject);
 		}
-		if (MeshEditorDataObject == nullptr)
+		if (!IsMeshEditorDataValid())
 		{
 			//The asset is created in the skeletalmesh package. We keep it private so the user cannot see it in the content browser
 			//StandAlone make sure the asset is save when we save the package(i.e. the skeletalmesh)
@@ -2104,6 +2118,7 @@ USkeletalMeshEditorData& USkeletalMesh::GetMeshEditorData() const
 	//Make sure we have a valid pointer
 	check(MeshEditorDataObject != nullptr);
 	return *MeshEditorDataObject;
+	PRAGMA_ENABLE_DEPRECATION_WARNINGS
 }
 
 void USkeletalMesh::LoadLODImportedData(const int32 LODIndex, FSkeletalMeshImportData& OutMesh) const
@@ -2217,7 +2232,7 @@ void USkeletalMesh::CreateUserSectionsDataForLegacyAssets()
 {
 	//We want to avoid changing the ddc if we load an old asset.
 	//This bool should be put to false at the end of the postload, if there is another posteditchange call after a new ddc will be created
-	UseLegacyMeshDerivedDataKey = true;
+	SetUseLegacyMeshDerivedDataKey(true);
 	//Fill up the Section ChunkedParentSectionIndex and OriginalDataSectionIndex
 	//We also want to create the UserSectionsData structure so the user can change the section data
 	for (int32 LodIndex = 0; LodIndex < LODInfo.Num(); LodIndex++)
@@ -2396,7 +2411,7 @@ void USkeletalMesh::PostLoad()
 	//       on the cloth assets even before reaching this point.
 	//       In these occurences, the cloth asset's RF_NeedsPostLoad flag is already cleared despite its PostLoad still
 	//       being un-executed, making the following block code ineffective.
-	for (UClothingAssetBase* MeshClothingAsset : MeshClothingAssets)
+	for (UClothingAssetBase* MeshClothingAsset : GetMeshClothingAssets())
 	{
 		MeshClothingAsset->ConditionalPostLoad();
 	}
@@ -2448,16 +2463,16 @@ void USkeletalMesh::PostLoad()
 		}
 
 		// load LODinfo if using shared asset, it can override existing bone remove settings
-		if (LODSettings != nullptr)
+		if (GetLODSettings() != nullptr)
 		{
 			//before we copy
 			if (GetLinkerCustomVersion(FFortniteMainBranchObjectVersion::GUID) < FFortniteMainBranchObjectVersion::AddBakePoseOverrideForSkeletalMeshReductionSetting)
 			{
 				// if LODsetting doesn't have BakePose, but this does, we'll have to copy that to BakePoseOverride
-				const int32 NumSettings = FMath::Min(LODSettings->GetNumberOfSettings(), GetLODNum());
+				const int32 NumSettings = FMath::Min(GetLODSettings()->GetNumberOfSettings(), GetLODNum());
 				for (int32 Index = 0; Index < NumSettings; ++Index)
 				{
-					const FSkeletalMeshLODGroupSettings& GroupSetting = LODSettings->GetSettingsForLODLevel(Index);
+					const FSkeletalMeshLODGroupSettings& GroupSetting = GetLODSettings()->GetSettingsForLODLevel(Index);
 					// if lod setting doesn't have bake pose, but this lod does, that means this bakepose has to move to BakePoseOverride
 					// since we want to match what GroupSetting has
 					if (GroupSetting.BakePose == nullptr && LODInfo[Index].BakePose)
@@ -2468,7 +2483,7 @@ void USkeletalMesh::PostLoad()
 					}
 				}
 			}
-			LODSettings->SetLODSettingsToMesh(this);
+			GetLODSettings()->SetLODSettingsToMesh(this);
 		}
 
 		if (GetLinkerUE4Version() < VER_UE4_SORT_ACTIVE_BONE_INDICES)
@@ -2487,7 +2502,7 @@ void USkeletalMesh::PostLoad()
 			for (int32 LodIndex = 0; LodIndex < LODInfo.Num(); LodIndex++)
 			{
 				FSkeletalMeshLODModel& ThisLODModel = ImportedModel->LODModels[LodIndex];
-				RefSkeleton.EnsureParentsExistAndSort(ThisLODModel.ActiveBoneIndices);
+				GetRefSkeleton().EnsureParentsExistAndSort(ThisLODModel.ActiveBoneIndices);
 			}
 		}
 
@@ -2545,12 +2560,12 @@ void USkeletalMesh::PostLoad()
 		}
 
 		//Make sure unused cloth are unbind
-		if (MeshClothingAssets.Num() > 0)
+		if (GetMeshClothingAssets().Num() > 0)
 		{
 			TArray<UClothingAssetBase *> InUsedClothingAssets;
 			GetClothingAssetsInUse(InUsedClothingAssets);
 			//Look if we have some cloth binding to unbind
-			for (UClothingAssetBase* MeshClothingAsset : MeshClothingAssets)
+			for (UClothingAssetBase* MeshClothingAsset : GetMeshClothingAssets())
 			{
 				if (MeshClothingAsset == nullptr)
 				{
@@ -2603,9 +2618,9 @@ void USkeletalMesh::PostLoad()
 	CalculateInvRefMatrices();
 
 #if WITH_EDITORONLY_DATA
-	if (RetargetBasePose.Num() == 0 && !GetOutermost()->bIsCookedForEditor)
+	if (GetRetargetBasePose().Num() == 0 && !GetOutermost()->bIsCookedForEditor)
 	{
-		RetargetBasePose = RefSkeleton.GetRefBonePose();
+		GetRetargetBasePose() = GetRefSkeleton().GetRefBonePose();
 	}
 
 	if (GetLinkerCustomVersion(FFortniteMainBranchObjectVersion::GUID) < FFortniteMainBranchObjectVersion::SupportVirtualBoneInRetargeting)
@@ -2618,7 +2633,7 @@ void USkeletalMesh::PostLoad()
 	CalculateExtendedBounds();
 
 #if WITH_EDITORONLY_DATA
-	if (bRequiresLODScreenSizeConversion || bRequiresLODHysteresisConversion)
+	if (GetRequiresLODScreenSizeConversion() || GetRequiresLODHysteresisConversion())
 	{
 		// Convert screen area to screen size
 		ConvertLegacyLODScreenSize();
@@ -2628,14 +2643,14 @@ void USkeletalMesh::PostLoad()
 	// If inverse masses have never been cached, invalidate data so it will be recalculated
 	if(GetLinkerCustomVersion(FSkeletalMeshCustomVersion::GUID) < FSkeletalMeshCustomVersion::CachedClothInverseMasses)
 	{
-		for(UClothingAssetBase* ClothingAsset : MeshClothingAssets)
+		for(UClothingAssetBase* ClothingAsset : GetMeshClothingAssets())
 		{
 			if(ClothingAsset) 
 			ClothingAsset->InvalidateCachedData();
 		}
 	}
 
-	bHasActiveClothingAssets = ComputeActiveClothingAssets();
+	SetHasActiveClothingAssets(ComputeActiveClothingAssets());
 
 #if WITH_EDITOR
 	if (GetLinkerCustomVersion(FNiagaraObjectVersion::GUID) < FNiagaraObjectVersion::SkeletalMeshVertexSampling)
@@ -2651,7 +2666,7 @@ void USkeletalMesh::PostLoad()
 #endif // !WITH_EDITOR
 #if WITH_EDITORONLY_DATA
 	//Next postedit change will use the new ddc key scheme
-	UseLegacyMeshDerivedDataKey = false;
+	SetUseLegacyMeshDerivedDataKey(false);
 #endif
 }
 
@@ -2661,7 +2676,7 @@ void USkeletalMesh::RebuildRefSkeletonNameToIndexMap()
 {
 	TArray<FBoneIndexType> DuplicateBones;
 	// Make sure we have no duplicate bones. Some content got corrupted somehow. :(
-	RefSkeleton.RemoveDuplicateBones(this, DuplicateBones);
+	GetRefSkeleton().RemoveDuplicateBones(this, DuplicateBones);
 
 	// If we have removed any duplicate bones, we need to fix up any broken LODs as well.
 	// Duplicate bones are given from highest index to lowest.
@@ -2701,7 +2716,7 @@ void USkeletalMesh::RebuildRefSkeletonNameToIndexMap()
 	}
 
 	// Rebuild name table.
-	RefSkeleton.RebuildNameToIndexMap();
+	GetRefSkeleton().RebuildNameToIndexMap();
 }
 
 #endif
@@ -2724,15 +2739,15 @@ void USkeletalMesh::GetAssetRegistryTags(TArray<FAssetRegistryTag>& OutTags) con
 	OutTags.Add(FAssetRegistryTag("Vertices", FString::FromInt(NumVertices), FAssetRegistryTag::TT_Numerical));
 	OutTags.Add(FAssetRegistryTag("Triangles", FString::FromInt(NumTriangles), FAssetRegistryTag::TT_Numerical));
 	OutTags.Add(FAssetRegistryTag("LODs", FString::FromInt(NumLODs), FAssetRegistryTag::TT_Numerical));
-	OutTags.Add(FAssetRegistryTag("Bones", FString::FromInt(RefSkeleton.GetRawBoneNum()), FAssetRegistryTag::TT_Numerical));
-	OutTags.Add(FAssetRegistryTag("MorphTargets", FString::FromInt(MorphTargets.Num()), FAssetRegistryTag::TT_Numerical));
+	OutTags.Add(FAssetRegistryTag("Bones", FString::FromInt(GetRefSkeleton().GetRawBoneNum()), FAssetRegistryTag::TT_Numerical));
+	OutTags.Add(FAssetRegistryTag("MorphTargets", FString::FromInt(GetMorphTargets().Num()), FAssetRegistryTag::TT_Numerical));
 
 #if WITH_EDITORONLY_DATA
-	if (AssetImportData)
+	if (GetAssetImportData())
 	{
-		OutTags.Add( FAssetRegistryTag(SourceFileTagName(), AssetImportData->GetSourceData().ToJson(), FAssetRegistryTag::TT_Hidden) );
+		OutTags.Add( FAssetRegistryTag(SourceFileTagName(), GetAssetImportData()->GetSourceData().ToJson(), FAssetRegistryTag::TT_Hidden) );
 #if WITH_EDITOR
-		AssetImportData->AppendAssetRegistryTags(OutTags);
+		GetAssetImportData()->AppendAssetRegistryTags(OutTags);
 #endif
 	}
 #endif
@@ -2809,12 +2824,12 @@ bool USkeletalMesh::RegisterMorphTarget(UMorphTarget* MorphTarget, bool bInvalid
 
 		bool bRegistered = false;
 
-		for ( int32 Index = 0; Index < MorphTargets.Num(); ++Index )
+		for ( int32 Index = 0; Index < GetMorphTargets().Num(); ++Index )
 		{
-			if ( MorphTargets[Index]->GetFName() == MorphTarget->GetFName() )
+			if (GetMorphTargets()[Index]->GetFName() == MorphTarget->GetFName() )
 			{
 				UE_LOG( LogSkeletalMesh, Verbose, TEXT("RegisterMorphTarget: %s already exists, replacing"), *MorphTarget->GetName() );
-				MorphTargets[Index] = MorphTarget;
+				GetMorphTargets()[Index] = MorphTarget;
 				bRegistered = true;
 				break;
 			}
@@ -2822,7 +2837,7 @@ bool USkeletalMesh::RegisterMorphTarget(UMorphTarget* MorphTarget, bool bInvalid
 
 		if (!bRegistered)
 		{
-			MorphTargets.Add( MorphTarget );
+			GetMorphTargets().Add( MorphTarget );
 			bRegistered = true;
 		}
 
@@ -2838,7 +2853,7 @@ bool USkeletalMesh::RegisterMorphTarget(UMorphTarget* MorphTarget, bool bInvalid
 
 void USkeletalMesh::UnregisterAllMorphTarget()
 {
-	MorphTargets.Empty();
+	GetMorphTargets().Empty();
 	InitMorphTargetsAndRebuildRenderData();
 }
 
@@ -2848,11 +2863,11 @@ void USkeletalMesh::UnregisterMorphTarget(UMorphTarget* MorphTarget)
 	{
 		// Do not remove with MorphTarget->GetFName(). The name might have changed
 		// Search the value, and delete	
-		for ( int32 I=0; I<MorphTargets.Num(); ++I)
+		for ( int32 I=0; I< GetMorphTargets().Num(); ++I)
 		{
-			if ( MorphTargets[I] == MorphTarget )
+			if (GetMorphTargets()[I] == MorphTarget )
 			{
-				MorphTargets.RemoveAt(I);
+				GetMorphTargets().RemoveAt(I);
 				--I;
 				InitMorphTargetsAndRebuildRenderData();
 				return;
@@ -2864,34 +2879,34 @@ void USkeletalMesh::UnregisterMorphTarget(UMorphTarget* MorphTarget)
 
 void USkeletalMesh::InitMorphTargets()
 {
-	MorphTargetIndexMap.Empty();
+	GetMorphTargetIndexMap().Empty();
 
-	for (int32 Index = 0; Index < MorphTargets.Num(); ++Index)
+	for (int32 Index = 0; Index < GetMorphTargets().Num(); ++Index)
 	{
-		UMorphTarget* MorphTarget = MorphTargets[Index];
+		UMorphTarget* MorphTarget = GetMorphTargets()[Index];
 		// if we don't have a valid data, just remove it
 		if (!MorphTarget->HasValidData())
 		{
-			MorphTargets.RemoveAt(Index);
+			GetMorphTargets().RemoveAt(Index);
 			--Index;
 			continue;
 		}
 
 		FName const ShapeName = MorphTarget->GetFName();
-		if (MorphTargetIndexMap.Find(ShapeName) == nullptr)
+		if (GetMorphTargetIndexMap().Find(ShapeName) == nullptr)
 		{
-			MorphTargetIndexMap.Add(ShapeName, Index);
+			GetMorphTargetIndexMap().Add(ShapeName, Index);
 
 			// register as morphtarget curves
-			if (Skeleton)
+			if (GetSkeleton())
 			{
 				FSmartName CurveName;
 				CurveName.DisplayName = ShapeName;
 				
 				// verify will make sure it adds to the curve if not found
 				// the reason of using this is to make sure it works in editor/non-editor
-				Skeleton->VerifySmartName(USkeleton::AnimCurveMappingName, CurveName);
-				Skeleton->AccumulateCurveMetaData(ShapeName, false, true);
+				GetSkeleton()->VerifySmartName(USkeleton::AnimCurveMappingName, CurveName);
+				GetSkeleton()->AccumulateCurveMetaData(ShapeName, false, true);
 			}
 		}
 	}
@@ -2908,11 +2923,11 @@ UMorphTarget* USkeletalMesh::FindMorphTargetAndIndex(FName MorphTargetName, int3
 	OutIndex = INDEX_NONE;
 	if( MorphTargetName != NAME_None )
 	{
-		const int32* Found = MorphTargetIndexMap.Find(MorphTargetName);
+		const int32* Found = GetMorphTargetIndexMap().Find(MorphTargetName);
 		if (Found)
 		{
 			OutIndex = *Found;
-			return MorphTargets[*Found];
+			return GetMorphTargets()[*Found];
 		}
 	}
 
@@ -2931,7 +2946,7 @@ USkeletalMesh::FSocketInfo::FSocketInfo(const USkeletalMesh* InSkeletalMesh, USk
 	: SocketLocalTransform(InSocket->GetSocketLocalTransform())
 	, Socket(InSocket)
 	, SocketIndex(InSocketIndex)
-	, SocketBoneIndex(InSkeletalMesh->RefSkeleton.FindBoneIndex(InSocket->BoneName))
+	, SocketBoneIndex(InSkeletalMesh->GetRefSkeleton().FindBoneIndex(InSocket->BoneName))
 {}
 
 #endif
@@ -2967,9 +2982,9 @@ USkeletalMeshSocket* USkeletalMesh::FindSocketAndIndex(FName InSocketName, int32
 	}
 
 	// If the socket isn't on the mesh, try to find it on the skeleton
-	if (Skeleton)
+	if (GetSkeleton())
 	{
-		USkeletalMeshSocket* SkeletonSocket = Skeleton->FindSocketAndIndex(InSocketName, OutIndex);
+		USkeletalMeshSocket* SkeletonSocket = GetSkeleton()->FindSocketAndIndex(InSocketName, OutIndex);
 		if (SkeletonSocket != nullptr)
 		{
 			OutIndex += Sockets.Num();
@@ -3012,20 +3027,20 @@ USkeletalMeshSocket* USkeletalMesh::FindSocketInfo(FName InSocketName, FTransfor
 		{
 			OutIndex = i;
 			OutTransform = Socket->GetSocketLocalTransform();
-			OutBoneIndex = RefSkeleton.FindBoneIndex(Socket->BoneName);
+			OutBoneIndex = GetRefSkeleton().FindBoneIndex(Socket->BoneName);
 			return Socket;
 		}
 	}
 
 	// If the socket isn't on the mesh, try to find it on the skeleton
-	if (Skeleton)
+	if (GetSkeleton())
 	{
-		USkeletalMeshSocket* SkeletonSocket = Skeleton->FindSocketAndIndex(InSocketName, OutIndex);
+		USkeletalMeshSocket* SkeletonSocket = GetSkeleton()->FindSocketAndIndex(InSocketName, OutIndex);
 		if (SkeletonSocket != nullptr)
 		{
 			OutIndex += Sockets.Num();
 			OutTransform = SkeletonSocket->GetSocketLocalTransform();
-			OutBoneIndex = RefSkeleton.FindBoneIndex(SkeletonSocket->BoneName);
+			OutBoneIndex = GetRefSkeleton().FindBoneIndex(SkeletonSocket->BoneName);
 		}
 		return SkeletonSocket;
 	}
@@ -3035,7 +3050,7 @@ USkeletalMeshSocket* USkeletalMesh::FindSocketInfo(FName InSocketName, FTransfor
 
 int32 USkeletalMesh::NumSockets() const
 {
-	return Sockets.Num() + (Skeleton ? Skeleton->Sockets.Num() : 0);
+	return Sockets.Num() + (GetSkeleton() ? GetSkeleton()->Sockets.Num() : 0);
 }
 
 USkeletalMeshSocket* USkeletalMesh::GetSocketByIndex(int32 Index) const
@@ -3046,9 +3061,9 @@ USkeletalMeshSocket* USkeletalMesh::GetSocketByIndex(int32 Index) const
 		return Sockets[Index];
 	}
 
-	if (Skeleton && (Index - NumMeshSockets) < Skeleton->Sockets.Num())
+	if (GetSkeleton() && (Index - NumMeshSockets) < GetSkeleton()->Sockets.Num())
 	{
-		return Skeleton->Sockets[Index - NumMeshSockets];
+		return GetSkeleton()->Sockets[Index - NumMeshSockets];
 	}
 
 	return nullptr;
@@ -3059,7 +3074,7 @@ TMap<FVector, FColor> USkeletalMesh::GetVertexColorData(const uint32 PaintingMes
 	TMap<FVector, FColor> VertexColorData;
 #if WITH_EDITOR
 	const FSkeletalMeshModel* SkeletalMeshModel = GetImportedModel();
-	if (bHasVertexColors && SkeletalMeshModel && SkeletalMeshModel->LODModels.IsValidIndex(PaintingMeshLODIndex))
+	if (GetHasVertexColors() && SkeletalMeshModel && SkeletalMeshModel->LODModels.IsValidIndex(PaintingMeshLODIndex))
 	{
 		const TArray<FSkelMeshSection>& Sections = SkeletalMeshModel->LODModels[PaintingMeshLODIndex].Sections;
 
@@ -3088,7 +3103,7 @@ void USkeletalMesh::RebuildSocketMap()
 	check(IsInGameThread());
 
 	SocketMap.Reset();
-	SocketMap.Reserve(Sockets.Num() + (Skeleton ? Skeleton->Sockets.Num() : 0));
+	SocketMap.Reserve(Sockets.Num() + (GetSkeleton() ? GetSkeleton()->Sockets.Num() : 0));
 
 	int32 SocketIndex;
 	for (SocketIndex = 0; SocketIndex < Sockets.Num(); ++SocketIndex)
@@ -3098,11 +3113,11 @@ void USkeletalMesh::RebuildSocketMap()
 	}
 
 	// If the socket isn't on the mesh, try to find it on the skeleton
-	if (Skeleton)
+	if (GetSkeleton())
 	{
-		for (SocketIndex = 0; SocketIndex < Skeleton->Sockets.Num(); ++SocketIndex)
+		for (SocketIndex = 0; SocketIndex < GetSkeleton()->Sockets.Num(); ++SocketIndex)
 		{
-			USkeletalMeshSocket* Socket = Skeleton->Sockets[SocketIndex];
+			USkeletalMeshSocket* Socket = GetSkeleton()->Sockets[SocketIndex];
 			if (!SocketMap.Contains(Socket->SocketName))
 			{
 				SocketMap.Add(Socket->SocketName, FSocketInfo(this, Socket, Sockets.Num() + SocketIndex));
@@ -3128,8 +3143,8 @@ FString USkeletalMesh::GetDetailedInfoInternal() const
 
 FMatrix USkeletalMesh::GetRefPoseMatrix( int32 BoneIndex ) const
 {
- 	check( BoneIndex >= 0 && BoneIndex < RefSkeleton.GetRawBoneNum() );
-	FTransform BoneTransform = RefSkeleton.GetRawRefBonePose()[BoneIndex];
+ 	check( BoneIndex >= 0 && BoneIndex < GetRefSkeleton().GetRawBoneNum() );
+	FTransform BoneTransform = GetRefSkeleton().GetRawRefBonePose()[BoneIndex];
 	// Make sure quaternion is normalized!
 	BoneTransform.NormalizeRotation();
 	return BoneTransform.ToMatrixWithScale();
@@ -3141,7 +3156,7 @@ FMatrix USkeletalMesh::GetComposedRefPoseMatrix( FName InBoneName ) const
 
 	if ( InBoneName != NAME_None )
 	{
-		int32 BoneIndex = RefSkeleton.FindBoneIndex(InBoneName);
+		int32 BoneIndex = GetRefSkeleton().FindBoneIndex(InBoneName);
 		if (BoneIndex != INDEX_NONE)
 		{
 			return GetComposedRefPoseMatrix(BoneIndex);
@@ -3152,7 +3167,7 @@ FMatrix USkeletalMesh::GetComposedRefPoseMatrix( FName InBoneName ) const
 
 			if(Socket != NULL)
 			{
-				BoneIndex = RefSkeleton.FindBoneIndex(Socket->BoneName);
+				BoneIndex = GetRefSkeleton().FindBoneIndex(Socket->BoneName);
 
 				if(BoneIndex != INDEX_NONE)
 				{
@@ -3187,7 +3202,7 @@ void USkeletalMesh::MoveDeprecatedShadowFlagToMaterials()
 	// First, the easy case where there's no LOD info (in which case, default to true!)
 	if ( LODInfo.Num() == 0 )
 	{
-		for ( auto Material = Materials.CreateIterator(); Material; ++Material )
+		for ( auto Material = GetMaterials().CreateIterator(); Material; ++Material )
 		{
 			Material->bEnableShadowCasting_DEPRECATED = true;
 		}
@@ -3223,7 +3238,7 @@ void USkeletalMesh::MoveDeprecatedShadowFlagToMaterials()
 	if ( !bDifferenceFound )
 	{
 		// All the same, so just copy the shadow casting flag to all materials
-		for ( auto Material = Materials.CreateIterator(); Material; ++Material )
+		for ( auto Material = GetMaterials().CreateIterator(); Material; ++Material )
 		{
 			Material->bEnableShadowCasting_DEPRECATED = PerLodShadowFlags.Num() ? PerLodShadowFlags[0] : true;
 		}
@@ -3234,6 +3249,7 @@ void USkeletalMesh::MoveDeprecatedShadowFlagToMaterials()
 		check( Resource->LODModels.Num() == LODInfo.Num() );
 
 		TArray<FSkeletalMaterial> NewMaterialArray;
+		TArray<FSkeletalMaterial>& CurrentMaterials = GetMaterials();
 
 		// There was a difference, so we need to build a new material list which has all the combinations of UMaterialInterface and shadow casting flag required
 		for ( int32 LODIndex = 0; LODIndex < Resource->LODModels.Num(); ++LODIndex )
@@ -3242,12 +3258,12 @@ void USkeletalMesh::MoveDeprecatedShadowFlagToMaterials()
 
 			for ( int32 i = 0; i < Resource->LODModels[LODIndex].Sections.Num(); ++i )
 			{
-				NewMaterialArray.Add( FSkeletalMaterial( Materials[ Resource->LODModels[LODIndex].Sections[i].MaterialIndex ].MaterialInterface, LODInfo[LODIndex].bEnableShadowCasting_DEPRECATED[i], false, NAME_None, NAME_None ) );
+				NewMaterialArray.Add( FSkeletalMaterial(CurrentMaterials[ Resource->LODModels[LODIndex].Sections[i].MaterialIndex ].MaterialInterface, LODInfo[LODIndex].bEnableShadowCasting_DEPRECATED[i], false, NAME_None, NAME_None ) );
 			}
 		}
 
 		// Reassign the materials array to the new one
-		Materials = NewMaterialArray;
+		SetMaterials(NewMaterialArray);
 		int32 NewIndex = 0;
 
 		// Remap the existing LODModels to point at the correct new material index
@@ -3272,17 +3288,18 @@ void USkeletalMesh::MoveMaterialFlagsToSections()
 		return;
 	}
 
+	TArray<FSkeletalMaterial>& CurrentMaterials = GetMaterials();
 	for (FSkeletalMeshLODModel &StaticLODModel : ImportedModel->LODModels)
 	{
 		for (int32 SectionIndex = 0; SectionIndex < StaticLODModel.Sections.Num(); ++SectionIndex)
 		{
 			FSkelMeshSection &Section = StaticLODModel.Sections[SectionIndex];
 			//Prior to FEditorObjectVersion::RefactorMeshEditorMaterials Material index match section index
-			if (Materials.IsValidIndex(SectionIndex))
+			if (CurrentMaterials.IsValidIndex(SectionIndex))
 			{
-				Section.bCastShadow = Materials[SectionIndex].bEnableShadowCasting_DEPRECATED;
+				Section.bCastShadow = CurrentMaterials[SectionIndex].bEnableShadowCasting_DEPRECATED;
 
-				Section.bRecomputeTangent = Materials[SectionIndex].bRecomputeTangent_DEPRECATED;
+				Section.bRecomputeTangent = CurrentMaterials[SectionIndex].bRecomputeTangent_DEPRECATED;
 			}
 			else
 			{
@@ -3429,9 +3446,9 @@ TArray<USkeletalMeshSocket*> USkeletalMesh::GetActiveSocketList() const
 	TArray<USkeletalMeshSocket*> ActiveSockets = Sockets;
 
 	// Then the skeleton sockets that aren't in the mesh
-	if (Skeleton)
+	if (GetSkeleton())
 	{
-		for (auto SkeletonSocketIt = Skeleton->Sockets.CreateConstIterator(); SkeletonSocketIt; ++SkeletonSocketIt)
+		for (auto SkeletonSocketIt = GetSkeleton()->Sockets.CreateConstIterator(); SkeletonSocketIt; ++SkeletonSocketIt)
 		{
 			USkeletalMeshSocket* Socket = *(SkeletonSocketIt);
 
@@ -3597,7 +3614,7 @@ void USkeletalMesh::ValidateBoneWeights(const ITargetPlatform* TargetPlatform)
 		FSkeletalMeshRenderData* SkelMeshRenderData = GetResourceForRendering();
 
 		int32 NumLODs = LODInfo.Num();
-		int32 MinFirstLOD = MinLod.GetValue();
+		int32 MinFirstLOD = GetMinLod().GetValue();
 		int32 MaxNumLODs = FMath::Clamp<int32>(NumLODs - MinFirstLOD, SkelMeshRenderData->NumInlinedLODs, NumLODs);
 
 		for (int32 LODIndex = 0; LODIndex < GetLODNum(); ++LODIndex)
@@ -3645,7 +3662,7 @@ FString USkeletalMesh::GetDerivedDataKey()
 
 int32 USkeletalMesh::ValidatePreviewAttachedObjects()
 {
-	int32 NumBrokenAssets = PreviewAttachedAssetContainer.ValidatePreviewAttachedObjects();
+	int32 NumBrokenAssets = GetPreviewAttachedAssetContainer().ValidatePreviewAttachedObjects();
 
 	if (NumBrokenAssets > 0)
 	{
@@ -3723,21 +3740,24 @@ void USkeletalMesh::ReleaseCPUResources()
 /** Allocate and initialise bone mirroring table for this skeletal mesh. Default is source = destination for each bone. */
 void USkeletalMesh::InitBoneMirrorInfo()
 {
-	SkelMirrorTable.Empty(RefSkeleton.GetNum());
-	SkelMirrorTable.AddZeroed(RefSkeleton.GetNum());
+	TArray<FBoneMirrorInfo>& LocalSkelMirrorTable = GetSkelMirrorTable();
+	LocalSkelMirrorTable.Empty(GetRefSkeleton().GetNum());
+	LocalSkelMirrorTable.AddZeroed(GetRefSkeleton().GetNum());
 
 	// By default, no bone mirroring, and source is ourself.
-	for(int32 i=0; i<SkelMirrorTable.Num(); i++)
+	for(int32 i=0; i< LocalSkelMirrorTable.Num(); i++)
 	{
-		SkelMirrorTable[i].SourceIndex = i;
+		LocalSkelMirrorTable[i].SourceIndex = i;
 	}
 }
 
 /** Utility for copying and converting a mirroring table from another SkeletalMesh. */
 void USkeletalMesh::CopyMirrorTableFrom(USkeletalMesh* SrcMesh)
 {
+	TArray<FBoneMirrorInfo>& SrcSkelMirrorTable = SrcMesh->GetSkelMirrorTable();
+	TArray<FBoneMirrorInfo>& LocalSkelMirrorTable = GetSkelMirrorTable();
 	// Do nothing if no mirror table in source mesh
-	if(SrcMesh->SkelMirrorTable.Num() == 0)
+	if(SrcSkelMirrorTable.Num() == 0)
 	{
 		return;
 	}
@@ -3747,36 +3767,36 @@ void USkeletalMesh::CopyMirrorTableFrom(USkeletalMesh* SrcMesh)
 
 	// Keep track of which entries in the source we have already copied
 	TArray<bool> EntryCopied;
-	EntryCopied.AddZeroed( SrcMesh->SkelMirrorTable.Num() );
+	EntryCopied.AddZeroed(SrcSkelMirrorTable.Num() );
 
 	// Mirror table must always be size of ref skeleton.
-	check(SrcMesh->SkelMirrorTable.Num() == SrcMesh->RefSkeleton.GetNum());
+	check(SrcSkelMirrorTable.Num() == SrcMesh->GetRefSkeleton().GetNum());
 
 	// Iterate over each entry in the source mesh mirror table.
 	// We assume that the src table is correct, and don't check for errors here (ie two bones using the same one as source).
-	for(int32 i=0; i<SrcMesh->SkelMirrorTable.Num(); i++)
+	for(int32 i=0; i< SrcSkelMirrorTable.Num(); i++)
 	{
 		if(!EntryCopied[i])
 		{
 			// Get name of source and dest bone for this entry in the source table.
-			FName DestBoneName = SrcMesh->RefSkeleton.GetBoneName(i);
-			int32 SrcBoneIndex = SrcMesh->SkelMirrorTable[i].SourceIndex;
-			FName SrcBoneName = SrcMesh->RefSkeleton.GetBoneName(SrcBoneIndex);
-			EAxis::Type FlipAxis = SrcMesh->SkelMirrorTable[i].BoneFlipAxis;
+			FName DestBoneName = SrcMesh->GetRefSkeleton().GetBoneName(i);
+			int32 SrcBoneIndex = SrcSkelMirrorTable[i].SourceIndex;
+			FName SrcBoneName = SrcMesh->GetRefSkeleton().GetBoneName(SrcBoneIndex);
+			EAxis::Type FlipAxis = SrcSkelMirrorTable[i].BoneFlipAxis;
 
 			// Look up bone names in target mesh (this one)
-			int32 DestBoneIndexTarget = RefSkeleton.FindBoneIndex(DestBoneName);
-			int32 SrcBoneIndexTarget = RefSkeleton.FindBoneIndex(SrcBoneName);
+			int32 DestBoneIndexTarget = GetRefSkeleton().FindBoneIndex(DestBoneName);
+			int32 SrcBoneIndexTarget = GetRefSkeleton().FindBoneIndex(SrcBoneName);
 
 			// If both bones found, copy data to this mesh's mirror table.
 			if( DestBoneIndexTarget != INDEX_NONE && SrcBoneIndexTarget != INDEX_NONE )
 			{
-				SkelMirrorTable[DestBoneIndexTarget].SourceIndex = SrcBoneIndexTarget;
-				SkelMirrorTable[DestBoneIndexTarget].BoneFlipAxis = FlipAxis;
+				LocalSkelMirrorTable[DestBoneIndexTarget].SourceIndex = SrcBoneIndexTarget;
+				LocalSkelMirrorTable[DestBoneIndexTarget].BoneFlipAxis = FlipAxis;
 
 
-				SkelMirrorTable[SrcBoneIndexTarget].SourceIndex = DestBoneIndexTarget;
-				SkelMirrorTable[SrcBoneIndexTarget].BoneFlipAxis = FlipAxis;
+				LocalSkelMirrorTable[SrcBoneIndexTarget].SourceIndex = DestBoneIndexTarget;
+				LocalSkelMirrorTable[SrcBoneIndexTarget].BoneFlipAxis = FlipAxis;
 
 				// Flag entries as copied, so we don't try and do it again.
 				EntryCopied[i] = true;
@@ -3789,25 +3809,26 @@ void USkeletalMesh::CopyMirrorTableFrom(USkeletalMesh* SrcMesh)
 /** Utility for copying and converting a mirroring table from another SkeletalMesh. */
 void USkeletalMesh::ExportMirrorTable(TArray<FBoneMirrorExport> &MirrorExportInfo) const
 {
+	const TArray<FBoneMirrorInfo>& LocalSkelMirrorTable = GetSkelMirrorTable();
 	// Do nothing if no mirror table in source mesh
-	if( SkelMirrorTable.Num() == 0 )
+	if(LocalSkelMirrorTable.Num() == 0 )
 	{
 		return;
 	}
 	
 	// Mirror table must always be size of ref skeleton.
-	check(SkelMirrorTable.Num() == RefSkeleton.GetNum());
+	check(LocalSkelMirrorTable.Num() == GetRefSkeleton().GetNum());
 
-	MirrorExportInfo.Empty(SkelMirrorTable.Num());
-	MirrorExportInfo.AddZeroed(SkelMirrorTable.Num());
+	MirrorExportInfo.Empty(LocalSkelMirrorTable.Num());
+	MirrorExportInfo.AddZeroed(LocalSkelMirrorTable.Num());
 
 	// Iterate over each entry in the source mesh mirror table.
 	// We assume that the src table is correct, and don't check for errors here (ie two bones using the same one as source).
-	for(int32 i=0; i<SkelMirrorTable.Num(); i++)
+	for(int32 i=0; i< LocalSkelMirrorTable.Num(); i++)
 	{
-		MirrorExportInfo[i].BoneName		= RefSkeleton.GetBoneName(i);
-		MirrorExportInfo[i].SourceBoneName	= RefSkeleton.GetBoneName(SkelMirrorTable[i].SourceIndex);
-		MirrorExportInfo[i].BoneFlipAxis	= SkelMirrorTable[i].BoneFlipAxis;
+		MirrorExportInfo[i].BoneName		= GetRefSkeleton().GetBoneName(i);
+		MirrorExportInfo[i].SourceBoneName	= GetRefSkeleton().GetBoneName(LocalSkelMirrorTable[i].SourceIndex);
+		MirrorExportInfo[i].BoneFlipAxis	= LocalSkelMirrorTable[i].BoneFlipAxis;
 	}
 }
 
@@ -3826,32 +3847,33 @@ void USkeletalMesh::ImportMirrorTable(const TArray<FBoneMirrorExport> &MirrorExp
 
 	// Keep track of which entries in the source we have already copied
 	TArray<bool> EntryCopied;
-	EntryCopied.AddZeroed( RefSkeleton.GetNum() );
+	EntryCopied.AddZeroed(GetRefSkeleton().GetNum() );
 
+	TArray<FBoneMirrorInfo>& LocalSkelMirrorTable = GetSkelMirrorTable();
 	// Mirror table must always be size of ref skeleton.
-	check(SkelMirrorTable.Num() == RefSkeleton.GetNum());
+	check(LocalSkelMirrorTable.Num() == GetRefSkeleton().GetNum());
 
 	// Iterate over each entry in the source mesh mirror table.
 	// We assume that the src table is correct, and don't check for errors here (ie two bones using the same one as source).
 	for(int32 i=0; i<MirrorExportInfo.Num(); i++)
 	{
 		FName DestBoneName	= MirrorExportInfo[i].BoneName;
-		int32 DestBoneIndex	= RefSkeleton.FindBoneIndex(DestBoneName);
+		int32 DestBoneIndex	= GetRefSkeleton().FindBoneIndex(DestBoneName);
 
 		if( DestBoneIndex != INDEX_NONE && !EntryCopied[DestBoneIndex] )
 		{
 			FName SrcBoneName	= MirrorExportInfo[i].SourceBoneName;
-			int32 SrcBoneIndex	= RefSkeleton.FindBoneIndex(SrcBoneName);
+			int32 SrcBoneIndex	= GetRefSkeleton().FindBoneIndex(SrcBoneName);
 			EAxis::Type FlipAxis		= MirrorExportInfo[i].BoneFlipAxis;
 
 			// If both bones found, copy data to this mesh's mirror table.
 			if( SrcBoneIndex != INDEX_NONE )
 			{
-				SkelMirrorTable[DestBoneIndex].SourceIndex = SrcBoneIndex;
-				SkelMirrorTable[DestBoneIndex].BoneFlipAxis = FlipAxis;
+				LocalSkelMirrorTable[DestBoneIndex].SourceIndex = SrcBoneIndex;
+				LocalSkelMirrorTable[DestBoneIndex].BoneFlipAxis = FlipAxis;
 
-				SkelMirrorTable[SrcBoneIndex].SourceIndex = DestBoneIndex;
-				SkelMirrorTable[SrcBoneIndex].BoneFlipAxis = FlipAxis;
+				LocalSkelMirrorTable[SrcBoneIndex].SourceIndex = DestBoneIndex;
+				LocalSkelMirrorTable[SrcBoneIndex].BoneFlipAxis = FlipAxis;
 
 				// Flag entries as copied, so we don't try and do it again.
 				EntryCopied[DestBoneIndex]	= true;
@@ -3870,10 +3892,12 @@ bool USkeletalMesh::MirrorTableIsGood(FString& ProblemBones) const
 {
 	TArray<int32>	BadBoneMirror;
 
-	for(int32 i=0; i<SkelMirrorTable.Num(); i++)
+	const TArray<FBoneMirrorInfo>& LocalSkelMirrorTable = GetSkelMirrorTable();
+
+	for(int32 i=0; i< LocalSkelMirrorTable.Num(); i++)
 	{
-		int32 SrcIndex = SkelMirrorTable[i].SourceIndex;
-		if( SkelMirrorTable[SrcIndex].SourceIndex != i)
+		int32 SrcIndex = LocalSkelMirrorTable[i].SourceIndex;
+		if(LocalSkelMirrorTable[SrcIndex].SourceIndex != i)
 		{
 			BadBoneMirror.Add(i);
 		}
@@ -3884,7 +3908,7 @@ bool USkeletalMesh::MirrorTableIsGood(FString& ProblemBones) const
 		for(int32 i=0; i<BadBoneMirror.Num(); i++)
 		{
 			int32 BoneIndex = BadBoneMirror[i];
-			FName BoneName = RefSkeleton.GetBoneName(BoneIndex);
+			FName BoneName = GetRefSkeleton().GetBoneName(BoneIndex);
 
 			ProblemBones += FString::Printf( TEXT("%s (%d)\n"), *BoneName.ToString(), BoneIndex );
 		}
@@ -3899,33 +3923,31 @@ bool USkeletalMesh::MirrorTableIsGood(FString& ProblemBones) const
 
 void USkeletalMesh::CreateBodySetup()
 {
+	PRAGMA_DISABLE_DEPRECATION_WARNINGS
 	if (BodySetup == nullptr)
 	{
-		BodySetup = NewObject<UBodySetup>(this);
+		SetBodySetup(NewObject<UBodySetup>(this));
 		BodySetup->bSharedCookedData = true;
 		BodySetup->AddToCluster(this);
 	}
-}
-
-UBodySetup* USkeletalMesh::GetBodySetup()
-{
-	CreateBodySetup();
-	return BodySetup;
+	PRAGMA_ENABLE_DEPRECATION_WARNINGS
 }
 
 #if WITH_EDITOR
 void USkeletalMesh::BuildPhysicsData()
 {
 	CreateBodySetup();
-	BodySetup->CookedFormatData.FlushData();	//we need to force a re-cook because we're essentially re-creating the bodysetup so that it swaps whether or not it has a trimesh
-	BodySetup->InvalidatePhysicsData();
-	BodySetup->CreatePhysicsMeshes();
+	const USkeletalMesh* ConstThis = this;
+	UBodySetup* LocalBodySetup = ConstThis->GetBodySetup();
+	LocalBodySetup->CookedFormatData.FlushData();	//we need to force a re-cook because we're essentially re-creating the bodysetup so that it swaps whether or not it has a trimesh
+	LocalBodySetup->InvalidatePhysicsData();
+	LocalBodySetup->CreatePhysicsMeshes();
 }
 #endif
 
 bool USkeletalMesh::ContainsPhysicsTriMeshData(bool InUseAllTriData) const
 {
-	return bEnablePerPolyCollision;
+	return GetEnablePerPolyCollision();
 }
 
 bool USkeletalMesh::GetPhysicsTriMeshData(FTriMeshCollisionData* CollisionData, bool bInUseAllTriData)
@@ -3933,7 +3955,7 @@ bool USkeletalMesh::GetPhysicsTriMeshData(FTriMeshCollisionData* CollisionData,
 #if WITH_EDITORONLY_DATA
 
 	// Fail if no mesh or not per poly collision
-	if (!GetResourceForRendering() || !bEnablePerPolyCollision)
+	if (!GetResourceForRendering() || !GetEnablePerPolyCollision())
 	{
 		return false;
 	}
@@ -4048,7 +4070,7 @@ FString USkeletalMesh::GetDesc()
 	if (Resource)
 	{
 		check(Resource->LODRenderData.Num() > 0);
-		DescString = FString::Printf(TEXT("%d Triangles, %d Bones"), Resource->LODRenderData[0].GetTotalFaces(), RefSkeleton.GetRawBoneNum());
+		DescString = FString::Printf(TEXT("%d Triangles, %d Bones"), Resource->LODRenderData[0].GetTotalFaces(), GetRefSkeleton().GetRawBoneNum());
 	}
 	return DescString;
 }
@@ -4114,7 +4136,7 @@ void USkeletalMesh::ConvertLegacyLODScreenSize()
 		{
 			FSkeletalMeshLODInfo& LODInfoEntry = LODInfo[LODIndex];
 
-			if (bRequiresLODScreenSizeConversion)
+			if (GetRequiresLODScreenSizeConversion())
 			{
 				if (LODInfoEntry.ScreenSize.Default == 0.0f)
 				{
@@ -4130,7 +4152,7 @@ void USkeletalMesh::ConvertLegacyLODScreenSize()
 				}
 			}
 
-			if (bRequiresLODHysteresisConversion)
+			if (GetRequiresLODHysteresisConversion())
 			{
 				if (LODInfoEntry.LODHysteresis != 0.0f)
 				{
@@ -4146,9 +4168,10 @@ void USkeletalMesh::ConvertLegacyLODScreenSize()
 
 class UNodeMappingContainer* USkeletalMesh::GetNodeMappingContainer(class UBlueprint* SourceAsset) const
 {
-	for (int32 Index = 0; Index < NodeMappingData.Num(); ++Index)
+	const TArray<UNodeMappingContainer*>& LocalNodeMappingData = GetNodeMappingData();
+	for (int32 Index = 0; Index < LocalNodeMappingData.Num(); ++Index)
 	{
-		UNodeMappingContainer* Iter = NodeMappingData[Index];
+		UNodeMappingContainer* Iter = LocalNodeMappingData[Index];
 		if (Iter && Iter->GetSourceAssetSoftObjectPtr() == TSoftObjectPtr<UObject>(SourceAsset))
 		{
 			return Iter;
@@ -4163,13 +4186,13 @@ const UAnimSequence* USkeletalMesh::GetBakePose(int32 LODIndex) const
 	const FSkeletalMeshLODInfo* LOD = GetLODInfo(LODIndex);
 	if (LOD)
 	{
-		if (LOD->BakePoseOverride && Skeleton == LOD->BakePoseOverride->GetSkeleton())
+		if (LOD->BakePoseOverride && GetSkeleton() == LOD->BakePoseOverride->GetSkeleton())
 		{
 			return LOD->BakePoseOverride;
 		}
 
 		// we make sure bake pose uses same skeleton
-		if (LOD->BakePose && Skeleton == LOD->BakePose->GetSkeleton())
+		if (LOD->BakePose && GetSkeleton() == LOD->BakePose->GetSkeleton())
 		{
 			return LOD->BakePose;
 		}
@@ -4181,9 +4204,9 @@ const UAnimSequence* USkeletalMesh::GetBakePose(int32 LODIndex) const
 const USkeletalMeshLODSettings* USkeletalMesh::GetDefaultLODSetting() const
 { 
 #if WITH_EDITORONLY_DATA
-	if (LODSettings)
+	if (GetLODSettings())
 	{
-		return LODSettings;
+		return GetLODSettings();
 	}
 #endif // WITH_EDITORONLY_DATA
 
@@ -4252,9 +4275,11 @@ void USkeletalMesh::RemoveLODInfo(int32 Index)
 	if (LODInfo.IsValidIndex(Index))
 	{
 #if WITH_EDITOR
-		if (MeshEditorDataObject != nullptr)
+		if (IsMeshEditorDataValid())
 		{
+			PRAGMA_DISABLE_DEPRECATION_WARNINGS
 			MeshEditorDataObject->RemoveLODImportedData(Index);
+			PRAGMA_ENABLE_DEPRECATION_WARNINGS
 		}
 		if (GetImportedModel() && GetImportedModel()->OriginalReductionSourceMeshData.IsValidIndex(Index))
 		{
@@ -4278,9 +4303,9 @@ bool USkeletalMesh::GetSupportsLODStreaming(const ITargetPlatform* TargetPlatfor
 	{
 		return false;
 	}
-	if (bOverrideLODStreamingSettings)
+	if (GetOverrideLODStreamingSettings())
 	{
-		return bSupportLODStreaming.GetValueForPlatform(*TargetPlatform->IniPlatformName());
+		return GetSupportLODStreaming().GetValueForPlatform(*TargetPlatform->IniPlatformName());
 	}
 	else
 	{
@@ -4291,9 +4316,9 @@ bool USkeletalMesh::GetSupportsLODStreaming(const ITargetPlatform* TargetPlatfor
 int32 USkeletalMesh::GetMaxNumStreamedLODs(const ITargetPlatform* TargetPlatform) const
 {
 	check(TargetPlatform);
-	if (bOverrideLODStreamingSettings)
+	if (GetOverrideLODStreamingSettings())
 	{
-		return MaxNumStreamedLODs.GetValueForPlatform(*TargetPlatform->IniPlatformName());
+		return GetMaxNumStreamedLODs().GetValueForPlatform(*TargetPlatform->IniPlatformName());
 	}
 	else
 	{
@@ -4304,9 +4329,9 @@ int32 USkeletalMesh::GetMaxNumStreamedLODs(const ITargetPlatform* TargetPlatform
 int32 USkeletalMesh::GetMaxNumOptionalLODs(const ITargetPlatform* TargetPlatform) const
 {
 	check(TargetPlatform);
-	if (bOverrideLODStreamingSettings)
+	if (GetOverrideLODStreamingSettings())
 	{
-		return MaxNumOptionalLODs.GetValueForPlatform(*TargetPlatform->IniPlatformName()) <= 0 ? 0 : MAX_MESH_LOD_COUNT;
+		return GetMaxNumOptionalLODs().GetValueForPlatform(*TargetPlatform->IniPlatformName()) <= 0 ? 0 : MAX_MESH_LOD_COUNT;
 	}
 	else
 	{
@@ -4317,6 +4342,7 @@ int32 USkeletalMesh::GetMaxNumOptionalLODs(const ITargetPlatform* TargetPlatform
 
 void USkeletalMesh::SetLODSettings(USkeletalMeshLODSettings* InLODSettings)
 {
+PRAGMA_DISABLE_DEPRECATION_WARNINGS
 #if WITH_EDITORONLY_DATA
 	LODSettings = InLODSettings;
 	if (LODSettings)
@@ -4324,22 +4350,27 @@ void USkeletalMesh::SetLODSettings(USkeletalMeshLODSettings* InLODSettings)
 		LODSettings->SetLODSettingsToMesh(this);
 	}
 #endif // WITH_EDITORONLY_DATA
+PRAGMA_ENABLE_DEPRECATION_WARNINGS
 }
 
 void USkeletalMesh::SetDefaultAnimatingRig(TSoftObjectPtr<UObject> InAnimatingRig)
 {
+	PRAGMA_DISABLE_DEPRECATION_WARNINGS
 #if WITH_EDITORONLY_DATA
 	DefaultAnimatingRig = InAnimatingRig;
 #endif // WITH_EDITORONLY_DATA
+	PRAGMA_ENABLE_DEPRECATION_WARNINGS
 }
 
-TSoftObjectPtr<UObject> USkeletalMesh::GetDefaultAnimatingRig()
+TSoftObjectPtr<UObject> USkeletalMesh::GetDefaultAnimatingRig() const
 {
+	PRAGMA_DISABLE_DEPRECATION_WARNINGS
 #if WITH_EDITORONLY_DATA
 	return DefaultAnimatingRig;
 #else // WITH_EDITORONLY_DATA
 	return nullptr;
 #endif // WITH_EDITORONLY_DATA
+	PRAGMA_ENABLE_DEPRECATION_WARNINGS
 }
 
 void USkeletalMesh::GetMappableNodeData(TArray<FName>& OutNames, TArray<FNodeItem>& OutNodeItems) const
@@ -4350,10 +4381,10 @@ void USkeletalMesh::GetMappableNodeData(TArray<FName>& OutNames, TArray<FNodeIte
 #else
 	// hasn't tested this route, but we don't have retarget base pose if not editor, wonder we should to non-editor soon
 	ensure(false);
-	FAnimationRuntime::FillUpComponentSpaceTransforms(RefSkeleton, RefSkeleton.GetRefBonePose(), ComponentSpaceRefPose);
+	FAnimationRuntime::FillUpComponentSpaceTransforms(GetRefSkeleton(), GetRefSkeleton().GetRefBonePose(), ComponentSpaceRefPose);
 #endif // 
 
-	const int32 NumJoint = RefSkeleton.GetNum();
+	const int32 NumJoint = GetRefSkeleton().GetNum();
 	// allocate buffer
 	OutNames.Reset(NumJoint);
 	OutNodeItems.Reset(NumJoint);
@@ -4363,7 +4394,7 @@ void USkeletalMesh::GetMappableNodeData(TArray<FName>& OutNames, TArray<FNodeIte
 		OutNames.AddDefaulted(NumJoint);
 		OutNodeItems.AddDefaulted(NumJoint);
 
-		const TArray<FMeshBoneInfo> MeshBoneInfo = RefSkeleton.GetRefBoneInfo();
+		const TArray<FMeshBoneInfo> MeshBoneInfo = GetRefSkeleton().GetRefBoneInfo();
 		for (int32 NodeIndex = 0; NodeIndex < NumJoint; ++NodeIndex)
 		{
 			OutNames[NodeIndex] = MeshBoneInfo[NodeIndex].Name;
@@ -4501,7 +4532,7 @@ FGuid FSkeletalMeshLODInfo::ComputeDeriveDataCacheKey(const FSkeletalMeshLODGrou
 TArray<FString> USkeletalMesh::K2_GetAllMorphTargetNames() const
 {
 	TArray<FString> Names;
-	for (UMorphTarget* MorphTarget : MorphTargets)
+	for (UMorphTarget* MorphTarget : GetMorphTargets())
 	{
 		Names.Add(MorphTarget->GetFName().ToString());
 	}
@@ -4776,10 +4807,10 @@ FSkeletalMeshSceneProxy::FSkeletalMeshSceneProxy(const USkinnedMeshComponent* Co
 			// If we are at a dropped LOD, route material index through the LODMaterialMap in the LODInfo struct.
 			int32 UseMaterialIndex = Section.MaterialIndex;			
 			{
-				if(SectionIndex < Info.LODMaterialMap.Num() && Component->SkeletalMesh->Materials.IsValidIndex(Info.LODMaterialMap[SectionIndex]))
+				if(SectionIndex < Info.LODMaterialMap.Num() && Component->SkeletalMesh->GetMaterials().IsValidIndex(Info.LODMaterialMap[SectionIndex]))
 				{
 					UseMaterialIndex = Info.LODMaterialMap[SectionIndex];
-					UseMaterialIndex = FMath::Clamp( UseMaterialIndex, 0, Component->SkeletalMesh->Materials.Num() );
+					UseMaterialIndex = FMath::Clamp( UseMaterialIndex, 0, Component->SkeletalMesh->GetMaterials().Num() );
 				}
 			}
 
@@ -4825,7 +4856,7 @@ FSkeletalMeshSceneProxy::FSkeletalMeshSceneProxy(const USkinnedMeshComponent* Co
 			}
 
 			bool bSectionCastsShadow = !bSectionHidden && bCastShadow &&
-				(Component->SkeletalMesh->Materials.IsValidIndex(UseMaterialIndex) == false || Section.bCastShadow);
+				(Component->SkeletalMesh->GetMaterials().IsValidIndex(UseMaterialIndex) == false || Section.bCastShadow);
 
 			bAnySectionCastsShadow |= bSectionCastsShadow;
 
@@ -4864,7 +4895,7 @@ FSkeletalMeshSceneProxy::FSkeletalMeshSceneProxy(const USkinnedMeshComponent* Co
 	// Copy out shadow physics asset data
 	if(SkinnedMeshComponent)
 	{
-		UPhysicsAsset* ShadowPhysicsAsset = SkinnedMeshComponent->SkeletalMesh->ShadowPhysicsAsset;
+		UPhysicsAsset* ShadowPhysicsAsset = SkinnedMeshComponent->SkeletalMesh->GetShadowPhysicsAsset();
 
 		if (ShadowPhysicsAsset
 			&& SkinnedMeshComponent->CastShadow
@@ -5688,7 +5719,7 @@ void FSkeletalMeshSceneProxy::DebugDrawSkeleton(int32 ViewIndex, FMeshElementCol
 
 	for (int32 Index = 0; Index < ComponentSpaceTransforms.Num(); ++Index)
 	{
-		const int32 ParentIndex = SkeletalMeshForDebug->RefSkeleton.GetParentIndex(Index);
+		const int32 ParentIndex = SkeletalMeshForDebug->GetRefSkeleton().GetParentIndex(Index);
 		FVector Start, End;
 		
 		FLinearColor LineColor = DebugDrawColor.Get(MakeRandomColorForSkeleton(GetPrimitiveComponentId().PrimIDValue));
@@ -5938,7 +5969,7 @@ FVector GetRefVertexLocationTyped(
 		const float	Weight = (float)SkinWeightVertexBuffer.GetBoneWeight(BufferVertIndex, InfluenceIndex) / 255.0f;
 		{
 			const FMatrix BoneTransformMatrix = FMatrix::Identity;//Mesh->GetComposedRefPoseMatrix(MeshBoneIndex);
-			const FMatrix RefToLocal = Mesh->RefBasesInvMatrix[MeshBoneIndex] * BoneTransformMatrix;
+			const FMatrix RefToLocal = Mesh->GetRefBasesInvMatrix()[MeshBoneIndex] * BoneTransformMatrix;
 
 			SkinnedPos += BoneTransformMatrix.TransformPosition(PositionBuffer.VertexPosition(BufferVertIndex)) * Weight;
 		}
diff --git a/Engine/Source/Runtime/Engine/Private/SkeletalMeshComponentPhysics.cpp b/Engine/Source/Runtime/Engine/Private/SkeletalMeshComponentPhysics.cpp
--- a/Engine/Source/Runtime/Engine/Private/SkeletalMeshComponentPhysics.cpp
+++ b/Engine/Source/Runtime/Engine/Private/SkeletalMeshComponentPhysics.cpp
@@ -142,7 +142,8 @@ void USkeletalMeshComponent::CreateBodySetup()
 
 	if (SkeletalMesh)
 	{
-		UBodySetup* OriginalBodySetup = SkeletalMesh->GetBodySetup();
+		const USkeletalMesh* SkeletalMeshConst = SkeletalMesh;
+		UBodySetup* OriginalBodySetup = SkeletalMeshConst->GetBodySetup();
 		BodySetup->CopyBodyPropertiesFrom(OriginalBodySetup);
 		BodySetup->CookedFormatDataOverride = &OriginalBodySetup->CookedFormatData;
 	}
@@ -164,9 +165,9 @@ UBodySetup* USkeletalMeshComponent::GetBodySetup()
 		UPhysicsAsset * const PhysicsAsset = GetPhysicsAsset();
 		if (SkeletalMesh && PhysicsAsset)
 		{
-			for (int32 i = 0; i < SkeletalMesh->RefSkeleton.GetNum(); i++)
+			for (int32 i = 0; i < SkeletalMesh->GetRefSkeleton().GetNum(); i++)
 			{
-				int32 BodyIndex = PhysicsAsset->FindBodyIndex(SkeletalMesh->RefSkeleton.GetBoneName(i));
+				int32 BodyIndex = PhysicsAsset->FindBodyIndex(SkeletalMesh->GetRefSkeleton().GetBoneName(i));
 				if (BodyIndex != INDEX_NONE)
 				{
 					return PhysicsAsset->SkeletalBodySetups[BodyIndex];
@@ -633,9 +634,9 @@ int32 USkeletalMeshComponent::FindRootBodyIndex() const
 	{
 		if(const UPhysicsAsset* PhysicsAsset = GetPhysicsAsset())
 		{
-			for (int32 i = 0; i<SkeletalMesh->RefSkeleton.GetNum(); i++)
+			for (int32 i = 0; i<SkeletalMesh->GetRefSkeleton().GetNum(); i++)
 			{
-				int32 BodyInstIndex = PhysicsAsset->FindBodyIndex(SkeletalMesh->RefSkeleton.GetBoneName(i));
+				int32 BodyInstIndex = PhysicsAsset->FindBodyIndex(SkeletalMesh->GetRefSkeleton().GetBoneName(i));
 				if (BodyInstIndex != INDEX_NONE)
 				{
 					RootBodyIndex = BodyInstIndex;
@@ -749,7 +750,7 @@ void USkeletalMeshComponent::InstantiatePhysicsAssetRefPose(const UPhysicsAsset&
 {
 	if(SkeletalMesh)
 	{
-		const FReferenceSkeleton& RefSkeleton = SkeletalMesh->RefSkeleton;
+		const FReferenceSkeleton& RefSkeleton = SkeletalMesh->GetRefSkeleton();
 
 		if (bCreateBodiesInRefPose)
 		{
@@ -1015,7 +1016,7 @@ void USkeletalMeshComponent::TermBodiesBelow(FName ParentBoneName)
 			int32 JointBoneIndex = GetBoneIndex(JointName);
 
 			// If constraint has bone in mesh, and is either the parent or child of it, term it
-			if(	JointBoneIndex != INDEX_NONE && (JointName == ParentBoneName ||	SkeletalMesh->RefSkeleton.BoneIsChildOf(JointBoneIndex, ParentBoneIndex)) )
+			if(	JointBoneIndex != INDEX_NONE && (JointName == ParentBoneName ||	SkeletalMesh->GetRefSkeleton().BoneIsChildOf(JointBoneIndex, ParentBoneIndex)) )
 			{
 				Constraints[i]->TermConstraint();
 			}
@@ -1031,7 +1032,7 @@ void USkeletalMeshComponent::TermBodiesBelow(FName ParentBoneName)
 				int32 BodyBoneIndex = GetBoneIndex(BodyName);
 
 				// If body has bone in mesh, and is either the parent or child of it, term it
-				if(	BodyBoneIndex != INDEX_NONE && (BodyName == ParentBoneName ||	SkeletalMesh->RefSkeleton.BoneIsChildOf(BodyBoneIndex, ParentBoneIndex)) )
+				if(	BodyBoneIndex != INDEX_NONE && (BodyName == ParentBoneName ||	SkeletalMesh->GetRefSkeleton().BoneIsChildOf(BodyBoneIndex, ParentBoneIndex)) )
 				{
 					Bodies[i]->TermBody();
 				}
@@ -1604,7 +1605,7 @@ void USkeletalMeshComponent::UpdateMeshForBrokenConstraints()
 				UBodySetup* PhysicsAssetBodySetup = PhysicsAsset->SkeletalBodySetups[BodySetupIndex];
 				int32 BoneIndex = GetBoneIndex(PhysicsAssetBodySetup->BoneName);
 				if( BoneIndex != INDEX_NONE && 
-					(BoneIndex == JointBoneIndex || SkeletalMesh->RefSkeleton.BoneIsChildOf(BoneIndex, JointBoneIndex)) )
+					(BoneIndex == JointBoneIndex || SkeletalMesh->GetRefSkeleton().BoneIsChildOf(BoneIndex, JointBoneIndex)) )
 				{
 					DEBUGBROKENCONSTRAINTUPDATE(UE_LOG(LogSkeletalMesh, Log, TEXT("    Found Child Bone: (%d) %s"), BoneIndex, *PhysicsAssetBodySetup->BoneName.ToString());)
 
@@ -1926,9 +1927,9 @@ void USkeletalMeshComponent::SetPhysicsAsset(UPhysicsAsset* InPhysicsAsset, bool
 		if( SkeletalMesh )
 		{
 			// Because we don't know what bones the new PhysicsAsset might want, we have to force an update to _all_ bones in the skeleton.
-			RequiredBones.Reset(SkeletalMesh->RefSkeleton.GetNum());
-			RequiredBones.AddUninitialized( SkeletalMesh->RefSkeleton.GetNum() );
-			for(int32 i=0; i<SkeletalMesh->RefSkeleton.GetNum(); i++)
+			RequiredBones.Reset(SkeletalMesh->GetRefSkeleton().GetNum());
+			RequiredBones.AddUninitialized( SkeletalMesh->GetRefSkeleton().GetNum() );
+			for(int32 i=0; i<SkeletalMesh->GetRefSkeleton().GetNum(); i++)
 			{
 				RequiredBones[i] = (FBoneIndexType)i;
 			}
@@ -2046,7 +2047,7 @@ FVector GetTypedSkinnedVertexPositionWithCloth(USkeletalMeshComponent* Component
 
 	// only if this component has clothing and is showing simulated results	
 	if (Component->SkeletalMesh &&
-		Component->SkeletalMesh->MeshClothingAssets.Num() > 0 &&
+		Component->SkeletalMesh->GetMeshClothingAssets().Num() > 0 &&
 		!Component->bDisableClothSimulation &&
 		Component->ClothBlendWeight > 0.0f // if cloth blend weight is 0.0, only showing skinned vertices regardless of simulation positions
 		)
@@ -2106,7 +2107,7 @@ void USkeletalMeshComponent::ComputeSkinnedPositions(USkeletalMeshComponent* Com
 		return;
 	}
 
-	if (Component->SkeletalMesh->MeshClothingAssets.Num() > 0 &&
+	if (Component->SkeletalMesh->GetMeshClothingAssets().Num() > 0 &&
 		!Component->bDisableClothSimulation &&
 		Component->ClothBlendWeight > 0.0f // if cloth blend weight is 0.0, only showing skinned vertices regardless of simulation positions
 		)
@@ -2315,7 +2316,7 @@ bool USkeletalMeshComponent::GetClosestPointOnPhysicsAsset(const FVector& WorldP
 
 	bool bSuccess = false;
 	const UPhysicsAsset* PhysicsAsset = GetPhysicsAsset();
-	const FReferenceSkeleton* RefSkeleton = SkeletalMesh ? &SkeletalMesh->RefSkeleton : nullptr;
+	const FReferenceSkeleton* RefSkeleton = SkeletalMesh ? &SkeletalMesh->GetRefSkeleton() : nullptr;
 	if(PhysicsAsset && RefSkeleton)
 	{
 		const TArray<FTransform>& BoneTransforms = GetComponentSpaceTransforms();
@@ -2532,7 +2533,7 @@ void USkeletalMeshComponent::RecreateClothingActors()
 		return;
 	}
 
-	if(SkeletalMesh->MeshClothingAssets.Num() > 0)
+	if(SkeletalMesh->GetMeshClothingAssets().Num() > 0)
 	{
 		UClass* SimFactoryClass = *ClothingSimulationFactory;
 		if (SimFactoryClass)
@@ -2570,10 +2571,10 @@ void USkeletalMeshComponent::RecreateClothingActors()
 
 				ClothingSimulation->FillContext(this, 0.f, ClothingSimulationContext);
 
-				const int32 NumMeshAssets = SkeletalMesh->MeshClothingAssets.Num();
+				const int32 NumMeshAssets = SkeletalMesh->GetMeshClothingAssets().Num();
 				for (int32 BaseAssetIndex = 0; BaseAssetIndex < NumMeshAssets; ++BaseAssetIndex)
 				{
-					UClothingAssetBase* const Asset = SkeletalMesh->MeshClothingAssets[BaseAssetIndex];
+					UClothingAssetBase* const Asset = SkeletalMesh->GetMeshClothingAssets()[BaseAssetIndex];
 					if (Asset && AssetsInUse.Contains(Asset))
 					{
 						ClothingSimulation->CreateActor(this, Asset, BaseAssetIndex);
@@ -2688,7 +2689,7 @@ void USkeletalMeshComponent::ExtractCollisionsForCloth(
 			for(const USkeletalBodySetup* SkeletalBodySetup : PhysicsAsset->SkeletalBodySetups)
 			{
 				// Cache bones
-				int32 MeshBoneIndex = SourceComponent->SkeletalMesh->RefSkeleton.FindBoneIndex(SkeletalBodySetup->BoneName);
+				int32 MeshBoneIndex = SourceComponent->SkeletalMesh->GetRefSkeleton().FindBoneIndex(SkeletalBodySetup->BoneName);
 				if(MeshBoneIndex != INDEX_NONE)
 				{
 					// Cache spheres & capsules form physics asset
diff --git a/Engine/Source/Runtime/Engine/Private/SkeletalMeshLODModel.cpp b/Engine/Source/Runtime/Engine/Private/SkeletalMeshLODModel.cpp
--- a/Engine/Source/Runtime/Engine/Private/SkeletalMeshLODModel.cpp
+++ b/Engine/Source/Runtime/Engine/Private/SkeletalMeshLODModel.cpp
@@ -817,7 +817,7 @@ void FSkeletalMeshLODModel::Serialize(FArchive& Ar, UObject* Owner, int32 Idx)
 			}
 
 			USkeletalMesh* SkelMeshOwner = CastChecked<USkeletalMesh>(Owner);
-			if (SkelMeshOwner->bHasVertexColors)
+			if (SkelMeshOwner->GetHasVertexColors())
 			{
 				// Handling for old color buffer data
 				if (Ar.IsLoading() && Ar.CustomVer(FSkeletalMeshCustomVersion::GUID) < FSkeletalMeshCustomVersion::UseSharedColorBufferFormat)
diff --git a/Engine/Source/Runtime/Engine/Private/SkeletalMeshLODRenderData.cpp b/Engine/Source/Runtime/Engine/Private/SkeletalMeshLODRenderData.cpp
--- a/Engine/Source/Runtime/Engine/Private/SkeletalMeshLODRenderData.cpp
+++ b/Engine/Source/Runtime/Engine/Private/SkeletalMeshLODRenderData.cpp
@@ -478,7 +478,7 @@ int32 FSkeletalMeshLODRenderData::GetPlatformMinLODIdx(const ITargetPlatform* Ta
 #if WITH_EDITOR
 	check(TargetPlatform && SkeletalMesh);
 	const FName IniPlatformName = TargetPlatform->GetPlatformInfo().IniPlatformName;
-	return  SkeletalMesh->MinLod.GetValueForPlatform(IniPlatformName);
+	return  SkeletalMesh->GetMinLod().GetValueForPlatform(IniPlatformName);
 #else
 	return 0;
 #endif
@@ -497,8 +497,8 @@ uint8 FSkeletalMeshLODRenderData::GenerateClassStripFlags(FArchive& Ar, const US
 	bool bMeshDisablesMinLodStrip = false;
 	if (bIsCook)
 	{
-		MinMeshLod = OwnerMesh ? OwnerMesh->MinLod.GetValueForPlatform(CookTarget->GetPlatformInfo().IniPlatformName) : 0;
-		bMeshDisablesMinLodStrip = OwnerMesh ? OwnerMesh->DisableBelowMinLodStripping.GetValueForPlatform(CookTarget->GetPlatformInfo().IniPlatformName) : false;
+		MinMeshLod = OwnerMesh ? OwnerMesh->GetMinLod().GetValueForPlatform(CookTarget->GetPlatformInfo().IniPlatformName) : 0;
+		bMeshDisablesMinLodStrip = OwnerMesh ? OwnerMesh->GetDisableBelowMinLodStripping().GetValueForPlatform(CookTarget->GetPlatformInfo().IniPlatformName) : false;
 	}
 	const bool bWantToStripBelowMinLod = bIsCook && GStripSkeletalMeshLodsDuringCooking != 0 && MinMeshLod > LODIdx && !bMeshDisablesMinLodStrip;
 
@@ -640,7 +640,7 @@ void FSkeletalMeshLODRenderData::SerializeStreamedData(FArchive& Ar, USkeletalMe
 	StaticVertexBuffers.StaticMeshVertexBuffer.Serialize(Ar, bNeedsCPUAccess);
 	Ar << SkinWeightVertexBuffer;
 
-	if (Owner && Owner->bHasVertexColors)
+	if (Owner && Owner->GetHasVertexColors())
 	{
 		StaticVertexBuffers.ColorVertexBuffer.Serialize(Ar, bForceKeepCPUResources);
 	}
diff --git a/Engine/Source/Runtime/Engine/Private/SkeletalMeshLODSettings.cpp b/Engine/Source/Runtime/Engine/Private/SkeletalMeshLODSettings.cpp
--- a/Engine/Source/Runtime/Engine/Private/SkeletalMeshLODSettings.cpp
+++ b/Engine/Source/Runtime/Engine/Private/SkeletalMeshLODSettings.cpp
@@ -54,7 +54,7 @@ bool USkeletalMeshLODSettings::SetLODSettingsToMesh(USkeletalMesh* InMesh, int32
 		// if we have available bake pose
 		// it's possible for skeleton to be null if this happens in the middle of importing
 		// so if skeleton is null, we allow copy (the GetBakePose will check correct skeleton when get it)
-		if (Setting.BakePose && (!InMesh->Skeleton || InMesh->Skeleton->IsCompatible(Setting.BakePose->GetSkeleton())))
+		if (Setting.BakePose && (!InMesh->GetSkeleton() || InMesh->GetSkeleton()->IsCompatible(Setting.BakePose->GetSkeleton())))
 		{
 			LODInfo->BakePose = Setting.BakePose;
 		}
@@ -66,7 +66,7 @@ bool USkeletalMeshLODSettings::SetLODSettingsToMesh(USkeletalMesh* InMesh, int32
 		// select joints that mesh has
 		// reset the list
 		LODInfo->BonesToRemove.Reset();
-		const FReferenceSkeleton& RefSkeleton = InMesh->RefSkeleton;
+		const FReferenceSkeleton& RefSkeleton = InMesh->GetRefSkeleton();
 
 		// copy the shared setting to mesh setting
 		// make sure we have joint in the mesh 
@@ -165,13 +165,13 @@ int32 USkeletalMeshLODSettings::SetLODSettingsToMesh(USkeletalMesh* InMesh) cons
 {
 	if (InMesh)
 	{
-		InMesh->MinLod = MinLod;
-		InMesh->DisableBelowMinLodStripping = DisableBelowMinLodStripping;
+		InMesh->SetMinLod(MinLod);
+		InMesh->SetDisableBelowMinLodStripping(DisableBelowMinLodStripping);
 #if WITH_EDITORONLY_DATA
-		InMesh->bOverrideLODStreamingSettings = bOverrideLODStreamingSettings;
-		InMesh->bSupportLODStreaming = bSupportLODStreaming;
-		InMesh->MaxNumStreamedLODs = MaxNumStreamedLODs;
-		InMesh->MaxNumOptionalLODs = MaxNumOptionalLODs;
+		InMesh->SetOverrideLODStreamingSettings(bOverrideLODStreamingSettings);
+		InMesh->SetSupportLODStreaming(bSupportLODStreaming);
+		InMesh->SetMaxNumStreamedLODs(MaxNumStreamedLODs);
+		InMesh->SetMaxNumOptionalLODs(MaxNumOptionalLODs);
 #endif
 		// we only fill up until we have enough LODs
 		const int32 NumSettings = FMath::Min(LODGroups.Num(), InMesh->GetLODNum());
@@ -191,13 +191,13 @@ int32 USkeletalMeshLODSettings::SetLODSettingsFromMesh(USkeletalMesh* InMesh)
 	// in this case, we just set all settings to Mesh
 	if (InMesh)
 	{
-		MinLod = InMesh->MinLod;
-		DisableBelowMinLodStripping = InMesh->DisableBelowMinLodStripping;
+		MinLod = InMesh->GetMinLod();
+		DisableBelowMinLodStripping = InMesh->GetDisableBelowMinLodStripping();
 #if WITH_EDITORONLY_DATA
-		bOverrideLODStreamingSettings = InMesh->bOverrideLODStreamingSettings;
-		bSupportLODStreaming = InMesh->bSupportLODStreaming;
-		MaxNumStreamedLODs = InMesh->MaxNumStreamedLODs;
-		MaxNumOptionalLODs = InMesh->MaxNumOptionalLODs;
+		bOverrideLODStreamingSettings = InMesh->GetOverrideLODStreamingSettings();
+		bSupportLODStreaming = InMesh->GetSupportLODStreaming();
+		MaxNumStreamedLODs = InMesh->GetMaxNumStreamedLODs();
+		MaxNumOptionalLODs = InMesh->GetMaxNumOptionalLODs();
 #endif
 		// we only fill up until we have enough LODs
 		const int32 NumSettings = InMesh->GetLODNum();
@@ -248,7 +248,7 @@ void USkeletalMeshLODSettings::PostEditChangeProperty(FPropertyChangedEvent& Pro
 	{
 		USkeletalMesh* Mesh = Cast<USkeletalMesh>(*Iter);
 		// if PhysicsAssetOverride is NULL, it uses SkeletalMesh Physics Asset, so I'll need to update here
-		if (Mesh->LODSettings == this)
+		if (Mesh->GetLODSettings() == this)
 		{
 			// apply the change
 			SetLODSettingsToMesh(Mesh);
diff --git a/Engine/Source/Runtime/Engine/Private/SkeletalMeshMerge.cpp b/Engine/Source/Runtime/Engine/Private/SkeletalMeshMerge.cpp
--- a/Engine/Source/Runtime/Engine/Private/SkeletalMeshMerge.cpp
+++ b/Engine/Source/Runtime/Engine/Private/SkeletalMeshMerge.cpp
@@ -83,25 +83,25 @@ void FSkeletalMeshMerge::MergeSkeleton(const TArray<FRefPoseOverride>* RefPoseOv
 
 	// Build the reference skeleton & sockets.
 
-	BuildReferenceSkeleton(SrcMeshList, NewRefSkeleton, MergeMesh->Skeleton);
+	BuildReferenceSkeleton(SrcMeshList, NewRefSkeleton, MergeMesh->GetSkeleton());
 	BuildSockets(SrcMeshList);
 
 	// Override the reference bone poses & sockets, if specified.
 
 	if (RefPoseOverrides)
 	{
-		OverrideReferenceSkeletonPose(*RefPoseOverrides, NewRefSkeleton, MergeMesh->Skeleton);
+		OverrideReferenceSkeletonPose(*RefPoseOverrides, NewRefSkeleton, MergeMesh->GetSkeleton());
 		OverrideMergedSockets(*RefPoseOverrides);
 	}
 
 	// Assign new referencer skeleton.
 
-	MergeMesh->RefSkeleton = NewRefSkeleton;
+	MergeMesh->SetRefSkeleton(NewRefSkeleton);
 
 	// Rebuild inverse ref pose matrices here as some access patterns 
 	// may need to access these matrices before FinalizeMesh is called
 	// (which would *normally* rebuild the inv ref matrices).
-	MergeMesh->RefBasesInvMatrix.Empty();
+	MergeMesh->GetRefBasesInvMatrix().Empty();
 	MergeMesh->CalculateInvRefMatrices();
 }
 
@@ -131,20 +131,20 @@ bool FSkeletalMeshMerge::FinalizeMesh()
 		USkeletalMesh* SrcMesh = SrcMeshList[MeshIdx];
 		if (SrcMesh)
 		{
-			if (SrcMesh->bHasVertexColors)
+			if (SrcMesh->GetHasVertexColors())
 			{
-				MergeMesh->bHasVertexColors = true;
+				MergeMesh->SetHasVertexColors(true);
 #if WITH_EDITORONLY_DATA
-				MergeMesh->VertexColorGuid = FGuid::NewGuid();
+				MergeMesh->SetVertexColorGuid(FGuid::NewGuid());
 #endif
 			}
 
 			FMergeMeshInfo& MeshInfo = SrcMeshInfo[MeshIdx];
-			MeshInfo.SrcToDestRefSkeletonMap.AddUninitialized(SrcMesh->RefSkeleton.GetRawBoneNum());
+			MeshInfo.SrcToDestRefSkeletonMap.AddUninitialized(SrcMesh->GetRefSkeleton().GetRawBoneNum());
 
-			for (int32 i = 0; i < SrcMesh->RefSkeleton.GetRawBoneNum(); i++)
+			for (int32 i = 0; i < SrcMesh->GetRefSkeleton().GetRawBoneNum(); i++)
 			{
-				FName SrcBoneName = SrcMesh->RefSkeleton.GetBoneName(i);
+				FName SrcBoneName = SrcMesh->GetRefSkeleton().GetBoneName(i);
 				int32 DestBoneIndex = NewRefSkeleton.FindBoneIndex(SrcBoneName);
 
 				if (DestBoneIndex == INDEX_NONE)
@@ -301,10 +301,10 @@ void FSkeletalMeshMerge::GenerateNewSectionArray( TArray<FNewSectionInfo>& NewSe
 				// use the remapping of material indices if there is a valid value
 				if(SrcLODInfo.LODMaterialMap.IsValidIndex(SectionIdx) && SrcLODInfo.LODMaterialMap[SectionIdx] != INDEX_NONE)
 				{
-					MaterialIndex = FMath::Clamp<int32>( SrcLODInfo.LODMaterialMap[SectionIdx], 0, SrcMesh->Materials.Num() );
+					MaterialIndex = FMath::Clamp<int32>( SrcLODInfo.LODMaterialMap[SectionIdx], 0, SrcMesh->GetMaterials().Num() );
 				}
 
-				const FSkeletalMaterial& SkeletalMaterial = SrcMesh->Materials[MaterialIndex];
+				const FSkeletalMaterial& SkeletalMaterial = SrcMesh->GetMaterials()[MaterialIndex];
 				UMaterialInterface* Material = SkeletalMaterial.MaterialInterface;
 
 				// see if there is an existing entry in the array of new sections that matches its material
@@ -478,11 +478,11 @@ void FSkeletalMeshMerge::GenerateLODModel( int32 LODIdx )
 
 
 		// find existing material index
-		check(MergeMesh->Materials.Num() == MaterialIds.Num());
+		check(MergeMesh->GetMaterials().Num() == MaterialIds.Num());
 		int32 MatIndex;
 		if(NewSectionInfo.MaterialId == -1)
 		{
-			MatIndex = MergeMesh->Materials.Find(NewSectionInfo.Material);
+			MatIndex = MergeMesh->GetMaterials().Find(NewSectionInfo.Material);
 		}
 		else
 		{
@@ -494,9 +494,9 @@ void FSkeletalMeshMerge::GenerateLODModel( int32 LODIdx )
 		{
 			FSkeletalMaterial SkeletalMaterial(NewSectionInfo.Material, true, false, NewSectionInfo.SlotName);
 			SkeletalMaterial.UVChannelData = NewSectionInfo.UVChannelData;
-			MergeMesh->Materials.Add(SkeletalMaterial);
+			MergeMesh->GetMaterials().Add(SkeletalMaterial);
 			MaterialIds.Add(NewSectionInfo.MaterialId);
-			Section.MaterialIndex = MergeMesh->Materials.Num()-1;
+			Section.MaterialIndex = MergeMesh->GetMaterials().Num()-1;
 		}
 		else
 		{
@@ -508,7 +508,7 @@ void FSkeletalMeshMerge::GenerateLODModel( int32 LODIdx )
 		// keep track of the current base index for this section in the merged index buffer
 		Section.BaseIndex = MergedIndexBuffer.Num();
 
-		FMeshUVChannelInfo& MergedUVData = MergeMesh->Materials[Section.MaterialIndex].UVChannelData;
+		FMeshUVChannelInfo& MergedUVData = MergeMesh->GetMaterials()[Section.MaterialIndex].UVChannelData;
 
 		// iterate over all of the sections that need to be merged together
 		for( int32 MergeIdx=0; MergeIdx < NewSectionInfo.MergeSections.Num(); MergeIdx++ )
@@ -518,9 +518,9 @@ void FSkeletalMeshMerge::GenerateLODModel( int32 LODIdx )
 
 			// Take the max UV density for each UVChannel between all sections that are being merged.
 			const int32 NewSectionMatId = MergeSectionInfo.Section->MaterialIndex;
-			if(MergeSectionInfo.SkelMesh->Materials.IsValidIndex(NewSectionMatId))
+			if(MergeSectionInfo.SkelMesh->GetMaterials().IsValidIndex(NewSectionMatId))
 			{
-				const FMeshUVChannelInfo& NewSectionUVData = MergeSectionInfo.SkelMesh->Materials[NewSectionMatId].UVChannelData;
+				const FMeshUVChannelInfo& NewSectionUVData = MergeSectionInfo.SkelMesh->GetMaterials()[NewSectionMatId].UVChannelData;
 				for (int32 i = 0; i < MAX_TEXCOORDS; i++)
 				{
 					const float NewSectionUVDensity = NewSectionUVData.LocalUVDensities[i];
@@ -560,7 +560,7 @@ void FSkeletalMeshMerge::GenerateLODModel( int32 LODIdx )
 			// add required bones from this source model entry to the merge model entry
 			for( int32 Idx=0; Idx < SrcLODData.RequiredBones.Num(); Idx++ )
 			{
-				FName SrcLODBoneName = MergeSectionInfo.SkelMesh->RefSkeleton.GetBoneName(SrcLODData.RequiredBones[Idx] );
+				FName SrcLODBoneName = MergeSectionInfo.SkelMesh->GetRefSkeleton().GetBoneName(SrcLODData.RequiredBones[Idx] );
 				int32 MergeBoneIndex = NewRefSkeleton.FindBoneIndex(SrcLODBoneName);
 				
 				if (MergeBoneIndex != INDEX_NONE)
@@ -601,7 +601,7 @@ void FSkeletalMeshMerge::GenerateLODModel( int32 LODIdx )
 				DestWeight = SrcLODData.GetSkinWeightVertexBuffer()->GetVertexSkinWeights(VertIdx);
 
 				// if the mesh uses vertex colors, copy the source color if possible or default to white
-				if( MergeMesh->bHasVertexColors )
+				if( MergeMesh->GetHasVertexColors() )
 				{
 					if( VertIdx < MaxColorIdx )
 					{
@@ -712,7 +712,7 @@ void FSkeletalMeshMerge::GenerateLODModel( int32 LODIdx )
 
 	// sort required bone array in strictly increasing order
 	MergeLODData.RequiredBones.Sort();
-	MergeMesh->RefSkeleton.EnsureParentsExistAndSort(MergeLODData.ActiveBoneIndices);
+	MergeMesh->GetRefSkeleton().EnsureParentsExistAndSort(MergeLODData.ActiveBoneIndices);
 	
 	// copy the new vertices and indices to the vertex buffer for the new model
 	MergeLODData.StaticVertexBuffers.StaticMeshVertexBuffer.SetUseFullPrecisionUVs(MergeLODInfo.BuildSettings.bUseFullPrecisionUVs);
@@ -737,7 +737,7 @@ void FSkeletalMeshMerge::GenerateLODModel( int32 LODIdx )
 	// copy vertex resource arrays
 	MergeLODData.SkinWeightVertexBuffer = MergedSkinWeightBuffer;
 
-	if( MergeMesh->bHasVertexColors )
+	if( MergeMesh->GetHasVertexColors() )
 	{
 		MergeLODData.StaticVertexBuffers.ColorVertexBuffer.InitFromColorArray(MergedColorBuffer);
 	}
@@ -758,7 +758,7 @@ bool FSkeletalMeshMerge::ProcessMergeMesh()
 	// copy settings and bone info from src meshes
 	bool bNeedsInit=true;
 
-	MergeMesh->SkelMirrorTable.Empty();
+	MergeMesh->GetSkelMirrorTable().Empty();
 
 	for( int32 MeshIdx=0; MeshIdx < SrcMeshList.Num(); MeshIdx++ )
 	{
@@ -770,8 +770,8 @@ bool FSkeletalMeshMerge::ProcessMergeMesh()
 				// initialize the merged mesh with the first src mesh entry used
 				MergeMesh->SetImportedBounds(SrcMesh->GetImportedBounds());
 
-				MergeMesh->SkelMirrorAxis = SrcMesh->SkelMirrorAxis;
-				MergeMesh->SkelMirrorFlipAxis = SrcMesh->SkelMirrorFlipAxis;
+				MergeMesh->SetSkelMirrorAxis(SrcMesh->GetSkelMirrorAxis());
+				MergeMesh->SetSkelMirrorFlipAxis(SrcMesh->GetSkelMirrorFlipAxis());
 
 				// only initialize once
 				bNeedsInit = false;
@@ -785,7 +785,7 @@ bool FSkeletalMeshMerge::ProcessMergeMesh()
 	}
 
 	// Rebuild inverse ref pose matrices.
-	MergeMesh->RefBasesInvMatrix.Empty();
+	MergeMesh->GetRefBasesInvMatrix().Empty();
 	MergeMesh->CalculateInvRefMatrices();
 
 	return Result;
@@ -840,15 +840,15 @@ void FSkeletalMeshMerge::BuildReferenceSkeleton(const TArray<USkeletalMesh*>& So
 
 		if (RefSkeleton.GetRawBoneNum() == 0)
 		{
-			RefSkeleton = SourceMesh->RefSkeleton;
+			RefSkeleton = SourceMesh->GetRefSkeleton();
 			continue;
 		}
 
 		// For subsequent meshes, add any missing bones.
 
-		for (int32 i = 1; i < SourceMesh->RefSkeleton.GetRawBoneNum(); ++i)
+		for (int32 i = 1; i < SourceMesh->GetRefSkeleton().GetRawBoneNum(); ++i)
 		{
-			FName SourceBoneName = SourceMesh->RefSkeleton.GetBoneName(i);
+			FName SourceBoneName = SourceMesh->GetRefSkeleton().GetBoneName(i);
 			int32 TargetBoneIndex = RefSkeleton.FindRawBoneIndex(SourceBoneName);
 
 			// If the source bone is present in the new RefSkeleton, we skip it.
@@ -860,8 +860,8 @@ void FSkeletalMeshMerge::BuildReferenceSkeleton(const TArray<USkeletalMesh*>& So
 
 			// Add the source bone to the RefSkeleton.
 
-			int32 SourceParentIndex = SourceMesh->RefSkeleton.GetParentIndex(i);
-			FName SourceParentName = SourceMesh->RefSkeleton.GetBoneName(SourceParentIndex);
+			int32 SourceParentIndex = SourceMesh->GetRefSkeleton().GetParentIndex(i);
+			FName SourceParentName = SourceMesh->GetRefSkeleton().GetBoneName(SourceParentIndex);
 			int32 TargetParentIndex = RefSkeleton.FindRawBoneIndex(SourceParentName);
 
 			if (TargetParentIndex == INDEX_NONE)
@@ -869,10 +869,10 @@ void FSkeletalMeshMerge::BuildReferenceSkeleton(const TArray<USkeletalMesh*>& So
 				continue;
 			}
 
-			FMeshBoneInfo MeshBoneInfo = SourceMesh->RefSkeleton.GetRefBoneInfo()[i];
+			FMeshBoneInfo MeshBoneInfo = SourceMesh->GetRefSkeleton().GetRefBoneInfo()[i];
 			MeshBoneInfo.ParentIndex = TargetParentIndex;
 
-			RefSkelModifier.Add(MeshBoneInfo, SourceMesh->RefSkeleton.GetRefBonePose()[i]);
+			RefSkelModifier.Add(MeshBoneInfo, SourceMesh->GetRefSkeleton().GetRefBonePose()[i]);
 		}
 	}
 }
@@ -882,7 +882,7 @@ void FSkeletalMeshMerge::OverrideReferenceSkeletonPose(const TArray<FRefPoseOver
 	for (int32 i = 0, PoseMax = PoseOverrides.Num(); i < PoseMax; ++i)
 	{
 		const FRefPoseOverride& PoseOverride = PoseOverrides[i];
-		const FReferenceSkeleton& SourceSkeleton = PoseOverride.SkeletalMesh->RefSkeleton;
+		const FReferenceSkeleton& SourceSkeleton = PoseOverride.SkeletalMesh->GetRefSkeleton();
 
 		FReferenceSkeletonModifier RefSkelModifier(TargetSkeleton, SkeletonAsset);
 
@@ -942,7 +942,7 @@ void FSkeletalMeshMerge::ReleaseResources(int32 Slack)
 	}
 
 	MergeMesh->ResetLODInfo();
-	MergeMesh->Materials.Empty();
+	MergeMesh->GetMaterials().Empty();
 }
 
 bool FSkeletalMeshMerge::AddSocket(const USkeletalMeshSocket* NewSocket, bool bIsSkeletonSocket)
@@ -962,9 +962,9 @@ bool FSkeletalMeshMerge::AddSocket(const USkeletalMeshSocket* NewSocket, bool bI
 	// (i.e. an existing mesh was used, or a created mesh was explicitly assigned a skeleton).
 	// In either case, we want to avoid adding sockets to the Skeleton (as it is shared), but we
 	// still need to check against it to prevent duplication.
-	if (bIsSkeletonSocket && MergeMesh->Skeleton)
+	if (bIsSkeletonSocket && MergeMesh->GetSkeleton())
 	{
-		for (USkeletalMeshSocket const * const ExistingSocket : MergeMesh->Skeleton->Sockets)
+		for (USkeletalMeshSocket const * const ExistingSocket : MergeMesh->GetSkeleton()->Sockets)
 		{
 			return false;
 		}
@@ -1004,9 +1004,9 @@ void FSkeletalMeshMerge::BuildSockets(const TArray<USkeletalMesh*>& SourceMeshLi
 
 	for (USkeletalMesh const * const SourceMesh : SourceMeshList)
 	{
-		if (SourceMesh)
+		if (SourceMesh && SourceMesh->GetSkeleton())
 		{
-			const TArray<USkeletalMeshSocket*>& NewSkeletonSocketList = SourceMesh->Skeleton->Sockets;
+			const TArray<USkeletalMeshSocket*>& NewSkeletonSocketList = SourceMesh->GetSkeleton()->Sockets;
 			AddSockets(NewSkeletonSocketList, true);
 		}
 	}
@@ -1050,9 +1050,15 @@ void FSkeletalMeshMerge::OverrideMergedSockets(const TArray<FRefPoseOverride>& P
 	for (int32 i = 0, PoseMax = PoseOverrides.Num(); i < PoseMax; ++i)
 	{
 		const FRefPoseOverride& PoseOverride = PoseOverrides[i];
-		const FReferenceSkeleton& SourceSkeleton = PoseOverride.SkeletalMesh->RefSkeleton;
+		const FReferenceSkeleton& SourceSkeleton = PoseOverride.SkeletalMesh->GetRefSkeleton();
+		
+		//We assume the skeleton is valid here
+		if (!ensure(PoseOverride.SkeletalMesh->GetSkeleton()))
+		{
+			continue;
+		}
 
-		const TArray<USkeletalMeshSocket*>& SkeletonSocketList = PoseOverride.SkeletalMesh->Skeleton->Sockets;
+		const TArray<USkeletalMeshSocket*>& SkeletonSocketList = PoseOverride.SkeletalMesh->GetSkeleton()->Sockets;
 		const TArray<USkeletalMeshSocket*>& MeshSocketList = const_cast<USkeletalMesh*>(PoseOverride.SkeletalMesh)->GetMeshOnlySocketList();
 
 		for (int32 j = 0, BoneMax = PoseOverride.Overrides.Num(); j < BoneMax; ++j)
diff --git a/Engine/Source/Runtime/Engine/Private/SkeletalMeshRenderData.cpp b/Engine/Source/Runtime/Engine/Private/SkeletalMeshRenderData.cpp
--- a/Engine/Source/Runtime/Engine/Private/SkeletalMeshRenderData.cpp
+++ b/Engine/Source/Runtime/Engine/Private/SkeletalMeshRenderData.cpp
@@ -41,15 +41,15 @@ static void SerializeLODInfoForDDC(USkeletalMesh* SkeletalMesh, FString& KeySuff
 		check(LODInfos.IsValidIndex(LODIndex));
 		FSkeletalMeshLODInfo& LODInfo = LODInfos[LODIndex];
 		bool bValidLODSettings = false;
-		if (SkeletalMesh->LODSettings != nullptr)
+		if (SkeletalMesh->GetLODSettings() != nullptr)
 		{
-			const int32 NumSettings = FMath::Min(SkeletalMesh->LODSettings->GetNumberOfSettings(), SkeletalMesh->GetLODNum());
+			const int32 NumSettings = FMath::Min(SkeletalMesh->GetLODSettings()->GetNumberOfSettings(), SkeletalMesh->GetLODNum());
 			if (LODIndex < NumSettings)
 			{
 				bValidLODSettings = true;
 			}
 		}
-		const FSkeletalMeshLODGroupSettings* SkeletalMeshLODGroupSettings = bValidLODSettings ? &SkeletalMesh->LODSettings->GetSettingsForLODLevel(LODIndex) : nullptr;
+		const FSkeletalMeshLODGroupSettings* SkeletalMeshLODGroupSettings = bValidLODSettings ? &SkeletalMesh->GetLODSettings()->GetSettingsForLODLevel(LODIndex) : nullptr;
 		LODInfo.BuildGUID = LODInfo.ComputeDeriveDataCacheKey(SkeletalMeshLODGroupSettings);
 		KeySuffix += LODInfo.BuildGUID.ToString(EGuidFormats::Digits);
 	}
@@ -71,7 +71,7 @@ static FString BuildSkeletalMeshDerivedDataKey(const ITargetPlatform* TargetPlat
 {
 	FString KeySuffix(TEXT(""));
 
-	if (SkelMesh->UseLegacyMeshDerivedDataKey )
+	if (SkelMesh->GetUseLegacyMeshDerivedDataKey() )
 	{
 		//Old asset will have the same LOD settings for bUseFullPrecisionUVs. We can use the LOD 0
 		const FSkeletalMeshLODInfo* BaseLODInfo = SkelMesh->GetLODInfo(0);
@@ -98,8 +98,8 @@ static FString BuildSkeletalMeshDerivedDataKey(const ITargetPlatform* TargetPlat
 		KeySuffix += tmpDebugString;
 	}
 
-	KeySuffix += SkelMesh->bHasVertexColors ? "1" : "0";
-	KeySuffix += SkelMesh->VertexColorGuid.ToString(EGuidFormats::Digits);
+	KeySuffix += SkelMesh->GetHasVertexColors() ? "1" : "0";
+	KeySuffix += SkelMesh->GetVertexColorGuid().ToString(EGuidFormats::Digits);
 
 	static auto* VarMeshStreaming = IConsoleManager::Get().FindConsoleVariable(TEXT("r.MeshStreaming"));
 	const bool bMeshStreamingEnabled = !VarMeshStreaming || VarMeshStreaming->GetInt() != 0;
@@ -160,7 +160,7 @@ void FSkeletalMeshRenderData::Cache(const ITargetPlatform* TargetPlatform, USkel
 			//We need to store those so we do not have to rerun the reduction to make them up to date
 			//with the serialize renderdata. This allow to use DDC when changing the reduction settings.
 			//The old workflow has to reduce the LODModel before getting the render data DDC.
-			if (!Owner->UseLegacyMeshDerivedDataKey)
+			if (!Owner->GetUseLegacyMeshDerivedDataKey())
 			{
 				FSkeletalMeshModel* SkelMeshModel = Owner->GetImportedModel();
 				check(SkelMeshModel);
@@ -168,8 +168,8 @@ void FSkeletalMeshRenderData::Cache(const ITargetPlatform* TargetPlatform, USkel
 				int32 MorphTargetNumber = 0;
 				Ar << MorphTargetNumber;
 				TArray<UMorphTarget*> ToDeleteMorphTargets;
-				ToDeleteMorphTargets.Append(Owner->MorphTargets);
-				Owner->MorphTargets.Empty();
+				ToDeleteMorphTargets.Append(Owner->GetMorphTargets());
+				Owner->GetMorphTargets().Empty();
 				//Rebuild the MorphTarget object
 				//We cannot serialize directly the UMorphTarget with a FMemoryArchive. This is not supported.
 				for (int32 MorphTargetIndex = 0; MorphTargetIndex < MorphTargetNumber; ++MorphTargetIndex)
@@ -187,8 +187,8 @@ void FSkeletalMeshRenderData::Cache(const ITargetPlatform* TargetPlatform, USkel
 						ToDeleteMorphTargets.Remove(MorphTarget);
 					}
 					MorphTarget->MorphLODModels.Empty();
-					Owner->MorphTargets.Add(MorphTarget);
-					check(MorphTargetIndex == Owner->MorphTargets.Num() - 1);
+					Owner->GetMorphTargets().Add(MorphTarget);
+					check(MorphTargetIndex == Owner->GetMorphTargets().Num() - 1);
 					int32 MorphLODModelNumber = 0;
 					Ar << MorphLODModelNumber;
 					MorphTarget->MorphLODModels.AddDefaulted(MorphLODModelNumber);
@@ -299,19 +299,19 @@ void FSkeletalMeshRenderData::Cache(const ITargetPlatform* TargetPlatform, USkel
 			
 			//If we load an old asset we want to be sure the serialize ddc will be the same has before the skeletalmesh build refactor
 			//So we do not serialize the LODModel sections.
-			if (!Owner->UseLegacyMeshDerivedDataKey)
+			if (!Owner->GetUseLegacyMeshDerivedDataKey())
 			{
-				int32 MorphTargetNumber = Owner->MorphTargets.Num();
+				int32 MorphTargetNumber = Owner->GetMorphTargets().Num();
 				Ar << MorphTargetNumber;
 				for (int32 MorphTargetIndex = 0; MorphTargetIndex < MorphTargetNumber; ++MorphTargetIndex)
 				{
-					FName MorphTargetName = Owner->MorphTargets[MorphTargetIndex]->GetFName();
+					FName MorphTargetName = Owner->GetMorphTargets()[MorphTargetIndex]->GetFName();
 					Ar << MorphTargetName;
-					int32 MorphLODModelNumber = Owner->MorphTargets[MorphTargetIndex]->MorphLODModels.Num();
+					int32 MorphLODModelNumber = Owner->GetMorphTargets()[MorphTargetIndex]->MorphLODModels.Num();
 					Ar << MorphLODModelNumber;
 					for (int32 MorphIndex = 0; MorphIndex < MorphLODModelNumber; ++MorphIndex)
 					{
-						Ar << Owner->MorphTargets[MorphTargetIndex]->MorphLODModels[MorphIndex];
+						Ar << Owner->GetMorphTargets()[MorphTargetIndex]->MorphLODModels[MorphIndex];
 					}
 				}
 				//No need to serialize the morph target mapping since we will rebuild the mapping when loading a ddc
diff --git a/Engine/Source/Runtime/Engine/Private/SkeletalMeshSampling.cpp b/Engine/Source/Runtime/Engine/Private/SkeletalMeshSampling.cpp
--- a/Engine/Source/Runtime/Engine/Private/SkeletalMeshSampling.cpp
+++ b/Engine/Source/Runtime/Engine/Private/SkeletalMeshSampling.cpp
@@ -248,7 +248,7 @@ FORCEINLINE_DEBUGGABLE void GetRegionValidData(USkeletalMesh* SkeletalMesh, FSke
 	
 	for (FSkelMeshRenderSection& Section : LODData.RenderSections)
 	{
-		FName MaterialName = SkeletalMesh->Materials[Section.MaterialIndex].MaterialSlotName;
+		FName MaterialName = SkeletalMesh->GetMaterials()[Section.MaterialIndex].MaterialSlotName;
 
 		if (SamplingRegion.IsMaterialAllowed(MaterialName))
 		{
@@ -319,7 +319,7 @@ void FSkeletalMeshSamplingInfo::BuildRegions(USkeletalMesh* SkeletalMesh)
 		FSkinWeightVertexBuffer* SkinWeightBuffer = &LODData.SkinWeightVertexBuffer;
 		check(SkinWeightBuffer);
 
-		FReferenceSkeleton& RefSkel = SkeletalMesh->RefSkeleton;
+		FReferenceSkeleton& RefSkel = SkeletalMesh->GetRefSkeleton();
 		TArray<int32> IncludedBoneIndices;
 		TArray<int32> ExcludedBoneIndices;
 		Region.GetFilteredBones(RefSkel, IncludedBoneIndices, ExcludedBoneIndices);
