To protect the skeletalmesh when using async operations like build or re-import we deprecate all public member variables of USkeletalMesh class and create the necessary getter and setter. Those member

diff --git a/Engine/Source/Runtime/Engine/Private/SkeletalRender.cpp b/Engine/Source/Runtime/Engine/Private/SkeletalRender.cpp
--- a/Engine/Source/Runtime/Engine/Private/SkeletalRender.cpp
+++ b/Engine/Source/Runtime/Engine/Private/SkeletalRender.cpp
@@ -310,13 +310,13 @@ void UpdateRefToLocalMatrices( TArray<FMatrix>& ReferenceToLocal, const USkinned
 	const USkinnedMeshComponent* const MasterComp = InMeshComponent->MasterPoseComponent.Get();
 	const FSkeletalMeshLODRenderData& LOD = InSkeletalMeshRenderData->LODRenderData[LODIndex];
 
-	const FReferenceSkeleton& RefSkeleton = ThisMesh->RefSkeleton;
+	const FReferenceSkeleton& RefSkeleton = ThisMesh->GetRefSkeleton();
 	const TArray<int32>& MasterBoneMap = InMeshComponent->GetMasterBoneMap();
-	const bool bIsMasterCompValid = MasterComp && MasterBoneMap.Num() == ThisMesh->RefSkeleton.GetNum();
+	const bool bIsMasterCompValid = MasterComp && MasterBoneMap.Num() == ThisMesh->GetRefSkeleton().GetNum();
 	const TArray<FTransform>& ComponentTransform = (bIsMasterCompValid)? MasterComp->GetComponentSpaceTransforms() : InMeshComponent->GetComponentSpaceTransforms();
 	const TArray<uint8>& BoneVisibilityStates = (bIsMasterCompValid) ? MasterComp->GetBoneVisibilityStates() : InMeshComponent->GetBoneVisibilityStates();
 	// Get inv ref pose matrices
-	const TArray<FMatrix>* RefBasesInvMatrix = &ThisMesh->RefBasesInvMatrix;
+	const TArray<FMatrix>* RefBasesInvMatrix = &ThisMesh->GetRefBasesInvMatrix();
 
 	// Check if there is an override (and it's the right size)
 	if( InMeshComponent->GetRefPoseOverride() && 
@@ -359,13 +359,13 @@ void UpdatePreviousRefToLocalMatrices(TArray<FMatrix>& ReferenceToLocal, const U
 	const USkinnedMeshComponent* const MasterComp = InMeshComponent->MasterPoseComponent.Get();
 	const FSkeletalMeshLODRenderData& LOD = InSkeletalMeshRenderData->LODRenderData[LODIndex];
 
-	const FReferenceSkeleton& RefSkeleton = ThisMesh->RefSkeleton;
+	const FReferenceSkeleton& RefSkeleton = ThisMesh->GetRefSkeleton();
 	const TArray<int32>& MasterBoneMap = InMeshComponent->GetMasterBoneMap();
-	const bool bIsMasterCompValid = MasterComp && MasterBoneMap.Num() == ThisMesh->RefSkeleton.GetNum();
+	const bool bIsMasterCompValid = MasterComp && MasterBoneMap.Num() == ThisMesh->GetRefSkeleton().GetNum();
 	const TArray<FTransform>& ComponentTransform = (bIsMasterCompValid) ? MasterComp->GetPreviousComponentTransformsArray() : InMeshComponent->GetPreviousComponentTransformsArray();
 	const TArray<uint8>& BoneVisibilityStates = (bIsMasterCompValid) ? MasterComp->GetPreviousBoneVisibilityStates() : InMeshComponent->GetPreviousBoneVisibilityStates();
 	// Get inv ref pose matrices
-	const TArray<FMatrix>* RefBasesInvMatrix = &ThisMesh->RefBasesInvMatrix;
+	const TArray<FMatrix>* RefBasesInvMatrix = &ThisMesh->GetRefBasesInvMatrix();
 	// Check if there is an override (and it's the right size)
 	if (InMeshComponent->GetRefPoseOverride() &&
 		InMeshComponent->GetRefPoseOverride()->RefBasesInvMatrix.Num() == RefBasesInvMatrix->Num())
diff --git a/Engine/Source/Runtime/Engine/Private/Streaming/SkeletalMeshUpdate.cpp b/Engine/Source/Runtime/Engine/Private/Streaming/SkeletalMeshUpdate.cpp
--- a/Engine/Source/Runtime/Engine/Private/Streaming/SkeletalMeshUpdate.cpp
+++ b/Engine/Source/Runtime/Engine/Private/Streaming/SkeletalMeshUpdate.cpp
@@ -201,7 +201,7 @@ void FSkeletalMeshStreamIn::DoFinishUpdate(const FContext& Context)
 			for (int32 LODIndex = PendingFirstLODIdx; LODIndex < CurrentFirstLODIdx; ++LODIndex)
 			{
 				FSkeletalMeshLODRenderData& LODResource = *Context.LODResourcesView[LODIndex];
-				LODResource.IncrementMemoryStats(Mesh->bHasVertexColors);
+				LODResource.IncrementMemoryStats(Mesh->GetHasVertexColors());
 				IntermediateBuffersArray[LODIndex].TransferBuffers(LODResource, Batcher);
 			}
 		}
diff --git a/Engine/Source/Runtime/Engine/Private/UnrealEngine.cpp b/Engine/Source/Runtime/Engine/Private/UnrealEngine.cpp
--- a/Engine/Source/Runtime/Engine/Private/UnrealEngine.cpp
+++ b/Engine/Source/Runtime/Engine/Private/UnrealEngine.cpp
@@ -5716,31 +5716,31 @@ bool UEngine::HandleListSkeletalMeshesCommand(const TCHAR* Cmd, FOutputDevice& A
 		int32 MaxNumStreamedLODs = -1;
 		int32 MaxNumOptionalLODs = -1;
 #if WITH_EDITORONLY_DATA
-		if (Mesh->MinLod.PerPlatform.Find(("Mobile")) != nullptr)
+		if (Mesh->GetMinLod().PerPlatform.Find(("Mobile")) != nullptr)
 		{
-			MobileMinLOD = *Mesh->MinLod.PerPlatform.Find(("Mobile"));
+			MobileMinLOD = *Mesh->GetMinLod().PerPlatform.Find(("Mobile"));
 		}
 
-		MeshDisablesMinLODStripping = Mesh->DisableBelowMinLodStripping.Default ? 1 : 0;
-		if (Mesh->DisableBelowMinLodStripping.PerPlatform.Find(("Mobile")) != nullptr)
+		MeshDisablesMinLODStripping = Mesh->GetDisableBelowMinLodStripping().Default ? 1 : 0;
+		if (Mesh->GetDisableBelowMinLodStripping().PerPlatform.Find(("Mobile")) != nullptr)
 		{
-			MeshDisablesMinLODStripping = *Mesh->DisableBelowMinLodStripping.PerPlatform.Find(("Mobile")) ? 1 : 0;
+			MeshDisablesMinLODStripping = *Mesh->GetDisableBelowMinLodStripping().PerPlatform.Find(("Mobile")) ? 1 : 0;
 		}
 
-		bSupportLODStreaming = Mesh->bSupportLODStreaming.Default;
-		if (const bool* Found = Mesh->bSupportLODStreaming.PerPlatform.Find("Mobile"))
+		bSupportLODStreaming = Mesh->GetSupportLODStreaming().Default;
+		if (const bool* Found = Mesh->GetSupportLODStreaming().PerPlatform.Find("Mobile"))
 		{
 			bSupportLODStreaming = *Found;
 		}
 
-		MaxNumStreamedLODs = Mesh->MaxNumStreamedLODs.Default;
-		if (const int32* Found = Mesh->MaxNumStreamedLODs.PerPlatform.Find("Mobile"))
+		MaxNumStreamedLODs = Mesh->GetMaxNumStreamedLODs().Default;
+		if (const int32* Found = Mesh->GetMaxNumStreamedLODs().PerPlatform.Find("Mobile"))
 		{
 			MaxNumStreamedLODs = *Found;
 		}
 
-		MaxNumOptionalLODs = Mesh->MaxNumOptionalLODs.Default;
-		if (const int32* Found = Mesh->MaxNumOptionalLODs.PerPlatform.Find("Mobile"))
+		MaxNumOptionalLODs = Mesh->GetMaxNumOptionalLODs().Default;
+		if (const int32* Found = Mesh->GetMaxNumOptionalLODs().PerPlatform.Find("Mobile"))
 		{
 			MaxNumOptionalLODs = *Found;
 		}
@@ -5782,16 +5782,17 @@ bool UEngine::HandleListSkeletalMeshesCommand(const TCHAR* Cmd, FOutputDevice& A
 		}
 
 		int32 VertexCountCollision = 0;
-		if (Mesh->BodySetup)
+		const USkeletalMesh* MeshConst = Mesh;
+		if (MeshConst->GetBodySetup())
 		{
 #if PHYSICS_INTERFACE_PHYSX
 			// Count PhysX trimesh mem usage
-			for (physx::PxTriangleMesh* TriMesh : Mesh->BodySetup->TriMeshes)
+			for (physx::PxTriangleMesh* TriMesh : MeshConst->GetBodySetup()->TriMeshes)
 			{
 				VertexCountCollision += TriMesh->getNbVertices();
 			}
 #elif WITH_CHAOS
-			for (auto& TriMesh : Mesh->BodySetup->ChaosTriMeshes)
+			for (auto& TriMesh : MeshConst->GetBodySetup()->ChaosTriMeshes)
 			{
 				VertexCountCollision += TriMesh->Particles().Size();
 			}
diff --git a/Engine/Source/Runtime/Experimental/GeometryCollectionEngine/Private/GeometryCollection/PhysicsAssetSimulation.cpp b/Engine/Source/Runtime/Experimental/GeometryCollectionEngine/Private/GeometryCollection/PhysicsAssetSimulation.cpp
--- a/Engine/Source/Runtime/Experimental/GeometryCollectionEngine/Private/GeometryCollection/PhysicsAssetSimulation.cpp
+++ b/Engine/Source/Runtime/Experimental/GeometryCollectionEngine/Private/GeometryCollection/PhysicsAssetSimulation.cpp
@@ -41,7 +41,7 @@ void FPhysicsAssetSimulationUtil::BuildParams(const UObject* Caller, const AActo
 		return;
 	}
 
-	const USkeleton* Skeleton = SkeletalMesh->Skeleton;
+	const USkeleton* Skeleton = SkeletalMesh->GetSkeleton();
 	const FSkeletalMeshRenderData* SkelMeshRenderData = SkeletalMesh->GetResourceForRendering();
 	if (!PhysicsAsset || !Skeleton || !SkelMeshRenderData || !SkelMeshRenderData->LODRenderData.Num())
 	{
diff --git a/Engine/Source/Runtime/Experimental/GeometryCollectionEngine/Private/GeometryCollection/SkeletalMeshSimulationComponent.cpp b/Engine/Source/Runtime/Experimental/GeometryCollectionEngine/Private/GeometryCollection/SkeletalMeshSimulationComponent.cpp
--- a/Engine/Source/Runtime/Experimental/GeometryCollectionEngine/Private/GeometryCollection/SkeletalMeshSimulationComponent.cpp
+++ b/Engine/Source/Runtime/Experimental/GeometryCollectionEngine/Private/GeometryCollection/SkeletalMeshSimulationComponent.cpp
@@ -149,7 +149,7 @@ void USkeletalMeshSimulationComponent::OnCreatePhysicsState()
 			USkeletalMesh* SkeletalMesh = SkelMeshComponent->SkeletalMesh;
 			if (SkeletalMesh)
 			{
-				UPhysicsAsset* PhysicsAsset = OverridePhysicsAsset ? OverridePhysicsAsset : SkelMeshComponent->SkeletalMesh->PhysicsAsset;
+				UPhysicsAsset* PhysicsAsset = OverridePhysicsAsset ? OverridePhysicsAsset : SkelMeshComponent->SkeletalMesh->GetPhysicsAsset();
 				FPhysicsAssetSimulationUtil::BuildParams(this, OwningActor, SkelMeshComponent, PhysicsAsset, OutPhysicsParams);
 			}
 
diff --git a/Engine/Source/Runtime/MovieSceneTracks/Private/Evaluation/MovieSceneSkeletalAnimationTemplate.cpp b/Engine/Source/Runtime/MovieSceneTracks/Private/Evaluation/MovieSceneSkeletalAnimationTemplate.cpp
--- a/Engine/Source/Runtime/MovieSceneTracks/Private/Evaluation/MovieSceneSkeletalAnimationTemplate.cpp
+++ b/Engine/Source/Runtime/MovieSceneTracks/Private/Evaluation/MovieSceneSkeletalAnimationTemplate.cpp
@@ -26,8 +26,8 @@ bool ShouldUsePreviewPlayback(IMovieScenePlayer& Player, UObject& RuntimeObject)
 
 bool CanPlayAnimation(USkeletalMeshComponent* SkeletalMeshComponent, UAnimSequenceBase* AnimAssetBase)
 {
-	return (SkeletalMeshComponent->SkeletalMesh && SkeletalMeshComponent->SkeletalMesh->Skeleton && 
-		(!AnimAssetBase || SkeletalMeshComponent->SkeletalMesh->Skeleton->IsCompatible(AnimAssetBase->GetSkeleton())));
+	return (SkeletalMeshComponent->SkeletalMesh && SkeletalMeshComponent->SkeletalMesh->GetSkeleton() && 
+		(!AnimAssetBase || SkeletalMeshComponent->SkeletalMesh->GetSkeleton()->IsCompatible(AnimAssetBase->GetSkeleton())));
 }
 
 void ResetAnimSequencerInstance(UObject& ObjectToRestore, IMovieScenePlayer& Player)
@@ -137,7 +137,7 @@ struct FPreAnimatedAnimationTokenProducer : IMovieScenePreAnimatedTokenProducer
 					if (Component->AnimScriptInstance && Component->SkeletalMesh)
 					{
 						//the skeleton may have changed so need to recalc required bones as needed.
-						Component->AnimScriptInstance->CurrentSkeleton = Component->SkeletalMesh->Skeleton;
+						Component->AnimScriptInstance->CurrentSkeleton = Component->SkeletalMesh->GetSkeleton();
 						//Need at least RecalcRequiredbones and UpdateMorphTargetrs
 						Component->InitializeAnimScriptInstance(true);
 					}
