Missing include, fix non-unity

diff --git a/Engine/Source/Runtime/Engine/Private/Materials/MaterialCachedData.cpp b/Engine/Source/Runtime/Engine/Private/Materials/MaterialCachedData.cpp
--- a/Engine/Source/Runtime/Engine/Private/Materials/MaterialCachedData.cpp
+++ b/Engine/Source/Runtime/Engine/Private/Materials/MaterialCachedData.cpp
@@ -34,6 +34,7 @@
 #include "Materials/MaterialFunctionInstance.h"
 #include "Materials/MaterialExpressionMaterialFunctionCall.h"
 #include "Materials/MaterialParameterCollection.h"
+#include "Engine/Font.h"
 #include "LandscapeGrassType.h"
 #include "ProfilingDebugging/LoadTimeTracker.h"
 #include "Curves/CurveLinearColor.h"
