After initial skeletal mesh class public member deprecation, we found that we miss one member (SkelMirrorAxis)
Fix QAGame that is using a deprecated member directly

diff --git a/Engine/Source/Runtime/Engine/Classes/Engine/SkeletalMesh.h b/Engine/Source/Runtime/Engine/Classes/Engine/SkeletalMesh.h
--- a/Engine/Source/Runtime/Engine/Classes/Engine/SkeletalMesh.h
+++ b/Engine/Source/Runtime/Engine/Classes/Engine/SkeletalMesh.h
@@ -1046,6 +1046,7 @@ public:
 	UFUNCTION(BlueprintGetter)
 	TSoftObjectPtr<UObject> GetDefaultAnimatingRig() const;
 
+	UE_DEPRECATED(4.27, "Please do not access this member directly; use USkeletalMesh::GetSkelMirrorAxis() or USkeletalMesh::SetSkelMirrorAxis().")
 	UPROPERTY(EditAnywhere, Category=Mirroring)
 	TEnumAsByte<EAxis::Type> SkelMirrorAxis;
 	static FName GetSkelMirrorAxisMemberName()
