Don't assume FMaterials with non-0 compiling shader map IDs are still compiling
- Instead check shader compiler to see if any pending compiles actually match that ID
- In theory we should be setting compiling ID to 0 once compile finishes, but this should avoid hangs if that doesn't happen

diff --git a/Engine/Source/Runtime/Engine/Private/Materials/MaterialShared.cpp b/Engine/Source/Runtime/Engine/Private/Materials/MaterialShared.cpp
--- a/Engine/Source/Runtime/Engine/Private/Materials/MaterialShared.cpp
+++ b/Engine/Source/Runtime/Engine/Private/Materials/MaterialShared.cpp
@@ -598,18 +598,7 @@ bool FMaterial::IsRequiredComplete() const
 #if WITH_EDITOR
 void FMaterial::GetShaderMapIDsWithUnfinishedCompilation(TArray<int32>& ShaderMapIds)
 {
-#if 0
-	// Build an array of the shader map Id's are not finished compiling.
-	if (GameThreadShaderMap && !GameThreadShaderMap->IsCompilationFinalized())
-	{
-		ShaderMapIds.Add(GameThreadShaderMap->GetCompilingId());
-	}
-	else if (CompilingShaderMapId != 0u)
-	{
-		ShaderMapIds.Add(CompilingShaderMapId);
-	}
-#endif
-	if (GameThreadCompilingShaderMapId != 0u)
+	if (GameThreadCompilingShaderMapId != 0u && GShaderCompilingManager->IsCompilingShaderMap(GameThreadCompilingShaderMapId))
 	{
 		ShaderMapIds.Add(GameThreadCompilingShaderMapId);
 	}
@@ -617,19 +606,11 @@ void FMaterial::GetShaderMapIDsWithUnfinishedCompilation(TArray<int32>& ShaderMa
 
 bool FMaterial::IsCompilationFinished() const
 {
-#if 0
-	// Build an array of the shader map Id's are not finished compiling.
-	if (GameThreadShaderMap && !GameThreadShaderMap->IsCompilationFinalized())
-	{
-		return false;
-	}
-	else if (CompilingShaderMapId != 0u)
+	if (GameThreadCompilingShaderMapId != 0u)
 	{
-		return false;
+		return !GShaderCompilingManager->IsCompilingShaderMap(GameThreadCompilingShaderMapId);
 	}
 	return true;
-#endif
-	return GameThreadCompilingShaderMapId == 0u;
 }
 
 void FMaterial::CancelCompilation()
diff --git a/Engine/Source/Runtime/Engine/Private/ShaderCompiler/ShaderCompiler.cpp b/Engine/Source/Runtime/Engine/Private/ShaderCompiler/ShaderCompiler.cpp
--- a/Engine/Source/Runtime/Engine/Private/ShaderCompiler/ShaderCompiler.cpp
+++ b/Engine/Source/Runtime/Engine/Private/ShaderCompiler/ShaderCompiler.cpp
@@ -2590,6 +2590,26 @@ void FShaderCompilingManager::SubmitJobs(TArray<FShaderCommonCompileJobPtr>& New
 	AllJobs.SubmitJobs(NewJobs);
 }
 
+bool FShaderCompilingManager::IsCompilingShaderMap(uint32 Id)
+{
+	if (Id != 0u)
+	{
+		FScopeLock Lock(&CompileQueueSection);
+		FPendingShaderMapCompileResultsPtr* PendingShaderMapPtr = ShaderMapJobs.Find(Id);
+		if (PendingShaderMapPtr)
+		{
+			return true;
+		}
+
+		FShaderMapFinalizeResults* FinalizedShaderMapPtr = PendingFinalizeShaderMaps.Find(Id);
+		if (FinalizedShaderMapPtr)
+		{
+			return true;
+		}
+	}
+	return false;
+}
+
 FShaderCompileJob* FShaderCompilingManager::PrepareShaderCompileJob(uint32 Id, const FShaderCompileJobKey& Key, EShaderCompileJobPriority Priority)
 {
 	return AllJobs.PrepareJob(Id, Key, Priority);
diff --git a/Engine/Source/Runtime/Engine/Public/ShaderCompiler.h b/Engine/Source/Runtime/Engine/Public/ShaderCompiler.h
--- a/Engine/Source/Runtime/Engine/Public/ShaderCompiler.h
+++ b/Engine/Source/Runtime/Engine/Public/ShaderCompiler.h
@@ -933,6 +933,8 @@ public:
 		bNoShaderCompilation = toggle;
 	}
 
+	ENGINE_API bool IsCompilingShaderMap(uint32 Id);
+
 	/** Prepares a job of the given type for compilation.  If a job with the given Id/Key already exists, it will attempt to adjust to the higher priority if possible, and nullptr will be returned.
 	  * If a non-nullptr is returned, the given job should be filled out with relevant information, then passed to SubmitJobs() when ready
 	  */
