Move legacy custom version in to its own files so it can be controlled independently of the rest of the file it was in

diff --git a/Engine/Source/Runtime/Engine/Private/WorldSettings.cpp b/Engine/Source/Runtime/Engine/Private/WorldSettings.cpp
--- a/Engine/Source/Runtime/Engine/Private/WorldSettings.cpp
+++ b/Engine/Source/Runtime/Engine/Private/WorldSettings.cpp
@@ -32,6 +32,7 @@
 #include "AI/NavigationSystemBase.h"
 #include "Engine/BookmarkBase.h"
 #include "Engine/BookMark.h"
+#include "WorldSettingsCustomVersion.h"
 
 #if WITH_EDITOR
 #include "Editor.h"
@@ -289,28 +290,6 @@ void AWorldSettings::GetLifetimeReplicatedProps( TArray< FLifetimeProperty > & O
 	DOREPLIFETIME( AWorldSettings, bHighPriorityLoading );
 }
 
-// Custom serialization version for all packages containing WorldSetting
-struct FWorldSettingCustomVersion
-{
-	enum Type
-	{
-		// Before any version changes were made
-		BeforeCustomVersionWasAdded = 0,
-
-		// Deprecated bEnableHierarchicalLODSystem
-		DeprecatedEnableHierarchicalLODSystem = 1,
-
-		// -----<new versions can be added above this line>-------------------------------------------------
-		VersionPlusOne,
-		LatestVersion = VersionPlusOne - 1
-	};
-
-	// The GUID for this custom version number
-	const static FGuid GUID;
-
-	FWorldSettingCustomVersion() = delete;
-};
-
 const FGuid FWorldSettingCustomVersion::GUID(0x1ED048F4, 0x2F2E4C68, 0x89D053A4, 0xF18F102D);
 // Register the custom version with core
 FCustomVersionRegistration GRegisterWorldSettingCustomVersion(FWorldSettingCustomVersion::GUID, FWorldSettingCustomVersion::LatestVersion, TEXT("WorldSettingVer"));
diff --git a/Engine/Source/Runtime/Engine/Private/WorldSettingsCustomVersion.h b/Engine/Source/Runtime/Engine/Private/WorldSettingsCustomVersion.h
new file mode 100644
--- /dev/null
+++ b/Engine/Source/Runtime/Engine/Private/WorldSettingsCustomVersion.h
@@ -0,0 +1,28 @@
+// Copyright Epic Games, Inc. All Rights Reserved.
+
+#pragma once
+
+#include "CoreMinimal.h"
+#include "Misc/Guid.h"
+
+// Custom serialization version for all packages containing WorldSetting
+struct FWorldSettingCustomVersion
+{
+	enum Type
+	{
+		// Before any version changes were made
+		BeforeCustomVersionWasAdded = 0,
+
+		// Deprecated bEnableHierarchicalLODSystem
+		DeprecatedEnableHierarchicalLODSystem = 1,
+
+		// DO NOT ADD NEW VERSIONS. THIS SHOULD NOT BE USED
+		VersionPlusOne,
+		LatestVersion = VersionPlusOne - 1
+	};
+
+	// The GUID for this custom version number
+	const static FGuid GUID;
+
+	FWorldSettingCustomVersion() = delete;
+};
