PackageResourceManager and PackagePaths

diff --git a/Engine/Source/Runtime/Engine/Private/Streaming/LandscapeMeshMobileUpdate.h b/Engine/Source/Runtime/Engine/Private/Streaming/LandscapeMeshMobileUpdate.h
--- a/Engine/Source/Runtime/Engine/Private/Streaming/LandscapeMeshMobileUpdate.h
+++ b/Engine/Source/Runtime/Engine/Private/Streaming/LandscapeMeshMobileUpdate.h
@@ -113,9 +113,9 @@ protected:
 
 	bool HasPendingIORequests() const;
 
-	FString GetIOFilename(const FContext& Context);
+	void GetIOPackagePath(const FContext& Context, FPackagePath& OutPackagePath, EPackageSegment& OutPackageSegment);
 	void SetAsyncFileCallback(const FContext& Context);
-	void SetIORequest(const FContext& Context, const FString& IOFilename);
+	void SetIORequest(const FContext& Context, const FPackagePath& PackagePath, EPackageSegment PackageSegment);
 	void GetIORequestResults(const FContext& Context);
 	void ClearIORequest(const FContext& Context);
 	void CancelIORequest();
diff --git a/Engine/Source/Runtime/Engine/Private/Streaming/SkeletalMeshUpdate.cpp b/Engine/Source/Runtime/Engine/Private/Streaming/SkeletalMeshUpdate.cpp
--- a/Engine/Source/Runtime/Engine/Private/Streaming/SkeletalMeshUpdate.cpp
+++ b/Engine/Source/Runtime/Engine/Private/Streaming/SkeletalMeshUpdate.cpp
@@ -464,8 +464,9 @@ void FSkeletalMeshStreamIn_IO::SetIORequest(const FContext& Context)
 	if (Mesh && RenderData)
 	{
 #if USE_BULKDATA_STREAMING_TOKEN
-		FString Filename;
-		verify(Mesh->GetMipDataFilename(PendingFirstLODIdx, Filename));
+		FPackagePath PackagePath;
+		EPackageSegment PackageSegment;
+		verify(Mesh->GetMipDataPackagePath(PendingFirstLODIdx, PackagePath, PackageSegment));
 #endif	
 
 		SetAsyncFileCallback(Context);
@@ -481,7 +482,8 @@ void FSkeletalMeshStreamIn_IO::SetIORequest(const FContext& Context)
 		TaskSynchronization.Increment();
 
 		IORequest = FBulkDataInterface::CreateStreamingRequestForRange(
-			STREAMINGTOKEN_PARAM(Filename)
+			STREAMINGTOKEN_PARAM(PackagePath)
+			STREAMINGTOKEN_PARAM(PackageSegment)
 			BulkDataArray,
 			bHighPrioIORequest ? AIOP_BelowNormal : AIOP_Low,
 			&AsyncFileCallback);
diff --git a/Engine/Source/Runtime/Engine/Private/Streaming/StaticMeshUpdate.cpp b/Engine/Source/Runtime/Engine/Private/Streaming/StaticMeshUpdate.cpp
--- a/Engine/Source/Runtime/Engine/Private/Streaming/StaticMeshUpdate.cpp
+++ b/Engine/Source/Runtime/Engine/Private/Streaming/StaticMeshUpdate.cpp
@@ -442,8 +442,9 @@ void FStaticMeshStreamIn_IO::SetIORequest(const FContext& Context)
 	if (Mesh && RenderData)
 	{
 #if USE_BULKDATA_STREAMING_TOKEN
-		FString Filename;
-		verify(Mesh->GetMipDataFilename(PendingFirstLODIdx, Filename));
+		FPackagePath PackagePath;
+		EPackageSegment PackageSegment;
+		verify(Mesh->GetMipDataPackagePath(PendingFirstLODIdx, PackagePath, PackageSegment));
 #endif	
 
 		SetAsyncFileCallback(Context);
@@ -459,7 +460,8 @@ void FStaticMeshStreamIn_IO::SetIORequest(const FContext& Context)
 		TaskSynchronization.Increment();
 
 		IORequest = FBulkDataInterface::CreateStreamingRequestForRange(
-			STREAMINGTOKEN_PARAM(Filename)
+			STREAMINGTOKEN_PARAM(PackagePath)
+			STREAMINGTOKEN_PARAM(PackageSegment)
 			BulkDataArray,
 			bHighPrioIORequest ? AIOP_BelowNormal : AIOP_Low,
 			&AsyncFileCallback);
diff --git a/Engine/Source/Runtime/Engine/Private/Streaming/Texture2DMipDataProvider_IO.cpp b/Engine/Source/Runtime/Engine/Private/Streaming/Texture2DMipDataProvider_IO.cpp
--- a/Engine/Source/Runtime/Engine/Private/Streaming/Texture2DMipDataProvider_IO.cpp
+++ b/Engine/Source/Runtime/Engine/Private/Streaming/Texture2DMipDataProvider_IO.cpp
@@ -9,8 +9,10 @@ Texture2DMipDataProvider_IO.cpp : Implementation of FTextureMipDataProvider usin
 #include "TextureResource.h"
 #include "HAL/PlatformFileManager.h"
 #include "HAL/FileManager.h"
+#include "Misc/PackageSegment.h"
 #include "Misc/Paths.h"
 #include "Streaming/TextureStreamingHelpers.h"
+#include "UObject/PackageResourceManager.h"
 
 FTexture2DMipDataProvider_IO::FTexture2DMipDataProvider_IO(const UTexture* InTexture, bool InPrioritizedIORequest)
 	: FTextureMipDataProvider(InTexture, ETickState::Init, ETickThread::Async)
@@ -41,12 +43,12 @@ void FTexture2DMipDataProvider_IO::Init(const FTextureUpdateContext& Context, co
 			continue;
 		}
 
-		FString IOFilename;
-		IOFilename = OwnerMip.BulkData.GetFilename();
+		FPackagePath PackagePath = OwnerMip.BulkData.GetPackagePath();
+		EPackageSegment PackageSegment = OwnerMip.BulkData.GetPackageSegment();
 		if (FileInfos.IsValidIndex(CurrentFileIndex))
 		{
 			FFileInfo& FileInfo = FileInfos[CurrentFileIndex];
-			if (FileInfo.IOFilename == IOFilename && FileInfo.LastMipIndex + 1 == MipIndex)
+			if (FileInfo.PackagePath == PackagePath && FileInfo.LastMipIndex + 1 == MipIndex)
 			{
 				FileInfo.LastMipIndex = MipIndex;
 			}
@@ -61,18 +63,19 @@ void FTexture2DMipDataProvider_IO::Init(const FTextureUpdateContext& Context, co
 			int64 IOFileOffset = 0;
 			if (GEventDrivenLoaderEnabled)
 			{
-				if (IOFilename.EndsWith(TEXT(".uasset")) || IOFilename.EndsWith(TEXT(".umap")))
+				if (PackageSegment == EPackageSegment::Header)
 				{
-					IOFileOffset = -IFileManager::Get().FileSize(*IOFilename);
+					IOFileOffset = -IPackageResourceManager::Get().FileSize(PackagePath, EPackageSegment::Header);
 					check(IOFileOffset < 0);
-					IOFilename = FPaths::GetBaseFilename(IOFilename, false) + TEXT(".uexp");
-					UE_LOG(LogTexture, Error, TEXT("Streaming from the .uexp file '%s' this MUST be in a ubulk instead for best performance."), *IOFilename);
+					PackageSegment = EPackageSegment::Exports;
+					UE_LOG(LogTexture, Error, TEXT("Streaming from the .uexp file '%s' this MUST be in a ubulk instead for best performance."), *PackagePath.GetDebugName(PackageSegment));
 				}
 			}
 
 			CurrentFileIndex = FileInfos.AddDefaulted();
 			FFileInfo& FileInfo = FileInfos[CurrentFileIndex];
-			FileInfo.IOFilename = MoveTemp(IOFilename);
+			FileInfo.PackagePath = MoveTemp(PackagePath);
+			FileInfo.PackageSegment = PackageSegment;
 			FileInfo.IOFileOffset = IOFileOffset;
 			FileInfo.FirstMipIndex = MipIndex;
 			FileInfo.LastMipIndex = MipIndex;
@@ -105,7 +108,7 @@ int32 FTexture2DMipDataProvider_IO::GetMips(
 
 			if (!FileInfo.IOFileHandle)
 			{
-				FileInfo.IOFileHandle.Reset(FPlatformFileManager::Get().GetPlatformFile().OpenAsyncRead(*FileInfo.IOFilename));
+				FileInfo.IOFileHandle.Reset(IPackageResourceManager::Get().OpenAsyncReadPackage(FileInfo.PackagePath, FileInfo.PackageSegment));
 				if (!FileInfo.IOFileHandle)
 				{
 					break;
diff --git a/Engine/Source/Runtime/Engine/Private/Streaming/Texture2DMipDataProvider_IO.h b/Engine/Source/Runtime/Engine/Private/Streaming/Texture2DMipDataProvider_IO.h
--- a/Engine/Source/Runtime/Engine/Private/Streaming/Texture2DMipDataProvider_IO.h
+++ b/Engine/Source/Runtime/Engine/Private/Streaming/Texture2DMipDataProvider_IO.h
@@ -8,6 +8,7 @@ Texture2DMipDataProvider_IO.h : Implementation of FTextureMipDataProvider using
 
 #include "CoreMinimal.h"
 #include "Streaming/TextureMipDataProvider.h"
+#include "Misc/PackagePath.h"
 #include "Async/AsyncFileHandle.h"
 
 /**
@@ -38,7 +39,8 @@ protected:
 	// A structured with information about which file contains which mips.
 	struct FFileInfo
 	{
-		FString IOFilename;
+		FPackagePath PackagePath;
+		EPackageSegment PackageSegment;
 		TUniquePtr<IAsyncReadFileHandle> IOFileHandle;
 		int64 IOFileOffset = 0;
 		int32 FirstMipIndex = INDEX_NONE;
diff --git a/Engine/Source/Runtime/Engine/Private/TextureDerivedData.cpp b/Engine/Source/Runtime/Engine/Private/TextureDerivedData.cpp
--- a/Engine/Source/Runtime/Engine/Private/TextureDerivedData.cpp
+++ b/Engine/Source/Runtime/Engine/Private/TextureDerivedData.cpp
@@ -1171,7 +1171,8 @@ bool FTexturePlatformData::TryLoadMips(int32 FirstMipToLoad, void** OutMipData,
 				// is it possible for streaming mips to be loaded in non streaming ways.
 				if (CVarSetTextureStreaming.GetValueOnAnyThread() != 0)
 				{
-					UE_CLOG(Mip.BulkData.IsInSeparateFile(), LogTexture, Error, TEXT("Loading non-streamed mips from an external bulk file.  This is not desireable.  File %s"), *(Mip.BulkData.GetFilename() ) );
+					UE_CLOG(Mip.BulkData.IsInSeparateFile(), LogTexture, Error, TEXT("Loading non-streamed mips from an external bulk file.  This is not desireable.  File %s"),
+						*(Mip.BulkData.GetPackagePath().GetDebugName()) );
 				}
 #endif
 				Mip.BulkData.GetCopy(&OutMipData[MipIndex - FirstMipToLoad], true);
diff --git a/Engine/Source/Runtime/Engine/Private/UnrealEngine.cpp b/Engine/Source/Runtime/Engine/Private/UnrealEngine.cpp
--- a/Engine/Source/Runtime/Engine/Private/UnrealEngine.cpp
+++ b/Engine/Source/Runtime/Engine/Private/UnrealEngine.cpp
@@ -9323,6 +9323,11 @@ void UEngine::AddTextureStreamingSlaveLoc(FVector InLoc, float BoostFactor, bool
 FGuid UEngine::GetPackageGuid(FName PackageName, bool bForPIE)
 {
 	FGuid Result(0,0,0,0);
+	FPackagePath PackagePath;
+	if (!FPackagePath::TryFromMountedName(PackageName.ToString(), PackagePath))
+	{
+		return Result;
+	}
 
 	uint32 LoadFlags = LOAD_NoWarn | LOAD_NoVerify;
 	if (bForPIE)
@@ -9330,7 +9335,8 @@ FGuid UEngine::GetPackageGuid(FName PackageName, bool bForPIE)
 		LoadFlags |= LOAD_PackageForPIE;
 	}
 	UPackage* PackageToReset = nullptr;
-	FLinkerLoad* Linker = LoadPackageLinker(nullptr, *PackageName.ToString(), LoadFlags, nullptr, nullptr, nullptr, [&PackageToReset, &Result](FLinkerLoad* InLinker)
+
+	FLinkerLoad* Linker = LoadPackageLinker(nullptr, PackagePath, LoadFlags, nullptr, nullptr, nullptr, [&PackageToReset, &Result](FLinkerLoad* InLinker)
 	{
 		if (InLinker != nullptr && InLinker->LinkerRoot != nullptr)
 		{
diff --git a/Engine/Source/Runtime/Engine/Private/VT/UploadingVirtualTexture.cpp b/Engine/Source/Runtime/Engine/Private/VT/UploadingVirtualTexture.cpp
--- a/Engine/Source/Runtime/Engine/Private/VT/UploadingVirtualTexture.cpp
+++ b/Engine/Source/Runtime/Engine/Private/VT/UploadingVirtualTexture.cpp
@@ -7,6 +7,8 @@
 #include "FileCache/FileCache.h"
 #include "CrunchCompression.h"
 #include "VirtualTextureChunkDDCCache.h"
+#include "UObject/PackageResourceManager.h"
+#include "Misc/PackageSegment.h"
 
 DECLARE_MEMORY_STAT(TEXT("File Cache Size"), STAT_FileCacheSize, STATGROUP_VirtualTextureMemory);
 DECLARE_MEMORY_STAT(TEXT("Total Header Size"), STAT_TotalHeaderSize, STATGROUP_VirtualTextureMemory);
@@ -359,16 +361,27 @@ FVTDataAndStatus FUploadingVirtualTexture::ReadData(FGraphEventArray& OutComplet
 	TUniquePtr<IFileCacheHandle>& Handle = HandlePerChunk[ChunkIndex];
 	if (!Handle)
 	{
+		enum class EChunkSource
+		{
+			PackagePath,
+			File,
+			BulkData,
+			Invalid
+		};
+		EChunkSource ChunkSource = EChunkSource::Invalid;
+		FPackagePath ChunkPackagePath;
+		EPackageSegment ChunkPackageSegment = EPackageSegment::Header;
 		FString ChunkFileName;
 		int64 ChunkOffsetInFile = 0;
 #if WITH_EDITOR
-		FString ChunkFileNameDCC;
 		// If the bulkdata has a file associated with it, we stream directly from it.
 		// This only happens for lightmaps atm
-		if (!BulkData.GetFilename().IsEmpty())
+		if (!BulkData.GetPackagePath().IsEmpty())
 		{
 			ensure(Size <= (size_t)BulkData.GetBulkDataSize());
-			ChunkFileName = BulkData.GetFilename();
+			ChunkPackagePath = BulkData.GetPackagePath();
+			ChunkPackageSegment = BulkData.GetPackageSegment();
+			ChunkSource = EChunkSource::PackagePath;
 			ChunkOffsetInFile = BulkData.GetBulkDataOffsetInFile();
 		}
 		// Else it should be VT data that is injected into the DDC (and stream from VT DDC cache)
@@ -376,6 +389,7 @@ FVTDataAndStatus FUploadingVirtualTexture::ReadData(FGraphEventArray& OutComplet
 		{
 			SCOPE_CYCLE_COUNTER(STAT_VTP_MakeChunkAvailable);
 			check(Chunk.DerivedDataKey.IsEmpty() == false);
+			FString ChunkFileNameDCC;
 
 			// If request is flagged as high priority, we will block here until DCC cache is populated
 			// This way we can service these high priority tasks immediately
@@ -388,6 +402,7 @@ FVTDataAndStatus FUploadingVirtualTexture::ReadData(FGraphEventArray& OutComplet
 				return EVTRequestPageStatus::Saturated;
 			}
 			ChunkFileName = ChunkFileNameDCC;
+			ChunkSource = EChunkSource::File;
 		}
 #else // WITH_EDITOR
 		ChunkOffsetInFile = BulkData.GetBulkDataOffsetInFile();
@@ -395,22 +410,32 @@ FVTDataAndStatus FUploadingVirtualTexture::ReadData(FGraphEventArray& OutComplet
 		{
 			if (!InvalidChunks[ChunkIndex])
 			{
-				UE_LOG(LogConsoleResponse, Display, TEXT("BulkData for chunk %d in file '%s' is empty."), ChunkIndex, *ChunkFileName);
+				UE_LOG(LogConsoleResponse, Display, TEXT("BulkData for chunk %d in file '%s' is empty."), ChunkIndex, *BulkData.GetPackagePath().GetDebugName());
 				InvalidChunks[ChunkIndex] = true;
 			}
 			return EVTRequestPageStatus::Invalid;
 		}
+		ChunkSource = EChunkSource::BulkData;
 #endif // !WITH_EDITOR
 		// If we have a valid file name then we can create a file handle directly to it.
 		// If we do not then pass in the BulkData object which will create the IAsyncReadFileHandle for us.
-		//
-		if (!ChunkFileName.IsEmpty())
+		switch (ChunkSource)
 		{
-			Handle.Reset(IFileCacheHandle::CreateFileCacheHandle(*ChunkFileName, ChunkOffsetInFile));
-		}
-		else
+		case EChunkSource::PackagePath:
 		{
+			IAsyncReadFileHandle* AsyncFileHandle = IPackageResourceManager::Get().OpenAsyncReadPackage(ChunkPackagePath, ChunkPackageSegment);
+			Handle.Reset(IFileCacheHandle::CreateFileCacheHandle(AsyncFileHandle, ChunkOffsetInFile));
+			break;
+		}
+		case EChunkSource::File:
+			Handle.Reset(IFileCacheHandle::CreateFileCacheHandle(*ChunkFileName, ChunkOffsetInFile));
+			break;
+		case EChunkSource::BulkData:
 			Handle.Reset(IFileCacheHandle::CreateFileCacheHandle(BulkData.OpenAsyncReadHandle(), ChunkOffsetInFile));
+			break;
+		default:
+			check(false);
+			return EVTRequestPageStatus::Invalid;
 		}
 		
 		// Don't expect CreateFileCacheHandle() to fail, async files should never fail to open
@@ -457,10 +482,10 @@ void FUploadingVirtualTexture::DumpToConsole(bool verbose)
 		}
 		else
 #endif
-			BulkDataFiles.Add(Chunk.BulkData.GetFilename());
+			BulkDataFiles.Add(Chunk.BulkData.GetPackagePath().GetLocalFullPath());
 	}
 
-	for (auto FileName : BulkDataFiles)
+	for (const auto& FileName : BulkDataFiles)
 	{
 		UE_LOG(LogConsoleResponse, Display, TEXT("Bulk data file / DDC entry: %s"), *FileName);
 	}
diff --git a/Engine/Source/Runtime/Engine/Private/World.cpp b/Engine/Source/Runtime/Engine/Private/World.cpp
--- a/Engine/Source/Runtime/Engine/Private/World.cpp
+++ b/Engine/Source/Runtime/Engine/Private/World.cpp
@@ -3301,7 +3301,7 @@ UWorld* UWorld::DuplicateWorldForPIE(const FString& PackageName, UWorld* OwningW
 	UPackage* PIELevelPackage = CreatePackage(*PrefixedLevelName);
 	PIELevelPackage->SetPackageFlags(PKG_PlayInEditor);
 	PIELevelPackage->PIEInstanceID = PIEInstanceID;
-	PIELevelPackage->FileName = PackageFName;
+	PIELevelPackage->SetLoadedPath(FPackagePath::FromPackageNameChecked(PackageFName));
 	PIELevelPackage->SetGuid( EditorLevelPackage->GetGuid() );
 	PIELevelPackage->MarkAsFullyLoaded();
 
@@ -6378,16 +6378,26 @@ void FSeamlessTravelHandler::StartLoadingDestination()
 #endif
 
 		// Set the world type in the static map, so that UWorld::PostLoad can set the world type
-		UWorld::WorldTypePreLoadMap.FindOrAdd(FName(*URLMapPackageName)) = CurrentWorld->WorldType;
-
-		LoadPackageAsync(
-			URLMapPackageName, 
-			PendingTravelGuid.IsValid() ? &PendingTravelGuid : NULL,
-			*URLMapPackageToLoadFrom,
-			FLoadPackageAsyncDelegate::CreateRaw(this, &FSeamlessTravelHandler::SeamlessTravelLoadCallback), 			
-			PackageFlags,
-			PIEInstanceID
+		FName PackageName(*URLMapPackageName);
+		UWorld::WorldTypePreLoadMap.FindOrAdd(PackageName) = CurrentWorld->WorldType;
+
+		FStringView PackageNameToLoad(!URLMapPackageToLoadFrom.IsEmpty() ? FStringView(URLMapPackageToLoadFrom) : FStringView(URLMapPackageName));
+		FPackagePath PackagePath;
+		if (FPackagePath::TryFromMountedName(PackageNameToLoad, PackagePath))
+		{
+			LoadPackageAsync(
+				PackagePath,
+				PackageName,
+				FLoadPackageAsyncDelegate::CreateRaw(this, &FSeamlessTravelHandler::SeamlessTravelLoadCallback),
+				PendingTravelGuid.IsValid() ? &PendingTravelGuid : NULL,
+				PackageFlags,
+				PIEInstanceID
 			);
+		}
+		else
+		{
+			UE_LOG(LogWorld, Error, TEXT("StartLoadingDestination: Invalid PackagePath \"%.*s\""), PackageNameToLoad.Len(), PackageNameToLoad.GetData());
+		}
 	}
 	else
 	{
@@ -7718,7 +7728,7 @@ static ULevel* DuplicateLevelWithPrefix(ULevel* InLevel, int32 InstanceID )
 	UPackage* NewPackage = CreatePackage( *PrefixedPackageName );
 	NewPackage->SetPackageFlags( PKG_PlayInEditor );
 	NewPackage->PIEInstanceID = InstanceID;
-	NewPackage->FileName = OriginalPackage->FileName;
+	NewPackage->SetLoadedPath(OriginalPackage->GetLoadedPath());
 	NewPackage->SetGuid( OriginalPackage->GetGuid() );
 	NewPackage->MarkAsFullyLoaded();
 
diff --git a/Engine/Source/Runtime/Engine/Private/WorldPartition/WorldPartition.cpp b/Engine/Source/Runtime/Engine/Private/WorldPartition/WorldPartition.cpp
--- a/Engine/Source/Runtime/Engine/Private/WorldPartition/WorldPartition.cpp
+++ b/Engine/Source/Runtime/Engine/Private/WorldPartition/WorldPartition.cpp
@@ -223,11 +223,11 @@ void UWorldPartition::Initialize(UWorld* InWorld, const FTransform& InTransform)
 		TArray<FAssetData> Assets;
 		UPackage* LevelPackage = OuterWorld->PersistentLevel->GetOutermost();
 
-		if (!LevelPackage->FileName.IsNone())
+		if (!LevelPackage->GetLoadedPath().GetPackageFName().IsNone())
 		{
 			TGuardValue<bool> IgnoreAssetRegistryEvents(bIgnoreAssetRegistryEvents, true);
 
-			const FString LevelPathStr = LevelPackage->FileName.ToString();
+			const FString LevelPathStr = LevelPackage->GetLoadedPath().GetPackageName();
 			const FString LevelExternalActorsPath = ULevel::GetExternalActorsPath(LevelPathStr);
 
 			// Do a synchronous scan of the level external actors path.			
@@ -249,9 +249,9 @@ void UWorldPartition::Initialize(UWorld* InWorld, const FTransform& InTransform)
 
 		if (bIsInstanced)
 		{
-			InstancingContext.AddMapping(LevelPackage->FileName, LevelPackage->GetFName());
+			InstancingContext.AddMapping(LevelPackage->GetLoadedPath().GetPackageFName(), LevelPackage->GetFName());
 
-			const FString SourceWorldName = FPaths::GetBaseFilename(LevelPackage->FileName.ToString());
+			const FString SourceWorldName = FPaths::GetBaseFilename(LevelPackage->GetLoadedPath().GetPackageName());
 			const FString DestWorldName = FPaths::GetBaseFilename(LevelPackage->GetFName().ToString());
 
 			ReplaceFrom = SourceWorldName + TEXT(".") + SourceWorldName;
@@ -940,7 +940,7 @@ bool UWorldPartition::ShouldHandleAssetEvent(const FAssetData& InAssetData)
 		return InValue;
 	};
 
-	const FString ThisLevelPath = RemoveAfterFirstDot(GetPackage()->FileName.ToString());
+	const FString ThisLevelPath = GetPackage()->GetLoadedPath().GetPackageName();
 	const FString AssetLevelPath = RemoveAfterFirstDot(InAssetData.ObjectPath.ToString());
 	return (ThisLevelPath == AssetLevelPath);
 }
diff --git a/Engine/Source/Runtime/Engine/Private/WorldPartition/WorldPartitionLevelStreamingDynamic.cpp b/Engine/Source/Runtime/Engine/Private/WorldPartition/WorldPartitionLevelStreamingDynamic.cpp
--- a/Engine/Source/Runtime/Engine/Private/WorldPartition/WorldPartitionLevelStreamingDynamic.cpp
+++ b/Engine/Source/Runtime/Engine/Private/WorldPartition/WorldPartitionLevelStreamingDynamic.cpp
@@ -57,7 +57,7 @@ void UWorldPartitionLevelStreamingDynamic::Initialize(const UWorldPartitionRunti
 	ChildPackages = InCell.GetPackages();
 
 	UWorld* OuterWorld = InCell.GetOuterUWorldPartition()->GetTypedOuter<UWorld>();
-	OriginalLevelPackageName = OuterWorld->GetPackage()->FileName;
+	OriginalLevelPackageName = OuterWorld->GetPackage()->GetLoadedPath().GetPackageFName();
 	PackageNameToLoad = GetWorldAssetPackageFName();
 	OuterWorldPartition = OuterWorld->GetWorldPartition();
 }
@@ -280,7 +280,10 @@ bool UWorldPartitionLevelStreamingDynamic::IssueLoadRequests()
 			}
 		});
 
-		LoadPackageAsync(*ActorPackageInstanceNames[ChildIndex], nullptr, *ActorPackageName, CompletionCallback, PKG_PlayInEditor, RuntimePackage->PIEInstanceID, 0, &InstancingContext);
+		FPackagePath PackagePath;
+		FPackagePath::TryFromMountedName(!ActorPackageName.IsEmpty() ? ActorPackageName : ActorPackageInstanceNames[ChildIndex], PackagePath);
+		FName PackageName(*ActorPackageInstanceNames[ChildIndex]);
+		LoadPackageAsync(PackagePath, PackageName, CompletionCallback, nullptr, PKG_PlayInEditor, RuntimePackage->PIEInstanceID, 0, &InstancingContext);
 	}
 
 	return !!NumPendingActorRequests;
diff --git a/Engine/Source/Runtime/Landscape/Classes/LandscapeComponent.h b/Engine/Source/Runtime/Landscape/Classes/LandscapeComponent.h
--- a/Engine/Source/Runtime/Landscape/Classes/LandscapeComponent.h
+++ b/Engine/Source/Runtime/Landscape/Classes/LandscapeComponent.h
@@ -121,7 +121,8 @@ class FLandscapeComponentDerivedData
 	/** Cached render data. Only valid on device. */
 	TSharedPtr<FLandscapeMobileRenderData, ESPMode::ThreadSafe > CachedRenderData;
 
-	FString CachedLODDataFileName;
+	FPackagePath CachedLODDataPackagePath;
+	EPackageSegment CachedLODDataPackageSegment;
 
 	friend class ULandscapeLODStreamingProxy;
 
@@ -398,7 +399,9 @@ class ULandscapeLODStreamingProxy : public UStreamableRenderAsset
 	virtual EStreamableRenderAssetType GetRenderAssetType() const final override { return EStreamableRenderAssetType::LandscapeMeshMobile; }
 	//~ End UStreamableRenderAsset Interface
 
+	UE_DEPRECATED(5.0, "Use GetMipDataPackagePath instead")
 	LANDSCAPE_API bool GetMipDataFilename(const int32 MipIndex, FString& OutBulkDataFilename) const;
+	LANDSCAPE_API bool GetMipDataPackagePath(const int32 MipIndex, FPackagePath& OutBulkDataPackagePath, EPackageSegment& OutPackageSegment) const;
 
 
 	LANDSCAPE_API TArray<float> GetLODScreenSizeArray() const;
diff --git a/Engine/Source/Runtime/Landscape/Private/Landscape.cpp b/Engine/Source/Runtime/Landscape/Private/Landscape.cpp
--- a/Engine/Source/Runtime/Landscape/Private/Landscape.cpp
+++ b/Engine/Source/Runtime/Landscape/Private/Landscape.cpp
@@ -29,6 +29,7 @@ Landscape.cpp: Terrain rendering
 #include "Logging/MessageLog.h"
 #include "Misc/UObjectToken.h"
 #include "Misc/MapErrors.h"
+#include "Misc/PackageSegment.h"
 #include "DerivedDataCacheInterface.h"
 #include "Interfaces/ITargetPlatform.h"
 #include "LandscapeMeshCollisionComponent.h"
@@ -3530,7 +3531,8 @@ void FLandscapeComponentDerivedData::Serialize(FArchive& Ar, UObject* Owner)
 		StreamingLODDataArray.AddDefaulted(NumStreamingLODs);
 	}
 
-	CachedLODDataFileName.Empty();
+	CachedLODDataPackagePath.Empty();
+	CachedLODDataPackageSegment = EPackageSegment::Header;
 
 	for (int32 Idx = 0; Idx < NumStreamingLODs; ++Idx)
 	{
@@ -3542,9 +3544,10 @@ void FLandscapeComponentDerivedData::Serialize(FArchive& Ar, UObject* Owner)
 		FByteBulkData& LODData = StreamingLODDataArray[Idx];
 		LODData.Serialize(Ar, Owner, Idx);
 #endif
-		if (CachedLODDataFileName.IsEmpty() && !!(LODData.GetBulkDataFlags() & BULKDATA_Force_NOT_InlinePayload))
+		if (CachedLODDataPackagePath.IsEmpty() && !!(LODData.GetBulkDataFlags() & BULKDATA_Force_NOT_InlinePayload))
 		{
-			CachedLODDataFileName = LODData.GetFilename();
+			CachedLODDataPackagePath = LODData.GetPackagePath();
+			CachedLODDataPackageSegment = LODData.GetPackageSegment();
 		}
 	}
 }
@@ -3989,12 +3992,25 @@ int32 ULandscapeLODStreamingProxy::CalcCumulativeLODSize(int32 NumLODs) const
 }
 
 bool ULandscapeLODStreamingProxy::GetMipDataFilename(const int32 MipIndex, FString& OutBulkDataFilename) const
+{
+	FPackagePath PackagePath;
+	EPackageSegment PackageSegment;
+	if (GetMipDataPackagePath(MipIndex, PackagePath, PackageSegment))
+	{
+		OutBulkDataFilename = PackagePath.GetLocalFullPath(PackageSegment);
+		return true;
+	}
+	return false;
+}
+
+bool ULandscapeLODStreamingProxy::GetMipDataPackagePath(const int32 MipIndex, FPackagePath& OutPackagePath, EPackageSegment& OutPackageSegment) const
 {
 	check(LandscapeComponent);
 	const int32 NumStreamingLODs = LandscapeComponent->PlatformData.StreamingLODDataArray.Num();
 	if (MipIndex >= 0 && MipIndex < NumStreamingLODs)
 	{
-		OutBulkDataFilename = LandscapeComponent->PlatformData.CachedLODDataFileName;
+		OutPackagePath = LandscapeComponent->PlatformData.CachedLODDataPackagePath;
+		OutPackageSegment = LandscapeComponent->PlatformData.CachedLODDataPackageSegment;
 		return true;
 	}
 	return false;
@@ -4003,10 +4019,11 @@ bool ULandscapeLODStreamingProxy::GetMipDataFilename(const int32 MipIndex, FStri
 FIoFilenameHash ULandscapeLODStreamingProxy::GetMipIoFilenameHash(const int32 MipIndex) const
 {
 #if LANDSCAPE_LOD_STREAMING_USE_TOKEN
-	FString MipFilename;
-	if (GetMipDataFilename(MipIndex, MipFilename))
+	FPackagePath PackagePath;
+	EPackageSegment PackageSegment;
+	if (GetMipDataPackagePath(MipIndex, PackagePath, PackageSegment))
 	{
-		return MakeIoFilenameHash(MipFilename);
+		return MakeIoFilenameHash(PackagePath);
 	}
 #else
 	if (LandscapeComponent && LandscapeComponent->PlatformData.StreamingLODDataArray.IsValidIndex(MipIndex))
diff --git a/Engine/Source/Runtime/Launch/Private/LaunchEngineLoop.cpp b/Engine/Source/Runtime/Launch/Private/LaunchEngineLoop.cpp
--- a/Engine/Source/Runtime/Launch/Private/LaunchEngineLoop.cpp
+++ b/Engine/Source/Runtime/Launch/Private/LaunchEngineLoop.cpp
@@ -72,6 +72,7 @@
 	#include "UObject/Package.h"
 	#include "UObject/Linker.h"
 	#include "UObject/LinkerLoad.h"
+	#include "UObject/PackageResourceManager.h"
 	#include "UObject/ReferencerFinder.h"
 #endif
 
@@ -3093,12 +3094,15 @@ int32 FEngineLoop::PreInitPostStartupScreen(const TCHAR* CmdLine)
 
 		SlowTask.EnterProgressFrame(5);
 
-#if USE_EVENT_DRIVEN_ASYNC_LOAD_AT_BOOT_TIME && !USE_PER_MODULE_UOBJECT_BOOTSTRAP
 		{
 		    SCOPED_BOOT_TIMING("LoadAssetRegistryModule");
 		    // If we don't do this now and the async loading thread is active, then we will attempt to load this module from a thread
 		    FModuleManager::Get().LoadModule("AssetRegistry");
 		}
+#if WITH_COREUOBJECT
+		// Initialize the PackageResourceManager, which is needed to load any (non-script) Packages. It is first used in ProcessNewlyLoadedObjects (due to the loading of asset references in Class Default Objects)
+		// It has to be intialized after the AssetRegistryModule; the editor implementations of PackageResourceManager relies on it
+		IPackageResourceManager::Initialize();
 #endif
 		EndInitGameTextLocalization();
 
@@ -4277,6 +4281,11 @@ void FEngineLoop::Exit()
 	}
 #endif
 
+#if WITH_COREUOBJECT
+	// PackageResourceManager depends on AssetRegistry, so must be shutdown before we unload the AssetRegistry module
+	IPackageResourceManager::Shutdown();
+#endif
+
 #if WITH_EDITOR
 	// These module must be shut down first because other modules may try to access them during shutdown.
 	// Accessing these modules at shutdown causes instability since the object system will have been shut down and these modules uses uobjects internally.
