Hide strata MSM from UI.

diff --git a/Engine/Source/Runtime/Engine/Classes/Engine/EngineTypes.h b/Engine/Source/Runtime/Engine/Classes/Engine/EngineTypes.h
--- a/Engine/Source/Runtime/Engine/Classes/Engine/EngineTypes.h
+++ b/Engine/Source/Runtime/Engine/Classes/Engine/EngineTypes.h
@@ -476,7 +476,7 @@ enum EMaterialShadingModel
 	MSM_Eye						UMETA(DisplayName="Eye"),
 	MSM_SingleLayerWater		UMETA(DisplayName="SingleLayerWater"),
 	MSM_ThinTranslucent			UMETA(DisplayName="Thin Translucent"),
-	MSM_Strata					UMETA(DisplayName="Strata"),
+	MSM_Strata					UMETA(DisplayName="Strata", Hidden),
 	/** Number of unique shading models. */
 	MSM_NUM						UMETA(Hidden),
 	/** Shading model will be determined by the Material Expression Graph,
