Strata material and categories update

diff --git a/Engine/Source/Runtime/Engine/Private/Materials/MaterialExpressions.cpp b/Engine/Source/Runtime/Engine/Private/Materials/MaterialExpressions.cpp
--- a/Engine/Source/Runtime/Engine/Private/Materials/MaterialExpressions.cpp
+++ b/Engine/Source/Runtime/Engine/Private/Materials/MaterialExpressions.cpp
@@ -19108,7 +19108,7 @@ UMaterialExpressionStrataDiffuseBSDF::UMaterialExpressionStrataDiffuseBSDF(const
 	struct FConstructorStatics
 	{
 		FText NAME_Strata;
-		FConstructorStatics() : NAME_Strata(LOCTEXT("Strata", "Strata")) { }
+		FConstructorStatics() : NAME_Strata(LOCTEXT("Strata BSDFs", "Strata BSDFs")) { }
 	};
 	static FConstructorStatics ConstructorStatics;
 #if WITH_EDITORONLY_DATA
@@ -19175,7 +19175,7 @@ UMaterialExpressionStrataDielectricBSDF::UMaterialExpressionStrataDielectricBSDF
 	struct FConstructorStatics
 	{
 		FText NAME_Strata;
-		FConstructorStatics() : NAME_Strata(LOCTEXT("Strata", "Strata")) { }
+		FConstructorStatics() : NAME_Strata(LOCTEXT("Strata BSDFs", "Strata BSDFs")) { }
 	};
 	static FConstructorStatics ConstructorStatics;
 #if WITH_EDITORONLY_DATA
@@ -19259,7 +19259,7 @@ UMaterialExpressionStrataConductorBSDF::UMaterialExpressionStrataConductorBSDF(c
 	struct FConstructorStatics
 	{
 		FText NAME_Strata;
-		FConstructorStatics() : NAME_Strata(LOCTEXT("Strata", "Strata")) { }
+		FConstructorStatics() : NAME_Strata(LOCTEXT("Strata BSDFs", "Strata BSDFs")) { }
 	};
 	static FConstructorStatics ConstructorStatics;
 #if WITH_EDITORONLY_DATA
@@ -19343,7 +19343,7 @@ UMaterialExpressionStrataVolumeBSDF::UMaterialExpressionStrataVolumeBSDF(const F
 	struct FConstructorStatics
 	{
 		FText NAME_Strata;
-		FConstructorStatics() : NAME_Strata(LOCTEXT("Strata", "Strata")) { }
+		FConstructorStatics() : NAME_Strata(LOCTEXT("Strata BSDFs", "Strata BSDFs")) { }
 	};
 	static FConstructorStatics ConstructorStatics;
 #if WITH_EDITORONLY_DATA
@@ -19418,7 +19418,7 @@ UMaterialExpressionStrataSheenBSDF::UMaterialExpressionStrataSheenBSDF(const FOb
 	struct FConstructorStatics
 	{
 		FText NAME_Strata;
-		FConstructorStatics() : NAME_Strata(LOCTEXT("Strata", "Strata")) { }
+		FConstructorStatics() : NAME_Strata(LOCTEXT("Strata BSDFs", "Strata BSDFs")) { }
 	};
 	static FConstructorStatics ConstructorStatics;
 #if WITH_EDITORONLY_DATA
@@ -19485,7 +19485,7 @@ UMaterialExpressionStrataHorizontalMixing::UMaterialExpressionStrataHorizontalMi
 	struct FConstructorStatics
 	{
 		FText NAME_Strata;
-		FConstructorStatics() : NAME_Strata(LOCTEXT("Strata", "Strata")) { }
+		FConstructorStatics() : NAME_Strata(LOCTEXT("Strata Ops", "Strata Ops")) { }
 	};
 	static FConstructorStatics ConstructorStatics;
 #if WITH_EDITORONLY_DATA
@@ -19552,7 +19552,7 @@ UMaterialExpressionStrataVerticalLayering::UMaterialExpressionStrataVerticalLaye
 	struct FConstructorStatics
 	{
 		FText NAME_Strata;
-		FConstructorStatics() : NAME_Strata(LOCTEXT("Strata", "Strata")) { }
+		FConstructorStatics() : NAME_Strata(LOCTEXT("Strata Ops", "Strata Ops")) { }
 	};
 	static FConstructorStatics ConstructorStatics;
 #if WITH_EDITORONLY_DATA
@@ -19616,7 +19616,7 @@ UMaterialExpressionStrataAdd::UMaterialExpressionStrataAdd(const FObjectInitiali
 	struct FConstructorStatics
 	{
 		FText NAME_Strata;
-		FConstructorStatics() : NAME_Strata(LOCTEXT("Strata", "Strata")) { }
+		FConstructorStatics() : NAME_Strata(LOCTEXT("Strata Ops", "Strata Ops")) { }
 	};
 	static FConstructorStatics ConstructorStatics;
 #if WITH_EDITORONLY_DATA
@@ -19680,7 +19680,7 @@ UMaterialExpressionStrataMultiply::UMaterialExpressionStrataMultiply(const FObje
 	struct FConstructorStatics
 	{
 		FText NAME_Strata;
-		FConstructorStatics() : NAME_Strata(LOCTEXT("Strata", "Strata")) { }
+		FConstructorStatics() : NAME_Strata(LOCTEXT("Strata Ops", "Strata Ops")) { }
 	};
 	static FConstructorStatics ConstructorStatics;
 #if WITH_EDITORONLY_DATA
@@ -19744,7 +19744,7 @@ UMaterialExpressionStrataArtisticIOR::UMaterialExpressionStrataArtisticIOR(const
 	struct FConstructorStatics
 	{
 		FText NAME_Strata;
-		FConstructorStatics() : NAME_Strata(LOCTEXT("Strata", "Strata")) { }
+		FConstructorStatics() : NAME_Strata(LOCTEXT("Strata Helpers", "Strata Helpers")) { }
 	};
 	static FConstructorStatics ConstructorStatics;
 #if WITH_EDITORONLY_DATA
@@ -19791,7 +19791,7 @@ UMaterialExpressionStrataPhysicalIOR::UMaterialExpressionStrataPhysicalIOR(const
 	struct FConstructorStatics
 	{
 		FText NAME_Strata;
-		FConstructorStatics() : NAME_Strata(LOCTEXT("Strata", "Strata")) { }
+		FConstructorStatics() : NAME_Strata(LOCTEXT("Strata Helpers", "Strata Helpers")) { }
 	};
 	static FConstructorStatics ConstructorStatics;
 #if WITH_EDITORONLY_DATA
@@ -19839,7 +19839,7 @@ UMaterialExpressionStrataAnisotropyToRoughness::UMaterialExpressionStrataAnisotr
 	struct FConstructorStatics
 	{
 		FText NAME_Strata;
-		FConstructorStatics() : NAME_Strata(LOCTEXT("Strata", "Strata")) { }
+		FConstructorStatics() : NAME_Strata(LOCTEXT("Strata Helpers", "Strata Helpers")) { }
 	};
 	static FConstructorStatics ConstructorStatics;
 #if WITH_EDITORONLY_DATA
