Swaped strata forground background to match lerp as reported for concistency

diff --git a/Engine/Source/Runtime/Engine/Classes/Materials/MaterialExpressionStrata.h b/Engine/Source/Runtime/Engine/Classes/Materials/MaterialExpressionStrata.h
--- a/Engine/Source/Runtime/Engine/Classes/Materials/MaterialExpressionStrata.h
+++ b/Engine/Source/Runtime/Engine/Classes/Materials/MaterialExpressionStrata.h
@@ -263,13 +263,13 @@ class UMaterialExpressionStrataHorizontalMixing : public UMaterialExpression
 	 * Strata material
 	 */
 	UPROPERTY()
-	FExpressionInput Foreground;
+	FExpressionInput Background;
 
 	/**
 	 * Strata material
 	 */
 	UPROPERTY()
-	FExpressionInput Background;
+	FExpressionInput Foreground;
 
 	/**
 	 * Lerp factor between Background (Mix == 0) and Foreground (Mix == 1).
