Fixed name.

diff --git a/Engine/Source/Runtime/Engine/Private/Materials/MaterialExpressions.cpp b/Engine/Source/Runtime/Engine/Private/Materials/MaterialExpressions.cpp
--- a/Engine/Source/Runtime/Engine/Private/Materials/MaterialExpressions.cpp
+++ b/Engine/Source/Runtime/Engine/Private/Materials/MaterialExpressions.cpp
@@ -19816,7 +19816,7 @@ int32 UMaterialExpressionStrataPhysicalIOR::Compile(class FMaterialCompiler* Com
 
 void UMaterialExpressionStrataPhysicalIOR::GetCaption(TArray<FString>& OutCaptions) const
 {
-	OutCaptions.Add(TEXT("Strata Artistic IOR"));
+	OutCaptions.Add(TEXT("Strata Physical IOR"));
 }
 
 uint32 UMaterialExpressionStrataPhysicalIOR::GetOutputType(int32 OutputIndex)
