Fixed blending of mateiral layers when strata is not used (default material code chunk needs to be registered too).

diff --git a/Engine/Source/Editor/MaterialEditor/Private/MaterialEditor.cpp b/Engine/Source/Editor/MaterialEditor/Private/MaterialEditor.cpp
--- a/Engine/Source/Editor/MaterialEditor/Private/MaterialEditor.cpp
+++ b/Engine/Source/Editor/MaterialEditor/Private/MaterialEditor.cpp
@@ -234,7 +234,7 @@ int32 FMatExpressionPreview::CompilePropertyAndSetMaterialProperty(EMaterialProp
 	}
 	else if (Property == MP_FrontMaterial)
 	{
-		Ret = Compiler->FrontMaterial();
+		Ret = Compiler->StrataCreateAndRegisterNullMaterial();
 	}
 	else
 	{
diff --git a/Engine/Source/Runtime/Engine/Private/Materials/HLSLMaterialTranslator.cpp b/Engine/Source/Runtime/Engine/Private/Materials/HLSLMaterialTranslator.cpp
--- a/Engine/Source/Runtime/Engine/Private/Materials/HLSLMaterialTranslator.cpp
+++ b/Engine/Source/Runtime/Engine/Private/Materials/HLSLMaterialTranslator.cpp
@@ -1109,55 +1109,60 @@ bool FHLSLMaterialTranslator::Translate()
 		{
 			const FStrataMaterialCompilationInfo& StrataCompilationInfo = CodeChunkToStrataCompilationInfoMap[Chunk[MP_FrontMaterial]];
 
-			static const auto CVarStrataBytePerPixel = IConsoleManager::Get().FindTConsoleVariableDataInt(TEXT("r.Strata.BytesPerPixel"));
-			check(CVarStrataBytePerPixel)
-			const uint32 StrataBytePerPixel = CVarStrataBytePerPixel ? CVarStrataBytePerPixel->GetValueOnAnyThread() : 0;
-			StrataMaterialAnalysis = StrataCompilationInfoMaterialAnalysis(this, StrataCompilationInfo, StrataBytePerPixel);
-
-			FString StrataMaterialDescription;
+			// The material can be null when some entries are automatically generated, for instance in the mateiral layer blendign system
+			const bool bIsNullMaterial = (StrataCompilationInfo.LayerCount == 0 && StrataCompilationInfo.TotalBSDFCount == 0);
+			if (!bIsNullMaterial)
+			{
+				static const auto CVarStrataBytePerPixel = IConsoleManager::Get().FindTConsoleVariableDataInt(TEXT("r.Strata.BytesPerPixel"));
+				check(CVarStrataBytePerPixel)
+					const uint32 StrataBytePerPixel = CVarStrataBytePerPixel ? CVarStrataBytePerPixel->GetValueOnAnyThread() : 0;
+				StrataMaterialAnalysis = StrataCompilationInfoMaterialAnalysis(this, StrataCompilationInfo, StrataBytePerPixel);
 
-			StrataMaterialDescription += FString::Printf(TEXT("----- STRATA -----\r\n"));
-			StrataMaterialDescription += FString::Printf(TEXT("StrataCompilationInfo -\r\n"));
-			StrataMaterialDescription += FString::Printf(TEXT(" - TotalBSDFCount    = %i\r\n"), StrataCompilationInfo.TotalBSDFCount);
-			StrataMaterialDescription += FString::Printf(TEXT(" - SharedNormalCount = %i\r\n"), StrataCompilationInfoGetSharedNormalCount());
+				FString StrataMaterialDescription;
 
-			for (uint32 LayerIt = 0; LayerIt < StrataCompilationInfo.LayerCount; ++LayerIt)
-			{
-				StrataMaterialDescription += FString::Printf(TEXT(" Layer %i BSDFs:\r\n"), LayerIt);
+				StrataMaterialDescription += FString::Printf(TEXT("----- STRATA -----\r\n"));
+				StrataMaterialDescription += FString::Printf(TEXT("StrataCompilationInfo -\r\n"));
+				StrataMaterialDescription += FString::Printf(TEXT(" - TotalBSDFCount    = %i\r\n"), StrataCompilationInfo.TotalBSDFCount);
+				StrataMaterialDescription += FString::Printf(TEXT(" - SharedNormalCount = %i\r\n"), StrataCompilationInfoGetSharedNormalCount());
 
-				const FStrataMaterialCompilationInfo::FLayer& Layer = StrataCompilationInfo.Layers[LayerIt];
-				for (uint32 BSDFIt = 0; BSDFIt < Layer.BSDFCount; ++BSDFIt)
+				for (uint32 LayerIt = 0; LayerIt < StrataCompilationInfo.LayerCount; ++LayerIt)
 				{
-					StrataMaterialDescription += FString::Printf(TEXT("     - %s - SharedNormalIndex = %u \r\n"), *GetStrataBSDFName(Layer.BSDFs[BSDFIt].Type), Layer.BSDFs[BSDFIt].SharedNormalIndex);
+					StrataMaterialDescription += FString::Printf(TEXT(" Layer %i BSDFs:\r\n"), LayerIt);
+
+					const FStrataMaterialCompilationInfo::FLayer& Layer = StrataCompilationInfo.Layers[LayerIt];
+					for (uint32 BSDFIt = 0; BSDFIt < Layer.BSDFCount; ++BSDFIt)
+					{
+						StrataMaterialDescription += FString::Printf(TEXT("     - %s - SharedNormalIndex = %u \r\n"), *GetStrataBSDFName(Layer.BSDFs[BSDFIt].Type), Layer.BSDFs[BSDFIt].SharedNormalIndex);
+					}
 				}
-			}
 
-			StrataMaterialDescription += FString::Printf(TEXT("Byte Per Pixel Budget      %u\r\n"), StrataBytePerPixel);
-			StrataMaterialDescription += FString::Printf(TEXT("Result.bFitInMemoryBudget  %s\r\n"), StrataMaterialAnalysis.bFitInMemoryBudget ? TEXT("YES") : TEXT("NO"));
-			StrataMaterialDescription += FString::Printf(TEXT("Result.RequestedLayerCount %u\r\n"), StrataMaterialAnalysis.RequestedLayerCount);
-			StrataMaterialDescription += FString::Printf(TEXT("Result.RequestedBSDFCount  %u\r\n"), StrataMaterialAnalysis.RequestedBSDFCount);
-			StrataMaterialDescription += FString::Printf(TEXT("Result.RequestedByteCount  %u\r\n"), StrataMaterialAnalysis.RequestedByteCount);
-			if (!StrataMaterialAnalysis.bFitInMemoryBudget)
-			{
-				StrataMaterialDescription += FString::Printf(TEXT("Result.ClampedLayerCount   %u\r\n"), StrataMaterialAnalysis.ClampedLayerCount);
-				StrataMaterialDescription += FString::Printf(TEXT("Result.ClampedBSDFCount    %u\r\n"), StrataMaterialAnalysis.ClampedBSDFCount);
-				StrataMaterialDescription += FString::Printf(TEXT("Result.UsedByteCount       %u\r\n"), StrataMaterialAnalysis.UsedByteCount);
-			}
-			StrataMaterialDescription += FString::Printf(TEXT("------------------\r\n"));
+				StrataMaterialDescription += FString::Printf(TEXT("Byte Per Pixel Budget      %u\r\n"), StrataBytePerPixel);
+				StrataMaterialDescription += FString::Printf(TEXT("Result.bFitInMemoryBudget  %s\r\n"), StrataMaterialAnalysis.bFitInMemoryBudget ? TEXT("YES") : TEXT("NO"));
+				StrataMaterialDescription += FString::Printf(TEXT("Result.RequestedLayerCount %u\r\n"), StrataMaterialAnalysis.RequestedLayerCount);
+				StrataMaterialDescription += FString::Printf(TEXT("Result.RequestedBSDFCount  %u\r\n"), StrataMaterialAnalysis.RequestedBSDFCount);
+				StrataMaterialDescription += FString::Printf(TEXT("Result.RequestedByteCount  %u\r\n"), StrataMaterialAnalysis.RequestedByteCount);
+				if (!StrataMaterialAnalysis.bFitInMemoryBudget)
+				{
+					StrataMaterialDescription += FString::Printf(TEXT("Result.ClampedLayerCount   %u\r\n"), StrataMaterialAnalysis.ClampedLayerCount);
+					StrataMaterialDescription += FString::Printf(TEXT("Result.ClampedBSDFCount    %u\r\n"), StrataMaterialAnalysis.ClampedBSDFCount);
+					StrataMaterialDescription += FString::Printf(TEXT("Result.UsedByteCount       %u\r\n"), StrataMaterialAnalysis.UsedByteCount);
+				}
+				StrataMaterialDescription += FString::Printf(TEXT("------------------\r\n"));
 
-			ResourcesString += TEXT("/*");
-			ResourcesString += StrataMaterialDescription;
-			ResourcesString += TEXT("*/");
+				ResourcesString += TEXT("/*");
+				ResourcesString += StrataMaterialDescription;
+				ResourcesString += TEXT("*/");
 
-			MaterialCompilationOutput.StrataMaterialDescription = StrataMaterialDescription;
+				MaterialCompilationOutput.StrataMaterialDescription = StrataMaterialDescription;
 
-			if (StrataMaterialAnalysis.ClampedLayerCount == 0)
-			{
-				UE_LOG(LogMaterial, Fatal, TEXT("Material %s cannot have any layers rendered due to its complexity (asset: %s).\r\n"), *Material->GetDebugName(), *Material->GetAssetPath());
-			}
-			else if (StrataMaterialAnalysis.RequestedLayerCount > StrataMaterialAnalysis.ClampedLayerCount)
-			{
-				UE_LOG(LogMaterial, Warning, TEXT("Material %s is not fitting the allocated byte per pixel budget for starta. Layers will be removed (asset: %s).\r\n"), *Material->GetDebugName(), *Material->GetAssetPath());
+				if (StrataMaterialAnalysis.ClampedLayerCount == 0)
+				{
+					UE_LOG(LogMaterial, Fatal, TEXT("Material %s cannot have any layers rendered due to its complexity (asset: %s).\r\n"), *Material->GetDebugName(), *Material->GetAssetPath());
+				}
+				else if (StrataMaterialAnalysis.RequestedLayerCount > StrataMaterialAnalysis.ClampedLayerCount)
+				{
+					UE_LOG(LogMaterial, Warning, TEXT("Material %s is not fitting the allocated byte per pixel budget for starta. Layers will be removed (asset: %s).\r\n"), *Material->GetDebugName(), *Material->GetAssetPath());
+				}
 			}
 		}
 		else
@@ -7438,9 +7443,11 @@ uint8 FHLSLMaterialTranslator::StrataCompilationInfoGetSharedNormalCount()
 	return CodeChunkToStrataSharedNormal.Num();
 }
 
-int32 FHLSLMaterialTranslator::FrontMaterial()
+int32 FHLSLMaterialTranslator::StrataCreateAndRegisterNullMaterial()
 {
-	return AddInlinedCodeChunk(MCT_Strata, TEXT("GetInitialisedStrataData()"));
+	int32 OutputCodeChunk = AddInlinedCodeChunk(MCT_Strata, TEXT("GetInitialisedStrataData()"));
+	StrataCompilationInfoCreateNullBSDF(this, OutputCodeChunk);
+	return OutputCodeChunk;
 }
 
 int32 FHLSLMaterialTranslator::StrataDiffuseBSDF(int32 Albedo, int32 Roughness, int32 Normal, uint8 SharedNormalIndex)
diff --git a/Engine/Source/Runtime/Engine/Private/Materials/HLSLMaterialTranslator.h b/Engine/Source/Runtime/Engine/Private/Materials/HLSLMaterialTranslator.h
--- a/Engine/Source/Runtime/Engine/Private/Materials/HLSLMaterialTranslator.h
+++ b/Engine/Source/Runtime/Engine/Private/Materials/HLSLMaterialTranslator.h
@@ -790,7 +790,7 @@ protected:
 	virtual int32 VirtualTextureOutput(uint8 MaterialAttributeMask) override;
 
 	// Strata
-	virtual int32 FrontMaterial() override;
+	virtual int32 StrataCreateAndRegisterNullMaterial() override;
 	virtual int32 StrataDiffuseBSDF(int32 Albedo, int32 Roughness, int32 Normal, uint8 SharedNormalIndex) override;
 	virtual int32 StrataDielectricBSDF(int32 RoughnessX, int32 RoughnessY, int32 IOR, int32 Tint, int32 Normal, int32 Tangent, uint8 SharedNormalIndex) override;
 	virtual int32 StrataConductorBSDF(int32 Reflectivity, int32 EdgeColor, int32 RoughnessX, int32 RoughnessY, int32 Normal, int32 Tangent, uint8 SharedNormalIndex) override;
diff --git a/Engine/Source/Runtime/Engine/Private/Materials/MaterialShared.cpp b/Engine/Source/Runtime/Engine/Private/Materials/MaterialShared.cpp
--- a/Engine/Source/Runtime/Engine/Private/Materials/MaterialShared.cpp
+++ b/Engine/Source/Runtime/Engine/Private/Materials/MaterialShared.cpp
@@ -4099,8 +4099,7 @@ int32 FMaterialAttributeDefintion::CompileDefaultValue(FMaterialCompiler* Compil
 	if (Property == MP_FrontMaterial)
 	{
 		check(ValueType == MCT_Strata);
-		// Default to the first shading model of the material. If the material is using a single shading model selected through the dropdown, this is how it gets written to the shader as a constant (optimizing out all the dynamic branches)
-		return Compiler->FrontMaterial();
+		return Compiler->StrataCreateAndRegisterNullMaterial();
 	}
 
 	if (TexCoordIndex == INDEX_NONE)
diff --git a/Engine/Source/Runtime/Engine/Private/Materials/StrataMaterial.cpp b/Engine/Source/Runtime/Engine/Private/Materials/StrataMaterial.cpp
--- a/Engine/Source/Runtime/Engine/Private/Materials/StrataMaterial.cpp
+++ b/Engine/Source/Runtime/Engine/Private/Materials/StrataMaterial.cpp
@@ -61,6 +61,14 @@ void StrataCompilationInfoCreateSingleBSDFMaterial(FMaterialCompiler* Compiler,
 	Compiler->StrataCompilationInfoRegisterCodeChunk(CodeChunk, StrataInfo);
 }
 
+void  StrataCompilationInfoCreateNullBSDF(FMaterialCompiler* Compiler, int32 CodeChunk)
+{
+	FStrataMaterialCompilationInfo StrataInfo;
+	StrataInfo.LayerCount = 0;
+	StrataInfo.TotalBSDFCount = 0;
+	Compiler->StrataCompilationInfoRegisterCodeChunk(CodeChunk, StrataInfo);
+}
+
 
 FStrataMaterialCompilationInfo StrataCompilationInfoAdd(FMaterialCompiler* Compiler, const FStrataMaterialCompilationInfo& A, const FStrataMaterialCompilationInfo& B)
 {
diff --git a/Engine/Source/Runtime/Engine/Private/Materials/StrataMaterial.h b/Engine/Source/Runtime/Engine/Private/Materials/StrataMaterial.h
--- a/Engine/Source/Runtime/Engine/Private/Materials/StrataMaterial.h
+++ b/Engine/Source/Runtime/Engine/Private/Materials/StrataMaterial.h
@@ -44,6 +44,7 @@ FString GetStrataBSDFName(uint8 BSDFType);
 
 uint8 StrataCompilationInfoCreateSharedNormal(FMaterialCompiler* Compiler, int32 NormalCodeChunk);
 uint8 StrataCompilationInfoCreateSharedNormal(FMaterialCompiler* Compiler, int32 NormalCodeChunk, int32 TangentCodeChunk);
+void  StrataCompilationInfoCreateNullBSDF(FMaterialCompiler* Compiler, int32 CodeChunk);
 void  StrataCompilationInfoCreateSingleBSDFMaterial(FMaterialCompiler* Compiler, int32 CodeChunk, uint8 SharedNormalIndex, uint8 BSDFType, bool bVolumeHasScattering = false);
 
 FStrataMaterialCompilationInfo StrataCompilationInfoAdd(FMaterialCompiler* Compiler, const FStrataMaterialCompilationInfo& A, const FStrataMaterialCompilationInfo& B);
diff --git a/Engine/Source/Runtime/Engine/Public/MaterialCompiler.h b/Engine/Source/Runtime/Engine/Public/MaterialCompiler.h
--- a/Engine/Source/Runtime/Engine/Public/MaterialCompiler.h
+++ b/Engine/Source/Runtime/Engine/Public/MaterialCompiler.h
@@ -390,7 +390,7 @@ public:
 	virtual int32 ShadingModel(EMaterialShadingModel InSelectedShadingModel) = 0;
 
 	// Strata
-	virtual int32 FrontMaterial() = 0;
+	virtual int32 StrataCreateAndRegisterNullMaterial() = 0;
 	virtual int32 StrataDiffuseBSDF(int32 Albedo, int32 Roughness, int32 Normal, uint8 SharedNormalIndex) = 0;
 	virtual int32 StrataDielectricBSDF(int32 RoughnessX, int32 RoughnessY, int32 IOR, int32 Tint, int32 Normal, int32 Tangent, uint8 SharedNormalIndex) = 0;
 	virtual int32 StrataConductorBSDF(int32 Reflectivity, int32 EdgeColor, int32 RoughnessX, int32 RoughnessY, int32 Normal, int32 Tangent, uint8 SharedNormalIndex) = 0;
@@ -800,9 +800,9 @@ public:
 		return Compiler->IsDevelopmentFeatureEnabled(FeatureName);
 	}
 
-	virtual int32 FrontMaterial() override
+	virtual int32 StrataCreateAndRegisterNullMaterial() override
 	{
-		return Compiler->FrontMaterial();
+		return Compiler->StrataCreateAndRegisterNullMaterial();
 	}
 
 	virtual int32 StrataDiffuseBSDF(int32 Albedo, int32 Roughness, int32 Normal, uint8 SharedNormalIndex) override
