Added Blueprint support in DataLayerSubsystem

diff --git a/Engine/Source/Runtime/Engine/Private/WorldPartition/DataLayer/DataLayerSubsystem.cpp b/Engine/Source/Runtime/Engine/Private/WorldPartition/DataLayer/DataLayerSubsystem.cpp
--- a/Engine/Source/Runtime/Engine/Private/WorldPartition/DataLayer/DataLayerSubsystem.cpp
+++ b/Engine/Source/Runtime/Engine/Private/WorldPartition/DataLayer/DataLayerSubsystem.cpp
@@ -9,18 +9,6 @@
 UDataLayerSubsystem::UDataLayerSubsystem()
 {}
 
-const UDataLayer* UDataLayerSubsystem::GetDataLayerFromLabel(const FName& InDataLayerLabel) const
-{
-	const AWorldDataLayers* WorldDataLayers = AWorldDataLayers::Get(GetWorld());
-	return WorldDataLayers ? WorldDataLayers->GetDataLayerFromLabel(InDataLayerLabel) : nullptr;
-}
-
-const UDataLayer* UDataLayerSubsystem::GetDataLayerFromName(const FName& InDataLayerName) const
-{
-	const AWorldDataLayers* WorldDataLayers = AWorldDataLayers::Get(GetWorld());
-	return WorldDataLayers ? WorldDataLayers->GetDataLayerFromName(InDataLayerName) : nullptr;
-}
-
 bool UDataLayerSubsystem::ShouldCreateSubsystem(UObject* Outer) const
 {
 	if (!Super::ShouldCreateSubsystem(Outer))
@@ -36,38 +24,80 @@ bool UDataLayerSubsystem::ShouldCreateSubsystem(UObject* Outer) const
 	return false;
 }
 
-void UDataLayerSubsystem::ActivateDataLayer(const FName& InDataLayerName, bool bActivate)
+UDataLayer* UDataLayerSubsystem::GetDataLayerFromLabel(const FName& InDataLayerLabel) const
 {
-	// First validate that DataLayer is known by the world
-	if (!GetDataLayerFromName(InDataLayerName))
+	const AWorldDataLayers* WorldDataLayers = AWorldDataLayers::Get(GetWorld());
+	return WorldDataLayers ? const_cast<UDataLayer*>(WorldDataLayers->GetDataLayerFromLabel(InDataLayerLabel)) : nullptr;
+}
+
+UDataLayer* UDataLayerSubsystem::GetDataLayerFromName(const FName& InDataLayerName) const
+{
+	const AWorldDataLayers* WorldDataLayers = AWorldDataLayers::Get(GetWorld());
+	return WorldDataLayers ? const_cast<UDataLayer*>(WorldDataLayers->GetDataLayerFromName(InDataLayerName)) : nullptr;
+}
+
+void UDataLayerSubsystem::ActivateDataLayer(const FActorDataLayer& InDataLayer, bool bInActivate)
+{
+	ActivateDataLayerByName(InDataLayer.Name, bInActivate);
+}
+
+void UDataLayerSubsystem::ActivateDataLayer(const UDataLayer* InDataLayer, bool bInActivate)
 {
+	if (!InDataLayer)
+	{
 		return;
 	}
 
-	const bool bIsLayerActive = IsDataLayerActive(InDataLayerName);
-	if (bIsLayerActive != bActivate)
+	FName DataLayerName = InDataLayer->GetFName();
+	const bool bIsLayerActive = IsDataLayerActiveByName(DataLayerName);
+	if (bIsLayerActive != bInActivate)
 	{
-		if (bActivate)
+		if (bInActivate)
 		{
-			ActiveDataLayerNames.Add(InDataLayerName);
+			ActiveDataLayerNames.Add(DataLayerName);
 		}
 		else
 		{
-			ActiveDataLayerNames.Remove(InDataLayerName);
+			ActiveDataLayerNames.Remove(DataLayerName);
 		}
 	}
 }
 
-bool UDataLayerSubsystem::IsDataLayerActive(const FName& InDataLayerName) const
+void UDataLayerSubsystem::ActivateDataLayerByName(const FName& InDataLayerName, bool bInActivate)
+{
+	ActivateDataLayer(GetDataLayerFromName(InDataLayerName), bInActivate);
+}
+
+void UDataLayerSubsystem::ActivateDataLayerByLabel(const FName& InDataLayerLabel, bool bInActivate)
+{
+	ActivateDataLayer(GetDataLayerFromLabel(InDataLayerLabel), bInActivate);
+}
+
+bool UDataLayerSubsystem::IsDataLayerActive(const FActorDataLayer& InDataLayer) const
+{
+	return IsDataLayerActiveByName(InDataLayer.Name);
+}
+
+bool UDataLayerSubsystem::IsDataLayerActive(const UDataLayer* InDataLayer) const
+{
+	return InDataLayer && IsDataLayerActiveByName(InDataLayer->GetFName());
+}
+
+bool UDataLayerSubsystem::IsDataLayerActiveByName(const FName& InDataLayerName) const
 {
 	return ActiveDataLayerNames.Contains(InDataLayerName);
 }
 
+bool UDataLayerSubsystem::IsDataLayerActiveByLabel(const FName& InDataLayerLabel) const
+{
+	return IsDataLayerActive(GetDataLayerFromLabel(InDataLayerLabel));
+}
+
 bool UDataLayerSubsystem::IsAnyDataLayerActive(const TArray<FName>& InDataLayerNames) const
 {
 	for (FName DataLayerName : InDataLayerNames)
 	{
-		if (IsDataLayerActive(DataLayerName))
+		if (IsDataLayerActiveByName(DataLayerName))
 		{
 			return true;
 		}
@@ -91,8 +121,7 @@ FAutoConsoleCommand UDataLayerSubsystem::ToggleDataLayerActivation(
 					UDataLayerSubsystem* DataLayerSubsystem = World->GetSubsystem<UDataLayerSubsystem>();
 					if (const UDataLayer* DataLayer = DataLayerSubsystem ? DataLayerSubsystem->GetDataLayerFromLabel(DataLayerLabel) : nullptr)
 					{
-						FName DataLayerName = DataLayer->GetFName();
-						DataLayerSubsystem->ActivateDataLayer(DataLayerName, !DataLayerSubsystem->IsDataLayerActive(DataLayerName));
+						DataLayerSubsystem->ActivateDataLayer(DataLayer, !DataLayerSubsystem->IsDataLayerActive(DataLayer));
 					}
 				}
 			}
diff --git a/Engine/Source/Runtime/Engine/Public/WorldPartition/DataLayer/ActorDataLayer.h b/Engine/Source/Runtime/Engine/Public/WorldPartition/DataLayer/ActorDataLayer.h
--- a/Engine/Source/Runtime/Engine/Public/WorldPartition/DataLayer/ActorDataLayer.h
+++ b/Engine/Source/Runtime/Engine/Public/WorldPartition/DataLayer/ActorDataLayer.h
@@ -22,6 +22,6 @@ struct ENGINE_API FActorDataLayer
 	FORCEINLINE bool operator<(const FActorDataLayer& Other) const { return Name.FastLess(Other.Name); }
 
 	/** The name of this layer */
-	UPROPERTY(EditAnywhere, BlueprintReadWrite, Category = DataLayer)
+	UPROPERTY(VisibleAnywhere, Category = DataLayer)
 	FName Name;
 };
\ No newline at end of file
diff --git a/Engine/Source/Runtime/Engine/Public/WorldPartition/DataLayer/DataLayerSubsystem.h b/Engine/Source/Runtime/Engine/Public/WorldPartition/DataLayer/DataLayerSubsystem.h
--- a/Engine/Source/Runtime/Engine/Public/WorldPartition/DataLayer/DataLayerSubsystem.h
+++ b/Engine/Source/Runtime/Engine/Public/WorldPartition/DataLayer/DataLayerSubsystem.h
@@ -3,6 +3,7 @@
 #pragma once
 
 #include "Subsystems/WorldSubsystem.h"
+#include "WorldPartition/DataLayer/ActorDataLayer.h"
 #include "WorldPartition/DataLayer/DataLayer.h"
 #include "DataLayerSubsystem.generated.h"
 
@@ -22,10 +23,26 @@ public:
 	virtual bool ShouldCreateSubsystem(UObject* Outer) const override;
 	//~ End USubsystem Interface.
 
-	const UDataLayer* GetDataLayerFromLabel(const FName& InDataLayerLabel) const;
-	const UDataLayer* GetDataLayerFromName(const FName& InDataLayerName) const;
-	void ActivateDataLayer(const FName& InDataLayerName, bool bActivate);
-	bool IsDataLayerActive(const FName& InDataLayerName) const;
+	//~ Begin Blueprint callable functions
+	UFUNCTION(BlueprintCallable, Category = DataLayers)
+	void ActivateDataLayer(const FActorDataLayer& InDataLayer, bool bInActivate);
+	
+	UFUNCTION(BlueprintCallable, Category = DataLayers)
+	void ActivateDataLayerByLabel(const FName& InDataLayerLabel, bool bInActivate);
+
+	UFUNCTION(BlueprintCallable, Category = DataLayers)
+	bool IsDataLayerActive(const FActorDataLayer& InDataLayer) const;
+
+	UFUNCTION(BlueprintCallable, Category = DataLayers)
+	bool IsDataLayerActiveByLabel(const FName& InDataLayerName) const;
+	//~ End Blueprint callable functions
+
+	UDataLayer* GetDataLayerFromLabel(const FName& InDataLayerLabel) const;
+	UDataLayer* GetDataLayerFromName(const FName& InDataLayerName) const;
+	void ActivateDataLayer(const UDataLayer* InDataLayer, bool bInActivate);
+	void ActivateDataLayerByName(const FName& InDataLayerName, bool bInActivate);
+	bool IsDataLayerActive(const UDataLayer* InDataLayer) const;
+	bool IsDataLayerActiveByName(const FName& InDataLayerName) const;
 	bool IsAnyDataLayerActive(const TArray<FName>& InDataLayerNames) const;
 
 private:
