Starta node only for volumetric fog/cloud. It never goes into the Strata GBuffer.
Only usable in the volumetric domain.
Added support for emissive on all BSDF, only taken advantage of by voluemtricFogCloud as of today.

Shootergame compiled and ran on pc.

diff --git a/Engine/Shaders/Private/Strata/Strata.ush b/Engine/Shaders/Private/Strata/Strata.ush
--- a/Engine/Shaders/Private/Strata/Strata.ush
+++ b/Engine/Shaders/Private/Strata/Strata.ush
@@ -25,14 +25,14 @@
 
 
 
-
 ///////////////////////////////////////////////////////////////////////////////
 // BSDFs representation
 struct FStrataBSDF
 {
+	float3				Weight;		// Weight is a float3 because, when loaded from the material buffer, it contains the BSDF throughput (coverage and transmittance combined)
 	uint				State;		// (27 unused bits) (1bit as HasAnisotropy) (2bits as SharedNormalIndex) (3bits as BSDFType)
 
-	float3				Weight;		// Weight is a float3 because, when loaded from the material buffer, it contains the BSDF throughput (coverage and transmittance combined)
+	float3				Emissive;	// All bsdf can potientially have emissive. This is never written out to the strata gbuffer but instead accumulated in the emissive buffer during the base pass.
 
 	float4 VGPRs[2];
 };
@@ -111,6 +111,11 @@ struct FStrataBSDF
 #define SHEEN_ALBEDO(X)					X.VGPRs[0].xyz
 #define SHEEN_ROUGHNESS(X)				X.VGPRs[0].w
 
+#define VOLUMETRICFOGCLOUD_ALBEDO(X)	X.VGPRs[0].xyz
+#define VOLUMETRICFOGCLOUD_EXTINCTION(X)X.VGPRs[1].xyz
+#define VOLUMETRICFOGCLOUD_AO(X)		X.VGPRs[0].w
+#define VOLUMETRICFOGCLOUD_EMISSIVE(X)	X.Emissive
+
 void StrataRequestSharedNormalTangentBasis(inout uint NormalTypes, uint Index);
 
 bool StrataHasAnisotropy(in float Roughness0, float Roughness1)
@@ -267,6 +272,23 @@ FStrataData GetStrataSheenBSDF(float3 Albedo, float Roughness, uint SharedNormal
 	return StrataData;
 }
 
+FStrataData GetStrataVolumeFogCloudBSDF(float3 Albedo, float3 Extinction, float3 Emissive, float AmbientOcclusion)
+{
+	FStrataData StrataData = GetInitialisedStrataData();
+	StrataData.LayerCount++;
+	StrataData.Layers[0].BSDFCount++;
+
+	BSDF_SETTYPE					(StrataData.Layers[0].BSDFs[0], STRATA_BSDF_TYPE_VOLUMETRICFOGCLOUD);
+	VOLUMETRICFOGCLOUD_ALBEDO		(StrataData.Layers[0].BSDFs[0]) = Albedo;
+	VOLUMETRICFOGCLOUD_EXTINCTION	(StrataData.Layers[0].BSDFs[0]) = Extinction;
+	VOLUMETRICFOGCLOUD_EMISSIVE		(StrataData.Layers[0].BSDFs[0]) = Emissive;
+	VOLUMETRICFOGCLOUD_AO			(StrataData.Layers[0].BSDFs[0]) = AmbientOcclusion;
+
+	StrataData.Layers[0].BSDFs[0].Weight = 1.0f;
+
+	return StrataData;
+}
+
 // STRATA_TODO put in a common file
 // Point lobe in off-specular peak direction
 float3 StrataGetOffSpecularPeakReflectionDir(float3 Normal, float3 ReflectionVector, float Roughness)
diff --git a/Engine/Shaders/Private/VolumetricCloud.usf b/Engine/Shaders/Private/VolumetricCloud.usf
--- a/Engine/Shaders/Private/VolumetricCloud.usf
+++ b/Engine/Shaders/Private/VolumetricCloud.usf
@@ -13,6 +13,8 @@
 #include "Random.ush"
 #include "VolumetricCloudCommon.ush"
 
+#define STRATA_ENABLED PROJECT_STRATA && MATERIAL_IS_STRATA
+
 #if defined(SHADER_MAINVS) || defined(SHADER_RENDERVIEW_PS) || defined(SHADER_SHADOW_PS)
 
 #include "/Engine/Generated/Material.ush"
@@ -32,22 +34,54 @@ struct RenderVolumetricCloudVSToPS
 
 float3 SampleExtinctionCoefficients(in FPixelMaterialInputs PixelMaterialInputs)
 {
-	return clamp(GetMaterialSubsurfaceDataRaw(PixelMaterialInputs).rgb, 0.0f, 65000.0f);
+	float3 Extinction = 0.0f;
+#if STRATA_ENABLED
+	FStrataBSDF BSDF = PixelMaterialInputs.FrontMaterial.Layers[0].BSDFs[0];
+	Extinction = VOLUMETRICFOGCLOUD_EXTINCTION(BSDF).r;
+#else
+	#if !MATERIAL_SHADINGMODEL_UNLIT
+	Extinction = GetMaterialSubsurfaceDataRaw(PixelMaterialInputs).rgb;
+	#endif
+#endif
+	return clamp(Extinction, 0.0f, 65000.0f);
 }
 
 float3 SampleEmissive(in FPixelMaterialInputs PixelMaterialInputs)
 {
-	return USES_EMISSIVE_COLOR ? GetMaterialEmissive(PixelMaterialInputs) : 0.0f;
+	float3 EmissiveColor = 0.0f;
+#if STRATA_ENABLED
+	FStrataBSDF BSDF = PixelMaterialInputs.FrontMaterial.Layers[0].BSDFs[0];
+	EmissiveColor = VOLUMETRICFOGCLOUD_EMISSIVE(BSDF);
+#else
+	EmissiveColor = GetMaterialEmissiveRaw(PixelMaterialInputs);
+#endif
+	return clamp(EmissiveColor, 0.0f, 65000.0f);
 }
 
 float3 SampleAlbedo(in FPixelMaterialInputs PixelMaterialInputs)
 {
-	return saturate(GetMaterialBaseColor(PixelMaterialInputs));
+	float3 Albedo = 0.0f;
+#if STRATA_ENABLED
+	FStrataBSDF BSDF = PixelMaterialInputs.FrontMaterial.Layers[0].BSDFs[0];
+	Albedo = VOLUMETRICFOGCLOUD_ALBEDO(BSDF);
+#else
+	#if !MATERIAL_SHADINGMODEL_UNLIT
+	Albedo = GetMaterialBaseColor(PixelMaterialInputs) * View.DiffuseOverrideParameter.w + View.DiffuseOverrideParameter.xyz;
+	#endif
+#endif
+	return saturate(Albedo);
 }
 
-float3 SampleAmbientOcclusion(in FPixelMaterialInputs PixelMaterialInputs)
+float SampleAmbientOcclusion(in FPixelMaterialInputs PixelMaterialInputs)
 {
-	return GetMaterialAmbientOcclusion(PixelMaterialInputs);
+	float AO = 0.0f;
+#if STRATA_ENABLED
+	FStrataBSDF BSDF = PixelMaterialInputs.FrontMaterial.Layers[0].BSDFs[0];
+	AO = VOLUMETRICFOGCLOUD_AO(BSDF);
+#else
+	AO = GetMaterialAmbientOcclusion(PixelMaterialInputs);
+#endif
+	return AO;
 }
 
 #endif
diff --git a/Engine/Shaders/Private/VolumetricFogVoxelization.usf b/Engine/Shaders/Private/VolumetricFogVoxelization.usf
--- a/Engine/Shaders/Private/VolumetricFogVoxelization.usf
+++ b/Engine/Shaders/Private/VolumetricFogVoxelization.usf
@@ -10,6 +10,8 @@
 
 #include "Common.ush"
 
+#define STRATA_ENABLED PROJECT_STRATA && MATERIAL_IS_STRATA
+
 #define SceneTexturesStruct VoxelizeVolumePass.SceneTextures 
 
 #include "/Engine/Generated/Material.ush"
@@ -19,6 +21,46 @@
 
 uint VoxelizationPassIndex;
 
+float SampleExtinctionCoefficients(in FPixelMaterialInputs PixelMaterialInputs)
+{
+	float Extinction = 0.0f;
+#if STRATA_ENABLED
+	FStrataBSDF BSDF = PixelMaterialInputs.FrontMaterial.Layers[0].BSDFs[0];
+	Extinction = VOLUMETRICFOGCLOUD_EXTINCTION(BSDF).r;
+#else
+	#if !MATERIAL_SHADINGMODEL_UNLIT
+	Extinction = GetMaterialSubsurfaceDataRaw(PixelMaterialInputs).r;
+	#endif
+#endif
+	return clamp(Extinction, 0.0f, 65000.0f);
+}
+
+float3 SampleEmissive(in FPixelMaterialInputs PixelMaterialInputs)
+{
+	float3 EmissiveColor = 0.0f;
+#if STRATA_ENABLED
+	FStrataBSDF BSDF = PixelMaterialInputs.FrontMaterial.Layers[0].BSDFs[0];
+	EmissiveColor = VOLUMETRICFOGCLOUD_EMISSIVE(BSDF);
+#else
+	EmissiveColor = GetMaterialEmissiveRaw(PixelMaterialInputs);
+#endif
+	return clamp(EmissiveColor, 0.0f, 65000.0f);
+}
+
+float3 SampleAlbedo(in FPixelMaterialInputs PixelMaterialInputs, in ViewState InputView)
+{
+	float3 Albedo = 0.0f;
+#if STRATA_ENABLED
+	FStrataBSDF BSDF = PixelMaterialInputs.FrontMaterial.Layers[0].BSDFs[0];
+	Albedo = VOLUMETRICFOGCLOUD_ALBEDO(BSDF);
+#else
+	#if !MATERIAL_SHADINGMODEL_UNLIT
+	Albedo = GetMaterialBaseColor(PixelMaterialInputs);
+	#endif
+#endif
+	return saturate(Albedo)* InputView.DiffuseOverrideParameter.w + InputView.DiffuseOverrideParameter.xyz;
+}
+
 struct FVoxelizeVolumePrimitiveVSToGS
 {
 	FVertexFactoryInterpolantsVSToPS FactoryInterpolants;
@@ -261,15 +303,9 @@ void VoxelizePS(
 	FPixelMaterialInputs PixelMaterialInputs;
 	CalcMaterialParameters(MaterialParameters, PixelMaterialInputs, SvPosition, true);
 
-	float3 EmissiveColor = clamp(GetMaterialEmissiveRaw(PixelMaterialInputs), 0.0f, 65000.0f);
-
-	float3 Albedo = 0;
-	float Extinction = 0;
-
-#if !MATERIAL_SHADINGMODEL_UNLIT
-	Albedo = GetMaterialBaseColor(PixelMaterialInputs) * View.DiffuseOverrideParameter.w + View.DiffuseOverrideParameter.xyz;
-	Extinction = clamp(GetMaterialSubsurfaceDataRaw(PixelMaterialInputs).r, 0.0f, 65000.0f);
-#endif
+	float3 EmissiveColor = SampleEmissive(PixelMaterialInputs);
+	float3 Albedo = SampleAlbedo(PixelMaterialInputs, ResolvedView);
+	float Extinction = SampleExtinctionCoefficients(PixelMaterialInputs);
 
 	float3 Scattering = Albedo * Extinction;
 	float Absorption = max(Extinction - Luminance(Scattering), 0.0f);
diff --git a/Engine/Shaders/Shared/StrataDefinitions.h b/Engine/Shaders/Shared/StrataDefinitions.h
--- a/Engine/Shaders/Shared/StrataDefinitions.h
+++ b/Engine/Shaders/Shared/StrataDefinitions.h
@@ -26,5 +26,7 @@
 #define STRATA_BSDF_TYPE_CONDUCTOR			2
 #define STRATA_BSDF_TYPE_VOLUME				3
 #define STRATA_BSDF_TYPE_SHEEN				4
+#define STRATA_BSDF_TYPE_VOLUMETRICFOGCLOUD	5
+// When more than 8 BSDF must exists, please update STATE_BITS_BSDF in Strata.ush
 
 
diff --git a/Engine/Source/Runtime/Engine/Classes/Materials/MaterialExpressionStrata.h b/Engine/Source/Runtime/Engine/Classes/Materials/MaterialExpressionStrata.h
--- a/Engine/Source/Runtime/Engine/Classes/Materials/MaterialExpressionStrata.h
+++ b/Engine/Source/Runtime/Engine/Classes/Materials/MaterialExpressionStrata.h
@@ -247,7 +247,45 @@ class UMaterialExpressionStrataSheenBSDF : public UMaterialExpression
 	//~ End UMaterialExpression Interface
 };
 
-// Subsurface, thinfilm, generalised schlick
+UCLASS(MinimalAPI, collapsecategories, hidecategories = Object)
+class UMaterialExpressionStrataVolumetricFogCloudBSDF : public UMaterialExpression
+{
+	GENERATED_UCLASS_BODY()
+
+	/**
+	* Albedo (type = float3, unit = unitless, default = 0)
+	*/
+	UPROPERTY()
+	FExpressionInput Albedo;
+
+	/**
+	 * The rate at which light is absorbed or scattered by the medium. Mean Free Path = 1 / Extinction. (type = float3, unit = 1/m, default = 0)
+	 */
+	UPROPERTY()
+	FExpressionInput Extinction;
+
+	/**
+	 * Emissive (type = float3, unit = luminance, default = 0)
+	 */
+	UPROPERTY()
+	FExpressionInput Emissive;
+
+	/**
+	 * Ambient occlusion: 1 means no occlusion while 0 means fully occluded. (type = float, unit = unitless, default = 1)
+	 */
+	UPROPERTY()
+	FExpressionInput AmbientOcclusion;
+
+	//~ Begin UMaterialExpression Interface
+#if WITH_EDITOR
+	virtual int32 Compile(class FMaterialCompiler* Compiler, int32 OutputIndex) override;
+	virtual void GetCaption(TArray<FString>& OutCaptions) const override;
+	virtual uint32 GetOutputType(int32 OutputIndex) override;
+	virtual uint32 GetInputType(int32 InputIndex) override;
+	virtual bool IsResultStrataMaterial(int32 OutputIndex) override;
+#endif
+	//~ End UMaterialExpression Interface
+};
 
 
 
diff --git a/Engine/Source/Runtime/Engine/Private/Materials/HLSLMaterialTranslator.cpp b/Engine/Source/Runtime/Engine/Private/Materials/HLSLMaterialTranslator.cpp
--- a/Engine/Source/Runtime/Engine/Private/Materials/HLSLMaterialTranslator.cpp
+++ b/Engine/Source/Runtime/Engine/Private/Materials/HLSLMaterialTranslator.cpp
@@ -944,7 +944,7 @@ bool FHLSLMaterialTranslator::Translate()
 
 		if (Material->IsStrataMaterial())
 		{
-			if (!(BlendMode == BLEND_Translucent || BlendMode == BLEND_Opaque))
+			if (Domain!= MD_Volume && !(BlendMode == BLEND_Translucent || BlendMode == BLEND_Opaque))
 			{
 				Errorf(TEXT("Strata materials must be opaque or translucent."));
 			}
@@ -1104,66 +1104,90 @@ bool FHLSLMaterialTranslator::Translate()
 			TextureExpression->GetExternalTextureParameterInfo(MaterialCompilationOutput.UniformExpressionSet.UniformExternalTextureParameters.AddDefaulted_GetRef());
 		}
 
-		// Output some debug info as comment in code and in the material stat window
-		if (Chunk[MP_FrontMaterial] != INDEX_NONE && CodeChunkToStrataCompilationInfoMap.Contains(Chunk[MP_FrontMaterial]) == true)
+		const bool bStrataFrontMaterialProvided = Chunk[MP_FrontMaterial] != INDEX_NONE && CodeChunkToStrataCompilationInfoMap.Contains(Chunk[MP_FrontMaterial]) == true;
+		bool bStrataFrontMaterialIsValid = bStrataFrontMaterialProvided;
+		if (bStrataFrontMaterialIsValid)
 		{
 			const FStrataMaterialCompilationInfo& StrataCompilationInfo = CodeChunkToStrataCompilationInfoMap[Chunk[MP_FrontMaterial]];
 
-			// The material can be null when some entries are automatically generated, for instance in the mateiral layer blendign system
+			// The material can be null when some entries are automatically generated, for instance in the material layer blending system
 			const bool bIsNullMaterial = (StrataCompilationInfo.LayerCount == 0 && StrataCompilationInfo.TotalBSDFCount == 0);
-			if (!bIsNullMaterial)
-			{
-				static const auto CVarStrataBytePerPixel = IConsoleManager::Get().FindTConsoleVariableDataInt(TEXT("r.Strata.BytesPerPixel"));
-				check(CVarStrataBytePerPixel)
-					const uint32 StrataBytePerPixel = CVarStrataBytePerPixel ? CVarStrataBytePerPixel->GetValueOnAnyThread() : 0;
-				StrataMaterialAnalysis = StrataCompilationInfoMaterialAnalysis(this, StrataCompilationInfo, StrataBytePerPixel);
-
-				FString StrataMaterialDescription;
+			bStrataFrontMaterialIsValid &= !bIsNullMaterial;
+		}
 
-				StrataMaterialDescription += FString::Printf(TEXT("----- STRATA -----\r\n"));
-				StrataMaterialDescription += FString::Printf(TEXT("StrataCompilationInfo -\r\n"));
-				StrataMaterialDescription += FString::Printf(TEXT(" - TotalBSDFCount    = %i\r\n"), StrataCompilationInfo.TotalBSDFCount);
-				StrataMaterialDescription += FString::Printf(TEXT(" - SharedNormalCount = %i\r\n"), StrataCompilationInfoGetSharedNormalCount());
+		if (bStrataFrontMaterialIsValid)
+		{
+			const FStrataMaterialCompilationInfo& StrataCompilationInfo = CodeChunkToStrataCompilationInfoMap[Chunk[MP_FrontMaterial]];
 
-				for (uint32 LayerIt = 0; LayerIt < StrataCompilationInfo.LayerCount; ++LayerIt)
+			// Verify the validity of the strata material
+			if (Domain == MD_Volume)
+			{
+				if (!StrataIsVolumetricFogCloudOnly(this, StrataCompilationInfo))
 				{
-					StrataMaterialDescription += FString::Printf(TEXT(" Layer %i BSDFs:\r\n"), LayerIt);
-
-					const FStrataMaterialCompilationInfo::FLayer& Layer = StrataCompilationInfo.Layers[LayerIt];
-					for (uint32 BSDFIt = 0; BSDFIt < Layer.BSDFCount; ++BSDFIt)
-					{
-						StrataMaterialDescription += FString::Printf(TEXT("     - %s - SharedNormalIndex = %u \r\n"), *GetStrataBSDFName(Layer.BSDFs[BSDFIt].Type), Layer.BSDFs[BSDFIt].SharedNormalIndex);
-					}
+					UE_LOG(LogMaterial, Error, TEXT("Material %s is in the VOLUME domain but a VolumetricFogCloud BSDF is not provided (asset: %s).\r\n"), *Material->GetDebugName(), *Material->GetAssetPath());
 				}
-
-				StrataMaterialDescription += FString::Printf(TEXT("Byte Per Pixel Budget      %u\r\n"), StrataBytePerPixel);
-				StrataMaterialDescription += FString::Printf(TEXT("Result.bFitInMemoryBudget  %s\r\n"), StrataMaterialAnalysis.bFitInMemoryBudget ? TEXT("YES") : TEXT("NO"));
-				StrataMaterialDescription += FString::Printf(TEXT("Result.RequestedLayerCount %u\r\n"), StrataMaterialAnalysis.RequestedLayerCount);
-				StrataMaterialDescription += FString::Printf(TEXT("Result.RequestedBSDFCount  %u\r\n"), StrataMaterialAnalysis.RequestedBSDFCount);
-				StrataMaterialDescription += FString::Printf(TEXT("Result.RequestedByteCount  %u\r\n"), StrataMaterialAnalysis.RequestedByteCount);
-				if (!StrataMaterialAnalysis.bFitInMemoryBudget)
+			}
+			else
+			{
+				if (StrataMaterialContainsAnyBSDF(this, StrataCompilationInfo, STRATA_BSDF_TYPE_VOLUMETRICFOGCLOUD))
 				{
-					StrataMaterialDescription += FString::Printf(TEXT("Result.ClampedLayerCount   %u\r\n"), StrataMaterialAnalysis.ClampedLayerCount);
-					StrataMaterialDescription += FString::Printf(TEXT("Result.ClampedBSDFCount    %u\r\n"), StrataMaterialAnalysis.ClampedBSDFCount);
-					StrataMaterialDescription += FString::Printf(TEXT("Result.UsedByteCount       %u\r\n"), StrataMaterialAnalysis.UsedByteCount);
+					UE_LOG(LogMaterial, Error, TEXT("Material %s is not in the VOLUME domain but a VolumetricFogCloud BSDF is used. It should not be used out of the VOLUME domain. (asset: %s).\r\n"), *Material->GetDebugName(), *Material->GetAssetPath());
 				}
-				StrataMaterialDescription += FString::Printf(TEXT("------------------\r\n"));
+			}
 
-				ResourcesString += TEXT("/*");
-				ResourcesString += StrataMaterialDescription;
-				ResourcesString += TEXT("*/");
+			// STRATA_TODO if not volume domain: verify that there isn't any VolumetricFogCloud node in the material
 
-				MaterialCompilationOutput.StrataMaterialDescription = StrataMaterialDescription;
+			// Output some debug info as comment in code and in the material stat window
+			static const auto CVarStrataBytePerPixel = IConsoleManager::Get().FindTConsoleVariableDataInt(TEXT("r.Strata.BytesPerPixel"));
+			check(CVarStrataBytePerPixel)
+				const uint32 StrataBytePerPixel = CVarStrataBytePerPixel ? CVarStrataBytePerPixel->GetValueOnAnyThread() : 0;
+			StrataMaterialAnalysis = StrataCompilationInfoMaterialAnalysis(this, StrataCompilationInfo, StrataBytePerPixel);
 
-				if (StrataMaterialAnalysis.ClampedLayerCount == 0)
-				{
-					UE_LOG(LogMaterial, Fatal, TEXT("Material %s cannot have any layers rendered due to its complexity (asset: %s).\r\n"), *Material->GetDebugName(), *Material->GetAssetPath());
-				}
-				else if (StrataMaterialAnalysis.RequestedLayerCount > StrataMaterialAnalysis.ClampedLayerCount)
+			FString StrataMaterialDescription;
+
+			StrataMaterialDescription += FString::Printf(TEXT("----- STRATA -----\r\n"));
+			StrataMaterialDescription += FString::Printf(TEXT("StrataCompilationInfo -\r\n"));
+			StrataMaterialDescription += FString::Printf(TEXT(" - TotalBSDFCount    = %i\r\n"), StrataCompilationInfo.TotalBSDFCount);
+			StrataMaterialDescription += FString::Printf(TEXT(" - SharedNormalCount = %i\r\n"), StrataCompilationInfoGetSharedNormalCount());
+
+			for (uint32 LayerIt = 0; LayerIt < StrataCompilationInfo.LayerCount; ++LayerIt)
+			{
+				StrataMaterialDescription += FString::Printf(TEXT(" Layer %i BSDFs:\r\n"), LayerIt);
+
+				const FStrataMaterialCompilationInfo::FLayer& Layer = StrataCompilationInfo.Layers[LayerIt];
+				for (uint32 BSDFIt = 0; BSDFIt < Layer.BSDFCount; ++BSDFIt)
 				{
-					UE_LOG(LogMaterial, Warning, TEXT("Material %s is not fitting the allocated byte per pixel budget for starta. Layers will be removed (asset: %s).\r\n"), *Material->GetDebugName(), *Material->GetAssetPath());
+					StrataMaterialDescription += FString::Printf(TEXT("     - %s - SharedNormalIndex = %u \r\n"), *GetStrataBSDFName(Layer.BSDFs[BSDFIt].Type), Layer.BSDFs[BSDFIt].SharedNormalIndex);
 				}
 			}
+
+			StrataMaterialDescription += FString::Printf(TEXT("Byte Per Pixel Budget      %u\r\n"), StrataBytePerPixel);
+			StrataMaterialDescription += FString::Printf(TEXT("Result.bFitInMemoryBudget  %s\r\n"), StrataMaterialAnalysis.bFitInMemoryBudget ? TEXT("YES") : TEXT("NO"));
+			StrataMaterialDescription += FString::Printf(TEXT("Result.RequestedLayerCount %u\r\n"), StrataMaterialAnalysis.RequestedLayerCount);
+			StrataMaterialDescription += FString::Printf(TEXT("Result.RequestedBSDFCount  %u\r\n"), StrataMaterialAnalysis.RequestedBSDFCount);
+			StrataMaterialDescription += FString::Printf(TEXT("Result.RequestedByteCount  %u\r\n"), StrataMaterialAnalysis.RequestedByteCount);
+			if (!StrataMaterialAnalysis.bFitInMemoryBudget)
+			{
+				StrataMaterialDescription += FString::Printf(TEXT("Result.ClampedLayerCount   %u\r\n"), StrataMaterialAnalysis.ClampedLayerCount);
+				StrataMaterialDescription += FString::Printf(TEXT("Result.ClampedBSDFCount    %u\r\n"), StrataMaterialAnalysis.ClampedBSDFCount);
+				StrataMaterialDescription += FString::Printf(TEXT("Result.UsedByteCount       %u\r\n"), StrataMaterialAnalysis.UsedByteCount);
+			}
+			StrataMaterialDescription += FString::Printf(TEXT("------------------\r\n"));
+
+			ResourcesString += TEXT("/*");
+			ResourcesString += StrataMaterialDescription;
+			ResourcesString += TEXT("*/");
+
+			MaterialCompilationOutput.StrataMaterialDescription = StrataMaterialDescription;
+
+			if (StrataMaterialAnalysis.ClampedLayerCount == 0)
+			{
+				UE_LOG(LogMaterial, Error, TEXT("Material %s cannot have any layers rendered due to its complexity (asset: %s).\r\n"), *Material->GetDebugName(), *Material->GetAssetPath());
+			}
+			else if (StrataMaterialAnalysis.RequestedLayerCount > StrataMaterialAnalysis.ClampedLayerCount)
+			{
+				UE_LOG(LogMaterial, Warning, TEXT("Material %s is not fitting the allocated byte per pixel budget for starta. Layers will be removed (asset: %s).\r\n"), *Material->GetDebugName(), *Material->GetAssetPath());
+			}
 		}
 		else
 		{
@@ -1526,7 +1550,7 @@ void FHLSLMaterialTranslator::GetMaterialEnvironment(EShaderPlatform InPlatform,
 	// if duals source blending (colored transmittance) is not supported on a platform, it will fall back to standard alpha blending (grey scale transmittance)
 	OutEnvironment.SetDefine(TEXT("DUAL_SOURCE_COLOR_BLENDING_ENABLED"), bMaterialRequestsDualSourceBlending && Material->IsDualBlendingEnabled(Platform) ? TEXT("1") : TEXT("0"));
 
-	// Translate() is called before getting here so we can create
+	// Translate() is called before getting here so we can create related define
 	if (StrataMaterialAnalysis.RequestedBSDFCount > 0)
 	{
 		OutEnvironment.SetDefine(TEXT("STRATA_CLAMPED_LAYER_COUNT"), StrataMaterialAnalysis.ClampedLayerCount);
@@ -7513,6 +7537,17 @@ int32 FHLSLMaterialTranslator::StrataSheenBSDF(int32 Albedo, int32 Roughness, in
 	);
 }
 
+int32 FHLSLMaterialTranslator::StrataVolumetricFogCloudBSDF(int32 Albedo, int32 Extinction, int32 Emissive, int32 AmbientOcclusion)
+{
+	return AddCodeChunk(
+		MCT_Strata, TEXT("GetStrataVolumeFogCloudBSDF(%s, %s, %s, %s)"),
+		*GetParameterCode(Albedo),
+		*GetParameterCode(Extinction),
+		*GetParameterCode(Emissive),
+		*GetParameterCode(AmbientOcclusion)
+	);
+}
+
 int32 FHLSLMaterialTranslator::StrataHorizontalMixing(int32 Foreground, int32 Background, int32 Mix)
 {
 	if (Foreground == INDEX_NONE || Background == INDEX_NONE || Mix == INDEX_NONE)
diff --git a/Engine/Source/Runtime/Engine/Private/Materials/HLSLMaterialTranslator.h b/Engine/Source/Runtime/Engine/Private/Materials/HLSLMaterialTranslator.h
--- a/Engine/Source/Runtime/Engine/Private/Materials/HLSLMaterialTranslator.h
+++ b/Engine/Source/Runtime/Engine/Private/Materials/HLSLMaterialTranslator.h
@@ -796,6 +796,7 @@ protected:
 	virtual int32 StrataConductorBSDF(int32 Reflectivity, int32 EdgeColor, int32 RoughnessX, int32 RoughnessY, int32 Normal, int32 Tangent, uint8 SharedNormalIndex) override;
 	virtual int32 StrataVolumeBSDF(int32 Albedo, int32 Extinction, int32 Anisotropy, int32 Thickness, int32 Normal, uint8 SharedNormalIndex) override;
 	virtual int32 StrataSheenBSDF(int32 Albedo, int32 Roughness, int32 Normal, uint8 SharedNormalIndex) override;
+	virtual int32 StrataVolumetricFogCloudBSDF(int32 Albedo, int32 Extinction, int32 Emissive, int32 AmbientOcclusion) override;
 	virtual int32 StrataHorizontalMixing(int32 Foreground, int32 Background, int32 Mix) override;
 	virtual int32 StrataVerticalLayering(int32 Top, int32 Base) override;
 	virtual int32 StrataAdd(int32 A, int32 B) override;
diff --git a/Engine/Source/Runtime/Engine/Private/Materials/Material.cpp b/Engine/Source/Runtime/Engine/Private/Materials/Material.cpp
--- a/Engine/Source/Runtime/Engine/Private/Materials/Material.cpp
+++ b/Engine/Source/Runtime/Engine/Private/Materials/Material.cpp
@@ -5995,7 +5995,8 @@ static bool IsPropertyActive_Internal(EMaterialProperty InProperty,
 		return InProperty == MP_EmissiveColor
 			|| InProperty == MP_SubsurfaceColor
 			|| InProperty == MP_BaseColor
-			|| InProperty == MP_AmbientOcclusion;
+			|| InProperty == MP_AmbientOcclusion
+			|| InProperty == MP_FrontMaterial;
 	}
 	else if (Domain == MD_UI)
 	{
diff --git a/Engine/Source/Runtime/Engine/Private/Materials/MaterialExpressions.cpp b/Engine/Source/Runtime/Engine/Private/Materials/MaterialExpressions.cpp
--- a/Engine/Source/Runtime/Engine/Private/Materials/MaterialExpressions.cpp
+++ b/Engine/Source/Runtime/Engine/Private/Materials/MaterialExpressions.cpp
@@ -19479,6 +19479,73 @@ bool UMaterialExpressionStrataSheenBSDF::IsResultStrataMaterial(int32 OutputInde
 
 
 
+UMaterialExpressionStrataVolumetricFogCloudBSDF::UMaterialExpressionStrataVolumetricFogCloudBSDF(const FObjectInitializer& ObjectInitializer)
+	: Super(ObjectInitializer)
+{
+	struct FConstructorStatics
+	{
+		FText NAME_Strata;
+		FConstructorStatics() : NAME_Strata(LOCTEXT("Strata BSDFs", "Strata BSDFs")) { }
+	};
+	static FConstructorStatics ConstructorStatics;
+#if WITH_EDITORONLY_DATA
+	MenuCategories.Add(ConstructorStatics.NAME_Strata);
+#endif
+}
+
+#if WITH_EDITOR
+int32 UMaterialExpressionStrataVolumetricFogCloudBSDF::Compile(class FMaterialCompiler* Compiler, int32 OutputIndex)
+{
+	int32 OutputCodeChunk = Compiler->StrataVolumetricFogCloudBSDF(
+		CompileWithDefaultFloat3(Compiler, Albedo, 0.0f, 0.0f, 0.0f),
+		CompileWithDefaultFloat3(Compiler, Extinction, 0.0f, 0.0f, 0.0f),
+		CompileWithDefaultFloat3(Compiler, Emissive, 0.0f, 0.0f, 0.0f),
+		CompileWithDefaultFloat1(Compiler, AmbientOcclusion, 1.0f));
+
+	uint8 FakeSharedNormalIndex = 0;
+	StrataCompilationInfoCreateSingleBSDFMaterial(Compiler, OutputCodeChunk, FakeSharedNormalIndex, STRATA_BSDF_TYPE_VOLUMETRICFOGCLOUD);
+
+	return OutputCodeChunk;
+}
+
+void UMaterialExpressionStrataVolumetricFogCloudBSDF::GetCaption(TArray<FString>& OutCaptions) const
+{
+	OutCaptions.Add(TEXT("Strata Sheen BSDF"));
+}
+
+uint32 UMaterialExpressionStrataVolumetricFogCloudBSDF::GetOutputType(int32 OutputIndex)
+{
+	return MCT_Strata;
+}
+
+uint32 UMaterialExpressionStrataVolumetricFogCloudBSDF::GetInputType(int32 InputIndex)
+{
+	switch (InputIndex)
+	{
+	case 0:
+		return MCT_Float3;
+		break;
+	case 1:
+		return MCT_Float3;
+		break;
+	case 2:
+		return MCT_Float3;
+		break;
+	case 3:
+		return MCT_Float;
+		break;
+	}
+
+	check(false);
+	return MCT_Float1;
+}
+
+bool UMaterialExpressionStrataVolumetricFogCloudBSDF::IsResultStrataMaterial(int32 OutputIndex)
+{
+	return true;
+}
+#endif // WITH_EDITOR
+
 UMaterialExpressionStrataHorizontalMixing::UMaterialExpressionStrataHorizontalMixing(const FObjectInitializer& ObjectInitializer)
 	: Super(ObjectInitializer)
 {
diff --git a/Engine/Source/Runtime/Engine/Private/Materials/StrataMaterial.cpp b/Engine/Source/Runtime/Engine/Private/Materials/StrataMaterial.cpp
--- a/Engine/Source/Runtime/Engine/Private/Materials/StrataMaterial.cpp
+++ b/Engine/Source/Runtime/Engine/Private/Materials/StrataMaterial.cpp
@@ -24,6 +24,9 @@ FString GetStrataBSDFName(uint8 BSDFType)
 	case STRATA_BSDF_TYPE_SHEEN:
 		return TEXT("SHEEN");
 		break;
+	case STRATA_BSDF_TYPE_VOLUMETRICFOGCLOUD:
+		return TEXT("VOLUMETRICFOGCLOUD");
+		break;
 	}
 	check(false);
 	return "";
@@ -135,6 +138,43 @@ FStrataMaterialCompilationInfo StrataCompilationInfoVerticalLayering(FMaterialCo
 	return StrataInfo;
 }
 
+bool StrataIsVolumetricFogCloudOnly(FMaterialCompiler* Compiler, const FStrataMaterialCompilationInfo& Material)
+{
+	if (Material.TotalBSDFCount == 0 || Material.LayerCount == 0)
+	{
+		Compiler->Error(TEXT("There is no layer or BSDF plugged in, but a material in the volume domain wants to read from a VolumeBSDF."));
+		return false;
+	}
+	if (Material.TotalBSDFCount > 1 || Material.LayerCount > 1)
+	{
+		Compiler->Error(TEXT("There is more than one layer or BSDF, but a material in the volume domain wants to read from a single VolumeBSDF only."));
+		return false;
+	}
+	if (Material.Layers[0].BSDFs[0].Type != STRATA_BSDF_TYPE_VOLUMETRICFOGCLOUD)
+	{
+		Compiler->Error(TEXT("The single BSDF resulting from the graph is not of type Volume."));
+		return false;
+	}
+
+	return true;
+}
+
+bool StrataMaterialContainsAnyBSDF(FMaterialCompiler* Compiler, const FStrataMaterialCompilationInfo& Material, uint8 BSDFType)
+{
+	for (uint32 LayerIt = 0; LayerIt < Material.LayerCount; ++LayerIt)
+	{
+		const FStrataMaterialCompilationInfo::FLayer& Layer = Material.Layers[LayerIt];
+		for (uint32 BSDF = 0; BSDF < Layer.BSDFCount; BSDF++)
+		{
+			if (Layer.BSDFs[BSDF].Type == BSDFType)
+			{
+				return true;
+			}
+		}
+	}
+	return false;
+}
+
 FStrataMaterialAnalysisResult::FStrataMaterialAnalysisResult()
 {
 	bFitInMemoryBudget = true;
diff --git a/Engine/Source/Runtime/Engine/Private/Materials/StrataMaterial.h b/Engine/Source/Runtime/Engine/Private/Materials/StrataMaterial.h
--- a/Engine/Source/Runtime/Engine/Private/Materials/StrataMaterial.h
+++ b/Engine/Source/Runtime/Engine/Private/Materials/StrataMaterial.h
@@ -52,6 +52,8 @@ FStrataMaterialCompilationInfo StrataCompilationInfoMultiply(FMaterialCompiler*
 FStrataMaterialCompilationInfo StrataCompilationInfoHorizontalMixing(FMaterialCompiler* Compiler, const FStrataMaterialCompilationInfo& A, const FStrataMaterialCompilationInfo& B);
 FStrataMaterialCompilationInfo StrataCompilationInfoVerticalLayering(FMaterialCompiler* Compiler, const FStrataMaterialCompilationInfo& Top, const FStrataMaterialCompilationInfo& Base);
 
+bool StrataIsVolumetricFogCloudOnly(FMaterialCompiler* Compiler, const FStrataMaterialCompilationInfo& Material);
+bool StrataMaterialContainsAnyBSDF(FMaterialCompiler* Compiler, const FStrataMaterialCompilationInfo& Material, uint8 BSDFType);
 
 struct FStrataMaterialAnalysisResult
 {
diff --git a/Engine/Source/Runtime/Engine/Public/MaterialCompiler.h b/Engine/Source/Runtime/Engine/Public/MaterialCompiler.h
--- a/Engine/Source/Runtime/Engine/Public/MaterialCompiler.h
+++ b/Engine/Source/Runtime/Engine/Public/MaterialCompiler.h
@@ -396,6 +396,7 @@ public:
 	virtual int32 StrataConductorBSDF(int32 Reflectivity, int32 EdgeColor, int32 RoughnessX, int32 RoughnessY, int32 Normal, int32 Tangent, uint8 SharedNormalIndex) = 0;
 	virtual int32 StrataVolumeBSDF(int32 Albedo, int32 Extinction, int32 Anisotropy, int32 Thickness, int32 Normal, uint8 SharedNormalIndex) = 0;
 	virtual int32 StrataSheenBSDF(int32 Albedo, int32 Roughness, int32 Normal, uint8 SharedNormalIndex) = 0;
+	virtual int32 StrataVolumetricFogCloudBSDF(int32 Albedo, int32 Extinction, int32 Emissive, int32 AmbientOcclusion) = 0;
 	virtual int32 StrataHorizontalMixing(int32 Foreground, int32 Background, int32 Mix) = 0;
 	virtual int32 StrataVerticalLayering(int32 Top, int32 Base) = 0;
 	virtual int32 StrataAdd(int32 A, int32 B) = 0;
@@ -830,6 +831,11 @@ public:
 		return Compiler->StrataSheenBSDF(Albedo, Roughness, Normal, SharedNormalIndex);
 	}
 
+	virtual int32 StrataVolumetricFogCloudBSDF(int32 Albedo, int32 Extinction, int32 Emissive, int32 AmbientOcclusion) override
+	{
+		return Compiler->StrataVolumetricFogCloudBSDF(Albedo, Extinction, Emissive, AmbientOcclusion);
+	}
+
 	virtual int32 StrataHorizontalMixing(int32 Foreground, int32 Background, int32 Mix) override
 	{
 		return Compiler->StrataHorizontalMixing(Foreground, Background, Mix);
