Misc Strata clean up

diff --git a/Engine/Source/Runtime/Engine/Private/Materials/HLSLMaterialTranslator.cpp b/Engine/Source/Runtime/Engine/Private/Materials/HLSLMaterialTranslator.cpp
--- a/Engine/Source/Runtime/Engine/Private/Materials/HLSLMaterialTranslator.cpp
+++ b/Engine/Source/Runtime/Engine/Private/Materials/HLSLMaterialTranslator.cpp
@@ -1124,14 +1124,16 @@ bool FHLSLMaterialTranslator::Translate()
 			{
 				if (!StrataIsVolumetricFogCloudOnly(this, StrataCompilationInfo))
 				{
-					UE_LOG(LogMaterial, Error, TEXT("Material %s is in the VOLUME domain but a VolumetricFogCloud BSDF is not provided (asset: %s).\r\n"), *Material->GetDebugName(), *Material->GetAssetPath());
+					FString ErrorMsg = FString::Printf(TEXT("Material %s is in the VOLUME domain but a VolumetricFogCloud BSDF is not provided (asset: %s).\r\n"), *Material->GetDebugName(), *Material->GetAssetPath());
+					Error(*ErrorMsg);
 				}
 			}
 			else
 			{
 				if (StrataMaterialContainsAnyBSDF(this, StrataCompilationInfo, STRATA_BSDF_TYPE_VOLUMETRICFOGCLOUD))
 				{
-					UE_LOG(LogMaterial, Error, TEXT("Material %s is not in the VOLUME domain but a VolumetricFogCloud BSDF is used. It should not be used out of the VOLUME domain. (asset: %s).\r\n"), *Material->GetDebugName(), *Material->GetAssetPath());
+					FString ErrorMsg = FString::Printf(TEXT("Material %s is not in the VOLUME domain but a VolumetricFogCloud BSDF is used. It should not be used out of the VOLUME domain. (asset: %s).\r\n"), *Material->GetDebugName(), *Material->GetAssetPath());
+					Error(*ErrorMsg);
 				}
 			}
 
diff --git a/Engine/Source/Runtime/Engine/Private/Materials/MaterialExpressions.cpp b/Engine/Source/Runtime/Engine/Private/Materials/MaterialExpressions.cpp
--- a/Engine/Source/Runtime/Engine/Private/Materials/MaterialExpressions.cpp
+++ b/Engine/Source/Runtime/Engine/Private/Materials/MaterialExpressions.cpp
@@ -19510,7 +19510,7 @@ int32 UMaterialExpressionStrataVolumetricFogCloudBSDF::Compile(class FMaterialCo
 
 void UMaterialExpressionStrataVolumetricFogCloudBSDF::GetCaption(TArray<FString>& OutCaptions) const
 {
-	OutCaptions.Add(TEXT("Strata Sheen BSDF"));
+	OutCaptions.Add(TEXT("Strata Volumetric Fog Cloud BSDF"));
 }
 
 uint32 UMaterialExpressionStrataVolumetricFogCloudBSDF::GetOutputType(int32 OutputIndex)
diff --git a/Engine/Source/Runtime/Engine/Private/Materials/StrataMaterial.cpp b/Engine/Source/Runtime/Engine/Private/Materials/StrataMaterial.cpp
--- a/Engine/Source/Runtime/Engine/Private/Materials/StrataMaterial.cpp
+++ b/Engine/Source/Runtime/Engine/Private/Materials/StrataMaterial.cpp
@@ -142,12 +142,12 @@ bool StrataIsVolumetricFogCloudOnly(FMaterialCompiler* Compiler, const FStrataMa
 {
 	if (Material.TotalBSDFCount == 0 || Material.LayerCount == 0)
 	{
-		Compiler->Error(TEXT("There is no layer or BSDF plugged in, but a material in the volume domain wants to read from a VolumeBSDF."));
+		Compiler->Error(TEXT("There is no layer or BSDF plugged in, but a material in the volume domain wants to read from a StrataVolumetricFogCloudBSDF."));
 		return false;
 	}
 	if (Material.TotalBSDFCount > 1 || Material.LayerCount > 1)
 	{
-		Compiler->Error(TEXT("There is more than one layer or BSDF, but a material in the volume domain wants to read from a single VolumeBSDF only."));
+		Compiler->Error(TEXT("There is more than one layer or BSDF, but a material in the volume domain wants to read from a single StrataVolumetricFogCloudBSDF only."));
 		return false;
 	}
 	if (Material.Layers[0].BSDFs[0].Type != STRATA_BSDF_TYPE_VOLUMETRICFOGCLOUD)
