Fix CIS due to missing include

diff --git a/Engine/Source/Runtime/Engine/Private/ShaderCompiler/ShaderGenerationUtil.cpp b/Engine/Source/Runtime/Engine/Private/ShaderCompiler/ShaderGenerationUtil.cpp
--- a/Engine/Source/Runtime/Engine/Private/ShaderCompiler/ShaderGenerationUtil.cpp
+++ b/Engine/Source/Runtime/Engine/Private/ShaderCompiler/ShaderGenerationUtil.cpp
@@ -11,6 +11,7 @@
 #include "Misc/CoreMisc.h"
 #include "Misc/ScopeLock.h"
 #include "Interfaces/ITargetPlatformManagerModule.h"
+#include "HAL/PlatformFileManager.h"
 #include "RenderUtils.h"
 #include "RHI.h"
 #include "SceneManagement.h"
