Cherry pick CL 14820425

Fix nullptr dereference in FillLevelList

#SOURCE: CL 14798352 via CL 14798929 via CL 14799101 via CL 14799187 via CL 14799331

diff --git a/Engine/Source/Runtime/Engine/Private/TickTaskManager.cpp b/Engine/Source/Runtime/Engine/Private/TickTaskManager.cpp
--- a/Engine/Source/Runtime/Engine/Private/TickTaskManager.cpp
+++ b/Engine/Source/Runtime/Engine/Private/TickTaskManager.cpp
@@ -1670,7 +1670,7 @@ private:
 		for( int32 LevelIndex = 0; LevelIndex < Levels.Num(); LevelIndex++ )
 		{
 			ULevel* Level = Levels[LevelIndex];
-			if (Level->bIsVisible)
+			if (Level && Level->bIsVisible)
 			{
 				check(Level->TickTaskLevel);
 				LevelList.Add(Level->TickTaskLevel);
