Fix static analysis error potentially dereferencing nullptr ChunkData

diff --git a/Engine/Source/Runtime/Engine/Private/VT/VirtualTextureBuiltData.cpp b/Engine/Source/Runtime/Engine/Private/VT/VirtualTextureBuiltData.cpp
--- a/Engine/Source/Runtime/Engine/Private/VT/VirtualTextureBuiltData.cpp
+++ b/Engine/Source/Runtime/Engine/Private/VT/VirtualTextureBuiltData.cpp
@@ -208,6 +208,8 @@ bool FVirtualTextureBuiltData::ValidateCompression(FStringView const& InDDCDebug
 		}
 #endif // WITH_EDITORONLY_DATA
 
+		bResult = bResult && (ChunkData != nullptr);
+
 		uint32 TileIndex = TileIndexPerChunk[ChunkIndex];
 		while (bResult && TileIndex < TileIndexPerChunk[ChunkIndex + 1])
 		{
