HLOD: Fixed fallback to the default HLODLayer not working for the first level of HLOD

diff --git a/Engine/Source/Runtime/Engine/Private/WorldPartition/HLOD/HLODLayer.cpp b/Engine/Source/Runtime/Engine/Private/WorldPartition/HLOD/HLODLayer.cpp
--- a/Engine/Source/Runtime/Engine/Private/WorldPartition/HLOD/HLODLayer.cpp
+++ b/Engine/Source/Runtime/Engine/Private/WorldPartition/HLOD/HLODLayer.cpp
@@ -92,7 +92,7 @@ UHLODLayer* UHLODLayer::GetHLODLayer(const AActor* InActor)
 
 	// Only fallback to the default HLODLayer for the first level of HLOD
 	bool bIsHLOD0 = !InActor->IsA<AWorldPartitionHLOD>();
-	if (!bIsHLOD0) 
+	if (bIsHLOD0) 
 	{
 		if (UWorldPartition* WorldPartition = InActor->GetWorld()->GetWorldPartition())
 		{
