HLOD: Fixed HLOD actors not getting cooked because world is not initialized during cook and default TActorIterator<ASpatialHashRuntimeGridInfo> skips actors if their level is not active.

diff --git a/Engine/Source/Runtime/Engine/Private/WorldPartition/DataLayer/WorldDataLayers.cpp b/Engine/Source/Runtime/Engine/Private/WorldPartition/DataLayer/WorldDataLayers.cpp
--- a/Engine/Source/Runtime/Engine/Private/WorldPartition/DataLayer/WorldDataLayers.cpp
+++ b/Engine/Source/Runtime/Engine/Private/WorldPartition/DataLayer/WorldDataLayers.cpp
@@ -22,15 +22,15 @@ const AWorldDataLayers* AWorldDataLayers::Get(UWorld* World)
 {
 	if (World)
 	{
-		EActorIteratorFlags Flags = EActorIteratorFlags::OnlyActiveLevels | EActorIteratorFlags::SkipPendingKill;
+		// Prepare flags for actor iterator. Don't use default Flags because it uses EActorIteratorFlags::OnlyActiveLevels 
+		// which will make this code return no actor when cooking (because world is not initialized)
+		EActorIteratorFlags Flags = EActorIteratorFlags::SkipPendingKill;
 #if WITH_EDITOR
-		if (IsRunningCookCommandlet())
+		if (!IsRunningCookCommandlet())
+#endif
 		{
-			// Don't use default Flags because EActorIteratorFlags::OnlyActiveLevels will make this code return no AWorldDataLayers actor
-			// when cooking (because world is not initialized)
-			Flags &= ~EActorIteratorFlags::OnlyActiveLevels;
+			Flags |= EActorIteratorFlags::OnlyActiveLevels;
 		}
-#endif
 		for (AWorldDataLayers* Actor : TActorRange<AWorldDataLayers>(World, AWorldDataLayers::StaticClass(), Flags))
 		{
 			if (Actor)
diff --git a/Engine/Source/Runtime/Engine/Private/WorldPartition/WorldPartitionRuntimeSpatialHash.cpp b/Engine/Source/Runtime/Engine/Private/WorldPartition/WorldPartitionRuntimeSpatialHash.cpp
--- a/Engine/Source/Runtime/Engine/Private/WorldPartition/WorldPartitionRuntimeSpatialHash.cpp
+++ b/Engine/Source/Runtime/Engine/Private/WorldPartition/WorldPartitionRuntimeSpatialHash.cpp
@@ -1246,8 +1246,15 @@ bool UWorldPartitionRuntimeSpatialHash::GenerateStreaming(EWorldPartitionStreami
 	TArray<FSpatialHashRuntimeGrid> AllGrids;
 	AllGrids.Append(Grids);
 
+	// Prepare flags for actor iterator. Don't use default Flags because it uses EActorIteratorFlags::OnlyActiveLevels 
+	// which will make this code return no actor when cooking (because world is not initialized)
+	EActorIteratorFlags Flags = EActorIteratorFlags::SkipPendingKill;
+	if (!IsRunningCookCommandlet())
+	{
+		Flags |= EActorIteratorFlags::OnlyActiveLevels;
+	}
 	// Append grids from runtime grid actors
-	for (TActorIterator<ASpatialHashRuntimeGridInfo> ItRuntimeGridActor(GetWorld()); ItRuntimeGridActor; ++ItRuntimeGridActor)
+	for (TActorIterator<ASpatialHashRuntimeGridInfo> ItRuntimeGridActor(GetWorld(), ASpatialHashRuntimeGridInfo::StaticClass(), Flags); ItRuntimeGridActor; ++ItRuntimeGridActor)
 	{
 		AllGrids.Add(ItRuntimeGridActor->GridSettings);
 	}
