Saved the Editor analytic sessions to disk in a background task to save time in the game thread when writing later is not problematic.
  - On my Windows Xeon 6 core, 3.4GHz, it takes 0.1 ms rather than 5 ms on the game thread when saving asynchrously. The session is usually saved every 30 seconds.
  - Replaced usage of TAtomic<> with equivalent std::atomic<>

diff --git a/Engine/Source/Runtime/Engine/Private/Analytics/EditorSessionSummaryWriter.cpp b/Engine/Source/Runtime/Engine/Private/Analytics/EditorSessionSummaryWriter.cpp
--- a/Engine/Source/Runtime/Engine/Private/Analytics/EditorSessionSummaryWriter.cpp
+++ b/Engine/Source/Runtime/Engine/Private/Analytics/EditorSessionSummaryWriter.cpp
@@ -163,12 +163,12 @@ bool FEditorSessionSummaryWriter::UpdateEditorIdleTime(double CurrTimeSecs, bool
 {
 	bool bSessionUpdated = false;
 
-	double LastActivityExpectedSecs = LastEditorActivityTimeSecs.Load();
+	double LastActivityExpectedSecs = LastEditorActivityTimeSecs.load();
 	double InactivitySeconds = CurrTimeSecs - LastActivityExpectedSecs;
 	if (InactivitySeconds >= EditorSessionWriterDefs::EditorInactivitySecondsForIdleState) // Was idle long enough to account this span of time as Idle?
 	{
 		// Ensure only one thread increments the counter.
-		if (LastEditorActivityTimeSecs.CompareExchange(LastActivityExpectedSecs, CurrTimeSecs))
+		if (LastEditorActivityTimeSecs.compare_exchange_strong(LastActivityExpectedSecs, CurrTimeSecs))
 		{
 			// Add up this span of inactivity and reset the counter to start another span.
 			FPlatformAtomics::InterlockedAdd(&CurrentSession->TotalEditorInactivitySeconds, FMath::FloorToInt(static_cast<float>(InactivitySeconds)));
@@ -179,7 +179,7 @@ bool FEditorSessionSummaryWriter::UpdateEditorIdleTime(double CurrTimeSecs, bool
 
 	if (bReset)
 	{
-		LastEditorActivityTimeSecs.Store(CurrTimeSecs);
+		LastEditorActivityTimeSecs.store(CurrTimeSecs);
 	}
 
 	return bSessionUpdated;
@@ -190,10 +190,10 @@ bool FEditorSessionSummaryWriter::UpdateUserIdleTime(double CurrTimeSecs, bool b
 	bool bSessionUpdated = false;
 
 	// How much time elapsed since the last activity.
-	double TotalIdleSecs = CurrTimeSecs - LastUserActivityTimeSecs.Load();
+	double TotalIdleSecs = CurrTimeSecs - LastUserActivityTimeSecs.load();
 	if (TotalIdleSecs > 60.0) // Less than a minute is always considered normal interaction delay.
 	{
-		double LastAccountedIdleSecs = AccountedUserIdleSecs.Load();
+		double LastAccountedIdleSecs = AccountedUserIdleSecs.load();
 		double UnaccountedIdleSecs   = TotalIdleSecs - LastAccountedIdleSecs;
 
 		// If one or more minute is unaccounted
@@ -208,7 +208,7 @@ bool FEditorSessionSummaryWriter::UpdateUserIdleTime(double CurrTimeSecs, bool b
 			double DeltaIdle30Min = FMath::Max(0.0, AccountedIdleMins + ToAccountIdleMins - 30) - FMath::Max(0.0, AccountedIdleMins - 30); // The 30 first minutes of this idle sequence are considered 'normal interaction delay' and are not accounted for the 30-min timer.
 
 			// Ensure only one thread adds the current delta time.
-			if (AccountedUserIdleSecs.CompareExchange(LastAccountedIdleSecs, LastAccountedIdleSecs + ToAccountIdleMins * 60.0)) // Only add the 'accounted' minutes and keep fraction of minutes running.
+			if (AccountedUserIdleSecs.compare_exchange_strong(LastAccountedIdleSecs, LastAccountedIdleSecs + ToAccountIdleMins * 60.0)) // Only add the 'accounted' minutes and keep fraction of minutes running.
 			{
 				FPlatformAtomics::InterlockedAdd(&CurrentSession->Idle1Min, FMath::RoundToInt(static_cast<float>(DeltaIdle1Min)));
 				FPlatformAtomics::InterlockedAdd(&CurrentSession->Idle5Min, FMath::RoundToInt(static_cast<float>(DeltaIdle5Min)));
@@ -220,8 +220,8 @@ bool FEditorSessionSummaryWriter::UpdateUserIdleTime(double CurrTimeSecs, bool b
 
 	if (bReset)
 	{
-		AccountedUserIdleSecs.Store(0);
-		LastUserActivityTimeSecs.Store(CurrTimeSecs);
+		AccountedUserIdleSecs.store(0);
+		LastUserActivityTimeSecs.store(CurrTimeSecs);
 	}
 
     // WARNING: The code is supposed to be concurrent safe, but does't block. Calling UpdateUserIdleTime() and reading the counter back may not read the latest value if another thread concurrently
@@ -342,8 +342,8 @@ void FEditorSessionSummaryWriter::Tick(float DeltaTime)
 
 	if (bSaveSession)
 	{
-		// TODO: Consider cloning the session and saving in another thread. It cost about 1 to 6ms on Windows to save in the main thread.
-		TrySaveCurrentSession(CurrentTimeUtc, CurrentTimeSecs); // Saving also updates session duration/timestamp/userIdle/editorIdle
+		// Saves asynchronously, saving up to 5ms in the game thread. Saving also updates session duration/timestamp/userIdle/editorIdle
+		TrySaveCurrentSession(CurrentTimeUtc, CurrentTimeSecs, /*bAsync*/true);
 	}
 }
 
@@ -352,7 +352,7 @@ void FEditorSessionSummaryWriter::LowDriveSpaceDetected()
 	if (CurrentSession)
 	{
 		CurrentSession->bIsLowDriveSpace = true;
-		TrySaveCurrentSession(FDateTime::UtcNow(), FPlatformTime::Seconds());
+		TrySaveCurrentSession(FDateTime::UtcNow(), FPlatformTime::Seconds(), /*bAsync*/true);
 	}
 }
 
@@ -361,6 +361,12 @@ void FEditorSessionSummaryWriter::Shutdown()
 	// NOTE: Initialize(), Shutdown() and ~FEditorSessionSummaryWriter() are expected to be called from the game thread only.
 	if (CurrentSession && !bShutdown)
 	{
+		if (SaveSessionTask)
+		{
+			SaveSessionTask->EnsureCompletion();
+			SaveSessionTask.Reset();
+		}
+
 		// NOTE: Shutdown() may crash if a delegate is broadcasted from another thread at the same time (that's a bug in 4.24.x, 4.25.x) the delegate are modified.
 		FEditorDelegates::PreBeginPIE.RemoveAll(this);
 		FEditorDelegates::EndPIE.RemoveAll(this);
@@ -376,12 +382,13 @@ void FEditorSessionSummaryWriter::Shutdown()
 		double CurrTimeSecs = FPlatformTime::Seconds();
 		FDateTime CurrTimeUtc = FDateTime::UtcNow();
 
-		if (!TrySaveCurrentSession(CurrTimeUtc, CurrTimeSecs)) // If the save fails (because the lock was already taken)
+		if (!TrySaveCurrentSession(CurrTimeUtc, CurrTimeSecs, /*bAsync*/false)) // Try saving NOW, the application is going to close.
 		{
+			// In case the session could not be saved, gather the minimum info and emits the shutdown event locklessly.
 			UpdateUserIdleTime(CurrTimeSecs, /*bReset*/false);
 			UpdateEditorIdleTime(CurrTimeSecs, /*bReset*/false);
 			UpdateSessionDuration(CurrTimeSecs);
-			CurrentSession->LogEvent(FEditorAnalyticsSession::EEventType::Shutdown, CurrTimeUtc); // Use the lockless mechanism. It doesn't save everything, but it carries the critical information.
+			CurrentSession->LogEvent(FEditorAnalyticsSession::EEventType::Shutdown, CurrTimeUtc); // The lockless mechanism doesn't save everything, but it carries the critical information.
 		}
 
 		CurrentSession.Reset();
@@ -544,7 +551,7 @@ void FEditorSessionSummaryWriter::OnVanillaStateChanged(bool bIsVanilla)
 	if (CurrentSession != nullptr && CurrentSession->bIsVanilla != bIsVanilla)
 	{
 		CurrentSession->bIsVanilla = bIsVanilla;
-		TrySaveCurrentSession(FDateTime::UtcNow(), FPlatformTime::Seconds());
+		TrySaveCurrentSession(FDateTime::UtcNow(), FPlatformTime::Seconds(), /*bAsync*/true);
 	}
 }
 
@@ -553,7 +560,7 @@ void FEditorSessionSummaryWriter::OnUserActivity(const FUserActivity& UserActivi
 	if (CurrentSession != nullptr)
 	{
 		CurrentSession->CurrentUserActivity = GetUserActivityString();
-		TrySaveCurrentSession(FDateTime::UtcNow(), FPlatformTime::Seconds());
+		TrySaveCurrentSession(FDateTime::UtcNow(), FPlatformTime::Seconds(), /*bAsync*/true);
 	}
 }
 
@@ -565,12 +572,13 @@ void FEditorSessionSummaryWriter::OnUserLoginChanged(bool bLoggingIn, int32, int
 
 		double CurrTimeSecs = FPlatformTime::Seconds();
 		FDateTime CurrTimeUtc = FDateTime::UtcNow();
-		if (!TrySaveCurrentSession(CurrTimeUtc, CurrTimeSecs)) // If the save fails (because the lock was already taken)
+		if (!TrySaveCurrentSession(CurrTimeUtc, CurrTimeSecs, /*bAsync*/false)) // Try saving NOW, the application is going to close.
 		{
+			// In case the session could not be saved, gather the minimum info and emits the shutdown event locklessly.
 			UpdateUserIdleTime(CurrTimeSecs, /*bReset*/false);
 			UpdateEditorIdleTime(CurrTimeSecs, /*bReset*/false);
 			UpdateSessionDuration(CurrTimeSecs);
-			CurrentSession->LogEvent(FEditorAnalyticsSession::EEventType::LogOut, FDateTime::UtcNow()); // Use the lockless mechanism. It doesn't save everything, but it carries the critical information.
+			CurrentSession->LogEvent(FEditorAnalyticsSession::EEventType::LogOut, FDateTime::UtcNow()); // The lockless mechanism doesn't save everything, but it carries the critical information.
 		}
 	}
 }
@@ -597,7 +605,7 @@ void FEditorSessionSummaryWriter::OnSlateUserInteraction(double CurrSlateInterac
 	bSave |= UpdateEditorIdleTime(CurrTimeSecs, /*bReset*/true);
 	if (bSave)
 	{
-		TrySaveCurrentSession(FDateTime::UtcNow(), CurrTimeSecs);
+		TrySaveCurrentSession(FDateTime::UtcNow(), CurrTimeSecs, /*bAsync*/true);
 	}
 }
 
@@ -606,7 +614,7 @@ void FEditorSessionSummaryWriter::OnEnterPIE(const bool /*bIsSimulating*/)
 	if (CurrentSession != nullptr)
 	{
 		CurrentSession->bIsInPIE = true;
-		TrySaveCurrentSession(FDateTime::UtcNow(), FPlatformTime::Seconds());
+		TrySaveCurrentSession(FDateTime::UtcNow(), FPlatformTime::Seconds(), /*bAsync*/false); // Save NOW, PIE is more susceptible to crash, don't delay the save.
 	}
 }
 
@@ -615,29 +623,64 @@ void FEditorSessionSummaryWriter::OnExitPIE(const bool /*bIsSimulating*/)
 	if (CurrentSession != nullptr)
 	{
 		CurrentSession->bIsInPIE = false;
-		TrySaveCurrentSession(FDateTime::UtcNow(), FPlatformTime::Seconds());
+		TrySaveCurrentSession(FDateTime::UtcNow(), FPlatformTime::Seconds(), /*bAsync*/false); // Save NOW, PIE is more susceptible to crash, don't delay the save.
+	}
+}
+
+// NOTE: TrySaveCurrentSession() is expected to always be called from the game thread.
+bool FEditorSessionSummaryWriter::TrySaveCurrentSession(const FDateTime& CurrTimeUtc, double CurrTimeSecs, bool bAsync)
+{
+	// Ensure any previous save has completed. NOTE: If this add too much waiting, return if the new save request is also async, but ensure to wait if bAsync is false (ex. when the Editor is about to exits)
+	if (SaveSessionTask.IsValid())
+	{
+		SaveSessionTask->EnsureCompletion();
+		SaveSessionTask.Reset();
+	}
+
+	UpdateOutOfProcessMonitorState(/*bQuickCheck*/true);
+	UpdateUserIdleTime(CurrTimeSecs, /*bReset*/false);
+	UpdateEditorIdleTime(CurrTimeSecs, /*bReset*/false);
+	UpdateSessionDuration(CurrTimeSecs);
+	UpdateSessionTimestamp(CurrTimeUtc);
+
+	// PERF NOTE: Saving the session can take up to 5 ms while copying the session and starting a task is usually less than 0.1 ms.
+	if (bAsync)
+	{
+		// Copy the session and save it asynchronously. That allows writing the current session while the copy is saving to disk.
+		SaveSessionTask = MakeUnique<FAsyncTask<FSaveEditorAnalyticSessionWorker>>(*CurrentSession, [this]()
+		{
+			// Record the last save time.
+			LastSaveTimeSecs = FPlatformTime::Seconds();
+		});
+		SaveSessionTask->StartBackgroundTask();
+		return true; // Task was started (async mode)
 	}
+	else if (FEditorAnalyticsSession::TryLock()) // Inter-process lock to grant this process exclusive access to the key-store file/registry.
+	{
+		CurrentSession->Save();
+		LastSaveTimeSecs = CurrTimeSecs;
+		FEditorAnalyticsSession::Unlock();
+		return true; // Session was was saved (synchronous mode)
+	}
+
+	return false; // Session wasn't saved - sync mode / Saving task wasn't started - async mode.
 }
 
-bool FEditorSessionSummaryWriter::TrySaveCurrentSession(const FDateTime& CurrTimeUtc, double CurrTimeSecs)
+void FEditorSessionSummaryWriter::FSaveEditorAnalyticSessionWorker::DoWork()
 {
 	if (FEditorAnalyticsSession::TryLock()) // Inter-process lock to grant this process exclusive access to the key-store file/registry.
 	{
-		if (SaveSessionLock.TryLock()) // Intra-process lock to grant the calling thread exclusive access to the key-store file/registry.
+		// Save the session.
+		Session.Save();
+
+		// Notify the caller.
+		if (OnSessionSavedFn)
 		{
-			UpdateOutOfProcessMonitorState(/*bQuickCheck*/true);
-			UpdateUserIdleTime(CurrTimeSecs, /*bReset*/false);
-			UpdateEditorIdleTime(CurrTimeSecs, /*bReset*/false);
-			UpdateSessionDuration(CurrTimeSecs);
-			UpdateSessionTimestamp(CurrTimeUtc);
-			CurrentSession->Save();
-			LastSaveTimeSecs = CurrTimeSecs;
-			SaveSessionLock.Unlock();
+			OnSessionSavedFn();
 		}
+
 		FEditorAnalyticsSession::Unlock();
-		return true;
 	}
-	return false;
 }
 
 #undef LOCTEXT_NAMESPACE
diff --git a/Engine/Source/Runtime/Engine/Private/Analytics/EditorSessionSummaryWriter.h b/Engine/Source/Runtime/Engine/Private/Analytics/EditorSessionSummaryWriter.h
--- a/Engine/Source/Runtime/Engine/Private/Analytics/EditorSessionSummaryWriter.h
+++ b/Engine/Source/Runtime/Engine/Private/Analytics/EditorSessionSummaryWriter.h
@@ -3,12 +3,15 @@
 #pragma once
 
 #include "CoreMinimal.h"
-#include "Templates/Atomic.h"
 
 #if WITH_EDITOR
 
+#include "EditorAnalyticsSession.h"
+#include "Async/AsyncWork.h"
+#include "Templates/Function.h"
+#include <atomic>
+
 struct FUserActivity;
-struct FEditorAnalyticsSession;
 
 /** Writer for SessionSummary events to track all editor sessions. */
 class FEditorSessionSummaryWriter
@@ -39,23 +42,48 @@ private:
 	bool UpdateEditorIdleTime(double InCurrTimeSecs, bool bReset);
 	bool UpdateUserIdleTime(double InCurrTimeSecs, bool bReset);
 	bool UpdateOutOfProcessMonitorState(bool bQuickCheck);
-	bool TrySaveCurrentSession(const FDateTime& CurrTimeUtc, double CurrTimeSecs);
+	bool TrySaveCurrentSession(const FDateTime& CurrTimeUtc, double CurrTimeSecs, bool bAsync);
+
+private:
+	/** Helper task class to save the session in the background. */
+	class FSaveEditorAnalyticSessionWorker : public FNonAbandonableTask
+	{
+	public:
+		/** Initialization constructor. */
+		FSaveEditorAnalyticSessionWorker(const FEditorAnalyticsSession& InSession, TFunction<void()> OnSavedFn)
+			: Session(InSession)
+			, OnSessionSavedFn(MoveTemp(OnSavedFn)) {}
+
+		/** Save the analytic session to disk. */
+		void DoWork();
+
+		TStatId GetStatId() const { return TStatId(); }
+
+		/** Return the name of this task */
+		static const TCHAR* Name() { return TEXT( "FSaveEditorAnalyticSessionWorker" ); }
+
+	private:
+		/** Pointer to the session to save. */
+		FEditorAnalyticsSession Session;
+
+		/** Callback function invoked once the session is saved. */
+		TFunction<void()> OnSessionSavedFn;
+	};
 
 private:
 	TUniquePtr<FEditorAnalyticsSession> CurrentSession;
-	FCriticalSection SaveSessionLock;
 
 	/** The next time to check if the debugger is attached */
 	double NextDebuggerCheckSecs;
 
 	/** Last activity (user input, crash, terminate, shutdown) timestamp from FPlatformTime::Seconds() to track user inactivity. */
-	TAtomic<double> LastUserActivityTimeSecs;
+	std::atomic<double> LastUserActivityTimeSecs;
 
 	/** The number of idle seconds in the current idle sequence that were accounted (saved in the session) for the user idle counters. */
-	TAtomic<double> AccountedUserIdleSecs;
+	std::atomic<double> AccountedUserIdleSecs;
 
 	/** Last activity (user input, crash, terminate, shutdown, CPU Burst) timestamp from FPlatformTime::Seconds(). */
-	TAtomic<double> LastEditorActivityTimeSecs;
+	std::atomic<double> LastEditorActivityTimeSecs;
 
 	/** Session timestamp from FDateTime::UtcNow(). Unreliable if user change system date/time (daylight saving or user altering it). */
 	FDateTime SessionStartTimeUtc;
@@ -64,13 +92,16 @@ private:
 	double SessionStartTimeSecs = 0.0;
 
 	/** The last save timestamp from FPlatformTime::Seconds(). */
-	double LastSaveTimeSecs = 0.0;
+	std::atomic<double> LastSaveTimeSecs = 0.0;
 
 	/** Non-zero if out-of process monitoring is set. To ensure one CrashReportClient(CRC) doesn't report the session of another CRC instance (race condition). */
 	const uint32 OutOfProcessMonitorProcessId;
 
 	/** True once Shutdown() is called. */
 	bool bShutdown = false;
+
+	/** Task used to save the session to disk */
+	TUniquePtr<FAsyncTask<FSaveEditorAnalyticSessionWorker>> SaveSessionTask;
 };
 
 #endif
