Fixed CIS error in EditorSessionSummaryWriter reported on Mac about std::atomic initialization.

diff --git a/Engine/Source/Runtime/Engine/Private/Analytics/EditorSessionSummaryWriter.cpp b/Engine/Source/Runtime/Engine/Private/Analytics/EditorSessionSummaryWriter.cpp
--- a/Engine/Source/Runtime/Engine/Private/Analytics/EditorSessionSummaryWriter.cpp
+++ b/Engine/Source/Runtime/Engine/Private/Analytics/EditorSessionSummaryWriter.cpp
@@ -66,6 +66,7 @@ FEditorSessionSummaryWriter::FEditorSessionSummaryWriter(uint32 InProcessMonitor
 	, LastEditorActivityTimeSecs(FPlatformTime::Seconds())
 	, SessionStartTimeUtc(FDateTime::UtcNow()) // Reliable only if system date/time doesn't change (like daylight savings or user altering it)
 	, SessionStartTimeSecs(FPlatformTime::Seconds()) // Don't rely on system date/time. May suffer from lack of precision over long period of time (few seconds over a day).
+	, LastSaveTimeSecs(0.0)
 	, OutOfProcessMonitorProcessId(InProcessMonitorProcessId)
 {
 }
diff --git a/Engine/Source/Runtime/Engine/Private/Analytics/EditorSessionSummaryWriter.h b/Engine/Source/Runtime/Engine/Private/Analytics/EditorSessionSummaryWriter.h
--- a/Engine/Source/Runtime/Engine/Private/Analytics/EditorSessionSummaryWriter.h
+++ b/Engine/Source/Runtime/Engine/Private/Analytics/EditorSessionSummaryWriter.h
@@ -92,7 +92,7 @@ private:
 	double SessionStartTimeSecs = 0.0;
 
 	/** The last save timestamp from FPlatformTime::Seconds(). */
-	std::atomic<double> LastSaveTimeSecs = 0.0;
+	std::atomic<double> LastSaveTimeSecs;
 
 	/** Non-zero if out-of process monitoring is set. To ensure one CrashReportClient(CRC) doesn't report the session of another CRC instance (race condition). */
 	const uint32 OutOfProcessMonitorProcessId;
