WorldPartition: Cleaner way to find world's ASpatialHashRuntimeGridInfo actors

diff --git a/Engine/Source/Runtime/Engine/Private/WorldPartition/WorldPartitionRuntimeSpatialHash.cpp b/Engine/Source/Runtime/Engine/Private/WorldPartition/WorldPartitionRuntimeSpatialHash.cpp
--- a/Engine/Source/Runtime/Engine/Private/WorldPartition/WorldPartitionRuntimeSpatialHash.cpp
+++ b/Engine/Source/Runtime/Engine/Private/WorldPartition/WorldPartitionRuntimeSpatialHash.cpp
@@ -1249,17 +1249,13 @@ bool UWorldPartitionRuntimeSpatialHash::GenerateStreaming(EWorldPartitionStreami
 	TArray<FSpatialHashRuntimeGrid> AllGrids;
 	AllGrids.Append(Grids);
 
-	// Prepare flags for actor iterator. Don't use default Flags because it uses EActorIteratorFlags::OnlyActiveLevels 
-	// which will make this code return no actor when cooking (because world is not initialized)
-	EActorIteratorFlags Flags = EActorIteratorFlags::SkipPendingKill;
-	if (!IsRunningCookCommandlet())
-	{
-		Flags |= EActorIteratorFlags::OnlyActiveLevels;
-	}
 	// Append grids from runtime grid actors
-	for (TActorIterator<ASpatialHashRuntimeGridInfo> ItRuntimeGridActor(GetWorld(), ASpatialHashRuntimeGridInfo::StaticClass(), Flags); ItRuntimeGridActor; ++ItRuntimeGridActor)
+	for (AActor* Actor : GetWorld()->PersistentLevel->Actors)
 	{
-		AllGrids.Add(ItRuntimeGridActor->GridSettings);
+		if (ASpatialHashRuntimeGridInfo* RuntimeGridActor = Cast<ASpatialHashRuntimeGridInfo>(Actor))
+		{
+			AllGrids.Add(RuntimeGridActor->GridSettings);
+		}
 	}
 
 	check(!StreamingGrids.Num());
