Fix crash when unloading a partitionned world: avoid calling code that will result in calling StaticFindObject during a GC when unloading a partitioned world due to recent changes.
This also fixes UWorldPartition::Tick not being called anymore.

diff --git a/Engine/Source/Runtime/Engine/Private/WorldPartition/WorldPartition.cpp b/Engine/Source/Runtime/Engine/Private/WorldPartition/WorldPartition.cpp
--- a/Engine/Source/Runtime/Engine/Private/WorldPartition/WorldPartition.cpp
+++ b/Engine/Source/Runtime/Engine/Private/WorldPartition/WorldPartition.cpp
@@ -300,11 +300,8 @@ void UWorldPartition::Initialize(UWorld* InWorld, const FTransform& InTransform)
 
 	InitState = EWorldPartitionInitState::Initialized;
 
-	if (World->IsGameWorld())
-	{
-		UWorldPartitionSubsystem* WorldPartitionSubsystem = World->GetSubsystem<UWorldPartitionSubsystem>();
-		WorldPartitionSubsystem->RegisterWorldPartition(this);
-	}
+	UWorldPartitionSubsystem* WorldPartitionSubsystem = World->GetSubsystem<UWorldPartitionSubsystem>();
+	WorldPartitionSubsystem->RegisterWorldPartition(this);
 
 #if WITH_EDITOR
 	if (World->IsPlayInEditor())
@@ -314,13 +311,6 @@ void UWorldPartition::Initialize(UWorld* InWorld, const FTransform& InTransform)
 #endif
 }
 
-void UWorldPartition::BeginDestroy()
-{
-	Uninitialize();
-
-	Super::BeginDestroy();
-}
-
 UWorld* UWorldPartition::GetWorld() const
 {
 	if (World)
@@ -373,11 +363,8 @@ void UWorldPartition::Uninitialize()
 		EditorHash = nullptr;
 #endif		
 
-		if (World->IsGameWorld())
-		{
-			UWorldPartitionSubsystem* WorldPartitionSubsystem = World->GetSubsystem<UWorldPartitionSubsystem>();
-			WorldPartitionSubsystem->UnregisterWorldPartition(this);
-		}
+		UWorldPartitionSubsystem* WorldPartitionSubsystem = World->GetSubsystem<UWorldPartitionSubsystem>();
+		WorldPartitionSubsystem->UnregisterWorldPartition(this);
 
 		World = nullptr;
 
diff --git a/Engine/Source/Runtime/Engine/Public/WorldPartition/WorldPartition.h b/Engine/Source/Runtime/Engine/Public/WorldPartition/WorldPartition.h
--- a/Engine/Source/Runtime/Engine/Public/WorldPartition/WorldPartition.h
+++ b/Engine/Source/Runtime/Engine/Public/WorldPartition/WorldPartition.h
@@ -111,7 +111,6 @@ private:
 #endif
 
 	virtual void Serialize(FArchive& Ar) override;
-	virtual void BeginDestroy() override;
 
 	static void WorldPartitionOnLevelRemovedFromWorld(class ULevel* Level, UWorld* InWorld);
 
