Remove unsupported `Rebuild HLODs` console command in the editor.

diff --git a/Engine/Source/Runtime/Engine/Private/WorldPartition/HLOD/HLODBuilder.cpp b/Engine/Source/Runtime/Engine/Private/WorldPartition/HLOD/HLODBuilder.cpp
--- a/Engine/Source/Runtime/Engine/Private/WorldPartition/HLOD/HLODBuilder.cpp
+++ b/Engine/Source/Runtime/Engine/Private/WorldPartition/HLOD/HLODBuilder.cpp
@@ -45,17 +45,10 @@ public:
 			FWorldPartitionActorDesc* HLODActorDesc = WorldPartition->GetActorDesc(HLODActorGuid);
 			check(HLODActorDesc);
 
-			if (IsRunningCommandlet())
-			{
-				check(!HLODActorDesc->GetLoadedRefCount());
-				HLODActorDesc->AddLoadedRefCount();
-				HLODActor = CastChecked<AWorldPartitionHLOD>(HLODActorDesc->Load());
-				HLODActor->GetLevel()->AddLoadedActor(HLODActor);
-			}
-			else
-			{
-				HLODActor = CastChecked<AWorldPartitionHLOD>(HLODActorDesc->GetActor());
-			}
+			check(!HLODActorDesc->GetLoadedRefCount());
+			HLODActorDesc->AddLoadedRefCount();
+			HLODActor = CastChecked<AWorldPartitionHLOD>(HLODActorDesc->Load());
+			HLODActor->GetLevel()->AddLoadedActor(HLODActor);
 		}
 
 		if (!HLODActor)
diff --git a/Engine/Source/Runtime/Engine/Private/WorldPartition/WorldPartition.cpp b/Engine/Source/Runtime/Engine/Private/WorldPartition/WorldPartition.cpp
--- a/Engine/Source/Runtime/Engine/Private/WorldPartition/WorldPartition.cpp
+++ b/Engine/Source/Runtime/Engine/Private/WorldPartition/WorldPartition.cpp
@@ -45,24 +45,6 @@ DEFINE_LOG_CATEGORY(LogWorldPartition);
 #define LOCTEXT_NAMESPACE "WorldPartitionEditor"
 
 #if WITH_EDITOR
-static FAutoConsoleCommand GenerateHLODCmd(
-	TEXT("wp.Editor.GenerateHLOD"),
-	TEXT("Generates HLOD data for runtime."),
-	FConsoleCommandWithArgsDelegate::CreateLambda([](const TArray<FString>& Args)
-	{
-		if (UWorld* World = GEditor->GetEditorWorldContext().World())
-		{
-			if (!World->IsPlayInEditor())
-			{
-				if (UWorldPartition* WorldPartition = World->GetWorldPartition())
-				{
-					WorldPartition->GenerateHLOD(nullptr);
-				}
-			}
-		}
-	})
-);
-
 static FAutoConsoleCommand DumpActorDescs(
 	TEXT("wp.Editor.DumpActorDescs"),
 	TEXT("Dump the list of actor descriptors in a CSV file."),
diff --git a/Engine/Source/Runtime/Engine/Private/WorldPartition/WorldPartitionRuntimeSpatialHash.cpp b/Engine/Source/Runtime/Engine/Private/WorldPartition/WorldPartitionRuntimeSpatialHash.cpp
--- a/Engine/Source/Runtime/Engine/Private/WorldPartition/WorldPartitionRuntimeSpatialHash.cpp
+++ b/Engine/Source/Runtime/Engine/Private/WorldPartition/WorldPartitionRuntimeSpatialHash.cpp
@@ -1547,29 +1547,23 @@ TArray<FGuid> GenerateHLODsForGrid(UWorldPartition* WorldPartition, const FSpati
 				{
 					FWorldPartitionActorDesc* ActorDesc = (FWorldPartitionActorDesc*)WorldPartition->GetActorDesc(ActorGuid);
 
-					if (IsRunningCommandlet())
-					{					
-						check(!ActorDesc->GetLoadedRefCount());
-						ActorDesc->AddLoadedRefCount();
+					check(!ActorDesc->GetLoadedRefCount());
+					ActorDesc->AddLoadedRefCount();
 
-						AActor* Actor = ActorDesc->GetActor();
+					AActor* Actor = ActorDesc->GetActor();
 					
-						if (!Actor)
-						{
-							Actor = ActorDesc->Load();
-							check(Actor);
-						}
-
-						Actor->GetLevel()->AddLoadedActor(Actor);
+					if (!Actor)
+					{
+						Actor = ActorDesc->Load();
+						check(Actor);
 					}
 
+					Actor->GetLevel()->AddLoadedActor(Actor);
+
 					CellActors.Add(ActorDesc->GetActor());
 				}
 
-				if (IsRunningCommandlet())
-				{
-					CollectGarbage(GARBAGE_COLLECTION_KEEPFLAGS, true);
-				}
+				CollectGarbage(GARBAGE_COLLECTION_KEEPFLAGS, true);
 
 				TArray<AWorldPartitionHLOD*> CellHLODActors;
 				CellHLODActors = UHLODLayer::GenerateHLODForCell(WorldPartition, &Context, CellName, CellBounds, HLODLevel, CellActors);
@@ -1583,73 +1577,64 @@ TArray<FGuid> GenerateHLODsForGrid(UWorldPartition* WorldPartition, const FSpati
 						FGuid ActorGuid = CellHLODActor->GetActorGuid();
 						GridHLODActors.Add(ActorGuid);
 
-						if (IsRunningCommandlet())
-						{
-							// Checkout packages
-							UPackage* CellHLODActorPackage = CellHLODActor->GetPackage();
-							CellHLODActorPackage->MarkAsFullyLoaded();
+						// Checkout packages
+						UPackage* CellHLODActorPackage = CellHLODActor->GetPackage();
+						CellHLODActorPackage->MarkAsFullyLoaded();
 
-							if (CellHLODActorPackage->HasAnyPackageFlags(PKG_NewlyCreated))
-							{
-								NewCellHLODActors.Add(CellHLODActor);
-							}
+						if (CellHLODActorPackage->HasAnyPackageFlags(PKG_NewlyCreated))
+						{
+							NewCellHLODActors.Add(CellHLODActor);
+						}
 
-							if (!SourceControlHelper->Checkout(CellHLODActorPackage))
-							{
-								UE_LOG(LogWorldPartitionRuntimeSpatialHash, Error, TEXT("Error checking out package %s."), *CellHLODActorPackage->GetName());
-								check(0);
-							}
+						if (!SourceControlHelper->Checkout(CellHLODActorPackage))
+						{
+							UE_LOG(LogWorldPartitionRuntimeSpatialHash, Error, TEXT("Error checking out package %s."), *CellHLODActorPackage->GetName());
+							check(0);
+						}
 
-							// Save packages
-							FString PackageFileName = SourceControlHelper->GetFilename(CellHLODActorPackage);
-							if (!UPackage::SavePackage(CellHLODActorPackage, nullptr, RF_Standalone, *PackageFileName))
-							{
-								UE_LOG(LogWorldPartitionRuntimeSpatialHash, Error, TEXT("Error saving package %s."), *CellHLODActorPackage->GetName());
-								check(0);
-							}
+						// Save packages
+						FString PackageFileName = SourceControlHelper->GetFilename(CellHLODActorPackage);
+						if (!UPackage::SavePackage(CellHLODActorPackage, nullptr, RF_Standalone, *PackageFileName))
+						{
+							UE_LOG(LogWorldPartitionRuntimeSpatialHash, Error, TEXT("Error saving package %s."), *CellHLODActorPackage->GetName());
+							check(0);
+						}
 
-							// Add new packages to source control
-							if (!SourceControlHelper->Add(CellHLODActorPackage))
-							{
-								UE_LOG(LogWorldPartitionRuntimeSpatialHash, Error, TEXT("Error adding package %s to source control."), *CellHLODActorPackage->GetName());
-								check(0);
-							}
+						// Add new packages to source control
+						if (!SourceControlHelper->Add(CellHLODActorPackage))
+						{
+							UE_LOG(LogWorldPartitionRuntimeSpatialHash, Error, TEXT("Error adding package %s to source control."), *CellHLODActorPackage->GetName());
+							check(0);
 						}
 					}
 
-					if (IsRunningCommandlet())
-					{
-						// Manually tick the directory watcher and the asset registry to register newly created actors
-						DirectoryWatcherModule.Get()->Tick(-1.0f);
-						AssetRegistryModule.Get().Tick(-1.0f);
+					// Manually tick the directory watcher and the asset registry to register newly created actors
+					DirectoryWatcherModule.Get()->Tick(-1.0f);
+					AssetRegistryModule.Get().Tick(-1.0f);
 
-						// Update newly created HLOD actors
-						for (AWorldPartitionHLOD* CellHLODActor: NewCellHLODActors)
-						{
-							FWorldPartitionActorDesc* ActorDesc = (FWorldPartitionActorDesc*)WorldPartition->GetActorDesc(CellHLODActor->GetActorGuid());
+					// Update newly created HLOD actors
+					for (AWorldPartitionHLOD* CellHLODActor: NewCellHLODActors)
+					{
+						FWorldPartitionActorDesc* ActorDesc = (FWorldPartitionActorDesc*)WorldPartition->GetActorDesc(CellHLODActor->GetActorGuid());
 					
-							check(!ActorDesc->GetLoadedRefCount());
-							ActorDesc->AddLoadedRefCount();
-						}
+						check(!ActorDesc->GetLoadedRefCount());
+						ActorDesc->AddLoadedRefCount();
 					}
 				}
 
 				// Unload actors
-				if (IsRunningCommandlet())
-				{
-					TArray<AActor*> ToUnloadActors = MoveTemp(CellActors);
-					ToUnloadActors.Append(CellHLODActors);
+				TArray<AActor*> ToUnloadActors = MoveTemp(CellActors);
+				ToUnloadActors.Append(CellHLODActors);
 
-					for (AActor* Actor: ToUnloadActors)
-					{
-						FWorldPartitionActorDesc* ActorDesc = (FWorldPartitionActorDesc*)WorldPartition->GetActorDesc(Actor->GetActorGuid());
+				for (AActor* Actor: ToUnloadActors)
+				{
+					FWorldPartitionActorDesc* ActorDesc = (FWorldPartitionActorDesc*)WorldPartition->GetActorDesc(Actor->GetActorGuid());
 					
-						check(ActorDesc->GetLoadedRefCount() == 1);
-						ActorDesc->RemoveLoadedRefCount();
+					check(ActorDesc->GetLoadedRefCount() == 1);
+					ActorDesc->RemoveLoadedRefCount();
 
-						Actor->GetLevel()->RemoveLoadedActor(Actor);
-						ActorDesc->Unload();
-					}
+					Actor->GetLevel()->RemoveLoadedActor(Actor);
+					ActorDesc->Unload();
 				}
 			}
 		}
@@ -1951,17 +1936,14 @@ bool UWorldPartitionRuntimeSpatialHash::GenerateHLOD(ISourceControlHelper* Sourc
 		GenerateHLODs(HLODGrids[HLODGridName], GridsDepth[HLODGridName] + 1, HLODActorClusters);
 	}
 
-	if (IsRunningCommandlet())
+	// Destroy all unreferenced HLOD actors
+	for (const auto& HLODActorPair: Context.HLODActorDescs)
 	{
-		// Destroy all unreferenced HLOD actors
-		for (const auto& HLODActorPair: Context.HLODActorDescs)
-		{
-			FWorldPartitionActorDesc* HLODActorDesc = WorldPartition->GetActorDesc(HLODActorPair.Value);
-			check(HLODActorDesc);
+		FWorldPartitionActorDesc* HLODActorDesc = WorldPartition->GetActorDesc(HLODActorPair.Value);
+		check(HLODActorDesc);
 
-			const FString ActorPackageFilename = SourceControlHelper->GetFilename(HLODActorDesc->GetActorPackage().ToString());		
-			SourceControlHelper->Delete(ActorPackageFilename);
-		}
+		const FString ActorPackageFilename = SourceControlHelper->GetFilename(HLODActorDesc->GetActorPackage().ToString());		
+		SourceControlHelper->Delete(ActorPackageFilename);
 	}
 
 	// Create/destroy HLOD grid actors
