-Adding the option to override BlendIn/BlendOut settings on Montage_Play/Stop.

diff --git a/Engine/Source/Runtime/Engine/Classes/Animation/AnimInstance.h b/Engine/Source/Runtime/Engine/Classes/Animation/AnimInstance.h
--- a/Engine/Source/Runtime/Engine/Classes/Animation/AnimInstance.h
+++ b/Engine/Source/Runtime/Engine/Classes/Animation/AnimInstance.h
@@ -600,10 +600,19 @@ public:
 	UFUNCTION(BlueprintCallable, Category = "Montage")
 	float Montage_Play(UAnimMontage* MontageToPlay, float InPlayRate = 1.f, EMontagePlayReturnType ReturnValueType = EMontagePlayReturnType::MontageLength, float InTimeToStartMontageAt=0.f, bool bStopAllMontages = true);
 
+	/** Plays an animation montage. Same as Montage_Play, but you can specify an AlphaBlend for Blend In settings. */
+	UFUNCTION(BlueprintCallable, Category = "Montage")
+	float Montage_PlayWithBlendIn(UAnimMontage* MontageToPlay, const FAlphaBlendArgs& BlendIn, float InPlayRate = 1.f, EMontagePlayReturnType ReturnValueType = EMontagePlayReturnType::MontageLength, float InTimeToStartMontageAt=0.f, bool bStopAllMontages = true);
+
 	/** Stops the animation montage. If reference is NULL, it will stop ALL active montages. */
+	/** Stopped montages will blend out using their montage asset's BlendOut, with InBlendOutTime as the BlendTime */
 	UFUNCTION(BlueprintCallable, Category = "Montage")
 	void Montage_Stop(float InBlendOutTime, const UAnimMontage* Montage = NULL);
 
+	/** Same as Montage_Stop, but all BlendOut settings are provided instead of using the ones on the montage asset*/
+	UFUNCTION(BlueprintCallable, Category = "Montage")
+	void Montage_StopWithBlendOut(const FAlphaBlendArgs& BlendOut, const UAnimMontage* Montage = nullptr);
+
 	/** Stops all active montages belonging to a group. */
 	UFUNCTION(BlueprintCallable, Category = "Montage")
 	void Montage_StopGroupByName(float InBlendOutTime, FName GroupName);
@@ -841,6 +850,9 @@ protected:
 	/** Advance montages **/
 	virtual void Montage_Advance(float DeltaSeconds);
 
+	void Montage_StopInternal(TFunctionRef<FAlphaBlend(const FAnimMontageInstance*)> AlphaBlendSelectorFunction, const UAnimMontage* Montage = nullptr);
+	float Montage_PlayInternal(UAnimMontage* MontageToPlay, const FAlphaBlend& BlendIn, float InPlayRate = 1.f, EMontagePlayReturnType ReturnValueType = EMontagePlayReturnType::MontageLength, float InTimeToStartMontageAt = 0.f, bool bStopAllMontages = true);
+
 public:
 	/** Queue a Montage BlendingOut Event to be triggered. */
 	void QueueMontageBlendingOutEvent(const FQueuedMontageBlendingOutEvent& MontageBlendingOutEvent);
diff --git a/Engine/Source/Runtime/Engine/Classes/Animation/AnimMontage.h b/Engine/Source/Runtime/Engine/Classes/Animation/AnimMontage.h
--- a/Engine/Source/Runtime/Engine/Classes/Animation/AnimMontage.h
+++ b/Engine/Source/Runtime/Engine/Classes/Animation/AnimMontage.h
@@ -434,7 +434,12 @@ public:
 	FAnimMontageInstance(UAnimInstance * InAnimInstance);
 
 	//~ Begin montage instance Interfaces
+
+	// Fills BlendIn from the associated montage 
 	void Play(float InPlayRate = 1.f);
+	// Use provided BlendIn
+	void Play(float InPlayRate, const FAlphaBlend& BlendIn);
+
 	void Stop(const FAlphaBlend& InBlendOut, bool bInterrupt=true);
 	void Pause();
 	void Initialize(class UAnimMontage * InMontage);
diff --git a/Engine/Source/Runtime/Engine/Private/AlphaBlend.cpp b/Engine/Source/Runtime/Engine/Private/AlphaBlend.cpp
--- a/Engine/Source/Runtime/Engine/Private/AlphaBlend.cpp
+++ b/Engine/Source/Runtime/Engine/Private/AlphaBlend.cpp
@@ -23,6 +23,16 @@ FAlphaBlend::FAlphaBlend(const FAlphaBlend& Other, float NewBlendTime)
 	Reset();
 }
 
+FAlphaBlend::FAlphaBlend(const FAlphaBlendArgs& InArgs)
+	: CustomCurve(InArgs.CustomCurve)
+	, BlendTime(InArgs.BlendTime)
+	, BeginValue(0.0f)
+	, DesiredValue(1.0f)
+	, BlendOption(InArgs.BlendOption)
+{
+	Reset();
+}
+
 void FAlphaBlend::ResetBlendTime()
 {
 	// if blend time is <= 0, then blending is done and complete
@@ -217,3 +227,10 @@ bool FAlphaBlend::IsComplete() const
 	return (CachedDesiredBlendedValue == BlendedValue);
 }
 
+FAlphaBlendArgs::FAlphaBlendArgs()
+	: CustomCurve(nullptr)
+	, BlendTime(0.2f)
+	, BlendOption(EAlphaBlendOption::Linear)
+{
+
+}
diff --git a/Engine/Source/Runtime/Engine/Private/Animation/AnimInstance.cpp b/Engine/Source/Runtime/Engine/Private/Animation/AnimInstance.cpp
--- a/Engine/Source/Runtime/Engine/Private/Animation/AnimInstance.cpp
+++ b/Engine/Source/Runtime/Engine/Private/Animation/AnimInstance.cpp
@@ -1907,8 +1907,7 @@ bool UAnimInstance::IsPlayingSlotAnimation(const UAnimSequenceBase* Asset, FName
 	return false;
 }
 
-/** Play a Montage. Returns Length of Montage in seconds. Returns 0.f if failed to play. */
-float UAnimInstance::Montage_Play(UAnimMontage* MontageToPlay, float InPlayRate/*= 1.f*/, EMontagePlayReturnType ReturnValueType, float InTimeToStartMontageAt, bool bStopAllMontages /*= true*/)
+float UAnimInstance::Montage_PlayInternal(UAnimMontage* MontageToPlay, const FAlphaBlend& BlendIn, float InPlayRate /*= 1.f*/, EMontagePlayReturnType ReturnValueType /*= EMontagePlayReturnType::MontageLength*/, float InTimeToStartMontageAt /*= 0.f*/, bool bStopAllMontages /*= true*/)
 {
 	LLM_SCOPE(ELLMTag::Animation);
 
@@ -1920,7 +1919,7 @@ float UAnimInstance::Montage_Play(UAnimMontage* MontageToPlay, float InPlayRate/
 			{
 				// Enforce 'a single montage at once per group' rule
 				FName NewMontageGroupName = MontageToPlay->GetGroupName();
-				StopAllMontagesByGroupName(NewMontageGroupName, MontageToPlay->BlendIn);
+				StopAllMontagesByGroupName(NewMontageGroupName, BlendIn);
 			}
 
 			// Enforce 'a single root motion montage at once' rule.
@@ -1939,7 +1938,7 @@ float UAnimInstance::Montage_Play(UAnimMontage* MontageToPlay, float InPlayRate/
 			const float MontageLength = MontageToPlay->GetPlayLength();
 
 			NewInstance->Initialize(MontageToPlay);
-			NewInstance->Play(InPlayRate);
+			NewInstance->Play(InPlayRate, BlendIn);
 			NewInstance->SetPosition(FMath::Clamp(InTimeToStartMontageAt, 0.f, MontageLength));
 			MontageInstances.Add(NewInstance);
 			ActiveMontagesMap.Add(MontageToPlay, NewInstance);
@@ -1967,14 +1966,26 @@ float UAnimInstance::Montage_Play(UAnimMontage* MontageToPlay, float InPlayRate/
 	return 0.f;
 }
 
-void UAnimInstance::Montage_Stop(float InBlendOutTime, const UAnimMontage* Montage)
+/** Play a Montage. Returns Length of Montage in seconds. Returns 0.f if failed to play. */
+float UAnimInstance::Montage_Play(UAnimMontage* MontageToPlay, float InPlayRate/*= 1.f*/, EMontagePlayReturnType ReturnValueType, float InTimeToStartMontageAt, bool bStopAllMontages /*= true*/)
+{
+	FAlphaBlend BlendIn = MontageToPlay ? MontageToPlay->BlendIn : FAlphaBlend();
+	return Montage_PlayInternal(MontageToPlay, BlendIn, InPlayRate, ReturnValueType, InTimeToStartMontageAt, bStopAllMontages);
+}
+
+float UAnimInstance::Montage_PlayWithBlendIn(UAnimMontage* MontageToPlay, const FAlphaBlendArgs& BlendIn, float InPlayRate /*= 1.f*/, EMontagePlayReturnType ReturnValueType /*= EMontagePlayReturnType::MontageLength*/, float InTimeToStartMontageAt/*=0.f*/, bool bStopAllMontages /*= true*/)
+{
+	return Montage_PlayInternal(MontageToPlay, FAlphaBlend(BlendIn), InPlayRate, ReturnValueType, InTimeToStartMontageAt, bStopAllMontages);
+}
+
+void UAnimInstance::Montage_StopInternal(TFunctionRef<FAlphaBlend(const FAnimMontageInstance*)> AlphaBlendSelectorFunction, const UAnimMontage* Montage /*= nullptr*/)
 {
 	if (Montage)
 	{
 		FAnimMontageInstance* MontageInstance = GetActiveInstanceForMontage(Montage);
 		if (MontageInstance)
 		{
-			MontageInstance->Stop(FAlphaBlend(Montage->BlendOut, InBlendOutTime));
+			MontageInstance->Stop(AlphaBlendSelectorFunction(MontageInstance));
 		}
 	}
 	else
@@ -1985,12 +1996,34 @@ void UAnimInstance::Montage_Stop(float InBlendOutTime, const UAnimMontage* Monta
 			FAnimMontageInstance* MontageInstance = MontageInstances[InstanceIndex];
 			if (MontageInstance && MontageInstance->IsActive())
 			{
-				MontageInstance->Stop(FAlphaBlend(MontageInstance->Montage->BlendOut, InBlendOutTime));
+				MontageInstance->Stop(AlphaBlendSelectorFunction(MontageInstance));
 			}
 		}
 	}
 }
 
+void UAnimInstance::Montage_Stop(float InBlendOutTime, const UAnimMontage* Montage)
+{
+	auto AlphaBlendFromInstanceAndInBlendOutTime = [InBlendOutTime](const FAnimMontageInstance* InMontageInstance)
+	{
+		check(InMontageInstance->Montage);
+		return FAlphaBlend(InMontageInstance->Montage->BlendOut, InBlendOutTime);
+	};
+
+	Montage_StopInternal(AlphaBlendFromInstanceAndInBlendOutTime, Montage);
+}
+
+void UAnimInstance::Montage_StopWithBlendOut(const FAlphaBlendArgs& BlendOutArgs, const UAnimMontage* Montage /*= NULL*/)
+{
+	FAlphaBlend BlendOut = FAlphaBlend(BlendOutArgs);
+	auto AlphaBlendPassthrough = [BlendOut](const FAnimMontageInstance* InMontageInstance)
+	{
+		return BlendOut;
+	};
+
+	Montage_StopInternal(AlphaBlendPassthrough, Montage);
+}
+
 void UAnimInstance::Montage_StopGroupByName(float InBlendOutTime, FName GroupName)
 {
 	for (int32 InstanceIndex = 0; InstanceIndex < MontageInstances.Num(); InstanceIndex++)
diff --git a/Engine/Source/Runtime/Engine/Private/Animation/AnimMontage.cpp b/Engine/Source/Runtime/Engine/Private/Animation/AnimMontage.cpp
--- a/Engine/Source/Runtime/Engine/Private/Animation/AnimMontage.cpp
+++ b/Engine/Source/Runtime/Engine/Private/Animation/AnimMontage.cpp
@@ -1330,6 +1330,11 @@ FAnimMontageInstance::FAnimMontageInstance(UAnimInstance * InAnimInstance)
 }
 
 void FAnimMontageInstance::Play(float InPlayRate)
+{
+	Play(InPlayRate, Montage->BlendIn);
+}
+
+void FAnimMontageInstance::Play(float InPlayRate, const FAlphaBlend& BlendIn)
 {
 	bPlaying = true;
 	PlayRate = InPlayRate;
@@ -1339,8 +1344,9 @@ void FAnimMontageInstance::Play(float InPlayRate)
 	
 	// set blend option
 	float CurrentWeight = Blend.GetBlendedValue();
-	InitializeBlend(Montage->BlendIn);	
-	Blend.SetBlendTime(Montage->BlendIn.GetBlendTime() * DefaultBlendTimeMultiplier);
+
+	InitializeBlend(BlendIn);
+	Blend.SetBlendTime(BlendIn.GetBlendTime() * DefaultBlendTimeMultiplier);
 	Blend.SetValueRange(CurrentWeight, 1.f);
 	bEnableAutoBlendOut = Montage->bEnableAutoBlendOut;
 }
diff --git a/Engine/Source/Runtime/Engine/Public/AlphaBlend.h b/Engine/Source/Runtime/Engine/Public/AlphaBlend.h
--- a/Engine/Source/Runtime/Engine/Public/AlphaBlend.h
+++ b/Engine/Source/Runtime/Engine/Public/AlphaBlend.h
@@ -43,6 +43,29 @@ enum class EAlphaBlendOption : uint8
 	Custom,
 };
 
+/**
+ * Alpha Blend construction arguments. Used only for creation of an AlphaBlend.
+ */
+USTRUCT(BlueprintType)
+struct ENGINE_API FAlphaBlendArgs
+{
+	GENERATED_BODY()
+
+	FAlphaBlendArgs();
+
+	/** If you're using Custom BlendOption, you can specify curve */
+	UPROPERTY(EditAnywhere, BlueprintReadWrite, Category = "Blend", meta=(DisplayAfter="BlendOption"))
+	UCurveFloat* CustomCurve;
+
+	/** Blend Time */
+	UPROPERTY(EditAnywhere, BlueprintReadWrite, Category = "Blend")
+	float BlendTime;
+
+	/** Type of blending used (Linear, Cubic, etc.) */
+	UPROPERTY(EditAnywhere, BlueprintReadWrite, Category = "Blend")
+	EAlphaBlendOption BlendOption;
+};
+
 /**
  * Alpha Blend class that supports different blend options as well as custom curves
  */
@@ -67,6 +90,9 @@ public:
 	/* Constructor */
 	FAlphaBlend(const FAlphaBlend& Other, float NewBlendTime);
 
+	/* Constructor */
+	explicit FAlphaBlend(const FAlphaBlendArgs& InArgs);
+
 	/** Setters - need to refresh cached value */
 	void SetBlendOption(EAlphaBlendOption InBlendOption);
 	void SetCustomCurve(UCurveFloat* InCustomCurve);
