Allow material cached data to be nullptr

diff --git a/Engine/Source/Runtime/Engine/Classes/Materials/Material.h b/Engine/Source/Runtime/Engine/Classes/Materials/Material.h
--- a/Engine/Source/Runtime/Engine/Classes/Materials/Material.h
+++ b/Engine/Source/Runtime/Engine/Classes/Materials/Material.h
@@ -991,7 +991,7 @@ private:
 #endif // WITH_EDITORONLY_DATA
 public:
 
-	const FMaterialCachedExpressionData& GetCachedExpressionData() const { check(CachedExpressionData); return *CachedExpressionData; }
+	const FMaterialCachedExpressionData& GetCachedExpressionData() const { return CachedExpressionData ? *CachedExpressionData : FMaterialCachedExpressionData::EmptyData; }
 
 	//~ Begin UMaterialInterface Interface.
 	ENGINE_API virtual UMaterial* GetMaterial() override;
diff --git a/Engine/Source/Runtime/Engine/Classes/Materials/MaterialInstance.h b/Engine/Source/Runtime/Engine/Classes/Materials/MaterialInstance.h
--- a/Engine/Source/Runtime/Engine/Classes/Materials/MaterialInstance.h
+++ b/Engine/Source/Runtime/Engine/Classes/Materials/MaterialInstance.h
@@ -665,7 +665,7 @@ public:
 #endif // WITH_EDITOR
 
 	/** Appends textures referenced by expressions, including nested functions. */
-	ENGINE_API virtual TArrayView<UObject* const> GetReferencedTextures() const override final { check(CachedData); return CachedData->ReferencedTextures; }
+	ENGINE_API virtual TArrayView<UObject* const> GetReferencedTextures() const override final { return CachedData ? CachedData->ReferencedTextures : TArrayView<UObject* const>(); }
 
 #if WITH_EDITOR
 	/** Add to the set any texture referenced by expressions, including nested functions, as well as any overrides from parameters. */
diff --git a/Engine/Source/Runtime/Engine/Private/Materials/Material.cpp b/Engine/Source/Runtime/Engine/Private/Materials/Material.cpp
--- a/Engine/Source/Runtime/Engine/Private/Materials/Material.cpp
+++ b/Engine/Source/Runtime/Engine/Private/Materials/Material.cpp
@@ -60,6 +60,7 @@
 #include "MaterialShaderQualitySettings.h"
 #include "Engine/RendererSettings.h"
 #include "ProfilingDebugging/LoadTimeTracker.h"
+#include "ProfilingDebugging/CookStats.h"
 #include "MaterialCompiler.h"
 #include "MaterialShaderType.h"
 #include "Materials/MaterialInstanceSupport.h"
@@ -88,6 +89,21 @@
 #include "HAL/ThreadHeartBeat.h"
 #include "Misc/ScopedSlowTask.h"
 
+#if ENABLE_COOK_STATS
+#include "ProfilingDebugging/ScopedTimers.h"
+namespace MaterialCookStats
+{
+	static double UpdateCachedExpressionDataSec = 0.0;
+
+	static FCookStatsManager::FAutoRegisterCallback RegisterCookStats([](FCookStatsManager::AddStatFuncRef AddStat)
+		{
+			AddStat(TEXT("Material"), FCookStatsManager::CreateKeyValueArray(
+				TEXT("UpdateCachedExpressionDataSec"), UpdateCachedExpressionDataSec
+			));
+		});
+}
+#endif
+
 #define LOCTEXT_NAMESPACE "Material"
 
 static TAutoConsoleVariable<int32> CVarMaterialParameterLegacyChecks(
@@ -2093,6 +2109,8 @@ bool UMaterial::GetParameterDesc(const FHashedMaterialParameterInfo& ParameterIn
 }
 void UMaterial::UpdateCachedExpressionData()
 {
+	COOK_STAT(FScopedDurationTimer BlockingTimer(MaterialCookStats::UpdateCachedExpressionDataSec));
+
 	if (bSavedCachedExpressionData)
 	{
 		// Don't need to rebuild cached data if it was serialized
@@ -2129,14 +2147,17 @@ bool UMaterial::GetScalarParameterValue(const FHashedMaterialParameterInfo& Para
 
 bool UMaterial::GetScalarParameterValue_New(const FHashedMaterialParameterInfo& ParameterInfo, float& OutValue, bool bOveriddenOnly) const
 {
-	const int32 Index = CachedExpressionData->Parameters.FindParameterIndex(EMaterialParameterType::Scalar, ParameterInfo, bOveriddenOnly);
-	if (Index == INDEX_NONE)
+	if (CachedExpressionData)
 	{
-		return false;
+		const int32 Index = CachedExpressionData->Parameters.FindParameterIndex(EMaterialParameterType::Scalar, ParameterInfo, bOveriddenOnly);
+		if (Index != INDEX_NONE)
+		{
+			OutValue = CachedExpressionData->Parameters.ScalarValues[Index];
+			return true;
+		}
 	}
 
-	OutValue = CachedExpressionData->Parameters.ScalarValues[Index];
-	return true;
+	return false;
 }
 
 #if WITH_EDITOR
@@ -2214,38 +2235,44 @@ bool UMaterial::GetScalarParameterValue_Legacy(const FHashedMaterialParameterInf
 
 bool UMaterial::IsScalarParameterUsedAsAtlasPosition(const FHashedMaterialParameterInfo& ParameterInfo, bool& OutValue, TSoftObjectPtr<UCurveLinearColor>& OutCurve, TSoftObjectPtr<UCurveLinearColorAtlas>& OutAtlas) const
 {
-	const int32 Index = CachedExpressionData->Parameters.FindParameterIndex(EMaterialParameterType::Scalar, ParameterInfo);
-	if (Index == INDEX_NONE)
-	{
-		return false;
-	}
-
-	UCurveLinearColor* Curve = CachedExpressionData->Parameters.ScalarCurveValues[Index];
-	UCurveLinearColorAtlas* Atlas = CachedExpressionData->Parameters.ScalarCurveAtlasValues[Index];
-	if (Curve && Atlas)
-	{
-		OutCurve = TSoftObjectPtr<UCurveLinearColor>(FSoftObjectPath(Curve->GetPathName()));
-		OutAtlas = TSoftObjectPtr<UCurveLinearColorAtlas>(FSoftObjectPath(Atlas->GetPathName()));
-		OutValue = true;
-	}
-	else
+	if (CachedExpressionData)
 	{
-		OutValue = false;
+		const int32 Index = CachedExpressionData->Parameters.FindParameterIndex(EMaterialParameterType::Scalar, ParameterInfo);
+		if (Index != INDEX_NONE)
+		{
+			UCurveLinearColor* Curve = CachedExpressionData->Parameters.ScalarCurveValues[Index];
+			UCurveLinearColorAtlas* Atlas = CachedExpressionData->Parameters.ScalarCurveAtlasValues[Index];
+			if (Curve && Atlas)
+			{
+				OutCurve = TSoftObjectPtr<UCurveLinearColor>(FSoftObjectPath(Curve->GetPathName()));
+				OutAtlas = TSoftObjectPtr<UCurveLinearColorAtlas>(FSoftObjectPath(Atlas->GetPathName()));
+				OutValue = true;
+			}
+			else
+			{
+				OutValue = false;
+			}
+			return true;
+		}
 	}
 
-	return true;
+	return false;
 }
 
 bool UMaterial::GetScalarParameterSliderMinMax(const FHashedMaterialParameterInfo& ParameterInfo, float& OutSliderMin, float& OutSliderMax) const
 {
-	const int32 Index = CachedExpressionData->Parameters.FindParameterIndex(EMaterialParameterType::Scalar, ParameterInfo);
-	if (Index == INDEX_NONE)
+	if (CachedExpressionData)
 	{
-		return false;
+		const int32 Index = CachedExpressionData->Parameters.FindParameterIndex(EMaterialParameterType::Scalar, ParameterInfo);
+		if (Index != INDEX_NONE)
+		{
+			OutSliderMin = CachedExpressionData->Parameters.ScalarMinMaxValues[Index].X;
+			OutSliderMax = CachedExpressionData->Parameters.ScalarMinMaxValues[Index].Y;
+			return true;
+		}
 	}
-	OutSliderMin = CachedExpressionData->Parameters.ScalarMinMaxValues[Index].X;
-	OutSliderMax = CachedExpressionData->Parameters.ScalarMinMaxValues[Index].Y;
-	return true;
+
+	return false;
 }
 #endif // WITH_EDITOR
 
@@ -2268,13 +2295,16 @@ bool UMaterial::GetVectorParameterValue(const FHashedMaterialParameterInfo& Para
 
 bool UMaterial::GetVectorParameterValue_New(const FHashedMaterialParameterInfo& ParameterInfo, FLinearColor& OutValue, bool bOveriddenOnly) const
 {
-	const int32 Index = CachedExpressionData->Parameters.FindParameterIndex(EMaterialParameterType::Vector, ParameterInfo, bOveriddenOnly);
-	if (Index == INDEX_NONE)
+	if (CachedExpressionData)
 	{
-		return false;
+		const int32 Index = CachedExpressionData->Parameters.FindParameterIndex(EMaterialParameterType::Vector, ParameterInfo, bOveriddenOnly);
+		if (Index != INDEX_NONE)
+		{
+			OutValue = CachedExpressionData->Parameters.VectorValues[Index];
+			return true;
+		}
 	}
-	OutValue = CachedExpressionData->Parameters.VectorValues[Index];
-	return true;
+	return false;
 }
 
 #if WITH_EDITOR
@@ -2349,24 +2379,30 @@ bool UMaterial::GetVectorParameterValue_Legacy(const FHashedMaterialParameterInf
 
 bool UMaterial::IsVectorParameterUsedAsChannelMask(const FHashedMaterialParameterInfo& ParameterInfo, bool& OutValue) const
 {
-	const int32 Index = CachedExpressionData->Parameters.FindParameterIndex(EMaterialParameterType::Vector, ParameterInfo);
-	if (Index == INDEX_NONE)
+	if (CachedExpressionData)
 	{
-		return false;
+		const int32 Index = CachedExpressionData->Parameters.FindParameterIndex(EMaterialParameterType::Vector, ParameterInfo);
+		if (Index != INDEX_NONE)
+		{
+			OutValue = CachedExpressionData->Parameters.VectorUsedAsChannelMaskValues[Index];
+			return true;
+		}
 	}
-	OutValue = CachedExpressionData->Parameters.VectorUsedAsChannelMaskValues[Index];
-	return true;
+	return false;
 }
 
 bool UMaterial::GetVectorParameterChannelNames(const FHashedMaterialParameterInfo& ParameterInfo, FParameterChannelNames& OutValue) const
 {
-	const int32 Index = CachedExpressionData->Parameters.FindParameterIndex(EMaterialParameterType::Vector, ParameterInfo);
-	if (Index == INDEX_NONE)
+	if (CachedExpressionData)
 	{
-		return false;
+		const int32 Index = CachedExpressionData->Parameters.FindParameterIndex(EMaterialParameterType::Vector, ParameterInfo);
+		if (Index != INDEX_NONE)
+		{
+			OutValue = CachedExpressionData->Parameters.VectorChannelNameValues[Index];
+			return true;
+		}
 	}
-	OutValue = CachedExpressionData->Parameters.VectorChannelNameValues[Index];
-	return true;
+	return false;
 }
 
 #endif // WITH_EDITOR
@@ -2393,14 +2429,18 @@ bool UMaterial::GetTextureParameterValue(const FHashedMaterialParameterInfo& Par
 
 bool UMaterial::GetTextureParameterValue_New(const FHashedMaterialParameterInfo& ParameterInfo, UTexture*& OutValue, bool bOveriddenOnly) const
 {
-	const int32 Index = CachedExpressionData->Parameters.FindParameterIndex(EMaterialParameterType::Texture, ParameterInfo, bOveriddenOnly);
-	if (Index == INDEX_NONE)
+	if (CachedExpressionData)
 	{
-		OutValue = nullptr;
-		return false;
+		const int32 Index = CachedExpressionData->Parameters.FindParameterIndex(EMaterialParameterType::Texture, ParameterInfo, bOveriddenOnly);
+		if (Index != INDEX_NONE)
+		{
+			OutValue = CachedExpressionData->Parameters.TextureValues[Index];
+			return true;
+		}
 	}
-	OutValue = CachedExpressionData->Parameters.TextureValues[Index];
-	return true;
+
+	OutValue = nullptr;
+	return false;
 }
 
 #if WITH_EDITOR
@@ -2498,14 +2538,18 @@ bool UMaterial::GetRuntimeVirtualTextureParameterValue(const FHashedMaterialPara
 
 bool UMaterial::GetRuntimeVirtualTextureParameterValue_New(const FHashedMaterialParameterInfo& ParameterInfo, URuntimeVirtualTexture*& OutValue, bool bOveriddenOnly) const
 {
-	const int32 Index = CachedExpressionData->Parameters.FindParameterIndex(EMaterialParameterType::RuntimeVirtualTexture, ParameterInfo, bOveriddenOnly);
-	if (Index == INDEX_NONE)
+	if (CachedExpressionData)
 	{
-		OutValue = nullptr;
-		return false;
+		const int32 Index = CachedExpressionData->Parameters.FindParameterIndex(EMaterialParameterType::RuntimeVirtualTexture, ParameterInfo, bOveriddenOnly);
+		if (Index != INDEX_NONE)
+		{
+			OutValue = CachedExpressionData->Parameters.RuntimeVirtualTextureValues[Index];
+			return true;
+		}
 	}
-	OutValue = CachedExpressionData->Parameters.RuntimeVirtualTextureValues[Index];
-	return true;
+
+	OutValue = nullptr;
+	return false;
 }
 
 #if WITH_EDITOR
@@ -2585,13 +2629,16 @@ bool UMaterial::GetRuntimeVirtualTextureParameterValue_Legacy(const FHashedMater
 #if WITH_EDITOR
 bool UMaterial::GetTextureParameterChannelNames(const FHashedMaterialParameterInfo& ParameterInfo, FParameterChannelNames& OutValue) const
 {
-	const int32 Index = CachedExpressionData->Parameters.FindParameterIndex(EMaterialParameterType::Texture, ParameterInfo);
-	if (Index == INDEX_NONE)
+	if (CachedExpressionData)
 	{
-		return false;
+		const int32 Index = CachedExpressionData->Parameters.FindParameterIndex(EMaterialParameterType::Texture, ParameterInfo);
+		if (Index != INDEX_NONE)
+		{
+			OutValue = CachedExpressionData->Parameters.TextureChannelNameValues[Index];
+			return true;
+		}
 	}
-	OutValue = CachedExpressionData->Parameters.TextureChannelNameValues[Index];
-	return true;
+	return false;
 }
 #endif // WITH_EDITOR
 
@@ -2617,16 +2664,20 @@ bool UMaterial::GetFontParameterValue(const FHashedMaterialParameterInfo& Parame
 
 bool UMaterial::GetFontParameterValue_New(const FHashedMaterialParameterInfo& ParameterInfo, UFont*& OutFontValue, int32& OutFontPage, bool bOveriddenOnly) const
 {
-	const int32 Index = CachedExpressionData->Parameters.FindParameterIndex(EMaterialParameterType::Font, ParameterInfo, bOveriddenOnly);
-	if (Index == INDEX_NONE)
+	if (CachedExpressionData)
 	{
-		OutFontValue = nullptr;
-		OutFontPage = INDEX_NONE;
-		return false;
+		const int32 Index = CachedExpressionData->Parameters.FindParameterIndex(EMaterialParameterType::Font, ParameterInfo, bOveriddenOnly);
+		if (Index != INDEX_NONE)
+		{
+			OutFontValue = CachedExpressionData->Parameters.FontValues[Index];
+			OutFontPage = CachedExpressionData->Parameters.FontPageValues[Index];
+			return true;
+		}
 	}
-	OutFontValue = CachedExpressionData->Parameters.FontValues[Index];
-	OutFontPage = CachedExpressionData->Parameters.FontPageValues[Index];
-	return true;
+
+	OutFontValue = nullptr;
+	OutFontPage = INDEX_NONE;
+	return false;
 }
 
 #if WITH_EDITOR
@@ -3499,7 +3550,7 @@ void UMaterial::BackwardsCompatibilityVirtualTextureOutputConversion()
 
 void UMaterial::GetQualityLevelUsage(TArray<bool, TInlineAllocator<EMaterialQualityLevel::Num> >& OutQualityLevelsUsed, EShaderPlatform ShaderPlatform, bool bCooking)
 {
-	OutQualityLevelsUsed = CachedExpressionData->QualityLevelsUsed;
+	OutQualityLevelsUsed = GetCachedExpressionData().QualityLevelsUsed;
 	if (OutQualityLevelsUsed.Num() == 0)
 	{
 		OutQualityLevelsUsed.AddDefaulted(EMaterialQualityLevel::Num);
@@ -3712,6 +3763,8 @@ void UMaterial::PostLoad()
 	UpdateCachedExpressionData();
 #endif // WITH_EDITOR
 
+	checkf(CachedExpressionData, TEXT("Missing cached expression data for material, should have been either serialized or created during PostLoad"));
+
 	for (int32 CollectionIndex = 0; CollectionIndex < CachedExpressionData->ParameterCollectionInfos.Num(); CollectionIndex++)
 	{
 		if (CachedExpressionData->ParameterCollectionInfos[CollectionIndex].ParameterCollection)
@@ -4994,11 +5047,14 @@ bool UMaterial::UpdateLightmassTextureTracking()
 int32 UMaterial::GetLayerParameterIndex(EMaterialParameterAssociation Association, UMaterialFunctionInterface* LayerFunction) const
 {
 	int32 Index = INDEX_NONE;
-	switch (Association)
+	if (CachedExpressionData)
 	{
-	case BlendParameter: Index = CachedExpressionData->DefaultLayerBlends.Find(LayerFunction); break;
-	case LayerParameter: Index = CachedExpressionData->DefaultLayers.Find(LayerFunction); break;
-	default: checkNoEntry(); break;
+		switch (Association)
+		{
+		case BlendParameter: Index = CachedExpressionData->DefaultLayerBlends.Find(LayerFunction); break;
+		case LayerParameter: Index = CachedExpressionData->DefaultLayers.Find(LayerFunction); break;
+		default: checkNoEntry(); break;
+		}
 	}
 	return Index;
 }
@@ -5007,11 +5063,14 @@ int32 UMaterial::GetLayerParameterIndex(EMaterialParameterAssociation Associatio
 
 void UMaterial::GetReferencedTexturesAndOverrides(TSet<const UTexture*>& InOutTextures) const
 {
-	for (UObject* UsedObject : CachedExpressionData->ReferencedTextures)
+	if (CachedExpressionData)
 	{
-		if (const UTexture* UsedTexture = Cast<UTexture>(UsedObject))
+		for (UObject* UsedObject : CachedExpressionData->ReferencedTextures)
 		{
-			InOutTextures.Add(UsedTexture);
+			if (const UTexture* UsedTexture = Cast<UTexture>(UsedObject))
+			{
+				InOutTextures.Add(UsedTexture);
+			}
 		}
 	}
 }
@@ -5559,17 +5618,23 @@ void UMaterial::RecursiveUpdateRealtimePreview( UMaterialExpression* InExpressio
 
 void UMaterial::AppendReferencedFunctionIdsTo(TArray<FGuid>& Ids) const
 {
-	for (int32 FunctionIndex = 0; FunctionIndex < CachedExpressionData->FunctionInfos.Num(); FunctionIndex++)
+	if (CachedExpressionData)
 	{
-		Ids.AddUnique(CachedExpressionData->FunctionInfos[FunctionIndex].StateId);
+		for (int32 FunctionIndex = 0; FunctionIndex < CachedExpressionData->FunctionInfos.Num(); FunctionIndex++)
+		{
+			Ids.AddUnique(CachedExpressionData->FunctionInfos[FunctionIndex].StateId);
+		}
 	}
 }
 
 void UMaterial::AppendReferencedParameterCollectionIdsTo(TArray<FGuid>& Ids) const
 {
-	for (int32 CollectionIndex = 0; CollectionIndex < CachedExpressionData->ParameterCollectionInfos.Num(); CollectionIndex++)
+	if (CachedExpressionData)
 	{
-		Ids.AddUnique(CachedExpressionData->ParameterCollectionInfos[CollectionIndex].StateId);
+		for (int32 CollectionIndex = 0; CollectionIndex < CachedExpressionData->ParameterCollectionInfos.Num(); CollectionIndex++)
+		{
+			Ids.AddUnique(CachedExpressionData->ParameterCollectionInfos[CollectionIndex].StateId);
+		}
 	}
 }
 
diff --git a/Engine/Source/Runtime/Engine/Private/Materials/MaterialCachedData.cpp b/Engine/Source/Runtime/Engine/Private/Materials/MaterialCachedData.cpp
--- a/Engine/Source/Runtime/Engine/Private/Materials/MaterialCachedData.cpp
+++ b/Engine/Source/Runtime/Engine/Private/Materials/MaterialCachedData.cpp
@@ -40,6 +40,8 @@
 #include "Curves/CurveLinearColor.h"
 #include "Curves/CurveLinearColorAtlas.h"
 
+const FMaterialCachedExpressionData FMaterialCachedExpressionData::EmptyData;
+
 void FMaterialCachedExpressionData::Reset()
 {
 	Parameters.Reset();
diff --git a/Engine/Source/Runtime/Engine/Private/Materials/MaterialInstance.cpp b/Engine/Source/Runtime/Engine/Private/Materials/MaterialInstance.cpp
--- a/Engine/Source/Runtime/Engine/Private/Materials/MaterialInstance.cpp
+++ b/Engine/Source/Runtime/Engine/Private/Materials/MaterialInstance.cpp
@@ -28,6 +28,7 @@
 #include "Materials/MaterialInstanceSupport.h"
 #include "Engine/SubsurfaceProfile.h"
 #include "ProfilingDebugging/LoadTimeTracker.h"
+#include "ProfilingDebugging/CookStats.h"
 #include "Interfaces/ITargetPlatform.h"
 #include "Interfaces/ITargetPlatformManagerModule.h"
 #include "Components.h"
@@ -42,6 +43,21 @@
 #include "ShaderPlatformQualitySettings.h"
 #include "MaterialShaderQualitySettings.h"
 
+#if ENABLE_COOK_STATS
+#include "ProfilingDebugging/ScopedTimers.h"
+namespace MaterialInstanceCookStats
+{
+	static double UpdateCachedExpressionDataSec = 0.0;
+
+	static FCookStatsManager::FAutoRegisterCallback RegisterCookStats([](FCookStatsManager::AddStatFuncRef AddStat)
+		{
+			AddStat(TEXT("MaterialInstance"), FCookStatsManager::CreateKeyValueArray(
+				TEXT("UpdateCachedExpressionDataSec"), UpdateCachedExpressionDataSec
+			));
+		});
+}
+#endif
+
 DECLARE_CYCLE_STAT(TEXT("MaterialInstance CopyMatInstParams"), STAT_MaterialInstance_CopyMatInstParams, STATGROUP_Shaders);
 DECLARE_CYCLE_STAT(TEXT("MaterialInstance Serialize"), STAT_MaterialInstance_Serialize, STATGROUP_Shaders);
 DECLARE_CYCLE_STAT(TEXT("MaterialInstance CopyUniformParamsInternal"), STAT_MaterialInstance_CopyUniformParamsInternal, STATGROUP_Shaders);
@@ -747,7 +763,7 @@ bool UMaterialInstance::GetScalarParameterSliderMinMax(const FHashedMaterialPara
 		return false;
 	}
 
-	if (ParameterInfo.Association != EMaterialParameterAssociation::GlobalParameter)
+	if (ParameterInfo.Association != EMaterialParameterAssociation::GlobalParameter && CachedData)
 	{
 		const int32 ParameterIndex = CachedData->Parameters.FindParameterIndex(EMaterialParameterType::Scalar, ParameterInfo);
 		if (ParameterIndex != INDEX_NONE)
@@ -786,7 +802,7 @@ bool UMaterialInstance::GetScalarParameterValue(const FHashedMaterialParameterIn
 	}
 	
 	// Instance-included default
-	if (ParameterInfo.Association != EMaterialParameterAssociation::GlobalParameter)
+	if (ParameterInfo.Association != EMaterialParameterAssociation::GlobalParameter && CachedData)
 	{
 		const int32 ParameterIndex = CachedData->Parameters.FindParameterIndex(EMaterialParameterType::Scalar, ParameterInfo);
 		if (ParameterIndex != INDEX_NONE)
@@ -830,7 +846,7 @@ bool UMaterialInstance::IsScalarParameterUsedAsAtlasPosition(const FHashedMateri
 #endif
 
 	// Instance-included default
-	if (ParameterInfo.Association != EMaterialParameterAssociation::GlobalParameter)
+	if (ParameterInfo.Association != EMaterialParameterAssociation::GlobalParameter && CachedData)
 	{
 		const int32 ParameterIndex = CachedData->Parameters.FindParameterIndex(EMaterialParameterType::Scalar, ParameterInfo);
 		if (ParameterIndex != INDEX_NONE)
@@ -880,7 +896,7 @@ bool UMaterialInstance::GetVectorParameterValue(const FHashedMaterialParameterIn
 	}
 	
 	// Instance-included default
-	if (ParameterInfo.Association != EMaterialParameterAssociation::GlobalParameter)
+	if (ParameterInfo.Association != EMaterialParameterAssociation::GlobalParameter && CachedData)
 	{
 		const int32 ParameterIndex = CachedData->Parameters.FindParameterIndex(EMaterialParameterType::Vector, ParameterInfo);
 		if (ParameterIndex != INDEX_NONE)
@@ -912,7 +928,7 @@ bool UMaterialInstance::IsVectorParameterUsedAsChannelMask(const FHashedMaterial
 	}
 	
 	// Instance-included default
-	if (ParameterInfo.Association != EMaterialParameterAssociation::GlobalParameter)
+	if (ParameterInfo.Association != EMaterialParameterAssociation::GlobalParameter && CachedData)
 	{
 		const int32 ParameterIndex = CachedData->Parameters.FindParameterIndex(EMaterialParameterType::Vector, ParameterInfo);
 		if (ParameterIndex != INDEX_NONE)
@@ -942,7 +958,7 @@ bool  UMaterialInstance::GetVectorParameterChannelNames(const FHashedMaterialPar
 	}
 
 	// Instance-included default
-	if (ParameterInfo.Association != EMaterialParameterAssociation::GlobalParameter)
+	if (ParameterInfo.Association != EMaterialParameterAssociation::GlobalParameter && CachedData)
 	{
 		const int32 ParameterIndex = CachedData->Parameters.FindParameterIndex(EMaterialParameterType::Vector, ParameterInfo);
 		if (ParameterIndex != INDEX_NONE)
@@ -982,7 +998,7 @@ bool UMaterialInstance::GetTextureParameterValue(const FHashedMaterialParameterI
 	}
 	
 	// Instance-included default
-	if (ParameterInfo.Association != EMaterialParameterAssociation::GlobalParameter)
+	if (ParameterInfo.Association != EMaterialParameterAssociation::GlobalParameter && CachedData)
 	{
 		const int32 ParameterIndex = CachedData->Parameters.FindParameterIndex(EMaterialParameterType::Texture, ParameterInfo);
 		if (ParameterIndex != INDEX_NONE)
@@ -1023,7 +1039,7 @@ bool UMaterialInstance::GetRuntimeVirtualTextureParameterValue(const FHashedMate
 	}
 
 	// Instance-included default
-	if (ParameterInfo.Association != EMaterialParameterAssociation::GlobalParameter)
+	if (ParameterInfo.Association != EMaterialParameterAssociation::GlobalParameter && CachedData)
 	{
 		const int32 ParameterIndex = CachedData->Parameters.FindParameterIndex(EMaterialParameterType::RuntimeVirtualTexture, ParameterInfo);
 		if (ParameterIndex != INDEX_NONE)
@@ -1057,7 +1073,7 @@ bool  UMaterialInstance::GetTextureParameterChannelNames(const FHashedMaterialPa
 	}
 
 	// Instance-included default
-	if (ParameterInfo.Association != EMaterialParameterAssociation::GlobalParameter)
+	if (ParameterInfo.Association != EMaterialParameterAssociation::GlobalParameter && CachedData)
 	{
 		const int32 ParameterIndex = CachedData->Parameters.FindParameterIndex(EMaterialParameterType::Texture, ParameterInfo);
 		if (ParameterIndex != INDEX_NONE)
@@ -1099,7 +1115,7 @@ bool UMaterialInstance::GetFontParameterValue(const FHashedMaterialParameterInfo
 	}
 	
 	// Instance-included default
-	if (ParameterInfo.Association != EMaterialParameterAssociation::GlobalParameter)
+	if (ParameterInfo.Association != EMaterialParameterAssociation::GlobalParameter && CachedData)
 	{
 		const int32 ParameterIndex = CachedData->Parameters.FindParameterIndex(EMaterialParameterType::Font, ParameterInfo);
 		if (ParameterIndex != INDEX_NONE)
@@ -2150,7 +2166,11 @@ void UMaterialInstance::GetStaticParameterValues(FStaticParameterSet& OutStaticP
 void UMaterialInstance::GetAllParametersOfType(EMaterialParameterType Type, TArray<FMaterialParameterInfo>& OutParameterInfo, TArray<FGuid>& OutParameterIds) const
 {
 	const UMaterial* Material = GetMaterial_Concurrent();
-	int32 NumParameters = CachedData->Parameters.GetNumParameters(Type);
+	int32 NumParameters = 0;
+	if (CachedData)
+	{
+		NumParameters += CachedData->Parameters.GetNumParameters(Type);
+	}
 	if (Material)
 	{
 		NumParameters += Material->GetCachedExpressionData().Parameters.GetNumParameters(Type);
@@ -2158,7 +2178,10 @@ void UMaterialInstance::GetAllParametersOfType(EMaterialParameterType Type, TArr
 
 	OutParameterInfo.Empty(NumParameters);
 	OutParameterIds.Empty(NumParameters);
-	CachedData->Parameters.GetAllParameterInfoOfType(Type, false, OutParameterInfo, OutParameterIds);
+	if (CachedData)
+	{
+		CachedData->Parameters.GetAllParameterInfoOfType(Type, false, OutParameterInfo, OutParameterIds);
+	}
 	if (Material)
 	{
 		Material->GetCachedExpressionData().Parameters.GetAllGlobalParameterInfoOfType(Type, false, OutParameterInfo, OutParameterIds);
@@ -2298,7 +2321,7 @@ bool UMaterialInstance::GetScalarParameterDefaultValue(const FHashedMaterialPara
 
 	// In the case of duplicate parameters with different values, this will return the
 	// first matching expression found, not necessarily the one that's used for rendering
-	if (ParameterInfo.Association != EMaterialParameterAssociation::GlobalParameter)
+	if (ParameterInfo.Association != EMaterialParameterAssociation::GlobalParameter && CachedData)
 	{
 		const int32 ParameterIndex = CachedData->Parameters.FindParameterIndex(EMaterialParameterType::Scalar, ParameterInfo);
 		if (ParameterIndex != INDEX_NONE)
@@ -2343,7 +2366,7 @@ bool UMaterialInstance::GetVectorParameterDefaultValue(const FHashedMaterialPara
 		}
 	}
 
-	if (ParameterInfo.Association != EMaterialParameterAssociation::GlobalParameter)
+	if (ParameterInfo.Association != EMaterialParameterAssociation::GlobalParameter && CachedData)
 	{
 		const int32 ParameterIndex = CachedData->Parameters.FindParameterIndex(EMaterialParameterType::Vector, ParameterInfo);
 		if (ParameterIndex != INDEX_NONE)
@@ -2388,7 +2411,7 @@ bool UMaterialInstance::GetTextureParameterDefaultValue(const FHashedMaterialPar
 		}
 	}
 
-	if (ParameterInfo.Association != EMaterialParameterAssociation::GlobalParameter)
+	if (ParameterInfo.Association != EMaterialParameterAssociation::GlobalParameter && CachedData)
 	{
 		const int32 ParameterIndex = CachedData->Parameters.FindParameterIndex(EMaterialParameterType::Texture, ParameterInfo);
 		if (ParameterIndex != INDEX_NONE)
@@ -2433,7 +2456,7 @@ bool UMaterialInstance::GetRuntimeVirtualTextureParameterDefaultValue(const FHas
 		}
 	}
 
-	if (ParameterInfo.Association != EMaterialParameterAssociation::GlobalParameter)
+	if (ParameterInfo.Association != EMaterialParameterAssociation::GlobalParameter && CachedData)
 	{
 		const int32 ParameterIndex = CachedData->Parameters.FindParameterIndex(EMaterialParameterType::RuntimeVirtualTexture, ParameterInfo);
 		if (ParameterIndex != INDEX_NONE)
@@ -2479,7 +2502,7 @@ bool UMaterialInstance::GetFontParameterDefaultValue(const FHashedMaterialParame
 		}
 	}
 
-	if (ParameterInfo.Association != EMaterialParameterAssociation::GlobalParameter)
+	if (ParameterInfo.Association != EMaterialParameterAssociation::GlobalParameter && CachedData)
 	{
 		const int32 ParameterIndex = CachedData->Parameters.FindParameterIndex(EMaterialParameterType::Font, ParameterInfo);
 		if (ParameterIndex != INDEX_NONE)
@@ -4041,11 +4064,14 @@ void UMaterialInstance::UpdateStaticPermutation(const FStaticParameterSet& NewPa
 
 void UMaterialInstance::GetReferencedTexturesAndOverrides(TSet<const UTexture*>& InOutTextures) const
 {
-	for (UObject* UsedObject : CachedData->ReferencedTextures)
+	if (CachedData)
 	{
-		if (const UTexture* UsedTexture = Cast<UTexture>(UsedObject))
+		for (UObject* UsedObject : CachedData->ReferencedTextures)
 		{
-			InOutTextures.Add(UsedTexture);
+			if (const UTexture* UsedTexture = Cast<UTexture>(UsedObject))
+			{
+				InOutTextures.Add(UsedTexture);
+			}
 		}
 	}
 
@@ -4069,6 +4095,8 @@ extern void FMaterialCachedParameters_UpdateForLayerParameters(FMaterialCachedPa
 
 void UMaterialInstance::UpdateCachedLayerParameters()
 {
+	COOK_STAT(FScopedDurationTimer BlockingTimer(MaterialInstanceCookStats::UpdateCachedExpressionDataSec));
+
 	if (bSavedCachedData)
 	{
 		// Don't need to rebuild cached data if it was serialized
diff --git a/Engine/Source/Runtime/Engine/Public/MaterialCachedData.h b/Engine/Source/Runtime/Engine/Public/MaterialCachedData.h
--- a/Engine/Source/Runtime/Engine/Public/MaterialCachedData.h
+++ b/Engine/Source/Runtime/Engine/Public/MaterialCachedData.h
@@ -220,6 +220,8 @@ struct FMaterialCachedExpressionData
 {
 	GENERATED_USTRUCT_BODY()
 
+	ENGINE_API static const FMaterialCachedExpressionData EmptyData;
+
 #if WITH_EDITOR
 	/** Returns 'false' if update is incomplete, due to missing expression data (stripped from non-editor build) */
 	bool UpdateForExpressions(const FMaterialCachedExpressionContext& Context, const TArray<UMaterialExpression*>& Expressions, EMaterialParameterAssociation Association, int32 ParameterIndex);
