Control Rig Anim Node: Add debug drawing functionality in PIE/AnimBP_Editor for Control Rig Anim Node. Draw commands are queued during evaluate/execute. These commands can have multiple consumers so I also moved the clearing of these commands to the beginning of the execute.

Validated that debug drawing works in the following cases:
Playi n PIE
Play in AnimBP
Control Rig Editor
Control Rig Editor Debug Other Instances
Sequencer

Does not work yet in the following cases:
Pause in PIE
Pause in AnimBP

diff --git a/Engine/Plugins/Experimental/ControlRig/Source/ControlRig/Private/AnimNode_ControlRigBase.cpp b/Engine/Plugins/Experimental/ControlRig/Source/ControlRig/Private/AnimNode_ControlRigBase.cpp
--- a/Engine/Plugins/Experimental/ControlRig/Source/ControlRig/Private/AnimNode_ControlRigBase.cpp
+++ b/Engine/Plugins/Experimental/ControlRig/Source/ControlRig/Private/AnimNode_ControlRigBase.cpp
@@ -8,6 +8,10 @@
 #include "AnimationRuntime.h"
 #include "Units/Execution/RigUnit_BeginExecution.h"
 
+#if ENABLE_ANIM_DEBUG
+TAutoConsoleVariable<int32> CVarAnimNodeControlRigDebug(TEXT("a.AnimNode.ControlRig.Debug"), 0, TEXT("Set to 1 to turn on debug drawing for AnimNode_ControlRigBase"));
+#endif
+
 FAnimNode_ControlRigBase::FAnimNode_ControlRigBase()
 	: FAnimNode_CustomProperty()
 	, InputSettings(FControlRigIOSettings())
@@ -233,6 +237,16 @@ void FAnimNode_ControlRigBase::ExecuteControlRig(FPoseContext& InOutput)
 		{
 			// first evaluate control rig
 			ControlRig->Evaluate_AnyThread();
+
+#if ENABLE_ANIM_DEBUG 
+			// When Control Rig is at editing time (in CR editor), draw instructions are consumed by ControlRigEditMode, so we need to skip drawing here.
+			bool bShowDebug = (CVarAnimNodeControlRigDebug.GetValueOnAnyThread() == 1 && ControlRig->ExecutionType != ERigExecutionType::Editing);
+
+			if (bShowDebug)
+			{ 
+				QueueControlRigDrawInstructions(ControlRig, InOutput.AnimInstanceProxy);
+			}
+#endif
 		}
 
 		// now update output
@@ -417,3 +431,56 @@ UClass* FAnimNode_ControlRigBase::GetTargetClass() const
 	return nullptr;
 }
 
+void FAnimNode_ControlRigBase::QueueControlRigDrawInstructions(UControlRig* ControlRig, FAnimInstanceProxy* Proxy) const
+{
+	ensure(ControlRig);
+	ensure(Proxy);
+
+	for (const FControlRigDrawInstruction& Instruction : ControlRig->GetDrawInterface())
+	{
+		if (!Instruction.IsValid())
+		{
+			continue;
+		}
+
+		FTransform InstructionTransform = Instruction.Transform * Proxy->GetComponentTransform();
+		switch (Instruction.PrimitiveType)
+		{
+			case EControlRigDrawSettings::Points:
+			{
+				for (const FVector& Point : Instruction.Positions)
+				{
+					Proxy->AnimDrawDebugPoint(InstructionTransform.TransformPosition(Point), Instruction.Thickness, Instruction.Color.ToFColor(true), false, -1.f, SDPG_Foreground);
+				}
+				break;
+			}
+			case EControlRigDrawSettings::Lines:
+			{
+				const TArray<FVector>& Points = Instruction.Positions;
+
+				for (int32 PointIndex = 0; PointIndex < Points.Num() - 1; PointIndex += 2)
+				{
+					Proxy->AnimDrawDebugLine(InstructionTransform.TransformPosition(Points[PointIndex]), InstructionTransform.TransformPosition(Points[PointIndex + 1]), Instruction.Color.ToFColor(true), false, -1.f, Instruction.Thickness, SDPG_Foreground);
+				}
+				break;
+			}
+			case EControlRigDrawSettings::LineStrip:
+			{
+				const TArray<FVector>& Points = Instruction.Positions;
+
+				for (int32 PointIndex = 0; PointIndex < Points.Num() - 1; PointIndex++)
+				{
+					Proxy->AnimDrawDebugLine(InstructionTransform.TransformPosition(Points[PointIndex]), InstructionTransform.TransformPosition(Points[PointIndex + 1]), Instruction.Color.ToFColor(true), false, -1.f, Instruction.Thickness, SDPG_Foreground);
+				}
+				break;
+			}
+
+			case EControlRigDrawSettings::DynamicMesh:
+			{
+				// TODO: Add support for this if anyone is actually using it. Currently it is only defined and referenced in an unused API, DrawCone in Control Rig.
+				break;
+			} 
+		}
+	} 
+}
+
diff --git a/Engine/Plugins/Experimental/ControlRig/Source/ControlRig/Private/ControlRig.cpp b/Engine/Plugins/Experimental/ControlRig/Source/ControlRig/Private/ControlRig.cpp
--- a/Engine/Plugins/Experimental/ControlRig/Source/ControlRig/Private/ControlRig.cpp
+++ b/Engine/Plugins/Experimental/ControlRig/Source/ControlRig/Private/ControlRig.cpp
@@ -460,6 +460,7 @@ void UControlRig::Execute(const EControlRigState InState, const FName& InEventNa
 	}
 
 	FRigUnitContext Context;
+	DrawInterface.Reset();
 	Context.DrawInterface = &DrawInterface;
 	Context.DrawContainer = &DrawContainer;
 
diff --git a/Engine/Plugins/Experimental/ControlRig/Source/ControlRig/Public/AnimNode_ControlRigBase.h b/Engine/Plugins/Experimental/ControlRig/Source/ControlRig/Public/AnimNode_ControlRigBase.h
--- a/Engine/Plugins/Experimental/ControlRig/Source/ControlRig/Public/AnimNode_ControlRigBase.h
+++ b/Engine/Plugins/Experimental/ControlRig/Source/ControlRig/Public/AnimNode_ControlRigBase.h
@@ -100,6 +100,8 @@ protected:
 	// execute control rig on the input pose and outputs the result
 	void ExecuteControlRig(FPoseContext& InOutput);
 
+	void QueueControlRigDrawInstructions(UControlRig* ControlRig, FAnimInstanceProxy* Proxy) const;
+
 	friend struct FControlRigSequencerAnimInstanceProxy;
 	friend struct FControlRigLayerInstanceProxy;
 };
diff --git a/Engine/Plugins/Experimental/ControlRig/Source/ControlRig/Public/ControlRig.h b/Engine/Plugins/Experimental/ControlRig/Source/ControlRig/Public/ControlRig.h
--- a/Engine/Plugins/Experimental/ControlRig/Source/ControlRig/Public/ControlRig.h
+++ b/Engine/Plugins/Experimental/ControlRig/Source/ControlRig/Public/ControlRig.h
@@ -401,6 +401,9 @@ public:
 	FControlRigExecuteEvent& OnExecuted_AnyThread() { return ExecutedEvent; }
 	FRigEventDelegate& OnRigEvent_AnyThread() { return RigEventDelegate; }
 
+	const FControlRigDrawInterface& GetDrawInterface() const { return DrawInterface; };
+	FControlRigDrawInterface& GetDrawInterface() { return DrawInterface; };
+
 private:
 
 	UPROPERTY()
diff --git a/Engine/Plugins/Experimental/ControlRig/Source/ControlRigEditor/Private/EditMode/ControlRigEditMode.cpp b/Engine/Plugins/Experimental/ControlRig/Source/ControlRigEditor/Private/EditMode/ControlRigEditMode.cpp
--- a/Engine/Plugins/Experimental/ControlRig/Source/ControlRigEditor/Private/EditMode/ControlRigEditMode.cpp
+++ b/Engine/Plugins/Experimental/ControlRig/Source/ControlRigEditor/Private/EditMode/ControlRigEditMode.cpp
@@ -615,8 +615,6 @@ void FControlRigEditMode::Render(const FSceneView* View, FViewport* Viewport, FP
 			}
 		}
 	}
-
-	GetControlRig(true)->DrawInterface.Reset();
 }
 
 bool FControlRigEditMode::InputKey(FEditorViewportClient* InViewportClient, FViewport* InViewport, FKey InKey, EInputEvent InEvent)
diff --git a/Engine/Source/Runtime/Engine/Private/Animation/AnimInstanceProxy.cpp b/Engine/Source/Runtime/Engine/Private/Animation/AnimInstanceProxy.cpp
--- a/Engine/Source/Runtime/Engine/Private/Animation/AnimInstanceProxy.cpp
+++ b/Engine/Source/Runtime/Engine/Private/Animation/AnimInstanceProxy.cpp
@@ -509,11 +509,12 @@ void FAnimInstanceProxy::PostUpdate(UAnimInstance* InAnimInstance) const
 	{
 		switch (DebugItem.ItemType)
 		{
-		case EDrawDebugItemType::OnScreenMessage: GEngine->AddOnScreenDebugMessage(INDEX_NONE, 0.f, DebugItem.Color, DebugItem.Message, false, DebugItem.TextScale); break;
-		case EDrawDebugItemType::DirectionalArrow: DrawDebugDirectionalArrow(InAnimInstance->GetSkelMeshComponent()->GetWorld(), DebugItem.StartLoc, DebugItem.EndLoc, DebugItem.Size, DebugItem.Color, DebugItem.bPersistentLines, DebugItem.LifeTime, 0, DebugItem.Thickness); break;
-		case EDrawDebugItemType::Sphere : DrawDebugSphere(InAnimInstance->GetSkelMeshComponent()->GetWorld(), DebugItem.Center, DebugItem.Radius, DebugItem.Segments, DebugItem.Color, DebugItem.bPersistentLines, DebugItem.LifeTime, 0, DebugItem.Thickness); break;
-		case EDrawDebugItemType::Line: DrawDebugLine(InAnimInstance->GetSkelMeshComponent()->GetWorld(), DebugItem.StartLoc, DebugItem.EndLoc, DebugItem.Color, DebugItem.bPersistentLines, DebugItem.LifeTime, 0, DebugItem.Thickness); break;
-		case EDrawDebugItemType::CoordinateSystem : DrawDebugCoordinateSystem(InAnimInstance->GetSkelMeshComponent()->GetWorld(), DebugItem.StartLoc, DebugItem.Rotation, DebugItem.Size, DebugItem.bPersistentLines, DebugItem.LifeTime, 0, DebugItem.Thickness); break;
+			case EDrawDebugItemType::OnScreenMessage: GEngine->AddOnScreenDebugMessage(INDEX_NONE, 0.f, DebugItem.Color, DebugItem.Message, false, DebugItem.TextScale); break;
+			case EDrawDebugItemType::DirectionalArrow: DrawDebugDirectionalArrow(InAnimInstance->GetSkelMeshComponent()->GetWorld(), DebugItem.StartLoc, DebugItem.EndLoc, DebugItem.Size, DebugItem.Color, DebugItem.bPersistentLines, DebugItem.LifeTime, DebugItem.DepthPriority, DebugItem.Thickness); break;
+			case EDrawDebugItemType::Sphere: DrawDebugSphere(InAnimInstance->GetSkelMeshComponent()->GetWorld(), DebugItem.Center, DebugItem.Radius, DebugItem.Segments, DebugItem.Color, DebugItem.bPersistentLines, DebugItem.LifeTime, DebugItem.DepthPriority, DebugItem.Thickness); break;
+			case EDrawDebugItemType::Line: DrawDebugLine(InAnimInstance->GetSkelMeshComponent()->GetWorld(), DebugItem.StartLoc, DebugItem.EndLoc, DebugItem.Color, DebugItem.bPersistentLines, DebugItem.LifeTime, DebugItem.DepthPriority, DebugItem.Thickness); break;
+			case EDrawDebugItemType::CoordinateSystem: DrawDebugCoordinateSystem(InAnimInstance->GetSkelMeshComponent()->GetWorld(), DebugItem.StartLoc, DebugItem.Rotation, DebugItem.Size, DebugItem.bPersistentLines, DebugItem.LifeTime, DebugItem.DepthPriority, DebugItem.Thickness); break;
+			case EDrawDebugItemType::Point: DrawDebugPoint(InAnimInstance->GetSkelMeshComponent()->GetWorld(), DebugItem.StartLoc, DebugItem.Size, DebugItem.Color, DebugItem.bPersistentLines, DebugItem.LifeTime, DebugItem.DepthPriority); break;
 		}
 	}
 #endif
@@ -1787,7 +1788,7 @@ void FAnimInstanceProxy::GatherDebugData_WithRoot(FNodeDebugData& DebugData, FAn
 
 #if ENABLE_ANIM_DRAW_DEBUG
 
-void FAnimInstanceProxy::AnimDrawDebugOnScreenMessage(const FString& DebugMessage, const FColor& Color, const FVector2D& TextScale)
+void FAnimInstanceProxy::AnimDrawDebugOnScreenMessage(const FString& DebugMessage, const FColor& Color, const FVector2D& TextScale, ESceneDepthPriorityGroup DepthPriority)
 {
 	FQueuedDrawDebugItem DrawDebugItem;
 
@@ -1795,11 +1796,12 @@ void FAnimInstanceProxy::AnimDrawDebugOnScreenMessage(const FString& DebugMessag
 	DrawDebugItem.Message = DebugMessage;
 	DrawDebugItem.Color = Color;
 	DrawDebugItem.TextScale = TextScale;
+	DrawDebugItem.DepthPriority = DepthPriority;
 
 	QueuedDrawDebugItems.Add(DrawDebugItem);
 }
 
-void FAnimInstanceProxy::AnimDrawDebugDirectionalArrow(const FVector& LineStart, const FVector& LineEnd, float ArrowSize, const FColor& Color, bool bPersistentLines, float LifeTime, float Thickness)
+void FAnimInstanceProxy::AnimDrawDebugDirectionalArrow(const FVector& LineStart, const FVector& LineEnd, float ArrowSize, const FColor& Color, bool bPersistentLines, float LifeTime, float Thickness, ESceneDepthPriorityGroup DepthPriority)
 {
 	FQueuedDrawDebugItem DrawDebugItem;
 
@@ -1811,11 +1813,12 @@ void FAnimInstanceProxy::AnimDrawDebugDirectionalArrow(const FVector& LineStart,
 	DrawDebugItem.bPersistentLines = bPersistentLines;
 	DrawDebugItem.LifeTime = LifeTime;
 	DrawDebugItem.Thickness = Thickness;
+	DrawDebugItem.DepthPriority = DepthPriority;
 
 	QueuedDrawDebugItems.Add(DrawDebugItem);
 }
 
-void FAnimInstanceProxy::AnimDrawDebugSphere(const FVector& Center, float Radius, int32 Segments, const FColor& Color, bool bPersistentLines, float LifeTime, float Thickness)
+void FAnimInstanceProxy::AnimDrawDebugSphere(const FVector& Center, float Radius, int32 Segments, const FColor& Color, bool bPersistentLines, float LifeTime, float Thickness, ESceneDepthPriorityGroup DepthPriority)
 {
 	FQueuedDrawDebugItem DrawDebugItem;
 
@@ -1827,11 +1830,12 @@ void FAnimInstanceProxy::AnimDrawDebugSphere(const FVector& Center, float Radius
 	DrawDebugItem.bPersistentLines = bPersistentLines;
 	DrawDebugItem.LifeTime = LifeTime;
 	DrawDebugItem.Thickness = Thickness;
+	DrawDebugItem.DepthPriority = DepthPriority;
 
 	QueuedDrawDebugItems.Add(DrawDebugItem);
 }
 
-void FAnimInstanceProxy::AnimDrawDebugCoordinateSystem(FVector const& AxisLoc, FRotator const& AxisRot, float Scale, bool bPersistentLines, float LifeTime, float Thickness)
+void FAnimInstanceProxy::AnimDrawDebugCoordinateSystem(FVector const& AxisLoc, FRotator const& AxisRot, float Scale, bool bPersistentLines, float LifeTime, float Thickness, ESceneDepthPriorityGroup DepthPriority)
 {
 	FQueuedDrawDebugItem DrawDebugItem;
 
@@ -1842,11 +1846,12 @@ void FAnimInstanceProxy::AnimDrawDebugCoordinateSystem(FVector const& AxisLoc, F
 	DrawDebugItem.bPersistentLines = bPersistentLines;
 	DrawDebugItem.LifeTime = LifeTime;
 	DrawDebugItem.Thickness = Thickness;
+	DrawDebugItem.DepthPriority = DepthPriority;
 
 	QueuedDrawDebugItems.Add(DrawDebugItem);
 }
 
-void FAnimInstanceProxy::AnimDrawDebugLine(const FVector& StartLoc, const FVector& EndLoc, const FColor& Color, bool bPersistentLines, float LifeTime, float Thickness)
+void FAnimInstanceProxy::AnimDrawDebugLine(const FVector& StartLoc, const FVector& EndLoc, const FColor& Color, bool bPersistentLines, float LifeTime, float Thickness, ESceneDepthPriorityGroup DepthPriority)
 {
 	FQueuedDrawDebugItem DrawDebugItem;
 
@@ -1857,11 +1862,12 @@ void FAnimInstanceProxy::AnimDrawDebugLine(const FVector& StartLoc, const FVecto
 	DrawDebugItem.bPersistentLines = bPersistentLines;
 	DrawDebugItem.LifeTime = LifeTime;
 	DrawDebugItem.Thickness = Thickness;
+	DrawDebugItem.DepthPriority = DepthPriority;
 
 	QueuedDrawDebugItems.Add(DrawDebugItem);
 }
 
-void FAnimInstanceProxy::AnimDrawDebugPlane(const FTransform& BaseTransform, float Radii, const FColor& Color, bool bPersistentLines /*= false*/, float LifeTime /*= -1.f*/, float Thickness /*= 0.f*/)
+void FAnimInstanceProxy::AnimDrawDebugPlane(const FTransform& BaseTransform, float Radii, const FColor& Color, bool bPersistentLines /*= false*/, float LifeTime /*= -1.f*/, float Thickness /*= 0.f*/, ESceneDepthPriorityGroup DepthPriority)
 {
 	// just draw two triangle from [-Radii,-Radii] to [Radii, Radii]
 	FQueuedDrawDebugItem DrawDebugItem;
@@ -1871,6 +1877,7 @@ void FAnimInstanceProxy::AnimDrawDebugPlane(const FTransform& BaseTransform, flo
 	DrawDebugItem.bPersistentLines = bPersistentLines;
 	DrawDebugItem.LifeTime = LifeTime;
 	DrawDebugItem.Thickness = Thickness;
+	DrawDebugItem.DepthPriority = DepthPriority;
 
 	DrawDebugItem.StartLoc = BaseTransform.TransformPosition(FVector(-Radii, -Radii, 0));
 	DrawDebugItem.EndLoc = BaseTransform.TransformPosition(FVector(-Radii, Radii, 0));
@@ -1892,6 +1899,21 @@ void FAnimInstanceProxy::AnimDrawDebugPlane(const FTransform& BaseTransform, flo
 	DrawDebugItem.EndLoc = BaseTransform.TransformPosition(FVector(Radii, -Radii, 0));
 	QueuedDrawDebugItems.Add(DrawDebugItem);
 }
+
+void FAnimInstanceProxy::AnimDrawDebugPoint(const FVector& Loc, float Size, const FColor& Color, bool bPersistentLines, float LifeTime, ESceneDepthPriorityGroup DepthPriority)
+{
+	FQueuedDrawDebugItem DrawDebugItem;
+
+	DrawDebugItem.ItemType = EDrawDebugItemType::Point;
+	DrawDebugItem.StartLoc = Loc;
+	DrawDebugItem.Color = Color;
+	DrawDebugItem.Size = Size;
+	DrawDebugItem.bPersistentLines = bPersistentLines;
+	DrawDebugItem.LifeTime = LifeTime;
+	DrawDebugItem.DepthPriority = DepthPriority;
+
+	QueuedDrawDebugItems.Add(DrawDebugItem);
+}
 #endif // ENABLE_ANIM_DRAW_DEBUG
 
 float FAnimInstanceProxy::GetInstanceAssetPlayerLength(int32 AssetPlayerIndex)
diff --git a/Engine/Source/Runtime/Engine/Public/Animation/AnimInstanceProxy.h b/Engine/Source/Runtime/Engine/Public/Animation/AnimInstanceProxy.h
--- a/Engine/Source/Runtime/Engine/Public/Animation/AnimInstanceProxy.h
+++ b/Engine/Source/Runtime/Engine/Public/Animation/AnimInstanceProxy.h
@@ -51,6 +51,7 @@ namespace EDrawDebugItemType
 		Line,
 		OnScreenMessage,
 		CoordinateSystem,
+		Point,
 	};
 }
 
@@ -100,6 +101,9 @@ struct FQueuedDrawDebugItem
 
 	UPROPERTY(Transient)
 	FVector2D TextScale;
+
+	UPROPERTY(Transient)
+	TEnumAsByte<ESceneDepthPriorityGroup> DepthPriority;
 };
 
 /** Proxy object passed around during animation tree update in lieu of a UAnimInstance */
@@ -419,19 +423,21 @@ public:
 #if ENABLE_ANIM_DRAW_DEBUG
 	TArray<FQueuedDrawDebugItem> QueuedDrawDebugItems;
 
-	void AnimDrawDebugOnScreenMessage(const FString& DebugMessage, const FColor& Color, const FVector2D& TextScale = FVector2D::UnitVector);
-	void AnimDrawDebugLine(const FVector& StartLoc, const FVector& EndLoc, const FColor& Color, bool bPersistentLines = false, float LifeTime = -1.f, float Thickness = 0.f);
-	void AnimDrawDebugDirectionalArrow(const FVector& LineStart, const FVector& LineEnd, float ArrowSize, const FColor& Color, bool bPersistentLines = false, float LifeTime = -1.f, float Thickness = 0.f);
-	void AnimDrawDebugSphere(const FVector& Center, float Radius, int32 Segments, const FColor& Color, bool bPersistentLines = false, float LifeTime = -1.f, float Thickness = 0.f);
-	void AnimDrawDebugCoordinateSystem(FVector const& AxisLoc, FRotator const& AxisRot, float Scale = 1.f, bool bPersistentLines = false, float LifeTime = -1.f, float Thickness = 0.f);
-	void AnimDrawDebugPlane(const FTransform& BaseTransform, float Radii, const FColor& Color, bool bPersistentLines = false, float LifeTime = -1.f, float Thickness = 0.f);
+	void AnimDrawDebugOnScreenMessage(const FString& DebugMessage, const FColor& Color, const FVector2D& TextScale = FVector2D::UnitVector, ESceneDepthPriorityGroup DepthPriority = SDPG_World);
+	void AnimDrawDebugLine(const FVector& StartLoc, const FVector& EndLoc, const FColor& Color, bool bPersistentLines = false, float LifeTime = -1.f, float Thickness = 0.f, ESceneDepthPriorityGroup DepthPriority = SDPG_World);
+	void AnimDrawDebugDirectionalArrow(const FVector& LineStart, const FVector& LineEnd, float ArrowSize, const FColor& Color, bool bPersistentLines = false, float LifeTime = -1.f, float Thickness = 0.f, ESceneDepthPriorityGroup DepthPriority = SDPG_World);
+	void AnimDrawDebugSphere(const FVector& Center, float Radius, int32 Segments, const FColor& Color, bool bPersistentLines = false, float LifeTime = -1.f, float Thickness = 0.f, ESceneDepthPriorityGroup DepthPriority = SDPG_World);
+	void AnimDrawDebugCoordinateSystem(FVector const& AxisLoc, FRotator const& AxisRot, float Scale = 1.f, bool bPersistentLines = false, float LifeTime = -1.f, float Thickness = 0.f, ESceneDepthPriorityGroup DepthPriority = SDPG_World);
+	void AnimDrawDebugPlane(const FTransform& BaseTransform, float Radii, const FColor& Color, bool bPersistentLines = false, float LifeTime = -1.f, float Thickness = 0.f, ESceneDepthPriorityGroup DepthPriority = SDPG_World);
+	void AnimDrawDebugPoint(const FVector& Loc, float Size, const FColor& Color, bool bPersistentLines = false, float LifeTime = -1.f, ESceneDepthPriorityGroup DepthPriority = SDPG_World);
 #else
-	void AnimDrawDebugOnScreenMessage(const FString& DebugMessage, const FColor& Color, const FVector2D& TextScale = FVector2D::UnitVector) {}
-	void AnimDrawDebugLine(const FVector& StartLoc, const FVector& EndLoc, const FColor& Color, bool bPersistentLines = false, float LifeTime = -1.f, float Thickness = 0.f) {}
-	void AnimDrawDebugDirectionalArrow(const FVector& LineStart, const FVector& LineEnd, float ArrowSize, const FColor& Color, bool bPersistentLines = false, float LifeTime = -1.f, float Thickness = 0.f) {}
-	void AnimDrawDebugSphere(const FVector& Center, float Radius, int32 Segments, const FColor& Color, bool bPersistentLines = false, float LifeTime = -1.f, float Thickness = 0.f) {}
-	void AnimDrawDebugCoordinateSystem(FVector const& AxisLoc, FRotator const& AxisRot, float Scale = 1.f, bool bPersistentLines = false, float LifeTime = -1.f, float Thickness = 0.f) {}
-	void AnimDrawDebugPlane(const FTransform& BaseTransform, float Radii, const FColor& Color, bool bPersistentLines = false, float LifeTime = -1.f, float Thickness = 0.f) {}
+	void AnimDrawDebugOnScreenMessage(const FString& DebugMessage, const FColor& Color, const FVector2D& TextScale = FVector2D::UnitVector, ESceneDepthPriorityGroup DepthPriority = SDPG_World) {}
+	void AnimDrawDebugLine(const FVector& StartLoc, const FVector& EndLoc, const FColor& Color, bool bPersistentLines = false, float LifeTime = -1.f, float Thickness = 0.f, ESceneDepthPriorityGroup DepthPriority = SDPG_World) {}
+	void AnimDrawDebugDirectionalArrow(const FVector& LineStart, const FVector& LineEnd, float ArrowSize, const FColor& Color, bool bPersistentLines = false, float LifeTime = -1.f, float Thickness = 0.f, ESceneDepthPriorityGroup DepthPriority = SDPG_World) {}
+	void AnimDrawDebugSphere(const FVector& Center, float Radius, int32 Segments, const FColor& Color, bool bPersistentLines = false, float LifeTime = -1.f, float Thickness = 0.f, ESceneDepthPriorityGroup DepthPriority = SDPG_World) {}
+	void AnimDrawDebugCoordinateSystem(FVector const& AxisLoc, FRotator const& AxisRot, float Scale = 1.f, bool bPersistentLines = false, float LifeTime = -1.f, float Thickness = 0.f, ESceneDepthPriorityGroup DepthPriority = SDPG_World) {}
+	void AnimDrawDebugPlane(const FTransform& BaseTransform, float Radii, const FColor& Color, bool bPersistentLines = false, float LifeTime = -1.f, float Thickness = 0.f, ESceneDepthPriorityGroup DepthPriority = SDPG_World) {}
+	void AnimDrawDebugPoint(const FVector& Loc, float Size, const FColor& Color, bool bPersistentLines = false, float LifeTime = -1.f, ESceneDepthPriorityGroup DepthPriority = SDPG_World) {}
 #endif // ENABLE_ANIM_DRAW_DEBUG
 
 #if ENABLE_ANIM_LOGGING
