Auto MSM/Domain set from graph for volumetric and unlit.
Added strata shading model as material info to infor domain/blend/MSM.
This work with material function and material layer.
Added unlit node and shader support for it.

diff --git a/Engine/Shaders/Private/BasePassPixelShader.usf b/Engine/Shaders/Private/BasePassPixelShader.usf
--- a/Engine/Shaders/Private/BasePassPixelShader.usf
+++ b/Engine/Shaders/Private/BasePassPixelShader.usf
@@ -1517,14 +1517,26 @@ void FPixelShaderInOut_MainPS(
 			Throughput);
 
 		DualBlendColorAdd = Color;
-		DualBlendColorMul = Throughput;	// TODO transmittance
+		DualBlendColorMul = Throughput;
 		Color = 0.0f;
-		Opacity = 1.0f;
+		Opacity = 1.0f; // 1 is needed if using translucent. STRATA_TODO: we should only allow premultiplied alpha
 	}
 #endif // STRATA_TRANSLUCENT_ENABLED
 
-#endif // STRATA_ENABLED
+#elif STRATA_ENABLED && MATERIAL_SHADINGMODEL_UNLIT
 
+	// Unlit forces a single BSDF
+	FStrataBSDF UnlitBSDF = PixelMaterialInputs.FrontMaterial.Layers[0].BSDFs[0];
+#if STRATA_OPAQUE_ENABLED
+	Color = BSDF_GETEMISSIVE(UnlitBSDF);
+#else // STRATA_TRANSLUCENT_ENABLED
+	DualBlendColorAdd = BSDF_GETEMISSIVE(UnlitBSDF);
+	DualBlendColorMul = UNLIT_TRANSMITTANCE(UnlitBSDF);
+	Color = 0.0f;
+	Opacity = 1.0f; // 1 is needed if using translucent. STRATA_TODO: we should only allow premultiplied alpha
+#endif
+	
+#endif
 
 	#if MATERIAL_DOMAIN_POSTPROCESS
 		#if MATERIAL_OUTPUT_OPACITY_AS_ALPHA
diff --git a/Engine/Shaders/Private/Strata/Strata.ush b/Engine/Shaders/Private/Strata/Strata.ush
--- a/Engine/Shaders/Private/Strata/Strata.ush
+++ b/Engine/Shaders/Private/Strata/Strata.ush
@@ -92,6 +92,14 @@ struct FStrataBSDF
 #define BSDF_GETWEIGHT10F(X)			BSDF_GETSTATE(X, STATE_BITS_WEIGHT10F, STATE_BIT_OFFSET_WEIGHT10F)
 #define BSDF_SETWEIGHT10F(X, Weight)	BSDF_SETSTATE(X, STATE_BITS_WEIGHT10F, STATE_BIT_OFFSET_WEIGHT10F, Weight)
 
+#if STRATA_INLINE_SHADING
+#define BSDF_SETEMISSIVE(X, V)	X.Emissive = V
+#define BSDF_GETEMISSIVE(X)		X.Emissive
+#else
+#define BSDF_SETEMISSIVE(X, V)	X
+#define BSDF_GETEMISSIVE(X)		0.0f
+#endif
+
 #define DIFFUSE_ALBEDO(X)				X.VGPRs[0].xyz
 #define DIFFUSE_ROUGHNESS(X)			X.VGPRs[0].w
 #define DIFFUSE_SSSPROFILEID(X)			X.VGPRs[1].x
@@ -118,13 +126,8 @@ struct FStrataBSDF
 #define VOLUMETRICFOGCLOUD_ALBEDO(X)	X.VGPRs[0].xyz
 #define VOLUMETRICFOGCLOUD_EXTINCTION(X)X.VGPRs[1].xyz
 #define VOLUMETRICFOGCLOUD_AO(X)		X.VGPRs[0].w
-#if STRATA_INLINE_SHADING
-#define VOLUMETRICFOGCLOUD_SETEMISSIVE(X, V)	X.Emissive = V
-#define VOLUMETRICFOGCLOUD_GETEMISSIVE(X)		X.Emissive
-#else
-#define VOLUMETRICFOGCLOUD_SETEMISSIVE(X, V)	X
-#define VOLUMETRICFOGCLOUD_GETEMISSIVE(X)		0.0f
-#endif
+
+#define UNLIT_TRANSMITTANCE(X)			X.VGPRs[0].xyz
 
 void StrataRequestSharedNormalTangentBasis(inout uint NormalTypes, uint Index);
 
@@ -300,7 +303,22 @@ FStrataData GetStrataVolumeFogCloudBSDF(float3 Albedo, float3 Extinction, float3
 	VOLUMETRICFOGCLOUD_ALBEDO		(StrataData.Layers[0].BSDFs[0]) = Albedo;
 	VOLUMETRICFOGCLOUD_EXTINCTION	(StrataData.Layers[0].BSDFs[0]) = Extinction;
 	VOLUMETRICFOGCLOUD_AO			(StrataData.Layers[0].BSDFs[0]) = AmbientOcclusion;
-	VOLUMETRICFOGCLOUD_SETEMISSIVE	(StrataData.Layers[0].BSDFs[0], Emissive);
+	BSDF_SETEMISSIVE				(StrataData.Layers[0].BSDFs[0], Emissive);
+
+	StrataData.Layers[0].BSDFs[0].Weight = 1.0f;
+
+	return StrataData;
+}
+
+FStrataData GetStrataUnlitBSDF(float3 EmissiveColor, float3 TransmittanceColor)
+{
+	FStrataData StrataData = GetInitialisedStrataData();
+	StrataData.LayerCount++;
+	StrataData.Layers[0].BSDFCount++;
+
+	BSDF_SETTYPE					(StrataData.Layers[0].BSDFs[0], STRATA_BSDF_TYPE_UNLIT);
+	UNLIT_TRANSMITTANCE				(StrataData.Layers[0].BSDFs[0]) = TransmittanceColor;
+	BSDF_SETEMISSIVE				(StrataData.Layers[0].BSDFs[0], EmissiveColor);
 
 	StrataData.Layers[0].BSDFs[0].Weight = 1.0f;
 
diff --git a/Engine/Shaders/Private/VolumetricCloud.usf b/Engine/Shaders/Private/VolumetricCloud.usf
--- a/Engine/Shaders/Private/VolumetricCloud.usf
+++ b/Engine/Shaders/Private/VolumetricCloud.usf
@@ -51,7 +51,7 @@ float3 SampleEmissive(in FPixelMaterialInputs PixelMaterialInputs)
 	float3 EmissiveColor = 0.0f;
 #if STRATA_ENABLED
 	FStrataBSDF BSDF = PixelMaterialInputs.FrontMaterial.Layers[0].BSDFs[0];
-	EmissiveColor = VOLUMETRICFOGCLOUD_GETEMISSIVE(BSDF);
+	EmissiveColor = BSDF_GETEMISSIVE(BSDF);
 #else
 	EmissiveColor = GetMaterialEmissiveRaw(PixelMaterialInputs);
 #endif
diff --git a/Engine/Shaders/Private/VolumetricFogVoxelization.usf b/Engine/Shaders/Private/VolumetricFogVoxelization.usf
--- a/Engine/Shaders/Private/VolumetricFogVoxelization.usf
+++ b/Engine/Shaders/Private/VolumetricFogVoxelization.usf
@@ -40,7 +40,7 @@ float3 SampleEmissive(in FPixelMaterialInputs PixelMaterialInputs)
 	float3 EmissiveColor = 0.0f;
 #if STRATA_ENABLED
 	FStrataBSDF BSDF = PixelMaterialInputs.FrontMaterial.Layers[0].BSDFs[0];
-	EmissiveColor = VOLUMETRICFOGCLOUD_GETEMISSIVE(BSDF);
+	EmissiveColor = BSDF_GETEMISSIVE(BSDF);
 #else
 	EmissiveColor = GetMaterialEmissiveRaw(PixelMaterialInputs);
 #endif
diff --git a/Engine/Shaders/Shared/StrataDefinitions.h b/Engine/Shaders/Shared/StrataDefinitions.h
--- a/Engine/Shaders/Shared/StrataDefinitions.h
+++ b/Engine/Shaders/Shared/StrataDefinitions.h
@@ -27,6 +27,7 @@
 #define STRATA_BSDF_TYPE_VOLUME				3
 #define STRATA_BSDF_TYPE_SHEEN				4
 #define STRATA_BSDF_TYPE_VOLUMETRICFOGCLOUD	5
+#define STRATA_BSDF_TYPE_UNLIT				6
 // When more than 8 BSDF must exists, please update STATE_BITS_BSDF in Strata.ush
 
 
diff --git a/Engine/Source/Runtime/Engine/Classes/Engine/EngineTypes.h b/Engine/Source/Runtime/Engine/Classes/Engine/EngineTypes.h
--- a/Engine/Source/Runtime/Engine/Classes/Engine/EngineTypes.h
+++ b/Engine/Source/Runtime/Engine/Classes/Engine/EngineTypes.h
@@ -531,6 +531,49 @@ private:
 	uint16 ShadingModelField = 0;
 };
 
+/**
+ * Specifies the Strata runtime shading model summarized from the material graph
+ */
+UENUM()
+enum EStrataShadingModel
+{
+	SSM_Unlit					UMETA(DisplayName = "Unlit"),
+	SSM_DefaultLit				UMETA(DisplayName = "DefaultLit"),
+	SSM_VolumetricFogCloud		UMETA(DisplayName = "VolumetricFogCloud"),
+	/** Number of unique shading models. */
+	SSM_NUM						UMETA(Hidden),
+};
+static_assert(SSM_NUM <= 8, "Do not exceed 16 shading models without expanding FStrataMaterialShadingModelField to support uint32 instead of uint16!");
+
+/** Gather information from the Strata material graph to setup material for runtime. */
+USTRUCT()
+struct ENGINE_API FStrataMaterialInfo
+{
+	GENERATED_USTRUCT_BODY()
+
+public:
+	FStrataMaterialInfo() {}
+	FStrataMaterialInfo(EStrataShadingModel InShadingModel) { AddShadingModel(InShadingModel); }
+
+	// Shading model
+	void AddShadingModel(EStrataShadingModel InShadingModel) { check(InShadingModel < SSM_NUM); ShadingModelField |= (1 << (uint16)InShadingModel); }
+	bool HasShadingModel(EStrataShadingModel InShadingModel) const { return (ShadingModelField & (1 << (uint16)InShadingModel)) != 0; }
+	bool HasOnlyShadingModel(EStrataShadingModel InShadingModel) const { return ShadingModelField == (1 << (uint16)InShadingModel); }
+	uint8 GetShadingModelField() const { return ShadingModelField; }
+	int32 CountShadingModels() const { return FMath::CountBits(ShadingModelField); }
+
+	// STRATA_TODO SSS LUT
+
+	bool IsValid() const { return (ShadingModelField > 0) && (ShadingModelField < (1 << SSM_NUM)); }
+
+	bool operator==(const FStrataMaterialInfo& Other) const { return ShadingModelField == Other.GetShadingModelField(); }
+	bool operator!=(const FStrataMaterialInfo& Other) const { return ShadingModelField != Other.GetShadingModelField(); }
+
+private:
+	UPROPERTY()
+	uint8 ShadingModelField = 0;
+};
+
 /** This is used by the drawing passes to determine tessellation policy, so changes here need to be supported in native code. */
 UENUM()
 enum EMaterialTessellationMode
diff --git a/Engine/Source/Runtime/Engine/Classes/Materials/MaterialExpression.h b/Engine/Source/Runtime/Engine/Classes/Materials/MaterialExpression.h
--- a/Engine/Source/Runtime/Engine/Classes/Materials/MaterialExpression.h
+++ b/Engine/Source/Runtime/Engine/Classes/Materials/MaterialExpression.h
@@ -16,6 +16,7 @@ class UEdGraphNode;
 class UMaterial;
 class UTexture;
 struct FPropertyChangedEvent;
+struct FMaterialShadingModelField;
 
 //@warning: FExpressionInput is mirrored in MaterialShared.h and manually "subclassed" in Material.h (FMaterialInput)
 #if !CPP      //noexport struct
@@ -331,6 +332,11 @@ class ENGINE_API UMaterialExpression : public UObject
 	 */
 	virtual bool IsResultStrataMaterial(int32 OutputIndex) { return false; }
 
+	/**
+	 * Recursively parse nodes outputing strata material in order to gather all the possible shading models used in a material graph output a Strata material.
+	 */
+	virtual void GatherStrataMaterialInfo(FStrataMaterialInfo& StrataMaterialInfo, int32 OutputIndex) { }
+
 	/**
 	 * If true, discards the output index when caching this expression which allows more cases to re-use the output instead of adding a separate instruction
 	 */
diff --git a/Engine/Source/Runtime/Engine/Classes/Materials/MaterialExpressionBlendMaterialAttributes.h b/Engine/Source/Runtime/Engine/Classes/Materials/MaterialExpressionBlendMaterialAttributes.h
--- a/Engine/Source/Runtime/Engine/Classes/Materials/MaterialExpressionBlendMaterialAttributes.h
+++ b/Engine/Source/Runtime/Engine/Classes/Materials/MaterialExpressionBlendMaterialAttributes.h
@@ -50,6 +50,8 @@ class UMaterialExpressionBlendMaterialAttributes : public UMaterialExpression
 	virtual FName GetInputName(int32 InputIndex) const override;
 	virtual bool IsInputConnectionRequired(int32 InputIndex) const override {return true;}
 	virtual bool IsResultMaterialAttributes(int32 OutputIndex) override {return true;}
+	virtual bool IsResultStrataMaterial(int32 OutputIndex) override;
+	virtual void GatherStrataMaterialInfo(FStrataMaterialInfo& StrataMaterialInfo, int32 OutputIndex) override;
 	virtual uint32 GetInputType(int32 InputIndex) override {return InputIndex == 2 ? MCT_Float1 : MCT_MaterialAttributes;}
 #endif
 	//~ End UMaterialExpression Interface
diff --git a/Engine/Source/Runtime/Engine/Classes/Materials/MaterialExpressionFunctionOutput.h b/Engine/Source/Runtime/Engine/Classes/Materials/MaterialExpressionFunctionOutput.h
--- a/Engine/Source/Runtime/Engine/Classes/Materials/MaterialExpressionFunctionOutput.h
+++ b/Engine/Source/Runtime/Engine/Classes/Materials/MaterialExpressionFunctionOutput.h
@@ -64,6 +64,7 @@ class UMaterialExpressionFunctionOutput : public UMaterialExpression
 	virtual uint32 GetInputType(int32 InputIndex) override;
 	virtual bool IsResultMaterialAttributes(int32 OutputIndex) override;
 	virtual bool IsResultStrataMaterial(int32 OutputIndex) override;
+	virtual void GatherStrataMaterialInfo(FStrataMaterialInfo& StrataMaterialInfo, int32 OutputIndex) override;
 #endif // WITH_EDITOR
 	//~ End UMaterialExpression Interface
 
diff --git a/Engine/Source/Runtime/Engine/Classes/Materials/MaterialExpressionGetMaterialAttributes.h b/Engine/Source/Runtime/Engine/Classes/Materials/MaterialExpressionGetMaterialAttributes.h
--- a/Engine/Source/Runtime/Engine/Classes/Materials/MaterialExpressionGetMaterialAttributes.h
+++ b/Engine/Source/Runtime/Engine/Classes/Materials/MaterialExpressionGetMaterialAttributes.h
@@ -43,7 +43,9 @@ class UMaterialExpressionGetMaterialAttributes : public UMaterialExpression
 	virtual uint32 GetOutputType(int32 OutputIndex) override;
 	virtual bool IsInputConnectionRequired(int32 InputIndex) const override {return true;}
 	virtual uint32 GetInputType(int32 InputIndex) override { return MCT_MaterialAttributes; }
-	virtual bool IsResultMaterialAttributes(int32 OutputIndex) override {return OutputIndex == 0;}
+	virtual bool IsResultMaterialAttributes(int32 OutputIndex) override { return OutputIndex == 0; }
+	virtual bool IsResultStrataMaterial(int32 OutputIndex) override;
+	virtual void GatherStrataMaterialInfo(FStrataMaterialInfo& StrataMaterialInfo, int32 OutputIndex) override;
 #endif
 	//~ End UMaterialExpression Interface
 };
diff --git a/Engine/Source/Runtime/Engine/Classes/Materials/MaterialExpressionMakeMaterialAttributes.h b/Engine/Source/Runtime/Engine/Classes/Materials/MaterialExpressionMakeMaterialAttributes.h
--- a/Engine/Source/Runtime/Engine/Classes/Materials/MaterialExpressionMakeMaterialAttributes.h
+++ b/Engine/Source/Runtime/Engine/Classes/Materials/MaterialExpressionMakeMaterialAttributes.h
@@ -90,6 +90,8 @@ class UMaterialExpressionMakeMaterialAttributes : public UMaterialExpression
 	virtual void GetCaption(TArray<FString>& OutCaptions) const override;
 	virtual bool IsResultMaterialAttributes(int32 OutputIndex) override {return true;}
 	virtual uint32 GetInputType(int32 InputIndex) override;
+	virtual bool IsResultStrataMaterial(int32 OutputIndex) override;
+	virtual void GatherStrataMaterialInfo(FStrataMaterialInfo& StrataMaterialInfo, int32 OutputIndex) override;
 #endif // WITH_EDITOR
 	//~ End UMaterialExpression Interface
 };
diff --git a/Engine/Source/Runtime/Engine/Classes/Materials/MaterialExpressionMaterialFunctionCall.h b/Engine/Source/Runtime/Engine/Classes/Materials/MaterialExpressionMaterialFunctionCall.h
--- a/Engine/Source/Runtime/Engine/Classes/Materials/MaterialExpressionMaterialFunctionCall.h
+++ b/Engine/Source/Runtime/Engine/Classes/Materials/MaterialExpressionMaterialFunctionCall.h
@@ -129,6 +129,7 @@ class UMaterialExpressionMaterialFunctionCall : public UMaterialExpression
 	virtual bool MatchesSearchQuery( const TCHAR* SearchQuery ) override;
 	virtual bool IsResultMaterialAttributes(int32 OutputIndex) override;
 	virtual bool IsResultStrataMaterial(int32 OutputIndex) override;
+	virtual void GatherStrataMaterialInfo(FStrataMaterialInfo& StrataMaterialInfo, int32 OutputIndex) override;
 	virtual uint32 GetInputType(int32 InputIndex) override;
 	//~ End UMaterialExpression Interface
 
diff --git a/Engine/Source/Runtime/Engine/Classes/Materials/MaterialExpressionStrata.h b/Engine/Source/Runtime/Engine/Classes/Materials/MaterialExpressionStrata.h
--- a/Engine/Source/Runtime/Engine/Classes/Materials/MaterialExpressionStrata.h
+++ b/Engine/Source/Runtime/Engine/Classes/Materials/MaterialExpressionStrata.h
@@ -75,6 +75,7 @@ class UMaterialExpressionStrataDiffuseBSDF : public UMaterialExpression // STRAT
 	virtual uint32 GetOutputType(int32 OutputIndex) override;
 	virtual uint32 GetInputType(int32 InputIndex) override;
 	virtual bool IsResultStrataMaterial(int32 OutputIndex) override;
+	virtual void GatherStrataMaterialInfo(FStrataMaterialInfo& StrataMaterialInfo, int32 OutputIndex) override;
 #endif
 	//~ End UMaterialExpression Interface
 };
@@ -127,6 +128,7 @@ class UMaterialExpressionStrataDielectricBSDF : public UMaterialExpression
 	virtual uint32 GetOutputType(int32 OutputIndex) override;
 	virtual uint32 GetInputType(int32 InputIndex) override;
 	virtual bool IsResultStrataMaterial(int32 OutputIndex) override;
+	virtual void GatherStrataMaterialInfo(FStrataMaterialInfo& StrataMaterialInfo, int32 OutputIndex) override;
 #endif
 	//~ End UMaterialExpression Interface
 };
@@ -179,6 +181,7 @@ class UMaterialExpressionStrataConductorBSDF : public UMaterialExpression
 	virtual uint32 GetOutputType(int32 OutputIndex) override;
 	virtual uint32 GetInputType(int32 InputIndex) override;
 	virtual bool IsResultStrataMaterial(int32 OutputIndex) override;
+	virtual void GatherStrataMaterialInfo(FStrataMaterialInfo& StrataMaterialInfo, int32 OutputIndex) override;
 #endif
 	//~ End UMaterialExpression Interface
 };
@@ -219,6 +222,7 @@ class UMaterialExpressionStrataVolumeBSDF : public UMaterialExpression
 	virtual uint32 GetOutputType(int32 OutputIndex) override;
 	virtual uint32 GetInputType(int32 InputIndex) override;
 	virtual bool IsResultStrataMaterial(int32 OutputIndex) override;
+	virtual void GatherStrataMaterialInfo(FStrataMaterialInfo& StrataMaterialInfo, int32 OutputIndex) override;
 #endif
 	//~ End UMaterialExpression Interface
 };
@@ -253,6 +257,7 @@ class UMaterialExpressionStrataSheenBSDF : public UMaterialExpression
 	virtual uint32 GetOutputType(int32 OutputIndex) override;
 	virtual uint32 GetInputType(int32 InputIndex) override;
 	virtual bool IsResultStrataMaterial(int32 OutputIndex) override;
+	virtual void GatherStrataMaterialInfo(FStrataMaterialInfo& StrataMaterialInfo, int32 OutputIndex) override;
 #endif
 	//~ End UMaterialExpression Interface
 };
@@ -293,6 +298,36 @@ class UMaterialExpressionStrataVolumetricFogCloudBSDF : public UMaterialExpressi
 	virtual uint32 GetOutputType(int32 OutputIndex) override;
 	virtual uint32 GetInputType(int32 InputIndex) override;
 	virtual bool IsResultStrataMaterial(int32 OutputIndex) override;
+	virtual void GatherStrataMaterialInfo(FStrataMaterialInfo& StrataMaterialInfo, int32 OutputIndex) override;
+#endif
+	//~ End UMaterialExpression Interface
+};
+
+UCLASS(MinimalAPI, collapsecategories, hidecategories = Object)
+class UMaterialExpressionStrataUnlitBSDF : public UMaterialExpression
+{
+	GENERATED_UCLASS_BODY()
+
+	/**
+	* EmissiveColor (type = float3, unit = Luminance, default = 0)
+	*/
+	UPROPERTY()
+	FExpressionInput EmissiveColor;
+
+	/**
+	 * The amount of transmitted light from the back side of the surface to the front side of the surface ()
+	 */
+	UPROPERTY()
+	FExpressionInput TransmittanceColor;
+
+	//~ Begin UMaterialExpression Interface
+#if WITH_EDITOR
+	virtual int32 Compile(class FMaterialCompiler* Compiler, int32 OutputIndex) override;
+	virtual void GetCaption(TArray<FString>& OutCaptions) const override;
+	virtual uint32 GetOutputType(int32 OutputIndex) override;
+	virtual uint32 GetInputType(int32 InputIndex) override;
+	virtual bool IsResultStrataMaterial(int32 OutputIndex) override;
+	virtual void GatherStrataMaterialInfo(FStrataMaterialInfo& StrataMaterialInfo, int32 OutputIndex) override;
 #endif
 	//~ End UMaterialExpression Interface
 };
@@ -332,6 +367,7 @@ class UMaterialExpressionStrataHorizontalMixing : public UMaterialExpression
 	virtual uint32 GetOutputType(int32 OutputIndex) override;
 	virtual uint32 GetInputType(int32 InputIndex) override;
 	virtual bool IsResultStrataMaterial(int32 OutputIndex) override;
+	virtual void GatherStrataMaterialInfo(FStrataMaterialInfo& StrataMaterialInfo, int32 OutputIndex) override;
 #endif
 	//~ End UMaterialExpression Interface
 };
@@ -360,6 +396,7 @@ class UMaterialExpressionStrataVerticalLayering : public UMaterialExpression
 	virtual uint32 GetOutputType(int32 OutputIndex) override;
 	virtual uint32 GetInputType(int32 InputIndex) override;
 	virtual bool IsResultStrataMaterial(int32 OutputIndex) override;
+	virtual void GatherStrataMaterialInfo(FStrataMaterialInfo& StrataMaterialInfo, int32 OutputIndex) override;
 #endif
 	//~ End UMaterialExpression Interface
 };
@@ -388,6 +425,7 @@ class UMaterialExpressionStrataAdd : public UMaterialExpression
 	virtual uint32 GetOutputType(int32 OutputIndex) override;
 	virtual uint32 GetInputType(int32 InputIndex) override;
 	virtual bool IsResultStrataMaterial(int32 OutputIndex) override;
+	virtual void GatherStrataMaterialInfo(FStrataMaterialInfo& StrataMaterialInfo, int32 OutputIndex) override;
 #endif
 	//~ End UMaterialExpression Interface
 };
@@ -416,6 +454,7 @@ class UMaterialExpressionStrataMultiply : public UMaterialExpression
 	virtual uint32 GetOutputType(int32 OutputIndex) override;
 	virtual uint32 GetInputType(int32 InputIndex) override;
 	virtual bool IsResultStrataMaterial(int32 OutputIndex) override;
+	virtual void GatherStrataMaterialInfo(FStrataMaterialInfo& StrataMaterialInfo, int32 OutputIndex) override;
 #endif
 	//~ End UMaterialExpression Interface
 };
diff --git a/Engine/Source/Runtime/Engine/Private/Materials/HLSLMaterialTranslator.cpp b/Engine/Source/Runtime/Engine/Private/Materials/HLSLMaterialTranslator.cpp
--- a/Engine/Source/Runtime/Engine/Private/Materials/HLSLMaterialTranslator.cpp
+++ b/Engine/Source/Runtime/Engine/Private/Materials/HLSLMaterialTranslator.cpp
@@ -946,7 +946,7 @@ bool FHLSLMaterialTranslator::Translate()
 		{
 			if (Domain!= MD_Volume && !(BlendMode == BLEND_Translucent || BlendMode == BLEND_Opaque))
 			{
-				Errorf(TEXT("Strata materials must be opaque or translucent."));
+				Errorf(TEXT("Strata materials must be an opaque or translucent surface, or a volume."));
 			}
 		}
 
@@ -1119,20 +1119,24 @@ bool FHLSLMaterialTranslator::Translate()
 		{
 			const FStrataMaterialCompilationInfo& StrataCompilationInfo = CodeChunkToStrataCompilationInfoMap[Chunk[MP_FrontMaterial]];
 
-			// Verify the validity of the strata material
-			if (Domain == MD_Volume)
+			// VolumetricFogCloud must be used in isolation
+			if (StrataMaterialContainsAnyBSDF(this, StrataCompilationInfo, STRATA_BSDF_TYPE_VOLUMETRICFOGCLOUD))
 			{
+				// If the unlit node is used, it must be the only one used
 				if (!StrataIsVolumetricFogCloudOnly(this, StrataCompilationInfo))
 				{
-					FString ErrorMsg = FString::Printf(TEXT("Material %s is in the VOLUME domain but a VolumetricFogCloud BSDF is not provided (asset: %s).\r\n"), *Material->GetDebugName(), *Material->GetAssetPath());
+					FString ErrorMsg = FString::Printf(TEXT("Material %s contains Unlit BSDF but it is not the only one (asset: %s).\r\n"), *Material->GetDebugName(), *Material->GetAssetPath());
 					Error(*ErrorMsg);
 				}
 			}
-			else
+
+			// Unlit must be used in isolation
+			if (StrataMaterialContainsAnyBSDF(this, StrataCompilationInfo, STRATA_BSDF_TYPE_UNLIT))
 			{
-				if (StrataMaterialContainsAnyBSDF(this, StrataCompilationInfo, STRATA_BSDF_TYPE_VOLUMETRICFOGCLOUD))
+				// If the unlit node is used, it must be the only one used
+				if (!StrataIsUnlitOnly(this, StrataCompilationInfo))
 				{
-					FString ErrorMsg = FString::Printf(TEXT("Material %s is not in the VOLUME domain but a VolumetricFogCloud BSDF is used. It should not be used out of the VOLUME domain. (asset: %s).\r\n"), *Material->GetDebugName(), *Material->GetAssetPath());
+					FString ErrorMsg = FString::Printf(TEXT("Material %s contains Unlit BSDF but it is not the only one (asset: %s).\r\n"), *Material->GetDebugName(), *Material->GetAssetPath());
 					Error(*ErrorMsg);
 				}
 			}
@@ -1547,6 +1551,9 @@ void FHLSLMaterialTranslator::GetMaterialEnvironment(EShaderPlatform InPlatform,
 
 	const bool bStrataMaterial = Material->IsStrataMaterial();
 	OutEnvironment.SetDefine(TEXT("MATERIAL_IS_STRATA"), bStrataMaterial ? TEXT("1") : TEXT("0"));
+
+	// STRATA_TODO do not request DualSourceBlending if gray scale transmittance is selected.
+	// bMaterialRequestsDualSourceBlending this base on limited set of blend mode: Opaque, Masked, TransmittanceCoverage, TransmittanceColored;
 	bMaterialRequestsDualSourceBlending |= bStrataMaterial;
 
 	// if duals source blending (colored transmittance) is not supported on a platform, it will fall back to standard alpha blending (grey scale transmittance)
@@ -7552,6 +7559,15 @@ int32 FHLSLMaterialTranslator::StrataVolumetricFogCloudBSDF(int32 Albedo, int32
 	);
 }
 
+int32 FHLSLMaterialTranslator::StrataUnlitBSDF(int32 EmissiveColor, int32 TransmittanceColor)
+{
+	return AddCodeChunk(
+		MCT_Strata, TEXT("GetStrataUnlitBSDF(%s, %s)"),
+		*GetParameterCode(EmissiveColor),
+		*GetParameterCode(TransmittanceColor)
+	);
+}
+
 int32 FHLSLMaterialTranslator::StrataHorizontalMixing(int32 Foreground, int32 Background, int32 Mix)
 {
 	if (Foreground == INDEX_NONE || Background == INDEX_NONE || Mix == INDEX_NONE)
diff --git a/Engine/Source/Runtime/Engine/Private/Materials/HLSLMaterialTranslator.h b/Engine/Source/Runtime/Engine/Private/Materials/HLSLMaterialTranslator.h
--- a/Engine/Source/Runtime/Engine/Private/Materials/HLSLMaterialTranslator.h
+++ b/Engine/Source/Runtime/Engine/Private/Materials/HLSLMaterialTranslator.h
@@ -797,6 +797,7 @@ protected:
 	virtual int32 StrataVolumeBSDF(int32 Albedo, int32 Extinction, int32 Anisotropy, int32 Thickness, int32 Normal, uint8 SharedNormalIndex) override;
 	virtual int32 StrataSheenBSDF(int32 Albedo, int32 Roughness, int32 Normal, uint8 SharedNormalIndex) override;
 	virtual int32 StrataVolumetricFogCloudBSDF(int32 Albedo, int32 Extinction, int32 Emissive, int32 AmbientOcclusion) override;
+	virtual int32 StrataUnlitBSDF(int32 EmissiveColor, int32 TransmittanceColor) override;
 	virtual int32 StrataHorizontalMixing(int32 Foreground, int32 Background, int32 Mix) override;
 	virtual int32 StrataVerticalLayering(int32 Top, int32 Base) override;
 	virtual int32 StrataAdd(int32 A, int32 B) override;
diff --git a/Engine/Source/Runtime/Engine/Private/Materials/Material.cpp b/Engine/Source/Runtime/Engine/Private/Materials/Material.cpp
--- a/Engine/Source/Runtime/Engine/Private/Materials/Material.cpp
+++ b/Engine/Source/Runtime/Engine/Private/Materials/Material.cpp
@@ -4474,9 +4474,66 @@ void UMaterial::RebuildShadingModelField()
 {
 	ShadingModels.ClearShadingModels();
 
+	const URendererSettings* RendererSettings = GetDefault<URendererSettings>();
+	const bool bStrataEnabled = RendererSettings && RendererSettings->bEnableStrata;
+	if (bStrataEnabled && FrontMaterial.IsConnected())
+	{
+		FStrataMaterialInfo StrataMaterialInfo;
+		check(this->FrontMaterial.Expression);
+		if (this->FrontMaterial.Expression->IsResultStrataMaterial(this->FrontMaterial.OutputIndex))
+		{
+			this->FrontMaterial.Expression->GatherStrataMaterialInfo(StrataMaterialInfo, this->FrontMaterial.OutputIndex);
+		}
+
+		bool bSanitizeMaterial = false;
+		if (!StrataMaterialInfo.IsValid())
+		{
+			bSanitizeMaterial = true;
+			UE_LOG(LogMaterial, Error, TEXT("%s: Material information is invalid."), *GetName());
+		}
+		if (StrataMaterialInfo.CountShadingModels() > 1)
+		{
+			bSanitizeMaterial = true;
+			UE_LOG(LogMaterial, Error, TEXT("%s: Material has more than a single material represented."), *GetName());
+		}
+		if (bSanitizeMaterial)
+		{
+			StrataMaterialInfo.AddShadingModel(SSM_DefaultLit);
+		}
+		
+
+		// Now derive some properties from the material
+		if (StrataMaterialInfo.HasOnlyShadingModel(SSM_Unlit))
+		{
+			MaterialDomain = EMaterialDomain::MD_Surface;
+			ShadingModel = MSM_Unlit;
+			if (BlendMode != EBlendMode::BLEND_Opaque && BlendMode != EBlendMode::BLEND_Masked)
+			{
+				BlendMode = EBlendMode::BLEND_Translucent; // This is to be able to use dual-source blending
+			}
+		}
+		else if (StrataMaterialInfo.HasOnlyShadingModel(SSM_DefaultLit))
+		{
+			MaterialDomain = EMaterialDomain::MD_Surface;
+			ShadingModel = MSM_DefaultLit;
+			if (BlendMode != EBlendMode::BLEND_Opaque && BlendMode != EBlendMode::BLEND_Masked)
+			{
+				BlendMode = EBlendMode::BLEND_Translucent; // This is to be able to use dual-source blending
+			}
+		}
+		else if (StrataMaterialInfo.HasOnlyShadingModel(SSM_VolumetricFogCloud))
+		{
+			MaterialDomain = EMaterialDomain::MD_Volume;
+			ShadingModel = MSM_DefaultLit;
+			BlendMode = EBlendMode::BLEND_Additive;
+		}
+
+		// Also update the ShadingModels for remaining pipeline operation
+		ShadingModels.AddShadingModel(ShadingModel);
+	}
 	// If using shading model from material expression, go through the expressions and look for the ShadingModel expression to figure out what shading models need to be supported in this material.
 	// This might not be the same as what is actually compiled in to the shader, since there might be feature switches, static switches etc. that skip certain shading models.
-	if (ShadingModel == MSM_FromMaterialExpression)
+	else if (ShadingModel == MSM_FromMaterialExpression)
 	{
 
 		TArray<UMaterialExpressionShadingModel*> ShadingModelExpressions;
diff --git a/Engine/Source/Runtime/Engine/Private/Materials/MaterialExpressions.cpp b/Engine/Source/Runtime/Engine/Private/Materials/MaterialExpressions.cpp
--- a/Engine/Source/Runtime/Engine/Private/Materials/MaterialExpressions.cpp
+++ b/Engine/Source/Runtime/Engine/Private/Materials/MaterialExpressions.cpp
@@ -5763,6 +5763,23 @@ uint32 UMaterialExpressionMakeMaterialAttributes::GetInputType(int32 InputIndex)
 		return UMaterialExpression::GetInputType(InputIndex);
 	}
 }
+
+bool UMaterialExpressionMakeMaterialAttributes::IsResultStrataMaterial(int32 OutputIndex)
+{
+	if (FrontMaterial.GetTracedInput().Expression)
+	{
+		return FrontMaterial.GetTracedInput().Expression->IsResultStrataMaterial(FrontMaterial.OutputIndex);
+	}
+	return false;
+}
+
+void UMaterialExpressionMakeMaterialAttributes::GatherStrataMaterialInfo(FStrataMaterialInfo& StrataMaterialInfo, int32 OutputIndex)
+{
+	if (FrontMaterial.GetTracedInput().Expression)
+	{
+		FrontMaterial.GetTracedInput().Expression->GatherStrataMaterialInfo(StrataMaterialInfo, FrontMaterial.OutputIndex);
+	}
+}
 #endif // WITH_EDITOR
 
 // -----
@@ -6099,6 +6116,23 @@ uint32 UMaterialExpressionGetMaterialAttributes::GetOutputType(int32 OutputIndex
 	return OutputType;
 }
 
+bool UMaterialExpressionGetMaterialAttributes::IsResultStrataMaterial(int32 OutputIndex)
+{
+	if (MaterialAttributes.Expression)
+	{
+		return MaterialAttributes.Expression->IsResultStrataMaterial(0);
+	}
+	return false;
+}
+
+void UMaterialExpressionGetMaterialAttributes::GatherStrataMaterialInfo(FStrataMaterialInfo& StrataMaterialInfo, int32 OutputIndex)
+{
+	if (MaterialAttributes.Expression)
+	{
+		MaterialAttributes.Expression->GatherStrataMaterialInfo(StrataMaterialInfo, 0);
+	}
+}
+
 void UMaterialExpressionGetMaterialAttributes::PreEditChange(FProperty* PropertyAboutToChange)
 {
 	// Backup attribute array so we can re-connect pins
@@ -6542,6 +6576,31 @@ FName UMaterialExpressionBlendMaterialAttributes::GetInputName(int32 InputIndex)
 
 	return Name;
 }
+
+bool UMaterialExpressionBlendMaterialAttributes::IsResultStrataMaterial(int32 OutputIndex)
+{
+	if (A.GetTracedInput().Expression)
+	{
+		return A.GetTracedInput().Expression->IsResultStrataMaterial(0); // can only blend Starta type together so one or the othjer input is enough
+	}
+	if (B.GetTracedInput().Expression)
+	{
+		return B.GetTracedInput().Expression->IsResultStrataMaterial(0);
+	}
+	return false;
+}
+
+void UMaterialExpressionBlendMaterialAttributes::GatherStrataMaterialInfo(FStrataMaterialInfo& StrataMaterialInfo, int32 OutputIndex)
+{
+	if (A.GetTracedInput().Expression)
+	{
+		A.GetTracedInput().Expression->GatherStrataMaterialInfo(StrataMaterialInfo, A.OutputIndex);
+	}
+	if (B.GetTracedInput().Expression)
+	{
+		B.GetTracedInput().Expression->GatherStrataMaterialInfo(StrataMaterialInfo, B.OutputIndex);
+	}
+}
 #endif // WITH_EDITOR
 
 //
@@ -14283,6 +14342,14 @@ bool UMaterialExpressionMaterialFunctionCall::IsResultStrataMaterial(int32 Outpu
 	}
 }
 
+void UMaterialExpressionMaterialFunctionCall::GatherStrataMaterialInfo(FStrataMaterialInfo& StrataMaterialInfo, int32 OutputIndex)
+{
+	if (OutputIndex >= 0 && OutputIndex < FunctionOutputs.Num() && FunctionOutputs[OutputIndex].ExpressionOutput)
+	{
+		return FunctionOutputs[OutputIndex].ExpressionOutput->GatherStrataMaterialInfo(StrataMaterialInfo, 0);
+	}
+}
+
 uint32 UMaterialExpressionMaterialFunctionCall::GetInputType(int32 InputIndex)
 {
 	if (InputIndex < FunctionInputs.Num())
@@ -14807,6 +14874,15 @@ bool UMaterialExpressionFunctionOutput::IsResultStrataMaterial(int32 OutputIndex
 	}
 }
 
+void UMaterialExpressionFunctionOutput::GatherStrataMaterialInfo(FStrataMaterialInfo& StrataMaterialInfo, int32 OutputIndex)
+{
+	// If there is a loop anywhere in this expression's inputs then we can't risk checking them
+	if( A.GetTracedInput().Expression && !A.Expression->ContainsInputLoop() )
+	{
+		A.Expression->GatherStrataMaterialInfo(StrataMaterialInfo, A.OutputIndex);
+	}
+}
+
 #endif // WITH_EDITOR
 
 ///////////////////////////////////////////////////////////////////////////////
@@ -19179,6 +19255,11 @@ bool UMaterialExpressionStrataDiffuseBSDF::IsResultStrataMaterial(int32 OutputIn
 {
 	return true;
 }
+
+void UMaterialExpressionStrataDiffuseBSDF::GatherStrataMaterialInfo(FStrataMaterialInfo& StrataMaterialInfo, int32 OutputIndex)
+{
+	StrataMaterialInfo.AddShadingModel(SSM_DefaultLit);
+}
 #endif // WITH_EDITOR
 
 
@@ -19263,6 +19344,11 @@ bool UMaterialExpressionStrataDielectricBSDF::IsResultStrataMaterial(int32 Outpu
 {
 	return true;
 }
+
+void UMaterialExpressionStrataDielectricBSDF::GatherStrataMaterialInfo(FStrataMaterialInfo& StrataMaterialInfo, int32 OutputIndex)
+{
+	StrataMaterialInfo.AddShadingModel(SSM_DefaultLit);
+}
 #endif // WITH_EDITOR
 
 
@@ -19347,6 +19433,11 @@ bool UMaterialExpressionStrataConductorBSDF::IsResultStrataMaterial(int32 Output
 {
 	return true;
 }
+
+void UMaterialExpressionStrataConductorBSDF::GatherStrataMaterialInfo(FStrataMaterialInfo& StrataMaterialInfo, int32 OutputIndex)
+{
+	StrataMaterialInfo.AddShadingModel(SSM_DefaultLit);
+}
 #endif // WITH_EDITOR
 
 
@@ -19422,6 +19513,11 @@ bool UMaterialExpressionStrataVolumeBSDF::IsResultStrataMaterial(int32 OutputInd
 {
 	return true;
 }
+
+void UMaterialExpressionStrataVolumeBSDF::GatherStrataMaterialInfo(FStrataMaterialInfo& StrataMaterialInfo, int32 OutputIndex)
+{
+	StrataMaterialInfo.AddShadingModel(SSM_DefaultLit);
+}
 #endif // WITH_EDITOR
 
 
@@ -19489,6 +19585,11 @@ bool UMaterialExpressionStrataSheenBSDF::IsResultStrataMaterial(int32 OutputInde
 {
 	return true;
 }
+
+void UMaterialExpressionStrataSheenBSDF::GatherStrataMaterialInfo(FStrataMaterialInfo& StrataMaterialInfo, int32 OutputIndex)
+{
+	StrataMaterialInfo.AddShadingModel(SSM_DefaultLit);
+}
 #endif // WITH_EDITOR
 
 
@@ -19558,8 +19659,80 @@ bool UMaterialExpressionStrataVolumetricFogCloudBSDF::IsResultStrataMaterial(int
 {
 	return true;
 }
+
+void UMaterialExpressionStrataVolumetricFogCloudBSDF::GatherStrataMaterialInfo(FStrataMaterialInfo& StrataMaterialInfo, int32 OutputIndex)
+{
+	StrataMaterialInfo.AddShadingModel(SSM_VolumetricFogCloud);
+}
 #endif // WITH_EDITOR
 
+
+
+UMaterialExpressionStrataUnlitBSDF::UMaterialExpressionStrataUnlitBSDF(const FObjectInitializer& ObjectInitializer)
+	: Super(ObjectInitializer)
+{
+	struct FConstructorStatics
+	{
+		FText NAME_Strata;
+		FConstructorStatics() : NAME_Strata(LOCTEXT("Strata BSDFs", "Strata BSDFs")) { }
+	};
+	static FConstructorStatics ConstructorStatics;
+#if WITH_EDITORONLY_DATA
+	MenuCategories.Add(ConstructorStatics.NAME_Strata);
+#endif
+}
+
+#if WITH_EDITOR
+int32 UMaterialExpressionStrataUnlitBSDF::Compile(class FMaterialCompiler* Compiler, int32 OutputIndex)
+{
+	int32 OutputCodeChunk = Compiler->StrataUnlitBSDF(
+		CompileWithDefaultFloat3(Compiler, EmissiveColor, 0.0f, 0.0f, 0.0f),
+		CompileWithDefaultFloat3(Compiler, TransmittanceColor, 1.0f, 1.0f, 1.0f));
+
+	uint8 FakeSharedNormalIndex = 0;
+	StrataCompilationInfoCreateSingleBSDFMaterial(Compiler, OutputCodeChunk, FakeSharedNormalIndex, STRATA_BSDF_TYPE_UNLIT);
+
+	return OutputCodeChunk;
+}
+
+void UMaterialExpressionStrataUnlitBSDF::GetCaption(TArray<FString>& OutCaptions) const
+{
+	OutCaptions.Add(TEXT("Strata Unlit BSDF"));
+}
+
+uint32 UMaterialExpressionStrataUnlitBSDF::GetOutputType(int32 OutputIndex)
+{
+	return MCT_Strata;
+}
+
+uint32 UMaterialExpressionStrataUnlitBSDF::GetInputType(int32 InputIndex)
+{
+	switch (InputIndex)
+	{
+	case 0:
+		return MCT_Float3;
+		break;
+	case 1:
+		return MCT_Float3;
+	}
+
+	check(false);
+	return MCT_Float1;
+}
+
+bool UMaterialExpressionStrataUnlitBSDF::IsResultStrataMaterial(int32 OutputIndex)
+{
+	return true;
+}
+
+void UMaterialExpressionStrataUnlitBSDF::GatherStrataMaterialInfo(FStrataMaterialInfo& StrataMaterialInfo, int32 OutputIndex)
+{
+	StrataMaterialInfo.AddShadingModel(SSM_Unlit);
+}
+#endif // WITH_EDITOR
+
+
+
 UMaterialExpressionStrataHorizontalMixing::UMaterialExpressionStrataHorizontalMixing(const FObjectInitializer& ObjectInitializer)
 	: Super(ObjectInitializer)
 {
@@ -19623,6 +19796,18 @@ bool UMaterialExpressionStrataHorizontalMixing::IsResultStrataMaterial(int32 Out
 {
 	return true;
 }
+
+void UMaterialExpressionStrataHorizontalMixing::GatherStrataMaterialInfo(FStrataMaterialInfo& StrataMaterialInfo, int32 OutputIndex)
+{
+	if (Foreground.GetTracedInput().Expression)
+	{
+		Foreground.GetTracedInput().Expression->GatherStrataMaterialInfo(StrataMaterialInfo, Foreground.OutputIndex);
+	}
+	if (Background.GetTracedInput().Expression)
+	{
+		Background.GetTracedInput().Expression->GatherStrataMaterialInfo(StrataMaterialInfo, Background.OutputIndex);
+	}
+}
 #endif // WITH_EDITOR
 
 
@@ -19687,6 +19872,18 @@ bool UMaterialExpressionStrataVerticalLayering::IsResultStrataMaterial(int32 Out
 {
 	return true;
 }
+
+void UMaterialExpressionStrataVerticalLayering::GatherStrataMaterialInfo(FStrataMaterialInfo& StrataMaterialInfo, int32 OutputIndex)
+{
+	if (Top.GetTracedInput().Expression)
+	{
+		Top.GetTracedInput().Expression->GatherStrataMaterialInfo(StrataMaterialInfo, Top.OutputIndex);
+	}
+	if (Base.GetTracedInput().Expression)
+	{
+		Base.GetTracedInput().Expression->GatherStrataMaterialInfo(StrataMaterialInfo, Base.OutputIndex);
+	}
+}
 #endif // WITH_EDITOR
 
 
@@ -19751,6 +19948,18 @@ bool UMaterialExpressionStrataAdd::IsResultStrataMaterial(int32 OutputIndex)
 {
 	return true;
 }
+
+void UMaterialExpressionStrataAdd::GatherStrataMaterialInfo(FStrataMaterialInfo& StrataMaterialInfo, int32 OutputIndex)
+{
+	if (A.GetTracedInput().Expression)
+	{
+		A.GetTracedInput().Expression->GatherStrataMaterialInfo(StrataMaterialInfo, A.OutputIndex);
+	}
+	if (B.GetTracedInput().Expression)
+	{
+		B.GetTracedInput().Expression->GatherStrataMaterialInfo(StrataMaterialInfo, B.OutputIndex);
+	}
+}
 #endif // WITH_EDITOR
 
 
@@ -19815,6 +20024,14 @@ bool UMaterialExpressionStrataMultiply::IsResultStrataMaterial(int32 OutputIndex
 {
 	return true;
 }
+
+void UMaterialExpressionStrataMultiply::GatherStrataMaterialInfo(FStrataMaterialInfo& StrataMaterialInfo, int32 OutputIndex)
+{
+	if (A.GetTracedInput().Expression)
+	{
+		A.GetTracedInput().Expression->GatherStrataMaterialInfo(StrataMaterialInfo, A.OutputIndex);
+	}
+}
 #endif // WITH_EDITOR
 
 
diff --git a/Engine/Source/Runtime/Engine/Private/Materials/StrataMaterial.cpp b/Engine/Source/Runtime/Engine/Private/Materials/StrataMaterial.cpp
--- a/Engine/Source/Runtime/Engine/Private/Materials/StrataMaterial.cpp
+++ b/Engine/Source/Runtime/Engine/Private/Materials/StrataMaterial.cpp
@@ -27,6 +27,9 @@ FString GetStrataBSDFName(uint8 BSDFType)
 	case STRATA_BSDF_TYPE_VOLUMETRICFOGCLOUD:
 		return TEXT("VOLUMETRICFOGCLOUD");
 		break;
+	case STRATA_BSDF_TYPE_UNLIT:
+		return TEXT("UNLIT");
+		break;
 	}
 	check(false);
 	return "";
@@ -159,6 +162,27 @@ bool StrataIsVolumetricFogCloudOnly(FMaterialCompiler* Compiler, const FStrataMa
 	return true;
 }
 
+bool StrataIsUnlitOnly(FMaterialCompiler* Compiler, const FStrataMaterialCompilationInfo& Material)
+{
+	if (Material.TotalBSDFCount == 0 || Material.LayerCount == 0)
+	{
+		Compiler->Error(TEXT("There is no layer or BSDF plugged in, but a material in the volume domain wants to read from a StrataUnlitBSDF."));
+		return false;
+	}
+	if (Material.TotalBSDFCount > 1 || Material.LayerCount > 1)
+	{
+		Compiler->Error(TEXT("There is more than one layer or BSDF, but a material in the volume domain wants to read from a single StrataUnlitBSDF only."));
+		return false;
+	}
+	if (Material.Layers[0].BSDFs[0].Type != STRATA_BSDF_TYPE_UNLIT)
+	{
+		Compiler->Error(TEXT("The single BSDF resulting from the graph is not of type Volume."));
+		return false;
+	}
+
+	return true;
+}
+
 bool StrataMaterialContainsAnyBSDF(FMaterialCompiler* Compiler, const FStrataMaterialCompilationInfo& Material, uint8 BSDFType)
 {
 	for (uint32 LayerIt = 0; LayerIt < Material.LayerCount; ++LayerIt)
diff --git a/Engine/Source/Runtime/Engine/Private/Materials/StrataMaterial.h b/Engine/Source/Runtime/Engine/Private/Materials/StrataMaterial.h
--- a/Engine/Source/Runtime/Engine/Private/Materials/StrataMaterial.h
+++ b/Engine/Source/Runtime/Engine/Private/Materials/StrataMaterial.h
@@ -53,6 +53,7 @@ FStrataMaterialCompilationInfo StrataCompilationInfoHorizontalMixing(FMaterialCo
 FStrataMaterialCompilationInfo StrataCompilationInfoVerticalLayering(FMaterialCompiler* Compiler, const FStrataMaterialCompilationInfo& Top, const FStrataMaterialCompilationInfo& Base);
 
 bool StrataIsVolumetricFogCloudOnly(FMaterialCompiler* Compiler, const FStrataMaterialCompilationInfo& Material);
+bool StrataIsUnlitOnly(FMaterialCompiler* Compiler, const FStrataMaterialCompilationInfo& Material);
 bool StrataMaterialContainsAnyBSDF(FMaterialCompiler* Compiler, const FStrataMaterialCompilationInfo& Material, uint8 BSDFType);
 
 struct FStrataMaterialAnalysisResult
diff --git a/Engine/Source/Runtime/Engine/Private/ShaderCompiler/ShaderGenerationUtil.cpp b/Engine/Source/Runtime/Engine/Private/ShaderCompiler/ShaderGenerationUtil.cpp
--- a/Engine/Source/Runtime/Engine/Private/ShaderCompiler/ShaderGenerationUtil.cpp
+++ b/Engine/Source/Runtime/Engine/Private/ShaderCompiler/ShaderGenerationUtil.cpp
@@ -226,7 +226,9 @@ void FShaderCompileUtilities::ApplyFetchEnvironment(FShaderMaterialPropertyDefin
 	FETCH_COMPILE_BOOL(IS_BASE_PASS);
 	FETCH_COMPILE_BOOL(IS_MATERIAL_SHADER);
 
-	FETCH_COMPILE_BOOL(STRATA_TRANSLUCENT_ENABLED);
+	FETCH_COMPILE_BOOL(PROJECT_STRATA);
+	FETCH_COMPILE_BOOL(MATERIAL_IS_STRATA);
+
 	FETCH_COMPILE_BOOL(DUAL_SOURCE_COLOR_BLENDING_ENABLED);
 
 	FETCH_COMPILE_INT(DECAL_RENDERTARGET_COUNT);
@@ -1837,7 +1839,7 @@ void FShaderCompileUtilities::ApplyDerivedDefines(FShaderCompilerEnvironment& Ou
 		{
 			bTargetUsage[0] = true;
 			// we also need MRT for thin translucency due to dual blending if we are not on the fallback path
-			bTargetUsage[1] = (DerivedDefines.WRITES_VELOCITY_TO_GBUFFER || (MaterialDefines.MATERIAL_SHADINGMODEL_THIN_TRANSLUCENT && DerivedDefines.THIN_TRANSLUCENT_USE_DUAL_BLEND));
+			bTargetUsage[1] = (DerivedDefines.WRITES_VELOCITY_TO_GBUFFER || (MaterialDefines.DUAL_SOURCE_COLOR_BLENDING_ENABLED && DerivedDefines.MATERIAL_WORKS_WITH_DUAL_SOURCE_COLOR_BLENDING));
 		}
 	}
 	else if (MaterialDefines.IS_VIRTUAL_TEXTURE_MATERIAL)
diff --git a/Engine/Source/Runtime/Engine/Public/MaterialCompiler.h b/Engine/Source/Runtime/Engine/Public/MaterialCompiler.h
--- a/Engine/Source/Runtime/Engine/Public/MaterialCompiler.h
+++ b/Engine/Source/Runtime/Engine/Public/MaterialCompiler.h
@@ -397,6 +397,7 @@ public:
 	virtual int32 StrataVolumeBSDF(int32 Albedo, int32 Extinction, int32 Anisotropy, int32 Thickness, int32 Normal, uint8 SharedNormalIndex) = 0;
 	virtual int32 StrataSheenBSDF(int32 Albedo, int32 Roughness, int32 Normal, uint8 SharedNormalIndex) = 0;
 	virtual int32 StrataVolumetricFogCloudBSDF(int32 Albedo, int32 Extinction, int32 Emissive, int32 AmbientOcclusion) = 0;
+	virtual int32 StrataUnlitBSDF(int32 EmissiveColor, int32 TransmittanceColor) = 0;
 	virtual int32 StrataHorizontalMixing(int32 Foreground, int32 Background, int32 Mix) = 0;
 	virtual int32 StrataVerticalLayering(int32 Top, int32 Base) = 0;
 	virtual int32 StrataAdd(int32 A, int32 B) = 0;
@@ -836,6 +837,11 @@ public:
 		return Compiler->StrataVolumetricFogCloudBSDF(Albedo, Extinction, Emissive, AmbientOcclusion);
 	}
 
+	virtual int32 StrataUnlitBSDF(int32 EmissiveColor, int32 TransmittanceColor) override
+	{
+		return Compiler->StrataUnlitBSDF(EmissiveColor, TransmittanceColor);
+	}
+
 	virtual int32 StrataHorizontalMixing(int32 Foreground, int32 Background, int32 Mix) override
 	{
 		return Compiler->StrataHorizontalMixing(Foreground, Background, Mix);
diff --git a/Engine/Source/Runtime/RenderCore/Private/ShaderMaterialDerivedHelpers.cpp b/Engine/Source/Runtime/RenderCore/Private/ShaderMaterialDerivedHelpers.cpp
--- a/Engine/Source/Runtime/RenderCore/Private/ShaderMaterialDerivedHelpers.cpp
+++ b/Engine/Source/Runtime/RenderCore/Private/ShaderMaterialDerivedHelpers.cpp
@@ -84,8 +84,7 @@ FShaderMaterialDerivedDefines RENDERCORE_API CalculateDerivedMaterialParameters(
 
 	Dst.USE_EDITOR_SHADERS = (Dst.PLATFORM_SUPPORTS_EDITOR_SHADERS && Dst.USE_DEVELOPMENT_SHADERS);
 
-//
-////Materials also have to opt in to these features.
+	//Materials also have to opt in to these features.
 	Dst.USE_EDITOR_COMPOSITING = (Dst.USE_EDITOR_SHADERS && Mat.EDITOR_PRIMITIVE_MATERIAL);
 	Dst.MATERIALBLENDING_ANY_TRANSLUCENT = (Mat.MATERIALBLENDING_TRANSLUCENT || Mat.MATERIALBLENDING_ADDITIVE || Mat.MATERIALBLENDING_MODULATE);
 
@@ -102,8 +101,11 @@ FShaderMaterialDerivedDefines RENDERCORE_API CalculateDerivedMaterialParameters(
 	// matching the logic in: bool FMaterialResource::IsDualBlendingEnabled(EShaderPlatform Platform) const
 	Dst.THIN_TRANSLUCENT_USE_DUAL_BLEND = Mat.MATERIAL_SHADINGMODEL_THIN_TRANSLUCENT && SrcGlobal.bSupportsDualBlending;
 
-
-	Dst.MATERIAL_WORKS_WITH_DUAL_SOURCE_COLOR_BLENDING = (Dst.MATERIAL_SHADINGMODEL_THIN_TRANSLUCENT || Mat.STRATA_TRANSLUCENT_ENABLED);
+	// matching the logic in BasePassPixelShader.usf
+	Dst.STRATA_ENABLED = Mat.PROJECT_STRATA | Mat.MATERIAL_IS_STRATA;
+	Dst.SHADER_STRATA_TRANSLUCENT_ENABLED = Dst.STRATA_ENABLED && Dst.MATERIALBLENDING_ANY_TRANSLUCENT;
+	// As of today, all translucent Strata material forces dual source color blending. STRATA_TODO: optimize out dual color blending when not needed or requested
+	Dst.MATERIAL_WORKS_WITH_DUAL_SOURCE_COLOR_BLENDING = (Dst.MATERIAL_SHADINGMODEL_THIN_TRANSLUCENT || Dst.SHADER_STRATA_TRANSLUCENT_ENABLED);
 
 
 	// There are 4 different ways of setting MRTs depending on which .usf file is the base shader, which sets defines which includes a different include file
diff --git a/Engine/Source/Runtime/RenderCore/Public/ShaderMaterial.h b/Engine/Source/Runtime/RenderCore/Public/ShaderMaterial.h
--- a/Engine/Source/Runtime/RenderCore/Public/ShaderMaterial.h
+++ b/Engine/Source/Runtime/RenderCore/Public/ShaderMaterial.h
@@ -177,7 +177,9 @@ struct FShaderMaterialPropertyDefines
 	uint8 OUT_BASECOLOR_NORMAL_SPECULAR : 1;
 	uint8 OUT_WORLDHEIGHT : 1;
 
-	uint8 STRATA_TRANSLUCENT_ENABLED : 1;
+	uint8 PROJECT_STRATA : 1;
+	uint8 MATERIAL_IS_STRATA : 1;
+
 	uint8 DUAL_SOURCE_COLOR_BLENDING_ENABLED : 1;
 
 	uint8 IS_MATERIAL_SHADER : 1;
@@ -272,6 +274,9 @@ struct FShaderMaterialDerivedDefines
 
 	bool PLATFORM_SUPPORTS_EDITOR_SHADERS;
 
+	bool STRATA_ENABLED;
+	bool SHADER_STRATA_TRANSLUCENT_ENABLED;
+
 	bool THIN_TRANSLUCENT_USE_DUAL_BLEND;
 	bool MATERIAL_WORKS_WITH_DUAL_SOURCE_COLOR_BLENDING;
 };
