Strata code always uses INDEX_NONE

diff --git a/Engine/Source/Runtime/Engine/Private/Materials/HLSLMaterialTranslator.cpp b/Engine/Source/Runtime/Engine/Private/Materials/HLSLMaterialTranslator.cpp
--- a/Engine/Source/Runtime/Engine/Private/Materials/HLSLMaterialTranslator.cpp
+++ b/Engine/Source/Runtime/Engine/Private/Materials/HLSLMaterialTranslator.cpp
@@ -7417,7 +7417,7 @@ uint8 FHLSLMaterialTranslator::StrataCompilationInfoRegisterSharedNormalIndex(in
 	{
 		// Allocate a new slot for a new shared normal
 		uint8 NewSharedNormalIndex = NextFreeStrataShaderNormalIndex++;
-		CodeChunkToStrataSharedNormal.Add(NormalCodeChunkHash, { NormalCodeChunk , -1, NewSharedNormalIndex, *GetParameterCode(NormalCodeChunk), FString() });
+		CodeChunkToStrataSharedNormal.Add(NormalCodeChunkHash, { NormalCodeChunk , INDEX_NONE, NewSharedNormalIndex, *GetParameterCode(NormalCodeChunk), FString() });
 		return NewSharedNormalIndex;
 	}
 	// Return the first existing code chunk which match the normal chunk code
@@ -7449,7 +7449,7 @@ uint8 FHLSLMaterialTranslator::StrataCompilationInfoRegisterSharedNormalIndex(in
 			bFound = true;
 			break;
 		}
-		else if (NormalInfo->TangentCodeChunk == -1)
+		else if (NormalInfo->TangentCodeChunk == INDEX_NONE)
 		{
 			NormalInfo->TangentCodeChunk = TangentCodeChunk;
 			NormalInfo->SharedTangentCode = *GetParameterCode(TangentCodeChunk);
diff --git a/Engine/Source/Runtime/Engine/Private/Materials/MaterialExpressions.cpp b/Engine/Source/Runtime/Engine/Private/Materials/MaterialExpressions.cpp
--- a/Engine/Source/Runtime/Engine/Private/Materials/MaterialExpressions.cpp
+++ b/Engine/Source/Runtime/Engine/Private/Materials/MaterialExpressions.cpp
@@ -19145,9 +19145,9 @@ static int32 CompileWithDefaultCodeChunk(class FMaterialCompiler* Compiler, FExp
 	int32 CodeChunk = Input.GetTracedInput().Expression ? Input.Compile(Compiler) : DefaultCodeChunk;
 	if (bDefaultIsUsed)
 	{
-		*bDefaultIsUsed |= CodeChunk == -1;
+		*bDefaultIsUsed |= CodeChunk == INDEX_NONE;
 	}
-	return CodeChunk == -1 ? DefaultCodeChunk : CodeChunk;
+	return CodeChunk == INDEX_NONE ? DefaultCodeChunk : CodeChunk;
 }
 static int32 CompileWithDefaultFloat1(class FMaterialCompiler* Compiler, FExpressionInput& Input, float X, bool* bDefaultIsUsed = nullptr)
 {
@@ -19158,9 +19158,9 @@ static int32 CompileWithDefaultFloat1(class FMaterialCompiler* Compiler, FExpres
 	int32 CodeChunk = Input.GetTracedInput().Expression ? Input.Compile(Compiler) : Compiler->Constant(X);
 	if (bDefaultIsUsed)
 	{
-		*bDefaultIsUsed |= CodeChunk == -1;
+		*bDefaultIsUsed |= CodeChunk == INDEX_NONE;
 	}
-	return CodeChunk == -1 ? Compiler->Constant(X) : CodeChunk;
+	return CodeChunk == INDEX_NONE ? Compiler->Constant(X) : CodeChunk;
 }
 static int32 CompileWithDefaultFloat3(class FMaterialCompiler* Compiler, FExpressionInput& Input, float X, float Y, float Z, bool* bDefaultIsUsed = nullptr)
 {
@@ -19171,9 +19171,9 @@ static int32 CompileWithDefaultFloat3(class FMaterialCompiler* Compiler, FExpres
 	int32 CodeChunk = Input.GetTracedInput().Expression ? Input.Compile(Compiler) : Compiler->Constant3(X, Y, Z);
 	if (bDefaultIsUsed)
 	{
-		*bDefaultIsUsed |= CodeChunk == -1;
+		*bDefaultIsUsed |= CodeChunk == INDEX_NONE;
 	}
-	return CodeChunk == -1 ? Compiler->Constant3(X, Y, Z) : CodeChunk;
+	return CodeChunk == INDEX_NONE ? Compiler->Constant3(X, Y, Z) : CodeChunk;
 }
 
 #endif // WITH_EDITOR
