Ensure that UID value of a SmartName is also serialized during duplication

diff --git a/Engine/Source/Runtime/Engine/Private/Animation/SmartName.cpp b/Engine/Source/Runtime/Engine/Private/Animation/SmartName.cpp
--- a/Engine/Source/Runtime/Engine/Private/Animation/SmartName.cpp
+++ b/Engine/Source/Runtime/Engine/Private/Animation/SmartName.cpp
@@ -338,7 +338,7 @@ bool FSmartName::Serialize(FArchive& Ar)
 		Ar << TempUID;
 	}
 #if WITH_EDITOR
-	else if (Ar.IsTransacting())
+	else if (Ar.IsTransacting() || Ar.HasAnyPortFlags(PPF_Duplicate))
 	{
 		Ar << UID;
 	}
