Build fix for clang.

#rnx

diff --git a/Engine/Source/Runtime/Engine/Private/Materials/MaterialCachedData.cpp b/Engine/Source/Runtime/Engine/Private/Materials/MaterialCachedData.cpp
--- a/Engine/Source/Runtime/Engine/Private/Materials/MaterialCachedData.cpp
+++ b/Engine/Source/Runtime/Engine/Private/Materials/MaterialCachedData.cpp
@@ -40,7 +40,7 @@
 #include "Curves/CurveLinearColor.h"
 #include "Curves/CurveLinearColorAtlas.h"
 
-const FMaterialCachedExpressionData FMaterialCachedExpressionData::EmptyData;
+const FMaterialCachedExpressionData FMaterialCachedExpressionData::EmptyData{};
 
 void FMaterialCachedExpressionData::Reset()
 {
