- Made some curve API const.
- Fixed problem with CopyCurve for FVectorCurve. Memcpy is not valid as it breaks key handle information within the rich curves.
- Added GetVectorCurveByIndex to FTransformCurve, used for retrieving individual component's sub-curves

diff --git a/Engine/Source/Runtime/Engine/Private/Animation/AnimCurveTypes.cpp b/Engine/Source/Runtime/Engine/Private/Animation/AnimCurveTypes.cpp
--- a/Engine/Source/Runtime/Engine/Private/Animation/AnimCurveTypes.cpp
+++ b/Engine/Source/Runtime/Engine/Private/Animation/AnimCurveTypes.cpp
@@ -91,7 +91,7 @@ FLinearColor FAnimCurveBase::MakeColor()
 //  FFloatCurve
 
 // we don't want to have = operator. This only copies curves, but leaving naming and everything else intact. 
-void FFloatCurve::CopyCurve(FFloatCurve& SourceCurve)
+void FFloatCurve::CopyCurve(const FFloatCurve& SourceCurve)
 {
 	FloatCurve = SourceCurve.FloatCurve;
 }
@@ -106,7 +106,7 @@ void FFloatCurve::UpdateOrAddKey(float NewKey, float CurrentTime)
 	FloatCurve.UpdateOrAddKey(CurrentTime, NewKey);
 }
 
-void FFloatCurve::GetKeys(TArray<float>& OutTimes, TArray<float>& OutValues)
+void FFloatCurve::GetKeys(TArray<float>& OutTimes, TArray<float>& OutValues) const
 {
 	const int32 NumKeys = FloatCurve.GetNumKeys();
 	OutTimes.Empty(NumKeys);
@@ -130,9 +130,11 @@ void FFloatCurve::Resize(float NewLength, bool bInsert/* whether insert or remov
 //  FVectorCurve
 
 // we don't want to have = operator. This only copies curves, but leaving naming and everything else intact. 
-void FVectorCurve::CopyCurve(FVectorCurve& SourceCurve)
+void FVectorCurve::CopyCurve(const FVectorCurve& SourceCurve)
 {
-	FMemory::Memcpy(FloatCurves, SourceCurve.FloatCurves);
+	FloatCurves[0] = SourceCurve.FloatCurves[0];
+	FloatCurves[1] = SourceCurve.FloatCurves[1];
+	FloatCurves[2] = SourceCurve.FloatCurves[2];
 }
 
 FVector FVectorCurve::Evaluate(float CurrentTime, float BlendWeight) const
@@ -153,7 +155,7 @@ void FVectorCurve::UpdateOrAddKey(const FVector& NewKey, float CurrentTime)
 	FloatCurves[(int32)EIndex::Z].UpdateOrAddKey(CurrentTime, NewKey.Z);
 }
 
-void FVectorCurve::GetKeys(TArray<float>& OutTimes, TArray<FVector>& OutValues)
+void FVectorCurve::GetKeys(TArray<float>& OutTimes, TArray<FVector>& OutValues) const
 {
 	// Determine curve with most keys
 	int32 MaxNumKeys = 0;
@@ -191,7 +193,7 @@ void FVectorCurve::Resize(float NewLength, bool bInsert/* whether insert or remo
 	FloatCurves[(int32)EIndex::Z].ReadjustTimeRange(0, NewLength, bInsert, OldStartTime, OldEndTime);
 }
 
-int32 FVectorCurve::GetNumKeys()
+int32 FVectorCurve::GetNumKeys() const
 {
 	int32 MaxNumKeys = 0;
 	for (int32 CurveIndex = 0; CurveIndex < 3; ++CurveIndex)
@@ -207,7 +209,7 @@ int32 FVectorCurve::GetNumKeys()
 //  FTransformCurve
 
 // we don't want to have = operator. This only copies curves, but leaving naming and everything else intact. 
-void FTransformCurve::CopyCurve(FTransformCurve& SourceCurve)
+void FTransformCurve::CopyCurve(const FTransformCurve& SourceCurve)
 {
 	TranslationCurve.CopyCurve(SourceCurve.TranslationCurve);
 	RotationCurve.CopyCurve(SourceCurve.RotationCurve);
@@ -250,9 +252,9 @@ void FTransformCurve::UpdateOrAddKey(const FTransform& NewKey, float CurrentTime
 	ScaleCurve.UpdateOrAddKey(NewKey.GetScale3D(), CurrentTime);
 }
 
-void FTransformCurve::GetKeys(TArray<float>& OutTimes, TArray<FTransform>& OutValues)
+void FTransformCurve::GetKeys(TArray<float>& OutTimes, TArray<FTransform>& OutValues) const
 {
-	FVectorCurve* UsedCurve = nullptr;
+	const FVectorCurve* UsedCurve = nullptr;
 	int32 MaxNumKeys = 0;
 
 	int32 NumKeys = TranslationCurve.GetNumKeys();
@@ -309,6 +311,46 @@ void FTransformCurve::Resize(float NewLength, bool bInsert/* whether insert or r
 	ScaleCurve.Resize(NewLength, bInsert, OldStartTime, OldEndTime);
 }
 
+const FVectorCurve* FTransformCurve::GetVectorCurveByIndex(int32 Index) const
+{
+	const FVectorCurve* Curve = nullptr;
+
+	if (Index == 0)
+	{
+		Curve = &TranslationCurve;
+	}
+	else if (Index == 1)
+	{
+		Curve = &RotationCurve;
+	}
+	else if (Index == 2)
+	{
+		Curve = &ScaleCurve;
+	}
+
+	return Curve;
+}
+
+FVectorCurve* FTransformCurve::GetVectorCurveByIndex(int32 Index)
+{
+	FVectorCurve* Curve = nullptr;
+
+	if (Index == 0)
+	{
+		Curve = &TranslationCurve;
+	}
+	else if (Index == 1)
+	{
+		Curve = &RotationCurve;
+	}
+	else if (Index == 2)
+	{
+		Curve = &ScaleCurve;
+	}
+
+	return Curve;
+}
+
 ////////////////////////////////////////////////////
 //  FCachedFloatCurve
 
diff --git a/Engine/Source/Runtime/Engine/Public/Animation/AnimCurveTypes.h b/Engine/Source/Runtime/Engine/Public/Animation/AnimCurveTypes.h
--- a/Engine/Source/Runtime/Engine/Public/Animation/AnimCurveTypes.h
+++ b/Engine/Source/Runtime/Engine/Public/Animation/AnimCurveTypes.h
@@ -170,10 +170,10 @@ struct FFloatCurve : public FAnimCurveBase
 	}
 
 	// we don't want to have = operator. This only copies curves, but leaving naming and everything else intact. 
-	void CopyCurve(FFloatCurve& SourceCurve);
+	void CopyCurve(const FFloatCurve& SourceCurve);
 	ENGINE_API float Evaluate(float CurrentTime) const;
 	ENGINE_API void UpdateOrAddKey(float NewKey, float CurrentTime);
-	ENGINE_API void GetKeys(TArray<float>& OutTimes, TArray<float>& OutValues);
+	ENGINE_API void GetKeys(TArray<float>& OutTimes, TArray<float>& OutValues) const;
 	void Resize(float NewLength, bool bInsert/* whether insert or remove*/, float OldStartTime, float OldEndTime);
 };
 
@@ -202,13 +202,13 @@ struct FVectorCurve : public FAnimCurveBase
 	}
 
 	// we don't want to have = operator. This only copies curves, but leaving naming and everything else intact. 
-	void CopyCurve(FVectorCurve& SourceCurve);
+	void CopyCurve(const FVectorCurve& SourceCurve);
 	ENGINE_API FVector Evaluate(float CurrentTime, float BlendWeight) const;
 	ENGINE_API void UpdateOrAddKey(const FVector& NewKey, float CurrentTime);
-	ENGINE_API void GetKeys(TArray<float>& OutTimes, TArray<FVector>& OutValues);
+	ENGINE_API void GetKeys(TArray<float>& OutTimes, TArray<FVector>& OutValues) const;
 	bool DoesContainKey() const { return (FloatCurves[0].GetNumKeys() > 0 || FloatCurves[1].GetNumKeys() > 0 || FloatCurves[2].GetNumKeys() > 0);}
 	void Resize(float NewLength, bool bInsert/* whether insert or remove*/, float OldStartTime, float OldEndTime);
-	int32 GetNumKeys();
+	int32 GetNumKeys() const;
 };
 
 USTRUCT()
@@ -238,11 +238,14 @@ struct FTransformCurve: public FAnimCurveBase
 	}
 
 	// we don't want to have = operator. This only copies curves, but leaving naming and everything else intact. 
-	void CopyCurve(FTransformCurve& SourceCurve);
+	void CopyCurve(const FTransformCurve& SourceCurve);
 	ENGINE_API FTransform Evaluate(float CurrentTime, float BlendWeight) const;
 	ENGINE_API void UpdateOrAddKey(const FTransform& NewKey, float CurrentTime);
-	ENGINE_API void GetKeys(TArray<float>& OutTimes, TArray<FTransform>& OutValues);
+	ENGINE_API void GetKeys(TArray<float>& OutTimes, TArray<FTransform>& OutValues) const;
 	void Resize(float NewLength, bool bInsert/* whether insert or remove*/, float OldStartTime, float OldEndTime);
+	
+	const FVectorCurve* GetVectorCurveByIndex(int32 Index) const;
+	FVectorCurve* GetVectorCurveByIndex(int32 Index);
 };
 
 USTRUCT(BlueprintType)
