Made USkeleton::GetSmartNameByUID/GetSmartNameByName const functions

diff --git a/Engine/Source/Runtime/Engine/Classes/Animation/Skeleton.h b/Engine/Source/Runtime/Engine/Classes/Animation/Skeleton.h
--- a/Engine/Source/Runtime/Engine/Classes/Animation/Skeleton.h
+++ b/Engine/Source/Runtime/Engine/Classes/Animation/Skeleton.h
@@ -474,8 +474,8 @@ public:
 
 	// quick wrapper function for Find UID by name, if not found, it will return SmartName::MaxUID
 	ENGINE_API SmartName::UID_Type GetUIDByName(const FName& ContainerName, const FName& Name) const;
-	ENGINE_API bool GetSmartNameByUID(const FName& ContainerName, SmartName::UID_Type UID, FSmartName& OutSmartName);
-	ENGINE_API bool GetSmartNameByName(const FName& ContainerName, const FName& InName, FSmartName& OutSmartName);
+	ENGINE_API bool GetSmartNameByUID(const FName& ContainerName, SmartName::UID_Type UID, FSmartName& OutSmartName) const;
+	ENGINE_API bool GetSmartNameByName(const FName& ContainerName, const FName& InName, FSmartName& OutSmartName) const;
 
 	// Get or add a smartname container with the given name
 	ENGINE_API const FSmartNameMapping* GetSmartNameContainer(const FName& ContainerName) const;
diff --git a/Engine/Source/Runtime/Engine/Private/Animation/Skeleton.cpp b/Engine/Source/Runtime/Engine/Private/Animation/Skeleton.cpp
--- a/Engine/Source/Runtime/Engine/Private/Animation/Skeleton.cpp
+++ b/Engine/Source/Runtime/Engine/Private/Animation/Skeleton.cpp
@@ -1368,7 +1368,7 @@ void USkeleton::RemoveSmartnamesAndModify(FName ContainerName, const TArray<FNam
 }
 #endif // WITH_EDITOR
 
-bool USkeleton::GetSmartNameByUID(const FName& ContainerName, SmartName::UID_Type UID, FSmartName& OutSmartName)
+bool USkeleton::GetSmartNameByUID(const FName& ContainerName, SmartName::UID_Type UID, FSmartName& OutSmartName) const
 {
 	const FSmartNameMapping* RequestedMapping = SmartNames.GetContainerInternal(ContainerName);
 	if (RequestedMapping)
@@ -1379,7 +1379,7 @@ bool USkeleton::GetSmartNameByUID(const FName& ContainerName, SmartName::UID_Typ
 	return false;
 }
 
-bool USkeleton::GetSmartNameByName(const FName& ContainerName, const FName& InName, FSmartName& OutSmartName)
+bool USkeleton::GetSmartNameByName(const FName& ContainerName, const FName& InName, FSmartName& OutSmartName) const
 {
 	const FSmartNameMapping* RequestedMapping = SmartNames.GetContainerInternal(ContainerName);
 	if (RequestedMapping)
