Removing previously deprecated UAnimInstance::PlaySlotAnimation

diff --git a/Engine/Source/Runtime/Engine/Classes/Animation/AnimInstance.h b/Engine/Source/Runtime/Engine/Classes/Animation/AnimInstance.h
--- a/Engine/Source/Runtime/Engine/Classes/Animation/AnimInstance.h
+++ b/Engine/Source/Runtime/Engine/Classes/Animation/AnimInstance.h
@@ -571,12 +571,6 @@ public:
 	********************************************************************************************* */
 public:
 
-	/** DEPRECATED. Use PlaySlotAnimationAsDynamicMontage instead, it returns the UAnimMontage created instead of time, allowing more control */
-	/** Play normal animation asset on the slot node. You can only play one asset (whether montage or animsequence) at a time. */
-	UE_DEPRECATED(4.9, "This function is deprecated, please use PlaySlotAnimationAsDynamicMontage instead.")
-	UFUNCTION(BlueprintCallable, Category="Animation")
-	float PlaySlotAnimation(UAnimSequenceBase* Asset, FName SlotNodeName, float BlendInTime = 0.25f, float BlendOutTime = 0.25f, float InPlayRate = 1.f, int32 LoopCount = 1);
-
 	/** Play normal animation asset on the slot node by creating a dynamic UAnimMontage. You can only play one asset (whether montage or animsequence) at a time per SlotGroup. */
 	UFUNCTION(BlueprintCallable, Category="Animation")
 	UAnimMontage* PlaySlotAnimationAsDynamicMontage(UAnimSequenceBase* Asset, FName SlotNodeName, float BlendInTime = 0.25f, float BlendOutTime = 0.25f, float InPlayRate = 1.f, int32 LoopCount = 1, float BlendOutTriggerTime = -1.f, float InTimeToStartMontageAt = 0.f);
diff --git a/Engine/Source/Runtime/Engine/Private/Animation/AnimInstance.cpp b/Engine/Source/Runtime/Engine/Private/Animation/AnimInstance.cpp
--- a/Engine/Source/Runtime/Engine/Private/Animation/AnimInstance.cpp
+++ b/Engine/Source/Runtime/Engine/Private/Animation/AnimInstance.cpp
@@ -1760,66 +1760,6 @@ void UAnimInstance::TriggerQueuedMontageEvents()
 	}
 }
 
-float UAnimInstance::PlaySlotAnimation(UAnimSequenceBase* Asset, FName SlotNodeName, float BlendInTime, float BlendOutTime, float InPlayRate, int32 LoopCount)
-{
-	// create temporary montage and play
-	bool bValidAsset = Asset && !Asset->IsA(UAnimMontage::StaticClass());
-	if (!bValidAsset)
-	{
-		// user warning
-		UE_LOG(LogAnimMontage, Warning, TEXT("PlaySlotAnimation: Invalid input asset(%s). If Montage, please use Montage_Play"), *GetNameSafe(Asset));
-		return 0.f;
-	}
-
-	if (SlotNodeName == NAME_None)
-	{
-		// user warning
-		UE_LOG(LogAnimMontage, Warning, TEXT("SlotNode Name is required. Make sure to add Slot Node in your anim graph and name it."));
-		return 0.f;
-	}
-
-	USkeleton* AssetSkeleton = Asset->GetSkeleton();
-	if (!CurrentSkeleton->IsCompatible(AssetSkeleton))
-	{
-		UE_LOG(LogAnimMontage, Warning, TEXT("The Skeleton '%s' isn't compatible with '%s' in AnimSequence '%s'!"), *GetPathNameSafe(AssetSkeleton), *GetPathNameSafe(CurrentSkeleton), *Asset->GetName());
-		return 0.f;
-	}
-
-	if (!Asset->CanBeUsedInComposition())
-	{
-		UE_LOG(LogAnimMontage, Warning, TEXT("This animation isn't supported to play as montage"));
-		return 0.f;
-	}
-
-	// now play
-	UAnimMontage* NewMontage = NewObject<UAnimMontage>();
-	NewMontage->SetSkeleton(AssetSkeleton);
-
-	// add new track
-	FSlotAnimationTrack& NewTrack = NewMontage->SlotAnimTracks[0];
-	NewTrack.SlotName = SlotNodeName;
-	FAnimSegment NewSegment;
-	NewSegment.AnimReference = Asset;
-	NewSegment.AnimStartTime = 0.f;
-	NewSegment.AnimEndTime = Asset->GetPlayLength();
-	NewSegment.AnimPlayRate = 1.f;
-	NewSegment.StartPos = 0.f;
-	NewSegment.LoopingCount = LoopCount;
-	NewMontage->SetSequenceLength(NewSegment.GetLength());
-	NewTrack.AnimTrack.AnimSegments.Add(NewSegment);
-		
-	FCompositeSection NewSection;
-	NewSection.SectionName = TEXT("Default");
-	NewSection.SetTime(0.0f);
-
-	// add new section
-	NewMontage->CompositeSections.Add(NewSection);
-	NewMontage->BlendIn.SetBlendTime(BlendInTime);
-	NewMontage->BlendOut.SetBlendTime(BlendOutTime);
-
-	return Montage_Play(NewMontage, InPlayRate);
-}
-
 UAnimMontage* UAnimInstance::PlaySlotAnimationAsDynamicMontage(UAnimSequenceBase* Asset, FName SlotNodeName, float BlendInTime, float BlendOutTime, float InPlayRate, int32 LoopCount, float BlendOutTriggerTime, float InTimeToStartMontageAt)
 {
 	if (Asset && CurrentSkeleton->IsCompatible(Asset->GetSkeleton()))
