- Deprecating UAnimSequence compression API for non-editor builds
- Deprecating public acces to ClearCompressedBone/CurveData and SerializeCompressedData

diff --git a/Engine/Source/Runtime/Engine/Classes/Animation/AnimSequence.h b/Engine/Source/Runtime/Engine/Classes/Animation/AnimSequence.h
--- a/Engine/Source/Runtime/Engine/Classes/Animation/AnimSequence.h
+++ b/Engine/Source/Runtime/Engine/Classes/Animation/AnimSequence.h
@@ -629,10 +629,36 @@ public:
 	FString GetDDCCacheKeySuffix(const bool bPerformStripping) const;
 	void ApplyCompressedData(const FString& DataCacheKeySuffix, const bool bPerformFrameStripping, const TArray<uint8>& Data);
 #endif
+
+#if !WITH_EDITOR
+	UE_DEPRECATED(5.0, "WaitOnExistingCompression will be marked EDITOR_ONLY")
+#endif // !WITH_EDITOR
 	void WaitOnExistingCompression(const bool bWantResults=true);
+
+#if !WITH_EDITOR
+	UE_DEPRECATED(5.0, "RequestAnimCompression will be marked EDITOR_ONLY")
+#endif // !WITH_EDITOR
 	void RequestAnimCompression(FRequestAnimCompressionParams Params);
-	void RequestSyncAnimRecompression(bool bOutput = false) { RequestAnimCompression(FRequestAnimCompressionParams(false, false, bOutput)); }
-	void RequestAsyncAnimRecompression(bool bOutput = false) { RequestAnimCompression(FRequestAnimCompressionParams(true, false, bOutput)); }
+
+#if !WITH_EDITOR
+	UE_DEPRECATED(5.0, "RequestSyncAnimRecompression will be marked EDITOR_ONLY")
+#endif // !WITH_EDITOR
+	void RequestSyncAnimRecompression(bool bOutput = false)
+	{
+		PRAGMA_DISABLE_DEPRECATION_WARNINGS
+		RequestAnimCompression(FRequestAnimCompressionParams(false, false, bOutput));
+		PRAGMA_ENABLE_DEPRECATION_WARNINGS
+	}
+
+#if !WITH_EDITOR
+	UE_DEPRECATED(5.0, "RequestAsyncAnimRecompression will be marked EDITOR_ONLY")
+#endif // !WITH_EDITOR
+	void RequestAsyncAnimRecompression(bool bOutput = false)
+	{ 
+		PRAGMA_DISABLE_DEPRECATION_WARNINGS
+		RequestAnimCompression(FRequestAnimCompressionParams(true, false, bOutput));
+		PRAGMA_ENABLE_DEPRECATION_WARNINGS
+	}
 
 protected:
 	void ApplyCompressedData(const TArray<uint8>& Data);
@@ -642,10 +668,12 @@ public:
 	bool IsCompressedDataValid() const;
 	bool IsCurveCompressedDataValid() const;
 
+	UE_DEPRECATED(5.0, "ClearCompressedBoneData will be marked protected")
 	void ClearCompressedBoneData();
+	UE_DEPRECATED(5.0, "ClearCompressedCurveData will be marked protected")
 	void ClearCompressedCurveData();
-
 	// Write the compressed data to the supplied FArchive
+	UE_DEPRECATED(5.0, "SerializeCompressedData will be marked protected")
 	void SerializeCompressedData(FArchive& Ar, bool bDDCData);
 
 	// End Memory related functions
diff --git a/Engine/Source/Runtime/Engine/Private/Animation/AnimSequence.cpp b/Engine/Source/Runtime/Engine/Private/Animation/AnimSequence.cpp
--- a/Engine/Source/Runtime/Engine/Private/Animation/AnimSequence.cpp
+++ b/Engine/Source/Runtime/Engine/Private/Animation/AnimSequence.cpp
@@ -781,7 +781,9 @@ void UAnimSequence::Serialize(FArchive& Ar)
 
 		if (bSerializeCompressedData)
 		{
+			PRAGMA_DISABLE_DEPRECATION_WARNINGS
 			SerializeCompressedData(Ar,false);
+			PRAGMA_ENABLE_DEPRECATION_WARNINGS
 			Ar << bUseRawDataOnly;
 		}
 	}
@@ -1209,13 +1211,19 @@ void UAnimSequence::VerifyTrackMap(USkeleton* MySkeleton)
 #endif // WITH_EDITOR
 void UAnimSequence::BeginDestroy()
 {
+#if WITH_EDITOR
 	// Could already be compressing
+	PRAGMA_DISABLE_DEPRECATION_WARNINGS
 	WaitOnExistingCompression(false);
+	PRAGMA_ENABLE_DEPRECATION_WARNINGS
+#endif // WITH_EDITOR
 
 	Super::BeginDestroy();
 
+	PRAGMA_DISABLE_DEPRECATION_WARNINGS
 	ClearCompressedCurveData();
 	ClearCompressedBoneData();
+	PRAGMA_ENABLE_DEPRECATION_WARNINGS
 }
 
 #if WITH_EDITOR
@@ -2524,7 +2532,9 @@ void UAnimSequence::ApplyCompressedData(const TArray<uint8>& Data)
 	if(Data.Num() > 0)
 	{
 		FMemoryReader MemAr(Data);
+		PRAGMA_DISABLE_DEPRECATION_WARNINGS
 		SerializeCompressedData(MemAr, true);
+		PRAGMA_ENABLE_DEPRECATION_WARNINGS
 		//This is only safe during sync anim compression
 		SetSkeletonVirtualBoneGuid(GetSkeleton()->GetVirtualBoneGuid());
 		bUseRawDataOnly = false;
@@ -2965,8 +2975,10 @@ void UAnimSequence::RecycleAnimSequence()
 	TrackToSkeletonMapTable.Empty();
 	RawCurveData.Empty();
 
+	PRAGMA_DISABLE_DEPRECATION_WARNINGS
 	ClearCompressedBoneData();
 	ClearCompressedCurveData();
+	PRAGMA_ENABLE_DEPRECATION_WARNINGS
 
 	AuthoredSyncMarkers.Empty();
 	UniqueMarkerNames.Empty();
@@ -2983,8 +2995,10 @@ void UAnimSequence::CleanAnimSequenceForImport()
 	RawDataGuid.Invalidate();
 	AnimationTrackNames.Empty();
 	TrackToSkeletonMapTable.Empty();
+	PRAGMA_DISABLE_DEPRECATION_WARNINGS
 	ClearCompressedBoneData();
 	ClearCompressedCurveData();
+	PRAGMA_ENABLE_DEPRECATION_WARNINGS
 }
 #endif // WITH_EDITOR
 
@@ -4652,8 +4666,10 @@ void UAnimSequence::ResetAnimation()
 	AnimationTrackNames.Empty();
 	TrackToSkeletonMapTable.Empty();
 
+	PRAGMA_DISABLE_DEPRECATION_WARNINGS
 	ClearCompressedBoneData();
 	ClearCompressedCurveData();
+	PRAGMA_ENABLE_DEPRECATION_WARNINGS
 
 	Notifies.Empty();
 	AuthoredSyncMarkers.Empty();
@@ -5976,7 +5992,10 @@ void UAnimSequence::EnableRootMotionSettingFromMontage(bool bInEnableRootMotion,
 #if WITH_EDITOR
 void UAnimSequence::OnRawDataChanged()
 {
+	PRAGMA_DISABLE_DEPRECATION_WARNINGS
 	ClearCompressedBoneData();
+	PRAGMA_ENABLE_DEPRECATION_WARNINGS
+
 	bUseRawDataOnly = true;
 
 	RequestAsyncAnimRecompression(false);
