- Basic set of automated tests covering parts of the AnimSequence API regarding sequence length, number of frames, bone tracks, notifies and curve evaluation.
- Removed notify offset to improve drag-ability (not needed anymore with new timeline)
- Replaced hardcoded minimum notify duration by a single animation frame (according to frame rate)
- Setting a sequence length of <= 0.f is now guarded against

diff --git a/Engine/Source/Runtime/Engine/Classes/Animation/AnimSequence.h b/Engine/Source/Runtime/Engine/Classes/Animation/AnimSequence.h
--- a/Engine/Source/Runtime/Engine/Classes/Animation/AnimSequence.h
+++ b/Engine/Source/Runtime/Engine/Classes/Animation/AnimSequence.h
@@ -199,8 +199,6 @@ struct ENGINE_API FRequestAnimCompressionParams
 UCLASS(config=Engine, hidecategories=(UObject, Length), BlueprintType)
 class ENGINE_API UAnimSequence : public UAnimSequenceBase
 {
-	friend class UAnimationBlueprintLibrary;
-
 	GENERATED_UCLASS_BODY()
 
 #if WITH_EDITORONLY_DATA
@@ -1068,4 +1066,6 @@ public:
 	friend class UAnimationBlueprintLibrary;
 	friend class UAnimBoneCompressionSettings;
 	friend class FCustomAttributeCustomization;
+	friend class UAnimationBlueprintLibrary;
+	friend class FAnimSequenceTestBase;
 };
diff --git a/Engine/Source/Runtime/Engine/Classes/Animation/AnimSequenceBase.h b/Engine/Source/Runtime/Engine/Classes/Animation/AnimSequenceBase.h
--- a/Engine/Source/Runtime/Engine/Classes/Animation/AnimSequenceBase.h
+++ b/Engine/Source/Runtime/Engine/Classes/Animation/AnimSequenceBase.h
@@ -43,7 +43,10 @@ public:
 	virtual void SetSequenceLength(float NewLength) 	
 	{
 		PRAGMA_DISABLE_DEPRECATION_WARNINGS
-		SequenceLength = NewLength; 
+		if (NewLength > 0.f)
+		{
+			SequenceLength = NewLength;
+		}
 		PRAGMA_ENABLE_DEPRECATION_WARNINGS
 	}
 
diff --git a/Engine/Source/Runtime/Engine/Private/Animation/AnimSequence.cpp b/Engine/Source/Runtime/Engine/Private/Animation/AnimSequence.cpp
--- a/Engine/Source/Runtime/Engine/Private/Animation/AnimSequence.cpp
+++ b/Engine/Source/Runtime/Engine/Private/Animation/AnimSequence.cpp
@@ -2090,8 +2090,8 @@ void UAnimSequence::ResizeSequence(float NewLength, int32 NewNumKeys, bool bInse
 				// If Notify is inside of the trimmed time frame, zero out the duration
 				if(NotifyEnd >= OldStartTime && NotifyEnd <= OldEndTime)
 				{
-					// small number @todo see if there is define for this
-					NewDuration = 0.1f;
+					// Set duration to a single frame according to the current sampling rate
+					NewDuration = (float)GetSamplingFrameRate().AsInterval();
 				}
 				// If Notify overlaps trimmed time frame, remove trimmed duration
 				else if (CurrentTime < OldEndTime && NotifyEnd > OldEndTime)
@@ -2099,7 +2099,8 @@ void UAnimSequence::ResizeSequence(float NewLength, int32 NewNumKeys, bool bInse
 					NewDuration = NotifyEnd-Duration-CurrentTime;
 				}
 
-				NewDuration = FMath::Max(NewDuration, 0.1f);
+				// Clamp duration to a single frame according to the current sampling rate
+				NewDuration = FMath::Max(NewDuration, (float)GetSamplingFrameRate().AsInterval());
 			}
 
 			if (CurrentTime >= OldStartTime && CurrentTime <= OldEndTime)
@@ -4661,7 +4662,9 @@ void UAnimSequence::ResetAnimation()
 {
 	// clear everything. Making new animation, so need to reset all the things that belong here
 	NumberOfKeys = 0;
-	SetSequenceLength(0.f);
+	PRAGMA_DISABLE_DEPRECATION_WARNINGS
+	SequenceLength = 0.f;
+	PRAGMA_ENABLE_DEPRECATION_WARNINGS
 	RawAnimationData.Empty();
 	AnimationTrackNames.Empty();
 	TrackToSkeletonMapTable.Empty();
diff --git a/Engine/Source/Runtime/Engine/Private/Animation/AnimSequenceBase.cpp b/Engine/Source/Runtime/Engine/Private/Animation/AnimSequenceBase.cpp
--- a/Engine/Source/Runtime/Engine/Private/Animation/AnimSequenceBase.cpp
+++ b/Engine/Source/Runtime/Engine/Private/Animation/AnimSequenceBase.cpp
@@ -577,7 +577,7 @@ void UAnimSequenceBase::UnregisterOnNotifyChanged(void* Unregister)
 
 void UAnimSequenceBase::ClampNotifiesAtEndOfSequence()
 {
-	const float NotifyClampTime = GetPlayLength() - 0.01f; //Slight offset so that notify is still draggable
+	const float NotifyClampTime = GetPlayLength();
 	for(int i = 0; i < Notifies.Num(); ++ i)
 	{
 		if(Notifies[i].GetTime() >= GetPlayLength())
diff --git a/Engine/Source/Runtime/Engine/Public/Animation/AnimTypes.h b/Engine/Source/Runtime/Engine/Public/Animation/AnimTypes.h
--- a/Engine/Source/Runtime/Engine/Public/Animation/AnimTypes.h
+++ b/Engine/Source/Runtime/Engine/Public/Animation/AnimTypes.h
@@ -508,6 +508,16 @@ struct FAnimSyncMarker
 
 	/** This can be used with the Sort() function on a TArray of FAnimSyncMarker to sort the notifies array by time, earliest first. */
 	ENGINE_API bool operator <(const FAnimSyncMarker& Other) const { return Time < Other.Time; }
+
+	ENGINE_API bool operator ==(const FAnimSyncMarker& Other) const
+	{
+		return MarkerName == Other.MarkerName &&
+#if WITH_EDITORONLY_DATA
+			TrackIndex == Other.TrackIndex &&
+			Guid == Other.Guid &&
+#endif
+			Time == Other.Time;
+	}
 };
 
 #if WITH_EDITORONLY_DATA
