Fix race in bulkdata access while saving textures still being built asynchronously

diff --git a/Engine/Source/Runtime/Engine/Private/Texture.cpp b/Engine/Source/Runtime/Engine/Private/Texture.cpp
--- a/Engine/Source/Runtime/Engine/Private/Texture.cpp
+++ b/Engine/Source/Runtime/Engine/Private/Texture.cpp
@@ -595,6 +595,13 @@ void UTexture::PreSave(const class ITargetPlatform* TargetPlatform)
 		Source.Compress();
 	}
 
+	// Ensure that compilation has finished before saving the package
+	// otherwise async compilation might try to read the bulkdata
+	// while it's being serialized to the package.
+	if (IsCompiling())
+	{
+		FTextureCompilingManager::Get().FinishCompilation({ this });
+	}
 #endif // #if WITH_EDITOR
 }
 
@@ -1098,11 +1105,12 @@ void FTextureSource::InitCubeWithMipChain(
 
 void FTextureSource::Compress()
 {
-	if (CanPNGCompress())
-	{
 #if WITH_EDITOR
-		FWriteScopeLock BulkDataExclusiveScope(BulkDataLock.Get());
+	FWriteScopeLock BulkDataExclusiveScope(BulkDataLock.Get());
 #endif
+
+	if (CanPNGCompress())
+	{
 		uint8* BulkDataPtr = (uint8*)BulkData.Lock(LOCK_READ_WRITE);
 		IImageWrapperModule& ImageWrapperModule = FModuleManager::LoadModuleChecked<IImageWrapperModule>( FName("ImageWrapper") );
 		TSharedPtr<IImageWrapper> ImageWrapper = ImageWrapperModule.CreateImageWrapper( EImageFormat::PNG );
diff --git a/Engine/Source/Runtime/Engine/Private/Texture2D.cpp b/Engine/Source/Runtime/Engine/Private/Texture2D.cpp
--- a/Engine/Source/Runtime/Engine/Private/Texture2D.cpp
+++ b/Engine/Source/Runtime/Engine/Private/Texture2D.cpp
@@ -643,6 +643,16 @@ void UTexture2D::PreSave(const class ITargetPlatform* TargetPlatform)
 		bTemporarilyDisableStreaming = false;
 		UpdateResource();
 	}
+
+	// #TODO DC This is redundant code coming from UTexture::Presave that can be removed once we remove the above streaming code.
+
+	// Ensure that compilation has finished before saving the package
+	// otherwise async compilation might try to read the bulkdata
+	// while it's being serialized to the package.
+	if (IsCompiling())
+	{
+		FTextureCompilingManager::Get().FinishCompilation({ this });
+	}
 #endif
 }
 
