Fix cook warnings for missing GetDefaultTriggerWeightThreshold because it's excluded from non-editor builds. Also, in retrospect, there's no reason to hide this function from non-editor builds. It's not expected that people will use this for anything, but there's no reason to stop them.

diff --git a/Engine/Source/Runtime/Engine/Private/Animation/AnimNotify.cpp b/Engine/Source/Runtime/Engine/Private/Animation/AnimNotify.cpp
--- a/Engine/Source/Runtime/Engine/Private/Animation/AnimNotify.cpp
+++ b/Engine/Source/Runtime/Engine/Private/Animation/AnimNotify.cpp
@@ -61,12 +61,10 @@ FString UAnimNotify::GetNotifyName_Implementation() const
 	return NotifyName;
 }
 
-#if WITH_EDITOR
 float UAnimNotify::GetDefaultTriggerWeightThreshold_Implementation() const
 {
 	return ZERO_ANIMWEIGHT_THRESH;
 }
-#endif
 
 /// @endcond
 
diff --git a/Engine/Source/Runtime/Engine/Private/Animation/AnimNotifyState.cpp b/Engine/Source/Runtime/Engine/Private/Animation/AnimNotifyState.cpp
--- a/Engine/Source/Runtime/Engine/Private/Animation/AnimNotifyState.cpp
+++ b/Engine/Source/Runtime/Engine/Private/Animation/AnimNotifyState.cpp
@@ -74,12 +74,10 @@ FString UAnimNotifyState::GetNotifyName_Implementation() const
 
 /// @endcond
 
-#if WITH_EDITOR
 float UAnimNotifyState::GetDefaultTriggerWeightThreshold_Implementation() const
 {
 	return ZERO_ANIMWEIGHT_THRESH;
 }
-#endif
 
 void UAnimNotifyState::PostLoad()
 {
