Add support for blend profiles inside anim montages and blend nodes.
Improved blend profile support adding a time based method.

diff --git a/Engine/Config/BaseEngine.ini b/Engine/Config/BaseEngine.ini
--- a/Engine/Config/BaseEngine.ini
+++ b/Engine/Config/BaseEngine.ini
@@ -843,6 +843,7 @@ NearClipPlane=10.0
 
 ;5.0
 +EnumRedirects=(OldName="/Script/Slate.ETextTransformPolicy", NewName="/Script/SlateCore.ETextTransformPolicy")
++PropertyRedirects=(OldName="BlendProfile.BlendProfileMode", NewName="Mode")
 
 [CoreUObject.Metadata]
 ; Note: UnrealHeaderTool should be rerun after making changes to MetadataRedirects to catch any keys specified in class headers
diff --git a/Engine/Source/Editor/AnimGraph/Private/AnimGraphNode_BlendListBase.cpp b/Engine/Source/Editor/AnimGraph/Private/AnimGraphNode_BlendListBase.cpp
--- a/Engine/Source/Editor/AnimGraph/Private/AnimGraphNode_BlendListBase.cpp
+++ b/Engine/Source/Editor/AnimGraph/Private/AnimGraphNode_BlendListBase.cpp
@@ -1,6 +1,8 @@
 // Copyright Epic Games, Inc. All Rights Reserved.
 
 #include "AnimGraphNode_BlendListBase.h"
+#include "AnimNodes/AnimNode_BlendListBase.h"
+#include "UObject/UE5MainStreamObjectVersion.h"
 
 /////////////////////////////////////////////////////
 // UAnimGraphNode_BlendListBase
diff --git a/Engine/Source/Editor/AnimGraph/Private/AnimGraphNode_LayeredBoneBlend.cpp b/Engine/Source/Editor/AnimGraph/Private/AnimGraphNode_LayeredBoneBlend.cpp
--- a/Engine/Source/Editor/AnimGraph/Private/AnimGraphNode_LayeredBoneBlend.cpp
+++ b/Engine/Source/Editor/AnimGraph/Private/AnimGraphNode_LayeredBoneBlend.cpp
@@ -165,7 +165,7 @@ void UAnimGraphNode_LayeredBoneBlend::ValidateAnimNodeDuringCompilation(class US
 				bCompilationError = true;
 				continue;
 			}
-			else if (BlendMask->BlendProfileMode != EBlendProfileMode::BlendMask)
+			else if (BlendMask->Mode != EBlendProfileMode::BlendMask)
 			{
 				MessageLog.Error(*FText::Format(LOCTEXT("LayeredBlendProfileModeError", "@@ is using a BlendProfile(@@) without a BlendMask mode for Blend Pose {0}. "), FText::AsNumber(MaskIndex)).ToString(), this, BlendMask);
 				bCompilationError = true;
diff --git a/Engine/Source/Editor/AnimGraph/Private/AnimGraphNode_StateMachineBase.cpp b/Engine/Source/Editor/AnimGraph/Private/AnimGraphNode_StateMachineBase.cpp
--- a/Engine/Source/Editor/AnimGraph/Private/AnimGraphNode_StateMachineBase.cpp
+++ b/Engine/Source/Editor/AnimGraph/Private/AnimGraphNode_StateMachineBase.cpp
@@ -383,7 +383,6 @@ void UAnimGraphNode_StateMachineBase::OnProcessDuringCompilation(IAnimBlueprintC
 			BakedTransition.BlendMode = TransitionNode->BlendMode;
 			BakedTransition.CustomCurve = TransitionNode->CustomBlendCurve;
 			BakedTransition.BlendProfile = TransitionNode->BlendProfile;
-			BakedTransition.BlendProfileMode = TransitionNode->BlendProfileMode;
 			BakedTransition.LogicType = TransitionNode->LogicType;
 
 			UAnimStateNodeBase* PreviousState = TransitionNode->GetPreviousState();
diff --git a/Engine/Source/Editor/AnimGraph/Private/AnimStateTransitionNode.cpp b/Engine/Source/Editor/AnimGraph/Private/AnimStateTransitionNode.cpp
--- a/Engine/Source/Editor/AnimGraph/Private/AnimStateTransitionNode.cpp
+++ b/Engine/Source/Editor/AnimGraph/Private/AnimStateTransitionNode.cpp
@@ -81,7 +81,6 @@ UAnimStateTransitionNode::UAnimStateTransitionNode(const FObjectInitializer& Obj
 	Bidirectional = false;
 	PriorityOrder = 1;
 	LogicType = ETransitionLogicType::TLT_StandardBlend;
-	BlendProfileMode = EBlendProfileMode::TimeFactor;
 }
 
 void UAnimStateTransitionNode::AllocateDefaultPins()
@@ -320,8 +319,7 @@ void UAnimStateTransitionNode::PostEditChangeProperty(struct FPropertyChangedEve
 	if (PropertyName == GET_MEMBER_NAME_CHECKED(UAnimStateTransitionNode, CrossfadeDuration) || 
 		PropertyName == GET_MEMBER_NAME_CHECKED(UAnimStateTransitionNode, BlendMode) ||
 		PropertyName == GET_MEMBER_NAME_CHECKED(UAnimStateTransitionNode, CustomBlendCurve) ||
-		PropertyName == GET_MEMBER_NAME_CHECKED(UAnimStateTransitionNode, BlendProfile) ||
-		PropertyName == GET_MEMBER_NAME_CHECKED(UAnimStateTransitionNode, BlendProfileMode))
+		PropertyName == GET_MEMBER_NAME_CHECKED(UAnimStateTransitionNode, BlendProfile))
 	{
 		PropagateCrossfadeSettings();
 	}
@@ -468,7 +466,6 @@ void UAnimStateTransitionNode::CopyCrossfadeSettings(const UAnimStateTransitionN
 	BlendMode = SrcNode->BlendMode;
 	CustomBlendCurve = SrcNode->CustomBlendCurve;
 	BlendProfile = SrcNode->BlendProfile;
-	BlendProfileMode = SrcNode->BlendProfileMode;
 	SharedCrossfadeIdx = SrcNode->SharedCrossfadeIdx;
 	SharedCrossfadeName = SrcNode->SharedCrossfadeName;
 	SharedCrossfadeGuid = SrcNode->SharedCrossfadeGuid;
@@ -551,12 +548,6 @@ void UAnimStateTransitionNode::Serialize(FArchive& Ar)
 	Super::Serialize(Ar);
 	Ar.UsingCustomVersion(FAnimPhysObjectVersion::GUID);
 	Ar.UsingCustomVersion(FUE5MainStreamObjectVersion::GUID);
-
-	// Make sure we use the old behavior on Blend Profiles for older files.
-	if (Ar.IsLoading() && Ar.CustomVer(FUE5MainStreamObjectVersion::GUID) < FUE5MainStreamObjectVersion::AnimationAddedBlendProfileModes)
-	{
-		BlendProfileMode = EBlendProfileMode::WeightFactor;
-	}
 }
 
 void UAnimStateTransitionNode::DestroyNode()
diff --git a/Engine/Source/Editor/AnimGraph/Public/AnimStateTransitionNode.h b/Engine/Source/Editor/AnimGraph/Public/AnimStateTransitionNode.h
--- a/Engine/Source/Editor/AnimGraph/Public/AnimStateTransitionNode.h
+++ b/Engine/Source/Editor/AnimGraph/Public/AnimStateTransitionNode.h
@@ -52,10 +52,6 @@ class UAnimStateTransitionNode : public UAnimStateNodeBase
 	UPROPERTY(EditAnywhere, Category=Transition, meta=(UseAsBlendProfile=true))
 	UBlendProfile* BlendProfile;
 
-	// The blend profile mode to use. This determines how the per bone blend weights are calculated when transitioning.
-	UPROPERTY(EditAnywhere, Category = Transition)
-	EBlendProfileMode BlendProfileMode;
-
 	// Try setting the rule automatically based on most relevant player node's remaining time and the CrossfadeDuration of the transition, ignoring the internal time
 	UPROPERTY(EditAnywhere, Category=Transition)
 	bool bAutomaticRuleBasedOnSequencePlayerInState;
diff --git a/Engine/Source/Editor/DetailCustomizations/Private/AnimTransitionNodeDetails.cpp b/Engine/Source/Editor/DetailCustomizations/Private/AnimTransitionNodeDetails.cpp
--- a/Engine/Source/Editor/DetailCustomizations/Private/AnimTransitionNodeDetails.cpp
+++ b/Engine/Source/Editor/DetailCustomizations/Private/AnimTransitionNodeDetails.cpp
@@ -229,25 +229,6 @@ void FAnimTransitionNodeDetails::CustomizeDetails( IDetailLayoutBuilder& DetailB
 		CrossfadeCategory.AddProperty(GET_MEMBER_NAME_CHECKED(UAnimStateTransitionNode, BlendMode)).DisplayName(LOCTEXT("ModeLabel", "Mode")).IsEnabled(BlendSettingsEnabledAttribute);
 		CrossfadeCategory.AddProperty(GET_MEMBER_NAME_CHECKED(UAnimStateTransitionNode, CustomBlendCurve)).DisplayName(LOCTEXT("CurveLabel", "Custom Blend Curve")).IsEnabled(BlendSettingsEnabledAttribute);
 		CrossfadeCategory.AddProperty(GET_MEMBER_NAME_CHECKED(UAnimStateTransitionNode, BlendProfile)).DisplayName(LOCTEXT("BlendProfileLabel", "Blend Profile")).IsEnabled(BlendSettingsEnabledAttribute);
-		
-		{// Blend Profile Mode
-			TSharedPtr<IPropertyHandle> BlendProfileHandle = DetailBuilder.GetProperty(GET_MEMBER_NAME_CHECKED(UAnimStateTransitionNode, BlendProfile));
-			auto BlendProfileEnabled = [BlendProfileHandle, BlendSettingsEnabledAttribute]()
-			{
-				if (BlendProfileHandle->IsValidHandle())
-				{
-					UObject* BlendProfilePropertyValueInternal = nullptr;
-					BlendProfileHandle->GetValue(BlendProfilePropertyValueInternal);
-					UBlendProfile* SelectedProfile = Cast<UBlendProfile>(BlendProfilePropertyValueInternal);
-					return (SelectedProfile != nullptr) && BlendSettingsEnabledAttribute.Get(false);
-				}
-				return false;
-			};
-			auto BlendProfileEnabledAttribute = TAttribute<bool>::Create(TAttribute<bool>::FGetter::CreateLambda(BlendProfileEnabled));
-			CrossfadeCategory.AddProperty(GET_MEMBER_NAME_CHECKED(UAnimStateTransitionNode, BlendProfileMode))
-				.DisplayName(LOCTEXT("BlendProfileModeLabel", "Blend Profile Mode"))
-				.IsEnabled(BlendProfileEnabledAttribute);
-		}
 
 		//////////////////////////////////////////////////////////////////////////
 
diff --git a/Engine/Source/Editor/SkeletonEditor/Private/EditableSkeleton.cpp b/Engine/Source/Editor/SkeletonEditor/Private/EditableSkeleton.cpp
--- a/Engine/Source/Editor/SkeletonEditor/Private/EditableSkeleton.cpp
+++ b/Engine/Source/Editor/SkeletonEditor/Private/EditableSkeleton.cpp
@@ -193,7 +193,7 @@ void FEditableSkeleton::SetBlendProfileMode(FName InBlendProfileName, EBlendProf
 		const bool bWasBlendMask = BlendProfile->IsBlendMask();
 		BlendProfile->SetFlags(RF_Transactional);
 		BlendProfile->Modify();
-		BlendProfile->BlendProfileMode = ProfileMode;
+		BlendProfile->Mode = ProfileMode;
 		const bool bIsBlendMask = BlendProfile->IsBlendMask();
 		// Re-set entry indices from the end to properly handle different default values when changing from/to blend mask
 		if (bWasBlendMask != bIsBlendMask)
diff --git a/Engine/Source/Editor/SkeletonEditor/Private/SBlendProfilePicker.cpp b/Engine/Source/Editor/SkeletonEditor/Private/SBlendProfilePicker.cpp
--- a/Engine/Source/Editor/SkeletonEditor/Private/SBlendProfilePicker.cpp
+++ b/Engine/Source/Editor/SkeletonEditor/Private/SBlendProfilePicker.cpp
@@ -134,7 +134,7 @@ public:
 	{
 		if (BlendProfile.IsValid())
 		{
-			return BlendProfile->BlendProfileMode;
+			return BlendProfile->Mode;
 		}
 		return EBlendProfileMode::WeightFactor;
 	}
@@ -143,7 +143,7 @@ public:
 	{
 		if (BlendProfile.IsValid())
 		{
-			return BlendProfile->BlendProfileMode == InMode;
+			return BlendProfile->Mode == InMode;
 		}
 		return false;
 	}
@@ -339,7 +339,7 @@ FText SBlendProfilePicker::GetSelectedProfileName() const
 	{
 		if (bIsStandalone)
 		{
-			return FText::Format(FText(LOCTEXT("SelectedNameEntryStandalone", "{0}: {1}")), BlendProfilePickerNames::GetNameForMode(SelectedProfile->BlendProfileMode), FText::FromName(SelectedProfileName));
+			return FText::Format(FText(LOCTEXT("SelectedNameEntryStandalone", "{0}: {1}")), BlendProfilePickerNames::GetNameForMode(SelectedProfile->Mode), FText::FromName(SelectedProfileName));
 		}
 		else
 		{
@@ -491,7 +491,7 @@ void SBlendProfilePicker::OnCreateNewProfileComitted(const FText& NewName, EText
 		else if(UBlendProfile* NewProfile = EditableSkeleton->CreateNewBlendProfile(NameToUse))
 		{
 			// Set our initial blend profile mode. Blend masks can't change this.
-			NewProfile->BlendProfileMode = InMode;
+			NewProfile->Mode = InMode;
 			OnProfileSelected(NewProfile->GetFName());
 		}
 	}
diff --git a/Engine/Source/Runtime/AnimGraphRuntime/Private/AnimNodes/AnimNode_BlendListBase.cpp b/Engine/Source/Runtime/AnimGraphRuntime/Private/AnimNodes/AnimNode_BlendListBase.cpp
--- a/Engine/Source/Runtime/AnimGraphRuntime/Private/AnimNodes/AnimNode_BlendListBase.cpp
+++ b/Engine/Source/Runtime/AnimGraphRuntime/Private/AnimNodes/AnimNode_BlendListBase.cpp
@@ -40,6 +40,12 @@ void FAnimNode_BlendListBase::Initialize_AnyThread(const FAnimationInitializeCon
 	Blends.Empty(NumPoses);
 	Blends.AddZeroed(NumPoses);
 
+	if (BlendProfile)
+	{
+		BlendStartAlphas.Empty(NumPoses);
+		BlendStartAlphas.AddZeroed(NumPoses);
+	}
+
 	LastActiveChildIndex = INDEX_NONE;
 
 	for(int32 i = 0 ; i < Blends.Num() ; ++i)
@@ -49,11 +55,18 @@ void FAnimNode_BlendListBase::Initialize_AnyThread(const FAnimationInitializeCon
 		Blend.SetBlendTime(0.0f);
 		Blend.SetBlendOption(BlendType);
 		Blend.SetCustomCurve(CustomBlendCurve);
+
+		if (BlendProfile)
+		{
+			BlendStartAlphas[i] = 0.0f;
+		}
 	}
 	Blends[0].SetAlpha(1.0f);
 
 	if(BlendProfile)
 	{
+		BlendStartAlphas[0] = 1.0f;
+
 		// Initialise per-bone data
 		PerBoneSampleData.Empty(NumPoses);
 		PerBoneSampleData.AddZeroed(NumPoses);
@@ -146,11 +159,22 @@ void FAnimNode_BlendListBase::Update_AnyThread(const FAnimationUpdateContext& Co
 				if(i == ChildIndex)
 				{
 					Blend.SetValueRange(BlendWeights[i], 1.0f);
+
+					if (BlendProfile)
+					{
+						Blend.ResetAlpha();
+						BlendStartAlphas[i] = Blend.GetAlpha();
+					}
 				}
 				else
 				{
 					Blend.SetValueRange(BlendWeights[i], 0.0f);
 				}
+
+				if (BlendProfile)
+				{
+					BlendStartAlphas[i] = Blend.GetAlpha();
+				}
 			}
 
 			// when this flag is true, we'll reinitialize the children
@@ -201,27 +225,13 @@ void FAnimNode_BlendListBase::Update_AnyThread(const FAnimationUpdateContext& Co
 		}
 
 		// If we're using a blend profile, extract the scales and build blend sample data
-		if(BlendProfile)
+		if (BlendProfile)
 		{
 			for(int32 i = 0; i < BlendPose.Num(); ++i)
 			{
-				// Update Per-Bone Info
-				const float BlendWeight = BlendWeights[i];
 				FBlendSampleData& PoseSampleData = PerBoneSampleData[i];
-				PoseSampleData.TotalWeight = BlendWeight;
-
-				for(int32 j = 0; j < PoseSampleData.PerBoneBlendData.Num(); ++j)
-				{
-					float& BoneBlend = PoseSampleData.PerBoneBlendData[j];
-					float WeightScale = BlendProfile->GetEntryBlendScale(j);
-
-					if(ChildIndex != i)
-					{
-						WeightScale = 1.0f / WeightScale;
-					}
-
-					BoneBlend = BlendWeight * WeightScale;
-				}
+				PoseSampleData.TotalWeight = BlendWeights[i];
+				BlendProfile->UpdateBoneWeights(PoseSampleData, Blends[i], BlendStartAlphas[i], BlendWeights[i]);
 			}
 
 			FBlendSampleData::NormalizeDataWeight(PerBoneSampleData);
diff --git a/Engine/Source/Runtime/AnimGraphRuntime/Public/AnimNodes/AnimNode_BlendListBase.h b/Engine/Source/Runtime/AnimGraphRuntime/Public/AnimNodes/AnimNode_BlendListBase.h
--- a/Engine/Source/Runtime/AnimGraphRuntime/Public/AnimNodes/AnimNode_BlendListBase.h
+++ b/Engine/Source/Runtime/AnimGraphRuntime/Public/AnimNodes/AnimNode_BlendListBase.h
@@ -36,7 +36,7 @@ public:
 
 	UPROPERTY(EditAnywhere, Category=BlendType)
 	EAlphaBlendOption BlendType;
-
+	
 protected:
 	/** This reinitializes child pose when re-activated. For example, when active child changes */
 	UPROPERTY(EditAnywhere, Category = Option)
@@ -58,6 +58,8 @@ protected:
 
 	TArray<float> RemainingBlendTimes;
 
+	TArray<float> BlendStartAlphas;
+
 	TArray<FBlendSampleData> PerBoneSampleData;
 
 	//Store which poses we need to evaluate
diff --git a/Engine/Source/Runtime/Engine/Classes/Animation/AnimInstance.h b/Engine/Source/Runtime/Engine/Classes/Animation/AnimInstance.h
--- a/Engine/Source/Runtime/Engine/Classes/Animation/AnimInstance.h
+++ b/Engine/Source/Runtime/Engine/Classes/Animation/AnimInstance.h
@@ -330,11 +330,12 @@ struct FMontageActiveSlotTracker
 
 struct FMontageEvaluationState
 {
-	FMontageEvaluationState(UAnimMontage* InMontage, float InWeight, float InDesiredWeight, float InPosition, bool bInIsPlaying, bool bInIsActive) 
+	FMontageEvaluationState(UAnimMontage* InMontage, float InPosition, bool bInIsPlaying, bool bInIsActive, const FAlphaBlend& InBlendInfo, const UBlendProfile* InActiveBlendProfile, float InBlendStartAlpha) 
 		: Montage(InMontage)
-		, MontageWeight(InWeight)
-		, DesiredWeight(InDesiredWeight)
+		, BlendInfo(InBlendInfo)
+		, ActiveBlendProfile(InActiveBlendProfile)
 		, MontagePosition(InPosition)
+		, BlendStartAlpha(InBlendStartAlpha)
 		, bIsPlaying(bInIsPlaying)
 		, bIsActive(bInIsActive)
 	{}
@@ -342,15 +343,18 @@ struct FMontageEvaluationState
 	// The montage to evaluate
 	TWeakObjectPtr<UAnimMontage> Montage;
 
-	// The weight to use for this montage
-	float MontageWeight;
+	// The current blend information.
+	FAlphaBlend BlendInfo;
 
-	// The desired weight of this montage
-	float DesiredWeight;
+	// The active blend profile. Montages have a profile for blending in and blending out.
+	const UBlendProfile* ActiveBlendProfile;
 
 	// The position to evaluate this montage at
 	float MontagePosition;
 
+	// The linear alpha value where to start blending from. So not the blended value that already has been curve sampled.
+	float BlendStartAlpha;
+
 	// Whether this montage is playing
 	bool bIsPlaying;
 
diff --git a/Engine/Source/Runtime/Engine/Classes/Animation/AnimMontage.h b/Engine/Source/Runtime/Engine/Classes/Animation/AnimMontage.h
--- a/Engine/Source/Runtime/Engine/Classes/Animation/AnimMontage.h
+++ b/Engine/Source/Runtime/Engine/Classes/Animation/AnimMontage.h
@@ -12,6 +12,7 @@
 #include "UObject/Object.h"
 #include "Animation/AnimLinkableElement.h"
 #include "Animation/AnimTypes.h"
+#include "Animation/BlendProfile.h"
 #include "Animation/Skeleton.h"
 #include "Animation/AnimationAsset.h"
 #include "AlphaBlend.h"
@@ -24,6 +25,7 @@ class UAnimMontage;
 class UAnimSequence;
 class USkeletalMeshComponent;
 
+
 /**
  * Section data for each track. Reference of data will be stored in the child class for the way they want
  * AnimComposite vs AnimMontage have different requirement for the actual data reference
@@ -362,9 +364,16 @@ private:
 	// transient value of previous position before move
 	float PreviousPosition;
 
+	// The current start linear alpha value of the blend. This is not stored inside the FAlphaBlend struct.
+	float BlendStartAlpha;
+
 	// sync group name
 	FName SyncGroupName;
 
+	// Active blend profile.
+	UBlendProfile* ActiveBlendProfile;
+	EBlendProfileMode ActiveBlendProfileMode;
+
 	/**
 	 * Optional evaluation range to use next update (ignoring the real delta time).
 	 * Used by external systems that are setting animation times directly. Will fire off notifies and other events provided the animation system is ticking.
@@ -407,6 +416,10 @@ public:
 	void SetWeight(float InValue) { Blend.SetAlpha(InValue); }
 	/** Set the Desired Weight */
 	void SetDesiredWeight(float InValue) { Blend.SetDesiredValue(InValue); }
+
+	/** Get the current blend info. */
+	const FAlphaBlend& GetBlend() const { return Blend; }
+
 private:
 	/** Followers this Montage will synchronize */
 	TArray<struct FAnimMontageInstance*> MontageSyncFollowers;
@@ -469,7 +482,10 @@ public:
 	float GetPlayRate() const { return PlayRate; }
 	float GetDeltaMoved() const { return DeltaMoved; }
 	float GetPreviousPosition() const { return PreviousPosition;  }
+	float GetBlendStartAlpha() const { return BlendStartAlpha; }
 	const FAnimMontageInstance* GetMontageSyncLeader() const { return MontageSyncLeader; } 
+	const UBlendProfile* GetActiveBlendProfile() const { return ActiveBlendProfile; }
+	const EBlendProfileMode GetActiveBlendProfileMode() const { return ActiveBlendProfileMode; }
 
 	/** 
 	 * Setters
@@ -624,6 +640,14 @@ class UAnimMontage : public UAnimCompositeBase
 	UPROPERTY(EditAnywhere, Category = BlendOption)
 	bool bEnableAutoBlendOut;
 
+	/** The blend profile to use. */
+	UPROPERTY(EditAnywhere, Category = BlendOption, meta = (UseAsBlendProfile = true))
+	UBlendProfile* BlendProfileIn;
+
+	/** The blend profile to use. */
+	UPROPERTY(EditAnywhere, Category = BlendOption, meta = (UseAsBlendProfile = true))
+	UBlendProfile* BlendProfileOut;
+
 	/** Root Bone will be locked to that position when extracting root motion. DEPRECATED in 4.5 root motion is controlled by anim sequences **/
 	UPROPERTY()
 	TEnumAsByte<ERootMotionRootLock::Type> RootMotionRootLock;
diff --git a/Engine/Source/Runtime/Engine/Classes/Animation/AnimNode_StateMachine.h b/Engine/Source/Runtime/Engine/Classes/Animation/AnimNode_StateMachine.h
--- a/Engine/Source/Runtime/Engine/Classes/Animation/AnimNode_StateMachine.h
+++ b/Engine/Source/Runtime/Engine/Classes/Animation/AnimNode_StateMachine.h
@@ -67,10 +67,6 @@ public:
 	UPROPERTY()
 	UBlendProfile* BlendProfile;
 
-	// The blend profile mode to use. This determines how the per bone blend weights are calculated when transitioning.
-	UPROPERTY()
-	EBlendProfileMode BlendProfileMode;
-
 	// Type of blend to use
 	EAlphaBlendOption BlendOption;
 
@@ -95,7 +91,6 @@ public:
 protected:
 	float CalculateInverseAlpha(EAlphaBlendOption BlendMode, float InFraction) const;
 	float CalculateAlpha(float InFraction) const;
-	void ProcessBlendProfile(int32 TransitionDirectionIndex);
 };
 
 USTRUCT()
diff --git a/Engine/Source/Runtime/Engine/Classes/Animation/AnimStateMachineTypes.h b/Engine/Source/Runtime/Engine/Classes/Animation/AnimStateMachineTypes.h
--- a/Engine/Source/Runtime/Engine/Classes/Animation/AnimStateMachineTypes.h
+++ b/Engine/Source/Runtime/Engine/Classes/Animation/AnimStateMachineTypes.h
@@ -146,9 +146,6 @@ struct FAnimationTransitionBetweenStates : public FAnimationStateBase
 	UPROPERTY()
 	UBlendProfile* BlendProfile;
 
-	UPROPERTY()
-	EBlendProfileMode BlendProfileMode;
-
 	UPROPERTY()
 	TEnumAsByte<ETransitionLogicType::Type> LogicType;
 
@@ -168,7 +165,6 @@ struct FAnimationTransitionBetweenStates : public FAnimationStateBase
 		, BlendMode(EAlphaBlendOption::CubicInOut)
 		, CustomCurve(nullptr)
 		, BlendProfile(nullptr)
-		, BlendProfileMode(EBlendProfileMode::TimeFactor)
 		, LogicType(ETransitionLogicType::TLT_StandardBlend)
 #if WITH_EDITOR
 		, ReverseTransition(false)
diff --git a/Engine/Source/Runtime/Engine/Classes/Animation/BlendProfile.h b/Engine/Source/Runtime/Engine/Classes/Animation/BlendProfile.h
--- a/Engine/Source/Runtime/Engine/Classes/Animation/BlendProfile.h
+++ b/Engine/Source/Runtime/Engine/Classes/Animation/BlendProfile.h
@@ -9,6 +9,12 @@
 #include "AnimationRuntime.h"
 #include "BlendProfile.generated.h"
 
+struct FAlphaBlend;
+struct FCompactPose;
+struct FBlendedCurve;
+struct FStackCustomAttributes;
+struct FSlotEvaluationPose;
+
 /** The mode in which the blend profile should be applied. */
 UENUM()
 enum class EBlendProfileMode : uint8
@@ -81,7 +87,7 @@ public:
 	 */
 	float GetBoneBlendScale(int32 InBoneIdx) const;
 
-	/** Get the set blend scale for the specified bone, will return 1.0f if no entry was found (no scale)
+	/** Get the set blend scale for the specified bone, will return 1.0f if no entry was found (no scale). The term bone factor and entry or bone scale refer to the same thing.
 	 *  @param InBoneName Name of the bone to retrieve
 	 */
 	float GetBoneBlendScale(const FName& InBoneName) const;
@@ -96,11 +102,39 @@ public:
 	 */
 	int32 GetEntryIndex(const FName& BoneName) const;
 
-	/** Get the blend scale stored in a specific entry 
+	/** Get the blend scale stored in a specific entry. The term bone factor and entry scale refer to the same thing.
 	 *  @param InEntryIdx Index of the entry to retrieve
 	 */
 	float GetEntryBlendScale(const int32 InEntryIdx) const;
 
+	/** Update all the bone weights for some provided FBlendSampleData.
+	 *  This internally will iterate over all entries inside the InOutCurrentData::PerBoneBlendData parameter and call CalculateBoneWeight for each entry.
+	 *  @param InOutCurrentData This is both input and output. The FBlendSampleData::PerBoneBlendData member will be updated with the weight values as calculated by CalculateBoneWeight for each entry.
+	 *  @param BlendInfo Information about the blend. This contains things like the blend duration and current alpha. Not all blend profile modes might use this data.
+	 *  @param BlendStartAlpha The linear alpha value, so not sampled from the curve, of where the blend started. This is mostly used when we are in the middle of a blend and suddenly reverse its direction.
+	 *  This value basically should contain the alpha value of the blend at the point of reversal.
+	 *  @param MainWeight The weight of the blend. This is used in the weight factor based mode, where this weight is multiplied by the bone factors.
+	 *  @param bInverse Should we inverse the weights? This can be used for things like transition reversal. In most cases you would want it set to false.
+	 */
+	void UpdateBoneWeights(FBlendSampleData& InOutCurrentData, const FAlphaBlend& BlendInfo, float BlendStartAlpha, float MainWeight, bool bInverse = false);
+
+	/** Calculate the blend weight for a given bone. This methoid basically defines how each blend profile mode works.
+	 *  @param BoneFactor This is the per bone value setup in the blend profile editor. The impact of this value depends on what blend profile mode is used.
+	 *  @param Mode The blend profile mode that should be used in combination with this bone factor.
+	 *  @param BlendInfo Information about the blend. This contains things like the blend duration and current alpha. Not all blend profile modes might use this data.
+	 *  @param BlendStartAlpha The linear alpha value, so not sampled from the curve, of where the blend started. This is mostly used when we are in the middle of a blend and suddenly reverse its direction.
+	 *  This value basically should contain the alpha value of the blend at the point of reversal.
+	 *  @param MainWeight The weight of the blend. This is used in the weight factor based mode, where this weight is multiplied by the bone factors.
+	 *  @param bInverse Should we inverse the weights? This can be used for things like transition reversal. In most cases you would want it set to false.
+	 */
+	static float CalculateBoneWeight(float BoneFactor, EBlendProfileMode Mode, const FAlphaBlend& BlendInfo, float BlendStartAlpha, float MainWeight, bool bInverse = false);
+
+	/** Resize and fill an array of floats with the bone factor values. One for each bone inside the compact pose.
+	 *  @param OutBoneBlendProfileFactors This array will be resized and filled with the factors for each bone in the compact pose, as setup in the blend profile editor.
+	 *  @param BoneContainer The bone container which is used to extract how many bones are inside the compact pose and to figure out what factor value to place at what array element.
+	 */
+	void FillBoneScalesArray(TArray<float>& OutBoneBlendProfileFactors, const FBoneContainer& BoneContainer) const;
+
 	// IInterpolationIndexProvider
 	virtual int32 GetPerBoneInterpolationIndex(int32 BoneIndex, const FBoneContainer& RequiredBones) const override;
 	// End IInterpolationIndexProvider
@@ -113,10 +147,11 @@ public:
 	// Default value of entries. Default values are not saved
 	virtual float GetDefaultBlendScale() const { return IsBlendMask() ? 0.0f : 1.0f; }
 
-	bool IsBlendMask() const { return BlendProfileMode == EBlendProfileMode::BlendMask;  }
+	bool IsBlendMask() const { return Mode == EBlendProfileMode::BlendMask;  }
 
-private:
+	EBlendProfileMode GetMode() const { return Mode; }
 
+private:
 	/** Sets the skeleton this blend profile is used with */
 	void SetSkeleton(USkeleton* InSkeleton);
 
@@ -134,5 +169,5 @@ public:
 
 	// Blend Profile Mode. Read EBlendProfileMode for more details
 	UPROPERTY()
-	EBlendProfileMode BlendProfileMode;
+	EBlendProfileMode Mode;
 };
diff --git a/Engine/Source/Runtime/Engine/Private/Animation/AnimInstance.cpp b/Engine/Source/Runtime/Engine/Private/Animation/AnimInstance.cpp
--- a/Engine/Source/Runtime/Engine/Private/Animation/AnimInstance.cpp
+++ b/Engine/Source/Runtime/Engine/Private/Animation/AnimInstance.cpp
@@ -3178,7 +3178,18 @@ void UAnimInstance::UpdateMontageEvaluationData()
 		{
 			UE_LOG(LogAnimMontage, Verbose, TEXT("UpdateMontageEvaluationData : AnimMontage: %s,  (DesiredWeight:%0.2f, Weight:%0.2f)"),
 						*MontageInstance->Montage->GetName(), MontageInstance->GetDesiredWeight(), MontageInstance->GetWeight());
-			Proxy.GetMontageEvaluationData().Add(FMontageEvaluationState(MontageInstance->Montage, MontageInstance->GetWeight(), MontageInstance->GetDesiredWeight(), MontageInstance->GetPosition(), MontageInstance->bPlaying, MontageInstance->IsActive()));
+
+			Proxy.GetMontageEvaluationData().Add(
+				FMontageEvaluationState
+				(
+					MontageInstance->Montage,
+					MontageInstance->GetPosition(),
+					MontageInstance->bPlaying,
+					MontageInstance->IsActive(),
+					MontageInstance->GetBlend(),
+					MontageInstance->GetActiveBlendProfile(),
+					MontageInstance->GetBlendStartAlpha()
+				));
 		}
 	}
 }
diff --git a/Engine/Source/Runtime/Engine/Private/Animation/AnimInstanceProxy.cpp b/Engine/Source/Runtime/Engine/Private/Animation/AnimInstanceProxy.cpp
--- a/Engine/Source/Runtime/Engine/Private/Animation/AnimInstanceProxy.cpp
+++ b/Engine/Source/Runtime/Engine/Private/Animation/AnimInstanceProxy.cpp
@@ -12,6 +12,8 @@
 #include "Animation/AnimNode_TransitionResult.h"
 #include "Animation/AnimNode_SaveCachedPose.h"
 #include "Animation/AnimNode_LinkedInputPose.h"
+#include "Animation/BlendProfile.h"
+#include "Animation/BlendProfileScratchData.h"
 #include "Engine/Engine.h"
 #include "DrawDebugHelpers.h"
 #include "GameFramework/WorldSettings.h"
@@ -1504,6 +1506,344 @@ void FAnimInstanceProxy::SlotEvaluatePose(const FName& SlotNodeName, const FComp
 	SlotEvaluatePose(SlotNodeName, SourceAnimationPoseData, InSourceWeight, BlendedAnimationPoseData, InBlendWeight, InTotalNodeWeight);
 }
 
+
+void FAnimInstanceProxy::SlotEvaluatePoseWithBlendProfiles(const FName& SlotNodeName, const FAnimationPoseData& SourceAnimationPoseData, float InSourceWeight, FAnimationPoseData& OutBlendedAnimationPoseData, float InBlendWeight)
+{
+	const FCompactPose& SourcePose = SourceAnimationPoseData.GetPose();
+	const FBlendedCurve& SourceCurve = SourceAnimationPoseData.GetCurve();
+	const FStackCustomAttributes& SourceAttributes = SourceAnimationPoseData.GetAttributes();
+
+	FCompactPose& BlendedPose = OutBlendedAnimationPoseData.GetPose();
+	FBlendedCurve& BlendedCurve = OutBlendedAnimationPoseData.GetCurve();
+	FStackCustomAttributes& BlendedAttributes = OutBlendedAnimationPoseData.GetAttributes();
+
+	if (InBlendWeight <= ZERO_ANIMWEIGHT_THRESH)
+	{
+		BlendedPose = SourcePose;
+		BlendedCurve = SourceCurve;
+		BlendedAttributes = SourceAttributes;
+		return;
+	}
+
+	// Get the array of per bone weight totals.
+	FBlendProfileScratchData& BlendProfileScratchData = FBlendProfileScratchData::Get();
+	TArray<TArray<float>>& PerBoneWeights = BlendProfileScratchData.PerBoneWeights;
+	TArray<float>& PerBoneWeightTotals = BlendProfileScratchData.PerBoneWeightTotals;
+	TArray<float>& PerBoneWeightTotalsAdditive = BlendProfileScratchData.PerBoneWeightTotalsAdditive;
+	TArray<float>& BoneBlendProfileScales = BlendProfileScratchData.BoneBlendProfileScales;
+	const int32 NumBones = RequiredBones.GetCompactPoseNumBones();
+	PerBoneWeightTotals.Reset(NumBones);
+	PerBoneWeightTotals.AddUninitialized(NumBones);
+	for (int32 BoneIndex = 0; BoneIndex < NumBones; ++BoneIndex)
+	{
+		PerBoneWeightTotals[BoneIndex] = 0.0f;
+	}
+	PerBoneWeightTotalsAdditive = PerBoneWeightTotals;
+	PerBoneWeights.Reset(MontageEvaluationData.Num());
+
+	//------------------------------------------
+	// Figure out what poses we need to blend.
+	//------------------------------------------
+	// Make a list of all the montages that use the slot we're interested in.
+	// Split this in an additive and non additive list.
+	check(MontageEvaluationData.Num() < 255); // Make sure we're in limits and not blending more than 255 poses (that would indicate another issue anyway)
+	TArray<FSlotEvaluationPose>& Poses = BlendProfileScratchData.Poses;
+	TArray<FSlotEvaluationPose>& AdditivePoses = BlendProfileScratchData.AdditivePoses;
+	TArray<uint8, TInlineAllocator<8>>& PoseIndices = BlendProfileScratchData.PoseIndices;
+	TArray<uint8, TInlineAllocator<8>>& AdditivePoseIndices = BlendProfileScratchData.AdditivePoseIndices;
+	TArray<float, TInlineAllocator<8>>& BlendingWeights = BlendProfileScratchData.BlendingWeights;
+	TArray<const FCompactPose*, TInlineAllocator<8>>& BlendingPoses = BlendProfileScratchData.BlendingPoses;
+	TArray<const FBlendedCurve*, TInlineAllocator<8>>& BlendingCurves = BlendProfileScratchData.BlendingCurves;
+	TArray<const FStackCustomAttributes*, TInlineAllocator<8>>& BlendingAttributes = BlendProfileScratchData.BlendingAttributes;
+	Poses.Reset();
+	AdditivePoses.Reset();
+	PoseIndices.Reset();
+	AdditivePoseIndices.Reset();
+	BlendingWeights.Reset();
+	BlendingPoses.Reset();
+	BlendingCurves.Reset();
+	BlendingAttributes.Reset();
+
+	// Gather all the poses we're interested in.
+	int32 CurrentPoseIndex = 0;
+	float NonAdditiveTotalWeight = 0.0f;
+	float AdditiveTotalWeight = 0.0f;
+	for (int32 Index = 0; Index < MontageEvaluationData.Num(); ++Index)
+	{
+		const FMontageEvaluationState& EvalState = MontageEvaluationData[Index];
+		const UAnimMontage* Montage = EvalState.Montage.Get();
+		if (Montage && Montage->IsValidSlot(SlotNodeName))
+		{
+			const UBlendProfile* BlendProfile = EvalState.ActiveBlendProfile;
+			const FAnimTrack* AnimTrack = Montage->GetAnimationData(SlotNodeName);
+			const EAdditiveAnimationType AdditiveAnimType = AnimTrack->IsAdditive()
+				? (AnimTrack->IsRotationOffsetAdditive() ? AAT_RotationOffsetMeshSpace : AAT_LocalSpaceBase)
+				: AAT_None;
+
+			// Bone array has to be allocated prior to calling GetPoseFromAnimTrack.
+			FSlotEvaluationPose NewPose(EvalState.BlendInfo.GetBlendedValue(), AdditiveAnimType);
+			NewPose.Pose.SetBoneContainer(&RequiredBones);
+			NewPose.Curve.InitFrom(RequiredBones);
+
+			// Extract pose from Track.
+			FAnimExtractContext ExtractionContext(EvalState.MontagePosition, Montage->HasRootMotion() && RootMotionMode != ERootMotionMode::NoRootMotionExtraction);
+			FAnimationPoseData NewAnimationPoseData(NewPose);
+			AnimTrack->GetAnimationPose(NewAnimationPoseData, ExtractionContext);
+
+			// Add montage curves.
+			FBlendedCurve MontageCurve;
+			MontageCurve.InitFrom(RequiredBones);
+			Montage->EvaluateCurveData(MontageCurve, EvalState.MontagePosition);
+			NewPose.Curve.Combine(MontageCurve);
+
+			// Capture non-additive poses only in the non-additive pass.
+			if (AdditiveAnimType == AAT_None)
+			{
+				Poses.Add(FSlotEvaluationPose(MoveTemp(NewPose)));
+				PoseIndices.Add(static_cast<uint8>(CurrentPoseIndex));
+				NonAdditiveTotalWeight += NewPose.Weight;
+			}
+			else // Grab additives in the additive pass.
+			{
+				AdditivePoses.Add(FSlotEvaluationPose(MoveTemp(NewPose)));
+				AdditivePoseIndices.Add(static_cast<uint8>(CurrentPoseIndex));
+				AdditiveTotalWeight += NewPose.Weight;
+			}
+
+			PerBoneWeights.AddDefaulted();
+			PerBoneWeights[CurrentPoseIndex].Reset(NumBones);
+			PerBoneWeights[CurrentPoseIndex].AddUninitialized(NumBones);
+			const float PoseWeight = EvalState.BlendInfo.GetBlendedValue();
+			if (BlendProfile)
+			{
+				BlendProfile->FillBoneScalesArray(BoneBlendProfileScales, RequiredBones);
+				for (int32 BoneIndex = 0; BoneIndex < NumBones; ++BoneIndex)
+				{
+					PerBoneWeights[CurrentPoseIndex][BoneIndex] = BlendProfile->CalculateBoneWeight(BoneBlendProfileScales[BoneIndex], BlendProfile->Mode, EvalState.BlendInfo, EvalState.BlendStartAlpha, NewPose.Weight, false);
+				}
+			}
+			else //  This pose doesn't use a blend profile, use the pose weight.
+			{
+				for (int32 BoneIndex = 0; BoneIndex < NumBones; ++BoneIndex)
+				{
+					PerBoneWeights[CurrentPoseIndex][BoneIndex] = PoseWeight;
+				}
+			}
+
+			// Sum up the total weights per bone.
+			if (AdditiveAnimType == AAT_None)
+			{
+				for (int32 BoneIndex = 0; BoneIndex < NumBones; ++BoneIndex)
+				{
+					PerBoneWeightTotals[BoneIndex] += PerBoneWeights[CurrentPoseIndex][BoneIndex];
+				}
+			}
+			else
+			{
+				for (int32 BoneIndex = 0; BoneIndex < NumBones; ++BoneIndex)
+				{
+					PerBoneWeightTotalsAdditive[BoneIndex] += PerBoneWeights[CurrentPoseIndex][BoneIndex];
+				}
+			}
+			
+			CurrentPoseIndex++;
+		} // If montage slot is valid.
+	} // For all montage eval data.
+
+	// Register the source pose if we have any, but don't include it in our normalizations.
+	const float SourceWeight = FMath::Clamp(InSourceWeight, 0.0f, 1.0f);
+	const bool bHasSourcePose = (SourceWeight > ZERO_ANIMWEIGHT_THRESH);
+	if (bHasSourcePose)
+	{
+		const int32 SourcePoseIndex = PerBoneWeights.AddDefaulted();
+		PerBoneWeights[SourcePoseIndex].Reset(NumBones);
+		PerBoneWeights[SourcePoseIndex].AddUninitialized(NumBones);
+
+		for (int32 BoneIndex = 0; BoneIndex < NumBones; ++BoneIndex)
+		{
+			PerBoneWeights[SourcePoseIndex][BoneIndex] = SourceWeight;
+		}
+
+		PoseIndices.Add(SourcePoseIndex);
+		CurrentPoseIndex++;
+	}
+	
+	// Normalize non additive weights.
+	const float NormalizeThreshold = bHasSourcePose ? (1.0f + ZERO_ANIMWEIGHT_THRESH) : ZERO_ANIMWEIGHT_THRESH;
+	const int32 NumPosesToNormalize = Poses.Num();
+	for (int32 PoseIndex = 0; PoseIndex < NumPosesToNormalize; ++PoseIndex)
+	{
+		const int32 PoseWeightIndex = PoseIndices[PoseIndex];
+		for (int32 BoneIndex = 0; BoneIndex < NumBones; ++BoneIndex)
+		{
+			const float TotalWeight = PerBoneWeightTotals[BoneIndex];
+			if (TotalWeight > NormalizeThreshold)
+			{
+				PerBoneWeights[PoseWeightIndex][BoneIndex] /= TotalWeight;
+			}
+			else
+			{
+				if (!bHasSourcePose)
+				{
+					PerBoneWeights[PoseWeightIndex][BoneIndex] = 1.0f;
+				}
+			}
+		}
+	}
+
+	// Calculate the source pose weights.
+	if (bHasSourcePose)
+	{
+		const int32 PoseWeightIndex = PoseIndices.Num() - 1;
+		for (int32 BoneIndex = 0; BoneIndex < NumBones; ++BoneIndex)
+		{
+			float OthersWeightSum = 0.0f;
+			for (int32 PoseIndex = 0; PoseIndex < PoseIndices.Num() - 1; ++PoseIndex)
+			{
+				OthersWeightSum += PerBoneWeights[PoseIndices[PoseIndex]][BoneIndex];
+			}
+
+			PerBoneWeights[PoseIndices[PoseWeightIndex]][BoneIndex] = 1.0f - OthersWeightSum;
+		}
+	}
+
+	// Normalize non additive pose weights, as we need them for blending curves and attributes.
+	if (NonAdditiveTotalWeight > 1.0f + ZERO_ANIMWEIGHT_THRESH)
+	{
+		for (int32 PoseIndex = 0; PoseIndex < Poses.Num(); ++PoseIndex)
+		{
+			Poses[PoseIndex].Weight /= NonAdditiveTotalWeight;
+		}
+	}
+
+	// NOTE: we don't normalize the additive poses.
+
+	//----------------------------------
+	// Build the blend arrays.
+	//----------------------------------
+	const int32 NumPoses = Poses.Num() + (bHasSourcePose ? 1 : 0); // Include the source input pose when doing non-additive blends.
+	BlendingPoses.AddUninitialized(NumPoses);
+	BlendingCurves.AddUninitialized(NumPoses);
+	BlendingAttributes.AddUninitialized(NumPoses);
+	BlendingWeights.AddUninitialized(NumPoses);
+
+	for (int32 Index = 0; Index < Poses.Num(); Index++)
+	{
+		BlendingPoses[Index] = &Poses[Index].Pose;
+		BlendingCurves[Index] = &Poses[Index].Curve;
+		BlendingAttributes[Index] = &Poses[Index].Attributes;
+		BlendingWeights[Index] = Poses[Index].Weight;
+	}
+
+	// Add the source pose to the blends.
+	if (bHasSourcePose)
+	{
+		const int32 SourceIndex = NumPoses - 1;
+		BlendingPoses[SourceIndex] = &SourcePose;
+		BlendingCurves[SourceIndex] = &SourceCurve;
+		BlendingAttributes[SourceIndex] = &SourceAttributes;
+		BlendingWeights[SourceIndex] = (NonAdditiveTotalWeight > 1.0f + ZERO_ANIMWEIGHT_THRESH) ? SourceWeight / NonAdditiveTotalWeight : SourceWeight; // Normalize the source weight if needed.
+	}
+
+	//------------------------------------------
+	// Blend the transforms.
+	//------------------------------------------
+	if (Poses.Num() == 0) // There are only additive poses.
+	{
+		check(AdditivePoseIndices.Num() > 0); // If there are no non-additive poses, there should be at least an additive.
+		BlendedPose = SourcePose;
+		BlendedCurve = SourceCurve;
+		BlendedAttributes = SourceAttributes;
+	}
+	else
+	{
+		// Perform the actual bone transform blending.
+		for (int32 PoseIndex = 0; PoseIndex < NumPoses; ++PoseIndex)
+		{
+			const int32 WeightPoseIndex = PoseIndices[PoseIndex];
+			if (PoseIndex == 0)
+			{
+				for (int32 BoneIndex = 0; BoneIndex < NumBones; ++BoneIndex)
+				{
+					const FCompactPoseBoneIndex CompactBoneIndex(BoneIndex);
+					BlendTransform<ETransformBlendMode::Overwrite>(
+						(*BlendingPoses[PoseIndex])[CompactBoneIndex],
+						BlendedPose[CompactBoneIndex],
+						PerBoneWeights[WeightPoseIndex][BoneIndex]);
+				}
+			}
+			else
+			{
+				for (int32 BoneIndex = 0; BoneIndex < NumBones; ++BoneIndex)
+				{
+					const FCompactPoseBoneIndex CompactBoneIndex(BoneIndex);
+					BlendTransform<ETransformBlendMode::Accumulate>(
+						(*BlendingPoses[PoseIndex])[CompactBoneIndex],
+						BlendedPose[CompactBoneIndex],
+						PerBoneWeights[WeightPoseIndex][BoneIndex]);
+				}
+			}
+		}
+
+		BlendedPose.NormalizeRotations();
+	}
+	
+	// Additive blends.
+	for (int32 PoseIndex = 0; PoseIndex < AdditivePoses.Num(); ++PoseIndex)
+	{
+		FCompactPose& AdditivePose = AdditivePoses[PoseIndex].Pose;
+		const int32 WeightPoseIndex = AdditivePoseIndices[PoseIndex];
+		if (AdditivePoses[PoseIndex].AdditiveType == AAT_RotationOffsetMeshSpace)
+		{
+			FAnimationRuntime::ConvertPoseToMeshRotation(BlendedPose);
+
+			// Apply additive.
+			for (int32 BoneIndex = 0; BoneIndex < NumBones; ++BoneIndex)
+			{
+				const FCompactPoseBoneIndex CompactBoneIndex(BoneIndex);
+				FTransform Additive = AdditivePose[CompactBoneIndex];
+				FTransform::BlendFromIdentityAndAccumulate(BlendedPose[CompactBoneIndex], Additive, ScalarRegister(PerBoneWeights[WeightPoseIndex][BoneIndex]));
+			}
+
+			FAnimationRuntime::ConvertMeshRotationPoseToLocalSpace(BlendedPose);
+		}
+		else
+		{
+			for (int32 BoneIndex = 0; BoneIndex < NumBones; ++BoneIndex)
+			{
+				const FCompactPoseBoneIndex CompactBoneIndex(BoneIndex);
+				FTransform Additive = AdditivePose[CompactBoneIndex];
+				FTransform::BlendFromIdentityAndAccumulate(BlendedPose[CompactBoneIndex], Additive, ScalarRegister(PerBoneWeights[WeightPoseIndex][BoneIndex]));
+			}
+		}
+
+		BlendedPose.NormalizeRotations();
+	} // For each additive pose.
+	
+	//------------------------------------------
+	// Blend curves and attributes.
+	//------------------------------------------
+	// Non-additives.
+	if (BlendingCurves.Num() > 0)
+	{
+		BlendCurves(BlendingCurves, BlendingWeights, OutBlendedAnimationPoseData.GetCurve());
+	}
+
+	if (BlendingAttributes.Num() > 0)
+	{
+		FCustomAttributesRuntime::BlendAttributes(BlendingAttributes, BlendingWeights, OutBlendedAnimationPoseData.GetAttributes());
+	}
+
+	// Additives.
+	for (int32 PoseIndex = 0; PoseIndex < Poses.Num(); ++PoseIndex)
+	{
+		FSlotEvaluationPose& AdditivePose = Poses[PoseIndex];
+		const FAnimationPoseData AdditiveAnimationPoseData(AdditivePose);
+		OutBlendedAnimationPoseData.GetCurve().Accumulate(AdditiveAnimationPoseData.GetCurve(), AdditivePose.Weight);
+		FCustomAttributesRuntime::AccumulateAttributes(AdditiveAnimationPoseData.GetAttributes(), OutBlendedAnimationPoseData.GetAttributes(), AdditivePose.Weight);
+	}
+}
+
 void FAnimInstanceProxy::SlotEvaluatePose(const FName& SlotNodeName, const FAnimationPoseData& SourceAnimationPoseData, float InSourceWeight, FAnimationPoseData& OutBlendedAnimationPoseData, float InBlendWeight, float InTotalNodeWeight)
 {
 	const FCompactPose& SourcePose = SourceAnimationPoseData.GetPose();
@@ -1513,9 +1853,9 @@ void FAnimInstanceProxy::SlotEvaluatePose(const FName& SlotNodeName, const FAnim
 	FCompactPose& BlendedPose = OutBlendedAnimationPoseData.GetPose();
 	FBlendedCurve& BlendedCurve = OutBlendedAnimationPoseData.GetCurve();
 	FStackCustomAttributes& BlendedAttributes = OutBlendedAnimationPoseData.GetAttributes();
-	
-	//Accessing MontageInstances from this function is not safe (as this can be called during Parallel Anim Evaluation!
-	//Any montage data you need to add should be part of MontageEvaluationData
+
+	// Accessing MontageInstances from this function is not safe (as this can be called during Parallel Anim Evaluation!
+	// Any montage data you need to add should be part of MontageEvaluationData
 	// nothing to blend, just get it out
 	if (InBlendWeight <= ZERO_ANIMWEIGHT_THRESH)
 	{
@@ -1525,15 +1865,35 @@ void FAnimInstanceProxy::SlotEvaluatePose(const FName& SlotNodeName, const FAnim
 		return;
 	}
 
+	// Check if we are blending a montage using a blend profile, if so take a special code path for this.
+	for (const FMontageEvaluationState& EvalState : MontageEvaluationData)
+	{
+		if (EvalState.Montage.IsValid() && EvalState.ActiveBlendProfile && EvalState.Montage->IsValidSlot(SlotNodeName))
+		{
+			SlotEvaluatePoseWithBlendProfiles(SlotNodeName, SourceAnimationPoseData, InSourceWeight, OutBlendedAnimationPoseData, InBlendWeight);
+			return;
+		}
+	}
+
 	// Split our data into additive and non additive.
-	TArray<FSlotEvaluationPose> AdditivePoses;
-	TArray<FSlotEvaluationPose> NonAdditivePoses;
+	FBlendProfileScratchData& BlendProfileScratchData = FBlendProfileScratchData::Get();
+	TArray<FSlotEvaluationPose>& AdditivePoses = BlendProfileScratchData.Poses;
+	TArray<FSlotEvaluationPose>& NonAdditivePoses = BlendProfileScratchData.AdditivePoses;
+	TArray<float, TInlineAllocator<8>>& BlendingWeights = BlendProfileScratchData.BlendingWeights;
+	TArray<const FCompactPose*, TInlineAllocator<8>>& BlendingPoses = BlendProfileScratchData.BlendingPoses;
+	TArray<const FBlendedCurve*, TInlineAllocator<8>>& BlendingCurves = BlendProfileScratchData.BlendingCurves;
+	TArray<const FStackCustomAttributes*, TInlineAllocator<8>>& BlendingAttributes = BlendProfileScratchData.BlendingAttributes;
+	NonAdditivePoses.Reset();
+	AdditivePoses.Reset();
+	BlendingWeights.Reset();
+	BlendingPoses.Reset();
+	BlendingCurves.Reset();
+	BlendingAttributes.Reset();
 
 	// first pass we go through collect weights and valid montages. 
 #if DEBUG_MONTAGEINSTANCE_WEIGHT
 	float TotalWeight = 0.f;
 #endif // DEBUG_MONTAGEINSTANCE_WEIGHT
-
 	for (const FMontageEvaluationState& EvalState : MontageEvaluationData)
 	{
 		// If MontageEvaluationData is not valid anymore, pass-through AnimSlot.
@@ -1557,7 +1917,8 @@ void FAnimInstanceProxy::SlotEvaluatePose(const FName& SlotNodeName, const FAnim
 				? (AnimTrack->IsRotationOffsetAdditive() ? AAT_RotationOffsetMeshSpace : AAT_LocalSpaceBase)
 				: AAT_None;
 
-			FSlotEvaluationPose NewPose(EvalState.MontageWeight, AdditiveAnimType);
+			const float MontageWeight = EvalState.BlendInfo.GetBlendedValue();
+			FSlotEvaluationPose NewPose(MontageWeight, AdditiveAnimType);
 			
 			// Bone array has to be allocated prior to calling GetPoseFromAnimTrack
 			NewPose.Pose.SetBoneContainer(&RequiredBones);
@@ -1576,8 +1937,9 @@ void FAnimInstanceProxy::SlotEvaluatePose(const FName& SlotNodeName, const FAnim
 			NewPose.Curve.Combine(MontageCurve);
 
 #if DEBUG_MONTAGEINSTANCE_WEIGHT
-			TotalWeight += EvalState.MontageWeight;
+			TotalWeight += MontageWeight;
 #endif // DEBUG_MONTAGEINSTANCE_WEIGHT
+
 			if (AdditiveAnimType == AAT_None)
 			{
 				NonAdditivePoses.Add(FSlotEvaluationPose(MoveTemp(NewPose)));
@@ -1586,8 +1948,8 @@ void FAnimInstanceProxy::SlotEvaluatePose(const FName& SlotNodeName, const FAnim
 			{
 				AdditivePoses.Add(FSlotEvaluationPose(MoveTemp(NewPose)));
 			}
-		}
-	}
+		} // if IsValidSlot
+	} // for all MontageEvaluationData
 
 	// allocate for blending
 	// If source has any weight, add it to the blend array.
@@ -1597,7 +1959,6 @@ void FAnimInstanceProxy::SlotEvaluatePose(const FName& SlotNodeName, const FAnim
 	ensure (FMath::IsNearlyEqual(InTotalNodeWeight, TotalWeight, KINDA_SMALL_NUMBER));
 #endif // DEBUG_MONTAGEINSTANCE_WEIGHT
 	ensure (InTotalNodeWeight > ZERO_ANIMWEIGHT_THRESH);
-
 	if (InTotalNodeWeight > (1.f + ZERO_ANIMWEIGHT_THRESH))
 	{
 		// Re-normalize additive poses
@@ -1623,22 +1984,13 @@ void FAnimInstanceProxy::SlotEvaluatePose(const FName& SlotNodeName, const FAnim
 			BlendedPose = SourcePose;
 			BlendedCurve = SourceCurve;
 			BlendedAttributes = SourceAttributes;
-		}
-		// Otherwise we need to blend non additive poses together
-		else
+		}		
+		else // Otherwise we need to blend non additive poses together
 		{
-			int32 const NumPoses = NonAdditivePoses.Num() + ((SourceWeight > ZERO_ANIMWEIGHT_THRESH) ? 1 : 0);
-
-			TArray<const FCompactPose*, TInlineAllocator<8>> BlendingPoses;
+			const int32 NumPoses = NonAdditivePoses.Num() + ((SourceWeight > ZERO_ANIMWEIGHT_THRESH) ? 1 : 0);
 			BlendingPoses.AddUninitialized(NumPoses);
-
-			TArray<float, TInlineAllocator<8>> BlendWeights;
-			BlendWeights.AddUninitialized(NumPoses);
-
-			TArray<const FBlendedCurve*, TInlineAllocator<8>> BlendingCurves;
+			BlendingWeights.AddUninitialized(NumPoses);
 			BlendingCurves.AddUninitialized(NumPoses);
-
-			TArray<const FStackCustomAttributes*, TInlineAllocator<8>> BlendingAttributes;
 			BlendingAttributes.AddUninitialized(NumPoses);
 
 			for (int32 Index = 0; Index < NonAdditivePoses.Num(); Index++)
@@ -1646,20 +1998,20 @@ void FAnimInstanceProxy::SlotEvaluatePose(const FName& SlotNodeName, const FAnim
 				BlendingPoses[Index] = &NonAdditivePoses[Index].Pose;
 				BlendingCurves[Index] = &NonAdditivePoses[Index].Curve;
 				BlendingAttributes[Index] = &NonAdditivePoses[Index].Attributes;
-				BlendWeights[Index] = NonAdditivePoses[Index].Weight;
+				BlendingWeights[Index] = NonAdditivePoses[Index].Weight;
 			}
 
 			if (SourceWeight > ZERO_ANIMWEIGHT_THRESH)
 			{
-				int32 const SourceIndex = BlendWeights.Num() - 1;
+				const int32 SourceIndex = NumPoses - 1;
 				BlendingPoses[SourceIndex] = &SourcePose;
 				BlendingCurves[SourceIndex] = &SourceCurve;
 				BlendingAttributes[SourceIndex] = &SourceAttributes;
-				BlendWeights[SourceIndex] = SourceWeight;
+				BlendingWeights[SourceIndex] = SourceWeight;
 			}
 
-			// now time to blend all montages
-			FAnimationRuntime::BlendPosesTogetherIndirect(BlendingPoses, BlendingCurves, BlendingAttributes, BlendWeights, OutBlendedAnimationPoseData);
+			// Blend all montages.
+			FAnimationRuntime::BlendPosesTogetherIndirect(BlendingPoses, BlendingCurves, BlendingAttributes, BlendingWeights, OutBlendedAnimationPoseData);
 		}
 	}
 
@@ -1673,9 +2025,6 @@ void FAnimInstanceProxy::SlotEvaluatePose(const FName& SlotNodeName, const FAnim
 			FAnimationRuntime::AccumulateAdditivePose(OutBlendedAnimationPoseData, AdditiveAnimationPoseData, AdditivePose.Weight, AdditivePose.AdditiveType);
 		}
 	}
-
-	// Normalize rotations after blending/accumulation
-	BlendedPose.NormalizeRotations();
 }
 
 //to debug montage weight
@@ -1692,6 +2041,7 @@ void FAnimInstanceProxy::GetSlotWeight(const FName& SlotNodeName, float& out_Slo
 #if DEBUGMONTAGEWEIGHT
 	float TotalDesiredWeight = 0.f;
 #endif
+
 	// first get all the montage instance weight this slot node has
 	for (const FMontageEvaluationState& EvalState : MontageEvaluationData)
 	{
@@ -1700,18 +2050,19 @@ void FAnimInstanceProxy::GetSlotWeight(const FName& SlotNodeName, float& out_Slo
 			const UAnimMontage* const Montage = EvalState.Montage.Get();
 			if (Montage->IsValidSlot(SlotNodeName))
 			{
-				NewSlotNodeWeight += EvalState.MontageWeight;
+				const float MontageWeight = EvalState.BlendInfo.GetBlendedValue();
+				NewSlotNodeWeight += MontageWeight;
 				if (!Montage->IsValidAdditiveSlot(SlotNodeName))
 				{
-					NonAdditiveTotalWeight += EvalState.MontageWeight;
+					NonAdditiveTotalWeight += MontageWeight;
 				}
 
 #if DEBUGMONTAGEWEIGHT			
 				TotalDesiredWeight += EvalState->DesiredWeight;
 #endif
 #if ENABLE_ANIM_LOGGING
-			UE_LOG(LogAnimation, Verbose, TEXT("GetSlotWeight : Owner: %s, AnimMontage: %s,  (DesiredWeight:%0.2f, Weight:%0.2f)"),
-						*GetActorName(), *EvalState.Montage->GetName(), EvalState.DesiredWeight, EvalState.MontageWeight);
+				UE_LOG(LogAnimation, Verbose, TEXT("GetSlotWeight : Owner: %s, AnimMontage: %s,  (DesiredWeight:%0.2f, Weight:%0.2f)"),
+							*GetActorName(), *EvalState.Montage->GetName(), EvalState.BlendInfo.GetDesiredValue(), MontageWeight);
 #endif
 			}
 		}
diff --git a/Engine/Source/Runtime/Engine/Private/Animation/AnimMontage.cpp b/Engine/Source/Runtime/Engine/Private/Animation/AnimMontage.cpp
--- a/Engine/Source/Runtime/Engine/Private/Animation/AnimMontage.cpp
+++ b/Engine/Source/Runtime/Engine/Private/Animation/AnimMontage.cpp
@@ -44,6 +44,9 @@ UAnimMontage::UAnimMontage(const FObjectInitializer& ObjectInitializer)
 	bEnableAutoBlendOut = true;
 	SyncSlotIndex = 0;
 
+	BlendProfileIn = nullptr;
+	BlendProfileOut = nullptr;
+
 	BlendInTime_DEPRECATED = -1.f;
 	BlendOutTime_DEPRECATED = -1.f;
 
@@ -1300,7 +1303,10 @@ FAnimMontageInstance::FAnimMontageInstance()
 	, NotifyWeight(0.f)
 	, DeltaMoved(0.f)
 	, PreviousPosition(0.f)
+	, BlendStartAlpha(0.0f)
 	, SyncGroupName(NAME_None)
+	, ActiveBlendProfile(nullptr)
+	, ActiveBlendProfileMode(EBlendProfileMode::TimeFactor)
 	, DisableRootMotionCount(0)
 	, MontageSyncLeader(NULL)
 	, MontageSyncUpdateFrameCounter(INDEX_NONE)
@@ -1322,7 +1328,10 @@ FAnimMontageInstance::FAnimMontageInstance(UAnimInstance * InAnimInstance)
 	, NotifyWeight(0.f)
 	, DeltaMoved(0.f)
 	, PreviousPosition(0.f)
+	, BlendStartAlpha(0.0f)
 	, SyncGroupName(NAME_None)
+	, ActiveBlendProfile(nullptr)
+	, ActiveBlendProfileMode(EBlendProfileMode::TimeFactor)
 	, DisableRootMotionCount(0)
 	, MontageSyncLeader(NULL)
 	, MontageSyncUpdateFrameCounter(INDEX_NONE)
@@ -1344,11 +1353,13 @@ void FAnimMontageInstance::Play(float InPlayRate, const FAlphaBlend& BlendIn)
 	
 	// set blend option
 	float CurrentWeight = Blend.GetBlendedValue();
-
-	InitializeBlend(BlendIn);
+	InitializeBlend(BlendIn);	
+	BlendStartAlpha = Blend.GetAlpha();
 	Blend.SetBlendTime(BlendIn.GetBlendTime() * DefaultBlendTimeMultiplier);
 	Blend.SetValueRange(CurrentWeight, 1.f);
 	bEnableAutoBlendOut = Montage->bEnableAutoBlendOut;
+
+	ActiveBlendProfile = Montage->BlendProfileIn;
 }
 
 void FAnimMontageInstance::InitializeBlend(const FAlphaBlend& InAlphaBlend)
@@ -1379,7 +1390,9 @@ void FAnimMontageInstance::Stop(const FAlphaBlend& InBlendOut, bool bInterrupt)
 		// do not use default Montage->BlendOut 
 		// depending on situation, the BlendOut time can change 
 		InitializeBlend(InBlendOut);
+		BlendStartAlpha = Blend.GetAlpha();
 		Blend.SetDesiredValue(0.f);
+		Blend.Update(0.0f);
 
 		if(Montage)
 		{
@@ -1402,6 +1415,7 @@ void FAnimMontageInstance::Stop(const FAlphaBlend& InBlendOut, bool bInterrupt)
 			// currently set up blend option, but it might be worse to switch between 
 			// blending out, but it is possible options in the future
 			Blend.SetBlendTime(InBlendOut.GetBlendTime());
+			BlendStartAlpha = Blend.GetAlpha();
 			// have to call this again to restart blending with new blend time
 			// we don't change blend options
 			Blend.SetDesiredValue(0.f);
@@ -1419,6 +1433,7 @@ void FAnimMontageInstance::Stop(const FAlphaBlend& InBlendOut, bool bInterrupt)
 
 	if (Montage != nullptr)
 	{
+		ActiveBlendProfile = Montage->BlendProfileOut;
 		UE_LOG(LogAnimMontage, Verbose, TEXT("Montage.Stop After: AnimMontage: %s,  (DesiredWeight:%0.2f, Weight:%0.2f)"),
 			*Montage->GetName(), GetDesiredWeight(), GetWeight());
 	}
@@ -1439,6 +1454,7 @@ void FAnimMontageInstance::Initialize(class UAnimMontage * InMontage)
 	{
 		Montage = InMontage;
 		SetPosition(0.f);
+		BlendStartAlpha = 0.0f;
 		// initialize Blend
 		Blend.SetValueRange(0.f, 1.0f);
 		RefreshNextPrevSections();
@@ -1534,6 +1550,7 @@ void FAnimMontageInstance::Terminate()
 	Blend.SetCustomCurve(NULL);
 	Blend.SetBlendOption(EAlphaBlendOption::Linear);
 
+	ActiveBlendProfile = nullptr;
 	Montage = nullptr;
 
 	UE_LOG(LogAnimMontage, Verbose, TEXT("Terminating: AnimMontage: %s"), *GetNameSafe(OldMontage));
@@ -1752,6 +1769,7 @@ void FAnimMontageInstance::MontageSync_PerformSyncToLeader()
 	}
 }
 
+
 void FAnimMontageInstance::UpdateWeight(float DeltaTime)
 {
 	if ( IsValid() )
@@ -1761,6 +1779,11 @@ void FAnimMontageInstance::UpdateWeight(float DeltaTime)
 		// update weight
 		Blend.Update(DeltaTime);
 
+		if (Blend.GetBlendTimeRemaining() < 0.0001f)
+		{
+			ActiveBlendProfile = nullptr;
+		}
+
 		// Notify weight is max of previous and current as notify could have come
 		// from any point between now and last tick
 		NotifyWeight = FMath::Max(PreviousWeight, Blend.GetBlendedValue());
@@ -2675,6 +2698,7 @@ UAnimMontage* FAnimMontageInstance::InitializeMatineeControl(FName SlotName, UAn
 			}
 
 			// Something animating this slot that's not us - set weight to be 0 on this slot
+			MontageInstance->BlendStartAlpha = 1.0f;
 			MontageInstance->Blend.SetDesiredValue(0.f);
 			MontageInstance->Blend.SetAlpha(1.f);
 		}
@@ -2718,6 +2742,7 @@ UAnimMontage* FAnimMontageInstance::SetMatineeAnimPositionInner(FName SlotName,
 		// ensure full weighting to this instance
 		AnimMontageInst->Blend.SetDesiredValue(1.f);
 		AnimMontageInst->Blend.SetAlpha(1.f);
+		AnimMontageInst->BlendStartAlpha = 1.0f;
 
 		AnimMontageInst->SetNextPositionWithEvents(InPosition);
 	}
@@ -2781,6 +2806,7 @@ UAnimMontage* FAnimMontageInstance::PreviewMatineeSetAnimPositionInner(FName Slo
 		// ensure full weighting to this instance
 		MontageInstanceToUpdate->Blend.SetDesiredValue(1.f);
 		MontageInstanceToUpdate->Blend.SetAlpha(1.f);
+		MontageInstanceToUpdate->BlendStartAlpha = 1.0f;
 
 		PreviousPosition = AnimInst->Montage_GetPosition(PlayingMontage);
 		AnimInst->Montage_SetPosition(PlayingMontage, InPosition);
@@ -2861,6 +2887,7 @@ UAnimMontage* FAnimMontageInstance::SetSequencerMontagePosition(FName SlotName,
 			// ensure full weighting to this instance
 			MontageInstanceToUpdate->Blend.SetDesiredValue(Weight);
 			MontageInstanceToUpdate->Blend.SetAlpha(Weight);
+			MontageInstanceToUpdate->BlendStartAlpha = MontageInstanceToUpdate->Blend.GetAlpha();
 			
 			if (bInPlaying)
 			{
diff --git a/Engine/Source/Runtime/Engine/Private/Animation/AnimNode_StateMachine.cpp b/Engine/Source/Runtime/Engine/Private/Animation/AnimNode_StateMachine.cpp
--- a/Engine/Source/Runtime/Engine/Private/Animation/AnimNode_StateMachine.cpp
+++ b/Engine/Source/Runtime/Engine/Private/Animation/AnimNode_StateMachine.cpp
@@ -33,7 +33,6 @@ FAnimationActiveTransitionEntry::FAnimationActiveTransitionEntry()
 	, EndNotify(INDEX_NONE)
 	, InterruptNotify(INDEX_NONE)
 	, BlendProfile(nullptr)
-	, BlendProfileMode(EBlendProfileMode::TimeFactor)
 	, BlendOption(EAlphaBlendOption::HermiteCubic)
 	, LogicType(ETransitionLogicType::TLT_StandardBlend)
 	, bActive(false)
@@ -49,7 +48,6 @@ FAnimationActiveTransitionEntry::FAnimationActiveTransitionEntry(int32 NextState
 	, EndNotify(ReferenceTransitionInfo.EndNotify)
 	, InterruptNotify(ReferenceTransitionInfo.InterruptNotify)
 	, BlendProfile(ReferenceTransitionInfo.BlendProfile)
-	, BlendProfileMode(ReferenceTransitionInfo.BlendProfileMode)
 	, BlendOption(ReferenceTransitionInfo.BlendMode)
 	, LogicType(ReferenceTransitionInfo.LogicType)
 	, bActive(true)
@@ -113,72 +111,6 @@ void FAnimationActiveTransitionEntry::InitializeCustomGraphLinks(const FAnimatio
 	}
 }
 
-
-void FAnimationActiveTransitionEntry::ProcessBlendProfile(int32 TransitionDirectionIndex)
-{
-	const bool bForwards = (TransitionDirectionIndex == 0);
-	FBlendSampleData& CurrentData = StateBlendData[TransitionDirectionIndex];
-
-	switch (BlendProfileMode)
-	{
-		// The per bone value is a factor of the transition time, where 0.5 means half the transition time, 0.1 means one tenth of the transition time, etc.
-		case EBlendProfileMode::TimeFactor:
-		{
-			for (int32 PerBoneIndex = 0; PerBoneIndex < CurrentData.PerBoneBlendData.Num(); ++PerBoneIndex)
-			{
-				// Calculate the blend weight for this bone.
-				const float BoneTimeScaleFactor = FMath::Clamp(BlendProfile->GetEntryBlendScale(PerBoneIndex), 0.0f, 1.0f);
-				const float BoneTransitionDuration = BoneTimeScaleFactor * CrossfadeDuration;
-				float NormalizedBoneWeight = (BoneTransitionDuration > ZERO_ANIMWEIGHT_THRESH) ? ElapsedTime / BoneTransitionDuration : 1.0f;
-				NormalizedBoneWeight = FMath::Clamp(NormalizedBoneWeight, 0.0f, 1.0f);
-				NormalizedBoneWeight = FAlphaBlend::AlphaToBlendOption(NormalizedBoneWeight, Blend.GetBlendOption(), Blend.GetCustomCurve());
-				if (NormalizedBoneWeight < ZERO_ANIMWEIGHT_THRESH)
-				{
-					NormalizedBoneWeight = ZERO_ANIMWEIGHT_THRESH;
-				}
-
-				// Reverse the weight when we go backward.
-				if (!bForwards)
-				{
-					NormalizedBoneWeight = 1.0f - NormalizedBoneWeight;
-				}
-
-				CurrentData.PerBoneBlendData[PerBoneIndex] = NormalizedBoneWeight;
-			}
-		}
-		break;
-
-		// The per bone value is a factor of the transition blend weight.
-		case EBlendProfileMode::WeightFactor:
-		{
-			for (int32 PerBoneIndex = 0; PerBoneIndex < CurrentData.PerBoneBlendData.Num(); ++PerBoneIndex)
-			{
-				// Get the bone weight factor from the profile.
-				float BoneWeightFactor = BlendProfile->GetEntryBlendScale(PerBoneIndex);
-				if (!bForwards)
-				{
-					BoneWeightFactor = (BoneWeightFactor > ZERO_ANIMWEIGHT_THRESH) ? 1.0f / BoneWeightFactor : 1.0f;
-				}
-
-				// Multiply it with the transition weight and make sure it isn't zero, as this can cause normalization issues.
-				BoneWeightFactor *= CurrentData.TotalWeight;
-				if (BoneWeightFactor < ZERO_ANIMWEIGHT_THRESH)
-				{
-					BoneWeightFactor = ZERO_ANIMWEIGHT_THRESH;
-				}
-
-				CurrentData.PerBoneBlendData[PerBoneIndex] = BoneWeightFactor;
-			}
-		}
-		break;
-
-		default:
-			checkf(false, TEXT("The selected Blend Profile Mode is not supported (Mode=%d)"), BlendProfileMode);
-			break;
-	}
-}
-
-
 void FAnimationActiveTransitionEntry::Update(const FAnimationUpdateContext& Context, int32 CurrentStateIndex, bool& bOutFinished)
 {
 	bOutFinished = false;
@@ -211,7 +143,7 @@ void FAnimationActiveTransitionEntry::Update(const FAnimationUpdateContext& Cont
 			{
 				const bool bForwards = (Idx == 0);
 				StateBlendData[Idx].TotalWeight = bForwards ? Alpha : 1.0f - Alpha;
-				ProcessBlendProfile(Idx);
+				BlendProfile->UpdateBoneWeights(StateBlendData[Idx], Blend, 0.0f, StateBlendData[Idx].TotalWeight, !bForwards);
 			}
 
 			FBlendSampleData::NormalizeDataWeight(StateBlendData);
diff --git a/Engine/Source/Runtime/Engine/Private/Animation/AnimationRuntime.cpp b/Engine/Source/Runtime/Engine/Private/Animation/AnimationRuntime.cpp
--- a/Engine/Source/Runtime/Engine/Private/Animation/AnimationRuntime.cpp
+++ b/Engine/Source/Runtime/Engine/Private/Animation/AnimationRuntime.cpp
@@ -138,7 +138,7 @@ FORCEINLINE void BlendCurves(const TArrayView<const FBlendedCurve> SourceCurves,
 	}
 }
 
-FORCEINLINE void BlendCurves(const TArrayView<const FBlendedCurve* const> SourceCurves, const TArrayView<const float> SourceWeights, FBlendedCurve& OutCurve)
+void BlendCurves(const TArrayView<const FBlendedCurve* const> SourceCurves, const TArrayView<const float> SourceWeights, FBlendedCurve& OutCurve)
 {
 	if(SourceCurves.Num() > 0)
 	{
@@ -1837,7 +1837,7 @@ void FAnimationRuntime::CreateMaskWeights(TArray<FPerBoneBlendWeight>& BoneBlend
 		{
 			const UBlendProfile* BlendMask = BlendMasks[MaskIndex];
 
-			if (!BlendMask || BlendMask->BlendProfileMode != EBlendProfileMode::BlendMask)
+			if (!BlendMask || BlendMask->Mode != EBlendProfileMode::BlendMask)
 			{
 				ensureMsgf(false, TEXT("FAnimationRuntime::CreateMaskWeights BlendMask null or BlendProfile mode is not blend mask.  BlendProfile=%s"), *GetNameSafe(BlendMask));
 				continue;
diff --git a/Engine/Source/Runtime/Engine/Private/Animation/BlendProfile.cpp b/Engine/Source/Runtime/Engine/Private/Animation/BlendProfile.cpp
--- a/Engine/Source/Runtime/Engine/Private/Animation/BlendProfile.cpp
+++ b/Engine/Source/Runtime/Engine/Private/Animation/BlendProfile.cpp
@@ -4,7 +4,7 @@
 
 UBlendProfile::UBlendProfile()
 	: OwningSkeleton(nullptr)
-	, BlendProfileMode(EBlendProfileMode::WeightFactor)
+	, Mode(EBlendProfileMode::WeightFactor)
 {
 	// Set up our owning skeleton and initialise bone references
 	if(USkeleton* OuterAsSkeleton = Cast<USkeleton>(GetOuter()))
@@ -129,7 +129,7 @@ float UBlendProfile::GetEntryBlendScale(const int32 InEntryIdx) const
 	{
 		return ProfileEntries[InEntryIdx].BlendScale;
 	}
-	// No overriden blend scale, return no scale
+	// No overridden blend scale, return no scale
 	return GetDefaultBlendScale();
 }
 
@@ -171,3 +171,96 @@ void UBlendProfile::SetSingleBoneBlendScale(int32 InBoneIdx, float InScale, bool
 		}
 	}
 }
+
+void UBlendProfile::FillBoneScalesArray(TArray<float>& OutBoneBlendProfileFactors, const FBoneContainer& BoneContainer) const
+{
+	const int32 NumBones = BoneContainer.GetCompactPoseNumBones();
+	OutBoneBlendProfileFactors.Reset(NumBones);
+	OutBoneBlendProfileFactors.AddUninitialized(NumBones);
+
+	// Fill the bone values with defaults values.
+	for (int32 Index = 0; Index < NumBones; ++Index)
+	{
+		OutBoneBlendProfileFactors[Index] = 1.0f;
+	}
+
+	// Overwrite the values of the bones that are inside the blend profile.
+	// Since the bones in the blend profile are stored as skeleton indices we need to remap them into our compact pose.
+	for (int32 Index = 0; Index < ProfileEntries.Num(); Index++)
+	{
+		const int32 SkeletonBoneIndex = ProfileEntries[Index].BoneReference.BoneIndex;
+		const FCompactPoseBoneIndex PoseBoneIndex = BoneContainer.GetCompactPoseIndexFromSkeletonIndex(SkeletonBoneIndex);
+		OutBoneBlendProfileFactors[PoseBoneIndex.GetInt()] = GetEntryBlendScale(Index);
+	}
+}
+
+float UBlendProfile::CalculateBoneWeight(float BoneFactor, EBlendProfileMode Mode, const FAlphaBlend& BlendInfo, float BlendStartAlpha, float MainWeight, bool bInverse)
+{
+	switch (Mode)
+	{
+		// The per bone value is a factor of the transition time, where 0.5 means half the transition time, 0.1 means one tenth of the transition time, etc.
+		case EBlendProfileMode::TimeFactor:
+		{
+			// Most bones will have a bone factor of 1, so let's optimize that case.
+			// Basically it means it will just follow the main weight.
+			if (BoneFactor >= 1.0f - ZERO_ANIMWEIGHT_THRESH)
+			{
+				return !bInverse ? MainWeight : 1.0f - MainWeight;
+			}
+
+			// Make sure our input values are valid, which is between 0 and 1.
+			const float ClampedFactor = FMath::Clamp(BoneFactor, 0.0f, 1.0f);
+
+			// Calculate where blend begin value is for this specific bone. So where did our blend start from?
+			// Note that this isn't just the BlendInfo.GetBlendedValue() because it can be different per bone as some bones are further ahead in time.
+			// We also need to sample the actual curve for this to get the real value.
+			const float BeginValue = (ClampedFactor > ZERO_ANIMWEIGHT_THRESH) ? FMath::Clamp(BlendStartAlpha / ClampedFactor, 0.0f, 1.0f) : 1.0f;
+			const float RealBeginValue = FAlphaBlend::AlphaToBlendOption(BeginValue, BlendInfo.GetBlendOption(), BlendInfo.GetCustomCurve());
+
+			// Calculate the current alpha value for the bone.
+			// As some bones can blend faster than others, we basically scale the current blend's alpha by the bone's factor.
+			// After that we sample the curve to get the real alpha blend value.
+			const float LinearAlpha = (ClampedFactor > ZERO_ANIMWEIGHT_THRESH) ? FMath::Clamp(BlendInfo.GetAlpha() / ClampedFactor, 0.0f, 1.0f) : 1.0f;
+			const float RealBoneAlpha = FAlphaBlend::AlphaToBlendOption(LinearAlpha, BlendInfo.GetBlendOption(), BlendInfo.GetCustomCurve());
+
+			// Now that we know the alpha for our blend, we can calculate the actual weight value.
+			// Also make sure the bone weight is valid. Values can't be zero because this could introduce issues during normalization internally in the pipeline.
+			const float BoneWeight = RealBeginValue + RealBoneAlpha * (BlendInfo.GetDesiredValue() - RealBeginValue);
+			const float ClampedBoneWeight = FMath::Clamp(BoneWeight, ZERO_ANIMWEIGHT_THRESH, 1.0f);
+
+			// Return our calculated weight, depending whether we'd like to invert it or not.
+			return !bInverse ? ClampedBoneWeight : (1.0f - ClampedBoneWeight);
+		}
+
+		// The per bone value is a factor of the main blend's weight.
+		case EBlendProfileMode::WeightFactor:
+		{
+			if (!bInverse)
+			{
+				return FMath::Clamp(MainWeight * BoneFactor, ZERO_ANIMWEIGHT_THRESH, 1.0f);
+			}
+
+			// We're inversing.
+			const float Weight = (BoneFactor > ZERO_ANIMWEIGHT_THRESH) ? MainWeight / BoneFactor : 1.0f;
+			return FMath::Clamp(Weight, ZERO_ANIMWEIGHT_THRESH, 1.0f);
+		}
+
+		// Handle unsupported modes.
+		// If you reach this point you have to add another case statement for your newly added blend profile mode.
+		default:
+		{
+			checkf(false, TEXT("The selected Blend Profile Mode is not supported (Mode=%d)"), Mode);
+			break;
+		}
+	}
+
+	return MainWeight;
+}
+
+void UBlendProfile::UpdateBoneWeights(FBlendSampleData& InOutCurrentData, const FAlphaBlend& BlendInfo, float BlendStartAlpha, float MainWeight, bool bInverse)
+{
+	for (int32 PerBoneIndex = 0; PerBoneIndex < InOutCurrentData.PerBoneBlendData.Num(); ++PerBoneIndex)
+	{
+		InOutCurrentData.PerBoneBlendData[PerBoneIndex] = CalculateBoneWeight(GetEntryBlendScale(PerBoneIndex), Mode, BlendInfo, BlendStartAlpha, MainWeight, bInverse);
+	}
+}
diff --git a/Engine/Source/Runtime/Engine/Private/Animation/BlendProfileScratchData.h b/Engine/Source/Runtime/Engine/Private/Animation/BlendProfileScratchData.h
new file mode 100644
--- /dev/null
+++ b/Engine/Source/Runtime/Engine/Private/Animation/BlendProfileScratchData.h
@@ -0,0 +1,34 @@
+// Copyright Epic Games, Inc. All Rights Reserved.
+
+#pragma once
+
+#include "CoreMinimal.h"
+#include "HAL/ThreadSingleton.h"
+
+struct FCompactPose;
+struct FBlendedCurve;
+struct FStackCustomAttributes;
+struct FSlotEvaluationPose;
+
+/**
+ * Temporary data that we use during blending of anim montages that use blend profiles.
+ * We use this in order to prevent runtime allocations.
+ */
+struct FBlendProfileScratchData : public TThreadSingleton<FBlendProfileScratchData>
+{
+	TArray<TArray<float>> PerBoneWeights;		// A set of bone weights, per montage instance. Index this like [PoseIndex][CompactPoseBoneIndex].
+	TArray<float> PerBoneWeightTotals;			// The bone weight totals for non-additive poses, used for normalizing weights.
+	TArray<float> PerBoneWeightTotalsAdditive;	// The bone weight totals for additive poses, used for normalizing weights.
+	TArray<float> BoneBlendProfileScales;		// The bone profile scale values.
+
+	TArray<FSlotEvaluationPose> Poses;			// Non additive poses.
+	TArray<FSlotEvaluationPose> AdditivePoses;	// Additive poses.
+
+	TArray<uint8, TInlineAllocator<8>> PoseIndices;			// The indices inside the PerBoneWeights array, for non additive poses.
+	TArray<uint8, TInlineAllocator<8>> AdditivePoseIndices;	// The indices inside the PerBoneWeights array, for additive poses.
+
+	TArray<float, TInlineAllocator<8>> BlendingWeights;					// The per pose blend weights.
+	TArray<const FCompactPose*, TInlineAllocator<8>> BlendingPoses;		// The non additive poses to blend.
+	TArray<const FBlendedCurve*, TInlineAllocator<8>> BlendingCurves;	// The curves to blend.
+	TArray<const FStackCustomAttributes*, TInlineAllocator<8>> BlendingAttributes;	// The attributes to blend.
+};
diff --git a/Engine/Source/Runtime/Engine/Public/Animation/AnimInstanceProxy.h b/Engine/Source/Runtime/Engine/Public/Animation/AnimInstanceProxy.h
--- a/Engine/Source/Runtime/Engine/Public/Animation/AnimInstanceProxy.h
+++ b/Engine/Source/Runtime/Engine/Public/Animation/AnimInstanceProxy.h
@@ -839,6 +839,9 @@ protected:
 	static void ResetCounterInputProxy(FAnimInstanceProxy* InputProxy);
 
 private:
+	/** Evaluate the slot when there are blend spaces involved in any of the active anim montages. */
+	void SlotEvaluatePoseWithBlendProfiles(const FName& SlotNodeName, const FAnimationPoseData& SourceAnimationPoseData, float InSourceWeight, FAnimationPoseData& OutBlendedAnimationPoseData, float InBlendWeight);
+
 	/** The component to world transform of the component we are running on */
 	FTransform ComponentTransform;
 
diff --git a/Engine/Source/Runtime/Engine/Public/AnimationRuntime.h b/Engine/Source/Runtime/Engine/Public/AnimationRuntime.h
--- a/Engine/Source/Runtime/Engine/Public/AnimationRuntime.h
+++ b/Engine/Source/Runtime/Engine/Public/AnimationRuntime.h
@@ -58,6 +58,8 @@ FORCEINLINE void BlendTransform<ETransformBlendMode::Accumulate>(const FTransfor
 }
 
 FORCEINLINE void BlendCurves(const TArrayView<const FBlendedCurve> SourceCurves, const TArrayView<const float> SourceWeights, FBlendedCurve& OutCurve);
+void BlendCurves(const TArrayView<const FBlendedCurve* const> SourceCurves, const TArrayView<const float> SourceWeights, FBlendedCurve& OutCurve);
+
 
 /////////////////////////////////////////////////////////
 /** Interface used to provide interpolation indices for per bone blends
