Added support for UVolumeTexture and UTexture2DArray streaming.
Texture 2D array streaming is implemented through FTexture2DArrayMipAllocator_Reallocate
Volume texture streaming is implemented through FVolumeTextureMipAllocator_Reallocate.
New generic stream out implementation in FTextureStreamOut.

Fixed different issues in FTexture2DMipDataProvider_IO relatively to optional bulks and IO errors.

Added UTexture::GetSurfaceDepth() for 3d textures.
Added UTexture::GetSurfaceArraySize() for array textures (where cubemaps return 6).

Moved UTexture2D::CancelPendingTextureStreaming() and UTexture2D::LevelIndex to UTexture.

Renamed UTexture2DArray::GetNumSlices() to UTexture2DArray::GetArraySize().

Fixed missing "virtual" keyword in base class UTextureMipDataProviderFactory.

Updated the texture editor UI so that it shows Texture2DArray array sizes and Texture3D depth.
Unified the UI for cubemaps so that it shows as an array with size 6.

Reimplemented "LODGroups" command to make it coherent with latest changes.

diff --git a/Engine/Plugins/FX/Niagara/Source/Niagara/Private/NiagaraDataInterface2DArrayTexture.cpp b/Engine/Plugins/FX/Niagara/Source/Niagara/Private/NiagaraDataInterface2DArrayTexture.cpp
--- a/Engine/Plugins/FX/Niagara/Source/Niagara/Private/NiagaraDataInterface2DArrayTexture.cpp
+++ b/Engine/Plugins/FX/Niagara/Source/Niagara/Private/NiagaraDataInterface2DArrayTexture.cpp
@@ -124,7 +124,7 @@ void UNiagaraDataInterface2DArrayTexture::GetTextureDimensions(FVectorVMContext&
 	{
 		TextureDimensions.X = Texture->GetSizeX();
 		TextureDimensions.Y = Texture->GetSizeY();
-		TextureDimensions.Z = Texture->GetNumSlices();
+		TextureDimensions.Z = Texture->GetArraySize();
 	}
 
 	for (int32 i = 0; i < Context.NumInstances; ++i)
@@ -237,7 +237,7 @@ void UNiagaraDataInterface2DArrayTexture::PushToRenderThreadImpl()
 	{
 		RT_TexDims.X = Texture->GetSizeX();
 		RT_TexDims.Y = Texture->GetSizeY();
-		RT_TexDims.Z = Texture->GetNumSlices();
+		RT_TexDims.Z = Texture->GetArraySize();
 	}
 
 	ENQUEUE_RENDER_COMMAND(FPushDITextureToRT)
diff --git a/Engine/Source/Editor/TextureEditor/Private/Models/TextureEditorViewportClient.cpp b/Engine/Source/Editor/TextureEditor/Private/Models/TextureEditorViewportClient.cpp
--- a/Engine/Source/Editor/TextureEditor/Private/Models/TextureEditorViewportClient.cpp
+++ b/Engine/Source/Editor/TextureEditor/Private/Models/TextureEditorViewportClient.cpp
@@ -89,17 +89,14 @@ void FTextureEditorViewportClient::Draw(FViewport* Viewport, FCanvas* Canvas)
 	UTextureRenderTargetVolume* RTTextureVolume = Cast<UTextureRenderTargetVolume>(Texture);
 
 	// Fully stream in the texture before drawing it.
-	if (Texture2D)
-	{
-		Texture2D->SetForceMipLevelsToBeResident(30.0f);
-		Texture2D->WaitForStreaming();
-	}
+	Texture->SetForceMipLevelsToBeResident(30.0f);
+	Texture->WaitForStreaming();
 
 	TextureEditorPinned->PopulateQuickInfo();
 
 	// Figure out the size we need
-	uint32 Width, Height;
-	TextureEditorPinned->CalculateTextureDimensions(Width, Height);
+	uint32 Width, Height, Depth, ArraySize;;
+	TextureEditorPinned->CalculateTextureDimensions(Width, Height, Depth, ArraySize);
 	const float MipLevel = (float)TextureEditorPinned->GetMipLevel();
 	const float LayerIndex = (float)TextureEditorPinned->GetLayer();
 
@@ -314,9 +311,8 @@ bool FTextureEditorViewportClient::InputAxis(FViewport* Viewport, int32 Controll
 		{
 			TSharedPtr<STextureEditorViewport> EditorViewport = TextureEditorViewportPtr.Pin();
 
-			uint32 Height = 1;
-			uint32 Width = 1;
-			TextureEditorPtr.Pin()->CalculateTextureDimensions(Width, Height);
+			uint32 Width, Height, Depth, ArraySize;
+			TextureEditorPtr.Pin()->CalculateTextureDimensions(Width, Height, Depth, ArraySize);
 
 			if (Key == EKeys::MouseY)
 			{
@@ -388,10 +384,21 @@ void FTextureEditorViewportClient::ModifyCheckerboardTextureColors()
 
 FText FTextureEditorViewportClient::GetDisplayedResolution() const
 {
-	uint32 Height = 1;
-	uint32 Width = 1;
-	TextureEditorPtr.Pin()->CalculateTextureDimensions(Width, Height);
-	return FText::Format( NSLOCTEXT("TextureEditor", "DisplayedResolution", "Displayed: {0}x{1}"), FText::AsNumber( FMath::Max((uint32)1, Width) ), FText::AsNumber( FMath::Max((uint32)1, Height)) );
+	// Zero is the default size 
+	uint32 Height, Width, Depth, ArraySize;
+	TextureEditorPtr.Pin()->CalculateTextureDimensions(Width, Height, Depth, ArraySize);
+	if (Depth > 0)
+	{
+		return FText::Format(NSLOCTEXT("TextureEditor", "DisplayedResolution", "Displayed: {0}x{1}x{2}"), FText::AsNumber(FMath::Max((uint32)1, Width)), FText::AsNumber(FMath::Max((uint32)1, Height)), FText::AsNumber(Depth));
+	}
+	else if (ArraySize > 0)
+	{
+		return FText::Format(NSLOCTEXT("TextureEditor", "DisplayedResolution", "Displayed: {0}x{1}*{2}"), FText::AsNumber(FMath::Max((uint32)1, Width)), FText::AsNumber(FMath::Max((uint32)1, Height)), FText::AsNumber(ArraySize));
+	}
+	else
+	{
+		return FText::Format(NSLOCTEXT("TextureEditor", "DisplayedResolution", "Displayed: {0}x{1}"), FText::AsNumber(FMath::Max((uint32)1, Width)), FText::AsNumber(FMath::Max((uint32)1, Height)));
+	}
 }
 
 
@@ -402,7 +409,8 @@ float FTextureEditorViewportClient::GetViewportVerticalScrollBarRatio() const
 	float WidgetHeight = 1.0f;
 	if (TextureEditorViewportPtr.Pin()->GetVerticalScrollBar().IsValid())
 	{
-		TextureEditorPtr.Pin()->CalculateTextureDimensions(Width, Height);
+		uint32 Depth, ArraySize;
+		TextureEditorPtr.Pin()->CalculateTextureDimensions(Width, Height, Depth, ArraySize);
 
 		WidgetHeight = TextureEditorViewportPtr.Pin()->GetViewport()->GetSizeXY().Y;
 	}
@@ -418,7 +426,8 @@ float FTextureEditorViewportClient::GetViewportHorizontalScrollBarRatio() const
 	float WidgetWidth = 1.0f;
 	if (TextureEditorViewportPtr.Pin()->GetHorizontalScrollBar().IsValid())
 	{
-		TextureEditorPtr.Pin()->CalculateTextureDimensions(Width, Height);
+		uint32 Depth, ArraySize;
+		TextureEditorPtr.Pin()->CalculateTextureDimensions(Width, Height, Depth, ArraySize);
 
 		WidgetWidth = TextureEditorViewportPtr.Pin()->GetViewport()->GetSizeXY().X;
 	}
@@ -472,14 +481,14 @@ FVector2D FTextureEditorViewportClient::GetViewportScrollBarPositions() const
 	FVector2D Positions = FVector2D::ZeroVector;
 	if (TextureEditorViewportPtr.Pin()->GetVerticalScrollBar().IsValid() && TextureEditorViewportPtr.Pin()->GetHorizontalScrollBar().IsValid())
 	{
-		uint32 Width, Height;
+		uint32 Width, Height, Depth, ArraySize;
 		UTexture* Texture = TextureEditorPtr.Pin()->GetTexture();
 		float VRatio = GetViewportVerticalScrollBarRatio();
 		float HRatio = GetViewportHorizontalScrollBarRatio();
 		float VDistFromBottom = TextureEditorViewportPtr.Pin()->GetVerticalScrollBar()->DistanceFromBottom();
 		float HDistFromBottom = TextureEditorViewportPtr.Pin()->GetHorizontalScrollBar()->DistanceFromBottom();
 	
-		TextureEditorPtr.Pin()->CalculateTextureDimensions(Width, Height);
+		TextureEditorPtr.Pin()->CalculateTextureDimensions(Width, Height, Depth, ArraySize);
 
 		if ((TextureEditorViewportPtr.Pin()->GetVerticalScrollBar()->GetVisibility() == EVisibility::Visible) && VDistFromBottom < 1.0f)
 		{
diff --git a/Engine/Source/Editor/TextureEditor/Private/TextureEditorToolkit.cpp b/Engine/Source/Editor/TextureEditor/Private/TextureEditorToolkit.cpp
--- a/Engine/Source/Editor/TextureEditor/Private/TextureEditorToolkit.cpp
+++ b/Engine/Source/Editor/TextureEditor/Private/TextureEditorToolkit.cpp
@@ -236,28 +236,28 @@ void FTextureEditorToolkit::InitTextureEditor( const EToolkitMode::Type Mode, co
 /* ITextureEditorToolkit interface
  *****************************************************************************/
 
-void FTextureEditorToolkit::CalculateTextureDimensions( uint32& Width, uint32& Height ) const
+void FTextureEditorToolkit::CalculateTextureDimensions( uint32& Width, uint32& Height, uint32& Depth, uint32& ArraySize ) const
 {
-	uint32 ImportedWidth = Texture->Source.GetSizeX();
-	uint32 ImportedHeight = Texture->Source.GetSizeY();
+	Width = Texture->Source.GetSizeX();
+	Height = Texture->Source.GetSizeY();
+	Depth = IsVolumeTexture() ? Texture->Source.GetNumLayers() : 0;
+	ArraySize = (Is2DArrayTexture() || IsCubeTexture()) ? Texture->Source.GetNumLayers() : 0;
 
-	// if Original Width and Height are 0, use the saved current width and height
-	if ((ImportedWidth == 0) && (ImportedHeight == 0))
+	if (!Width && !Height)
 	{
-		ImportedWidth = Texture->GetSurfaceWidth();
-		ImportedHeight = Texture->GetSurfaceHeight();
+		Width = (uint32)Texture->GetSurfaceWidth();
+		Height = (uint32)Texture->GetSurfaceHeight();
+		Depth = (uint32)Texture->GetSurfaceDepth();
+		ArraySize = Texture->GetSurfaceArraySize();
 	}
 
-	Width = ImportedWidth;
-	Height = ImportedHeight;
-
-
 	// catch if the Width and Height are still zero for some reason
 	if ((Width == 0) || (Height == 0))
 	{
 		Width = 0;
 		Height= 0;
-
+		Depth = 0;
+		ArraySize = 0;
 		return;
 	}
 
@@ -446,25 +446,18 @@ void FTextureEditorToolkit::PopulateQuickInfo( )
 	UVolumeTexture* VolumeTexture = Cast<UVolumeTexture>(Texture);
 	UTextureRenderTargetVolume* VolumeTextureRT = Cast<UTextureRenderTargetVolume>(Texture);
 
+	const bool bIsVolume = VolumeTexture || VolumeTextureRT;
+	const bool bIsArray2D = Texture2DArray || Texture2DArrayRT;
+	const bool bIsCube = !!TextureCube;
+
 	const uint32 SurfaceWidth = (uint32)Texture->GetSurfaceWidth();
 	const uint32 SurfaceHeight = (uint32)Texture->GetSurfaceHeight();
-	const uint32 SurfaceDepth =
-		[&]() -> uint32
-		{
-			if (VolumeTexture)
-			{
-				return (uint32)VolumeTexture->GetSizeZ();
-			}
-			else if (VolumeTextureRT)
-			{
-				return (uint32)VolumeTextureRT->SizeZ;
-			}
-			return 1;
-		}();
+	const uint32 SurfaceDepth = (uint32)Texture->GetSurfaceDepth();
+	const uint32 ArraySize = (uint32)Texture->GetSurfaceArraySize();
 
 	const uint32 ImportedWidth = FMath::Max<uint32>(SurfaceWidth, Texture->Source.GetSizeX());
 	const uint32 ImportedHeight =  FMath::Max<uint32>(SurfaceHeight, Texture->Source.GetSizeY());
-	const uint32 ImportedDepth = FMath::Max<uint32>(SurfaceDepth, VolumeTexture || VolumeTextureRT ? Texture->Source.GetNumSlices() : 1);
+	const uint32 ImportedDepth = FMath::Max<uint32>(SurfaceDepth, bIsVolume ? Texture->Source.GetNumSlices() : 0);
 
 	const FStreamableRenderResourceState SRRState = Texture->GetStreamableResourceState();
 	const int32 ActualMipBias = SRRState.IsValid() ? (SRRState.ResidentFirstLODIdx() + SRRState.AssetLODBias) : Texture->GetCachedLODBias();
@@ -499,7 +492,7 @@ void FTextureEditorToolkit::PopulateQuickInfo( )
 	Options.UseGrouping = false;
 
 
-	if (VolumeTexture || VolumeTextureRT)
+	if (bIsVolume)
 	{
 		ImportedText->SetText(FText::Format( NSLOCTEXT("TextureEditor", "QuickInfo_Imported_3x", "Imported: {0}x{1}x{2}"), FText::AsNumber(ImportedWidth, &Options), FText::AsNumber(ImportedHeight, &Options), FText::AsNumber(ImportedDepth, &Options)));
 		CurrentText->SetText(FText::Format( NSLOCTEXT("TextureEditor", "QuickInfo_Displayed_3x", "Displayed: {0}x{1}x{2}"), FText::AsNumber(PreviewEffectiveTextureWidth, &Options ), FText::AsNumber(PreviewEffectiveTextureHeight, &Options), FText::AsNumber(PreviewEffectiveTextureDepth, &Options)));
@@ -519,12 +512,18 @@ void FTextureEditorToolkit::PopulateQuickInfo( )
 			PreviewEffectiveTextureHeight *= (uint32)NumTilesY;
 		}
 	}
+	else if (bIsArray2D)
+	{
+		ImportedText->SetText(FText::Format(NSLOCTEXT("TextureEditor", "QuickInfo_Imported_3x", "Imported: {0}x{1}*{2}"), FText::AsNumber(ImportedWidth, &Options), FText::AsNumber(ImportedHeight, &Options), FText::AsNumber(ArraySize, &Options)));
+		CurrentText->SetText(FText::Format(NSLOCTEXT("TextureEditor", "QuickInfo_Displayed_3x", "Displayed: {0}x{1}*{2}"), FText::AsNumber(PreviewEffectiveTextureWidth, &Options), FText::AsNumber(PreviewEffectiveTextureHeight, &Options), FText::AsNumber(ArraySize, &Options)));
+		MaxInGameText->SetText(FText::Format(NSLOCTEXT("TextureEditor", "QuickInfo_MaxInGame_3x", "Max In-Game: {0}x{1}*{2}"), FText::AsNumber(MaxInGameWidth, &Options), FText::AsNumber(MaxInGameHeight, &Options), FText::AsNumber(ArraySize, &Options)));
+	}
 	else
 	{
 	    FText CubemapAdd;
-	    if(TextureCube)
+	    if(bIsCube)
 	    {
-		    CubemapAdd = NSLOCTEXT("TextureEditor", "QuickInfo_PerCubeSide", "x6 (CubeMap)");
+		    CubemapAdd = NSLOCTEXT("TextureEditor", "QuickInfo_PerCubeSide", "*6 (CubeMap)");
 	    }
     
 	    ImportedText->SetText(FText::Format( NSLOCTEXT("TextureEditor", "QuickInfo_Imported_2x", "Imported: {0}x{1}"), FText::AsNumber(ImportedWidth, &Options), FText::AsNumber(ImportedHeight, &Options)));
@@ -652,8 +651,8 @@ double FTextureEditorToolkit::CalculateDisplayedZoomLevel() const
 		return Zoom;
 	}
 
-	uint32 DisplayWidth, DisplayHeight;
-	CalculateTextureDimensions(DisplayWidth, DisplayHeight);
+	uint32 DisplayWidth, DisplayHeight, DisplayDepth, DisplayArraySize;
+	CalculateTextureDimensions(DisplayWidth, DisplayHeight, DisplayDepth, DisplayArraySize);
 	if (PreviewEffectiveTextureHeight != 0)
 	{
 		return (double)DisplayHeight / PreviewEffectiveTextureHeight;
@@ -1145,6 +1144,15 @@ bool FTextureEditorToolkit::IsCubeTexture( ) const
 	return (Texture->IsA(UTextureCube::StaticClass()) || Texture->IsA(UTextureRenderTargetCube::StaticClass()));
 }
 
+bool FTextureEditorToolkit::IsVolumeTexture() const
+{
+	return (Texture->IsA(UVolumeTexture::StaticClass()) || Texture->IsA(UTextureRenderTargetVolume::StaticClass()));
+}
+
+bool FTextureEditorToolkit::Is2DArrayTexture() const
+{
+	return (Texture->IsA(UTexture2DArray::StaticClass()) || Texture->IsA(UTextureRenderTarget2DArray::StaticClass()));
+}
 
 TSharedRef<SWidget> FTextureEditorToolkit::OnGenerateMipMapLevelMenu()
 {
diff --git a/Engine/Source/Editor/TextureEditor/Private/TextureEditorToolkit.h b/Engine/Source/Editor/TextureEditor/Private/TextureEditorToolkit.h
--- a/Engine/Source/Editor/TextureEditor/Private/TextureEditorToolkit.h
+++ b/Engine/Source/Editor/TextureEditor/Private/TextureEditorToolkit.h
@@ -60,7 +60,7 @@ public:
 	virtual bool UsesCustomToolbarPlacement() const { return false; }
 
 	// ITextureEditorToolkit interface
-	virtual void CalculateTextureDimensions( uint32& Width, uint32& Height ) const override;
+	virtual void CalculateTextureDimensions( uint32& Width, uint32& Height, uint32& Depth, uint32& ArraySize ) const override;
 	virtual ESimpleElementBlendMode GetColourChannelBlendMode( ) const override;
 	virtual int32 GetMipLevel( ) const override;
 	virtual int32 GetLayer() const override;
@@ -80,11 +80,6 @@ public:
 	virtual void SetVolumeOpacity( float VolumeOpacity ) override;
 	virtual const FRotator& GetVolumeOrientation( ) const override;
 	virtual void SetVolumeOrientation( const FRotator& InOrientation ) override;
-	virtual bool IsVolumeTexture() const override
-	{
-		return bIsVolumeTexture;
-	}
-
 	virtual int32 GetExposureBias() const override
 	{
 		return ExposureBias;
@@ -142,6 +137,16 @@ protected:
 	 */
 	bool IsCubeTexture( ) const;
 
+	/**
+	 * Checks whether the texture being edited is a volume texture.
+	 */
+	bool IsVolumeTexture() const;
+
+	/**
+	 * Checks whether the texture being edited is a texture 2D array.
+	 */
+	bool Is2DArrayTexture() const;
+
 	TSharedRef<SWidget> OnGenerateMipMapLevelMenu();
 	TSharedRef<SWidget> OnGenerateSettingsMenu();
 private:
diff --git a/Engine/Source/Editor/TextureEditor/Public/Interfaces/ITextureEditorToolkit.h b/Engine/Source/Editor/TextureEditor/Public/Interfaces/ITextureEditorToolkit.h
--- a/Engine/Source/Editor/TextureEditor/Public/Interfaces/ITextureEditorToolkit.h
+++ b/Engine/Source/Editor/TextureEditor/Public/Interfaces/ITextureEditorToolkit.h
@@ -27,7 +27,7 @@ public:
 	virtual void PopulateQuickInfo() = 0;
 
 	/** Calculates the display size of the texture */
-	virtual void CalculateTextureDimensions(uint32& Width, uint32& Height) const = 0;
+	virtual void CalculateTextureDimensions(uint32& Width, uint32& Height, uint32& Depth, uint32& ArraySize) const = 0;
 
 	/** Accessors */ 
 	virtual int32 GetMipLevel() const = 0;
diff --git a/Engine/Source/Programs/UnrealLightmass/Private/Lighting/LightmapData.h b/Engine/Source/Programs/UnrealLightmass/Private/Lighting/LightmapData.h
--- a/Engine/Source/Programs/UnrealLightmass/Private/Lighting/LightmapData.h
+++ b/Engine/Source/Programs/UnrealLightmass/Private/Lighting/LightmapData.h
@@ -115,10 +115,6 @@ namespace Lightmass
 		uint32 GetSizeX() const { return SizeX; }
 		uint32 GetSizeY() const { return SizeY; }
 
-		// USurface interface
-		virtual float GetSurfaceWidth() const { return SizeX; }
-		virtual float GetSurfaceHeight() const { return SizeY; }
-
 		void Quantize(int32 DebugSampleIndex);
 
 		const FShadowSample* GetData() { return Data.GetData(); }
@@ -179,10 +175,6 @@ namespace Lightmass
 		uint32 GetSizeX() const { return SizeX; }
 		uint32 GetSizeY() const { return SizeY; }
 
-		// USurface interface
-		virtual float GetSurfaceWidth() const { return SizeX; }
-		virtual float GetSurfaceHeight() const { return SizeY; }
-
 		void Quantize(int32 DebugSampleIndex);
 
 		const FSignedDistanceFieldShadowSample* GetData() { return Data.GetData(); }
diff --git a/Engine/Source/Runtime/AugmentedReality/Public/ARTextures.h b/Engine/Source/Runtime/AugmentedReality/Public/ARTextures.h
--- a/Engine/Source/Runtime/AugmentedReality/Public/ARTextures.h
+++ b/Engine/Source/Runtime/AugmentedReality/Public/ARTextures.h
@@ -46,6 +46,8 @@ public:
 	// UTexture interface implementation
 	virtual float GetSurfaceWidth() const override { return Size.X; }
 	virtual float GetSurfaceHeight() const override { return Size.Y; }
+	virtual float GetSurfaceDepth() const override { return 0; }
+	virtual uint32 GetSurfaceArraySize() const override { return 0; }
 	virtual FGuid GetExternalTextureGuid() const override { return ExternalTextureGuid; }
 	// End UTexture interface
 	
@@ -129,6 +131,8 @@ public:
 	// UTexture interface implementation
 	virtual float GetSurfaceWidth() const override { return Size.X; }
 	virtual float GetSurfaceHeight() const override { return Size.Y; }
+	virtual float GetSurfaceDepth() const override { return 0; }
+	virtual uint32 GetSurfaceArraySize() const override { return 6; }
 	virtual FGuid GetExternalTextureGuid() const override { return ExternalTextureGuid; }
 	// End UTexture interface
 	
diff --git a/Engine/Source/Runtime/Engine/Classes/Engine/Texture.h b/Engine/Source/Runtime/Engine/Classes/Engine/Texture.h
--- a/Engine/Source/Runtime/Engine/Classes/Engine/Texture.h
+++ b/Engine/Source/Runtime/Engine/Classes/Engine/Texture.h
@@ -795,6 +795,17 @@ public:
 		Properties needed at runtime below.
 	--------------------------------------------------------------------------*/
 
+	/*
+	 * Level scope index of this texture. It is used to reduce the amount of lookup to map a texture to its level index.
+	 * Useful when building texture streaming data, as well as when filling the texture streamer with precomputed data.
+     * It relates to FStreamingTextureBuildInfo::TextureLevelIndex and also the index in ULevel::StreamingTextureGuids. 
+	 * Default value of -1, indicates that the texture has an unknown index (not yet processed). At level load time, 
+	 * -2 is also used to indicate that the texture has been processed but no entry were found in the level table.
+	 * After any of these processes, the LevelIndex is reset to INDEX_NONE. Making it ready for the next level task.
+	 */
+	UPROPERTY(transient, duplicatetransient, NonTransactional)
+	int32 LevelIndex = INDEX_NONE;
+
 	/** A bias to the index of the top mip level to use. */
 	UPROPERTY(EditAnywhere, BlueprintReadWrite, Category=LevelOfDetail, meta=(DisplayName="LOD Bias"), AssetRegistrySearchable)
 	int32 LODBias;
@@ -1046,6 +1057,12 @@ public:
 	/** @return the height of the surface represented by the texture. */
 	virtual float GetSurfaceHeight() const PURE_VIRTUAL(UTexture::GetSurfaceHeight,return 0;);
 
+	/** @return the depth of the surface represented by the texture. */
+	virtual float GetSurfaceDepth() const PURE_VIRTUAL(UTexture::GetSurfaceDepth, return 0;);
+
+	/** @return the array size of the surface represented by the texture. */
+	virtual uint32 GetSurfaceArraySize() const PURE_VIRTUAL(UTexture::GetSurfaceArraySize, return 0;);
+
 	/**
 	 * Access the GUID which defines this texture's resources externally through FExternalTextureRegistry
 	 */
@@ -1090,6 +1107,12 @@ public:
 	ENGINE_API virtual bool ShouldMipLevelsBeForcedResident() const final override;
 	//~ End UStreamableRenderAsset Interface
 
+	/**
+	 * Cancels any pending texture streaming actions if possible.
+	 * Returns when no more async loading requests are in flight.
+	 */
+	ENGINE_API static void CancelPendingTextureStreaming();
+
 	/**
 	 *	Gets the average brightness of the texture (in linear space)
 	 *
diff --git a/Engine/Source/Runtime/Engine/Classes/Engine/Texture2D.h b/Engine/Source/Runtime/Engine/Classes/Engine/Texture2D.h
--- a/Engine/Source/Runtime/Engine/Classes/Engine/Texture2D.h
+++ b/Engine/Source/Runtime/Engine/Classes/Engine/Texture2D.h
@@ -22,17 +22,6 @@ class UTexture2D : public UTexture
 
 public:
 
-	/*
-	 * Level scope index of this texture. It is used to reduce the amount of lookup to map a texture to its level index.
-	 * Useful when building texture streaming data, as well as when filling the texture streamer with precomputed data.
-     * It relates to FStreamingTextureBuildInfo::TextureLevelIndex and also the index in ULevel::StreamingTextureGuids. 
-	 * Default value of -1, indicates that the texture has an unknown index (not yet processed). At level load time, 
-	 * -2 is also used to indicate that the texture has been processed but no entry were found in the level table.
-	 * After any of these processes, the LevelIndex is reset to INDEX_NONE. Making it ready for the next level task.
-	 */
-	UPROPERTY(transient, duplicatetransient, NonTransactional)
-	int32 LevelIndex;
-
 	/** keep track of first mip level used for ResourceMem creation */
 	UPROPERTY()
 	int32 FirstResourceMemMip;
@@ -127,6 +116,8 @@ public:
 	//~ Begin UTexture Interface.
 	virtual float GetSurfaceWidth() const override { return GetSizeX(); }
 	virtual float GetSurfaceHeight() const override { return GetSizeY(); }
+	virtual float GetSurfaceDepth() const override { return 0; }
+	virtual uint32 GetSurfaceArraySize() const override { return 0; }
 	virtual FTextureResource* CreateResource() override;
 	virtual EMaterialValueType GetMaterialType() const override;
 	virtual void UpdateResource() override;
@@ -228,13 +219,6 @@ public:
 	 */
 	virtual void GetResourceSizeEx(FResourceSizeEx& CumulativeResourceSize) override;
 
-	/**
-	 * Cancels any pending texture streaming actions if possible.
-	 * Returns when no more async loading requests are in flight.
-	 */
-	ENGINE_API static void CancelPendingTextureStreaming();
-
-
 	/**
 	 * Returns the global mip map bias applied as an offset for 2d textures.
 	 */
diff --git a/Engine/Source/Runtime/Engine/Classes/Engine/Texture2DArray.cpp b/Engine/Source/Runtime/Engine/Classes/Engine/Texture2DArray.cpp
--- a/Engine/Source/Runtime/Engine/Classes/Engine/Texture2DArray.cpp
+++ b/Engine/Source/Runtime/Engine/Classes/Engine/Texture2DArray.cpp
@@ -12,9 +12,13 @@
 #include "DeviceProfiles/DeviceProfileManager.h"
 #include "Containers/ResourceArray.h"
 #include "Rendering/Texture2DArrayResource.h"
+#include "Streaming/Texture2DArrayStreaming.h"
+#include "Streaming/TextureStreamIn.h"
+#include "Streaming/TextureStreamOut.h"
+#include "Engine/TextureMipDataProviderFactory.h"
 
 // Master switch to control whether streaming is enabled for texture2darray. 
-bool GSupportsTexture2DArrayStreaming = false;
+bool GSupportsTexture2DArrayStreaming = true;
 
 static TAutoConsoleVariable<int32> CVarAllowTexture2DArrayAssetCreation(
 	TEXT("r.AllowTexture2DArrayCreation"),
@@ -76,7 +80,7 @@ uint32 UTexture2DArray::CalcTextureMemorySize(int32 MipCount) const
 			// Must be consistent with the logic in FTexture2DResource::InitRHI
 			const FIntPoint MipExtents = CalcMipMapExtent(GetSizeX(), GetSizeY(), Format, FirstMip);
 			uint32 TextureAlign = 0;
-			Size = (uint32)(GetNumSlices() * RHICalcTexture2DPlatformSize(MipExtents.X, MipExtents.Y, Format, FMath::Max(1, MipCount), 1, Flags, FRHIResourceCreateInfo(PlatformData->GetExtData()), TextureAlign));
+			Size = (uint32)RHICalcTexture2DArrayPlatformSize(MipExtents.X, MipExtents.Y, GetArraySize(), Format, FMath::Max(1, MipCount), 1, Flags, FRHIResourceCreateInfo(PlatformData->GetExtData()), TextureAlign);
 		}
 	}
 	return Size;
@@ -106,9 +110,9 @@ ENGINE_API bool UTexture2DArray::CheckArrayTexturesCompatibility()
 			return false;
 		}
 
-		FTextureSource& TextureSource = SourceTextures[TextureIndex]->Source;
+		const FTextureSource& TextureSource = SourceTextures[TextureIndex]->Source;
 		// const int32 FormatDataSize = TextureSource.GetBytesPerPixel();
-		const EPixelFormat PixelFormat = SourceTextures[TextureIndex]->GetPixelFormat();
+		const ETextureSourceFormat SourceFormat = TextureSource.GetFormat();
 		const int32 SizeX = TextureSource.GetSizeX();
 		const int32 SizeY = TextureSource.GetSizeY();
 
@@ -120,10 +124,10 @@ ENGINE_API bool UTexture2DArray::CheckArrayTexturesCompatibility()
 				return false;
 			}
 
-			FTextureSource& TextureSourceCmp = SourceTextures[TextureCmpIndex]->Source;
-			FString TextureName = SourceTextures[TextureIndex]->GetFName().ToString();
-			FString TextureNameCmp = SourceTextures[TextureCmpIndex]->GetFName().ToString();
-			const EPixelFormat PixelFormatCmp = SourceTextures[TextureCmpIndex]->GetPixelFormat();
+			const FTextureSource& TextureSourceCmp = SourceTextures[TextureCmpIndex]->Source;
+			const FString TextureName = SourceTextures[TextureIndex]->GetFName().ToString();
+			const FString TextureNameCmp = SourceTextures[TextureCmpIndex]->GetFName().ToString();
+			const ETextureSourceFormat SourceFormatCmp = TextureSourceCmp.GetFormat();
 
 			if (TextureSourceCmp.GetSizeX() != SizeX || TextureSourceCmp.GetSizeY() != SizeY)
 			{
@@ -131,17 +135,11 @@ ENGINE_API bool UTexture2DArray::CheckArrayTexturesCompatibility()
 				bError = true;
 			}
 
-			if (PixelFormatCmp != PixelFormat)
+			if (SourceFormatCmp != SourceFormat)
 			{
 				UE_LOG(LogTexture, Warning, TEXT("Texture2DArray creation failed. Textures %s and %s have incompatible pixel formats."), *TextureName, *TextureNameCmp);
 				bError = true;
 			}
-
-			//if (TextureSourceCmp.GetBytesPerPixel() != FormatDataSize)
-			//{
-			//	UE_LOG(LogTexture, Warning, TEXT("Texture2DArray creation failed. Textures %s and %s have incompatible pixel formats."), *TextureName, *TextureNameCmp);
-			//	bError = true;
-			//}
 		}
 	}
 
@@ -166,7 +164,7 @@ ENGINE_API bool UTexture2DArray::UpdateSourceFromSourceTextures(bool bCreatingNe
 		// X,Y,Z size of the array.
 		int32 SizeX = SourceTextures[0]->GetSizeX();
 		int32 SizeY = SourceTextures[0]->GetSizeY();
-		uint32 NumSlices = SourceTextures.Num();
+		uint32 ArraySize = SourceTextures.Num();
 		// Only copy the first mip from the source textures to array texture.
 		uint32 NumMips = 1;
 
@@ -182,16 +180,18 @@ ENGINE_API bool UTexture2DArray::UpdateSourceFromSourceTextures(bool bCreatingNe
 		}
 
 		// Create the source texture for this UTexture.
-		Source.Init(SizeX, SizeY, NumSlices, NumMips, Format);
+		Source.Init(SizeX, SizeY, ArraySize, NumMips, Format);
 
 		// We only copy the top level Mip map.
-		uint8* DestMipData[MAX_TEXTURE_MIP_COUNT] = { 0 };
-		int32 MipSizeBytes[MAX_TEXTURE_MIP_COUNT] = { 0 };
+		TArray<uint8*, TInlineAllocator<MAX_TEXTURE_MIP_COUNT> > DestMipData;
+		TArray<uint64, TInlineAllocator<MAX_TEXTURE_MIP_COUNT> > MipSizeBytes;
+		DestMipData.AddZeroed(NumMips);
+		MipSizeBytes.AddZeroed(NumMips);
 			
 		for (uint32 MipIndex = 0; MipIndex < NumMips; ++MipIndex)
 		{
 			DestMipData[MipIndex] =  Source.LockMip(MipIndex);
-			MipSizeBytes[MipIndex] = Source.CalcMipSize(MipIndex) / NumSlices;
+			MipSizeBytes[MipIndex] = Source.CalcMipSize(MipIndex) / ArraySize;
 		}
 
 		for (int32 SourceTexIndex = 0; SourceTexIndex < SourceTextures.Num(); ++SourceTexIndex)
@@ -268,13 +268,13 @@ void UTexture2DArray::GetAssetRegistryTags(TArray<FAssetRegistryTag>& OutTags) c
 #if WITH_EDITOR
 	int32 SizeX = Source.GetSizeX();
 	int32 SizeY = Source.GetSizeY();
-	int32 SizeZ = Source.GetNumSlices(); //GetSizeZ()
+	int32 ArraySize = Source.GetNumSlices();
 #else
 	int32 SizeX = 0;
 	int32 SizeY = 0;
-	int32 SizeZ = 0;
+	int32 ArraySize = 0;
 #endif
-	const FString Dimensions = FString::Printf(TEXT("%dx%dx%d"), SizeX, SizeY, SizeZ);
+	const FString Dimensions = FString::Printf(TEXT("%dx%d*%d"), SizeX, SizeY, ArraySize);
 	OutTags.Add(FAssetRegistryTag("Dimensions", Dimensions, FAssetRegistryTag::TT_Dimensional));
 	OutTags.Add(FAssetRegistryTag("Format", GPixelFormats[GetPixelFormat()].Name, FAssetRegistryTag::TT_Alphabetical));
 
@@ -283,10 +283,10 @@ void UTexture2DArray::GetAssetRegistryTags(TArray<FAssetRegistryTag>& OutTags) c
 
 FString UTexture2DArray::GetDesc() 
 {
-	return FString::Printf(TEXT("Array: %dx%dx%d [%s]"),
+	return FString::Printf(TEXT("Array: %dx%d*%d [%s]"),
 		GetSizeX(),
 		GetSizeY(),
-		GetNumSlices(),
+		GetArraySize(),
 		GPixelFormats[GetPixelFormat()].Name
 	);
 }
@@ -370,10 +370,51 @@ ENGINE_API void UTexture2DArray::PostEditChangeProperty(FPropertyChangedEvent &
 
 bool UTexture2DArray::StreamOut(int32 NewMipCount)
 {
+	FTexture2DArrayResource* Texture2DArrayResource = Resource ? Resource->GetTexture2DArrayResource() : nullptr;
+	if (!HasPendingInitOrStreaming() && CachedSRRState.StreamOut(NewMipCount) && ensure(Texture2DArrayResource))
+	{
+		FTextureMipAllocator* MipAllocator = new FTexture2DArrayMipAllocator_Reallocate(this);
+		PendingUpdate = new FTextureStreamOut(this, MipAllocator);
+		return !PendingUpdate->IsCancelled();
+	}
 	return false;
 }
 
 bool UTexture2DArray::StreamIn(int32 NewMipCount, bool bHighPrio)
 {
+	FTexture2DArrayResource* Texture2DArrayResource = Resource ? Resource->GetTexture2DArrayResource() : nullptr;
+	if (!HasPendingInitOrStreaming() && CachedSRRState.StreamIn(NewMipCount) && ensure(Texture2DArrayResource))
+	{
+		FTextureMipDataProvider* CustomMipDataProvider = nullptr;
+		for (UAssetUserData* UserData : AssetUserData)
+		{
+			UTextureMipDataProviderFactory* CustomMipDataProviderFactory = Cast<UTextureMipDataProviderFactory>(UserData);
+			if (CustomMipDataProviderFactory)
+			{
+				CustomMipDataProvider = CustomMipDataProviderFactory->AllocateMipDataProvider(this);
+				if (CustomMipDataProvider)
+				{
+					break;
+				}
+			}
+		}
+
+		FTextureMipAllocator* MipAllocator = new FTexture2DArrayMipAllocator_Reallocate(this);
+		FTextureMipDataProvider* DefaultMipDataProvider = nullptr;
+
+#if WITH_EDITORONLY_DATA
+		if (FPlatformProperties::HasEditorOnlyData() && !GetOutermost()->bIsCookedForEditor)
+		{
+			DefaultMipDataProvider = new FTexture2DArrayMipDataProvider_DDC(this);
+		}
+		else
+#endif
+		{
+			DefaultMipDataProvider = new FTexture2DArrayMipDataProvider_IO(this, bHighPrio);
+		}
+
+		PendingUpdate = new FTextureStreamIn(this, MipAllocator, CustomMipDataProvider, DefaultMipDataProvider);
+		return !PendingUpdate->IsCancelled();
+	}
 	return false;
 }
diff --git a/Engine/Source/Runtime/Engine/Classes/Engine/Texture2DArray.h b/Engine/Source/Runtime/Engine/Classes/Engine/Texture2DArray.h
--- a/Engine/Source/Runtime/Engine/Classes/Engine/Texture2DArray.h
+++ b/Engine/Source/Runtime/Engine/Classes/Engine/Texture2DArray.h
@@ -34,7 +34,7 @@ public:
 	{
 		return PlatformData ? PlatformData->SizeY : 0;
 	}
-	FORCEINLINE int32 GetNumSlices() const
+	FORCEINLINE int32 GetArraySize() const
 	{
 		return PlatformData ? PlatformData->GetNumSlices() : 0;
 	}
@@ -55,6 +55,8 @@ public:
 	virtual void GetResourceSizeEx(FResourceSizeEx& CumulativeResourceSize) override;
 	virtual float GetSurfaceWidth() const override { return GetSizeX(); }
 	virtual float GetSurfaceHeight() const override { return GetSizeY(); }
+	virtual float GetSurfaceDepth() const override { return 0; }
+	virtual uint32 GetSurfaceArraySize() const override { return GetArraySize(); }
 	virtual FTextureResource* CreateResource() override;
 #if WITH_EDITOR
 	ENGINE_API virtual void PostEditChangeProperty(FPropertyChangedEvent& PropertyChangedEvent) override;
diff --git a/Engine/Source/Runtime/Engine/Classes/Engine/Texture2DDynamic.h b/Engine/Source/Runtime/Engine/Classes/Engine/Texture2DDynamic.h
--- a/Engine/Source/Runtime/Engine/Classes/Engine/Texture2DDynamic.h
+++ b/Engine/Source/Runtime/Engine/Classes/Engine/Texture2DDynamic.h
@@ -58,6 +58,8 @@ public:
 	virtual EMaterialValueType GetMaterialType() const override { return MCT_Texture2D; }
 	ENGINE_API virtual float GetSurfaceWidth() const override;
 	ENGINE_API virtual float GetSurfaceHeight() const override;
+	ENGINE_API virtual float GetSurfaceDepth() const override { return 0; }
+	ENGINE_API virtual uint32 GetSurfaceArraySize() const override { return 0; }
 	//~ End UTexture Interface.
 	
 	/**
diff --git a/Engine/Source/Runtime/Engine/Classes/Engine/TextureMipDataProviderFactory.h b/Engine/Source/Runtime/Engine/Classes/Engine/TextureMipDataProviderFactory.h
--- a/Engine/Source/Runtime/Engine/Classes/Engine/TextureMipDataProviderFactory.h
+++ b/Engine/Source/Runtime/Engine/Classes/Engine/TextureMipDataProviderFactory.h
@@ -32,5 +32,5 @@ public:
 	* 
 	* @param InAsset - the texture on which the stream in operation will be performed.
 	*/
-	FTextureMipDataProvider* AllocateMipDataProvider(UTexture* Asset) { return nullptr; }
+	virtual FTextureMipDataProvider* AllocateMipDataProvider(UTexture* Asset) { return nullptr; }
 };
diff --git a/Engine/Source/Runtime/Engine/Classes/Engine/TextureRenderTarget2D.h b/Engine/Source/Runtime/Engine/Classes/Engine/TextureRenderTarget2D.h
--- a/Engine/Source/Runtime/Engine/Classes/Engine/TextureRenderTarget2D.h
+++ b/Engine/Source/Runtime/Engine/Classes/Engine/TextureRenderTarget2D.h
@@ -177,6 +177,8 @@ class UTextureRenderTarget2D : public UTextureRenderTarget
 	//~ Begin UTexture Interface.
 	virtual float GetSurfaceWidth() const override { return SizeX; }
 	virtual float GetSurfaceHeight() const override { return SizeY; }
+	virtual float GetSurfaceDepth() const override { return 0; }
+	virtual uint32 GetSurfaceArraySize() const override { return 0; }
 	ENGINE_API virtual FTextureResource* CreateResource() override;
 	ENGINE_API virtual EMaterialValueType GetMaterialType() const override;
 	//~ End UTexture Interface.
diff --git a/Engine/Source/Runtime/Engine/Classes/Engine/TextureRenderTarget2DArray.h b/Engine/Source/Runtime/Engine/Classes/Engine/TextureRenderTarget2DArray.h
--- a/Engine/Source/Runtime/Engine/Classes/Engine/TextureRenderTarget2DArray.h
+++ b/Engine/Source/Runtime/Engine/Classes/Engine/TextureRenderTarget2DArray.h
@@ -80,6 +80,8 @@ class UTextureRenderTarget2DArray : public UTextureRenderTarget
 	// PVS-Studio notices that the implementation of GetSurfaceWidth is identical to this one
 	// and warns us. In this case, it is intentional, so we disable the warning:
 	virtual float GetSurfaceHeight()const  override { return SizeX; }	 //-V524
+	virtual float GetSurfaceDepth() const override { return 0; }
+	virtual uint32 GetSurfaceArraySize() const override { return Slices; }
 	virtual FTextureResource* CreateResource() override;
 	virtual EMaterialValueType GetMaterialType() const override;
 	//~ End UTexture Interface.
diff --git a/Engine/Source/Runtime/Engine/Classes/Engine/TextureRenderTargetCube.h b/Engine/Source/Runtime/Engine/Classes/Engine/TextureRenderTargetCube.h
--- a/Engine/Source/Runtime/Engine/Classes/Engine/TextureRenderTargetCube.h
+++ b/Engine/Source/Runtime/Engine/Classes/Engine/TextureRenderTargetCube.h
@@ -69,6 +69,8 @@ class UTextureRenderTargetCube : public UTextureRenderTarget
 	// PVS-Studio notices that the implementation of GetSurfaceWidth is identical to this one
 	// and warns us. In this case, it is intentional, so we disable the warning:
 	virtual float GetSurfaceHeight()const  override { return SizeX; }	 //-V524
+	virtual float GetSurfaceDepth() const override { return 0; }
+	virtual uint32 GetSurfaceArraySize() const override { return 6; }
 	virtual FTextureResource* CreateResource() override;
 	virtual EMaterialValueType GetMaterialType() const override;
 	//~ End UTexture Interface.
diff --git a/Engine/Source/Runtime/Engine/Classes/Engine/TextureRenderTargetVolume.h b/Engine/Source/Runtime/Engine/Classes/Engine/TextureRenderTargetVolume.h
--- a/Engine/Source/Runtime/Engine/Classes/Engine/TextureRenderTargetVolume.h
+++ b/Engine/Source/Runtime/Engine/Classes/Engine/TextureRenderTargetVolume.h
@@ -76,9 +76,9 @@ class UTextureRenderTargetVolume : public UTextureRenderTarget
 
 	//~ Begin UTexture Interface.
 	virtual float GetSurfaceWidth() const  override { return SizeX; }
-	// PVS-Studio notices that the implementation of GetSurfaceWidth is identical to this one
-	// and warns us. In this case, it is intentional, so we disable the warning:
-	virtual float GetSurfaceHeight()const  override { return SizeX; }	 //-V524
+	virtual float GetSurfaceHeight()const  override { return SizeY; }
+	virtual float GetSurfaceDepth()const  override { return SizeZ; }
+	virtual uint32 GetSurfaceArraySize() const override { return 0; }
 	virtual FTextureResource* CreateResource() override;
 	virtual EMaterialValueType GetMaterialType() const override;
 	//~ End UTexture Interface.
diff --git a/Engine/Source/Runtime/Engine/Classes/Engine/TextureStreamingTypes.h b/Engine/Source/Runtime/Engine/Classes/Engine/TextureStreamingTypes.h
--- a/Engine/Source/Runtime/Engine/Classes/Engine/TextureStreamingTypes.h
+++ b/Engine/Source/Runtime/Engine/Classes/Engine/TextureStreamingTypes.h
@@ -15,7 +15,7 @@
 class ULevel;
 class UMaterialInterface;
 class UPrimitiveComponent;
-class UTexture2D;
+class UTexture;
 struct FMaterialTextureInfo;
 struct FMeshUVChannelInfo;
 struct FSlowTask;
@@ -24,7 +24,6 @@ struct FStreamingTextureBuildInfo;
 ENGINE_API DECLARE_LOG_CATEGORY_EXTERN(TextureStreamingBuild, Log, All);
 
 class UTexture;
-class UTexture2D;
 class UStreamableRenderAsset;
 struct FStreamingTextureBuildInfo;
 struct FMaterialTextureInfo;
@@ -142,7 +141,7 @@ struct FStreamingTextureBuildInfo
 	/**
 	 *	Set this struct to match the unpacked params.
 	 *
-	 *	@param	LevelTextures	[in,out]	The list of textures referred by all component of a level. The array index maps to UTexture2D::LevelIndex.
+	 *	@param	LevelTextures	[in,out]	The list of textures referred by all component of a level. The array index maps to UTexture::LevelIndex.
 	 *	@param	RefBounds		[in]		The reference bounds used to compute the packed relative box.
 	 *	@param	Info			[in]		The unpacked params.
 	 */
@@ -211,14 +210,14 @@ class FStreamingTextureLevelContext
 	{
 		FTextureBoundState() {}
 
-		FTextureBoundState(UTexture2D* InTexture) : BuildDataTimestamp(0), BuildDataIndex(0), Texture(InTexture) {}
+		FTextureBoundState(UTexture* InTexture) : BuildDataTimestamp(0), BuildDataIndex(0), Texture(InTexture) {}
 
 		/** The timestamp of the build data to indentify whether BuildDataIndex is valid or not. */
 		int32 BuildDataTimestamp;
 		/** The ComponentBuildData Index referring this texture. */
 		int32 BuildDataIndex;
 		/**  The texture relative to this entry. */
-		UTexture2D* Texture;
+		UTexture* Texture;
 	};
 
 	/*
@@ -230,7 +229,7 @@ class FStreamingTextureLevelContext
 	EMaterialQualityLevel::Type QualityLevel;
 	ERHIFeatureLevel::Type FeatureLevel;
 
-	int32* GetBuildDataIndexRef(UTexture2D* Texture2D);
+	int32* GetBuildDataIndexRef(UTexture* Texture);
 
 public:
 
diff --git a/Engine/Source/Runtime/Engine/Classes/Engine/VolumeTexture.h b/Engine/Source/Runtime/Engine/Classes/Engine/VolumeTexture.h
--- a/Engine/Source/Runtime/Engine/Classes/Engine/VolumeTexture.h
+++ b/Engine/Source/Runtime/Engine/Classes/Engine/VolumeTexture.h
@@ -84,6 +84,8 @@ public:
 	//~ Begin UTexture Interface
 	virtual float GetSurfaceWidth() const override { return GetSizeX(); }
 	virtual float GetSurfaceHeight() const override { return GetSizeY(); }
+	virtual float GetSurfaceDepth() const override { return GetSizeZ(); }
+	virtual uint32 GetSurfaceArraySize() const override { return 0; }
 	virtual FTextureResource* CreateResource() override;
 #if WITH_EDITOR
 	ENGINE_API void SetDefaultSource2DTileSize();
diff --git a/Engine/Source/Runtime/Engine/Private/Rendering/StreamableTextureResource.cpp b/Engine/Source/Runtime/Engine/Private/Rendering/StreamableTextureResource.cpp
--- a/Engine/Source/Runtime/Engine/Private/Rendering/StreamableTextureResource.cpp
+++ b/Engine/Source/Runtime/Engine/Private/Rendering/StreamableTextureResource.cpp
@@ -107,7 +107,12 @@ FStreamableTextureResource::FStreamableTextureResource(UTexture* InOwner, const
 	SizeZ = Mip0.SizeZ;
 
 	MipFadeSetting = (LODGroup == TEXTUREGROUP_Lightmap || LODGroup == TEXTUREGROUP_Shadowmap) ? MipFade_Slow : MipFade_Normal;
-	CreationFlags = (InOwner->SRGB ? TexCreate_SRGB : TexCreate_None)  | (InOwner->bNotOfflineProcessed ? TexCreate_None : TexCreate_OfflineProcessed) | TexCreate_ShaderResource | TexCreate_Streamable | (InOwner->bNoTiling ? TexCreate_NoTiling : TexCreate_None);
+	CreationFlags = (InOwner->SRGB ? TexCreate_SRGB : TexCreate_None)  | (InOwner->bNotOfflineProcessed ? TexCreate_None : TexCreate_OfflineProcessed) | TexCreate_ShaderResource | (InOwner->bNoTiling ? TexCreate_NoTiling : TexCreate_None);
+
+	if (InPostInitState.MaxNumLODs > 1)
+	{
+		CreationFlags |= TexCreate_Streamable;
+	}
 
 	// Whether the virtual update path is enabled for this texture. This allows to map / unmap top mip memory in an out.
 	// Whether the texture will be created with TexCreate_Virtual depends on the requested mip count and "r.VirtualTextureReducedMemory"
@@ -118,6 +123,12 @@ FStreamableTextureResource::FStreamableTextureResource(UTexture* InOwner, const
 }
 
 #if STATS
+
+void FStreamableTextureResource::CalcRequestedMipsSize() 
+{ 
+	TextureSize = GetPlatformMipsSize(State.NumRequestedLODs); 
+}
+
 void FStreamableTextureResource::IncrementTextureStats() const
 {
 	INC_DWORD_STAT_BY( STAT_TextureMemory, TextureSize );
diff --git a/Engine/Source/Runtime/Engine/Private/Rendering/Texture2DArrayResource .cpp b/Engine/Source/Runtime/Engine/Private/Rendering/Texture2DArrayResource .cpp
--- a/Engine/Source/Runtime/Engine/Private/Rendering/Texture2DArrayResource .cpp	
+++ b/Engine/Source/Runtime/Engine/Private/Rendering/Texture2DArrayResource .cpp	
@@ -120,18 +120,16 @@ void FTexture2DArrayResource::CreatePartiallyResidentTexture()
 	TextureRHI.SafeRelease();
 }
 
-#if STATS
-void FTexture2DArrayResource::CalcRequestedMipsSize()
+uint64 FTexture2DArrayResource::GetPlatformMipsSize(uint32 NumMips) const
 {
-	if (PlatformData && State.NumRequestedLODs > 0)
+	if (PlatformData && NumMips > 0)
 	{
-		const FIntPoint MipExtents = CalcMipMapExtent(SizeX, SizeY, PixelFormat, State.RequestedFirstLODIdx());
+		const FIntPoint MipExtents = CalcMipMapExtent(SizeX, SizeY, PixelFormat, State.LODCountToFirstLODIdx(NumMips));
 		uint32 TextureAlign = 0;
-		TextureSize = SizeZ * RHICalcTexture2DPlatformSize(MipExtents.X, MipExtents.Y, PixelFormat, State.NumRequestedLODs, 1, CreationFlags, FRHIResourceCreateInfo(PlatformData->GetExtData()), TextureAlign);
+		return RHICalcTexture2DArrayPlatformSize(MipExtents.X, MipExtents.Y, SizeZ, PixelFormat, NumMips, 1, CreationFlags, FRHIResourceCreateInfo(PlatformData->GetExtData()), TextureAlign);
 	}
 	else
 	{
-		TextureSize = 0;
+		return 0;
 	}
 }
-#endif
diff --git a/Engine/Source/Runtime/Engine/Private/Rendering/Texture2DArrayResource.h b/Engine/Source/Runtime/Engine/Private/Rendering/Texture2DArrayResource.h
--- a/Engine/Source/Runtime/Engine/Private/Rendering/Texture2DArrayResource.h
+++ b/Engine/Source/Runtime/Engine/Private/Rendering/Texture2DArrayResource.h
@@ -19,15 +19,18 @@ public:
 
 	FTexture2DArrayResource(UTexture2DArray* InOwner, const FStreamableRenderResourceState& InState);
 
+	// Dynamic cast methods.
+	ENGINE_API virtual FTexture2DArrayResource* GetTexture2DArrayResource() { return this; }
+	// Dynamic cast methods (const).
+	ENGINE_API virtual const FTexture2DArrayResource* GetTexture2DArrayResource() const { return this; }
+
+	virtual uint64 GetPlatformMipsSize(uint32 NumMips) const override;
+
 protected:
 
 	void CreateTexture() final override;
 	void CreatePartiallyResidentTexture() final override;
 
-#if STATS
-	virtual void CalcRequestedMipsSize() final override;
-#endif
-
 	/** The initial data for all mips of a single slice. */
 	typedef TArray<TArrayView<uint8>, TInlineAllocator<MAX_TEXTURE_MIP_COUNT> > FSingleSliceMipDataView;
 	/** All slices initial data. */
diff --git a/Engine/Source/Runtime/Engine/Private/Rendering/Texture2DResource.cpp b/Engine/Source/Runtime/Engine/Private/Rendering/Texture2DResource.cpp
--- a/Engine/Source/Runtime/Engine/Private/Rendering/Texture2DResource.cpp
+++ b/Engine/Source/Runtime/Engine/Private/Rendering/Texture2DResource.cpp
@@ -167,32 +167,31 @@ void FTexture2DResource::CreatePartiallyResidentTexture()
 	TextureRHI = Texture2DRHI;
 }
 
-#if STATS
-void FTexture2DResource::CalcRequestedMipsSize()
+uint64 FTexture2DResource::GetPlatformMipsSize(uint32 NumMips) const
 {
-	if (PlatformData && State.NumRequestedLODs > 0)
+	if (PlatformData && NumMips > 0)
 	{
 		static TConsoleVariableData<int32>* CVarReducedMode = IConsoleManager::Get().FindTConsoleVariableDataInt(TEXT("r.VirtualTextureReducedMemory"));
 		check(CVarReducedMode);
 
 		uint32 TextureAlign = 0;
 		// Must be consistent with the logic in FTexture2DResource::InitRHI
-		if (bUsePartiallyResidentMips && (!CVarReducedMode->GetValueOnRenderThread() || State.NumRequestedLODs > State.NumNonStreamingLODs))
+		if (bUsePartiallyResidentMips && (!CVarReducedMode->GetValueOnRenderThread() || NumMips > State.NumNonStreamingLODs))
 		{
-			TextureSize = RHICalcVMTexture2DPlatformSize(SizeX, SizeY, PixelFormat, State.NumRequestedLODs, State.RequestedFirstLODIdx(), 1, CreationFlags | TexCreate_Virtual, TextureAlign);
+			return RHICalcVMTexture2DPlatformSize(SizeX, SizeY, PixelFormat, NumMips, State.LODCountToFirstLODIdx(NumMips), 1, CreationFlags | TexCreate_Virtual, TextureAlign);
 		}
 		else
 		{
 			const FIntPoint MipExtents = CalcMipMapExtent(SizeX, SizeY, PixelFormat, 0);
-			TextureSize = RHICalcTexture2DPlatformSize(MipExtents.X, MipExtents.Y, PixelFormat, State.NumRequestedLODs, 1, CreationFlags, FRHIResourceCreateInfo(PlatformData->GetExtData()), TextureAlign);
+			return RHICalcTexture2DPlatformSize(MipExtents.X, MipExtents.Y, PixelFormat, NumMips, 1, CreationFlags, FRHIResourceCreateInfo(PlatformData->GetExtData()), TextureAlign);
 		}
 	}
 	else
 	{
-		TextureSize = 0;
+		return 0;
 	}
 }
-#endif
+
 
 /**
  * Writes the data for a single mip-level into a destination buffer.
diff --git a/Engine/Source/Runtime/Engine/Private/Rendering/Texture3DResource.cpp b/Engine/Source/Runtime/Engine/Private/Rendering/Texture3DResource.cpp
--- a/Engine/Source/Runtime/Engine/Private/Rendering/Texture3DResource.cpp
+++ b/Engine/Source/Runtime/Engine/Private/Rendering/Texture3DResource.cpp
@@ -42,7 +42,7 @@ void FVolumeTextureBulkData::MergeMips(int32 NumMips)
 	// Don't do anything if there is nothing to merge
 	if (MergedSize > MipSize[FirstMipIdx])
 	{
-		uint8* MergedAlloc = (uint8*)FMemory::Malloc(MergedSize);
+		uint8* MergedAlloc = (uint8*)FMemory::Malloc(MergedSize, MALLOC_ALIGNMENT);
 		uint8* CurrPos = MergedAlloc;
 		for (int32 MipIndex = FirstMipIdx; MipIndex < NumMips; ++MipIndex)
 		{
@@ -78,8 +78,17 @@ FTexture3DResource::FTexture3DResource(UVolumeTexture* InOwner, const FStreamabl
 			{
 				const FTexture2DMipMap& MipMap = PlatformData->Mips[MipIndex];
 				
+				uint32 MipExtentX = 0;
+				uint32 MipExtentY = 0;
+				uint32 MipExtentZ = 0;
+				CalcMipMapExtent3D(SizeX, SizeY, SizeZ, PixelFormat, State.RequestedFirstLODIdx(), MipExtentX, MipExtentY, MipExtentZ);
+
+
+				uint32 TextureAlign = 0;
+				uint64 PlatformMipSize = RHICalcTexture3DPlatformSize(MipExtentX, MipExtentY, MipExtentZ, (EPixelFormat)PixelFormat, State.NumRequestedLODs, CreationFlags, FRHIResourceCreateInfo(PlatformData->GetExtData()), TextureAlign);
+
 				// The bulk data can be bigger because of memory alignment constraints on each slice and mips.
-				InitialData.GetMipSize()[MipIndex] = FMath::Max<int32>(
+				InitialData.GetMipSize()[MipIndex] = FMath::Max<uint64>(
 					MipMap.BulkData.GetBulkDataSize(), 
 					CalcTextureMipMapSize3D(SizeX, SizeY, SizeZ, (EPixelFormat)PixelFormat, MipIndex)
 				);
@@ -144,22 +153,20 @@ void FTexture3DResource::CreatePartiallyResidentTexture()
 	TextureRHI.SafeRelease();
 }
 
-#if STATS
-void FTexture3DResource::CalcRequestedMipsSize()
+uint64 FTexture3DResource::GetPlatformMipsSize(uint32 NumMips) const
 {
-	if (PlatformData && State.NumRequestedLODs > 0)
+	if (PlatformData && NumMips > 0)
 	{
 		uint32 MipExtentX = 0;
 		uint32 MipExtentY = 0;
 		uint32 MipExtentZ = 0;
-		CalcMipMapExtent3D(SizeX, SizeY, SizeZ, PixelFormat, State.RequestedFirstLODIdx(), MipExtentX, MipExtentY, MipExtentZ);
+		CalcMipMapExtent3D(SizeX, SizeY, SizeZ, PixelFormat, State.LODCountToFirstLODIdx(NumMips), MipExtentX, MipExtentY, MipExtentZ);
 
 		uint32 TextureAlign = 0;
-		TextureSize = (uint32)RHICalcTexture3DPlatformSize(MipExtentX, MipExtentY, MipExtentZ, PixelFormat, State.NumRequestedLODs, CreationFlags, FRHIResourceCreateInfo(PlatformData->GetExtData()), TextureAlign);
+		return RHICalcTexture3DPlatformSize(MipExtentX, MipExtentY, MipExtentZ, PixelFormat, NumMips, CreationFlags, FRHIResourceCreateInfo(PlatformData->GetExtData()), TextureAlign);
 	}
 	else
 	{
-		TextureSize = 0;
+		return 0;
 	}
 }
-#endif
diff --git a/Engine/Source/Runtime/Engine/Private/Rendering/Texture3DResource.h b/Engine/Source/Runtime/Engine/Private/Rendering/Texture3DResource.h
--- a/Engine/Source/Runtime/Engine/Private/Rendering/Texture3DResource.h
+++ b/Engine/Source/Runtime/Engine/Private/Rendering/Texture3DResource.h
@@ -16,6 +16,8 @@ class FVolumeTextureBulkData : public FResourceBulkDataInterface
 {
 public:
 
+	static const uint32 MALLOC_ALIGNMENT = 16;
+
 	FVolumeTextureBulkData(int32 InFirstMipIdx)
 	: FirstMipIdx(InFirstMipIdx)
 	{
@@ -33,22 +35,28 @@ public:
 		return MipData[FirstMipIdx];
 	}
 
+	void* GetResourceBulkData()
+	{
+		return MipData[FirstMipIdx];
+	}
+
 	uint32 GetResourceBulkDataSize() const override
 	{
-		return MipSize[FirstMipIdx];
+
+		return (uint32)MipSize[FirstMipIdx];
 	}
 
 	void Discard() override;
 	void MergeMips(int32 NumMips);
 
 	void** GetMipData() { return MipData; }
-	uint32* GetMipSize() { return MipSize; }
+	uint64* GetMipSize() { return MipSize; }
 	int32 GetFirstMipIdx() const { return FirstMipIdx; }
 
 protected:
 
 	void* MipData[MAX_TEXTURE_MIP_COUNT];
-	uint32 MipSize[MAX_TEXTURE_MIP_COUNT];
+	uint64 MipSize[MAX_TEXTURE_MIP_COUNT];
 	int32 FirstMipIdx;
 };
 
@@ -59,15 +67,19 @@ public:
 
 	FTexture3DResource(UVolumeTexture* InOwner, const FStreamableRenderResourceState& InState);
 
+	// Dynamic cast methods.
+	ENGINE_API virtual FTexture3DResource* GetTexture3DResource() override { return this; }
+	// Dynamic cast methods (const).
+	ENGINE_API virtual const FTexture3DResource* GetTexture3DResource() const override { return this; }
+
+	/** Returns the platform mip size for the given mip count. */
+	virtual uint64 GetPlatformMipsSize(uint32 NumMips) const override;
+
 private:
 
 	void CreateTexture() final override;
 	void CreatePartiallyResidentTexture() final override;
 
-#if STATS
-	virtual void CalcRequestedMipsSize() final override;
-#endif
-
 protected:
 	FVolumeTextureBulkData InitialData;
 
diff --git a/Engine/Source/Runtime/Engine/Private/Streaming/StreamingManagerTexture.cpp b/Engine/Source/Runtime/Engine/Private/Streaming/StreamingManagerTexture.cpp
--- a/Engine/Source/Runtime/Engine/Private/Streaming/StreamingManagerTexture.cpp
+++ b/Engine/Source/Runtime/Engine/Private/Streaming/StreamingManagerTexture.cpp
@@ -10,6 +10,9 @@
 #include "Engine/TextureStreamingTypes.h"
 #include "Engine/StaticMesh.h"
 #include "Engine/SkeletalMesh.h"
+#include "Engine/Texture2D.h"
+#include "Engine/VolumeTexture.h"
+#include "Engine/Texture2DArray.h"
 #include "LandscapeComponent.h"
 #include "Materials/MaterialInterface.h"
 #include "Misc/CommandLine.h"
@@ -81,8 +84,6 @@ FORCEINLINE float ClampMeshToCameraDistanceSquared(float MeshToCameraDistanceSqu
 /** Constructor, initializing all members and  */
 FRenderAssetStreamingManager::FRenderAssetStreamingManager()
 :	CurrentUpdateStreamingRenderAssetIndex(0)
-,	bTriggerDumpTextureGroupStats( false )
-,	bDetailedDumpTextureGroupStats( false )
 ,	AsyncWork( nullptr )
 ,	CurrentPendingMipCopyRequestIdx(0)
 ,	ProcessingStage( 0 )
@@ -255,7 +256,7 @@ void FRenderAssetStreamingManager::CancelForcedResources()
 		}
 	}
 
-	// Reset the streaming system, so it picks up any changes to UTexture2D right away.
+	// Reset the streaming system, so it picks up any changes to the asset right away.
 	ProcessingStage = 0;
 }
 
@@ -886,7 +887,7 @@ void FRenderAssetStreamingManager::RemoveStaticReferences(const UPrimitiveCompon
 
 /**
  * Called when a primitive is detached from an actor or another component.
- * Note: We should not be accessing the primitive or the UTexture2D after this call!
+ * Note: We should not be accessing the primitive or the UTexture after this call!
  */
 void FRenderAssetStreamingManager::NotifyPrimitiveDetached( const UPrimitiveComponent* Primitive )
 {
@@ -2010,27 +2011,46 @@ bool FRenderAssetStreamingManager::HandleListStreamingRenderAssetsCommand( const
 		const int32 WantedMipIndex = ResourceState.LODCountToAssetFirstLODIdx(StreamingRenderAsset.GetPerfectWantedMips());
 		const int32 MaxAllowedMipIndex = ResourceState.LODCountToAssetFirstLODIdx(StreamingRenderAsset.MaxAllowedMips);
 
-		if (AssetType == EStreamableRenderAssetType::Texture)
+		FTexturePlatformData** TexturePlatformData = Cast<UTexture>(RenderAsset) ? const_cast<UTexture*>(Cast<UTexture>(RenderAsset))->GetRunningPlatformData() : nullptr;
+		if (AssetType == EStreamableRenderAssetType::Texture && TexturePlatformData)
 		{
-			const UTexture2D* Texture = CastChecked<UTexture2D>(RenderAsset);
-			const TIndirectArray<struct FTexture2DMipMap>& Mips = Texture->PlatformData->Mips;
+			const UTexture2D* Texture2D = CastChecked<UTexture2D>(RenderAsset);
+			const UVolumeTexture* VolumeTexture = CastChecked<UVolumeTexture>(RenderAsset);
+			const UTexture2DArray* Texture2DArray = CastChecked<UTexture2DArray>(RenderAsset);
+			const TIndirectArray<struct FTexture2DMipMap>& TextureMips = (*TexturePlatformData)->Mips;
+
+			auto OutputMipsString = [&](int32 OutputIndex)->FString
+			{
+				const FTexture2DMipMap& OutputMips = TextureMips[OutputIndex];
+				if (Texture2D)
+				{
+					return FString::Printf(TEXT("%dx%d"), OutputMips.SizeX, OutputMips.SizeY);
+				}
+				else if (VolumeTexture)
+				{
+					return FString::Printf(TEXT("%dx%dx%d"), OutputMips.SizeX, OutputMips.SizeY, OutputMips.SizeZ);
+				}
+				else if (Texture2DArray)
+				{
+					return FString::Printf(TEXT("%dx%d*%d"), OutputMips.SizeX, OutputMips.SizeY, OutputMips.SizeZ);
+				}
+				else // Unkown type fallback
+				{
+					return FString::Printf(TEXT("%d?%d?%d"), OutputMips.SizeX, OutputMips.SizeY, OutputMips.SizeZ);
+				}
+			};
 
 			if (StreamingRenderAsset.LastRenderTime != MAX_FLT)
 			{
-				UE_LOG(LogContentStreaming, Log, TEXT("    Current=%dx%d Wanted=%dx%d MaxAllowed=%dx%d LastRenderTime=%.3f BudgetBias=%d Group=%s"),
-					Mips[CurrentMipIndex].SizeX, Mips[CurrentMipIndex].SizeY,
-					Mips[WantedMipIndex].SizeX, Mips[WantedMipIndex].SizeY,
-					Mips[MaxAllowedMipIndex].SizeX, Mips[MaxAllowedMipIndex].SizeY,
-					StreamingRenderAsset.LastRenderTime,
-					StreamingRenderAsset.BudgetMipBias,
+				UE_LOG(LogContentStreaming, Log, TEXT("    Current=%s  Wanted=%s MaxAllowed=%s LastRenderTime=%.3f BudgetBias=%d Group=%s"),
+					*OutputMipsString(CurrentMipIndex), *OutputMipsString(WantedMipIndex), *OutputMipsString(MaxAllowedMipIndex),
+					StreamingRenderAsset.LastRenderTime, StreamingRenderAsset.BudgetMipBias,
 					UTexture::GetTextureGroupString(static_cast<TextureGroup>(StreamingRenderAsset.LODGroup)));
 			}
 			else
 			{
-				UE_LOG(LogContentStreaming, Log, TEXT("    Current=%dx%d Wanted=%dx%d MaxAllowed=%dx%d BudgetBias=%d Group=%s"),
-					Mips[CurrentMipIndex].SizeX, Mips[CurrentMipIndex].SizeY,
-					Mips[WantedMipIndex].SizeX, Mips[WantedMipIndex].SizeY,
-					Mips[MaxAllowedMipIndex].SizeX, Mips[MaxAllowedMipIndex].SizeY,
+				UE_LOG(LogContentStreaming, Log, TEXT("    Current=%s Wanted=%s MaxAllowed=%s BudgetBias=%d Group=%s"),
+					*OutputMipsString(CurrentMipIndex), *OutputMipsString(WantedMipIndex), *OutputMipsString(MaxAllowedMipIndex),
 					StreamingRenderAsset.BudgetMipBias,
 					UTexture::GetTextureGroupString(static_cast<TextureGroup>(StreamingRenderAsset.LODGroup)));
 			}
@@ -2087,7 +2107,7 @@ bool FRenderAssetStreamingManager::HandleCancelRenderAssetStreamingCommand( cons
 {
 	FScopeLock ScopeLock(&CriticalSection);
 
-	UTexture2D::CancelPendingTextureStreaming();
+	UTexture::CancelPendingTextureStreaming();
 	UStaticMesh::CancelAllPendingStreamingActions();
 	USkeletalMesh::CancelAllPendingStreamingActions();
 	ULandscapeLODStreamingProxy::CancelAllPendingStreamingActions();
@@ -2344,16 +2364,115 @@ bool FRenderAssetStreamingManager::HandleStreamingManagerMemoryCommand( const TC
 
 bool FRenderAssetStreamingManager::HandleLODGroupsCommand( const TCHAR* Cmd, FOutputDevice& Ar )
 {
-	bDetailedDumpTextureGroupStats = FParse::Param(Cmd, TEXT("Detailed"));
-	bTriggerDumpTextureGroupStats = true;
-	// TODO: mesh LOD groups
+	FScopeLock ScopeLock(&CriticalSection);
+	SyncStates(true);
+
+	struct FTextureGroupStats
+	{
+		// Streaming texture stats
+		int32 NumStreamingTextures = 0;
+		uint64 CurrentTextureSize = 0;
+		uint64 WantedTextureSize = 0;
+		uint64 MaxTextureSize = 0;
+		// Non Streaming texture stats
+		int32 NumNonStreamingTextures = 0;
+		uint64 NonStreamingSize = 0;
+		// No resource texture
+		int32 NumNoResourceTextures = 0;
+	};
+	FTextureGroupStats TextureGroupStats[TEXTUREGROUP_MAX];
+
+	// Gather stats.
+	for (TObjectIterator<UTexture> It; It; ++It)
+	{
+		UTexture* Texture = *It;
+		check(Texture);
+
+		FTextureGroupStats& LODStats = TextureGroupStats[Texture->LODGroup];
+
+		const EPixelFormat PixelFormat = [&]()->EPixelFormat
+		{
+			if (Texture->GetRunningPlatformData())
+			{
+				return (*Texture->GetRunningPlatformData())->PixelFormat;
+			}
+			else if (Texture->Resource && Texture->Resource->TextureRHI)
+			{
+				return Texture->Resource->TextureRHI->GetFormat();
+			}
+			else
+			{
+				return PF_Unknown;
+			}
+		}();
+
+		// No resource no size taken
+		if (!Texture->Resource)
+		{
+			LODStats.NumNoResourceTextures++;
+		}
+		else if (Texture->IsStreamable())
+		{
+			FStreamingRenderAsset* StreamingTexture = GetStreamingRenderAsset(Texture);
+			if (ensure(StreamingTexture))
+			{
+				LODStats.NumStreamingTextures++;
+				LODStats.CurrentTextureSize += StreamingTexture->GetSize(StreamingTexture->ResidentMips);;
+				LODStats.WantedTextureSize += StreamingTexture->GetSize(StreamingTexture->WantedMips);
+				LODStats.MaxTextureSize += StreamingTexture->GetSize(StreamingTexture->MaxAllowedMips);
+			}
+		}
+		else
+		{
+			LODStats.NumNonStreamingTextures++;
+			LODStats.NonStreamingSize += Texture->CalcTextureMemorySizeEnum(TMC_ResidentMips);
+		}
+	}
+
+	// Output stats.
+	{
+		UE_LOG(LogContentStreaming, Log, TEXT("Texture memory usage:"));
+		FTextureGroupStats TotalStats;
+		for (int32 GroupIndex = 0; GroupIndex < TEXTUREGROUP_MAX; ++GroupIndex)
+		{
+			FTextureGroupStats& Stat = TextureGroupStats[GroupIndex];
+			if (Stat.NumStreamingTextures || Stat.NumNonStreamingTextures || Stat.NumNoResourceTextures)
+			{
+				TotalStats.NumStreamingTextures += Stat.NumStreamingTextures;
+				TotalStats.NumNonStreamingTextures += Stat.NumNonStreamingTextures;
+				TotalStats.CurrentTextureSize += Stat.CurrentTextureSize;
+				TotalStats.WantedTextureSize += Stat.WantedTextureSize;
+				TotalStats.MaxTextureSize += Stat.MaxTextureSize;
+				TotalStats.NonStreamingSize += Stat.NonStreamingSize;
+				TotalStats.NumNoResourceTextures += Stat.NumNoResourceTextures;
+				UE_LOG(LogContentStreaming, Log, TEXT("%34s: NumStreamingTextures=%4d { Current=%8.1f KB, Wanted=%8.1f KB, OnDisk=%8.1f KB }, NumNonStreaming=%4d { Size=%8.1f KB }, NumWithNoResource=%4d"),
+					UTexture::GetTextureGroupString((TextureGroup)GroupIndex),
+					Stat.NumStreamingTextures,
+					Stat.CurrentTextureSize / 1024.0f,
+					Stat.WantedTextureSize / 1024.0f,
+					Stat.MaxTextureSize / 1024.0f,
+					Stat.NumNonStreamingTextures,
+					Stat.NonStreamingSize / 1024.0f,
+					Stat.NumNoResourceTextures);
+			}
+		}
+		UE_LOG(LogContentStreaming, Log, TEXT("%34s: NumStreamingTextures=%4d { Current=%8.1f KB, Wanted=%8.1f KB, OnDisk=%8.1f KB }, NumNonStreaming=%4d { Size=%8.1f KB }, NumWithNoResource=%4d"),
+			TEXT("Total"),
+			TotalStats.NumStreamingTextures,
+			TotalStats.CurrentTextureSize / 1024.0f,
+			TotalStats.WantedTextureSize / 1024.0f,
+			TotalStats.MaxTextureSize / 1024.0f,
+			TotalStats.NumNonStreamingTextures,
+			TotalStats.NonStreamingSize / 1024.0f,
+			TotalStats.NumNoResourceTextures);
+	}
 	return true;
 }
 
 bool FRenderAssetStreamingManager::HandleInvestigateRenderAssetCommand(const TCHAR* Cmd, FOutputDevice& Ar, UWorld* InWorld)
 {
-	FScopeLock ScopeLock(&CriticalSection);
 
+	FScopeLock ScopeLock(&CriticalSection);
 	SyncStates(true);
 
 	FString InvestigateAssetName(FParse::Token(Cmd, 0));
@@ -2374,14 +2493,14 @@ bool FRenderAssetStreamingManager::HandleInvestigateRenderAssetCommand(const TCH
 				if (!RenderAsset) continue;
 				const EStreamableRenderAssetType AssetType = StreamingRenderAsset.RenderAssetType;
 				const FStreamableRenderResourceState ResourceState = RenderAsset->GetStreamableResourceState();
-				UTexture2D* Texture2D = Cast<UTexture2D>(RenderAsset);
+				UTexture* Texture = Cast<UTexture>(RenderAsset);
 				UStaticMesh* StaticMesh = Cast<UStaticMesh>(RenderAsset);
 				int32 CurrentMipIndex = ResourceState.LODCountToAssetFirstLODIdx(StreamingRenderAsset.ResidentMips);
 				int32 WantedMipIndex = ResourceState.LODCountToAssetFirstLODIdx(StreamingRenderAsset.GetPerfectWantedMips());
 				int32 MaxMipIndex = ResourceState.LODCountToAssetFirstLODIdx(StreamingRenderAsset.MaxAllowedMips);
 
 				UE_LOG(LogContentStreaming, Log, TEXT("%s: %s"), FStreamingRenderAsset::GetStreamingAssetTypeStr(AssetType), *AssetName);
-				FString LODGroupName = Texture2D ? UTexture::GetTextureGroupString((TextureGroup)StreamingRenderAsset.LODGroup) : TEXT("Unknown");
+				FString LODGroupName = Texture ? UTexture::GetTextureGroupString((TextureGroup)StreamingRenderAsset.LODGroup) : TEXT("Unknown");
 #if WITH_EDITORONLY_DATA
 				if (StaticMesh)
 				{
@@ -2392,13 +2511,16 @@ bool FRenderAssetStreamingManager::HandleInvestigateRenderAssetCommand(const TCH
 				const TCHAR* BucketNames[] = { TEXT("Largest"), TEXT("Larger"), TEXT("Default"), TEXT("Smaller"), TEXT("Smallest"), TEXT("Tiniest") };
 				if ((int32)FPlatformMemory::GetMemorySizeBucket() < UE_ARRAY_COUNT(BucketNames))
 				{
-					UE_LOG(LogContentStreaming, Log, TEXT("  LOD group:   %s [Bucket=%s]"), *LODGroupName, BucketNames[(int32)FPlatformMemory::GetMemorySizeBucket()]);
+					UE_LOG(LogContentStreaming, Log, TEXT("  LOD group:       %s [Bucket=%s]"), *LODGroupName, BucketNames[(int32)FPlatformMemory::GetMemorySizeBucket()]);
 				}
 				else
 				{
-					UE_LOG(LogContentStreaming, Log, TEXT("  LOD group:   %s [Unkown Bucket]"), *LODGroupName);
+					UE_LOG(LogContentStreaming, Log, TEXT("  LOD group:       %s [Unkown Bucket]"), *LODGroupName);
+				}
+				if (Texture && Texture->GetRunningPlatformData())
+				{
+					UE_LOG(LogContentStreaming, Log, TEXT("  Format:          %s"), GPixelFormats[(*Texture->GetRunningPlatformData())->PixelFormat].Name);
 				}
-
 				if (RenderAsset->bGlobalForceMipLevelsToBeResident)
 				{
 					UE_LOG(LogContentStreaming, Log, TEXT("  Force all mips:  bGlobalForceMipLevelsToBeResident"));
@@ -2414,17 +2536,34 @@ bool FRenderAssetStreamingManager::HandleInvestigateRenderAssetCommand(const TCH
 				}
 				else if (StreamingRenderAsset.bForceFullyLoadHeuristic)
 				{
-					UE_LOG(LogContentStreaming, Log, TEXT("  Force all mips: bForceFullyLoad"));
+					UE_LOG(LogContentStreaming, Log, TEXT("  Force all mips:  bForceFullyLoad"));
 				}
 				else if (ResourceState.MaxNumLODs == 1)
 				{
 					UE_LOG(LogContentStreaming, Log, TEXT("  Force all mips:  No mip-maps"));
 				}
 				
-				if (Texture2D)
+				if (Texture && Texture->GetRunningPlatformData())
 				{
-					UE_LOG(LogContentStreaming, Log, TEXT("  Current size [Mips]: %dx%d [%d]"), Texture2D->PlatformData->Mips[CurrentMipIndex].SizeX, Texture2D->PlatformData->Mips[CurrentMipIndex].SizeY, StreamingRenderAsset.ResidentMips);
-					UE_LOG(LogContentStreaming, Log, TEXT("  Wanted size [Mips]:  %dx%d [%d]"), Texture2D->PlatformData->Mips[WantedMipIndex].SizeX, Texture2D->PlatformData->Mips[WantedMipIndex].SizeY, StreamingRenderAsset.GetPerfectWantedMips());
+					const TIndirectArray<struct FTexture2DMipMap>& TextureMips = (*Texture->GetRunningPlatformData())->Mips;
+					UTexture2D* Texture2D = Cast<UTexture2D>(Texture);
+					UVolumeTexture* VolumeTexture = Cast<UVolumeTexture>(Texture);
+					UTexture2DArray* Texture2DArray = Cast<UTexture2DArray>(Texture);
+					if (Texture2D)
+					{
+						UE_LOG(LogContentStreaming, Log, TEXT("  Current size [2D Mips]: %dx%d [%d]"), TextureMips[CurrentMipIndex].SizeX, TextureMips[CurrentMipIndex].SizeY, StreamingRenderAsset.ResidentMips);
+						UE_LOG(LogContentStreaming, Log, TEXT("  Wanted size [2D Mips]:  %dx%d [%d]"), TextureMips[WantedMipIndex].SizeX, TextureMips[WantedMipIndex].SizeY, StreamingRenderAsset.GetPerfectWantedMips());
+					}
+					else if (VolumeTexture)
+					{
+						UE_LOG(LogContentStreaming, Log, TEXT("  Current size [3D Mips]: %dx%dx%d [%d]"), TextureMips[CurrentMipIndex].SizeX, TextureMips[CurrentMipIndex].SizeY, TextureMips[CurrentMipIndex].SizeZ, StreamingRenderAsset.ResidentMips);
+						UE_LOG(LogContentStreaming, Log, TEXT("  Wanted size [3D Mips]:  %dx%dx%d [%d]"), TextureMips[WantedMipIndex].SizeX, TextureMips[WantedMipIndex].SizeY, TextureMips[CurrentMipIndex].SizeZ, StreamingRenderAsset.GetPerfectWantedMips());
+					}
+					else if (Texture2DArray)
+					{
+						UE_LOG(LogContentStreaming, Log, TEXT("  Current size [2D Array Mips]: %dx%d*%d [%d]"), TextureMips[CurrentMipIndex].SizeX, TextureMips[CurrentMipIndex].SizeY, TextureMips[CurrentMipIndex].SizeZ, StreamingRenderAsset.ResidentMips);
+						UE_LOG(LogContentStreaming, Log, TEXT("  Wanted size [2D Array Mips]:  %dx%d*%d [%d]"), TextureMips[WantedMipIndex].SizeX, TextureMips[WantedMipIndex].SizeY, TextureMips[CurrentMipIndex].SizeZ, StreamingRenderAsset.GetPerfectWantedMips());
+					}
 				}
 				else
 				{
@@ -2467,7 +2606,7 @@ bool FRenderAssetStreamingManager::HandleInvestigateRenderAssetCommand(const TCH
 						// LOD group MaxResolution clamp : see UTextureLODSettings::CalculateLODBias(), included in CachedCombinedLODBias
 						const int32 MipCountBeforeMaxRes = ResourceState.MaxNumLODs - RenderAsset->NumCinematicMipLevels -
 							(StreamingRenderAsset.LODGroup == TEXTUREGROUP_UI ? GUITextureLODBias : 0) - 
-							(FPlatformProperties::RequiresCookedData() ? 0 : (LODGroupInfo.LODBias + (Texture2D ? Texture2D->LODBias : 0)));
+							(FPlatformProperties::RequiresCookedData() ? 0 : (LODGroupInfo.LODBias + (Texture ? Texture->LODBias : 0)));
 						const int32 MaxResBias = MipCountBeforeMaxRes - (LODGroupInfo.MaxLODMipCount + 1);
 						if (MaxResBias > 0)
 						{
@@ -2476,16 +2615,16 @@ bool FRenderAssetStreamingManager::HandleInvestigateRenderAssetCommand(const TCH
 						}
 
 						// Asset LODBias : see UTextureLODSettings::CalculateLODBias(), included in CachedCombinedLODBias
-						if (Texture2D && Texture2D->LODBias)
+						if (Texture && Texture->LODBias)
 						{
 							if (FPlatformProperties::RequiresCookedData())
 							{
-								BiasDesc += FString::Printf(TEXT(" [Asset.Bias:0(%d)]"), Texture2D->LODBias);
+								BiasDesc += FString::Printf(TEXT(" [Asset.Bias:0(%d)]"), Texture->LODBias);
 							}
 							else
 							{
-								BiasDesc += FString::Printf(TEXT(" [Asset.Bias:%d]"), Texture2D->LODBias);
-								CumuBias += Texture2D->LODBias;
+								BiasDesc += FString::Printf(TEXT(" [Asset.Bias:%d]"), Texture->LODBias);
+								CumuBias += Texture->LODBias;
 							}
 						}
 
@@ -2659,204 +2798,3 @@ bool FRenderAssetStreamingManager::Exec( UWorld* InWorld, const TCHAR* Cmd, FOut
 
 	return false;
 }
-
-void FRenderAssetStreamingManager::DumpTextureGroupStats( bool bDetailedStats )
-{
-	FScopeLock ScopeLock(&CriticalSection);
-
-	bTriggerDumpTextureGroupStats = false;
-#if !UE_BUILD_SHIPPING
-	struct FTextureGroupStats
-	{
-		FTextureGroupStats()
-		{
-			FMemory::Memzero( this, sizeof(FTextureGroupStats) );
-		}
-		int32 NumTextures;
-		int32 NumNonStreamingTextures;
-		int64 CurrentTextureSize;
-		int64 WantedTextureSize;
-		int64 MaxTextureSize;
-		int64 NonStreamingSize;
-	};
-	FTextureGroupStats TextureGroupStats[TEXTUREGROUP_MAX];
-	FTextureGroupStats TextureGroupWaste[TEXTUREGROUP_MAX];
-	int64 NumNonStreamingTextures = 0;
-	int64 NonStreamingSize = 0;
-	int32 NumNonStreamingPoolTextures = 0;
-	int64 NonStreamingPoolSize = 0;
-	int64 TotalSavings = 0;
-//	int32 UITexels = 0;
-	int32 NumDXT[PF_MAX];
-	int32 NumNonSaved[PF_MAX];
-	int32 NumOneMip[PF_MAX];
-	int32 NumBadAspect[PF_MAX];
-	int32 NumTooSmall[PF_MAX];
-	int32 NumNonPow2[PF_MAX];
-	int32 NumNULLResource[PF_MAX];
-	FMemory::Memzero( &NumDXT, sizeof(NumDXT) );
-	FMemory::Memzero( &NumNonSaved, sizeof(NumNonSaved) );
-	FMemory::Memzero( &NumOneMip, sizeof(NumOneMip) );
-	FMemory::Memzero( &NumBadAspect, sizeof(NumBadAspect) );
-	FMemory::Memzero( &NumTooSmall, sizeof(NumTooSmall) );
-	FMemory::Memzero( &NumNonPow2, sizeof(NumNonPow2) );
-	FMemory::Memzero( &NumNULLResource, sizeof(NumNULLResource) );
-
-	// Gather stats.
-	for( TObjectIterator<UTexture> It; It; ++It )
-	{
-		UTexture* Texture = *It;
-		UTexture2D* Texture2D = Cast<UTexture2D>(Texture);
-		FTextureGroupStats& Stat = TextureGroupStats[Texture->LODGroup];
-		FTextureGroupStats& Waste = TextureGroupWaste[Texture->LODGroup];
-		FStreamingRenderAsset* StreamingTexture = GetStreamingRenderAsset(Texture2D);
-		uint32 TextureAlign = 0;
-		FRHIResourceCreateInfo CreateInfo(Texture2D ? Texture2D->GetExtData() : 0);
-		if ( StreamingTexture )
-		{
-			Stat.NumTextures++;
-			Stat.CurrentTextureSize += StreamingTexture->GetSize( StreamingTexture->ResidentMips );
-			Stat.WantedTextureSize += StreamingTexture->GetSize( StreamingTexture->WantedMips );
-			Stat.MaxTextureSize += StreamingTexture->GetSize( StreamingTexture->MaxAllowedMips );
-			
-			int64 WasteCurrent = StreamingTexture->GetSize( StreamingTexture->ResidentMips ) - RHICalcTexture2DPlatformSize(Texture2D->GetSizeX(), Texture2D->GetSizeY(), Texture2D->GetPixelFormat(), StreamingTexture->ResidentMips, 1, TexCreate_None, CreateInfo, TextureAlign);
-
-			int64 WasteWanted = StreamingTexture->GetSize( StreamingTexture->WantedMips ) - RHICalcTexture2DPlatformSize(Texture2D->GetSizeX(), Texture2D->GetSizeY(), Texture2D->GetPixelFormat(), StreamingTexture->WantedMips, 1, TexCreate_None, CreateInfo, TextureAlign);
-
-			int64 WasteMaxSize = StreamingTexture->GetSize( StreamingTexture->MaxAllowedMips ) - RHICalcTexture2DPlatformSize(Texture2D->GetSizeX(), Texture2D->GetSizeY(), Texture2D->GetPixelFormat(), StreamingTexture->MaxAllowedMips, 1, TexCreate_None, CreateInfo, TextureAlign);
-
-			Waste.NumTextures++;
-			Waste.CurrentTextureSize += FMath::Max<int64>(WasteCurrent,0);
-			Waste.WantedTextureSize += FMath::Max<int64>(WasteWanted,0);
-			Waste.MaxTextureSize += FMath::Max<int64>(WasteMaxSize,0);
-		}
-		else
-		{
-
-			bool bIsPooledTexture = Texture->Resource && IsValidRef(Texture->Resource->TextureRHI) && appIsPoolTexture( Texture->Resource->TextureRHI );
-			int64 TextureSize = Texture->CalcTextureMemorySizeEnum(TMC_ResidentMips);
-			Stat.NumNonStreamingTextures++;
-			Stat.NonStreamingSize += TextureSize;
-			if ( Texture2D && Texture2D->Resource )
-			{				
-				int64 WastedSize = TextureSize - RHICalcTexture2DPlatformSize(Texture2D->GetSizeX(), Texture2D->GetSizeY(), Texture2D->GetPixelFormat(), Texture2D->GetNumMips(), 1, TexCreate_None, CreateInfo, TextureAlign);
-
-				Waste.NumNonStreamingTextures++;
-				Waste.NonStreamingSize += FMath::Max<int64>(WastedSize, 0);
-			}
-			if ( bIsPooledTexture )
-			{
-				NumNonStreamingPoolTextures++;
-				NonStreamingPoolSize += TextureSize;
-			}
-			else
-			{
-				NumNonStreamingTextures++;
-				NonStreamingSize += TextureSize;
-			}
-		}
-
-		if ( Texture2D && (Texture2D->GetPixelFormat() == PF_DXT1 || Texture2D->GetPixelFormat() == PF_DXT5) )
-		{
-			NumDXT[Texture2D->GetPixelFormat()]++;
-			if ( Texture2D->Resource )
-			{
-				// Track the reasons we couldn't save any memory from the mip-tail.
-				NumNonSaved[Texture2D->GetPixelFormat()]++;
-				if ( Texture2D->GetNumMips() < 2 )
-				{
-					NumOneMip[Texture2D->GetPixelFormat()]++;
-				}
-				else if ( Texture2D->GetSizeX() > Texture2D->GetSizeY() * 2 || Texture2D->GetSizeY() > Texture2D->GetSizeX() * 2 )
-				{
-					NumBadAspect[Texture2D->GetPixelFormat()]++;
-				}
-				else if ( Texture2D->GetSizeX() < 16 || Texture2D->GetSizeY() < 16 || Texture2D->GetNumMips() < 5 )
-				{
-					NumTooSmall[Texture2D->GetPixelFormat()]++;
-				}
-				else if ( (Texture2D->GetSizeX() & (Texture2D->GetSizeX() - 1)) != 0 || (Texture2D->GetSizeY() & (Texture2D->GetSizeY() - 1)) != 0 )
-				{
-					NumNonPow2[Texture2D->GetPixelFormat()]++;
-				}
-				else
-				{
-					// Unknown reason
-					int32 Q=0;
-				}
-			}
-			else
-			{
-				NumNULLResource[Texture2D->GetPixelFormat()]++;
-			}
-		}
-	}
-
-	// Output stats.
-	{
-		UE_LOG(LogContentStreaming, Log, TEXT("Texture memory usage:"));
-		FTextureGroupStats TotalStats;
-		for ( int32 GroupIndex=0; GroupIndex < TEXTUREGROUP_MAX; ++GroupIndex )
-		{
-			FTextureGroupStats& Stat = TextureGroupStats[GroupIndex];
-			TotalStats.NumTextures				+= Stat.NumTextures;
-			TotalStats.NumNonStreamingTextures	+= Stat.NumNonStreamingTextures;
-			TotalStats.CurrentTextureSize		+= Stat.CurrentTextureSize;
-			TotalStats.WantedTextureSize		+= Stat.WantedTextureSize;
-			TotalStats.MaxTextureSize			+= Stat.MaxTextureSize;
-			TotalStats.NonStreamingSize			+= Stat.NonStreamingSize;
-			UE_LOG(LogContentStreaming, Log, TEXT("%34s: NumTextures=%4d, Current=%8.1f KB, Wanted=%8.1f KB, OnDisk=%8.1f KB, NumNonStreaming=%4d, NonStreaming=%8.1f KB"),
-				UTexture::GetTextureGroupString((TextureGroup)GroupIndex),
-				Stat.NumTextures,
-				Stat.CurrentTextureSize / 1024.0f,
-				Stat.WantedTextureSize / 1024.0f,
-				Stat.MaxTextureSize / 1024.0f,
-				Stat.NumNonStreamingTextures,
-				Stat.NonStreamingSize / 1024.0f );
-		}
-		UE_LOG(LogContentStreaming, Log, TEXT("%34s: NumTextures=%4d, Current=%8.1f KB, Wanted=%8.1f KB, OnDisk=%8.1f KB, NumNonStreaming=%4d, NonStreaming=%8.1f KB"),
-			TEXT("Total"),
-			TotalStats.NumTextures,
-			TotalStats.CurrentTextureSize / 1024.0f,
-			TotalStats.WantedTextureSize / 1024.0f,
-			TotalStats.MaxTextureSize / 1024.0f,
-			TotalStats.NumNonStreamingTextures,
-			TotalStats.NonStreamingSize / 1024.0f );
-	}
-	if ( bDetailedStats )
-	{
-		UE_LOG(LogContentStreaming, Log, TEXT("Wasted memory due to inefficient texture storage:"));
-		FTextureGroupStats TotalStats;
-		for ( int32 GroupIndex=0; GroupIndex < TEXTUREGROUP_MAX; ++GroupIndex )
-		{
-			FTextureGroupStats& Stat = TextureGroupWaste[GroupIndex];
-			TotalStats.NumTextures				+= Stat.NumTextures;
-			TotalStats.NumNonStreamingTextures	+= Stat.NumNonStreamingTextures;
-			TotalStats.CurrentTextureSize		+= Stat.CurrentTextureSize;
-			TotalStats.WantedTextureSize		+= Stat.WantedTextureSize;
-			TotalStats.MaxTextureSize			+= Stat.MaxTextureSize;
-			TotalStats.NonStreamingSize			+= Stat.NonStreamingSize;
-			UE_LOG(LogContentStreaming, Log, TEXT("%34s: NumTextures=%4d, Current=%8.1f KB, Wanted=%8.1f KB, OnDisk=%8.1f KB, NumNonStreaming=%4d, NonStreaming=%8.1f KB"),
-				UTexture::GetTextureGroupString((TextureGroup)GroupIndex),
-				Stat.NumTextures,
-				Stat.CurrentTextureSize / 1024.0f,
-				Stat.WantedTextureSize / 1024.0f,
-				Stat.MaxTextureSize / 1024.0f,
-				Stat.NumNonStreamingTextures,
-				Stat.NonStreamingSize / 1024.0f );
-		}
-		UE_LOG(LogContentStreaming, Log, TEXT("%34s: NumTextures=%4d, Current=%8.1f KB, Wanted=%8.1f KB, OnDisk=%8.1f KB, NumNonStreaming=%4d, NonStreaming=%8.1f KB"),
-			TEXT("Total Wasted"),
-			TotalStats.NumTextures,
-			TotalStats.CurrentTextureSize / 1024.0f,
-			TotalStats.WantedTextureSize / 1024.0f,
-			TotalStats.MaxTextureSize / 1024.0f,
-			TotalStats.NumNonStreamingTextures,
-			TotalStats.NonStreamingSize / 1024.0f );
-	}
-
-	//@TODO: Calculate memory usage for non-pool textures properly!
-//	UE_LOG(LogContentStreaming, Log,  TEXT("%34s: NumTextures=%4d, Current=%7.1f KB"), TEXT("Non-streaming pool textures"), NumNonStreamingPoolTextures, NonStreamingPoolSize/1024.0f );
-//	UE_LOG(LogContentStreaming, Log,  TEXT("%34s: NumTextures=%4d, Current=%7.1f KB"), TEXT("Non-streaming non-pool textures"), NumNonStreamingTextures, NonStreamingSize/1024.0f );
-#endif // !UE_BUILD_SHIPPING
-}
diff --git a/Engine/Source/Runtime/Engine/Private/Streaming/StreamingManagerTexture.h b/Engine/Source/Runtime/Engine/Private/Streaming/StreamingManagerTexture.h
--- a/Engine/Source/Runtime/Engine/Private/Streaming/StreamingManagerTexture.h
+++ b/Engine/Source/Runtime/Engine/Private/Streaming/StreamingManagerTexture.h
@@ -270,8 +270,6 @@ private:
 	 */
 	void	SetRenderAssetsRemovedTimestamp(const FRemovedRenderAssetArray& RemovedRenderAssets);
 
-	void	DumpTextureGroupStats( bool bDetailedStats );
-
 	void	SetLastUpdateTime();
 	void	UpdateStats();
 	void	UpdateCSVOnlyStats();
@@ -306,12 +304,6 @@ private:
 	 */
 	void TickDeferredMipLevelChangeCallbacks();
 
-	/** Next sync, dump texture group stats. */
-	bool	bTriggerDumpTextureGroupStats;
-
-	/** Whether to the dumped texture group stats should contain extra information. */
-	bool	bDetailedDumpTextureGroupStats;
-
 	/** Cached from the system settings. */
 	int32 NumStreamedMips_Texture[TEXTUREGROUP_MAX];
 	TArray<int32> NumStreamedMips_StaticMesh;
diff --git a/Engine/Source/Runtime/Engine/Private/Streaming/StreamingTexture.h b/Engine/Source/Runtime/Engine/Private/Streaming/StreamingTexture.h
--- a/Engine/Source/Runtime/Engine/Private/Streaming/StreamingTexture.h
+++ b/Engine/Source/Runtime/Engine/Private/Streaming/StreamingTexture.h
@@ -7,13 +7,14 @@ StreamingTexture.h: Definitions of classes used for texture streaming.
 #pragma once
 
 #include "CoreMinimal.h"
-#include "Engine/Texture2D.h"
+#include "Engine/StreamableRenderAsset.h"
+#include "Engine/TextureDefines.h"
 
 struct FRenderAssetStreamingManager;
 struct FRenderAssetStreamingSettings;
 
 /*-----------------------------------------------------------------------------
-	FStreamingRenderAsset, the streaming system's version of UTexture2D.
+	FStreamingRenderAsset, the streaming system's version of UStreamableRenderAsset.
 -----------------------------------------------------------------------------*/
 
 /** Self-contained structure to manage a streaming texture/mesh, possibly on a separate thread. */
diff --git a/Engine/Source/Runtime/Engine/Private/Streaming/Texture2DArrayStreaming.cpp b/Engine/Source/Runtime/Engine/Private/Streaming/Texture2DArrayStreaming.cpp
new file mode 100644
--- /dev/null
+++ b/Engine/Source/Runtime/Engine/Private/Streaming/Texture2DArrayStreaming.cpp
@@ -0,0 +1,183 @@
+// Copyright Epic Games, Inc. All Rights Reserved.
+
+/*=============================================================================
+Texture2DArrayStreaming.cpp: Helpers to stream in and out texture 2D array LODs.
+=============================================================================*/
+
+#include "Streaming/Texture2DArrayStreaming.h"
+
+//*****************************************************************************
+//***************************** Global Definitions ****************************
+//*****************************************************************************
+
+void RHICopySharedMips(FRHICommandList& RHICmdList, FRHITexture2DArray* DestTexture, FRHITexture2DArray* SrcTexture)
+{
+	// Transition to copy source and dest
+	{
+		FRHITransitionInfo TransitionsBefore[] = { FRHITransitionInfo(SrcTexture, ERHIAccess::SRVMask, ERHIAccess::CopySrc), FRHITransitionInfo(DestTexture, ERHIAccess::SRVMask, ERHIAccess::CopyDest) };
+		RHICmdList.Transition(MakeArrayView(TransitionsBefore, UE_ARRAY_COUNT(TransitionsBefore)));
+	}
+
+	// Copy 
+	{
+		FRHICopyTextureInfo CopyInfo;
+
+		auto SetCopyInfo = [&](FRHITexture2DArray* Texture2DArray)
+		{
+			CopyInfo.Size.X = Texture2DArray->GetSizeX();
+			CopyInfo.Size.Y = Texture2DArray->GetSizeY();
+			CopyInfo.NumSlices = Texture2DArray->GetSizeZ();
+			CopyInfo.NumMips = Texture2DArray->GetNumMips();
+		};
+
+		if (DestTexture->GetNumMips() < SrcTexture->GetNumMips())
+		{
+			SetCopyInfo(DestTexture);
+		}
+		else
+		{
+			SetCopyInfo(SrcTexture);
+		}
+
+		CopyInfo.SourceMipIndex = SrcTexture->GetNumMips() - CopyInfo.NumMips;
+		CopyInfo.DestMipIndex = DestTexture->GetNumMips() - CopyInfo.NumMips;
+		RHICmdList.CopyTexture(SrcTexture, DestTexture, CopyInfo);
+	}
+
+	// Transition to SRV
+	{
+		FRHITransitionInfo TransitionsAfter[] = { FRHITransitionInfo(SrcTexture, ERHIAccess::CopySrc, ERHIAccess::SRVMask), FRHITransitionInfo(DestTexture, ERHIAccess::CopyDest, ERHIAccess::SRVMask) };
+		RHICmdList.Transition(MakeArrayView(TransitionsAfter, UE_ARRAY_COUNT(TransitionsAfter)));
+	}
+}
+
+//*****************************************************************************
+//******************* FTexture2DArrayMipAllocator_Reallocate ******************
+//*****************************************************************************
+
+FTexture2DArrayMipAllocator_Reallocate::FTexture2DArrayMipAllocator_Reallocate(UTexture* Texture)
+	: FTextureMipAllocator(Texture, ETickState::AllocateMips, ETickThread::Async)
+{
+}
+
+FTexture2DArrayMipAllocator_Reallocate::~FTexture2DArrayMipAllocator_Reallocate()
+{
+}
+
+bool FTexture2DArrayMipAllocator_Reallocate::AllocateMips(const FTextureUpdateContext& Context, FTextureMipInfoArray& OutMipInfos, const FTextureUpdateSyncOptions& SyncOptions)
+{
+	OutMipInfos.AddDefaulted(CurrentFirstLODIdx);
+	StreamedMipData.AddDefaulted(CurrentFirstLODIdx);
+	StreamedSliceSize.AddZeroed(CurrentFirstLODIdx);
+
+	// Allocate the mip memory as temporary buffers so that the FTextureMipDataProvider implementation can write to it.
+	for (int32 MipIdx = PendingFirstLODIdx; MipIdx < CurrentFirstLODIdx; ++MipIdx)
+	{
+		const FTexture2DMipMap& OwnerMip = *Context.MipsView[MipIdx];
+		FTextureMipInfo& MipInfo = OutMipInfos[MipIdx];
+
+		// Note that streamed in mips size will always be at least as big as block size.
+		MipInfo.Format = Context.Resource->GetPixelFormat();
+		MipInfo.SizeX = OwnerMip.SizeX;
+		MipInfo.SizeY = OwnerMip.SizeY;
+		MipInfo.ArraySize = OwnerMip.SizeZ;
+
+		uint32 TextureAlign = 0;
+
+		MipInfo.DataSize = MipInfo.ArraySize * CalcTextureMipMapSize(MipInfo.SizeX, MipInfo.SizeY, MipInfo.Format, 0);
+		MipInfo.DestData = FMemory::Malloc(MipInfo.DataSize, 16);
+
+		check(MipInfo.DataSize % MipInfo.ArraySize == 0);
+
+		StreamedMipData[MipIdx].Reset((uint8*)MipInfo.DestData);
+		StreamedSliceSize[MipIdx] = MipInfo.DataSize / MipInfo.ArraySize;
+	}
+
+	AdvanceTo(ETickState::FinalizeMips, ETickThread::Render);
+	return true;
+}
+
+bool FTexture2DArrayMipAllocator_Reallocate::FinalizeMips(const FTextureUpdateContext& Context, const FTextureUpdateSyncOptions& SyncOptions)
+{
+	const uint32 NumSlices = Context.Resource ? Context.Resource->GetSizeZ() : 0;
+
+	// Create new Texture.
+	{
+		FRHIResourceCreateInfo CreateInfo;
+		const FTexture2DMipMap& FirstMip = *Context.MipsView[PendingFirstLODIdx];
+		CreateInfo.ExtData = Context.Resource->GetExtData();
+		IntermediateTextureRHI = RHICreateTexture2DArray(FirstMip.SizeX, FirstMip.SizeY, FirstMip.SizeZ, Context.Resource->GetPixelFormat(), ResourceState.NumRequestedLODs, 1, Context.Resource->GetCreationFlags(), CreateInfo);
+	}
+
+	// Copy shared mips.
+	{
+		bool bCopySharedMipsDone = false;
+		ENQUEUE_RENDER_COMMAND(FCopySharedMipsForTexture2DArray)(
+			[&](FRHICommandListImmediate& RHICmdList)
+		{
+			RHICopySharedMips(RHICmdList, IntermediateTextureRHI.GetReference(), Context.Resource->GetTexture2DArrayRHI());
+			bCopySharedMipsDone = true;
+		});
+		// Expected to execute immediately since ran on the renderthread.
+		check(bCopySharedMipsDone);
+	}
+
+	// Update the streamed in mips if they were not initialized from the bulk data.
+	for (int32 MipIdx = PendingFirstLODIdx; MipIdx < CurrentFirstLODIdx; ++MipIdx)
+	{
+		const FTexture2DMipMap& Mip = *Context.MipsView[MipIdx];
+		const uint8* MipData = StreamedMipData[MipIdx].Get();
+		const uint64 SlizeSize = StreamedSliceSize[MipIdx];
+		if (MipData && SlizeSize)
+		{
+			const int32 RHIMipIndex = MipIdx - PendingFirstLODIdx;
+			for (uint32 SliceIdx = 0; SliceIdx < NumSlices; ++SliceIdx)
+			{
+				uint32 DestStride = 0;
+				void* DestData = RHILockTexture2DArray(IntermediateTextureRHI, SliceIdx, RHIMipIndex, RLM_WriteOnly, DestStride, false);
+				FMemory::Memcpy(DestData, MipData + SliceIdx * SlizeSize, SlizeSize);
+				RHIUnlockTexture2DArray(IntermediateTextureRHI, SliceIdx, RHIMipIndex, false);
+			}
+
+		}
+	}
+	StreamedMipData.Empty();
+	StreamedSliceSize.Empty();
+
+	// Use the new texture resource for the texture asset, must run on the renderthread.
+	Context.Resource->FinalizeStreaming(IntermediateTextureRHI);
+	// No need for the intermediate texture anymore.
+	IntermediateTextureRHI.SafeRelease();
+
+	// Update is complete, nothing more to do.
+	AdvanceTo(ETickState::Done, ETickThread::None);
+	return true;
+}
+
+void FTexture2DArrayMipAllocator_Reallocate::Cancel(const FTextureUpdateSyncOptions& SyncOptions)
+{
+	// Release the intermediate texture. If not null, this will be on the renderthread.
+	IntermediateTextureRHI.SafeRelease();
+	// Release the temporary mip data. Can be run on either renderthread or async threads.
+	StreamedMipData.Empty();
+	StreamedSliceSize.Empty();
+}
+
+FTextureMipAllocator::ETickThread FTexture2DArrayMipAllocator_Reallocate::GetCancelThread() const
+{
+	// If there is an  intermediate texture, it is safer to released on the renderthread.
+	if (IntermediateTextureRHI)
+	{
+		return ETickThread::Render;
+	}
+	// Otherwise, if there are only temporary mip data, they can be freed on any threads.
+	else if (StreamedMipData.Num())
+	{
+		return ETickThread::Async;
+	}
+	// Nothing to do.
+	else
+	{
+		return ETickThread::None;
+	}
+}
diff --git a/Engine/Source/Runtime/Engine/Private/Streaming/Texture2DArrayStreaming.h b/Engine/Source/Runtime/Engine/Private/Streaming/Texture2DArrayStreaming.h
new file mode 100644
--- /dev/null
+++ b/Engine/Source/Runtime/Engine/Private/Streaming/Texture2DArrayStreaming.h
@@ -0,0 +1,59 @@
+// Copyright Epic Games, Inc. All Rights Reserved.
+
+/*=============================================================================
+Texture2DArrayStreaming.h: Helpers to stream in and out texture 2D array LODs.
+=============================================================================*/
+
+#pragma once
+
+#include "CoreMinimal.h"
+#include "RenderAssetUpdate.h"
+#include "Engine/Texture2DArray.h"
+#include "Rendering/Texture2DResource.h"
+#include "Texture2DMipDataProvider_DDC.h"
+#include "Texture2DMipDataProvider_IO.h"
+#include "TextureMipAllocator.h"
+
+#if WITH_EDITORONLY_DATA
+typedef FTexture2DMipDataProvider_DDC FTexture2DArrayMipDataProvider_DDC;
+#endif
+
+typedef FTexture2DMipDataProvider_IO FTexture2DArrayMipDataProvider_IO;
+
+/**
+* Implements a texture3d mip resizing strategy based on creating a duplicate texture and copyinng the shared mips.
+*/
+class FTexture2DArrayMipAllocator_Reallocate : public FTextureMipAllocator
+{
+public:
+
+	FTexture2DArrayMipAllocator_Reallocate(UTexture* Texture);
+	~FTexture2DArrayMipAllocator_Reallocate();
+
+	// ********************************************************
+	// ********* FTextureMipAllocator implementation **********
+	// ********************************************************
+
+	bool AllocateMips(const FTextureUpdateContext& Context, FTextureMipInfoArray& OutMipInfos, const FTextureUpdateSyncOptions& SyncOptions) final override;
+	bool FinalizeMips(const FTextureUpdateContext& Context, const FTextureUpdateSyncOptions& SyncOptions) final override;
+	void Cancel(const FTextureUpdateSyncOptions& SyncOptions) final override;
+	ETickThread GetCancelThread() const final override;
+
+protected:
+
+	// Used to free memory allocated from a TUniquePtr
+	struct FMallocDeleter
+	{
+		void operator() (uint8* Buffer)
+		{
+			FMemory::Free(Buffer);
+		}
+	};
+
+	// The intermediate texture async created in the update process.
+	FTexture2DArrayRHIRef IntermediateTextureRHI;
+
+	// The temporary main memory allocations holding the mip data
+	TArray<TUniquePtr<uint8, FMallocDeleter>, TInlineAllocator<MAX_TEXTURE_MIP_COUNT>> StreamedMipData;
+	TArray<uint64, TInlineAllocator<MAX_TEXTURE_MIP_COUNT>> StreamedSliceSize;
+};
diff --git a/Engine/Source/Runtime/Engine/Private/Streaming/Texture2DMipDataProvider_DDC.cpp b/Engine/Source/Runtime/Engine/Private/Streaming/Texture2DMipDataProvider_DDC.cpp
--- a/Engine/Source/Runtime/Engine/Private/Streaming/Texture2DMipDataProvider_DDC.cpp
+++ b/Engine/Source/Runtime/Engine/Private/Streaming/Texture2DMipDataProvider_DDC.cpp
@@ -6,9 +6,10 @@ Texture2DMipDataProvider_DDC.cpp : Implementation of FTextureMipDataProvider usi
 
 #include "Texture2DMipDataProvider_DDC.h"
 #include "DerivedDataCacheInterface.h"
-#include "Engine/Texture2D.h"
+#include "Engine/Texture.h"
 #include "TextureResource.h"
 #include "Streaming/Texture2DStreamIn_DDC.h" // for FAbandonedDDCHandleManager
+#include "Streaming/TextureStreamingHelpers.h"
 
 #if WITH_EDITORONLY_DATA
 
@@ -76,13 +77,38 @@ int32 FTexture2DMipDataProvider_DDC::GetMips(
 				int32 MipSize = 0;
 				Ar << MipSize;
 
+				const uint32 DepthOrArraySize = FMath::Max<uint32>(MipInfo.ArraySize, MipInfo.SizeZ);
 				if (MipSize == MipInfo.DataSize)
 				{
 					Ar.Serialize(MipInfo.DestData, MipSize);
 					bSuccess = true;
 				}
+				else if (MipSize < MipInfo.DataSize && DepthOrArraySize > 1 && MipInfo.DataSize % DepthOrArraySize == 0 && MipSize % DepthOrArraySize == 0)
+				{
+					UE_LOG(LogTexture, Verbose, TEXT("DDC mip size smaller than streaming buffer size. (%s, Mip %d): %d KB / %d KB."), *Context.Resource->GetTextureName().ToString(), ResourceState.MaxNumLODs - MipIndex, MipInfo.DataSize / 1024, MipSize / 1024);
+
+					const uint64 SourceSubSize = MipSize / DepthOrArraySize;
+					const uint64 DestSubSize = MipInfo.DataSize / DepthOrArraySize;
+					const uint64 PaddingSubSize = DestSubSize - SourceSubSize;
+
+					uint8* DestData = (uint8*)MipInfo.DestData;
+					for (uint32 SubIdx = 0; SubIdx < DepthOrArraySize; ++SubIdx)
+					{
+						Ar.Serialize(DestData, SourceSubSize);
+						DestData += SourceSubSize;
+						FMemory::Memzero(DestData, PaddingSubSize);
+						DestData += PaddingSubSize;
+					}
+					bSuccess = true;
+				}
+				else
+				{
+					UE_LOG(LogTexture, Warning, TEXT("Mismatch between DDC mip size and streaming buffer size. (%s, Mip %d): %d KB / %d KB."), *Context.Resource->GetTextureName().ToString(), ResourceState.MaxNumLODs - MipIndex, MipInfo.DataSize / 1024, MipSize / 1024);
+					FMemory::Memzero(MipInfo.DestData, MipInfo.DataSize);
+				}
 			}
 		}
+
 		if (!bSuccess)
 		{
 			AdvanceTo(ETickState::CleanUp, ETickThread::Async);
diff --git a/Engine/Source/Runtime/Engine/Private/Streaming/Texture2DMipDataProvider_IO.cpp b/Engine/Source/Runtime/Engine/Private/Streaming/Texture2DMipDataProvider_IO.cpp
--- a/Engine/Source/Runtime/Engine/Private/Streaming/Texture2DMipDataProvider_IO.cpp
+++ b/Engine/Source/Runtime/Engine/Private/Streaming/Texture2DMipDataProvider_IO.cpp
@@ -5,19 +5,22 @@ Texture2DMipDataProvider_IO.cpp : Implementation of FTextureMipDataProvider usin
 =============================================================================*/
 
 #include "Texture2DMipDataProvider_IO.h"
-#include "Engine/Texture2D.h"
+#include "Engine/Texture.h"
 #include "TextureResource.h"
 #include "HAL/PlatformFileManager.h"
 #include "HAL/FileManager.h"
-#include "Misc/PackageSegment.h"
 #include "Misc/Paths.h"
 #include "Streaming/TextureStreamingHelpers.h"
-#include "UObject/PackageResourceManager.h"
+#include "ContentStreaming.h"
 
 FTexture2DMipDataProvider_IO::FTexture2DMipDataProvider_IO(const UTexture* InTexture, bool InPrioritizedIORequest)
 	: FTextureMipDataProvider(InTexture, ETickState::Init, ETickThread::Async)
 	, bPrioritizedIORequest(InPrioritizedIORequest)
 {
+	if (InTexture)
+	{
+		TextureName = InTexture->GetFName();
+	}
 }
 
 FTexture2DMipDataProvider_IO::~FTexture2DMipDataProvider_IO()
@@ -27,11 +30,24 @@ FTexture2DMipDataProvider_IO::~FTexture2DMipDataProvider_IO()
 
 void FTexture2DMipDataProvider_IO::Init(const FTextureUpdateContext& Context, const FTextureUpdateSyncOptions& SyncOptions)
 {
-	int32 CurrentFileIndex = INDEX_NONE;
+	IORequests.AddDefaulted(CurrentFirstLODIdx);
 
-	for (int32 MipIndex = PendingFirstLODIdx; MipIndex < CurrentFirstLODIdx; ++MipIndex)
+	// If this resource has optional LODs and we are streaming one of them.
+	if (ResourceState.NumNonOptionalLODs < ResourceState.MaxNumLODs && PendingFirstLODIdx < ResourceState.LODCountToFirstLODIdx(ResourceState.NumNonOptionalLODs))
+	{
+		// Generate the FilenameHash of each optional LOD before the first one requested, so that we can handle properly PAK unmount events.
+		// Note that streamer only stores the hash for the first optional mip.
+		for (int32 MipIdx = 0; MipIdx < PendingFirstLODIdx; ++MipIdx)
+		{
+			const FTexture2DMipMap& OwnerMip = *Context.MipsView[MipIdx];
+			IORequests[MipIdx].FilenameHash = OwnerMip.BulkData.GetIoFilenameHash();
+		}
+	}
+	
+	// Otherwise validate each streamed in mip.
+	for (int32 MipIdx = PendingFirstLODIdx; MipIdx < CurrentFirstLODIdx; ++MipIdx)
 	{
-		const FTexture2DMipMap& OwnerMip = *Context.MipsView[MipIndex];
+		const FTexture2DMipMap& OwnerMip = *Context.MipsView[MipIdx];
 		if (OwnerMip.BulkData.IsStoredCompressedOnDisk())
 		{
 			// Compression at the package level is no longer supported
@@ -42,43 +58,9 @@ void FTexture2DMipDataProvider_IO::Init(const FTextureUpdateContext& Context, co
 			// Invalid bulk data size.
 			continue;
 		}
-
-		FPackagePath PackagePath = OwnerMip.BulkData.GetPackagePath();
-		EPackageSegment PackageSegment = OwnerMip.BulkData.GetPackageSegment();
-		if (FileInfos.IsValidIndex(CurrentFileIndex))
-		{
-			FFileInfo& FileInfo = FileInfos[CurrentFileIndex];
-			if (FileInfo.PackagePath == PackagePath && FileInfo.LastMipIndex + 1 == MipIndex)
-			{
-				FileInfo.LastMipIndex = MipIndex;
-			}
-			else // Otherwise create a new entry
-			{
-				CurrentFileIndex = INDEX_NONE;
-			}
-		}
-
-		if (CurrentFileIndex == INDEX_NONE)
+		else
 		{
-			int64 IOFileOffset = 0;
-			if (GEventDrivenLoaderEnabled)
-			{
-				if (PackageSegment == EPackageSegment::Header)
-				{
-					IOFileOffset = -IPackageResourceManager::Get().FileSize(PackagePath, EPackageSegment::Header);
-					check(IOFileOffset < 0);
-					PackageSegment = EPackageSegment::Exports;
-					UE_LOG(LogTexture, Error, TEXT("Streaming from the .uexp file '%s' this MUST be in a ubulk instead for best performance."), *PackagePath.GetDebugName(PackageSegment));
-				}
-			}
-
-			CurrentFileIndex = FileInfos.AddDefaulted();
-			FFileInfo& FileInfo = FileInfos[CurrentFileIndex];
-			FileInfo.PackagePath = MoveTemp(PackagePath);
-			FileInfo.PackageSegment = PackageSegment;
-			FileInfo.IOFileOffset = IOFileOffset;
-			FileInfo.FirstMipIndex = MipIndex;
-			FileInfo.LastMipIndex = MipIndex;
+			IORequests[MipIdx].FilenameHash = OwnerMip.BulkData.GetIoFilenameHash();
 		}
 	}
 
@@ -92,49 +74,46 @@ int32 FTexture2DMipDataProvider_IO::GetMips(
 	const FTextureUpdateSyncOptions& SyncOptions)
 {
 	SetAsyncFileCallback(SyncOptions);
-	check(SyncOptions.Counter && !IORequests.Num());
+	check(SyncOptions.Counter);
 	
-	IORequests.AddDefaulted(CurrentFirstLODIdx);
-
-	for (FFileInfo& FileInfo : FileInfos)
+	while (StartingMipIndex < CurrentFirstLODIdx && MipInfos.IsValidIndex(StartingMipIndex))
 	{
-		while (StartingMipIndex >= FileInfo.FirstMipIndex && StartingMipIndex <= FileInfo.LastMipIndex && StartingMipIndex < CurrentFirstLODIdx)
-		{
-			if (!MipInfos.IsValidIndex(StartingMipIndex))
-			{
-				break;
-			}
-			const FTextureMipInfo& MipInfo = MipInfos[StartingMipIndex];
+		const FTexture2DMipMap& OwnerMip = *Context.MipsView[StartingMipIndex];
+		const FTextureMipInfo& MipInfo = MipInfos[StartingMipIndex];
 
-			if (!FileInfo.IOFileHandle)
-			{
-				FileInfo.IOFileHandle.Reset(IPackageResourceManager::Get().OpenAsyncReadPackage(FileInfo.PackagePath, FileInfo.PackageSegment));
-				if (!FileInfo.IOFileHandle)
-				{
-					break;
-				}
-			}
-
-			const FTexture2DMipMap& OwnerMip = *Context.MipsView[StartingMipIndex];
-			// If Data size is specified check compatibility for safety
-			if (MipInfo.DataSize && OwnerMip.BulkData.GetBulkDataSize() > MipInfo.DataSize)
-			{
-				break;
-			}
-
-			// Increment as we push the requests. If a requests complete immediately, then it will call the callback
-			// but that won't do anything because the tick would not try to acquire the lock since it is already locked.
-			SyncOptions.Counter->Increment();
-
-			IORequests[StartingMipIndex].Reset(FileInfo.IOFileHandle->ReadRequest(
-				OwnerMip.BulkData.GetBulkDataOffsetInFile() + FileInfo.IOFileOffset, 
-				OwnerMip.BulkData.GetBulkDataSize(), 
-				bPrioritizedIORequest ? AIOP_BelowNormal : AIOP_Low, 
-				&AsyncFileCallBack, 
-				(uint8*)MipInfo.DestData));
+		// Check the validity of the filename.
+		if (IORequests[StartingMipIndex].FilenameHash == INVALID_IO_FILENAME_HASH)
+		{
+			break;
+		}
 
-			++StartingMipIndex;
+		// If Data size is specified, check compatibility for safety
+		if (MipInfo.DataSize && (uint64)OwnerMip.BulkData.GetBulkDataSize() > MipInfo.DataSize)
+		{
+			break;
 		}
+		
+		// Ensure there is some valid place to read into.
+		if (!MipInfo.DestData)
+		{
+			break;
+		}
+
+		// Increment as we push the requests. If a request completes immediately, then it will call the callback
+		// but that won't do anything because the tick would not try to acquire the lock since it is already locked.
+		SyncOptions.Counter->Increment();
+
+		IORequests[StartingMipIndex].BulkDataIORequest.Reset(
+			OwnerMip.BulkData.CreateStreamingRequest(
+				0,
+				OwnerMip.BulkData.GetBulkDataSize(),
+				bPrioritizedIORequest ? (AIOP_FLAG_DONTCACHE | AIOP_BelowNormal) : (AIOP_FLAG_DONTCACHE | AIOP_Low),
+				&AsyncFileCallBack,
+				(uint8*)MipInfo.DestData
+			)
+		);
+
+		++StartingMipIndex;
 	}
 
 	AdvanceTo(ETickState::PollMips, ETickThread::Async);
@@ -143,19 +122,33 @@ int32 FTexture2DMipDataProvider_IO::GetMips(
 
 bool FTexture2DMipDataProvider_IO::PollMips(const FTextureUpdateSyncOptions& SyncOptions)
 {
+	// Notify that some files have possibly been unmounted / missing.
+	if (bIORequestCancelled && !bIORequestAborted)
+	{
+		IRenderAssetStreamingManager& StreamingManager = IStreamingManager::Get().GetRenderAssetStreamingManager();
+		for (FIORequest& IORequest : IORequests)
+		{
+			StreamingManager.MarkMountedStateDirty(IORequest.FilenameHash);
+		}
+
+		UE_LOG(LogContentStreaming, Warning, TEXT("[%s] Texture stream in request failed due to IO error (Mip %d-%d)."), *TextureName.ToString(), ResourceState.AssetLODBias + PendingFirstLODIdx, ResourceState.AssetLODBias + CurrentFirstLODIdx - 1);
+	}
+
 	ClearIORequests();
 	AdvanceTo(ETickState::Done, ETickThread::None);
+
 	return !bIORequestCancelled;
 }
 
 void FTexture2DMipDataProvider_IO::AbortPollMips() 
 {
-	for (TUniquePtr<IAsyncReadRequest>& IORequest : IORequests)
+	for (FIORequest& IORequest : IORequests)
 	{
-		if (IORequest)
+		if (IORequest.BulkDataIORequest)
 		{
-			// Calling IAsyncReadRequest::cancel() here will trigger the AsyncFileCallBack and precipitate the execution of Cancel().
-			IORequest->Cancel();
+			// Calling cancel() here will trigger the AsyncFileCallBack and precipitate the execution of Cancel().
+			IORequest.BulkDataIORequest->Cancel();
+			bIORequestAborted = true;
 		}
 	}
 }
@@ -182,7 +175,7 @@ void FTexture2DMipDataProvider_IO::SetAsyncFileCallback(const FTextureUpdateSync
 	FTextureUpdateSyncOptions::FCallback RescheduleCallback = SyncOptions.RescheduleCallback;
 	check(Counter && RescheduleCallback);
 
-	AsyncFileCallBack = [this, Counter, RescheduleCallback](bool bWasCancelled, IAsyncReadRequest* Req)
+	AsyncFileCallBack = [this, Counter, RescheduleCallback](bool bWasCancelled, IBulkDataIORequest* Req)
 	{
 		// At this point task synchronization would hold the number of pending requests.
 		Counter->Decrement();
@@ -208,16 +201,13 @@ void FTexture2DMipDataProvider_IO::SetAsyncFileCallback(const FTextureUpdateSync
 
 void FTexture2DMipDataProvider_IO::ClearIORequests()
 {
-	for (TUniquePtr<IAsyncReadRequest>& IORequest : IORequests)
+	for (FIORequest& IORequest : IORequests)
 	{
-		if (IORequest)
+		// If requests are not yet completed, cancel and wait.
+		if (IORequest.BulkDataIORequest && !IORequest.BulkDataIORequest->PollCompletion())
 		{
-			// If clearing requests not yet completed, cancel and wait.
-			if (!IORequest->PollCompletion())
-			{
-				IORequest->Cancel();
-				IORequest->WaitCompletion();
-			}
+			IORequest.BulkDataIORequest->Cancel();
+			IORequest.BulkDataIORequest->WaitCompletion();
 		}
 	}
 	IORequests.Empty();
diff --git a/Engine/Source/Runtime/Engine/Private/Streaming/Texture2DMipDataProvider_IO.h b/Engine/Source/Runtime/Engine/Private/Streaming/Texture2DMipDataProvider_IO.h
--- a/Engine/Source/Runtime/Engine/Private/Streaming/Texture2DMipDataProvider_IO.h
+++ b/Engine/Source/Runtime/Engine/Private/Streaming/Texture2DMipDataProvider_IO.h
@@ -8,7 +8,6 @@ Texture2DMipDataProvider_IO.h : Implementation of FTextureMipDataProvider using
 
 #include "CoreMinimal.h"
 #include "Streaming/TextureMipDataProvider.h"
-#include "Misc/PackagePath.h"
 #include "Async/AsyncFileHandle.h"
 
 /**
@@ -37,28 +36,26 @@ public:
 protected:
 
 	// A structured with information about which file contains which mips.
-	struct FFileInfo
+	struct FIORequest
 	{
-		FPackagePath PackagePath;
-		EPackageSegment PackageSegment;
-		TUniquePtr<IAsyncReadFileHandle> IOFileHandle;
-		int64 IOFileOffset = 0;
-		int32 FirstMipIndex = INDEX_NONE;
-		int32 LastMipIndex = INDEX_NONE;
+		FIoFilenameHash FilenameHash = INVALID_IO_FILENAME_HASH;
+		TUniquePtr<IBulkDataIORequest> BulkDataIORequest;
 	};
 
 	// Pending async requests created in GetMips().
-	TArray<TUniquePtr<IAsyncReadRequest>, TInlineAllocator<MAX_TEXTURE_MIP_COUNT>> IORequests;
-	// The list of relevant files used for reading texture mips.
-	TArray<FFileInfo, TInlineAllocator<2>> FileInfos;
+	TArray<FIORequest, TInlineAllocator<MAX_TEXTURE_MIP_COUNT>> IORequests;
 
+	// The asset name, used to log IO errors.
+	FName TextureName;
 	// Whether async read requests must be created with high priority (executes faster). 
 	bool bPrioritizedIORequest = false;
 	// Whether async read requests where cancelled for any reasons.
 	bool bIORequestCancelled = false;
+	// Whether async read requests were required to abort through AbortPollMips().
+	bool bIORequestAborted = false;
 
 	// A callback to be executed once all IO pending requests are completed.
-	FAsyncFileCallBack AsyncFileCallBack;
+	FBulkDataIORequestCallBack AsyncFileCallBack;
 
 	// Helper to configure the AsyncFileCallBack.
 	void SetAsyncFileCallback(const FTextureUpdateSyncOptions& SyncOptions);
diff --git a/Engine/Source/Runtime/Engine/Private/Streaming/TextureStreamOut.cpp b/Engine/Source/Runtime/Engine/Private/Streaming/TextureStreamOut.cpp
new file mode 100644
--- /dev/null
+++ b/Engine/Source/Runtime/Engine/Private/Streaming/TextureStreamOut.cpp
@@ -0,0 +1,207 @@
+// Copyright Epic Games, Inc. All Rights Reserved.
+
+/*=============================================================================
+TextureStreamOut.cpp : Implement a generic texture stream out strategy.
+=============================================================================*/
+
+#include "Streaming/TextureStreamOut.h"
+#include "Engine/Texture.h"
+#include "Streaming/RenderAssetUpdate.inl"
+
+template class TRenderAssetUpdate<FTextureUpdateContext>;
+
+bool FTextureStreamOut::IsSameThread(FTextureMipAllocator::ETickThread TickThread, int32 TaskThread)
+{
+	if (TaskThread == TT_Async)
+	{
+		return TickThread == FTextureMipAllocator::ETickThread::Async;
+	}
+	else // Expected to be called from either the async thread or the renderthread.
+	{
+		check(TaskThread == TT_Render);
+		return TickThread == FTextureMipAllocator::ETickThread::Render;
+	}
+}
+
+FTextureStreamOut::FTextureStreamOut(const UTexture* InTexture, FTextureMipAllocator* InMipAllocator)
+	: TRenderAssetUpdate<FTextureUpdateContext>(InTexture)
+{
+	check(InMipAllocator);
+
+	// Init the allocator.
+	MipAllocator.Reset(InMipAllocator);
+
+	// Init sync options
+	SyncOptions.bSnooze = &bDeferExecution;
+	SyncOptions.Counter = &TaskSynchronization;
+	SyncOptions.RescheduleCallback = [this]() 
+	{ 
+		if (!IsLocked())
+		{
+			Tick(TT_None); 
+		}
+	};
+
+	// Schedule the first update step.
+	FContext Context(InTexture, TT_None);
+	const EThreadType NextThread = GetMipAllocatorThread(FTextureMipAllocator::ETickState::AllocateMips);
+	if (NextThread != TT_None)
+	{
+		PushTask(Context, NextThread, SRA_UPDATE_CALLBACK(AllocateNewMips), GetCancelThread(), SRA_UPDATE_CALLBACK(Cancel));
+	}
+	else // Otherwise if it is impossible to allocate the new mips, abort.
+	{
+		MarkAsCancelled();
+		PushTask(Context, TT_None, nullptr, GetCancelThread(), SRA_UPDATE_CALLBACK(Cancel));
+	}
+}
+
+FTextureStreamOut::~FTextureStreamOut()
+{
+}
+
+FRenderAssetUpdate::EThreadType FTextureStreamOut::GetMipAllocatorThread(FTextureMipAllocator::ETickState TickState) const
+{
+	check(MipAllocator);
+	if (!IsCancelled() && MipAllocator->GetNextTickState() == TickState)
+	{
+		switch (MipAllocator->GetNextTickThread())
+		{
+		case FTextureMipAllocator::ETickThread::Async:
+			return TT_Async;
+		case FTextureMipAllocator::ETickThread::Render:
+			return TT_Render;
+		default:
+			check(false);
+		}
+	}
+	return TT_None;
+}
+
+FRenderAssetUpdate::EThreadType FTextureStreamOut::GetCancelThread() const
+{
+	// Cancel the mip allocator.
+	if (MipAllocator)
+	{
+		switch (MipAllocator->GetCancelThread())
+		{
+		case FTextureMipAllocator::ETickThread::Async:
+			return TT_Async;
+		case FTextureMipAllocator::ETickThread::Render:
+			return TT_Render;
+		}
+	}
+
+	// If both are cancelled, then run the final cleanup on the async thread.
+	return TT_Async;
+}
+
+// ****************************
+// **** Update Steps Work *****
+// ****************************
+
+bool FTextureStreamOut::DoAllocateNewMips(const FContext& Context)
+{
+	check(MipAllocator && IsSameThread(MipAllocator->GetNextTickThread(), Context.CurrentThread));
+	FTextureMipInfoArray MipInfos;
+	return MipAllocator->AllocateMips(Context, MipInfos, SyncOptions);
+}
+
+bool FTextureStreamOut::DoFinalizeNewMips(const FContext& Context)
+{
+	check(MipAllocator && IsSameThread(MipAllocator->GetNextTickThread(), Context.CurrentThread));
+	return MipAllocator->FinalizeMips(Context, SyncOptions);
+}
+
+// ****************************
+// ******* Update Steps *******
+// ****************************
+
+void FTextureStreamOut::AllocateNewMips(const FContext& Context)
+{
+	DECLARE_SCOPE_CYCLE_COUNTER(TEXT("FTextureStreamOut::AllocateNewMips"), STAT_TextureStreamIn_AllocateNewMips, STATGROUP_StreamingDetails);
+
+	// Execute
+	if (!DoAllocateNewMips(Context))
+	{
+		MarkAsCancelled();
+	}
+
+	// Schedule the next update step.
+	EThreadType NextThread = GetMipAllocatorThread(FTextureMipAllocator::ETickState::AllocateMips);
+	if (NextThread != TT_None) // Loop on this state.
+	{
+		PushTask(Context, NextThread, SRA_UPDATE_CALLBACK(AllocateNewMips), GetCancelThread(), SRA_UPDATE_CALLBACK(Cancel));
+	}
+	else if (!IsCancelled())
+	{
+		NextThread = GetMipAllocatorThread(FTextureMipAllocator::ETickState::FinalizeMips);
+		// All mips must be handled before moving to next stage.
+		if (NextThread != TT_None)
+		{
+			PushTask(Context, NextThread, SRA_UPDATE_CALLBACK(FinalizeNewMips), GetCancelThread(), SRA_UPDATE_CALLBACK(Cancel));
+		}
+		else // Otherwise nothing to execute but schedule next phase (edge case).
+		{
+			FinalizeNewMips(Context);
+		}
+	}
+	else
+	{
+		PushTask(Context, TT_None, nullptr, GetCancelThread(), SRA_UPDATE_CALLBACK(Cancel));
+	}
+}
+
+void FTextureStreamOut::FinalizeNewMips(const FContext& Context)
+{
+	DECLARE_SCOPE_CYCLE_COUNTER(TEXT("FTextureStreamOut::FinalizeNewMips"), STAT_TextureStreamIn_FinalizeNewMips, STATGROUP_StreamingDetails);
+
+	// Execute
+	if (!DoFinalizeNewMips(Context))
+	{
+		MarkAsCancelled();
+	}
+
+	// Schedule the next update step.
+	EThreadType NextThread = GetMipAllocatorThread(FTextureMipAllocator::ETickState::FinalizeMips);
+	if (NextThread != TT_None) // Loop on this state.
+	{
+		PushTask(Context, NextThread, SRA_UPDATE_CALLBACK(FinalizeNewMips), GetCancelThread(), SRA_UPDATE_CALLBACK(Cancel));
+	}
+	else
+	{
+		// Release the mip allocator.
+		MipAllocator.Reset();
+	}
+}
+
+void FTextureStreamOut::Cancel(const FContext& Context)
+{
+	DECLARE_SCOPE_CYCLE_COUNTER(TEXT("FTextureStreamOut::Cancel"), STAT_TextureStreamIn_Cancel, STATGROUP_StreamingDetails);
+
+	EThreadType NextThread = TT_None;
+
+	// Then cancel the mip allocator.
+	if (MipAllocator)
+	{
+		if (IsSameThread(MipAllocator->GetCancelThread(), Context.CurrentThread))
+		{
+			MipAllocator->Cancel(SyncOptions);
+		}
+		switch (MipAllocator->GetCancelThread())
+		{
+		case FTextureMipAllocator::ETickThread::Async:
+			NextThread = TT_Async;
+			break;
+		case FTextureMipAllocator::ETickThread::Render:
+			NextThread = TT_Render;
+			break;
+		}
+		if (NextThread != TT_None)
+		{
+			PushTask(Context, TT_None, nullptr, NextThread, SRA_UPDATE_CALLBACK(Cancel));
+			return;
+		}
+		MipAllocator.Reset();
+	}
+}
diff --git a/Engine/Source/Runtime/Engine/Private/Streaming/TextureStreamOut.h b/Engine/Source/Runtime/Engine/Private/Streaming/TextureStreamOut.h
new file mode 100644
--- /dev/null
+++ b/Engine/Source/Runtime/Engine/Private/Streaming/TextureStreamOut.h
@@ -0,0 +1,53 @@
+// Copyright Epic Games, Inc. All Rights Reserved.
+
+/*=============================================================================
+TextureStreamOut.h: Implement a generic texture stream out strategy.
+=============================================================================*/
+
+#pragma once
+
+#include "CoreMinimal.h"
+#include "RenderAssetUpdate.h"
+#include "TextureMipAllocator.h"
+
+class UTexture;
+
+/**
+* FTextureStreamOut implements a generic framework to stream out using  FTextureMipAllocator.
+*/
+class FTextureStreamOut : public TRenderAssetUpdate<FTextureUpdateContext>
+{
+public:
+
+	FTextureStreamOut(const UTexture* InTexture, FTextureMipAllocator* InMipAllocator);
+	~FTextureStreamOut();
+	
+	// Route FTextureMipAllocator::AllocateMips()
+	void AllocateNewMips(const FContext& Context);
+	
+	// Route FTextureMipAllocator::FinalizeMips()
+	void FinalizeNewMips(const FContext& Context);
+
+	// Execute the cancellation for both MipAllocator.
+	void Cancel(const FContext& Context);
+
+protected:
+
+	TUniquePtr<FTextureMipAllocator> MipAllocator;
+
+	FTextureUpdateSyncOptions SyncOptions;
+
+	// The actual implementation of the update steps which redirect to FTextureMipAllocator and FTextureMipDataProvider.
+	bool DoAllocateNewMips(const FContext& Context);
+	bool DoFinalizeNewMips(const FContext& Context);
+
+	EThreadType GetMipAllocatorThread(FTextureMipAllocator::ETickState TickState) const;
+
+	EThreadType GetCancelThread() const;
+
+private:
+
+	// Whether FTextureMipAllocator::ETickThread is the same thread as EThreadType (as int32).
+	FORCEINLINE static bool IsSameThread(FTextureMipAllocator::ETickThread TickThread, int32 TaskThread);
+};
+
diff --git a/Engine/Source/Runtime/Engine/Private/Streaming/TextureStreamingBuild.cpp b/Engine/Source/Runtime/Engine/Private/Streaming/TextureStreamingBuild.cpp
--- a/Engine/Source/Runtime/Engine/Private/Streaming/TextureStreamingBuild.cpp
+++ b/Engine/Source/Runtime/Engine/Private/Streaming/TextureStreamingBuild.cpp
@@ -14,7 +14,7 @@ TextureStreamingBuild.cpp : Contains definitions to build texture streaming data
 #include "Components/PrimitiveComponent.h"
 #include "Components/StaticMeshComponent.h"
 #include "Misc/FeedbackContext.h"
-#include "Engine/Texture2D.h"
+#include "Engine/Texture.h"
 #include "ShaderCompiler.h"
 #include "Engine/StaticMesh.h"
 #include "Streaming/TextureStreamingHelpers.h"
@@ -45,13 +45,13 @@ ENGINE_API bool BuildTextureStreamingComponentData(UWorld* InWorld, EMaterialQua
 	const float OneOverNumActorsInWorld = 1.f / (float)FMath::Max<int32>(NumActorsInWorld, 1); // Prevent div by 0
 
 	// Used to reset per level index for textures.
-	TArray<UTexture2D*> AllTextures;
-	for (FObjectIterator Iter(UTexture2D::StaticClass()); Iter && bFullRebuild; ++Iter)
+	TArray<UTexture*> AllTextures;
+	for (FObjectIterator Iter(UTexture::StaticClass()); Iter && bFullRebuild; ++Iter)
 	{
-		UTexture2D* Texture2D = Cast<UTexture2D>(*Iter);
-		if (Texture2D)
+		UTexture* Texture = Cast<UTexture>(*Iter);
+		if (Texture)
 		{
-			AllTextures.Add(Texture2D);
+			AllTextures.Add(Texture);
 		}
 	}
 
@@ -123,10 +123,10 @@ ENGINE_API bool BuildTextureStreamingComponentData(UWorld* InWorld, EMaterialQua
 		if (bFullRebuild)
 		{
 			// Reset LevelIndex to default for next use and build the level Guid array.
-			for (UTexture2D* Texture2D : AllTextures)
+			for (UTexture* Texture : AllTextures)
 			{
-				checkSlow(Texture2D);
-				Texture2D->LevelIndex = INDEX_NONE;
+				checkSlow(Texture);
+				Texture->LevelIndex = INDEX_NONE;
 			}
 
 			// Cleanup the asset references.
@@ -241,14 +241,14 @@ void FStreamingTextureBuildInfo::PackFrom(ULevel* Level, const FBoxSphereBounds&
 
 	PackedRelativeBox = PackRelativeBox(RefBounds.Origin, RefBounds.BoxExtent, Info.Bounds.Origin, Info.Bounds.BoxExtent);
 
-	UTexture2D* Texture2D = Cast<UTexture2D>(Info.RenderAsset);
-	check(Texture2D);
-	if (Texture2D->LevelIndex == INDEX_NONE)
+	UTexture* Texture = Cast<UTexture>(Info.RenderAsset);
+	check(Texture);
+	if (Texture->LevelIndex == INDEX_NONE)
 	{
 		// If this is the first time this texture gets processed in the packing process, encode it.
-		Texture2D->LevelIndex = Level->StreamingTextureGuids.Add(Texture2D->GetLightingGuid());
+		Texture->LevelIndex = Level->StreamingTextureGuids.Add(Texture->GetLightingGuid());
 	}
-	TextureLevelIndex = (uint16)Texture2D->LevelIndex;
+	TextureLevelIndex = (uint16)Texture->LevelIndex;
 
 	TexelFactor = Info.TexelFactor;
 }
@@ -350,14 +350,14 @@ void FStreamingTextureLevelContext::BindBuildData(const TArray<FStreamingTexture
 	}
 }
 
-int32* FStreamingTextureLevelContext::GetBuildDataIndexRef(UTexture2D* Texture2D)
+int32* FStreamingTextureLevelContext::GetBuildDataIndexRef(UTexture* Texture)
 {
 	if (ComponentBuildData) // If there is some build data to map to.
 	{
-		if (Texture2D->LevelIndex == INDEX_NONE)
+		if (Texture->LevelIndex == INDEX_NONE)
 		{
 			check(TextureGuidToLevelIndex); // Can't bind ComponentData without the remapping.
-			const int32* LevelIndex = TextureGuidToLevelIndex->Find(Texture2D->GetLightingGuid());
+			const int32* LevelIndex = TextureGuidToLevelIndex->Find(Texture->GetLightingGuid());
 			if (LevelIndex) // If the index is found in the map, the index is valid in BoundStates
 			{
 				// Here we need to support the invalid case where 2 textures have the same GUID.
@@ -365,14 +365,14 @@ int32* FStreamingTextureLevelContext::GetBuildDataIndexRef(UTexture2D* Texture2D
 				FTextureBoundState& BoundState = BoundStates[*LevelIndex];
 				if (!BoundState.Texture)
 				{
-					Texture2D->LevelIndex = *LevelIndex;
-					BoundState.Texture = Texture2D; // Update the mapping now!
+					Texture->LevelIndex = *LevelIndex;
+					BoundState.Texture = Texture; // Update the mapping now!
 				}
-				else // Don't allow 2 textures to be using the same level index otherwise UTexture2D::LevelIndex won't be reset properly in the destructor.
+				else // Don't allow 2 textures to be using the same level index otherwise UTexturD::LevelIndex won't be reset properly in the destructor.
 				{
 					FMessageLog("AssetCheck").Error()
 						->AddToken(FUObjectToken::Create(BoundState.Texture))
-						->AddToken(FUObjectToken::Create(Texture2D))
+						->AddToken(FUObjectToken::Create(Texture))
 						->AddToken(FTextToken::Create( NSLOCTEXT("AssetCheck", "TextureError_NonUniqueLightingGuid", "Same lighting guid, modify or touch any property in the texture editor to generate a new guid and fix the issue.") ) );
 
 					// This will fallback not using the precomputed data. Note also that the other texture might be using the wrong precomputed data.
@@ -381,12 +381,12 @@ int32* FStreamingTextureLevelContext::GetBuildDataIndexRef(UTexture2D* Texture2D
 			}
 			else // Otherwise add a dummy entry to prevent having to search in the map multiple times.
 			{
-				Texture2D->LevelIndex = BoundStates.Add(FTextureBoundState(Texture2D));
+				Texture->LevelIndex = BoundStates.Add(FTextureBoundState(Texture));
 			}
 		}
 
-		FTextureBoundState& BoundState = BoundStates[Texture2D->LevelIndex];
-		check(BoundState.Texture == Texture2D);
+		FTextureBoundState& BoundState = BoundStates[Texture->LevelIndex];
+		check(BoundState.Texture == Texture);
 
 		if (BoundState.BuildDataTimestamp == BuildDataTimestamp)
 		{
@@ -405,13 +405,12 @@ void FStreamingTextureLevelContext::ProcessMaterial(const FBoxSphereBounds& Comp
 
 	for (UTexture* Texture : Textures)
 	{
-		UTexture2D* Texture2D = Cast<UTexture2D>(Texture);
-		if (!Texture2D || !Texture2D->IsStreamable())
+		if (!Texture->IsStreamable())
 		{
 			continue;
 		}
 
-		int32* BuildDataIndex = GetBuildDataIndexRef(Texture2D);
+		int32* BuildDataIndex = GetBuildDataIndexRef(Texture);
 		if (BuildDataIndex)
 		{
 			if (*BuildDataIndex != INDEX_NONE)
@@ -419,7 +418,7 @@ void FStreamingTextureLevelContext::ProcessMaterial(const FBoxSphereBounds& Comp
 				FStreamingRenderAssetPrimitiveInfo& Info = *new(OutStreamingTextures) FStreamingRenderAssetPrimitiveInfo();
 				const FStreamingTextureBuildInfo& BuildInfo = (*ComponentBuildData)[*BuildDataIndex];
 
-				Info.RenderAsset = Texture2D;
+				Info.RenderAsset = Texture;
 				Info.TexelFactor = BuildInfo.TexelFactor * ComponentScaling;
 				Info.PackedRelativeBox = bUseRelativeBoxes ? BuildInfo.PackedRelativeBox : PackedRelativeBox_Identity;
 				UnpackRelativeBox(ComponentBounds, Info.PackedRelativeBox, Info.Bounds);
@@ -443,7 +442,7 @@ void FStreamingTextureLevelContext::ProcessMaterial(const FBoxSphereBounds& Comp
 			{
 				FStreamingRenderAssetPrimitiveInfo& Info = *new(OutStreamingTextures) FStreamingRenderAssetPrimitiveInfo();
 
-				Info.RenderAsset = Texture2D;
+				Info.RenderAsset = Texture;
 				Info.TexelFactor = TextureDensity * ComponentScaling;
 				Info.PackedRelativeBox = bUseRelativeBoxes ? MaterialData.PackedRelativeBox : PackedRelativeBox_Identity;
 				UnpackRelativeBox(ComponentBounds, Info.PackedRelativeBox, Info.Bounds);
diff --git a/Engine/Source/Runtime/Engine/Private/Streaming/VolumeTextureStreaming.cpp b/Engine/Source/Runtime/Engine/Private/Streaming/VolumeTextureStreaming.cpp
new file mode 100644
--- /dev/null
+++ b/Engine/Source/Runtime/Engine/Private/Streaming/VolumeTextureStreaming.cpp
@@ -0,0 +1,203 @@
+// Copyright Epic Games, Inc. All Rights Reserved.
+
+/*=============================================================================
+VolumeTextureStreaming.cpp: Helpers to stream in and out volume texture LODs.
+=============================================================================*/
+
+#include "Streaming/VolumeTextureStreaming.h"
+
+extern RHI_API bool GUseTexture3DBulkDataRHI;
+
+//*****************************************************************************
+//***************************** Global Definitions ****************************
+//*****************************************************************************
+
+void RHICopySharedMips(FRHICommandList& RHICmdList, FRHITexture3D* DestTexture, FRHITexture3D* SrcTexture)
+{
+	// Transition to copy source and dest
+	{
+		FRHITransitionInfo TransitionsBefore[] = { FRHITransitionInfo(SrcTexture, ERHIAccess::SRVMask, ERHIAccess::CopySrc), FRHITransitionInfo(DestTexture, ERHIAccess::SRVMask, ERHIAccess::CopyDest) };
+		RHICmdList.Transition(MakeArrayView(TransitionsBefore, UE_ARRAY_COUNT(TransitionsBefore)));
+	}
+
+	// Copy 
+	{
+		FRHICopyTextureInfo CopyInfo;
+
+		auto SetCopyInfo = [&](FRHITexture3D* Texture3D)
+		{
+			CopyInfo.Size.X = Texture3D->GetSizeX();
+			CopyInfo.Size.Y = Texture3D->GetSizeY();
+			CopyInfo.Size.Z = Texture3D->GetSizeZ();
+			CopyInfo.NumMips = Texture3D->GetNumMips();
+		};
+
+		if (DestTexture->GetNumMips() < SrcTexture->GetNumMips())
+		{
+			SetCopyInfo(DestTexture);
+		}
+		else
+		{
+			SetCopyInfo(SrcTexture);
+		}
+
+		CopyInfo.SourceMipIndex = SrcTexture->GetNumMips() - CopyInfo.NumMips;
+		CopyInfo.DestMipIndex = DestTexture->GetNumMips() - CopyInfo.NumMips;
+		RHICmdList.CopyTexture(SrcTexture, DestTexture, CopyInfo);
+	}
+
+	// Transition to SRV
+	{
+		FRHITransitionInfo TransitionsAfter[] = { FRHITransitionInfo(SrcTexture, ERHIAccess::CopySrc, ERHIAccess::SRVMask), FRHITransitionInfo(DestTexture, ERHIAccess::CopyDest, ERHIAccess::SRVMask) };
+		RHICmdList.Transition(MakeArrayView(TransitionsAfter, UE_ARRAY_COUNT(TransitionsAfter)));
+	}
+}
+
+//*****************************************************************************
+//******************** FVolumeTextureMipAllocator_Reallocate ******************
+//*****************************************************************************
+
+FVolumeTextureMipAllocator_Reallocate::FVolumeTextureMipAllocator_Reallocate(UTexture* Texture)
+	: FTextureMipAllocator(Texture, ETickState::AllocateMips, ETickThread::Async)
+	, StreamedInMipData(PendingFirstLODIdx)
+{
+}
+
+FVolumeTextureMipAllocator_Reallocate::~FVolumeTextureMipAllocator_Reallocate()
+{
+}
+
+bool FVolumeTextureMipAllocator_Reallocate::AllocateMips(const FTextureUpdateContext& Context, FTextureMipInfoArray& OutMipInfos, const FTextureUpdateSyncOptions& SyncOptions)
+{
+	OutMipInfos.AddDefaulted(CurrentFirstLODIdx);
+
+	// Allocate the mip memory as temporary buffers so that the FTextureMipDataProvider implementation can write to it.
+	for (int32 MipIdx = PendingFirstLODIdx; MipIdx < CurrentFirstLODIdx; ++MipIdx)
+	{
+		const FTexture2DMipMap& OwnerMip = *Context.MipsView[MipIdx];
+		FTextureMipInfo& MipInfo = OutMipInfos[MipIdx];
+
+		// Note that streamed in mips size will always be at least as big as block size.
+		MipInfo.Format = Context.Resource->GetPixelFormat();
+		MipInfo.SizeX = OwnerMip.SizeX;
+		MipInfo.SizeY = OwnerMip.SizeY;
+		MipInfo.SizeZ = OwnerMip.SizeZ;
+
+		uint32 TextureAlign = 0;
+		MipInfo.DataSize = RHICalcTexture3DPlatformSize(MipInfo.SizeX, MipInfo.SizeY, MipInfo.SizeZ, MipInfo.Format, 1, Context.Resource->GetCreationFlags(), FRHIResourceCreateInfo(Context.Resource->GetExtData()), TextureAlign);
+		StreamedInMipData.GetMipSize()[MipIdx] = MipInfo.DataSize;
+
+		// When initializing the texture with it's bulk data, we proceed with a single allocation.
+		if (!GUseTexture3DBulkDataRHI)
+		{
+			MipInfo.DestData = FMemory::Malloc(MipInfo.DataSize, FVolumeTextureBulkData::MALLOC_ALIGNMENT);
+			StreamedInMipData.GetMipData()[MipIdx] = MipInfo.DestData;
+		}
+	}
+
+	if (GUseTexture3DBulkDataRHI)
+	{
+		StreamedInMipData.MergeMips(ResourceState.MaxNumLODs);
+
+		uint8* MergedMipData = (uint8*)StreamedInMipData.GetResourceBulkData();
+		for (int32 MipIdx = PendingFirstLODIdx; MipIdx < CurrentFirstLODIdx; ++MipIdx)
+		{
+			FTextureMipInfo& MipInfo = OutMipInfos[MipIdx];
+			MipInfo.DestData = MergedMipData;
+			MergedMipData += MipInfo.DataSize;
+		}
+	}
+
+	AdvanceTo(ETickState::FinalizeMips, ETickThread::Render);
+	return true;
+}
+
+bool FVolumeTextureMipAllocator_Reallocate::FinalizeMips(const FTextureUpdateContext& Context, const FTextureUpdateSyncOptions& SyncOptions)
+{
+	EPixelFormat PixelFormat = Context.Resource->GetPixelFormat();
+
+	// Create new Texture.
+	{
+		FRHIResourceCreateInfo CreateInfo;
+		if (GUseTexture3DBulkDataRHI)
+		{
+			CreateInfo.BulkData = &StreamedInMipData;
+		}
+
+		const FTexture2DMipMap& FirstMip = *Context.MipsView[PendingFirstLODIdx];
+		CreateInfo.ExtData = Context.Resource->GetExtData();
+		IntermediateTextureRHI = RHICreateTexture3D(FirstMip.SizeX, FirstMip.SizeY, FirstMip.SizeZ, PixelFormat, ResourceState.NumRequestedLODs, Context.Resource->GetCreationFlags(), CreateInfo);
+	}
+
+	// Copy shared mips.
+	{
+		bool bCopySharedMipsDone = false;
+		ENQUEUE_RENDER_COMMAND(FCopySharedMipsForTexture3D)(
+			[&](FRHICommandListImmediate& RHICmdList)
+		{
+			RHICopySharedMips(RHICmdList, IntermediateTextureRHI.GetReference(), Context.Resource->GetTexture3DRHI());
+			bCopySharedMipsDone = true;
+		});
+		// Expected to execute immediately since ran on the renderthread.
+		check(bCopySharedMipsDone);
+	}
+
+	// Update the streamed in mips if they were not initialized from the bulk data.
+	if (!GUseTexture3DBulkDataRHI)
+	{
+		const int32 BlockSizeX = GPixelFormats[PixelFormat].BlockSizeX;
+		const int32 BlockSizeY = GPixelFormats[PixelFormat].BlockSizeY;
+		const int32 BlockBytes = GPixelFormats[PixelFormat].BlockBytes;
+		ensure(GPixelFormats[PixelFormat].BlockSizeZ == 1);
+
+		for (int32 MipIdx = PendingFirstLODIdx; MipIdx < CurrentFirstLODIdx; ++MipIdx)
+		{
+			const FTexture2DMipMap& Mip = *Context.MipsView[MipIdx];
+			const uint8* MipData = (uint8*)StreamedInMipData.GetMipData()[MipIdx];
+			if (MipData)
+			{
+				FUpdateTextureRegion3D UpdateRegion(0, 0, 0, 0, 0, 0, Mip.SizeX, Mip.SizeY, Mip.SizeZ);
+				const uint32 NumBlockX = (uint32)FMath::DivideAndRoundUp<int32>(Mip.SizeX, BlockSizeX);
+				const uint32 NumBlockY = (uint32)FMath::DivideAndRoundUp<int32>(Mip.SizeY, BlockSizeY);
+				RHIUpdateTexture3D(IntermediateTextureRHI, MipIdx - PendingFirstLODIdx, UpdateRegion, NumBlockX * BlockBytes, NumBlockX * NumBlockY * BlockBytes, MipData);
+			}
+		}
+	}
+	StreamedInMipData.Discard();
+
+	// Use the new texture resource for the texture asset, must run on the renderthread.
+	Context.Resource->FinalizeStreaming(IntermediateTextureRHI);
+	// No need for the intermediate texture anymore.
+	IntermediateTextureRHI.SafeRelease();
+
+	// Update is complete, nothing more to do.
+	AdvanceTo(ETickState::Done, ETickThread::None);
+	return true;
+}
+
+void FVolumeTextureMipAllocator_Reallocate::Cancel(const FTextureUpdateSyncOptions& SyncOptions)
+{
+	// Release the intermediate texture. If not null, this will be on the renderthread.
+	IntermediateTextureRHI.SafeRelease();
+	// Release the temporary mip data. Can be run on either renderthread or async threads.
+	StreamedInMipData.Discard();
+}
+
+FTextureMipAllocator::ETickThread FVolumeTextureMipAllocator_Reallocate::GetCancelThread() const
+{
+	// If there is an  intermediate texture, it is safer to released on the renderthread.
+	if (IntermediateTextureRHI)
+	{
+		return ETickThread::Render;
+	}
+	// Otherwise, if there are only temporary mip data, they can be freed on any threads.
+	else if (StreamedInMipData.GetResourceBulkData())
+	{
+		return ETickThread::Async;
+	}
+	// Nothing to do.
+	else
+	{
+		return ETickThread::None;
+	}
+}
diff --git a/Engine/Source/Runtime/Engine/Private/Streaming/VolumeTextureStreaming.h b/Engine/Source/Runtime/Engine/Private/Streaming/VolumeTextureStreaming.h
new file mode 100644
--- /dev/null
+++ b/Engine/Source/Runtime/Engine/Private/Streaming/VolumeTextureStreaming.h
@@ -0,0 +1,49 @@
+// Copyright Epic Games, Inc. All Rights Reserved.
+
+/*=============================================================================
+VolumeTextureStreaming.h: Helpers to stream in and out volume texture LODs.
+=============================================================================*/
+
+#pragma once
+
+#include "CoreMinimal.h"
+#include "RenderAssetUpdate.h"
+#include "Engine/VolumeTexture.h"
+#include "Rendering/Texture3DResource.h"
+#include "Texture2DMipDataProvider_DDC.h"
+#include "Texture2DMipDataProvider_IO.h"
+#include "TextureMipAllocator.h"
+
+#if WITH_EDITORONLY_DATA
+typedef FTexture2DMipDataProvider_DDC FVolumeTextureMipDataProvider_DDC;
+#endif
+
+typedef FTexture2DMipDataProvider_IO FVolumeTextureMipDataProvider_IO;
+
+/**
+* Implements a texture3d mip resizing strategy based on creating a duplicate texture and copyinng the shared mips.
+*/
+class FVolumeTextureMipAllocator_Reallocate : public FTextureMipAllocator
+{
+public:
+
+	FVolumeTextureMipAllocator_Reallocate(UTexture* Texture);
+	~FVolumeTextureMipAllocator_Reallocate();
+
+	// ********************************************************
+	// ********* FTextureMipAllocator implementation **********
+	// ********************************************************
+
+	bool AllocateMips(const FTextureUpdateContext& Context, FTextureMipInfoArray& OutMipInfos, const FTextureUpdateSyncOptions& SyncOptions) final override;
+	bool FinalizeMips(const FTextureUpdateContext& Context, const FTextureUpdateSyncOptions& SyncOptions) final override;
+	void Cancel(const FTextureUpdateSyncOptions& SyncOptions) final override;
+	ETickThread GetCancelThread() const final override;
+
+protected:
+
+	// The intermediate texture async created in the update process.
+	FTexture3DRHIRef IntermediateTextureRHI;
+
+	// The temporary main memory allocations holding the mip data
+	FVolumeTextureBulkData StreamedInMipData;
+};
diff --git a/Engine/Source/Runtime/Engine/Private/Texture.cpp b/Engine/Source/Runtime/Engine/Private/Texture.cpp
--- a/Engine/Source/Runtime/Engine/Private/Texture.cpp
+++ b/Engine/Source/Runtime/Engine/Private/Texture.cpp
@@ -686,6 +686,17 @@ bool UTexture::ShouldMipLevelsBeForcedResident() const
 	return false;
 }
 
+void UTexture::CancelPendingTextureStreaming()
+{
+	for( TObjectIterator<UTexture> It; It; ++It )
+	{
+		UTexture* CurrentTexture = *It;
+		CurrentTexture->CancelPendingStreamingRequest();
+	}
+
+	// No need to call FlushResourceStreaming(), since calling CancelPendingMipChangeRequest has an immediate effect.
+}
+
 float UTexture::GetAverageBrightness(bool bIgnoreTrueBlack, bool bUseGrayscale)
 {
 	// Indicate the action was not performed...
diff --git a/Engine/Source/Runtime/Engine/Private/Texture2D.cpp b/Engine/Source/Runtime/Engine/Private/Texture2D.cpp
--- a/Engine/Source/Runtime/Engine/Private/Texture2D.cpp
+++ b/Engine/Source/Runtime/Engine/Private/Texture2D.cpp
@@ -62,8 +62,6 @@ UTexture2D::UTexture2D(const FObjectInitializer& ObjectInitializer)
 		[this](FTexturePlatformData* InPlatformData) { SetPlatformData(InPlatformData); })
 {
 	PendingUpdate = nullptr;
-	StreamingIndex = INDEX_NONE;
-	LevelIndex = INDEX_NONE;
 	SRGB = true;
 }
 
@@ -585,17 +583,6 @@ float UTexture2D::GetAverageBrightness(bool bIgnoreTrueBlack, bool bUseGrayscale
 	return AvgBrightness;
 }
 
-void UTexture2D::CancelPendingTextureStreaming()
-{
-	for( TObjectIterator<UTexture2D> It; It; ++It )
-	{
-		UTexture2D* CurrentTexture = *It;
-		CurrentTexture->CancelPendingStreamingRequest();
-	}
-
-	// No need to call FlushResourceStreaming(), since calling CancelPendingMipChangeRequest has an immediate effect.
-}
-
 bool UTexture2D::IsReadyForAsyncPostLoad() const
 {
 #if WITH_EDITOR
diff --git a/Engine/Source/Runtime/Engine/Private/VolumeTexture.cpp b/Engine/Source/Runtime/Engine/Private/VolumeTexture.cpp
--- a/Engine/Source/Runtime/Engine/Private/VolumeTexture.cpp
+++ b/Engine/Source/Runtime/Engine/Private/VolumeTexture.cpp
@@ -12,11 +12,15 @@
 #include "DeviceProfiles/DeviceProfileManager.h"
 #include "Containers/ResourceArray.h"
 #include "Rendering/Texture3DResource.h"
+#include "Streaming/VolumeTextureStreaming.h"
+#include "Streaming/TextureStreamIn.h"
+#include "Streaming/TextureStreamOut.h"
+#include "Engine/TextureMipDataProviderFactory.h"
 
 //*****************************************************************************
 
 // Master switch to control whether streaming is enabled for volume texture. 
-bool GSupportsVolumeTextureStreaming = false;
+bool GSupportsVolumeTextureStreaming = true;
 
 // Limit the possible depth of volume texture otherwise when the user converts 2D textures, he can crash the engine.
 const int32 MAX_VOLUME_TEXTURE_DEPTH = 512;
@@ -392,10 +396,58 @@ bool UVolumeTexture::ShaderPlatformSupportsCompression(FStaticShaderPlatform Sha
 
 bool UVolumeTexture::StreamOut(int32 NewMipCount)
 {
+	FTexture3DResource* Texture3DResource = Resource ? Resource->GetTexture3DResource() : nullptr;
+	if (!HasPendingInitOrStreaming() && CachedSRRState.StreamOut(NewMipCount) && ensure(Texture3DResource))
+	{
+		FTextureMipAllocator* MipAllocator = nullptr;
+
+		// FVolumeTextureMipAllocator_Virtual?
+		MipAllocator = new FVolumeTextureMipAllocator_Reallocate(this);
+
+		PendingUpdate = new FTextureStreamOut(this, MipAllocator);
+		return !PendingUpdate->IsCancelled();
+	}
 	return false;
 }
 
 bool UVolumeTexture::StreamIn(int32 NewMipCount, bool bHighPrio)
 {
+	FTexture3DResource* Texture3DResource = Resource ? Resource->GetTexture3DResource() : nullptr;
+	if (!HasPendingInitOrStreaming() && CachedSRRState.StreamIn(NewMipCount) && ensure(Texture3DResource))
+	{
+		FTextureMipDataProvider* CustomMipDataProvider = nullptr;
+		for (UAssetUserData* UserData : AssetUserData)
+		{
+			UTextureMipDataProviderFactory* CustomMipDataProviderFactory = Cast<UTextureMipDataProviderFactory>(UserData);
+			if (CustomMipDataProviderFactory)
+			{
+				CustomMipDataProvider = CustomMipDataProviderFactory->AllocateMipDataProvider(this);
+				if (CustomMipDataProvider)
+				{
+					break;
+				}
+			}
+		}
+
+		FTextureMipAllocator* MipAllocator = nullptr;
+		FTextureMipDataProvider* DefaultMipDataProvider = nullptr;
+
+#if WITH_EDITORONLY_DATA
+		if (FPlatformProperties::HasEditorOnlyData() && !GetOutermost()->bIsCookedForEditor)
+		{
+			DefaultMipDataProvider = new FVolumeTextureMipDataProvider_DDC(this);
+		}
+		else 
+#endif
+		{
+			DefaultMipDataProvider = new FVolumeTextureMipDataProvider_IO(this, bHighPrio);
+		}
+
+		// FVolumeTextureMipAllocator_Virtual?
+		MipAllocator = new FVolumeTextureMipAllocator_Reallocate(this);
+
+		PendingUpdate = new FTextureStreamIn(this, MipAllocator, CustomMipDataProvider, DefaultMipDataProvider);
+		return !PendingUpdate->IsCancelled();
+	}
 	return false;
 }
diff --git a/Engine/Source/Runtime/Engine/Public/Rendering/StreamableTextureResource.h b/Engine/Source/Runtime/Engine/Public/Rendering/StreamableTextureResource.h
--- a/Engine/Source/Runtime/Engine/Public/Rendering/StreamableTextureResource.h
+++ b/Engine/Source/Runtime/Engine/Public/Rendering/StreamableTextureResource.h
@@ -28,6 +28,7 @@ public:
 
 	virtual uint32 GetSizeX() const final override { return SizeX; }
 	virtual uint32 GetSizeY() const final override { return SizeY; }
+	// Depth for 3D texture or ArraySize for texture 2d arrays
 	virtual uint32 GetSizeZ() const final override { return SizeZ; }
 	virtual void InitRHI() override;
 	virtual void ReleaseRHI() final override;
@@ -70,6 +71,11 @@ public:
 		return State.IsValid() ? State.ResidentFirstLODIdx() : 0;
 	}
 
+	FORCEINLINE uint32 GetExtData() const { return PlatformData->GetExtData(); }
+
+	/** Returns the platform mip size for the given mip count. */
+	virtual uint64 GetPlatformMipsSize(uint32 NumMips) const = 0;
+
 protected:
 
 	virtual void CreateTexture() = 0;
@@ -92,7 +98,7 @@ protected:
 	uint32 SizeX = 0;
 	/** The height when all mips are streamed in. */
 	uint32 SizeY = 0;
-	/** The depth when all mips are streamed in. */
+	/** The 3d depth for volume texture or num  slices for 2d array when all mips are streamed in. */
 	uint32 SizeZ = 0;
 
 	/** The FName of the texture asset */
@@ -109,7 +115,8 @@ protected:
 	bool bUsePartiallyResidentMips = false;
 
 #if STATS
-	virtual void CalcRequestedMipsSize() = 0;
+private:
+	void CalcRequestedMipsSize();
 	void IncrementTextureStats() const;
 	void DecrementTextureStats() const;
 
diff --git a/Engine/Source/Runtime/Engine/Public/Rendering/Texture2DResource.h b/Engine/Source/Runtime/Engine/Public/Rendering/Texture2DResource.h
--- a/Engine/Source/Runtime/Engine/Public/Rendering/Texture2DResource.h
+++ b/Engine/Source/Runtime/Engine/Public/Rendering/Texture2DResource.h
@@ -37,6 +37,9 @@ public:
 	/** Set the value of Filter, AddressU, AddressV, AddressW and MipBias from FStreamableTextureResource on the gamethread. */
 	void CacheSamplerStateInitializer(const UTexture2D* InOwner);
 
+	/** Returns the platform mip size for the given mip count. */
+	virtual uint64 GetPlatformMipsSize(uint32 NumMips) const override;
+
 	virtual void InitRHI() override;
 	virtual bool IsProxy() const override { return ProxiedResource != nullptr; }
 private:
@@ -50,9 +53,6 @@ private:
 
 	virtual void CreateTexture() final override;
 	virtual void CreatePartiallyResidentTexture() final override;
-#if STATS
-	virtual void CalcRequestedMipsSize() final override;
-#endif
 
 	/** Texture streaming command classes that need to be friends in order to call Update/FinalizeMipCount.	*/
 	friend class UTexture2D;
diff --git a/Engine/Source/Runtime/Engine/Public/Streaming/TextureMipDataProvider.h b/Engine/Source/Runtime/Engine/Public/Streaming/TextureMipDataProvider.h
--- a/Engine/Source/Runtime/Engine/Public/Streaming/TextureMipDataProvider.h
+++ b/Engine/Source/Runtime/Engine/Public/Streaming/TextureMipDataProvider.h
@@ -22,11 +22,12 @@ struct FTexture2DMipMap;
 struct FTextureMipInfo
 {
 	void* DestData = nullptr; // The mip data where the content must be copied too.
-	uint32 DataSize = 0; // Optional, might be 0 in cooked. Mostly for safeguard.
+	uint64 DataSize = 0; // Optional, might be 0 in cooked. Mostly for safeguard.
 	EPixelFormat Format = PF_Unknown;
 	uint32 SizeX = 0;
 	uint32 SizeY = 0;
-	uint32 SizeZ = 0;
+	uint32 SizeZ = 0; // For 3d textures
+	uint32 ArraySize = 0; // For texture arrays / cubemaps
 	uint32 RowPitch = 0;
 	uint32 DepthPitch = 0;	// For volume texture, texture array only.
 	ECubeFace CubeFace = CubeFace_MAX; // Cubemap only.
diff --git a/Engine/Source/Runtime/Engine/Public/TextureResource.h b/Engine/Source/Runtime/Engine/Public/TextureResource.h
--- a/Engine/Source/Runtime/Engine/Public/TextureResource.h
+++ b/Engine/Source/Runtime/Engine/Public/TextureResource.h
@@ -29,6 +29,8 @@ class IVirtualTexture;
 struct FTexturePlatformData;
 class FStreamableTextureResource;
 class FTexture2DResource;
+class FTexture3DResource;
+class FTexture2DArrayResource;
 
 /** Maximum number of slices in texture source art. */
 #define MAX_TEXTURE_SOURCE_SLICES 6
@@ -91,9 +93,13 @@ public:
 
 	// Dynamic cast methods.
 	ENGINE_API virtual FTexture2DResource* GetTexture2DResource() { return nullptr; }
+	ENGINE_API virtual FTexture3DResource* GetTexture3DResource() { return nullptr; }
+	ENGINE_API virtual FTexture2DArrayResource* GetTexture2DArrayResource() { return nullptr; }
 	ENGINE_API virtual FStreamableTextureResource* GetStreamableTextureResource() { return nullptr; }
 	// Dynamic cast methods (const).
 	ENGINE_API virtual const FTexture2DResource* GetTexture2DResource() const { return nullptr; }
+	ENGINE_API virtual const FTexture3DResource* GetTexture3DResource() const { return nullptr; }
+	ENGINE_API virtual const FTexture2DArrayResource* GetTexture2DArrayResource() const { return nullptr; }
 	ENGINE_API virtual const FStreamableTextureResource* GetStreamableTextureResource() const { return nullptr; }
 
 	// Current mip count. We use "current" to specify that it is not computed from SizeX() which is the size when fully streamed in.
@@ -112,6 +118,16 @@ public:
 		return TextureRHI.IsValid() ? TextureRHI->GetTexture2D() : nullptr;
 	}
 
+	FORCEINLINE FRHITexture3D* GetTexture3DRHI() const
+	{
+		return TextureRHI.IsValid() ? TextureRHI->GetTexture3D() : nullptr;
+	}
+
+	FORCEINLINE FRHITexture2DArray* GetTexture2DArrayRHI() const
+	{
+		return TextureRHI.IsValid() ? TextureRHI->GetTexture2DArray() : nullptr;
+	}
+
 	void SetTextureReference(FRHITextureReference* TextureReference)
 	{
 		TextureReferenceRHI = TextureReference;
diff --git a/Engine/Source/Runtime/MediaAssets/Public/MediaTexture.h b/Engine/Source/Runtime/MediaAssets/Public/MediaTexture.h
--- a/Engine/Source/Runtime/MediaAssets/Public/MediaTexture.h
+++ b/Engine/Source/Runtime/MediaAssets/Public/MediaTexture.h
@@ -187,6 +187,8 @@ public:
 	virtual EMaterialValueType GetMaterialType() const override;
 	virtual float GetSurfaceWidth() const override;
 	virtual float GetSurfaceHeight() const override;
+	virtual float GetSurfaceDepth() const override { return 0; }
+	virtual uint32 GetSurfaceArraySize() const override { return 0; }
 	virtual FGuid GetExternalTextureGuid() const override;
 	void SetRenderedExternalTextureGuid(const FGuid& InNewGuid);
 
diff --git a/Engine/Source/Runtime/WebBrowserTexture/Public/WebBrowserTexture.h b/Engine/Source/Runtime/WebBrowserTexture/Public/WebBrowserTexture.h
--- a/Engine/Source/Runtime/WebBrowserTexture/Public/WebBrowserTexture.h
+++ b/Engine/Source/Runtime/WebBrowserTexture/Public/WebBrowserTexture.h
@@ -69,6 +69,8 @@ public:
 	virtual EMaterialValueType GetMaterialType() const override;
 	virtual float GetSurfaceWidth() const override;
 	virtual float GetSurfaceHeight() const override;
+	virtual float GetSurfaceDepth() const override { return 0; }
+	virtual uint32 GetSurfaceArraySize() const override { return 0; }
 	virtual FGuid GetExternalTextureGuid() const override;
 
 public:
