Fixed duplicate instanciation of TRenderAssetUpdate<FTextureUpdateContext>

diff --git a/Engine/Source/Runtime/Engine/Private/Streaming/TextureStreamOut.cpp b/Engine/Source/Runtime/Engine/Private/Streaming/TextureStreamOut.cpp
--- a/Engine/Source/Runtime/Engine/Private/Streaming/TextureStreamOut.cpp
+++ b/Engine/Source/Runtime/Engine/Private/Streaming/TextureStreamOut.cpp
@@ -8,8 +8,6 @@ TextureStreamOut.cpp : Implement a generic texture stream out strategy.
 #include "Engine/Texture.h"
 #include "Streaming/RenderAssetUpdate.inl"
 
-template class TRenderAssetUpdate<FTextureUpdateContext>;
-
 bool FTextureStreamOut::IsSameThread(FTextureMipAllocator::ETickThread TickThread, int32 TaskThread)
 {
 	if (TaskThread == TT_Async)
