* Enforce Subsurface shading model when a strata diffuse model with subsurface is used within a graph.
 * Check and enforce that a single subsusurface profile is used accross a material graph

diff --git a/Engine/Source/Runtime/Engine/Classes/Engine/EngineTypes.h b/Engine/Source/Runtime/Engine/Classes/Engine/EngineTypes.h
--- a/Engine/Source/Runtime/Engine/Classes/Engine/EngineTypes.h
+++ b/Engine/Source/Runtime/Engine/Classes/Engine/EngineTypes.h
@@ -21,6 +21,7 @@ class UPhysicalMaterial;
 class UPhysicalMaterialMask;
 class UPrimitiveComponent;
 class USceneComponent;
+class USubsurfaceProfile;
 
 /**
  * Default number of components to expect in TInlineAllocators used with AActor component arrays.
@@ -539,6 +540,7 @@ enum EStrataShadingModel
 {
 	SSM_Unlit					UMETA(DisplayName = "Unlit"),
 	SSM_DefaultLit				UMETA(DisplayName = "DefaultLit"),
+	SSM_SubsurfaceLit			UMETA(DisplayName = "SubsurfaceLit"),
 	SSM_VolumetricFogCloud		UMETA(DisplayName = "VolumetricFogCloud"),
 	/** Number of unique shading models. */
 	SSM_NUM						UMETA(Hidden),
@@ -562,7 +564,10 @@ public:
 	uint8 GetShadingModelField() const { return ShadingModelField; }
 	int32 CountShadingModels() const { return FMath::CountBits(ShadingModelField); }
 
-	// STRATA_TODO SSS LUT
+	// Subsurface profiles
+	void AddSubsurfaceProfile(USubsurfaceProfile* InProfile) { if (InProfile) SubsurfaceProfiles.Add(InProfile); }
+	int32 CountSubsurfaceProfiles() const { return SubsurfaceProfiles.Num(); }
+	USubsurfaceProfile* GetSubsurfaceProfile() const { return SubsurfaceProfiles.Num() > 0 ? SubsurfaceProfiles[0] : nullptr; }
 
 	bool IsValid() const { return (ShadingModelField > 0) && (ShadingModelField < (1 << SSM_NUM)); }
 
@@ -572,6 +577,9 @@ public:
 private:
 	UPROPERTY()
 	uint8 ShadingModelField = 0;
+
+	UPROPERTY()
+	TArray<USubsurfaceProfile*> SubsurfaceProfiles;
 };
 
 /** This is used by the drawing passes to determine tessellation policy, so changes here need to be supported in native code. */
diff --git a/Engine/Source/Runtime/Engine/Private/Materials/Material.cpp b/Engine/Source/Runtime/Engine/Private/Materials/Material.cpp
--- a/Engine/Source/Runtime/Engine/Private/Materials/Material.cpp
+++ b/Engine/Source/Runtime/Engine/Private/Materials/Material.cpp
@@ -4493,11 +4493,20 @@ void UMaterial::RebuildShadingModelField()
 		}
 		if (StrataMaterialInfo.CountShadingModels() > 1)
 		{
-			bSanitizeMaterial = true;
-			UE_LOG(LogMaterial, Error, TEXT("%s: Material has more than a single material represented."), *GetName());
+			if (StrataMaterialInfo.CountShadingModels() == 2 && StrataMaterialInfo.HasShadingModel(EStrataShadingModel::SSM_DefaultLit) && StrataMaterialInfo.HasShadingModel(EStrataShadingModel::SSM_SubsurfaceLit))
+			{
+				StrataMaterialInfo = FStrataMaterialInfo();
+				StrataMaterialInfo.AddShadingModel(SSM_SubsurfaceLit);
+			}
+			else
+			{
+				bSanitizeMaterial = true;
+				UE_LOG(LogMaterial, Error, TEXT("%s: Material has more than a single material represented."), *GetName());
+			}
 		}
 		if (bSanitizeMaterial)
 		{
+			StrataMaterialInfo = FStrataMaterialInfo();
 			StrataMaterialInfo.AddShadingModel(SSM_DefaultLit);
 		}
 		
@@ -4512,6 +4521,16 @@ void UMaterial::RebuildShadingModelField()
 				BlendMode = EBlendMode::BLEND_Translucent; // This is to be able to use dual-source blending
 			}
 		}
+		else if (StrataMaterialInfo.HasOnlyShadingModel(SSM_SubsurfaceLit))
+		{
+			MaterialDomain = EMaterialDomain::MD_Surface;
+			ShadingModel = MSM_SubsurfaceProfile;
+			if (BlendMode != EBlendMode::BLEND_Opaque && BlendMode != EBlendMode::BLEND_Masked)
+			{
+				UE_LOG(LogMaterial, Error, TEXT("%s: Material has subsurface data, and its blending mode is not set to Opaque or Masked. Forcing blend mode to Opaque."), *GetName());
+				BlendMode = EBlendMode::BLEND_Opaque;
+			}
+		}
 		else if (StrataMaterialInfo.HasOnlyShadingModel(SSM_DefaultLit))
 		{
 			MaterialDomain = EMaterialDomain::MD_Surface;
@@ -4530,6 +4549,16 @@ void UMaterial::RebuildShadingModelField()
 
 		// Also update the ShadingModels for remaining pipeline operation
 		ShadingModels.AddShadingModel(ShadingModel);
+
+		// Subsurface profil
+		if (StrataMaterialInfo.HasOnlyShadingModel(SSM_SubsurfaceLit) && StrataMaterialInfo.CountSubsurfaceProfiles() > 0)
+		{
+			if (StrataMaterialInfo.CountSubsurfaceProfiles() > 1)
+			{
+				UE_LOG(LogMaterial, Error, TEXT("%s: Material has more than a single sub-surface profile used."), *GetName());
+			}
+			SubsurfaceProfile = StrataMaterialInfo.GetSubsurfaceProfile();
+		}
 	}
 	// If using shading model from material expression, go through the expressions and look for the ShadingModel expression to figure out what shading models need to be supported in this material.
 	// This might not be the same as what is actually compiled in to the shader, since there might be feature switches, static switches etc. that skip certain shading models.
diff --git a/Engine/Source/Runtime/Engine/Private/Materials/MaterialExpressions.cpp b/Engine/Source/Runtime/Engine/Private/Materials/MaterialExpressions.cpp
--- a/Engine/Source/Runtime/Engine/Private/Materials/MaterialExpressions.cpp
+++ b/Engine/Source/Runtime/Engine/Private/Materials/MaterialExpressions.cpp
@@ -19258,7 +19258,19 @@ bool UMaterialExpressionStrataDiffuseBSDF::IsResultStrataMaterial(int32 OutputIn
 
 void UMaterialExpressionStrataDiffuseBSDF::GatherStrataMaterialInfo(FStrataMaterialInfo& StrataMaterialInfo, int32 OutputIndex)
 {
-	StrataMaterialInfo.AddShadingModel(SSM_DefaultLit);
+	const bool bIsScaleInputConnected = true; // STRATA_TODO
+	if (SubsurfaceProfile || bIsScaleInputConnected)
+	{
+		StrataMaterialInfo.AddShadingModel(SSM_SubsurfaceLit);
+		if (SubsurfaceProfile)
+		{
+			StrataMaterialInfo.AddSubsurfaceProfile(SubsurfaceProfile);
+		}
+	}
+	else
+	{
+		StrataMaterialInfo.AddShadingModel(SSM_DefaultLit);
+	}
 }
 #endif // WITH_EDITOR
 
