- Collect all BLAS build and updates together for batch build instead of build during request
- Add flush command list option when trying BLAS build/update request is too big (required for TDRs)

diff --git a/Engine/Source/Runtime/Engine/Private/GPUSkinCache.cpp b/Engine/Source/Runtime/Engine/Private/GPUSkinCache.cpp
--- a/Engine/Source/Runtime/Engine/Private/GPUSkinCache.cpp
+++ b/Engine/Source/Runtime/Engine/Private/GPUSkinCache.cpp
@@ -139,6 +139,14 @@ FAutoConsoleVariableRef CVarGPUSkinCacheMemoryLimitForBatchedRayTracingGeometryU
 	TEXT(""),
 	ECVF_RenderThreadSafe
 );
+
+static int32 GMaxRayTracingPrimitivesPerCmdList = -1;
+FAutoConsoleVariableRef CVarGPUSkinCacheMaxRayTracingPrimitivesPerCmdList(
+	TEXT("r.SkinCache.MaxRayTracingPrimitivesPerCmdList"),
+	GMaxRayTracingPrimitivesPerCmdList,
+	TEXT("Maximum amount of primitives which are batched together into a single command list to fix potential TDRs."),
+	ECVF_RenderThreadSafe
+);
 #endif
 
 static inline bool DoesPlatformSupportGPUSkinCache(const FStaticShaderPlatform Platform)
@@ -770,24 +778,79 @@ void FGPUSkinCache::TransitionAllToReadable(FRHICommandList& RHICmdList)
 }
 
 #if RHI_RAYTRACING
-void FGPUSkinCache::CommitRayTracingGeometryUpdates(FRHICommandList& RHICmdList)
+
+void FGPUSkinCache::AddRayTracingGeometryToUpdate(FRayTracingGeometry* InRayTracingGeometry, EAccelerationStructureBuildMode InBuildMode)
 {
-	SCOPED_DRAW_EVENT(RHICmdList, CommitSkeletalRayTracingGeometryUpdates);
+	EAccelerationStructureBuildMode* CurrentBuildMode = RayTracingGeometriesToUpdate.Find(InRayTracingGeometry);
+	if (CurrentBuildMode == nullptr)
+	{
+		RayTracingGeometriesToUpdate.Add(InRayTracingGeometry, InBuildMode);
+	}
+	// If currently updating but need full rebuild then update the stored build mode
+	else if (*CurrentBuildMode == EAccelerationStructureBuildMode::Update && InBuildMode == EAccelerationStructureBuildMode::Build)
+	{
+		RayTracingGeometriesToUpdate[InRayTracingGeometry] = InBuildMode;
+	}
+}
+
+void FGPUSkinCache::CommitRayTracingGeometryUpdates(FRHICommandListImmediate& RHICmdList)
+{
+	TRACE_CPUPROFILER_EVENT_SCOPE(FGPUSkinCache::CommitRayTracingGeometryUpdates);
+
+	SCOPED_DRAW_EVENT(RHICmdList, CommitSkeletalRayTracingGeometryUpdates);	
 
 	if (RayTracingGeometriesToUpdate.Num())
 	{
-		TArray<FAccelerationStructureBuildParams> Updates;
-		for (const FRayTracingGeometry* RayTracingGeometry : RayTracingGeometriesToUpdate)
+		// If we have more deferred deleted data than set limit then force flush to make sure all pending releases have actually been freed
+		// before reallocating a lot of new BLAS data
+		if (RayTracingGeometryMemoryPendingRelease >= GMemoryLimitForBatchedRayTracingGeometryUpdates * 1024ull * 1024ull)
 		{
-			FAccelerationStructureBuildParams Params;
-			Params.BuildMode = EAccelerationStructureBuildMode::Update;
-			Params.Geometry = RayTracingGeometry->RayTracingGeometryRHI;
-			Params.Segments = RayTracingGeometry->Initializer.Segments;
-			Updates.Add(Params);
+			RHICmdList.ImmediateFlush(EImmediateFlushType::FlushRHIThreadFlushResources);
+			UE_LOG(LogSkinCache, Display, TEXT("Flushing RHI resource pending deletes due to %d MB limit"), GMemoryLimitForBatchedRayTracingGeometryUpdates);
 		}
+				
+		// Track the amount of primitives which need to be build/updated in a single batch
+		uint64 PrimitivesToUpdates = 0;
+		TArray<FAccelerationStructureBuildParams> BatchedBuildParams;
+		BatchedBuildParams.Reserve(RayTracingGeometriesToUpdate.Num());
 
-		RHICmdList.BuildAccelerationStructures(Updates);
-		RayTracingGeometriesToUpdate.Reset();
+		// Iterate all the geometries which need an update
+		for (TMap<FRayTracingGeometry*, EAccelerationStructureBuildMode>::TRangedForIterator Iter = RayTracingGeometriesToUpdate.begin(); Iter != RayTracingGeometriesToUpdate.end(); ++Iter)
+		{
+			FRayTracingGeometry* RayTracingGeometry = Iter.Key();
+			FAccelerationStructureBuildParams BuildParams;
+			BuildParams.Geometry	= RayTracingGeometry->RayTracingGeometryRHI;
+			BuildParams.BuildMode   = Iter.Value();
+			BuildParams.Segments	= RayTracingGeometry->Initializer.Segments;
+
+			// Make 'Build' 10 times more expensive than 1 'Update' of the BVH
+			uint32 PrimitiveCount = RayTracingGeometry->Initializer.TotalPrimitiveCount;
+			if (BuildParams.BuildMode == EAccelerationStructureBuildMode::Build)
+			{
+				PrimitiveCount *= 10;
+			}
+			PrimitivesToUpdates += PrimitiveCount;
+
+			// Add to batch
+			BatchedBuildParams.Add(BuildParams);
+
+			// Flush batch when limit is reached
+			if (GMaxRayTracingPrimitivesPerCmdList > 0 && PrimitivesToUpdates >= GMaxRayTracingPrimitivesPerCmdList)
+			{
+				RHICmdList.BuildAccelerationStructures(BatchedBuildParams);
+				RHICmdList.SubmitCommandsHint();
+
+				BatchedBuildParams.Empty(BatchedBuildParams.Max());
+				PrimitivesToUpdates = 0;
+			}
+		}
+
+		// Enqueue the last batch
+		RHICmdList.BuildAccelerationStructures(BatchedBuildParams);
+
+		// Clear working data
+		RayTracingGeometriesToUpdate.Empty(RayTracingGeometriesToUpdate.Num());
+		RayTracingGeometryMemoryPendingRelease = 0;
 	}
 }
 #endif // RHI_RAYTRACING
@@ -1472,22 +1535,24 @@ void FGPUSkinCache::ProcessRayTracingGeometryToUpdate(
 
 			if (RayTracingGeometry.RayTracingGeometryRHI.IsValid())
 			{
-				// RayTracingGeometry.UpdateRHI() releases the old RT geometry, however due to the deferred deletion nature of RHI resources
+				// RayTracingGeometry.ReleaseRHI() releases the old RT geometry, however due to the deferred deletion nature of RHI resources
 				// they will not be released until the end of the frame. We may get OOM in the middle of batched updates if not flushing.
 				// This memory size is an estimation based on vertex & index buffer size. In reality the flush happens at 2-3x of the number specified.
 				RayTracingGeometryMemoryPendingRelease += MemoryEstimation;
 
-				if (RayTracingGeometryMemoryPendingRelease >= GMemoryLimitForBatchedRayTracingGeometryUpdates * 1024ull * 1024ull)
-				{
-					RayTracingGeometryMemoryPendingRelease = 0;
-					RHICmdList.SubmitCommandsHint();
-					RHICmdList.ImmediateFlush(EImmediateFlushType::FlushRHIThreadFlushResources);
-					UE_LOG(LogSkinCache, Display, TEXT("Flushing RHI resource pending deletes due to %d MB limit"), GMemoryLimitForBatchedRayTracingGeometryUpdates);
-				}
+				// Release the old data (make sure it's not pending build anymore either)
+				RemoveRayTracingGeometryUpdate(&RayTracingGeometry);
+				RayTracingGeometry.ReleaseRHI();
 			}
 
+			// Update the new init data
 			RayTracingGeometry.SetInitializer(Initializer);
-			RayTracingGeometry.UpdateRHI();
+
+			// Only create RHI object but enqueue actual BLAS creation so they can be accumulated
+			RayTracingGeometry.RayTracingGeometryRHI = RHICreateRayTracingGeometry(Initializer);
+
+			// Request the build
+			AddRayTracingGeometryToUpdate(&RayTracingGeometry, EAccelerationStructureBuildMode::Build);
 		}
 		else
 		{
@@ -1496,7 +1561,9 @@ void FGPUSkinCache::ProcessRayTracingGeometryToUpdate(
 			{
 				// Refit BLAS with new vertex buffer data
 				FGPUSkinCache::GetRayTracingSegmentVertexBuffers(*SkinCacheEntry, RayTracingGeometry.Initializer.Segments);
-				AddRayTracingGeometryToUpdate(&RayTracingGeometry);
+
+				// Request the update
+				AddRayTracingGeometryToUpdate(&RayTracingGeometry, EAccelerationStructureBuildMode::Update);
 			}
 			else
 			{
diff --git a/Engine/Source/Runtime/Engine/Public/GPUSkinCache.h b/Engine/Source/Runtime/Engine/Public/GPUSkinCache.h
--- a/Engine/Source/Runtime/Engine/Public/GPUSkinCache.h
+++ b/Engine/Source/Runtime/Engine/Public/GPUSkinCache.h
@@ -367,17 +367,14 @@ public:
 	ENGINE_API void TransitionAllToReadable(FRHICommandList& RHICmdList);
 
 #if RHI_RAYTRACING
-	void AddRayTracingGeometryToUpdate(FRayTracingGeometry* RayTracingGeometry)
-	{
-		RayTracingGeometriesToUpdate.Add(RayTracingGeometry);
-	}
-
-	void CommitRayTracingGeometryUpdates(FRHICommandList& RHICmdList);
-
+	void AddRayTracingGeometryToUpdate(FRayTracingGeometry* InRayTracingGeometry, EAccelerationStructureBuildMode InBuildMode);
+	void CommitRayTracingGeometryUpdates(FRHICommandListImmediate& RHICmdList);
 	void RemoveRayTracingGeometryUpdate(FRayTracingGeometry* RayTracingGeometry)
 	{
 		if (RayTracingGeometriesToUpdate.Find(RayTracingGeometry) != nullptr)
+		{
 			RayTracingGeometriesToUpdate.Remove(RayTracingGeometry);
+		}
 	}
 
 	void ProcessRayTracingGeometryToUpdate(
@@ -399,7 +396,7 @@ protected:
 
 	TSet<FRHIUnorderedAccessView*> BuffersToTransition;
 #if RHI_RAYTRACING
-	TSet<FRayTracingGeometry*> RayTracingGeometriesToUpdate;
+	TMap<FRayTracingGeometry*, EAccelerationStructureBuildMode> RayTracingGeometriesToUpdate;
 	uint64 RayTracingGeometryMemoryPendingRelease = 0;
 #endif // RHI_RAYTRACING
 
