WorldPartition: Remove FWorldPartitionActorDescFactory and move creation of ActorDescs in AActor

diff --git a/Engine/Source/Runtime/Engine/Classes/GameFramework/Actor.h b/Engine/Source/Runtime/Engine/Classes/GameFramework/Actor.h
--- a/Engine/Source/Runtime/Engine/Classes/GameFramework/Actor.h
+++ b/Engine/Source/Runtime/Engine/Classes/GameFramework/Actor.h
@@ -787,6 +787,24 @@ public:
 
 	/** Returns this actor's current Guid. Actor Guids are only available in development builds. */
 	inline const FGuid& GetActorGuid() const { return ActorGuid; }
+
+private:
+	friend class UWorldPartition;
+
+	/**
+	 * Creates an uninitialized actor descriptor. Can be overriden.
+	 */
+	virtual TUniquePtr<class FWorldPartitionActorDesc> CreateClassActorDesc() const;
+
+	/**
+	 * Creates an initialized actor descriptor.
+	 */
+	TUniquePtr<class FWorldPartitionActorDesc> CreateActorDesc() const;
+
+	/**
+	 * Creates an uninitialized actor descriptor.
+	 */
+	static TUniquePtr<class FWorldPartitionActorDesc> CreateClassActorDesc(const TSubclassOf<AActor>& ActorClass);
 #endif // WITH_EDITOR
 
 public:
diff --git a/Engine/Source/Runtime/Engine/Private/Actor.cpp b/Engine/Source/Runtime/Engine/Private/Actor.cpp
--- a/Engine/Source/Runtime/Engine/Private/Actor.cpp
+++ b/Engine/Source/Runtime/Engine/Private/Actor.cpp
@@ -420,7 +420,7 @@ void AActor::GetAssetRegistryTags(TArray<FAssetRegistryTag>& OutTags) const
 	{
 		if (UWorldPartition* WorldPartition = GetWorldSettings()->GetWorldPartition())
 		{
-			TUniquePtr<FWorldPartitionActorDesc> ActorDesc(UWorldPartition::CreateActorDesc(this));
+			TUniquePtr<FWorldPartitionActorDesc> ActorDesc(CreateActorDesc());
 			
 			const FString ActorMetaDataClass = GetParentNativeClass(GetClass())->GetName();
 			static FName NAME_ActorMetaDataClass(TEXT("ActorMetaDataClass"));
diff --git a/Engine/Source/Runtime/Engine/Private/ActorEditor.cpp b/Engine/Source/Runtime/Engine/Private/ActorEditor.cpp
--- a/Engine/Source/Runtime/Engine/Private/ActorEditor.cpp
+++ b/Engine/Source/Runtime/Engine/Private/ActorEditor.cpp
@@ -27,6 +27,7 @@
 #if WITH_EDITOR
 
 #include "Editor.h"
+#include "WorldPartition/WorldPartitionActorDesc.h"
 
 #define LOCTEXT_NAMESPACE "ErrorChecking"
 
@@ -884,6 +885,27 @@ EActorGridPlacement AActor::GetDefaultGridPlacement() const
 	return EActorGridPlacement::None;
 }
 
+TUniquePtr<FWorldPartitionActorDesc> AActor::CreateClassActorDesc() const
+{
+	return TUniquePtr<FWorldPartitionActorDesc>(new FWorldPartitionActorDesc());
+}
+
+TUniquePtr<FWorldPartitionActorDesc> AActor::CreateActorDesc() const
+{
+	check(!HasAnyFlags(RF_ArchetypeObject | RF_ClassDefaultObject));
+	
+	TUniquePtr<FWorldPartitionActorDesc> ActorDesc(CreateClassActorDesc());
+		
+	ActorDesc->Init(this);
+	
+	return ActorDesc;
+}
+
+TUniquePtr<class FWorldPartitionActorDesc> AActor::CreateClassActorDesc(const TSubclassOf<AActor>& ActorClass)
+{
+	return CastChecked<AActor>(ActorClass->GetDefaultObject())->CreateClassActorDesc();
+}
+
 const FString& AActor::GetActorLabel() const
 {
 	// If the label string is empty then we'll use the default actor label (usually the actor's class name.)
diff --git a/Engine/Source/Runtime/Engine/Private/ActorPartition/ActorPartitionSubsystem.cpp b/Engine/Source/Runtime/Engine/Private/ActorPartition/ActorPartitionSubsystem.cpp
--- a/Engine/Source/Runtime/Engine/Private/ActorPartition/ActorPartitionSubsystem.cpp
+++ b/Engine/Source/Runtime/Engine/Private/ActorPartition/ActorPartitionSubsystem.cpp
@@ -303,12 +303,8 @@ bool UActorPartitionSubsystem::IsLevelPartition() const
 #if WITH_EDITOR
 void UActorPartitionSubsystem::Initialize(FSubsystemCollectionBase& Collection)
 {
-	UWorldPartitionSubsystem* WorldPartitionSubsystem = Collection.InitializeDependency<UWorldPartitionSubsystem>();
-	if (WorldPartitionSubsystem)
-	{
-		WorldPartitionSubsystem->RegisterActorDescFactory(APartitionActor::StaticClass(), &PartitionActorDescFactory);
-	}
-
+	Collection.InitializeDependency<UWorldPartitionSubsystem>();
+	
 	// Will need to register to WorldPartition setup changes events here...
 	InitializeActorPartition();
 }
diff --git a/Engine/Source/Runtime/Engine/Private/ActorPartition/PartitionActor.cpp b/Engine/Source/Runtime/Engine/Private/ActorPartition/PartitionActor.cpp
--- a/Engine/Source/Runtime/Engine/Private/ActorPartition/PartitionActor.cpp
+++ b/Engine/Source/Runtime/Engine/Private/ActorPartition/PartitionActor.cpp
@@ -4,6 +4,8 @@
 
 #if WITH_EDITOR
 #include "Components/BoxComponent.h"
+#include "WorldPartition/WorldPartitionActorDesc.h"
+#include "WorldPartition/ActorPartition/PartitionActorDesc.h"
 #endif
 
 APartitionActor::APartitionActor(const FObjectInitializer& ObjectInitializer)
@@ -15,4 +17,11 @@ APartitionActor::APartitionActor(const FObjectInitializer& ObjectInitializer)
 	USceneComponent* SceneComponent = CreateDefaultSubobject<USceneComponent>(TEXT("RootComponent0"));
 	RootComponent = SceneComponent;
 	RootComponent->Mobility = EComponentMobility::Static;
-}
\ No newline at end of file
+}
+
+#if WITH_EDITOR
+TUniquePtr<FWorldPartitionActorDesc> APartitionActor::CreateClassActorDesc() const
+{
+	return TUniquePtr<FWorldPartitionActorDesc>(new FPartitionActorDesc());
+}
+#endif
\ No newline at end of file
diff --git a/Engine/Source/Runtime/Engine/Private/LevelInstance/LevelInstanceActor.cpp b/Engine/Source/Runtime/Engine/Private/LevelInstance/LevelInstanceActor.cpp
--- a/Engine/Source/Runtime/Engine/Private/LevelInstance/LevelInstanceActor.cpp
+++ b/Engine/Source/Runtime/Engine/Private/LevelInstance/LevelInstanceActor.cpp
@@ -13,6 +13,8 @@
 #include "Logging/TokenizedMessage.h"
 #include "Misc/MapErrors.h"
 #include "Misc/UObjectToken.h"
+#include "WorldPartition/WorldPartitionActorDesc.h"
+#include "WorldPartition/LevelInstance/LevelInstanceActorDesc.h"
 #endif
 
 #define LOCTEXT_NAMESPACE "LevelInstanceActor"
@@ -175,6 +177,11 @@ const FGuid& ALevelInstance::GetLevelInstanceActorGuid() const
 
 #if WITH_EDITOR
 
+TUniquePtr<FWorldPartitionActorDesc> ALevelInstance::CreateClassActorDesc() const
+{
+	return TUniquePtr<FWorldPartitionActorDesc>(new FLevelInstanceActorDesc());
+}
+
 AActor* ALevelInstance::FindEditorInstanceActor() const
 {
 	AActor* FoundActor = nullptr;
diff --git a/Engine/Source/Runtime/Engine/Private/LevelInstance/LevelInstanceSubsystem.cpp b/Engine/Source/Runtime/Engine/Private/LevelInstance/LevelInstanceSubsystem.cpp
--- a/Engine/Source/Runtime/Engine/Private/LevelInstance/LevelInstanceSubsystem.cpp
+++ b/Engine/Source/Runtime/Engine/Private/LevelInstance/LevelInstanceSubsystem.cpp
@@ -16,7 +16,6 @@
 
 #if WITH_EDITOR
 #include "LevelInstance/LevelInstanceEditorLevelStreaming.h"
-#include "WorldPartition/LevelInstance/LevelInstanceActorDescFactory.h"
 #include "Misc/ScopedSlowTask.h"
 #include "Misc/ITransaction.h"
 #include "AssetRegistryModule.h"
@@ -58,14 +57,6 @@ void ULevelInstanceSubsystem::Initialize(FSubsystemCollectionBase& Collection)
 	Super::Initialize(Collection);
 
 #if WITH_EDITOR
-	LevelInstanceActorDescFactory.Reset(new FLevelInstanceActorDescFactory());
-
-	UWorldPartitionSubsystem* WorldPartitionSubsystem = Collection.InitializeDependency<UWorldPartitionSubsystem>();
-	if (WorldPartitionSubsystem)
-	{
-		RegisterActorDescFactories(WorldPartitionSubsystem);
-	}
-
 	if (GEditor)
 	{
 		FModuleManager::LoadModuleChecked<ILevelInstanceEditorModule>("LevelInstanceEditor");
@@ -957,11 +948,6 @@ void ULevelInstanceSubsystem::OnActorDeleted(AActor* Actor)
 	}
 }
 
-void ULevelInstanceSubsystem::RegisterActorDescFactories(UWorldPartitionSubsystem* WorldPartitionSubsystem)
-{
-	WorldPartitionSubsystem->RegisterActorDescFactory(ALevelInstance::StaticClass(), LevelInstanceActorDescFactory.Get());
-}
-
 bool ULevelInstanceSubsystem::ShouldIgnoreDirtyPackage(UPackage* DirtyPackage, const UWorld* EditingWorld)
 {
 	if (DirtyPackage == EditingWorld->GetOutermost())
diff --git a/Engine/Source/Runtime/Engine/Private/WorldPartition/ActorPartition/PartitionActorDescFactory.cpp b/Engine/Source/Runtime/Engine/Private/WorldPartition/ActorPartition/PartitionActorDescFactory.cpp
deleted file mode 100644
diff --git a/Engine/Source/Runtime/Engine/Private/WorldPartition/HLOD/HLODActor.cpp b/Engine/Source/Runtime/Engine/Private/WorldPartition/HLOD/HLODActor.cpp
--- a/Engine/Source/Runtime/Engine/Private/WorldPartition/HLOD/HLODActor.cpp
+++ b/Engine/Source/Runtime/Engine/Private/WorldPartition/HLOD/HLODActor.cpp
@@ -7,6 +7,7 @@
 #if WITH_EDITOR
 #include "WorldPartition/WorldPartition.h"
 #include "WorldPartition/WorldPartitionActorDesc.h"
+#include "WorldPartition/HLOD/HLODActorDesc.h"
 #endif
 
 AWorldPartitionHLOD::AWorldPartitionHLOD(const FObjectInitializer& ObjectInitializer)
@@ -154,6 +155,11 @@ EActorGridPlacement AWorldPartitionHLOD::GetDefaultGridPlacement() const
 	return EActorGridPlacement::Location;
 }
 
+TUniquePtr<FWorldPartitionActorDesc> AWorldPartitionHLOD::CreateClassActorDesc() const
+{
+	return TUniquePtr<FWorldPartitionActorDesc>(new FHLODActorDesc());
+}
+
 void AWorldPartitionHLOD::SetHLODPrimitives(const TArray<UPrimitiveComponent*>& InHLODPrimitives)
 {
 	TRACE_CPUPROFILER_EVENT_SCOPE(AWorldPartitionHLOD::SetHLODPrimitive);
diff --git a/Engine/Source/Runtime/Engine/Private/WorldPartition/HLOD/HLODActorDescFactory.cpp b/Engine/Source/Runtime/Engine/Private/WorldPartition/HLOD/HLODActorDescFactory.cpp
deleted file mode 100644
diff --git a/Engine/Source/Runtime/Engine/Private/WorldPartition/HLOD/HLODSubsystem.cpp b/Engine/Source/Runtime/Engine/Private/WorldPartition/HLOD/HLODSubsystem.cpp
--- a/Engine/Source/Runtime/Engine/Private/WorldPartition/HLOD/HLODSubsystem.cpp
+++ b/Engine/Source/Runtime/Engine/Private/WorldPartition/HLOD/HLODSubsystem.cpp
@@ -42,13 +42,7 @@ void UHLODSubsystem::Initialize(FSubsystemCollectionBase& Collection)
 {
 	Super::Initialize(Collection);
 
-	UWorldPartitionSubsystem* WorldPartitionSubsystem = Collection.InitializeDependency<UWorldPartitionSubsystem>();
-
 #if WITH_EDITOR
-	HLODActorDescFactory.Reset(new FHLODActorDescFactory());
-
-	RegisterActorDescFactories(WorldPartitionSubsystem);
-
 	if (!IsRunningCommandlet())
 	{
 		GetWorld()->PersistentLevel->OnLoadedActorAddedToLevelEvent.AddUObject(this, &UHLODSubsystem::OnActorLoaded);
@@ -106,11 +100,6 @@ void UHLODSubsystem::OnActorUnloaded(AActor& Actor)
 		}
 	}
 }
-
-void UHLODSubsystem::RegisterActorDescFactories(UWorldPartitionSubsystem* WorldPartitionSubsystem)
-{
-	WorldPartitionSubsystem->RegisterActorDescFactory(AWorldPartitionHLOD::StaticClass(), HLODActorDescFactory.Get());
-}
 #endif
 
 void UHLODSubsystem::RegisterHLODActor(AWorldPartitionHLOD* InWorldPartitionHLOD)
diff --git a/Engine/Source/Runtime/Engine/Private/WorldPartition/Landscape/LandscapeActorDescFactory.cpp b/Engine/Source/Runtime/Engine/Private/WorldPartition/Landscape/LandscapeActorDescFactory.cpp
deleted file mode 100644
diff --git a/Engine/Source/Runtime/Engine/Private/WorldPartition/LevelInstance/LevelInstanceActorDescFactory.cpp b/Engine/Source/Runtime/Engine/Private/WorldPartition/LevelInstance/LevelInstanceActorDescFactory.cpp
deleted file mode 100644
diff --git a/Engine/Source/Runtime/Engine/Private/WorldPartition/WorldPartition.cpp b/Engine/Source/Runtime/Engine/Private/WorldPartition/WorldPartition.cpp
--- a/Engine/Source/Runtime/Engine/Private/WorldPartition/WorldPartition.cpp
+++ b/Engine/Source/Runtime/Engine/Private/WorldPartition/WorldPartition.cpp
@@ -31,7 +31,6 @@
 #include "ScopedTransaction.h"
 #include "UnrealEdMisc.h"
 #include "AssetRegistryModule.h"
-#include "WorldPartition/WorldPartitionActorDescFactory.h"
 #include "WorldPartition/WorldPartitionLevelStreamingDynamic.h"
 #include "WorldPartition/WorldPartitionEditorHash.h"
 #include "WorldPartition/WorldPartitionEditorSpatialHash.h"
@@ -339,7 +338,6 @@ void UWorldPartition::Uninitialize()
 				UpdateLoadingEditorCell(Cell, /*bShouldBeLoaded*/false);
 			});
 			UnregisterDelegates();
-			ActorDescFactories.Empty();
 		}
 
 		EditorHash = nullptr;
@@ -408,53 +406,6 @@ void UWorldPartition::ApplyActorTransform(AActor* InActor, const FTransform& InT
 	}
 }
 
-TUniquePtr<FWorldPartitionActorDescFactory> UWorldPartition::DefaultActorDescFactory;
-TMap<FName, FWorldPartitionActorDescFactory*> UWorldPartition::ActorDescFactories;
-
-void UWorldPartition::RegisterActorDescFactory(TSubclassOf<AActor> Class, FWorldPartitionActorDescFactory* Factory)
-{
-	for (TObjectIterator<UClass> ClassIt; ClassIt; ++ClassIt)
-	{
-		if (ClassIt->IsChildOf(Class))
-		{
-			const FName ClassName = ClassIt->GetFName();
-			if (!ActorDescFactories.Contains(ClassName))
-			{
-				ActorDescFactories.Add(ClassName, Factory);
-			}
-		}
-	}
-}
-
-FWorldPartitionActorDescFactory* UWorldPartition::GetActorDescFactory(TSubclassOf<AActor> Class)
-{
-	Class = GetParentNativeClass(Class);
-
-	FName ClassName = Class->GetFName();
-	if (FWorldPartitionActorDescFactory** Factory = ActorDescFactories.Find(ClassName))
-	{
-		return *Factory;
-	}
-
-	if (!DefaultActorDescFactory)
-	{
-		DefaultActorDescFactory.Reset(new FWorldPartitionActorDescFactory());
-	}
-	return DefaultActorDescFactory.Get();
-}
-
-FWorldPartitionActorDescFactory* UWorldPartition::GetActorDescFactory(const AActor* Actor)
-{
-	return GetActorDescFactory(Actor->GetClass());
-}
-
-TUniquePtr<FWorldPartitionActorDesc> UWorldPartition::CreateActorDesc(const AActor* Actor)
-{
-	TUniquePtr<FWorldPartitionActorDesc> ActorDesc(GetActorDescFactory(Actor)->Create());
-	ActorDesc->Init(Actor);
-	return ActorDesc;
-}
-
 void UWorldPartition::ForEachIntersectingActorDesc(const FBox& Box, TSubclassOf<AActor> ActorClass, TFunctionRef<bool(const FWorldPartitionActorDesc*)> Predicate) const
 {
 	EditorHash->ForEachIntersectingActor(Box, [&ActorClass, Predicate](const FWorldPartitionActorDesc* ActorDesc)
@@ -933,10 +884,9 @@ TUniquePtr<FWorldPartitionActorDesc> UWorldPartition::GetActorDescriptor(const F
 			ActorDescInitData.ActorPath = InAssetData.ObjectPath;
 			FBase64::Decode(ActorMetaDataStr, ActorDescInitData.SerializedData);
 
-			FWorldPartitionActorDesc* NewActorDesc(GetActorDescFactory(ActorDescInitData.NativeClass)->Create());
+			TUniquePtr<FWorldPartitionActorDesc> NewActorDesc(AActor::CreateClassActorDesc(ActorDescInitData.NativeClass));
 			NewActorDesc->Init(ActorDescInitData);
-
-			return TUniquePtr<FWorldPartitionActorDesc>(NewActorDesc);
+			return NewActorDesc;
 		}
 	}
 
diff --git a/Engine/Source/Runtime/Engine/Private/WorldPartition/WorldPartitionActorDescFactory.cpp b/Engine/Source/Runtime/Engine/Private/WorldPartition/WorldPartitionActorDescFactory.cpp
deleted file mode 100644
diff --git a/Engine/Source/Runtime/Engine/Private/WorldPartition/WorldPartitionCookPackageSplitter.cpp b/Engine/Source/Runtime/Engine/Private/WorldPartition/WorldPartitionCookPackageSplitter.cpp
--- a/Engine/Source/Runtime/Engine/Private/WorldPartition/WorldPartitionCookPackageSplitter.cpp
+++ b/Engine/Source/Runtime/Engine/Private/WorldPartition/WorldPartitionCookPackageSplitter.cpp
@@ -10,23 +10,10 @@
 #include "WorldPartition/WorldPartitionRuntimeHash.h"
 #include "WorldPartition/WorldPartitionSubsystem.h"
 #include "WorldPartition/WorldPartition.h"
-#include "WorldPartition/HLOD/HLODActor.h"
-#include "WorldPartition/HLOD/HLODActorDescFactory.h"
-#include "WorldPartition/LevelInstance/LevelInstanceActorDescFactory.h"
-#include "WorldPartition/Landscape/LandscapeActorDescFactory.h"
-#include "WorldPartition/ActorPartition/PartitionActorDescFactory.h"
-#include "LevelInstance/LevelInstanceActor.h"
-#include "LandscapeProxy.h"
-#include "ActorPartition/PartitionActor.h"
 
 // Register FWorldPartitionCookPackageSplitter for UWorld class
 REGISTER_COOKPACKAGE_SPLITTER(FWorldPartitionCookPackageSplitter, UWorld);
 
-static FHLODActorDescFactory GHLODActorDescFactory;
-static FLevelInstanceActorDescFactory GLevelInstanceActorDescFactory;
-static FLandscapeActorDescFactory GLandscapeActorDescFactory;
-static FPartitionActorDescFactory GPartitionActorDescFactory;
-
 bool FWorldPartitionCookPackageSplitter::ShouldSplit(UObject* SplitData)
 {
 	UWorld* World = Cast<UWorld>(SplitData);
@@ -46,13 +33,6 @@ TArray<ICookPackageSplitter::FGeneratedPackage> FWorldPartitionCookPackageSplitt
 	// World is not initialized
 	ensure(!PartitionedWorld->bIsWorldInitialized);
 
-	// Manually register ActorDesc factories.
-	// @todo_ow: Revisit how these factories are declared. This could be automatic.
-	UWorldPartition::RegisterActorDescFactory(AWorldPartitionHLOD::StaticClass(), &GHLODActorDescFactory);
-	UWorldPartition::RegisterActorDescFactory(ALevelInstance::StaticClass(), &GLevelInstanceActorDescFactory);
-	UWorldPartition::RegisterActorDescFactory(ALandscapeProxy::StaticClass(), &GLandscapeActorDescFactory);
-	UWorldPartition::RegisterActorDescFactory(APartitionActor::StaticClass(), &GPartitionActorDescFactory);
-
 	// Manually initialize WorldPartition
 	UWorldPartition* WorldPartition = PartitionedWorld->PersistentLevel->GetWorldPartition();
 	WorldPartition->Initialize(PartitionedWorld, FTransform::Identity);
diff --git a/Engine/Source/Runtime/Engine/Private/WorldPartition/WorldPartitionSubsystem.cpp b/Engine/Source/Runtime/Engine/Private/WorldPartition/WorldPartitionSubsystem.cpp
--- a/Engine/Source/Runtime/Engine/Private/WorldPartition/WorldPartitionSubsystem.cpp
+++ b/Engine/Source/Runtime/Engine/Private/WorldPartition/WorldPartitionSubsystem.cpp
@@ -3,7 +3,6 @@
 #include "WorldPartition/WorldPartitionSubsystem.h"
 #include "WorldPartition/WorldPartition.h"
 #include "WorldPartition/WorldPartitionStreamingPolicy.h"
-#include "WorldPartition/WorldPartitionActorDescFactory.h"
 #include "Engine/Canvas.h"
 #include "Engine/Console.h"
 #include "ConsoleSettings.h"
@@ -206,9 +205,4 @@ void UWorldPartitionSubsystem::ForEachActorDesc(TSubclassOf<AActor> ActorClass,
 		MainPartition->ForEachActorDesc(ActorClass, Predicate);
 	}
 }
-
-void UWorldPartitionSubsystem::RegisterActorDescFactory(TSubclassOf<AActor> Class, FWorldPartitionActorDescFactory* Factory)
-{
-	UWorldPartition::RegisterActorDescFactory(Class, Factory);
-}
 #endif // WITH_EDITOR
\ No newline at end of file
diff --git a/Engine/Source/Runtime/Engine/Public/ActorPartition/ActorPartitionSubsystem.h b/Engine/Source/Runtime/Engine/Public/ActorPartition/ActorPartitionSubsystem.h
--- a/Engine/Source/Runtime/Engine/Public/ActorPartition/ActorPartitionSubsystem.h
+++ b/Engine/Source/Runtime/Engine/Public/ActorPartition/ActorPartitionSubsystem.h
@@ -6,7 +6,6 @@
 #include "Templates/SubclassOf.h"
 #include "Misc/HashBuilder.h"
 #include "Misc/Guid.h"
-#include "WorldPartition/ActorPartition/PartitionActorDescFactory.h"
 #include "ActorPartitionSubsystem.generated.h"
 
 class FBaseActorPartition;
@@ -150,8 +149,6 @@ private:
 	TMap<FCellCoord, TMap<UClass*, TMap<FGuid, TWeakObjectPtr<APartitionActor>>>> PartitionedActors;
 	TUniquePtr<FBaseActorPartition> ActorPartition;
 	
-	FPartitionActorDescFactory PartitionActorDescFactory;
-
 	FDelegateHandle ActorPartitionHashInvalidatedHandle;
 #endif
 };
diff --git a/Engine/Source/Runtime/Engine/Public/ActorPartition/PartitionActor.h b/Engine/Source/Runtime/Engine/Public/ActorPartition/PartitionActor.h
--- a/Engine/Source/Runtime/Engine/Public/ActorPartition/PartitionActor.h
+++ b/Engine/Source/Runtime/Engine/Public/ActorPartition/PartitionActor.h
@@ -21,6 +21,7 @@ public:
 	//~ Begin AActor Interface
 #if WITH_EDITOR	
 	virtual EActorGridPlacement GetDefaultGridPlacement() const override { return EActorGridPlacement::Location; }
+	virtual TUniquePtr<class FWorldPartitionActorDesc> CreateClassActorDesc() const override;
 	virtual uint32 GetDefaultGridSize(UWorld* InWorld) const PURE_VIRTUAL(APartitionActor, return 0;)
 	virtual FGuid GetGridGuid() const { return FGuid(); }
 #endif
diff --git a/Engine/Source/Runtime/Engine/Public/LevelInstance/LevelInstanceActor.h b/Engine/Source/Runtime/Engine/Public/LevelInstance/LevelInstanceActor.h
--- a/Engine/Source/Runtime/Engine/Public/LevelInstance/LevelInstanceActor.h
+++ b/Engine/Source/Runtime/Engine/Public/LevelInstance/LevelInstanceActor.h
@@ -48,6 +48,7 @@ public:
 	// AActor overrides
 	virtual void CheckForErrors() override;
 	virtual EActorGridPlacement GetDefaultGridPlacement() const override { return EActorGridPlacement::Bounds; }
+	virtual TUniquePtr<class FWorldPartitionActorDesc> CreateClassActorDesc() const override;
 	virtual bool CanDeleteSelectedActor(FText& OutReason) const override;
 	virtual void SetIsTemporarilyHiddenInEditor(bool bIsHidden) override;
 	virtual void EditorGetUnderlyingActors(TSet<AActor*>& OutUnderlyingActors) override;
diff --git a/Engine/Source/Runtime/Engine/Public/LevelInstance/LevelInstanceSubsystem.h b/Engine/Source/Runtime/Engine/Public/LevelInstance/LevelInstanceSubsystem.h
--- a/Engine/Source/Runtime/Engine/Public/LevelInstance/LevelInstanceSubsystem.h
+++ b/Engine/Source/Runtime/Engine/Public/LevelInstance/LevelInstanceSubsystem.h
@@ -7,10 +7,6 @@
 #include "UObject/NameTypes.h"
 #include "LevelInstance/LevelInstanceTypes.h"
 
-#if WITH_EDITOR
-#include "WorldPartition/LevelInstance/LevelInstanceActorDescFactory.h"
-#endif
-
 #include "LevelInstanceSubsystem.generated.h"
 
 class ALevelInstance;
@@ -70,7 +66,6 @@ public:
 	bool IsEditingLevelInstanceDirty(const ALevelInstance* LevelInstanceActor) const;
 	bool IsEditingLevelInstance(const ALevelInstance* LevelInstanceActor) const { return GetLevelInstanceEdit(LevelInstanceActor) != nullptr; }
 	
-	FLevelInstanceActorDescFactory* GetActorDescFactory() const { return LevelInstanceActorDescFactory.Get(); }
 	bool GetLevelInstanceBounds(const ALevelInstance* LevelInstanceActor, FBox& OutBounds) const;
 	
 	void ForEachActorInLevelInstance(const ALevelInstance* LevelInstanceActor, TFunctionRef<bool(AActor * LevelActor)> Operation) const;
@@ -110,7 +105,6 @@ private:
 #if WITH_EDITOR
 	void CommitChildrenLevelInstances(ALevelInstance* LevelInstanceActor);
 
-	void RegisterActorDescFactories(UWorldPartitionSubsystem* WorldPartitionSubsystem);
 	static bool ShouldIgnoreDirtyPackage(UPackage* DirtyPackage, const UWorld* EditingWorld);
 
 	struct FLevelInstanceEdit
@@ -141,8 +135,6 @@ private:
 	TUniquePtr<FLevelsToRemoveScope> LevelsToRemoveScope;
 	
 	TMap<FName, FLevelInstanceEdit> LevelInstanceEdits;
-
-	TUniquePtr<FLevelInstanceActorDescFactory> LevelInstanceActorDescFactory;
 #endif
 
 	struct FLevelInstance
diff --git a/Engine/Source/Runtime/Engine/Public/WorldPartition/ActorPartition/PartitionActorDescFactory.h b/Engine/Source/Runtime/Engine/Public/WorldPartition/ActorPartition/PartitionActorDescFactory.h
deleted file mode 100644
diff --git a/Engine/Source/Runtime/Engine/Public/WorldPartition/HLOD/HLODActor.h b/Engine/Source/Runtime/Engine/Public/WorldPartition/HLOD/HLODActor.h
--- a/Engine/Source/Runtime/Engine/Public/WorldPartition/HLOD/HLODActor.h
+++ b/Engine/Source/Runtime/Engine/Public/WorldPartition/HLOD/HLODActor.h
@@ -61,6 +61,7 @@ protected:
 	virtual void PostRegisterAllComponents() override;
 	virtual void PostUnregisterAllComponents() override;
 	virtual EActorGridPlacement GetDefaultGridPlacement() const override;
+	virtual TUniquePtr<class FWorldPartitionActorDesc> CreateClassActorDesc() const override;
 	virtual void PostActorCreated() override;
 #endif
 	//~ End AActor Interface.
diff --git a/Engine/Source/Runtime/Engine/Public/WorldPartition/HLOD/HLODActorDescFactory.h b/Engine/Source/Runtime/Engine/Public/WorldPartition/HLOD/HLODActorDescFactory.h
deleted file mode 100644
diff --git a/Engine/Source/Runtime/Engine/Public/WorldPartition/HLOD/HLODSubsystem.h b/Engine/Source/Runtime/Engine/Public/WorldPartition/HLOD/HLODSubsystem.h
--- a/Engine/Source/Runtime/Engine/Public/WorldPartition/HLOD/HLODSubsystem.h
+++ b/Engine/Source/Runtime/Engine/Public/WorldPartition/HLOD/HLODSubsystem.h
@@ -6,10 +6,6 @@
 #include "Containers/Map.h"
 #include "UObject/NameTypes.h"
 
-#if WITH_EDITOR
-#include "WorldPartition/HLOD/HLODActorDescFactory.h"
-#endif
-
 #include "HLODSubsystem.generated.h"
 
 class UWorldPartitionSubsystem;
@@ -39,23 +35,13 @@ public:
 	void OnCellShown(const UWorldPartitionRuntimeCell* InCell);
 	void OnCellHidden(const UWorldPartitionRuntimeCell* InCell);
 	
-#if WITH_EDITOR
-	FHLODActorDescFactory* GetActorDescFactory() const { return HLODActorDescFactory.Get(); }
-#endif
-
 private:
 #if WITH_EDITOR
 	void OnActorLoaded(AActor& Actor);
 	void OnActorUnloaded(AActor& Actor);
-	
-	void RegisterActorDescFactories(UWorldPartitionSubsystem* WorldPartitionSubsystem);
 #endif
 
 private:
-#if WITH_EDITORONLY_DATA
-	TUniquePtr<FHLODActorDescFactory> HLODActorDescFactory;
-#endif
-
 	// Mapping between an HLOD GUID & the loaded actor.
 	TMap<FGuid, AWorldPartitionHLOD*> RegisteredHLODActors;
 
diff --git a/Engine/Source/Runtime/Engine/Public/WorldPartition/Landscape/LandscapeActorDesc.h b/Engine/Source/Runtime/Engine/Public/WorldPartition/Landscape/LandscapeActorDesc.h
--- a/Engine/Source/Runtime/Engine/Public/WorldPartition/Landscape/LandscapeActorDesc.h
+++ b/Engine/Source/Runtime/Engine/Public/WorldPartition/Landscape/LandscapeActorDesc.h
@@ -4,7 +4,6 @@
 
 #include "CoreMinimal.h"
 #include "UObject/ObjectMacros.h"
-#include "Containers/Map.h"
 #include "WorldPartition/ActorPartition/PartitionActorDesc.h"
 
 /**
@@ -13,8 +12,6 @@
 class ENGINE_API FLandscapeActorDesc : public FPartitionActorDesc
 {
 #if WITH_EDITOR
-	friend class FLandscapeActorDescFactory;
-
 protected:
 	virtual void Init(const AActor* InActor) override;
 #endif
diff --git a/Engine/Source/Runtime/Engine/Public/WorldPartition/Landscape/LandscapeActorDescFactory.h b/Engine/Source/Runtime/Engine/Public/WorldPartition/Landscape/LandscapeActorDescFactory.h
deleted file mode 100644
diff --git a/Engine/Source/Runtime/Engine/Public/WorldPartition/LevelInstance/LevelInstanceActorDescFactory.h b/Engine/Source/Runtime/Engine/Public/WorldPartition/LevelInstance/LevelInstanceActorDescFactory.h
deleted file mode 100644
diff --git a/Engine/Source/Runtime/Engine/Public/WorldPartition/WorldPartition.h b/Engine/Source/Runtime/Engine/Public/WorldPartition/WorldPartition.h
--- a/Engine/Source/Runtime/Engine/Public/WorldPartition/WorldPartition.h
+++ b/Engine/Source/Runtime/Engine/Public/WorldPartition/WorldPartition.h
@@ -7,7 +7,6 @@
 #include "GameFramework/Actor.h"
 #include "Templates/SubclassOf.h"
 #include "WorldPartition/WorldPartitionActorDesc.h"
-#include "WorldPartition/WorldPartitionActorDescFactory.h"
 
 #if WITH_EDITOR
 #include "PackageSourceControlHelper.h"
@@ -90,9 +89,6 @@ private:
 	DECLARE_MULTICAST_DELEGATE_OneParam(FWorldPartitionChangedEvent, UWorld*);
 	static FWorldPartitionChangedEvent WorldPartitionChangedEvent;
 
-	static FWorldPartitionActorDescFactory* GetActorDescFactory(TSubclassOf<AActor> Class);
-	static FWorldPartitionActorDescFactory* GetActorDescFactory(const AActor* Actor);
-
 	void FlushStreaming();
 
 	// PIE Methods
@@ -121,16 +117,12 @@ private:
 	void ForEachActorDesc(TSubclassOf<AActor> ActorClass, TFunctionRef<bool(const FWorldPartitionActorDesc*)> Predicate) const;
 	// UActorPartitionSubsystem interface-
 public:
-	static void RegisterActorDescFactory(TSubclassOf<AActor> Class, FWorldPartitionActorDescFactory* Factory);
-
 	FName GetWorldPartitionEditorName();
 	bool IsSimulating() const;
 
 	FWorldPartitionActorDesc* GetActorDesc(const FGuid& Guid);
 	const FWorldPartitionActorDesc* GetActorDesc(const FGuid& Guid) const;
 
-	static TUniquePtr<FWorldPartitionActorDesc> CreateActorDesc(const AActor* Actor);
-
 	void LoadEditorCells(const FBox& Box);
 	void UnloadEditorCells(const FBox& Box);
 	bool RefreshLoadedEditorCells();
@@ -186,9 +178,6 @@ public:
 #if WITH_EDITOR
 	TMap<FGuid, TUniquePtr<FWorldPartitionActorDesc>> Actors;
 	
-	static TUniquePtr<FWorldPartitionActorDescFactory> DefaultActorDescFactory;
-	static TMap<FName, FWorldPartitionActorDescFactory*> ActorDescFactories;
-
 	DECLARE_EVENT_TwoParams(UWorldPartition, FWorldPartitionActorRegisteredEvent, AActor&, bool);
 	FWorldPartitionActorRegisteredEvent OnActorRegisteredEvent;
 
diff --git a/Engine/Source/Runtime/Engine/Public/WorldPartition/WorldPartitionActorDesc.h b/Engine/Source/Runtime/Engine/Public/WorldPartition/WorldPartitionActorDesc.h
--- a/Engine/Source/Runtime/Engine/Public/WorldPartition/WorldPartitionActorDesc.h
+++ b/Engine/Source/Runtime/Engine/Public/WorldPartition/WorldPartitionActorDesc.h
@@ -8,7 +8,6 @@
 #include "GameFramework/Actor.h"
 #include "Containers/Set.h"
 #include "Misc/Guid.h"
-#include "WorldPartition/WorldPartitionActorDescFactory.h"
 
 #if WITH_EDITOR
 // Struct used to create actor descriptor
@@ -27,7 +26,7 @@ struct FWorldPartitionActorDescInitData
 class ENGINE_API FWorldPartitionActorDesc 
 {
 #if WITH_EDITOR
-	friend class FWorldPartitionActorDescFactory;
+	friend class AActor;
 	friend class UWorldPartition;
 
 public:
diff --git a/Engine/Source/Runtime/Engine/Public/WorldPartition/WorldPartitionActorDescFactory.h b/Engine/Source/Runtime/Engine/Public/WorldPartition/WorldPartitionActorDescFactory.h
deleted file mode 100644
diff --git a/Engine/Source/Runtime/Engine/Public/WorldPartition/WorldPartitionSubsystem.h b/Engine/Source/Runtime/Engine/Public/WorldPartition/WorldPartitionSubsystem.h
--- a/Engine/Source/Runtime/Engine/Public/WorldPartition/WorldPartitionSubsystem.h
+++ b/Engine/Source/Runtime/Engine/Public/WorldPartition/WorldPartitionSubsystem.h
@@ -9,7 +9,6 @@
 
 class UWorldPartition;
 class UWorldPartitionEditorCell;
-class FWorldPartitionActorDescFactory;
 class FWorldPartitionActorDesc;
 
 /**
@@ -45,7 +44,6 @@ public:
 #if WITH_EDITOR
 	void ForEachIntersectingActorDesc(const FBox& Box, TSubclassOf<AActor> ActorClass, TFunctionRef<bool(const FWorldPartitionActorDesc*)> Predicate) const;
 	void ForEachActorDesc(TSubclassOf<AActor> ActorClass, TFunctionRef<bool(const FWorldPartitionActorDesc*)> Predicate) const;
-	void RegisterActorDescFactory(TSubclassOf<AActor> Class, FWorldPartitionActorDescFactory* Factory);
 #endif
 	void ToggleDrawRuntimeHash2D();
 
diff --git a/Engine/Source/Runtime/Landscape/Classes/LandscapeProxy.h b/Engine/Source/Runtime/Landscape/Classes/LandscapeProxy.h
--- a/Engine/Source/Runtime/Landscape/Classes/LandscapeProxy.h
+++ b/Engine/Source/Runtime/Landscape/Classes/LandscapeProxy.h
@@ -863,6 +863,7 @@ public:
 	//~ End AActor Interface
 	virtual uint32 GetDefaultGridSize(UWorld* InWorld) const override { return 1; }
 	virtual FGuid GetGridGuid() const override { return LandscapeGuid; }
+	virtual TUniquePtr<class FWorldPartitionActorDesc> CreateClassActorDesc() const override;
 #endif	//WITH_EDITOR
 
 	virtual FGuid GetLandscapeGuid() const override { return LandscapeGuid; }
diff --git a/Engine/Source/Runtime/Landscape/Private/LandscapeEdit.cpp b/Engine/Source/Runtime/Landscape/Private/LandscapeEdit.cpp
--- a/Engine/Source/Runtime/Landscape/Private/LandscapeEdit.cpp
+++ b/Engine/Source/Runtime/Landscape/Private/LandscapeEdit.cpp
@@ -58,6 +58,8 @@ LandscapeEdit.cpp: Landscape editing
 #include "ScopedTransaction.h"
 #include "Editor.h"
 #include "PhysicalMaterials/PhysicalMaterial.h"
+#include "WorldPartition/WorldPartitionActorDesc.h"
+#include "WorldPartition/Landscape/LandscapeActorDesc.h"
 #endif
 #include "Algo/Count.h"
 #include "Serialization/MemoryWriter.h"
@@ -632,6 +634,11 @@ void ULandscapeComponent::PostEditUndo()
 	}
 }
 
+TUniquePtr<FWorldPartitionActorDesc> ALandscapeProxy::CreateClassActorDesc() const
+{
+	return TUniquePtr<FWorldPartitionActorDesc>(new FLandscapeActorDesc());
+}
+
 void ALandscapeProxy::FixupWeightmaps()
 {
 	WeightmapUsageMap.Empty();
diff --git a/Engine/Source/Runtime/Landscape/Private/LandscapeSubsystem.cpp b/Engine/Source/Runtime/Landscape/Private/LandscapeSubsystem.cpp
--- a/Engine/Source/Runtime/Landscape/Private/LandscapeSubsystem.cpp
+++ b/Engine/Source/Runtime/Landscape/Private/LandscapeSubsystem.cpp
@@ -49,14 +49,6 @@ void ULandscapeSubsystem::Initialize(FSubsystemCollectionBase& Collection)
 {
 	Super::Initialize(Collection);
 
-#if WITH_EDITOR
-	UWorldPartitionSubsystem* WorldPartitionSubsystem = Collection.InitializeDependency<UWorldPartitionSubsystem>();
-	if (WorldPartitionSubsystem)
-	{
-		WorldPartitionSubsystem->RegisterActorDescFactory(ALandscapeProxy::StaticClass(), &LandscapeActorDescFactory);
-	}
-#endif
-
 #if WITH_EDITOR
 	GrassMapsBuilder = new FLandscapeGrassMapsBuilder(GetWorld());
 	BakedGITextureBuilder = new FLandscapeBakedGITextureBuilder(GetWorld());
diff --git a/Engine/Source/Runtime/Landscape/Public/LandscapeSubsystem.h b/Engine/Source/Runtime/Landscape/Public/LandscapeSubsystem.h
--- a/Engine/Source/Runtime/Landscape/Public/LandscapeSubsystem.h
+++ b/Engine/Source/Runtime/Landscape/Public/LandscapeSubsystem.h
@@ -6,7 +6,6 @@
 #include "UObject/ObjectMacros.h"
 #include "Subsystems/WorldSubsystem.h"
 #include "Tickable.h"
-#include "WorldPartition/Landscape/LandscapeActorDescFactory.h"
 #include "LandscapeSubsystem.generated.h"
 
 class ALandscapeProxy;
@@ -54,6 +53,5 @@ private:
 #if WITH_EDITOR
 	class FLandscapeGrassMapsBuilder* GrassMapsBuilder;
 	class FLandscapeBakedGITextureBuilder* BakedGITextureBuilder;
-	FLandscapeActorDescFactory LandscapeActorDescFactory;
 #endif
 };
