WorldPartition: LandscapeSplineActorDesc

diff --git a/Engine/Source/Runtime/Core/Public/UObject/UE5MainStreamObjectVersion.h b/Engine/Source/Runtime/Core/Public/UObject/UE5MainStreamObjectVersion.h
--- a/Engine/Source/Runtime/Core/Public/UObject/UE5MainStreamObjectVersion.h
+++ b/Engine/Source/Runtime/Core/Public/UObject/UE5MainStreamObjectVersion.h
@@ -55,6 +55,9 @@ struct CORE_API FUE5MainStreamObjectVersion
 		// Nanite data is now transient in Geometry Collection similar to how RenderData is transient in StaticMesh.
 		GeometryCollectionNaniteTransient,
 
+		// Added FLandscapeSplineActorDesc
+		AddedLandscapeSplineActorDesc,
+
 		// -----<new versions can be added above this line>-------------------------------------------------
 		VersionPlusOne,
 		LatestVersion = VersionPlusOne - 1
diff --git a/Engine/Source/Runtime/Engine/Private/WorldPartition/Landscape/LandscapeSplineActorDesc.cpp b/Engine/Source/Runtime/Engine/Private/WorldPartition/Landscape/LandscapeSplineActorDesc.cpp
new file mode 100644
--- /dev/null
+++ b/Engine/Source/Runtime/Engine/Private/WorldPartition/Landscape/LandscapeSplineActorDesc.cpp
@@ -0,0 +1,31 @@
+// Copyright Epic Games, Inc. All Rights Reserved.
+
+#include "WorldPartition/Landscape/LandscapeSplineActorDesc.h"
+
+#if WITH_EDITOR
+
+#include "LandscapeSplineActor.h"
+#include "UObject/UE5MainStreamObjectVersion.h"
+
+void FLandscapeSplineActorDesc::Init(const AActor* InActor)
+{
+	FWorldPartitionActorDesc::Init(InActor);
+
+	const ALandscapeSplineActor* LandscapeSplineActor = CastChecked<ALandscapeSplineActor>(InActor);
+	check(LandscapeSplineActor);
+	LandscapeGuid = LandscapeSplineActor->GetLandscapeGuid();
+}
+
+void FLandscapeSplineActorDesc::Serialize(FArchive& Ar)
+{
+	Ar.UsingCustomVersion(FUE5MainStreamObjectVersion::GUID);
+
+	FWorldPartitionActorDesc::Serialize(Ar);
+
+	if (Ar.CustomVer(FUE5MainStreamObjectVersion::GUID) >= FUE5MainStreamObjectVersion::AddedLandscapeSplineActorDesc)
+	{
+		Ar << LandscapeGuid;
+	}
+}
+
+#endif
diff --git a/Engine/Source/Runtime/Engine/Public/WorldPartition/Landscape/LandscapeSplineActorDesc.h b/Engine/Source/Runtime/Engine/Public/WorldPartition/Landscape/LandscapeSplineActorDesc.h
new file mode 100644
--- /dev/null
+++ b/Engine/Source/Runtime/Engine/Public/WorldPartition/Landscape/LandscapeSplineActorDesc.h
@@ -0,0 +1,21 @@
+// Copyright Epic Games, Inc. All Rights Reserved.
+
+#pragma once
+
+#include "CoreMinimal.h"
+#include "UObject/ObjectMacros.h"
+#include "WorldPartition/WorldPartitionActorDesc.h"
+
+/**
+ * ActorDesc for LandscapeSplineActors.
+ */
+class ENGINE_API FLandscapeSplineActorDesc : public FWorldPartitionActorDesc
+{
+#if WITH_EDITOR
+protected:
+	virtual void Init(const AActor* InActor) override;
+	virtual void Serialize(FArchive& Ar) override;
+
+	FGuid LandscapeGuid;
+#endif
+};
diff --git a/Engine/Source/Runtime/Landscape/Classes/LandscapeSplineActor.h b/Engine/Source/Runtime/Landscape/Classes/LandscapeSplineActor.h
--- a/Engine/Source/Runtime/Landscape/Classes/LandscapeSplineActor.h
+++ b/Engine/Source/Runtime/Landscape/Classes/LandscapeSplineActor.h
@@ -26,6 +26,7 @@ public:
 		
 #if WITH_EDITOR
 	virtual EActorGridPlacement GetDefaultGridPlacement() const override { return EActorGridPlacement::Bounds; }
+	virtual TUniquePtr<class FWorldPartitionActorDesc> CreateClassActorDesc() const override;
 
 	void GetSharedProperties(ULandscapeInfo* InLandscapeInfo);
 
diff --git a/Engine/Source/Runtime/Landscape/Private/LandscapeSplineActor.cpp b/Engine/Source/Runtime/Landscape/Private/LandscapeSplineActor.cpp
--- a/Engine/Source/Runtime/Landscape/Private/LandscapeSplineActor.cpp
+++ b/Engine/Source/Runtime/Landscape/Private/LandscapeSplineActor.cpp
@@ -6,6 +6,11 @@
 #include "Landscape.h"
 #include "LandscapeSplinesComponent.h"
 
+#if WITH_EDITOR
+#include "WorldPartition/WorldPartitionActorDesc.h"
+#include "WorldPartition/Landscape/LandscapeSplineActorDesc.h"
+#endif
+
 ALandscapeSplineActor::ALandscapeSplineActor(const FObjectInitializer& ObjectInitializer)
 	: Super(ObjectInitializer)
 {
@@ -31,6 +36,11 @@ ULandscapeInfo* ALandscapeSplineActor::GetLandscapeInfo() const
 }
 
 #if WITH_EDITOR
+TUniquePtr<FWorldPartitionActorDesc> ALandscapeSplineActor::CreateClassActorDesc() const
+{
+	return TUniquePtr<FWorldPartitionActorDesc>(new FLandscapeSplineActorDesc());
+}
+
 void ALandscapeSplineActor::GetSharedProperties(ULandscapeInfo* InLandscapeInfo)
 {
 	Modify();
