DataLayer: DataLayerSubsystem now loads DataLayerEditorModule at initialization (in editor) to avoid relying on DataLayer tab to be spawned by Level editor. We need to register custom PropertyTypeLayout for actor details to work properly.

diff --git a/Engine/Source/Editor/DataLayerEditor/Private/DataLayer/DataLayerEditorSubsystem.cpp b/Engine/Source/Editor/DataLayerEditor/Private/DataLayer/DataLayerEditorSubsystem.cpp
--- a/Engine/Source/Editor/DataLayerEditor/Private/DataLayer/DataLayerEditorSubsystem.cpp
+++ b/Engine/Source/Editor/DataLayerEditor/Private/DataLayer/DataLayerEditorSubsystem.cpp
@@ -146,10 +146,6 @@ void FDataLayersBroadcast::OnLevelViewportClientListChanged()
 
 UDataLayerEditorSubsystem* UDataLayerEditorSubsystem::Get()
 {
-	if (!FModuleManager::Get().IsModuleLoaded("DataLayerEditor"))
-	{
-		FModuleManager::LoadModuleChecked<FDataLayerEditorModule>("DataLayerEditor");
-	}
 	return GEditor ? GEditor->GetEditorSubsystem<UDataLayerEditorSubsystem>() : nullptr;
 }
 
diff --git a/Engine/Source/Editor/DataLayerEditor/Public/DataLayerEditorModule.h b/Engine/Source/Editor/DataLayerEditor/Public/DataLayerEditorModule.h
--- a/Engine/Source/Editor/DataLayerEditor/Public/DataLayerEditorModule.h
+++ b/Engine/Source/Editor/DataLayerEditor/Public/DataLayerEditorModule.h
@@ -3,16 +3,16 @@
 #pragma once
 
 #include "CoreMinimal.h"
-#include "Modules/ModuleInterface.h"
 #include "Framework/Commands/UICommandList.h"
 #include "Framework/MultiBox/MultiBoxExtender.h"
+#include "WorldPartition/DataLayer/IDataLayerEditorModule.h"
 
 class AActor;
 
 /**
  * The module holding all of the UI related pieces for DataLayers
  */
-class FDataLayerEditorModule : public IModuleInterface
+class FDataLayerEditorModule : public IDataLayerEditorModule
 {
 public:
 
diff --git a/Engine/Source/Runtime/Engine/Private/WorldPartition/DataLayer/DataLayerSubsystem.cpp b/Engine/Source/Runtime/Engine/Private/WorldPartition/DataLayer/DataLayerSubsystem.cpp
--- a/Engine/Source/Runtime/Engine/Private/WorldPartition/DataLayer/DataLayerSubsystem.cpp
+++ b/Engine/Source/Runtime/Engine/Private/WorldPartition/DataLayer/DataLayerSubsystem.cpp
@@ -5,6 +5,10 @@
 #include "WorldPartition/DataLayer/DataLayer.h"
 #include "Engine/Engine.h"
 #include "Engine/World.h"
+#if WITH_EDITOR
+#include "Editor.h"
+#include "WorldPartition/DataLayer/IDataLayerEditorModule.h"
+#endif
 
 UDataLayerSubsystem::UDataLayerSubsystem()
 {}
@@ -24,6 +28,18 @@ bool UDataLayerSubsystem::ShouldCreateSubsystem(UObject* Outer) const
 	return false;
 }
 
+void UDataLayerSubsystem::Initialize(FSubsystemCollectionBase& Collection)
+{
+	Super::Initialize(Collection);
+
+#if WITH_EDITOR
+	if (GEditor)
+	{
+		FModuleManager::LoadModuleChecked<IDataLayerEditorModule>("DataLayerEditor");
+	}
+#endif
+}
+
 UDataLayer* UDataLayerSubsystem::GetDataLayerFromLabel(const FName& InDataLayerLabel) const
 {
 	const AWorldDataLayers* WorldDataLayers = AWorldDataLayers::Get(GetWorld());
diff --git a/Engine/Source/Runtime/Engine/Public/WorldPartition/DataLayer/DataLayerSubsystem.h b/Engine/Source/Runtime/Engine/Public/WorldPartition/DataLayer/DataLayerSubsystem.h
--- a/Engine/Source/Runtime/Engine/Public/WorldPartition/DataLayer/DataLayerSubsystem.h
+++ b/Engine/Source/Runtime/Engine/Public/WorldPartition/DataLayer/DataLayerSubsystem.h
@@ -21,6 +21,7 @@ public:
 
 	//~ Begin USubsystem Interface.
 	virtual bool ShouldCreateSubsystem(UObject* Outer) const override;
+	virtual void Initialize(FSubsystemCollectionBase& Collection) override;
 	//~ End USubsystem Interface.
 
 	//~ Begin Blueprint callable functions
diff --git a/Engine/Source/Runtime/Engine/Public/WorldPartition/DataLayer/IDataLayerEditorModule.h b/Engine/Source/Runtime/Engine/Public/WorldPartition/DataLayer/IDataLayerEditorModule.h
new file mode 100644
--- /dev/null
+++ b/Engine/Source/Runtime/Engine/Public/WorldPartition/DataLayer/IDataLayerEditorModule.h
@@ -0,0 +1,14 @@
+// Copyright Epic Games, Inc. All Rights Reserved.
+#pragma once
+
+#include "CoreMinimal.h"
+#include "Modules/ModuleInterface.h"
+
+/**
+ * The module holding all of the UI related pieces for DataLayer management
+ */
+class IDataLayerEditorModule : public IModuleInterface
+{
+public:
+	virtual ~IDataLayerEditorModule() {}
+};
\ No newline at end of file
