WorldPartition: Fix Cook

diff --git a/Engine/Source/Runtime/Engine/Private/WorldPartition/WorldPartition.cpp b/Engine/Source/Runtime/Engine/Private/WorldPartition/WorldPartition.cpp
--- a/Engine/Source/Runtime/Engine/Private/WorldPartition/WorldPartition.cpp
+++ b/Engine/Source/Runtime/Engine/Private/WorldPartition/WorldPartition.cpp
@@ -281,8 +281,13 @@ void UWorldPartition::Initialize(UWorld* InWorld, const FTransform& InTransform)
 
 	InitState = EWorldPartitionInitState::Initialized;
 
-	UWorldPartitionSubsystem* WorldPartitionSubsystem = World->GetSubsystem<UWorldPartitionSubsystem>();
-	WorldPartitionSubsystem->RegisterWorldPartition(this);
+#if WITH_EDITOR
+	if (!IsRunningCookCommandlet())
+#endif
+	{
+		UWorldPartitionSubsystem* WorldPartitionSubsystem = World->GetSubsystem<UWorldPartitionSubsystem>();
+		WorldPartitionSubsystem->RegisterWorldPartition(this);
+	}
 
 #if WITH_EDITOR
 	if (World->IsPlayInEditor())
@@ -343,8 +348,13 @@ void UWorldPartition::Uninitialize()
 		EditorHash = nullptr;
 #endif		
 
-		UWorldPartitionSubsystem* WorldPartitionSubsystem = World->GetSubsystem<UWorldPartitionSubsystem>();
-		WorldPartitionSubsystem->UnregisterWorldPartition(this);
+#if WITH_EDITOR
+		if (!IsRunningCookCommandlet())
+#endif
+		{
+			UWorldPartitionSubsystem* WorldPartitionSubsystem = World->GetSubsystem<UWorldPartitionSubsystem>();
+			WorldPartitionSubsystem->UnregisterWorldPartition(this);
+		}
 
 		World = nullptr;
 
