Fix CIS build skeletalmesh public member deprecation

diff --git a/Engine/Source/Runtime/Engine/Private/SkeletalMesh.cpp b/Engine/Source/Runtime/Engine/Private/SkeletalMesh.cpp
--- a/Engine/Source/Runtime/Engine/Private/SkeletalMesh.cpp
+++ b/Engine/Source/Runtime/Engine/Private/SkeletalMesh.cpp
@@ -436,7 +436,9 @@ USkeletalMesh::USkeletalMesh(FVTableHelper& Helper)
 {
 }
 
+PRAGMA_DISABLE_DEPRECATION_WARNINGS
 USkeletalMesh::~USkeletalMesh() = default;
+PRAGMA_ENABLE_DEPRECATION_WARNINGS
 
 void USkeletalMesh::PostInitProperties()
 {
