Fixed static analysis warning around Context.Resource dereferencing

diff --git a/Engine/Source/Runtime/Engine/Private/Streaming/Texture2DArrayStreaming.cpp b/Engine/Source/Runtime/Engine/Private/Streaming/Texture2DArrayStreaming.cpp
--- a/Engine/Source/Runtime/Engine/Private/Streaming/Texture2DArrayStreaming.cpp
+++ b/Engine/Source/Runtime/Engine/Private/Streaming/Texture2DArrayStreaming.cpp
@@ -99,7 +99,7 @@ bool FTexture2DArrayMipAllocator_Reallocate::AllocateMips(const FTextureUpdateCo
 
 bool FTexture2DArrayMipAllocator_Reallocate::FinalizeMips(const FTextureUpdateContext& Context, const FTextureUpdateSyncOptions& SyncOptions)
 {
-	const uint32 NumSlices = Context.Resource ? Context.Resource->GetSizeZ() : 0;
+	const uint32 NumSlices = Context.Resource->GetSizeZ();
 
 	// Create new Texture.
 	{
