Add missing includes breaking CIS non-Unity Android builds

diff --git a/Engine/Source/Runtime/Core/Private/Async/Fundamental/Scheduler.cpp b/Engine/Source/Runtime/Core/Private/Async/Fundamental/Scheduler.cpp
--- a/Engine/Source/Runtime/Core/Private/Async/Fundamental/Scheduler.cpp
+++ b/Engine/Source/Runtime/Core/Private/Async/Fundamental/Scheduler.cpp
@@ -3,6 +3,7 @@
 #include "Async/Fundamental/Scheduler.h"
 #include "Async/Fundamental/Task.h"
 #include "Logging/LogMacros.h"
+#include "Misc/ScopeLock.h" 
 
 namespace LowLevelTasks
 {
diff --git a/Engine/Source/Runtime/Core/Public/Async/Fundamental/Scheduler.h b/Engine/Source/Runtime/Core/Public/Async/Fundamental/Scheduler.h
--- a/Engine/Source/Runtime/Core/Public/Async/Fundamental/Scheduler.h
+++ b/Engine/Source/Runtime/Core/Public/Async/Fundamental/Scheduler.h
@@ -3,6 +3,8 @@
 #pragma once 
 #include "Async/Fundamental/Task.h"
 #include "HAL/Thread.h"
+#include "HAL/Event.h"
+#include "HAL/PlatformProcess.h"
 #include "LocalQueue.h"
 #include "Containers/Array.h"
 #include "Templates/UniquePtr.h"
diff --git a/Engine/Source/Runtime/Engine/Private/Animation/BlendProfile.cpp b/Engine/Source/Runtime/Engine/Private/Animation/BlendProfile.cpp
--- a/Engine/Source/Runtime/Engine/Private/Animation/BlendProfile.cpp
+++ b/Engine/Source/Runtime/Engine/Private/Animation/BlendProfile.cpp
@@ -1,6 +1,7 @@
 // Copyright Epic Games, Inc. All Rights Reserved.
 
 #include "Animation/BlendProfile.h"
+#include "AlphaBlend.h"
 
 UBlendProfile::UBlendProfile()
 	: OwningSkeleton(nullptr)
diff --git a/Engine/Source/Runtime/Engine/Private/Streaming/Texture2DArrayStreaming.cpp b/Engine/Source/Runtime/Engine/Private/Streaming/Texture2DArrayStreaming.cpp
--- a/Engine/Source/Runtime/Engine/Private/Streaming/Texture2DArrayStreaming.cpp
+++ b/Engine/Source/Runtime/Engine/Private/Streaming/Texture2DArrayStreaming.cpp
@@ -5,6 +5,7 @@ Texture2DArrayStreaming.cpp: Helpers to stream in and out texture 2D array LODs.
 =============================================================================*/
 
 #include "Streaming/Texture2DArrayStreaming.h"
+#include "RenderUtils.h"
 
 //*****************************************************************************
 //***************************** Global Definitions ****************************
diff --git a/Engine/Source/Runtime/Engine/Private/Streaming/Texture2DMipDataProvider_IO.h b/Engine/Source/Runtime/Engine/Private/Streaming/Texture2DMipDataProvider_IO.h
--- a/Engine/Source/Runtime/Engine/Private/Streaming/Texture2DMipDataProvider_IO.h
+++ b/Engine/Source/Runtime/Engine/Private/Streaming/Texture2DMipDataProvider_IO.h
@@ -9,6 +9,7 @@ Texture2DMipDataProvider_IO.h : Implementation of FTextureMipDataProvider using
 #include "CoreMinimal.h"
 #include "Streaming/TextureMipDataProvider.h"
 #include "Async/AsyncFileHandle.h"
+#include "Serialization/BulkData2.h"
 
 /**
 * FTexture2DMipDataProvider_IO implements FTextureMipAllocator using file IO (cooked content).
diff --git a/Engine/Source/Runtime/Engine/Public/Rendering/StreamableTextureResource.h b/Engine/Source/Runtime/Engine/Public/Rendering/StreamableTextureResource.h
--- a/Engine/Source/Runtime/Engine/Public/Rendering/StreamableTextureResource.h
+++ b/Engine/Source/Runtime/Engine/Public/Rendering/StreamableTextureResource.h
@@ -9,6 +9,7 @@
 #include "CoreMinimal.h"
 #include "TextureResource.h"
 #include "Streaming/StreamableRenderResourceState.h"
+#include "Engine/Texture.h"
 
 bool CanCreateWithPartiallyResidentMips(uint32 TexCreateFlags);
 
