VT respects LOD settings on the texture object
LOD group settings are still ignored, since those values are typically tuned for regular textures, and not appropriate for VT

diff --git a/Engine/Source/Runtime/Engine/Private/TextureLODSettings.cpp b/Engine/Source/Runtime/Engine/Private/TextureLODSettings.cpp
--- a/Engine/Source/Runtime/Engine/Private/TextureLODSettings.cpp
+++ b/Engine/Source/Runtime/Engine/Private/TextureLODSettings.cpp
@@ -4,6 +4,7 @@
 #include "Engine/Texture2D.h"
 #include "Interfaces/ITargetPlatform.h"
 #include "Interfaces/ITargetPlatformManagerModule.h"
+#include "VT/VirtualTextureBuiltData.h"
 #include "PlatformInfo.h"
 
 int32 GUITextureLODBias = 0;
@@ -142,14 +143,9 @@ int32 UTextureLODSettings::CalculateLODBias(int32 Width, int32 Height, int32 Max
 	const FTextureLODGroup& LODGroupInfo = TextureLODGroups[LODGroup];
 
 	// Test to see if we have no mip generation as in which case the LOD bias will be ignored
+	// VTs don't respect NoMipmaps, mips are required for VTs
 	const TextureMipGenSettings FinalMipGenSetting = (InMipGenSetting == TMGS_FromTextureGroup) ? (TextureMipGenSettings)LODGroupInfo.MipGenSettings : InMipGenSetting;
-	if ( FinalMipGenSetting == TMGS_NoMipmaps )
-	{
-		return 0;
-	}
-
-	// Ignore LODBias for virtual textures, as these are not restricted by any GPU max size
-	if (bVirtualTexture)
+	if ( !bVirtualTexture && FinalMipGenSetting == TMGS_NoMipmaps )
 	{
 		return 0;
 	}
@@ -163,12 +159,22 @@ int32 UTextureLODSettings::CalculateLODBias(int32 Width, int32 Height, int32 Max
 	int32 TextureMaxLOD	= FMath::CeilLogTwo( FMath::Max( Width, Height ) );
 
 	// Calculate LOD bias.
-	int32 UsedLODBias	= NumCinematicMipLevels;
+	int32 UsedLODBias = 0;
+	if (!bVirtualTexture)
+	{
+		// VT doesn't support cinematic mips
+		UsedLODBias += NumCinematicMipLevels;
+	}
 	if (!FPlatformProperties::RequiresCookedData())
 	{
 		// When cooking, LODBias and LODGroupInfo.LODBias are taken into account to strip the top mips.
 		// Considering them again here would apply them twice.
-		UsedLODBias	+= LODBias + LODGroupInfo.LODBias;
+		UsedLODBias	+= LODBias;
+		if (!bVirtualTexture)
+		{
+			// Don't include LOD group's bias for VT, could include a separate bias that applies only to VT if needed
+			UsedLODBias += LODGroupInfo.LODBias;
+		}
 	}
 	
 	if (LODGroup == TEXTUREGROUP_UI)
@@ -176,8 +182,10 @@ int32 UTextureLODSettings::CalculateLODBias(int32 Width, int32 Height, int32 Max
 		UsedLODBias += GUITextureLODBias;
 	}
 
-	int32 MinLOD		= LODGroupInfo.MinLODMipCount;
-	int32 MaxLOD		= LODGroupInfo.MaxLODMipCount;
+	// Don't apply LOD group mip clamping to VT, since VT max size isn't constrained by any GPU capabilities
+	const int32 MinLOD	= bVirtualTexture ? 0 : LODGroupInfo.MinLODMipCount;
+	const int32 MaxLOD	= bVirtualTexture ? VIRTUALTEXTURE_DATA_MAXMIPS : LODGroupInfo.MaxLODMipCount;
+
 	int32 WantedMaxLOD	= FMath::Clamp( TextureMaxLOD - UsedLODBias, MinLOD, MaxLOD );
 	WantedMaxLOD		= FMath::Clamp( WantedMaxLOD, 0, TextureMaxLOD );
 	UsedLODBias			= TextureMaxLOD - WantedMaxLOD;
