Prevent garbage collection from trying to finish async static mesh compilation, which might lead to new objects being created in the process

diff --git a/Engine/Source/Runtime/Engine/Private/StaticMesh.cpp b/Engine/Source/Runtime/Engine/Private/StaticMesh.cpp
--- a/Engine/Source/Runtime/Engine/Private/StaticMesh.cpp
+++ b/Engine/Source/Runtime/Engine/Private/StaticMesh.cpp
@@ -6234,7 +6234,10 @@ bool UStaticMesh::DoesMipDataExist(const int32 MipIndex) const
 
 bool UStaticMesh::HasPendingRenderResourceInitialization() const
 {
-	return GetRenderData() && !GetRenderData()->bReadyForStreaming;
+	// Verify we're not compiling before accessing the renderdata to avoid forcing the compilation
+	// to finish during garbage collection. If we're still compiling, the render data has not
+	// yet been created, hence it is not possible we're actively streaming anything from it...
+	return !IsCompiling() && GetRenderData() && !GetRenderData()->bReadyForStreaming;
 }
 
 bool UStaticMesh::StreamOut(int32 NewMipCount)
