Fixed layered blend mask not working w/ multipe input poses
#rnx

diff --git a/Engine/Source/Runtime/Engine/Private/Animation/AnimationRuntime.cpp b/Engine/Source/Runtime/Engine/Private/Animation/AnimationRuntime.cpp
--- a/Engine/Source/Runtime/Engine/Private/Animation/AnimationRuntime.cpp
+++ b/Engine/Source/Runtime/Engine/Private/Animation/AnimationRuntime.cpp
@@ -1848,7 +1848,11 @@ void FAnimationRuntime::CreateMaskWeights(TArray<FPerBoneBlendWeight>& BoneBlend
 				int32 BoneIndex = BlendMask->ProfileEntries[EntryIndex].BoneReference.BoneIndex;
 				if (BoneBlendWeights.IsValidIndex(BoneIndex))
 				{
-					BoneBlendWeights[BoneIndex].BlendWeight = BlendMask->ProfileEntries[EntryIndex].BlendScale;
+					// Match the BoneBlendWeight's input pose with BlendMasks's MaskIndex and use the blend mask's weight
+					FPerBoneBlendWeight& BoneBlendWeight = BoneBlendWeights[BoneIndex];
+
+					BoneBlendWeight.SourceIndex = MaskIndex;
+					BoneBlendWeight.BlendWeight = BlendMask->ProfileEntries[EntryIndex].BlendScale;
 				}
 			}
 		}
