Prototype Wrapped Object Pointers (ObjectHandle/ObjectPtr)

Avoiding use of implicit conversion to pointer reference on TObjectPtr (making it either explicit or using a reference to a TObjectPtr where that is an easy option).

diff --git a/Engine/Plugins/Media/MediaPlayerEditor/Source/MediaPlayerEditor/Private/Customizations/PlatformMediaSourceCustomization.cpp b/Engine/Plugins/Media/MediaPlayerEditor/Source/MediaPlayerEditor/Private/Customizations/PlatformMediaSourceCustomization.cpp
--- a/Engine/Plugins/Media/MediaPlayerEditor/Source/MediaPlayerEditor/Private/Customizations/PlatformMediaSourceCustomization.cpp
+++ b/Engine/Plugins/Media/MediaPlayerEditor/Source/MediaPlayerEditor/Private/Customizations/PlatformMediaSourceCustomization.cpp
@@ -142,7 +142,7 @@ void FPlatformMediaSourceCustomization::SetPlatformMediaSourcesValue(FString Pla
 
 	for (auto Object : OuterObjects)
 	{
-		UMediaSource*& OldMediaSource = Cast<UPlatformMediaSource>(Object)->PlatformMediaSources.FindOrAdd(PlatformName);;
+		UE_TRANSITIONAL_OBJECT_PTR(UMediaSource)& OldMediaSource = Cast<UPlatformMediaSource>(Object)->PlatformMediaSources.FindOrAdd(PlatformName);;
 
 		if (OldMediaSource != MediaSource)
 		{
@@ -165,7 +165,7 @@ void FPlatformMediaSourceCustomization::HandleMediaSourceEntryBoxChanged(const F
 
 	for (auto Object : OuterObjects)
 	{
-		UMediaSource*& OldMediaSource = Cast<UPlatformMediaSource>(Object)->PlatformMediaSources.FindOrAdd(PlatformName);;
+		UE_TRANSITIONAL_OBJECT_PTR(UMediaSource)& OldMediaSource = Cast<UPlatformMediaSource>(Object)->PlatformMediaSources.FindOrAdd(PlatformName);;
 
 		if (OldMediaSource != AssetData.GetAsset())
 		{
diff --git a/Engine/Plugins/Runtime/ReplicationGraph/Source/Private/ReplicationGraph.cpp b/Engine/Plugins/Runtime/ReplicationGraph/Source/Private/ReplicationGraph.cpp
--- a/Engine/Plugins/Runtime/ReplicationGraph/Source/Private/ReplicationGraph.cpp
+++ b/Engine/Plugins/Runtime/ReplicationGraph/Source/Private/ReplicationGraph.cpp
@@ -5647,8 +5647,8 @@ void UReplicationGraphNode_AlwaysRelevant_ForConnection::GatherActorListsForConn
 
 		check(LastData != nullptr);
 
-		UpdateActor(CurViewer.InViewer, LastData->LastViewer);
-		UpdateActor(CurViewer.ViewTarget, LastData->LastViewTarget);
+		UpdateActor(CurViewer.InViewer, static_cast<AActor*&>(LastData->LastViewer));
+		UpdateActor(CurViewer.ViewTarget, static_cast<AActor*&>(LastData->LastViewTarget));
 	}
 
 	// Remove excess
diff --git a/Engine/Source/Editor/UMGEditor/Private/WidgetBlueprintFactory.cpp b/Engine/Source/Editor/UMGEditor/Private/WidgetBlueprintFactory.cpp
--- a/Engine/Source/Editor/UMGEditor/Private/WidgetBlueprintFactory.cpp
+++ b/Engine/Source/Editor/UMGEditor/Private/WidgetBlueprintFactory.cpp
@@ -81,7 +81,7 @@ bool UWidgetBlueprintFactory::ConfigureProperties()
 		Filter->AllowedChildrenOfClasses.Add(UPanelWidget::StaticClass());
 
 		const FText TitleText = LOCTEXT("CreateWidgetBlueprint", "Pick Root Widget for New Widget Blueprint");
-		return SClassPickerDialog::PickClass(TitleText, Options, RootWidgetClass, UPanelWidget::StaticClass());
+		return SClassPickerDialog::PickClass(TitleText, Options, static_cast<UClass*&>(RootWidgetClass), UPanelWidget::StaticClass());
 
 	}
 	return true;
diff --git a/Engine/Source/Editor/UnrealEd/Private/Commandlets/WorldPartitionConvertCommandlet.cpp b/Engine/Source/Editor/UnrealEd/Private/Commandlets/WorldPartitionConvertCommandlet.cpp
--- a/Engine/Source/Editor/UnrealEd/Private/Commandlets/WorldPartitionConvertCommandlet.cpp
+++ b/Engine/Source/Editor/UnrealEd/Private/Commandlets/WorldPartitionConvertCommandlet.cpp
@@ -671,7 +671,7 @@ void UWorldPartitionConvertCommandlet::CreateWorldMiniMapTexture(UWorld* World)
 		return;
 	}
 
-	FWorldPartitionMiniMapHelper::CaptureWorldMiniMapToTexture(World, WorldMiniMap, WorldMiniMap->MiniMapSize, WorldMiniMap->MiniMapTexture, WorldMiniMap->MiniMapWorldBounds);
+	FWorldPartitionMiniMapHelper::CaptureWorldMiniMapToTexture(World, WorldMiniMap, WorldMiniMap->MiniMapSize, static_cast<UTexture2D*&>(WorldMiniMap->MiniMapTexture), WorldMiniMap->MiniMapWorldBounds);
 }
 
 int32 UWorldPartitionConvertCommandlet::Main(const FString& Params)
diff --git a/Engine/Source/Editor/UnrealEd/Private/Factories/EditorFactories.cpp b/Engine/Source/Editor/UnrealEd/Private/Factories/EditorFactories.cpp
--- a/Engine/Source/Editor/UnrealEd/Private/Factories/EditorFactories.cpp
+++ b/Engine/Source/Editor/UnrealEd/Private/Factories/EditorFactories.cpp
@@ -5522,7 +5522,7 @@ EReimportResult::Type UReimportTextureFactory::Reimport( UObject* Obj )
 
 	UTexture* pTex = Cast<UTexture>(Obj);
 	
-	TGuardValue<UTexture*> OriginalTexGuardValue(pOriginalTex, pTex);
+	TGuardValue<decltype(pOriginalTex)> OriginalTexGuardValue(pOriginalTex, pTex);
 
 	const FString ResolvedSourceFilePath = pTex->AssetImportData->GetFirstFilename();
 	if (!ResolvedSourceFilePath.Len())
diff --git a/Engine/Source/Editor/UnrealEd/Private/PreviewMaterial.cpp b/Engine/Source/Editor/UnrealEd/Private/PreviewMaterial.cpp
--- a/Engine/Source/Editor/UnrealEd/Private/PreviewMaterial.cpp
+++ b/Engine/Source/Editor/UnrealEd/Private/PreviewMaterial.cpp
@@ -1041,7 +1041,7 @@ void UMaterialEditorInstanceConstant::RegenerateArrays()
 			ParameterValue.ExpressionId = Guids[ParameterIdx];
 
 			ParameterValue.ParameterValue = nullptr;
-			SourceInstance->GetTextureParameterValue(ParameterInfo, ParameterValue.ParameterValue);
+			SourceInstance->GetTextureParameterValue(ParameterInfo, static_cast<UTexture*&>(ParameterValue.ParameterValue));
 			SourceInstance->GetTextureParameterChannelNames(ParameterInfo, ParameterValue.ChannelNames);
 
 			// @todo: This is kind of slow, maybe store these in a map for lookup?
@@ -1073,7 +1073,7 @@ void UMaterialEditorInstanceConstant::RegenerateArrays()
 			ParameterValue.ExpressionId = Guids[ParameterIdx];
 
 			ParameterValue.ParameterValue = nullptr;
-			SourceInstance->GetRuntimeVirtualTextureParameterValue(ParameterInfo, ParameterValue.ParameterValue);
+			SourceInstance->GetRuntimeVirtualTextureParameterValue(ParameterInfo, static_cast<URuntimeVirtualTexture*&>(ParameterValue.ParameterValue));
 
 			// @todo: This is kind of slow, maybe store these in a map for lookup?
 			// See if this keyname exists in the source instance.
@@ -1110,7 +1110,7 @@ void UMaterialEditorInstanceConstant::RegenerateArrays()
 
 			ParameterValue.ParameterValue.FontValue = nullptr;
 			ParameterValue.ParameterValue.FontPage = 0;
-			SourceInstance->GetFontParameterValue(ParameterInfo, ParameterValue.ParameterValue.FontValue, ParameterValue.ParameterValue.FontPage);
+			SourceInstance->GetFontParameterValue(ParameterInfo, static_cast<UFont*&>(ParameterValue.ParameterValue.FontValue), ParameterValue.ParameterValue.FontPage);
 
 			// @todo: This is kind of slow, maybe store these in a map for lookup?
 			// See if this keyname exists in the source instance.
diff --git a/Engine/Source/Editor/UnrealEd/Private/WorldPartition/WorldPartitionMiniMapBuilder.cpp b/Engine/Source/Editor/UnrealEd/Private/WorldPartition/WorldPartitionMiniMapBuilder.cpp
--- a/Engine/Source/Editor/UnrealEd/Private/WorldPartition/WorldPartitionMiniMapBuilder.cpp
+++ b/Engine/Source/Editor/UnrealEd/Private/WorldPartition/WorldPartitionMiniMapBuilder.cpp
@@ -25,7 +25,7 @@ bool UWorldPartitionMiniMapBuilder::Run(UWorld* World, FPackageSourceControlHelp
 	}
 
 	WorldMiniMap->MiniMapSize = this->MiniMapSize;
-	FWorldPartitionMiniMapHelper::CaptureWorldMiniMapToTexture(World, WorldMiniMap, WorldMiniMap->MiniMapSize, WorldMiniMap->MiniMapTexture, WorldMiniMap->MiniMapWorldBounds);
+	FWorldPartitionMiniMapHelper::CaptureWorldMiniMapToTexture(World, WorldMiniMap, WorldMiniMap->MiniMapSize, static_cast<UTexture2D*&>(WorldMiniMap->MiniMapTexture), WorldMiniMap->MiniMapWorldBounds);
 
 	// Save MiniMap Package
 	auto WorldMiniMapExternalPackage = WorldMiniMap->GetExternalPackage();
diff --git a/Engine/Source/Editor/WorldPartitionEditor/Private/WorldPartition/SWorldPartitionEditorGridSpatialHash.cpp b/Engine/Source/Editor/WorldPartitionEditor/Private/WorldPartition/SWorldPartitionEditorGridSpatialHash.cpp
--- a/Engine/Source/Editor/WorldPartitionEditor/Private/WorldPartition/SWorldPartitionEditorGridSpatialHash.cpp
+++ b/Engine/Source/Editor/WorldPartitionEditor/Private/WorldPartition/SWorldPartitionEditorGridSpatialHash.cpp
@@ -68,7 +68,7 @@ FReply SWorldPartitionEditorGridSpatialHash::ReloadMiniMap()
 	}
 
 	WorldMiniMap->Modify();
-	FWorldPartitionMiniMapHelper::CaptureWorldMiniMapToTexture(World, WorldMiniMap, WorldMiniMap->MiniMapSize, WorldMiniMap->MiniMapTexture, WorldMiniMap->MiniMapWorldBounds);
+	FWorldPartitionMiniMapHelper::CaptureWorldMiniMapToTexture(World, WorldMiniMap, WorldMiniMap->MiniMapSize, static_cast<UTexture2D*&>(WorldMiniMap->MiniMapTexture), WorldMiniMap->MiniMapWorldBounds);
 
 	UpdateWorldMiniMapDetails();
 
diff --git a/Engine/Source/Runtime/CoreUObject/Public/UObject/ObjectPtr.h b/Engine/Source/Runtime/CoreUObject/Public/UObject/ObjectPtr.h
--- a/Engine/Source/Runtime/CoreUObject/Public/UObject/ObjectPtr.h
+++ b/Engine/Source/Runtime/CoreUObject/Public/UObject/ObjectPtr.h
@@ -16,6 +16,14 @@
 	#define UE_OBJPTR_DEPRECATED(Version, Message) 
 #endif
 
+/** 
+ * Wrapper macro for use in places where code needs to allow for a pointer type that could be a TObjectPtr<T> or a raw object pointer during a transitional period.
+ * The coding standard disallows general use of the auto keyword, but in wrapping it in this macro, we have a record
+ * of the explicit type meant to be used, and an avenue to go back and change these instances to an explicit
+ * TObjectPtr<T> after the transition is complete and the type won't be toggling back and forth anymore.
+ */
+#define UE_TRANSITIONAL_OBJECT_PTR(Type) auto 
+
 template <typename T>
 struct TObjectPtr;
 
diff --git a/Engine/Source/Runtime/Engine/Private/Animation/AssetMappingTable.cpp b/Engine/Source/Runtime/Engine/Private/Animation/AssetMappingTable.cpp
--- a/Engine/Source/Runtime/Engine/Private/Animation/AssetMappingTable.cpp
+++ b/Engine/Source/Runtime/Engine/Private/Animation/AssetMappingTable.cpp
@@ -191,7 +191,7 @@ void UAssetMappingTable::ReplaceReferredAnimations(const TMap<UAnimationAsset*,
 {
 	for (FAssetMapping& AssetMapping : MappedAssets)
 	{
-		UAnimationAsset*& AnimAsset = AssetMapping.SourceAsset;
+		UE_TRANSITIONAL_OBJECT_PTR(UAnimationAsset)& AnimAsset = AssetMapping.SourceAsset;
 		if (AnimAsset)
 		{
 			// now fix everythign else
diff --git a/Engine/Source/Runtime/Engine/Private/Components/BoxComponent.cpp b/Engine/Source/Runtime/Engine/Private/Components/BoxComponent.cpp
--- a/Engine/Source/Runtime/Engine/Private/Components/BoxComponent.cpp
+++ b/Engine/Source/Runtime/Engine/Private/Components/BoxComponent.cpp
@@ -41,8 +41,8 @@ void UBoxComponent::SetBoxExtent(FVector NewBoxExtent, bool bUpdateOverlaps)
 }
 
 
-template <EShapeBodySetupHelper UpdateBodySetupAction>
-bool InvalidateOrUpdateBoxBodySetup(UBodySetup*& ShapeBodySetup, bool bUseArchetypeBodySetup, FVector BoxExtent)
+template <EShapeBodySetupHelper UpdateBodySetupAction, typename BodySetupType>
+bool InvalidateOrUpdateBoxBodySetup(BodySetupType& ShapeBodySetup, bool bUseArchetypeBodySetup, FVector BoxExtent)
 {
 	check((bUseArchetypeBodySetup && UpdateBodySetupAction == EShapeBodySetupHelper::InvalidateSharingIfStale) || (!bUseArchetypeBodySetup && UpdateBodySetupAction == EShapeBodySetupHelper::UpdateBodySetup));
 	check(ShapeBodySetup->AggGeom.BoxElems.Num() == 1);
diff --git a/Engine/Source/Runtime/Engine/Private/Components/CapsuleComponent.cpp b/Engine/Source/Runtime/Engine/Private/Components/CapsuleComponent.cpp
--- a/Engine/Source/Runtime/Engine/Private/Components/CapsuleComponent.cpp
+++ b/Engine/Source/Runtime/Engine/Private/Components/CapsuleComponent.cpp
@@ -183,8 +183,8 @@ void UCapsuleComponent::SetCapsuleSize(float NewRadius, float NewHalfHeight, boo
 }
 
 
-template <EShapeBodySetupHelper UpdateBodySetupAction>
-bool InvalidateOrUpdateCapsuleBodySetup(UBodySetup*& ShapeBodySetup, bool bUseArchetypeBodySetup, float CapsuleRadius, float CapsuleHalfHeight)
+template <EShapeBodySetupHelper UpdateBodySetupAction, typename BodySetupType>
+bool InvalidateOrUpdateCapsuleBodySetup(BodySetupType& ShapeBodySetup, bool bUseArchetypeBodySetup, float CapsuleRadius, float CapsuleHalfHeight)
 {
 	check((bUseArchetypeBodySetup && UpdateBodySetupAction == EShapeBodySetupHelper::InvalidateSharingIfStale) || (!bUseArchetypeBodySetup && UpdateBodySetupAction == EShapeBodySetupHelper::UpdateBodySetup));
 	check(ShapeBodySetup->AggGeom.SphylElems.Num() == 1);
diff --git a/Engine/Source/Runtime/Engine/Private/Components/ChildActorComponent.cpp b/Engine/Source/Runtime/Engine/Private/Components/ChildActorComponent.cpp
--- a/Engine/Source/Runtime/Engine/Private/Components/ChildActorComponent.cpp
+++ b/Engine/Source/Runtime/Engine/Private/Components/ChildActorComponent.cpp
@@ -535,7 +535,7 @@ void UChildActorComponent::SetChildActorClass(TSubclassOf<AActor> Class, AActor*
 			DestroyChildActor();
 
 			// If an actor template was supplied, temporarily set ChildActorTemplate to create the new Actor with ActorTemplate used as the template
-			TGuardValue<AActor*> ChildActorTemplateGuard(ChildActorTemplate, (ActorTemplate ? ActorTemplate : ToRawPtr(ChildActorTemplate)));
+			TGuardValue<decltype(ChildActorTemplate)> ChildActorTemplateGuard(ChildActorTemplate, (ActorTemplate ? ActorTemplate : ToRawPtr(ChildActorTemplate)));
 
 			CreateChildActor();
 		}
diff --git a/Engine/Source/Runtime/Engine/Private/Components/SphereComponent.cpp b/Engine/Source/Runtime/Engine/Private/Components/SphereComponent.cpp
--- a/Engine/Source/Runtime/Engine/Private/Components/SphereComponent.cpp
+++ b/Engine/Source/Runtime/Engine/Private/Components/SphereComponent.cpp
@@ -30,8 +30,8 @@ void USphereComponent::CalcBoundingCylinder(float& CylinderRadius, float& Cylind
 	CylinderHalfHeight = CylinderRadius;
 }
 
-template <EShapeBodySetupHelper UpdateBodySetupAction>
-bool InvalidateOrUpdateSphereBodySetup(UBodySetup*& ShapeBodySetup, bool bUseArchetypeBodySetup, float SphereRadius)
+template <EShapeBodySetupHelper UpdateBodySetupAction, typename BodySetupType>
+bool InvalidateOrUpdateSphereBodySetup(BodySetupType& ShapeBodySetup, bool bUseArchetypeBodySetup, float SphereRadius)
 {
 	check((bUseArchetypeBodySetup && UpdateBodySetupAction == EShapeBodySetupHelper::InvalidateSharingIfStale) || (!bUseArchetypeBodySetup && UpdateBodySetupAction == EShapeBodySetupHelper::UpdateBodySetup) );
 	check(ShapeBodySetup->AggGeom.SphereElems.Num() == 1);
diff --git a/Engine/Source/Runtime/Engine/Private/GameViewportClient.cpp b/Engine/Source/Runtime/Engine/Private/GameViewportClient.cpp
--- a/Engine/Source/Runtime/Engine/Private/GameViewportClient.cpp
+++ b/Engine/Source/Runtime/Engine/Private/GameViewportClient.cpp
@@ -440,7 +440,7 @@ void UGameViewportClient::SetEnabledStats(const TArray<FString>& InEnabledStats)
 void UGameViewportClient::Init(struct FWorldContext& WorldContext, UGameInstance* OwningGameInstance, bool bCreateNewAudioDevice)
 {
 	// set reference to world context
-	WorldContext.AddRef(World);
+	WorldContext.AddRef(static_cast<UWorld*&>(World));
 
 	// remember our game instance
 	GameInstance = OwningGameInstance;
diff --git a/Engine/Source/Runtime/Engine/Private/Materials/MaterialInstance.cpp b/Engine/Source/Runtime/Engine/Private/Materials/MaterialInstance.cpp
--- a/Engine/Source/Runtime/Engine/Private/Materials/MaterialInstance.cpp
+++ b/Engine/Source/Runtime/Engine/Private/Materials/MaterialInstance.cpp
@@ -587,7 +587,7 @@ UMaterialInstance::UMaterialInstance(const FObjectInitializer& ObjectInitializer
 	ShadingModels = MSM_Unlit;
 
 	PhysMaterial = nullptr;
-	for (UPhysicalMaterial*& PhysMat : PhysicalMaterialMap)
+	for (UE_TRANSITIONAL_OBJECT_PTR(UPhysicalMaterial)& PhysMat : PhysicalMaterialMap)
 	{
 		PhysMat = nullptr;
 	}
@@ -5015,7 +5015,7 @@ void UMaterialInstance::CopyMaterialUniformParametersInternal(UMaterialInterface
 							{
 								FTextureParameterValue* ParameterValue = new(TextureParameterValues) FTextureParameterValue;
 								ParameterValue->ParameterInfo.Name = Parameter.ParameterInfo.GetName();
-								Parameter.GetGameThreadTextureValue(AsMaterial, *MaterialResource, ParameterValue->ParameterValue);
+								Parameter.GetGameThreadTextureValue(AsMaterial, *MaterialResource, static_cast<UTexture *&>(ParameterValue->ParameterValue));
 							}
 						}
 					}
diff --git a/Engine/Source/Runtime/Engine/Private/UnrealEngine.cpp b/Engine/Source/Runtime/Engine/Private/UnrealEngine.cpp
--- a/Engine/Source/Runtime/Engine/Private/UnrealEngine.cpp
+++ b/Engine/Source/Runtime/Engine/Private/UnrealEngine.cpp
@@ -2446,7 +2446,8 @@ void UEngine::ParseCommandline()
 * @param Material Reference to a material object pointer that will be filled out
 * @param bCheckUsage Check if the material has been marked to be used as a special engine material
 */
-void LoadSpecialMaterial(const FString& MaterialName, const FString& MaterialNamePath, UMaterial*& Material, bool bCheckUsage)
+template <typename T>
+void LoadSpecialMaterial(const FString& MaterialName, const FString& MaterialNamePath, T& Material, bool bCheckUsage)
 {
 	// only bother with materials that aren't already loaded
 	if (Material == NULL)
@@ -2784,11 +2785,11 @@ void UEngine::InitializeObjectReferences()
 		};
 
 		// Standard fonts.
-		ConditionalLoadEngineFont(TinyFont, TinyFontName.ToString(), GetDefaultTinyFont);
-		ConditionalLoadEngineFont(SmallFont, SmallFontName.ToString(), GetDefaultRegularFont);
-		ConditionalLoadEngineFont(MediumFont, MediumFontName.ToString(), GetDefaultRegularFont);
-		ConditionalLoadEngineFont(LargeFont, LargeFontName.ToString(), GetDefaultRegularFont);
-		ConditionalLoadEngineFont(SubtitleFont, SubtitleFontName.ToString(), GetDefaultRegularFont);
+		ConditionalLoadEngineFont(static_cast<UFont*&>(TinyFont), TinyFontName.ToString(), GetDefaultTinyFont);
+		ConditionalLoadEngineFont(static_cast<UFont*&>(SmallFont), SmallFontName.ToString(), GetDefaultRegularFont);
+		ConditionalLoadEngineFont(static_cast<UFont*&>(MediumFont), MediumFontName.ToString(), GetDefaultRegularFont);
+		ConditionalLoadEngineFont(static_cast<UFont*&>(LargeFont), LargeFontName.ToString(), GetDefaultRegularFont);
+		ConditionalLoadEngineFont(static_cast<UFont*&>(SubtitleFont), SubtitleFontName.ToString(), GetDefaultRegularFont);
 
 		// Additional fonts.
 		AdditionalFonts.Empty(AdditionalFontNames.Num());
diff --git a/Engine/Source/Runtime/Engine/Private/WorldSettings.cpp b/Engine/Source/Runtime/Engine/Private/WorldSettings.cpp
--- a/Engine/Source/Runtime/Engine/Private/WorldSettings.cpp
+++ b/Engine/Source/Runtime/Engine/Private/WorldSettings.cpp
@@ -790,7 +790,7 @@ class UBookmarkBase* AWorldSettings::GetOrAddBookmark(const uint32 BookmarkIndex
 {
 	if (BookmarkArray.IsValidIndex(BookmarkIndex))
 	{
-		UBookmarkBase*& Bookmark = BookmarkArray[BookmarkIndex];
+		UE_TRANSITIONAL_OBJECT_PTR(UBookmarkBase)& Bookmark = BookmarkArray[BookmarkIndex];
 
 		if (Bookmark == nullptr || (bRecreateOnClassMismatch && Bookmark->GetClass() != GetDefaultBookmarkClass()))
 		{
@@ -840,7 +840,7 @@ void AWorldSettings::ClearBookmark(const uint32 BookmarkIndex)
 {
 	if (BookmarkArray.IsValidIndex(BookmarkIndex))
 	{
-		if (UBookmarkBase*& Bookmark = BookmarkArray[BookmarkIndex])
+		if (UE_TRANSITIONAL_OBJECT_PTR(UBookmarkBase)& Bookmark = BookmarkArray[BookmarkIndex])
 		{
 			Modify();
 			Bookmark->OnCleared();
@@ -852,7 +852,7 @@ void AWorldSettings::ClearBookmark(const uint32 BookmarkIndex)
 void AWorldSettings::ClearAllBookmarks()
 {
 	Modify();
-	for (UBookmarkBase*& Bookmark : BookmarkArray)
+	for (UE_TRANSITIONAL_OBJECT_PTR(UBookmarkBase)& Bookmark : BookmarkArray)
 	{
 		if (Bookmark)
 		{
@@ -901,7 +901,7 @@ void AWorldSettings::SanitizeBookmarkClasses()
 		bool bFoundInvalidBookmarks = false;
 		for (int32 i = 0; i < BookmarkArray.Num(); ++i)
 		{
-			if (UBookmarkBase*& Bookmark = BookmarkArray[i])
+			if (UE_TRANSITIONAL_OBJECT_PTR(UBookmarkBase)& Bookmark = BookmarkArray[i])
 			{
 				if (Bookmark->GetClass() != ExpectedClass)
 				{
diff --git a/Engine/Source/Runtime/Landscape/Private/LandscapeEditInterface.cpp b/Engine/Source/Runtime/Landscape/Private/LandscapeEditInterface.cpp
--- a/Engine/Source/Runtime/Landscape/Private/LandscapeEditInterface.cpp
+++ b/Engine/Source/Runtime/Landscape/Private/LandscapeEditInterface.cpp
@@ -4564,7 +4564,7 @@ void FLandscapeEditDataInterface::SetSelectData(int32 X1, int32 Y1, int32 X2, in
 	auto ReturnComponentTexture = [&](ULandscapeComponent* Component) -> UTexture2D*&
 	{
 		check(Component);
-		return Component->EditToolRenderData.DataTexture;
+		return static_cast<UTexture2D*&>(Component->EditToolRenderData.DataTexture);
 	};
 	SetEditToolTextureData(X1, Y1, X2, Y2, Data, Stride, ReturnComponentTexture);
 }
@@ -4602,7 +4602,7 @@ void FLandscapeEditDataInterface::SetLayerContributionData(int32 X1, int32 Y1, i
 	auto ReturnComponentTexture = [](ULandscapeComponent* Component) -> UTexture2D*&
 	{
 		check(Component);
-		return Component->EditToolRenderData.LayerContributionTexture;
+		return static_cast<UTexture2D*&>(Component->EditToolRenderData.LayerContributionTexture);
 	};
 	SetEditToolTextureData(X1, Y1, X2, Y2, Data, Stride, ReturnComponentTexture);
 }
@@ -4640,7 +4640,7 @@ void FLandscapeEditDataInterface::SetDirtyData(int32 X1, int32 Y1, int32 X2, int
 	auto ReturnComponentTexture = [](ULandscapeComponent* Component) -> UTexture2D*&
 	{
 		check(Component);
-		return Component->EditToolRenderData.DirtyTexture;
+		return static_cast<UTexture2D*&>(Component->EditToolRenderData.DirtyTexture);
 	};
 	SetEditToolTextureData(X1, Y1, X2, Y2, Data, Stride, ReturnComponentTexture, TEXTUREGROUP_8BitData);
 }
