Fixing anisotropic rendering.

diff --git a/Engine/Shaders/Private/BasePassPixelShader.usf b/Engine/Shaders/Private/BasePassPixelShader.usf
--- a/Engine/Shaders/Private/BasePassPixelShader.usf
+++ b/Engine/Shaders/Private/BasePassPixelShader.usf
@@ -1021,7 +1021,9 @@ void FPixelShaderInOut_MainPS(
 		);
 
 #if USES_GBUFFER
-	GBuffer.SelectiveOutputMask = GetSelectiveOutputMask();
+	// This requires cleanup. Shader code that uses GBuffer.SelectiveOutputMask expects the outputmask to be in
+	// bits [4:7], but it gets packed as bits [0:3] in the flexible gbuffer since we might move it around.
+	GBuffer.SelectiveOutputMask = GetSelectiveOutputMask() >> 4;
 	GBuffer.Velocity = 0;
 #endif
 
diff --git a/Engine/Shaders/Private/GBufferHelpers.ush b/Engine/Shaders/Private/GBufferHelpers.ush
--- a/Engine/Shaders/Private/GBufferHelpers.ush
+++ b/Engine/Shaders/Private/GBufferHelpers.ush
@@ -429,7 +429,6 @@ void GBufferPostDecode(inout FGBufferData Ret, bool bChecker, bool bGetNormalize
 		#endif //USE_DEVELOPMENT_SHADERS
 	}
 	
-#if GBUFFER_HAS_TANGENT
 	{
 		Ret.WorldTangent = DecodeNormal(Ret.WorldTangent);
 		Ret.Anisotropy = Ret.Anisotropy * 2.0f - 1.0f;
@@ -439,10 +438,8 @@ void GBufferPostDecode(inout FGBufferData Ret, bool bChecker, bool bGetNormalize
 			Ret.WorldTangent = normalize(Ret.WorldTangent);
 		}
 	}
-#else
-	{
-		Ret.WorldTangent = 0;
-		Ret.Anisotropy = 0;
-	}
-#endif
+
+	// This requires cleanup. Shader code that uses GBuffer.SelectiveOutputMask expects the outputmask to be in
+	// bits [4:7], but it gets packed as bits [0:3] in the flexible gbuffer since we might move it around.
+	Ret.SelectiveOutputMask = Ret.SelectiveOutputMask << 4;
 }
diff --git a/Engine/Source/Runtime/Engine/Private/ShaderCompiler/ShaderGenerationUtil.cpp b/Engine/Source/Runtime/Engine/Private/ShaderCompiler/ShaderGenerationUtil.cpp
--- a/Engine/Source/Runtime/Engine/Private/ShaderCompiler/ShaderGenerationUtil.cpp
+++ b/Engine/Source/Runtime/Engine/Private/ShaderCompiler/ShaderGenerationUtil.cpp
@@ -253,7 +253,7 @@ void FShaderCompileUtilities::ApplyFetchEnvironment(FShaderCompilerDefines& SrcD
 
 
 // if we change the logic, increment this number to force a DDC key change
-static const int32 GBufferGeneratorVersion = 2;
+static const int32 GBufferGeneratorVersion = 3;
 
 static FShaderGlobalDefines FetchShaderGlobalDefines(EShaderPlatform TargetPlatform)
 {
@@ -1097,6 +1097,7 @@ static FString CreateGBufferDecodeFunctionDirect(const FGBufferInfo& BufferInfo)
 		}
 	}
 	FullStr += TEXT(",\n\tfloat CustomNativeDepth");
+	FullStr += TEXT(",\n\tfloat4 AnisotropicData");
 	FullStr += TEXT(",\n\tuint CustomStencil");
 	FullStr += TEXT(",\n\tfloat SceneDepth");
 	FullStr += TEXT(",\n\tbool bGetNormalizedNormal");
@@ -1275,6 +1276,8 @@ static FString CreateGBufferDecodeFunctionDirect(const FGBufferInfo& BufferInfo)
 		}
 	}
 
+	FullStr += TEXT("\tRet.WorldTangent = AnisotropicData.xyz;\n");
+	FullStr += TEXT("\tRet.Anisotropy = AnisotropicData.w;\n");
 
 	FullStr += TEXT("\n");
 	FullStr += TEXT("\tGBufferPostDecode(Ret,bChecker,bGetNormalizedNormal);\n");
@@ -1340,7 +1343,6 @@ static FString CreateGBufferDecodeFunctionVariation(const FGBufferInfo& BufferIn
 
 	FullStr += TEXT("{\n");
 
-
 	if (DecodeType == CoordUV)
 	{
 		FullStr += FString::Printf(TEXT("\tfloat CustomNativeDepth = Texture2DSampleLevel(SceneTexturesStruct.CustomDepthTexture, SceneTexturesStruct_CustomDepthTextureSampler, %s, 0).r;\n"), CoordName.GetCharArray().GetData());
@@ -1356,12 +1358,14 @@ static FString CreateGBufferDecodeFunctionVariation(const FGBufferInfo& BufferIn
 		}
 
 		FullStr += FString::Printf(TEXT("\tfloat SceneDepth = CalcSceneDepth(%s);\n"), CoordName.GetCharArray().GetData());
+		FullStr += TEXT("\tfloat4 AnisotropicData = Texture2DSampleLevel(SceneTexturesStruct.GBufferFTexture, SceneTexturesStruct_GBufferFTextureSampler, UV, 0).xyzw;\n");
 	}
 	else if (DecodeType == CoordUInt)
 	{
 		FullStr += FString::Printf(TEXT("\tfloat CustomNativeDepth = SceneTexturesStruct.CustomDepthTexture.Load(int3(%s, 0)).r;\n"), CoordName.GetCharArray().GetData());
 		FullStr += FString::Printf(TEXT("\tuint CustomStencil = SceneTexturesStruct.CustomStencilTexture.Load(int3(%s, 0)) STENCIL_COMPONENT_SWIZZLE;\n"), CoordName.GetCharArray().GetData());
 		FullStr += FString::Printf(TEXT("\tfloat SceneDepth = CalcSceneDepth(%s);\n"), CoordName.GetCharArray().GetData());
+		FullStr += TEXT("\tfloat4 AnisotropicData = SceneTexturesStruct.GBufferFTexture.Load(int3(PixelPos, 0)).xyzw;\n");
 	}
 	else if (DecodeType == SceneTextures)
 	{
@@ -1369,6 +1373,7 @@ static FString CreateGBufferDecodeFunctionVariation(const FGBufferInfo& BufferIn
 		FullStr += TEXT("\tfloat CustomNativeDepth = 0;\n");
 		FullStr += FString::Printf(TEXT("\tfloat DeviceZ = SampleDeviceZFromSceneTexturesTempCopy(%s);\n"), CoordName.GetCharArray().GetData());
 		FullStr += TEXT("\tfloat SceneDepth = ConvertFromDeviceZ(DeviceZ);\n");
+		FullStr += TEXT("\tfloat4 AnisotropicData = GBufferFTexture.SampleLevel(GBufferFTextureSampler, UV, 0).xyzw;\n");
 	}
 	else if (DecodeType == SceneTexturesLoad)
 	{
@@ -1376,6 +1381,7 @@ static FString CreateGBufferDecodeFunctionVariation(const FGBufferInfo& BufferIn
 		FullStr += TEXT("\tfloat CustomNativeDepth = 0;\n");
 		FullStr += FString::Printf(TEXT("\tfloat DeviceZ = SceneDepthTexture.Load(int3(%s, 0)).r;\n"), CoordName.GetCharArray().GetData());
 		FullStr += TEXT("\tfloat SceneDepth = ConvertFromDeviceZ(DeviceZ);\n");
+		FullStr += TEXT("\tfloat4 AnisotropicData = GBufferFTexture.Load(int3(PixelCoord, 0)).xyzw;\n");
 	}
 	else
 	{
@@ -1490,6 +1496,7 @@ static FString CreateGBufferDecodeFunctionVariation(const FGBufferInfo& BufferIn
 		}
 	}
 	FullStr += TEXT(",\n\t\tCustomNativeDepth");
+	FullStr += TEXT(",\n\t\tAnisotropicData");
 	FullStr += TEXT(",\n\t\tCustomStencil");
 	FullStr += TEXT(",\n\t\tSceneDepth");
 	FullStr += TEXT(",\n\t\tbGetNormalizedNormal");
diff --git a/Engine/Source/Runtime/Renderer/Private/AnisotropyRendering.cpp b/Engine/Source/Runtime/Renderer/Private/AnisotropyRendering.cpp
--- a/Engine/Source/Runtime/Renderer/Private/AnisotropyRendering.cpp
+++ b/Engine/Source/Runtime/Renderer/Private/AnisotropyRendering.cpp
@@ -291,7 +291,7 @@ END_SHADER_PARAMETER_STRUCT()
 
 void FDeferredShadingSceneRenderer::RenderAnisotropyPass(
 	FRDGBuilder& GraphBuilder, 
-	FRDGTextureRef SceneDepthTexture,
+	FSceneTextures& SceneTextures,
 	bool bDoParallelPass
 	)
 {
@@ -300,13 +300,16 @@ void FDeferredShadingSceneRenderer::RenderAnisotropyPass(
 	SCOPE_CYCLE_COUNTER(STAT_AnisotropyPassDrawTime);
 	RDG_GPU_STAT_SCOPE(GraphBuilder, RenderAnisotropyPass);
 
-	FSceneRenderTargets& SceneContext = FSceneRenderTargets::Get();
-	if (!SceneContext.GBufferF)
+	// allocate if necessary, and register new GBufferF to SceneTexures
 	{
-		SceneContext.AllocateAnisotropyTarget(GraphBuilder.RHICmdList);
-		check(SceneContext.GBufferF);
+		FSceneRenderTargets& SceneContext = FSceneRenderTargets::Get();
+		if (!SceneContext.GBufferF)
+		{
+			SceneContext.AllocateAnisotropyTarget(GraphBuilder.RHICmdList);
+			check(SceneContext.GBufferF);
+		}
+		SceneTextures.GBufferF = GraphBuilder.RegisterExternalTexture(SceneContext.GBufferF);
 	}
-	FRDGTextureRef GBufferFTexture = GraphBuilder.RegisterExternalTexture(SceneContext.GBufferF);
 
 	for (int32 ViewIndex = 0; ViewIndex < Views.Num(); ViewIndex++)
 	{
@@ -322,13 +325,13 @@ void FDeferredShadingSceneRenderer::RenderAnisotropyPass(
 			}
 
 			auto* PassParameters = GraphBuilder.AllocParameters<FAnisotropyPassParameters>();
-			PassParameters->RenderTargets.DepthStencil = FDepthStencilBinding(SceneDepthTexture, ERenderTargetLoadAction::ELoad, FExclusiveDepthStencil::DepthRead_StencilNop);
+			PassParameters->RenderTargets.DepthStencil = FDepthStencilBinding(SceneTextures.Depth.Target, ERenderTargetLoadAction::ELoad, FExclusiveDepthStencil::DepthRead_StencilNop);
 
 			if (bDoParallelPass)
 			{
-				AddClearRenderTargetPass(GraphBuilder, GBufferFTexture);
+				AddClearRenderTargetPass(GraphBuilder, SceneTextures.GBufferF);
 
-				PassParameters->RenderTargets[0] = FRenderTargetBinding(GBufferFTexture, ERenderTargetLoadAction::ELoad);
+				PassParameters->RenderTargets[0] = FRenderTargetBinding(SceneTextures.GBufferF, ERenderTargetLoadAction::ELoad);
 
 				GraphBuilder.AddPass(
 					RDG_EVENT_NAME("AnisotropyPassParallel"),
@@ -344,7 +347,7 @@ void FDeferredShadingSceneRenderer::RenderAnisotropyPass(
 			}
 			else
 			{
-				PassParameters->RenderTargets[0] = FRenderTargetBinding(GBufferFTexture, ERenderTargetLoadAction::EClear);
+				PassParameters->RenderTargets[0] = FRenderTargetBinding(SceneTextures.GBufferF, ERenderTargetLoadAction::EClear);
 
 				GraphBuilder.AddPass(
 					RDG_EVENT_NAME("AnisotropyPass"),
diff --git a/Engine/Source/Runtime/Renderer/Private/BasePassRendering.cpp b/Engine/Source/Runtime/Renderer/Private/BasePassRendering.cpp
--- a/Engine/Source/Runtime/Renderer/Private/BasePassRendering.cpp
+++ b/Engine/Source/Runtime/Renderer/Private/BasePassRendering.cpp
@@ -725,7 +725,7 @@ END_SHADER_PARAMETER_STRUCT()
 
 void FDeferredShadingSceneRenderer::RenderBasePass(
 	FRDGBuilder& GraphBuilder,
-	const FSceneTextures& SceneTextures,
+	FSceneTextures& SceneTextures,
 	const FDBufferTextures& DBufferTextures,
 	FExclusiveDepthStencil::Type BasePassDepthStencilAccess,
 	FRDGTextureRef ForwardShadowMaskTexture)
@@ -927,7 +927,7 @@ void FDeferredShadingSceneRenderer::RenderBasePass(
 	if (ShouldRenderAnisotropyPass())
 	{
 		GraphBuilder.SetCommandListStat(GET_STATID(STAT_CLM_AnisotropyPass));
-		RenderAnisotropyPass(GraphBuilder, SceneTextures.Depth.Target, bEnableParallelBasePasses);
+		RenderAnisotropyPass(GraphBuilder, SceneTextures, bEnableParallelBasePasses);
 		GraphBuilder.SetCommandListStat(GET_STATID(STAT_CLM_AfterAnisotropyPass));
 	}
 }
diff --git a/Engine/Source/Runtime/Renderer/Private/DeferredShadingRenderer.h b/Engine/Source/Runtime/Renderer/Private/DeferredShadingRenderer.h
--- a/Engine/Source/Runtime/Renderer/Private/DeferredShadingRenderer.h
+++ b/Engine/Source/Runtime/Renderer/Private/DeferredShadingRenderer.h
@@ -198,7 +198,7 @@ public:
 
 	void RenderBasePass(
 		FRDGBuilder& GraphBuilder,
-		const FSceneTextures& SceneTextures,
+		FSceneTextures& SceneTextures,
 		const FDBufferTextures& DBufferTextures,
 		FExclusiveDepthStencil::Type BasePassDepthStencilAccess,
 		FRDGTextureRef ForwardShadowMaskTexture);
@@ -216,7 +216,7 @@ public:
 
 	void RenderAnisotropyPass(
 		FRDGBuilder& GraphBuilder,
-		FRDGTextureRef SceneDepthTexture,
+		FSceneTextures& SceneTextures,
 		bool bDoParallelPass);
 
 	void RenderSingleLayerWater(
