WP HLOD Generation
- Only collect garbage when exceeding memory requirements.
- Don't keep RF_StandAlone objects around when collecting garbage.
- Wait for async SDF and MeshCards computations before GC.

diff --git a/Engine/Source/Editor/UnrealEd/Private/Commandlets/WorldPartitionBuilderCommandlet.cpp b/Engine/Source/Editor/UnrealEd/Private/Commandlets/WorldPartitionBuilderCommandlet.cpp
--- a/Engine/Source/Editor/UnrealEd/Private/Commandlets/WorldPartitionBuilderCommandlet.cpp
+++ b/Engine/Source/Editor/UnrealEd/Private/Commandlets/WorldPartitionBuilderCommandlet.cpp
@@ -164,21 +164,5 @@ int32 UWorldPartitionBuilderCommandlet::Main(const FString& Params)
 	WorldContext.SetCurrentWorld(nullptr);
 	GWorld = nullptr;
 
-    FPlatformMemoryStats PlatformMemoryStats = FPlatformMemory::GetStats();
-    
-    UE_LOG(LogWorldPartitionBuilderCommandlet, Warning, TEXT("MemoryStats:")\
-           TEXT("\n\tAvailablePhysical %llu")\
-           TEXT("\n\t AvailableVirtual %llu")\
-           TEXT("\n\t     UsedPhysical %llu")\
-           TEXT("\n\t PeakUsedPhysical %llu")\
-           TEXT("\n\t      UsedVirtual %llu")\
-           TEXT("\n\t  PeakUsedVirtual %llu"),
-           (uint64)PlatformMemoryStats.AvailablePhysical,
-           (uint64)PlatformMemoryStats.AvailableVirtual,
-           (uint64)PlatformMemoryStats.UsedPhysical,
-           (uint64)PlatformMemoryStats.PeakUsedPhysical,
-           (uint64)PlatformMemoryStats.UsedVirtual,
-           (uint64)PlatformMemoryStats.PeakUsedVirtual);
-
 	return 0;
 }
diff --git a/Engine/Source/Runtime/Engine/Private/WorldPartition/HLOD/HLODBuilder.cpp b/Engine/Source/Runtime/Engine/Private/WorldPartition/HLOD/HLODBuilder.cpp
--- a/Engine/Source/Runtime/Engine/Private/WorldPartition/HLOD/HLODBuilder.cpp
+++ b/Engine/Source/Runtime/Engine/Private/WorldPartition/HLOD/HLODBuilder.cpp
@@ -76,8 +76,24 @@ public:
 		}
 		else
 		{
-			World->DestroyActor(HLODActor);
-			HLODActor = nullptr;
+			if (HLODActorGuid.IsValid())
+			{
+				FWorldPartitionActorDesc* HLODActorDesc = WorldPartition->GetActorDesc(HLODActorGuid);
+				check(HLODActorDesc);
+
+				check(HLODActorDesc->GetLoadedRefCount() == 1);
+				HLODActorDesc->RemoveLoadedRefCount();
+
+				HLODActor->GetLevel()->RemoveLoadedActor(HLODActor);
+				HLODActorDesc->Unload();
+
+				Context->HLODActorDescs.Add(CellHash, HLODActorGuid);
+			}
+			else
+			{
+				World->DestroyActor(HLODActor);
+				HLODActor = nullptr;
+			}
 		}
 
 		if (HLODActor)
diff --git a/Engine/Source/Runtime/Engine/Private/WorldPartition/WorldPartitionRuntimeSpatialHash.cpp b/Engine/Source/Runtime/Engine/Private/WorldPartition/WorldPartitionRuntimeSpatialHash.cpp
--- a/Engine/Source/Runtime/Engine/Private/WorldPartition/WorldPartitionRuntimeSpatialHash.cpp
+++ b/Engine/Source/Runtime/Engine/Private/WorldPartition/WorldPartitionRuntimeSpatialHash.cpp
@@ -27,6 +27,9 @@
 #include "Algo/Find.h"
 #include "EngineUtils.h"
 
+#include "DistanceFieldAtlas.h"
+#include "MeshCardRepresentation.h"
+
 #include "WorldPartition/HLOD/HLODLayer.h"
 #include "WorldPartition/HLOD/HLODActor.h"
 #include "WorldPartition/HLOD/HLODActorDesc.h"
@@ -1487,6 +1490,36 @@ void UWorldPartitionRuntimeSpatialHash::FlushStreaming()
 
 TArray<FGuid> GenerateHLODsForGrid(UWorldPartition* WorldPartition, const FSpatialHashRuntimeGrid& RuntimeGrid, uint32 HLODLevel, FHLODGenerationContext& Context, ISourceControlHelper* SourceControlHelper, const TArray<FActorCluster>& ActorClusters)
 {
+	auto HasExceededMaxMemory = []()
+	{
+		const uint64 MemoryMinFreePhysical = 1024ll * 1024 * 1024;
+		const uint64 MemoryMaxUsedPhysical = 16384ll * 1024 * 1024l;		
+		const FPlatformMemoryStats MemStats = FPlatformMemory::GetStats();
+		return (MemStats.AvailablePhysical < MemoryMinFreePhysical) || (MemStats.UsedPhysical >= MemoryMaxUsedPhysical);
+	};
+
+	auto DoCollectGarbage = []()
+	{
+		if (GDistanceFieldAsyncQueue)
+		{
+			GDistanceFieldAsyncQueue->BlockUntilAllBuildsComplete();
+		}
+
+		if (GCardRepresentationAsyncQueue)
+		{
+			GCardRepresentationAsyncQueue->BlockUntilAllBuildsComplete();
+		}
+
+		const FPlatformMemoryStats MemStatsBefore = FPlatformMemory::GetStats();
+		CollectGarbage(RF_NoFlags, true);
+		const FPlatformMemoryStats MemStatsAfter = FPlatformMemory::GetStats();
+
+		UE_LOG(LogWorldPartitionRuntimeSpatialHash, Warning, TEXT("AvailablePhysical:%.2fGB AvailableVirtual %.2fGB"),
+			((int64)MemStatsAfter.AvailablePhysical - (int64)MemStatsBefore.AvailablePhysical) / (1024.0 * 1024.0 * 1024.0),
+			((int64)MemStatsAfter.AvailableVirtual - (int64)MemStatsBefore.AvailableVirtual) / (1024.0 * 1024.0 * 1024.0)
+		);
+	};
+
 	const FBox WorldBounds = WorldPartition->GetWorldBounds();
 
 	const FSquare2DGridHelper PartitionedActors = GetPartitionedActors(WorldPartition, WorldBounds, RuntimeGrid, ActorClusters);
@@ -1517,7 +1550,7 @@ TArray<FGuid> GenerateHLODsForGrid(UWorldPartition* WorldPartition, const FSpati
 	FAssetRegistryModule& AssetRegistryModule = FModuleManager::LoadModuleChecked<FAssetRegistryModule>(TEXT("AssetRegistry"));
 
 	TArray<FGuid> GridHLODActors;
-	PartitionedActors.ForEachCells([&PartitionedActors, &GridHLODActors, &SlowTask, RuntimeGrid, HLODLevel, WorldPartition, WorldBounds, &Context, SourceControlHelper, &DirectoryWatcherModule, &AssetRegistryModule](const FIntVector& CellCoord)
+	PartitionedActors.ForEachCells([&PartitionedActors, &GridHLODActors, &SlowTask, RuntimeGrid, HLODLevel, WorldPartition, WorldBounds, &Context, SourceControlHelper, &DirectoryWatcherModule, &AssetRegistryModule, &HasExceededMaxMemory, &DoCollectGarbage](const FIntVector& CellCoord)
 	{
 		const FSquare2DGridHelper::FGridLevel::FGridCell& GridCell = PartitionedActors.GetCell(CellCoord);
 
@@ -1563,7 +1596,10 @@ TArray<FGuid> GenerateHLODsForGrid(UWorldPartition* WorldPartition, const FSpati
 					CellActors.Add(ActorDesc->GetActor());
 				}
 
-				CollectGarbage(GARBAGE_COLLECTION_KEEPFLAGS, true);
+				if (HasExceededMaxMemory())
+				{
+					DoCollectGarbage();
+				}
 
 				TArray<AWorldPartitionHLOD*> CellHLODActors;
 				CellHLODActors = UHLODLayer::GenerateHLODForCell(WorldPartition, &Context, CellName, CellBounds, HLODLevel, CellActors);
@@ -1629,7 +1665,7 @@ TArray<FGuid> GenerateHLODsForGrid(UWorldPartition* WorldPartition, const FSpati
 				for (AActor* Actor: ToUnloadActors)
 				{
 					FWorldPartitionActorDesc* ActorDesc = (FWorldPartitionActorDesc*)WorldPartition->GetActorDesc(Actor->GetActorGuid());
-					
+
 					check(ActorDesc->GetLoadedRefCount() == 1);
 					ActorDesc->RemoveLoadedRefCount();
 
@@ -1639,6 +1675,11 @@ TArray<FGuid> GenerateHLODsForGrid(UWorldPartition* WorldPartition, const FSpati
 			}
 		}
 	});
+
+	// Need to collect garbage here since some HLOD actors have been marked pending kill when destroying them
+	// and they may be loaded when generating the nexty HLOD layer.
+	DoCollectGarbage();
+
 	return GridHLODActors;
 }
 
