Fixed non-unity build

diff --git a/Engine/Source/Runtime/Engine/Private/WorldPartition/DataLayer/DataLayerSubsystem.cpp b/Engine/Source/Runtime/Engine/Private/WorldPartition/DataLayer/DataLayerSubsystem.cpp
--- a/Engine/Source/Runtime/Engine/Private/WorldPartition/DataLayer/DataLayerSubsystem.cpp
+++ b/Engine/Source/Runtime/Engine/Private/WorldPartition/DataLayer/DataLayerSubsystem.cpp
@@ -7,6 +7,7 @@
 #include "Engine/World.h"
 #if WITH_EDITOR
 #include "Editor.h"
+#include "Modules/ModuleManager.h"
 #include "WorldPartition/DataLayer/IDataLayerEditorModule.h"
 #endif
 
