Fix crash when exiting the editor with a partitioned world loaded.

diff --git a/Engine/Source/Runtime/Engine/Private/WorldPartition/WorldPartition.cpp b/Engine/Source/Runtime/Engine/Private/WorldPartition/WorldPartition.cpp
--- a/Engine/Source/Runtime/Engine/Private/WorldPartition/WorldPartition.cpp
+++ b/Engine/Source/Runtime/Engine/Private/WorldPartition/WorldPartition.cpp
@@ -334,12 +334,15 @@ void UWorldPartition::Uninitialize()
 		}
 		else if (IsMainWorldPartition())
 		{
-			// Unload all Editor cells
-			// @todo_ow: Once Metadata is removed from external actor's package, this won't be necessary anymore.
-			EditorHash->ForEachCell([this](UWorldPartitionEditorCell* Cell)
+			if (!GIsRequestingExit)
 			{
-				UpdateLoadingEditorCell(Cell, /*bShouldBeLoaded*/false);
-			});
+				// Unload all Editor cells
+				// @todo_ow: Once Metadata is removed from external actor's package, this won't be necessary anymore.
+				EditorHash->ForEachCell([this](UWorldPartitionEditorCell* Cell)
+				{
+					UpdateLoadingEditorCell(Cell, /*bShouldBeLoaded*/false);
+				});
+			}
 			UnregisterDelegates();
 		}
 
