Fix GIsRequestingExit compile warning.

diff --git a/Engine/Source/Runtime/Engine/Private/WorldPartition/WorldPartition.cpp b/Engine/Source/Runtime/Engine/Private/WorldPartition/WorldPartition.cpp
--- a/Engine/Source/Runtime/Engine/Private/WorldPartition/WorldPartition.cpp
+++ b/Engine/Source/Runtime/Engine/Private/WorldPartition/WorldPartition.cpp
@@ -334,7 +334,7 @@ void UWorldPartition::Uninitialize()
 		}
 		else if (IsMainWorldPartition())
 		{
-			if (!GIsRequestingExit)
+			if (!IsEngineExitRequested())
 			{
 				// Unload all Editor cells
 				// @todo_ow: Once Metadata is removed from external actor's package, this won't be necessary anymore.
