Set flag to store cached expression data in materials directly from Serialize function when cooking, to avoid potential for missing it

diff --git a/Engine/Source/Runtime/Engine/Private/Materials/Material.cpp b/Engine/Source/Runtime/Engine/Private/Materials/Material.cpp
--- a/Engine/Source/Runtime/Engine/Private/Materials/Material.cpp
+++ b/Engine/Source/Runtime/Engine/Private/Materials/Material.cpp
@@ -3198,9 +3198,6 @@ void UMaterial::CacheResourceShadersForCooking(EShaderPlatform ShaderPlatform, T
 
 	CacheShadersForResources(ShaderPlatform, NewResourcesToCache, EMaterialShaderPrecompileMode::Background, TargetPlatform);
 
-	// Make sure cached data is saved in cooked builds
-	bSavedCachedExpressionData = true;
-
 	OutCachedMaterialResources.Append(NewResourcesToCache);
 }
 
@@ -3301,6 +3298,19 @@ void UMaterial::Serialize(FArchive& Ar)
 	Ar.UsingCustomVersion(FRenderingObjectVersion::GUID);
 	Ar.UsingCustomVersion(FFortniteMainBranchObjectVersion::GUID);
 
+	if (Ar.IsCooking())
+	{
+		if (CachedExpressionData)
+		{
+			bSavedCachedExpressionData = true;
+		}
+		else
+		{
+			// ClassDefault object is expected to be missing cached data, but in all other cases it should have been created when the material was loaded, in PostLoad
+			checkf(HasAllFlags(RF_ClassDefaultObject), TEXT("Trying to save cooked material %s, missing CachedExpressionData"), *GetName());
+		}
+	}
+
 	Super::Serialize(Ar);
 
 	if (Ar.UE4Ver() >= VER_UE4_PURGED_FMATERIAL_COMPILE_OUTPUTS)
@@ -3338,6 +3348,7 @@ void UMaterial::Serialize(FArchive& Ar)
 		{
 			CachedExpressionData = new FMaterialCachedExpressionData();
 		}
+		check(CachedExpressionData);
 		UScriptStruct* Struct = FMaterialCachedExpressionData::StaticStruct();
 		Struct->SerializeTaggedProperties(Ar, (uint8*)CachedExpressionData, Struct, nullptr);
 	}
diff --git a/Engine/Source/Runtime/Engine/Private/Materials/MaterialInstance.cpp b/Engine/Source/Runtime/Engine/Private/Materials/MaterialInstance.cpp
--- a/Engine/Source/Runtime/Engine/Private/Materials/MaterialInstance.cpp
+++ b/Engine/Source/Runtime/Engine/Private/Materials/MaterialInstance.cpp
@@ -2903,9 +2903,6 @@ void UMaterialInstance::CacheResourceShadersForRendering(EMaterialShaderPrecompi
 
 void UMaterialInstance::CacheResourceShadersForCooking(EShaderPlatform ShaderPlatform, TArray<FMaterialResource*>& OutCachedMaterialResources, EMaterialShaderPrecompileMode PrecompileMode, const ITargetPlatform* TargetPlatform)
 {
-	// Make sure cached data is saved in cooked builds
-	bSavedCachedData = true;
-
 	if (bHasStaticPermutationResource)
 	{
 		UMaterial* BaseMaterial = GetMaterial();
@@ -3336,6 +3333,20 @@ void UMaterialInstance::Serialize(FArchive& Ar)
 	SCOPE_CYCLE_COUNTER(STAT_MaterialInstance_Serialize);
 
 	Ar.UsingCustomVersion(FRenderingObjectVersion::GUID);
+
+	if (Ar.IsCooking())
+	{
+		if (CachedData)
+		{
+			bSavedCachedData = true;
+		}
+		else
+		{
+			// ClassDefault object is expected to be missing cached data, but in all other cases it should have been created when the material was loaded, in PostLoad
+			checkf(HasAllFlags(RF_ClassDefaultObject), TEXT("Trying to save cooked material instance %s, missing CachedExpressionData"), *GetName());
+		}
+	}
+
 	Super::Serialize(Ar);
 		
 #if WITH_EDITOR
@@ -3370,6 +3381,7 @@ void UMaterialInstance::Serialize(FArchive& Ar)
 		{
 			CachedData = new FMaterialInstanceCachedData();
 		}
+		check(CachedData);
 		UScriptStruct* Struct = FMaterialInstanceCachedData::StaticStruct();
 		Struct->SerializeTaggedProperties(Ar, (uint8*)CachedData, Struct, nullptr);
 	}
