Add a NotEqual oiperator for Doubles

#rnx

diff --git a/Engine/Source/Runtime/Engine/Classes/Kismet/KismetMathLibrary.h b/Engine/Source/Runtime/Engine/Classes/Kismet/KismetMathLibrary.h
--- a/Engine/Source/Runtime/Engine/Classes/Kismet/KismetMathLibrary.h
+++ b/Engine/Source/Runtime/Engine/Classes/Kismet/KismetMathLibrary.h
@@ -601,6 +601,10 @@ class ENGINE_API UKismetMathLibrary : public UBlueprintFunctionLibrary
 	UFUNCTION(BlueprintPure, meta=(DisplayName = "NotEqual (float)", CompactNodeTitle = "!=", Keywords = "!= not equal"), Category="Math|Float")
 	static bool NotEqual_FloatFloat(float A, float B);
 
+	/** Returns true if A does not equal B (A != B)*/
+	UFUNCTION(BlueprintPure, meta = (DisplayName = "NotEqual (Double)", CompactNodeTitle = "!=", Keywords = "!= not equal"), Category = "Math|Double")
+	static bool NotEqual_DoubleDouble(double A, double B);
+
 	/** Returns true if value is between Min and Max (V >= Min && V <= Max)
 	 * If InclusiveMin is true, value needs to be equal or larger than Min, else it needs to be larger
 	 * If InclusiveMax is true, value needs to be smaller or equal than Max, else it needs to be smaller
diff --git a/Engine/Source/Runtime/Engine/Classes/Kismet/KismetMathLibrary.inl b/Engine/Source/Runtime/Engine/Classes/Kismet/KismetMathLibrary.inl
--- a/Engine/Source/Runtime/Engine/Classes/Kismet/KismetMathLibrary.inl
+++ b/Engine/Source/Runtime/Engine/Classes/Kismet/KismetMathLibrary.inl
@@ -581,6 +581,12 @@ bool UKismetMathLibrary::NotEqual_FloatFloat(float A, float B)
 	return A != B;
 }
 
+KISMET_MATH_FORCEINLINE
+bool UKismetMathLibrary::NotEqual_DoubleDouble(double A, double B)
+{
+	return A != B;
+}
+
 KISMET_MATH_FORCEINLINE
 float UKismetMathLibrary::GridSnap_Float(float Location, float GridSize)
 {
