Disabled capture everyframe option when capturing Minimap in world partition.

diff --git a/Engine/Source/Runtime/Engine/Private/WorldPartition/WorldPartitionMiniMapHelper.cpp b/Engine/Source/Runtime/Engine/Private/WorldPartition/WorldPartitionMiniMapHelper.cpp
--- a/Engine/Source/Runtime/Engine/Private/WorldPartition/WorldPartitionMiniMapHelper.cpp
+++ b/Engine/Source/Runtime/Engine/Private/WorldPartition/WorldPartitionMiniMapHelper.cpp
@@ -117,6 +117,8 @@ void FWorldPartitionMiniMapHelper::CaptureWorldMiniMapToTexture(UWorld* InWorld,
 		CaptureComponent->OrthoWidth = ViewportWidth;
 		CaptureComponent->bUseCustomProjectionMatrix = true;
 		CaptureComponent->CustomProjectionMatrix = ProjectionMatrix;
+		CaptureComponent->bCaptureEveryFrame = false;
+		CaptureComponent->bCaptureOnMovement = false;
 		CaptureComponent->CaptureScene();
 
 		auto RenderTargetResource = RenderTargetTexture->GameThread_GetRenderTargetResource();
