Add better log messages for failed VT compression

diff --git a/Engine/Source/Runtime/Engine/Private/Texture2D.cpp b/Engine/Source/Runtime/Engine/Private/Texture2D.cpp
--- a/Engine/Source/Runtime/Engine/Private/Texture2D.cpp
+++ b/Engine/Source/Runtime/Engine/Private/Texture2D.cpp
@@ -1259,7 +1259,7 @@ void FVirtualTexture2DResource::InitRHI()
 		ProducerDesc.PhysicalGroupIndex[LayerIndex] = bSinglePhysicalSpace ? 0 : LayerIndex;
 	}
 
-	FUploadingVirtualTexture* VirtualTexture = new FUploadingVirtualTexture(VTData, FirstMipToUse);
+	FUploadingVirtualTexture* VirtualTexture = new FUploadingVirtualTexture(ProducerDesc.Name, VTData, FirstMipToUse);
 	ProducerHandle = GetRendererModule().RegisterVirtualTextureProducer(ProducerDesc, VirtualTexture);
 
 	// Only create the miptails mini-texture in-editor.
diff --git a/Engine/Source/Runtime/Engine/Private/VT/RuntimeVirtualTexture.cpp b/Engine/Source/Runtime/Engine/Private/VT/RuntimeVirtualTexture.cpp
--- a/Engine/Source/Runtime/Engine/Private/VT/RuntimeVirtualTexture.cpp
+++ b/Engine/Source/Runtime/Engine/Private/VT/RuntimeVirtualTexture.cpp
@@ -619,7 +619,7 @@ namespace RuntimeVirtualTexture
 					const int32 AdjustedTransitionLevel = TransitionLevel + FirstStreamingMip;
 					OutTransitionLevel = TransitionLevel;
 
-					IVirtualTexture* StreamingProducer = new FUploadingVirtualTexture(VTData, FirstStreamingMip);
+					IVirtualTexture* StreamingProducer = new FUploadingVirtualTexture(InStreamingTexture->GetFName(), VTData, FirstStreamingMip);
 					return new FVirtualTextureLevelRedirector(InProducer, StreamingProducer, AdjustedTransitionLevel);
 				}
 			}
diff --git a/Engine/Source/Runtime/Engine/Private/VT/UploadingVirtualTexture.cpp b/Engine/Source/Runtime/Engine/Private/VT/UploadingVirtualTexture.cpp
--- a/Engine/Source/Runtime/Engine/Private/VT/UploadingVirtualTexture.cpp
+++ b/Engine/Source/Runtime/Engine/Private/VT/UploadingVirtualTexture.cpp
@@ -59,8 +59,9 @@ FVirtualTextureCodec* FVirtualTextureCodec::ListHead = nullptr;
 FVirtualTextureCodec FVirtualTextureCodec::ListTail;
 uint32 FVirtualTextureCodec::NumCodecs = 0u;
 
-FUploadingVirtualTexture::FUploadingVirtualTexture(FVirtualTextureBuiltData* InData, int32 InFirstMipToUse)
-	: Data(InData)
+FUploadingVirtualTexture::FUploadingVirtualTexture(const FName& InName, FVirtualTextureBuiltData* InData, int32 InFirstMipToUse)
+	: Name(InName)
+	, Data(InData)
 	, FirstMipOffset(InFirstMipToUse)
 {
 	HandlePerChunk.AddDefaulted(InData->Chunks.Num());
diff --git a/Engine/Source/Runtime/Engine/Private/VT/UploadingVirtualTexture.h b/Engine/Source/Runtime/Engine/Private/VT/UploadingVirtualTexture.h
--- a/Engine/Source/Runtime/Engine/Private/VT/UploadingVirtualTexture.h
+++ b/Engine/Source/Runtime/Engine/Private/VT/UploadingVirtualTexture.h
@@ -56,7 +56,7 @@ struct FVTDataAndStatus
 class FUploadingVirtualTexture : public IVirtualTexture
 {
 public:
-	FUploadingVirtualTexture(FVirtualTextureBuiltData* InData, int32 FirstMipToUse);
+	FUploadingVirtualTexture(const FName& InName, FVirtualTextureBuiltData* InData, int32 FirstMipToUse);
 	virtual ~FUploadingVirtualTexture();
 
 	// IVirtualTexture interface
@@ -72,6 +72,7 @@ public:
 	virtual void DumpToConsole(bool verbose) override;
 	// End IVirtualTexture interface
 
+	inline const FName& GetName() const { return Name; }
 	inline const FVirtualTextureBuiltData* GetVTData() const { return Data; }
 
 	// gets the codec for the given chunk, data is not valid until returned OutCompletionEvents are complete
@@ -83,6 +84,7 @@ public:
 private:
 	friend class FVirtualTextureCodec;
 
+	FName Name;
 	FVirtualTextureBuiltData* Data;
 	TArray< TUniquePtr<IFileCacheHandle> > HandlePerChunk;
 	TArray< TUniquePtr<FVirtualTextureCodec> > CodecPerChunk;
diff --git a/Engine/Source/Runtime/Engine/Private/VT/VirtualTextureChunkManager.cpp b/Engine/Source/Runtime/Engine/Private/VT/VirtualTextureChunkManager.cpp
--- a/Engine/Source/Runtime/Engine/Private/VT/VirtualTextureChunkManager.cpp
+++ b/Engine/Source/Runtime/Engine/Private/VT/VirtualTextureChunkManager.cpp
@@ -143,6 +143,7 @@ FVTRequestPageResult FVirtualTextureChunkStreamingManager::RequestTile(FUploadin
 	TranscodeParams.vLevel = vLevel;
 	TranscodeParams.LayerMask = LayerMask;
 	TranscodeParams.Codec = CodecResult.Codec;
+	TranscodeParams.Name = VTexture->GetName();
 	const FVTTranscodeTileHandle TranscodeHandle = TranscodeCache.SubmitTask(UploadCache, TranscodeKey, TranscodeParams, &GraphCompletionEvents);
 	return FVTRequestPageResult(EVTRequestPageStatus::Pending, TranscodeHandle.PackedData);
 }
diff --git a/Engine/Source/Runtime/Engine/Private/VT/VirtualTextureTranscodeCache.cpp b/Engine/Source/Runtime/Engine/Private/VT/VirtualTextureTranscodeCache.cpp
--- a/Engine/Source/Runtime/Engine/Private/VT/VirtualTextureTranscodeCache.cpp
+++ b/Engine/Source/Runtime/Engine/Private/VT/VirtualTextureTranscodeCache.cpp
@@ -62,6 +62,11 @@ struct FTranscodeTask
 	}
 
 	void DoTask(ENamedThreads::Type CurrentThread, const FGraphEventRef& MyCompletionGraphEvent)
+	{
+		DoTask(true);
+	}
+
+	void DoTask(bool bAsync)
 	{
 		SCOPED_NAMED_EVENT(VirtualTextureTranscodeTask_DoTask, FColor::Cyan);
 
@@ -210,6 +215,7 @@ struct FTranscodeTask
 
 			if (!bDecodeResult && GShowDecodeErrors)
 			{
+				UE_LOG(LogVirtualTexturing, Error, TEXT("Failed to decode VT tile (vAddress %08X, vLevel %d) for '%s'"), Params.vAddress, Params.vLevel, *Params.Name.ToString());
 				UniformColorPixels(StagingBufferForLayer, TilePixelSize, TilePixelSize, LayerFormat, ErrorColor);
 			}
 
@@ -415,31 +421,8 @@ FVTTranscodeTileHandle FVirtualTextureTranscodeCache::SubmitTask(FVirtualTexture
 			// If there aren't any prerequisites, it's not worth launching a task for a raw memcpy, just do the work now
 			if (!Prerequisites && Chunk.CodecType[LayerIndex] == EVirtualTextureCodec::RawGPU)
 			{
-				const uint32 DataOffset = TileLayerOffset - TileBaseOffset;
-				const uint32 TileLayerSize = NextTileLayerOffset - TileLayerOffset;
-
-				const uint32 TileWidthInBlocks = FMath::DivideAndRoundUp(TilePixelSize, (uint32)GPixelFormats[LayerFormat].BlockSizeX);
-				const uint32 TileHeightInBlocks = FMath::DivideAndRoundUp(TilePixelSize, (uint32)GPixelFormats[LayerFormat].BlockSizeY);
-				const uint32 PackedStride = TileWidthInBlocks * GPixelFormats[LayerFormat].BlockBytes;
-				const size_t PackedOutputSize = PackedStride * TileHeightInBlocks;
-
-				if (StagingBufferForLayer.Stride == PackedStride)
-				{
-					check(TileLayerSize <= StagingBufferForLayer.MemorySize);
-					InParams.Data->CopyTo(StagingBufferForLayer.Memory, DataOffset, TileLayerSize);
-				}
-				else
-				{
-					check(PackedStride <= StagingBufferForLayer.Stride);
-					check(TileLayerSize <= PackedOutputSize);
-					check(TileHeightInBlocks * StagingBufferForLayer.Stride <= StagingBufferForLayer.MemorySize);
-					TempBuffer.SetNumUninitialized(PackedOutputSize, false);
-					InParams.Data->CopyTo(TempBuffer.GetData(), DataOffset, TileLayerSize);
-					for (uint32 y = 0; y < TileHeightInBlocks; ++y)
-					{
-						FMemory::Memcpy((uint8*)StagingBufferForLayer.Memory + y * StagingBufferForLayer.Stride, TempBuffer.GetData() + y * PackedStride, PackedStride);
-					}
-				}
+				FTranscodeTask LocalTask(StagingBuffer, InParams);
+				LocalTask.DoTask(false);
 			}
 			else
 			{
diff --git a/Engine/Source/Runtime/Engine/Private/VT/VirtualTextureTranscodeCache.h b/Engine/Source/Runtime/Engine/Private/VT/VirtualTextureTranscodeCache.h
--- a/Engine/Source/Runtime/Engine/Private/VT/VirtualTextureTranscodeCache.h
+++ b/Engine/Source/Runtime/Engine/Private/VT/VirtualTextureTranscodeCache.h
@@ -16,6 +16,7 @@ struct FVTTranscodeParams
 	IMemoryReadStreamRef Data;
 	const FVirtualTextureCodec* Codec;
 	const FVirtualTextureBuiltData* VTData;
+	FName Name;
 	uint32 ChunkIndex;
 	uint32 vAddress;
 	uint8 vLevel;
