Add which type of compression was used to store the texture source in the .uasset in the asset registry

diff --git a/Engine/Source/Runtime/Engine/Private/Texture.cpp b/Engine/Source/Runtime/Engine/Private/Texture.cpp
--- a/Engine/Source/Runtime/Engine/Private/Texture.cpp
+++ b/Engine/Source/Runtime/Engine/Private/Texture.cpp
@@ -613,6 +613,8 @@ void UTexture::GetAssetRegistryTags(TArray<FAssetRegistryTag>& OutTags) const
 		OutTags.Add( FAssetRegistryTag(SourceFileTagName(), AssetImportData->GetSourceData().ToJson(), FAssetRegistryTag::TT_Hidden) );
 	}
 
+	OutTags.Add(FAssetRegistryTag("SourceCompression", Source.IsPNGCompressed() ? TEXT("PNG") : LexToString(NAME_Zlib), FAssetRegistryTag::TT_Alphabetical));
+	
 	Super::GetAssetRegistryTags(OutTags);
 }
 #endif
