Make sure StateId is initialized for all UMaterialFunctionInterface derived types, previously was missed for function-instances
This broke shader map ID for any material using blend/layer instances

diff --git a/Engine/Source/Runtime/Engine/Classes/Materials/MaterialFunctionInterface.h b/Engine/Source/Runtime/Engine/Classes/Materials/MaterialFunctionInterface.h
--- a/Engine/Source/Runtime/Engine/Classes/Materials/MaterialFunctionInterface.h
+++ b/Engine/Source/Runtime/Engine/Classes/Materials/MaterialFunctionInterface.h
@@ -36,6 +36,7 @@ class UMaterialFunctionInterface : public UObject
 	GENERATED_UCLASS_BODY()
 
 	//~ Begin UObject Interface.
+	virtual void PostLoad() override;
 	virtual void GetAssetRegistryTags(TArray<FAssetRegistryTag>& OutTags) const override;
 	//~ End UObject Interface.
 
diff --git a/Engine/Source/Runtime/Engine/Private/Materials/MaterialExpressions.cpp b/Engine/Source/Runtime/Engine/Private/Materials/MaterialExpressions.cpp
--- a/Engine/Source/Runtime/Engine/Private/Materials/MaterialExpressions.cpp
+++ b/Engine/Source/Runtime/Engine/Private/Materials/MaterialExpressions.cpp
@@ -11915,6 +11915,16 @@ UMaterialFunctionInterface::UMaterialFunctionInterface(const FObjectInitializer&
 {
 }
 
+void UMaterialFunctionInterface::PostLoad()
+{
+	Super::PostLoad();
+
+	if (!StateId.IsValid())
+	{
+		StateId = FGuid::NewGuid();
+	}
+}
+
 void UMaterialFunctionInterface::GetAssetRegistryTags(TArray<FAssetRegistryTag>& OutTags) const
 {
 	Super::GetAssetRegistryTags(OutTags);
@@ -12094,11 +12104,6 @@ void UMaterialFunction::PostLoad()
 	LLM_SCOPE(ELLMTag::Materials);
 
 	Super::PostLoad();
-	
-	if (!StateId.IsValid())
-	{
-		StateId = FGuid::NewGuid();
-	}
 
 #if WITH_EDITORONLY_DATA
 	for (int32 ExpressionIndex = 0; ExpressionIndex < FunctionExpressions.Num(); ExpressionIndex++)
@@ -13298,7 +13303,15 @@ const FMaterialLayersFunctions::ID FMaterialLayersFunctions::GetID() const
 	for (int i=0; i<Layers.Num(); ++i)
 	{
 		const UMaterialFunctionInterface* Layer = Layers[i];
-		Result.LayerIDs[i] = (Layer) ? Layer->StateId : FGuid();
+		if (Layer)
+		{
+			check(Layer->StateId.IsValid());
+			Result.LayerIDs[i] = Layer->StateId;
+		}
+		else
+		{
+			Result.LayerIDs[i] = FGuid();
+		}
 	}
 
 	// Store the blend IDs in following format - stateID per function
@@ -13306,7 +13319,15 @@ const FMaterialLayersFunctions::ID FMaterialLayersFunctions::GetID() const
 	for (int i = 0; i < Blends.Num(); ++i)
 	{
 		const UMaterialFunctionInterface* Blend = Blends[i];
-		Result.BlendIDs[i] = (Blend) ? Blend->StateId : FGuid();
+		if (Blend)
+		{
+			check(Blend->StateId.IsValid());
+			Result.BlendIDs[i] = Blend->StateId;
+		}
+		else
+		{
+			Result.BlendIDs[i] = FGuid();
+		}
 	}
 
 	// Store the states copy
