Added missing implementation of UTextureCube::GetSurfaceDepth() and UTextureCube::GetSurfaceArraySize().

diff --git a/Engine/Source/Runtime/Engine/Classes/Engine/TextureCube.h b/Engine/Source/Runtime/Engine/Classes/Engine/TextureCube.h
--- a/Engine/Source/Runtime/Engine/Classes/Engine/TextureCube.h
+++ b/Engine/Source/Runtime/Engine/Classes/Engine/TextureCube.h
@@ -60,6 +60,8 @@ public:
 	//~ Begin UTexture Interface
 	virtual float GetSurfaceWidth() const override { return GetSizeX(); }
 	virtual float GetSurfaceHeight() const override { return GetSizeY(); }
+	virtual float GetSurfaceDepth() const override { return 0; }
+	virtual uint32 GetSurfaceArraySize() const override { return 6; }
 	ENGINE_API virtual FTextureResource* CreateResource() override;
 	ENGINE_API virtual void UpdateResource() override;
 	virtual EMaterialValueType GetMaterialType() const override { return MCT_TextureCube; }
