Made strata experimental and not beta (we do nto bother with data versioning/compatibility for now)

diff --git a/Engine/Source/Runtime/Engine/Classes/Engine/RendererSettings.h b/Engine/Source/Runtime/Engine/Classes/Engine/RendererSettings.h
--- a/Engine/Source/Runtime/Engine/Classes/Engine/RendererSettings.h
+++ b/Engine/Source/Runtime/Engine/Classes/Engine/RendererSettings.h
@@ -714,8 +714,8 @@ class ENGINE_API URendererSettings : public UDeveloperSettings
 	"Enable Strata materials (Beta)."
 	*/
 	UPROPERTY(config, EditAnywhere, Category = Strata, meta = (
-		ConsoleVariable = "r.Strata", DisplayName = "Strata materials (Beta)",
-		ToolTip = "Enable Strata materials (Beta).",
+		ConsoleVariable = "r.Strata", DisplayName = "Strata materials (Experimental)",
+		ToolTip = "Enable Strata materials (Experimental).",
 		ConfigRestartRequired = true))
 		uint32 bEnableStrata : 1;
 
