Actor descriptors handles/references
- Actor descriptors now have a soft (handle) and a hard (reference) reference count; soft reference count will not load/unload of actors, and hard reference will.
- Actor descriptors unique pointers are now allocated in non-volatile chunks, to allow handles/references to point to them directly.
- Converted manual loading of actors in HLOD builder code to use references.

diff --git a/Engine/Source/Runtime/CoreUObject/Public/AssetRegistry/AssetData.h b/Engine/Source/Runtime/CoreUObject/Public/AssetRegistry/AssetData.h
--- a/Engine/Source/Runtime/CoreUObject/Public/AssetRegistry/AssetData.h
+++ b/Engine/Source/Runtime/CoreUObject/Public/AssetRegistry/AssetData.h
@@ -401,6 +401,12 @@ public:
 		return Package;
 	}
 
+	/** Try to find the given tag  */
+	bool FindTag(const FName InTagName) const
+	{
+		return TagsAndValues.FindTag(InTagName).IsSet();
+	}
+
 	/** Try and get the value associated with the given tag as a type converted value */
 	template <typename ValueType>
 	bool GetTagValue(const FName InTagName, ValueType& OutTagValue) const;
diff --git a/Engine/Source/Runtime/Engine/Private/WorldPartition/HLOD/HLODBuilder.cpp b/Engine/Source/Runtime/Engine/Private/WorldPartition/HLOD/HLODBuilder.cpp
--- a/Engine/Source/Runtime/Engine/Private/WorldPartition/HLOD/HLODBuilder.cpp
+++ b/Engine/Source/Runtime/Engine/Private/WorldPartition/HLOD/HLODBuilder.cpp
@@ -5,6 +5,7 @@
 #if WITH_EDITOR
 
 #include "WorldPartition/WorldPartition.h"
+#include "WorldPartition/WorldPartitionHandle.h"
 #include "WorldPartition/HLOD/HLODActor.h"
 #include "WorldPartition/HLOD/HLODLayer.h"
 #include "WorldPartition/HLOD/HLODActorDesc.h"
@@ -39,16 +40,12 @@ public:
 		// Compute HLODActor hash
 		uint64 CellHash = FHLODActorDesc::ComputeCellHash(HLODLayer->GetName(), Context->GridIndexX, Context->GridIndexY, Context->GridIndexZ);
 
-		FGuid HLODActorGuid;
-		if (Context->HLODActorDescs.RemoveAndCopyValue(CellHash, HLODActorGuid))
+		int32 HLODActorRefIndex = INDEX_NONE;
+		FWorldPartitionHandle HLODActorHandle;
+		if (Context->HLODActorDescs.RemoveAndCopyValue(CellHash, HLODActorHandle))
 		{
-			FWorldPartitionActorDesc* HLODActorDesc = WorldPartition->GetActorDesc(HLODActorGuid);
-			check(HLODActorDesc);
-
-			check(!HLODActorDesc->GetLoadedRefCount());
-			HLODActorDesc->AddLoadedRefCount();
-			HLODActor = CastChecked<AWorldPartitionHLOD>(HLODActorDesc->Load());
-			HLODActor->GetLevel()->AddLoadedActor(HLODActor);
+			HLODActorRefIndex = Context->ActorReferences.Add(FWorldPartitionHandleHelpers::ConvertHandleToReference(HLODActorHandle));
+			HLODActor = CastChecked<AWorldPartitionHLOD>(HLODActorHandle->GetActor());
 		}
 
 		if (!HLODActor)
@@ -76,18 +73,10 @@ public:
 		}
 		else
 		{
-			if (HLODActorGuid.IsValid())
+			if (HLODActorRefIndex != INDEX_NONE)
 			{
-				FWorldPartitionActorDesc* HLODActorDesc = WorldPartition->GetActorDesc(HLODActorGuid);
-				check(HLODActorDesc);
-
-				check(HLODActorDesc->GetLoadedRefCount() == 1);
-				HLODActorDesc->RemoveLoadedRefCount();
-
-				HLODActor->GetLevel()->RemoveLoadedActor(HLODActor);
-				HLODActorDesc->Unload();
-
-				Context->HLODActorDescs.Add(CellHash, HLODActorGuid);
+				Context->HLODActorDescs.Add(CellHash, FWorldPartitionHandle(WorldPartition, HLODActor->GetActorGuid()));
+				Context->ActorReferences.RemoveAtSwap(HLODActorRefIndex);
 			}
 			else
 			{
diff --git a/Engine/Source/Runtime/Engine/Private/WorldPartition/HLOD/HLODBuilder.h b/Engine/Source/Runtime/Engine/Private/WorldPartition/HLOD/HLODBuilder.h
--- a/Engine/Source/Runtime/Engine/Private/WorldPartition/HLOD/HLODBuilder.h
+++ b/Engine/Source/Runtime/Engine/Private/WorldPartition/HLOD/HLODBuilder.h
@@ -6,12 +6,24 @@
 
 #include "CoreMinimal.h"
 #include "Containers/Array.h"
+#include "WorldPartition/WorldPartitionHandle.h"
 
 class AActor;
 class UWorldPartition;
 class UHLODLayer;
 class AWorldPartitionHLOD;
-struct FHLODGenerationContext;
+
+struct ENGINE_API FHLODGenerationContext
+{
+	TMap<uint64, FWorldPartitionHandle> HLODActorDescs;
+	TArray<FWorldPartitionReference> ActorReferences;
+	
+	// Everything needed to build the cell hash
+	int64 GridIndexX;
+	int64 GridIndexY;
+	int64 GridIndexZ;
+	FName HLODLayerName;
+};
 
 /**
  * Tools for building HLODs in WorldPartition
diff --git a/Engine/Source/Runtime/Engine/Private/WorldPartition/WorldPartition.cpp b/Engine/Source/Runtime/Engine/Private/WorldPartition/WorldPartition.cpp
--- a/Engine/Source/Runtime/Engine/Private/WorldPartition/WorldPartition.cpp
+++ b/Engine/Source/Runtime/Engine/Private/WorldPartition/WorldPartition.cpp
@@ -64,9 +64,7 @@ static FAutoConsoleCommand DumpActorDescs(
 		}
 	})
 );
-#endif
 
-#if WITH_EDITOR
 // Helper class to avoid sending global events until all cells updates are processed.
 struct FWorldPartionCellUpdateContext
 {
@@ -100,9 +98,7 @@ struct FWorldPartionCellUpdateContext
 };
 
 int32 FWorldPartionCellUpdateContext::UpdatesInProgress = 0;
-#endif
 
-#if WITH_EDITOR
 UWorldPartition::FEnableWorldPartitionEvent UWorldPartition::EnableWorldPartitionEvent;
 UWorldPartition::FWorldPartitionChangedEvent UWorldPartition::WorldPartitionChangedEvent;
 #endif
@@ -241,7 +237,8 @@ void UWorldPartition::Initialize(UWorld* InWorld, const FTransform& InTransform)
 
 		for (const FAssetData& Asset : Assets)
 		{
-			TUniquePtr<FWorldPartitionActorDesc> NewActorDesc = GetActorDescriptor(Asset);
+			TUniquePtr<FWorldPartitionActorDesc>* NewActorDesc = ActorDescList.Emplace(GetActorDescriptor(Asset));
+			check(NewActorDesc->IsValid());
 
 			if (bIsInstanced)
 			{
@@ -251,19 +248,19 @@ void UWorldPartition::Initialize(UWorld* InWorld, const FTransform& InTransform)
 
 				InstancingContext.AddMapping(*LongActorPackageName, *InstancedName);
 
-				NewActorDesc->TransformInstance(ReplaceFrom, ReplaceTo, InstanceTransform);
+				(*NewActorDesc)->TransformInstance(ReplaceFrom, ReplaceTo, InstanceTransform);
 			}
 
-			Actors.Add(NewActorDesc->GetGuid(), MoveTemp(NewActorDesc));
-		}
+			Actors.Add((*NewActorDesc)->GetGuid(), NewActorDesc);
 
-		if (bEditorOnly)
-		{
-			for (const auto& Pair : Actors)
+			if (bEditorOnly)
 			{
-				HashActorDesc(Pair.Value.Get());
+				HashActorDesc(NewActorDesc->Get());
 			}
+		}
 
+		if (bEditorOnly)
+		{
 			// Load the always loaded cell, don't call LoadCells to avoid creating a transaction
 			UpdateLoadingEditorCell(EditorHash->GetAlwaysLoadedCell(), true);
 
@@ -428,9 +425,9 @@ void UWorldPartition::ForEachIntersectingActorDesc(const FBox& Box, TSubclassOf<
 
 void UWorldPartition::ForEachActorDesc(TSubclassOf<AActor> ActorClass, TFunctionRef<bool(const FWorldPartitionActorDesc*)> Predicate) const
 {
-	for (const TPair<FGuid, TUniquePtr<FWorldPartitionActorDesc>>& Pair : Actors)
+	for (const TPair<FGuid, TUniquePtr<FWorldPartitionActorDesc>*> Pair : Actors)
 	{
-		FWorldPartitionActorDesc* ActorDesc = Pair.Value.Get();
+		FWorldPartitionActorDesc* ActorDesc = Pair.Value->Get();
 		if (ActorDesc->GetActorClass()->IsChildOf(ActorClass))
 		{
 			if (!Predicate(ActorDesc))
@@ -443,14 +440,14 @@ void UWorldPartition::ForEachActorDesc(TSubclassOf<AActor> ActorClass, TFunction
 
 const FWorldPartitionActorDesc* UWorldPartition::GetActorDesc(const FGuid& Guid) const
 {
-	const TUniquePtr<FWorldPartitionActorDesc>* ActorDescObj = Actors.Find(Guid);
-	return ActorDescObj != nullptr ? ActorDescObj->Get() : nullptr;
+	const TUniquePtr<FWorldPartitionActorDesc>* const * ActorDesc = Actors.Find(Guid);
+	return ActorDesc ? (*ActorDesc)->Get() : nullptr;
 }
 
 FWorldPartitionActorDesc* UWorldPartition::GetActorDesc(const FGuid& Guid)
 {
-	TUniquePtr<FWorldPartitionActorDesc>* ActorDescObj = Actors.Find(Guid);
-	return ActorDescObj != nullptr ? ActorDescObj->Get() : nullptr;
+	TUniquePtr<FWorldPartitionActorDesc>** ActorDesc = Actors.Find(Guid);
+	return ActorDesc ? (*ActorDesc)->Get() : nullptr;
 }
 
 bool UWorldPartition::IsSimulating() const
@@ -538,7 +535,7 @@ bool UWorldPartition::UpdateEditorCells(TFunctionRef<bool(TArray<UWorldPartition
 	{
 		FWorldPartitionActorDesc* ActorDesc = Pair.Key;
 		// Only prompt if the actor will get unloaded by the unloading cells
-		if (ActorDesc->GetLoadedRefCount() == Pair.Value)
+		if (ActorDesc->GetHardRefCount() == Pair.Value)
 		{
 			AActor* LoadedActor = ActorDesc->GetActor();
 			check(LoadedActor);
@@ -643,7 +640,7 @@ void UWorldPartition::UpdateLoadingEditorCell(UWorldPartitionEditorCell* Cell, b
 	{
 		check(Actor);
 
-		const uint32 ActorRefCount = ActorDesc->RemoveLoadedRefCount();
+		const uint32 ActorRefCount = ActorDesc->DecHardRefCount();
 		UE_LOG(LogWorldPartition, Verbose, TEXT(" ==> Unreferenced loaded actor %s(%d) [UWorldPartition::UpdateLoadingEditorCell]"), *Actor->GetFullName(), ActorRefCount);
 
 		if (!ActorRefCount)
@@ -678,7 +675,7 @@ void UWorldPartition::UpdateLoadingEditorCell(UWorldPartitionEditorCell* Cell, b
 
 			// The actor could be either loaded but with no cells loaded (in the case of a reference from another actor, for example)
 			// or directly referenced by a loaded cell.
-			check(Actor || !ActorDesc->GetLoadedRefCount());
+			check(Actor || !ActorDesc->GetHardRefCount());
 
 			// Filter actor against DataLayers
 			if (ShouldActorBeLoaded(ActorDesc))
@@ -690,12 +687,12 @@ void UWorldPartition::UpdateLoadingEditorCell(UWorldPartitionEditorCell* Cell, b
 				{
 					// We already hold a reference to this actor
 					check(Actor);
-					check(ActorDesc->GetLoadedRefCount());
+					check(ActorDesc->GetHardRefCount());
 					UE_LOG(LogWorldPartition, Verbose, TEXT(" ==> Skipped already loaded actor %s"), *Actor->GetFullName());
 				}
 				else
 				{
-					const uint32 ActorRefCount = ActorDesc->AddLoadedRefCount();
+					const uint32 ActorRefCount = ActorDesc->IncHardRefCount();
 
 					if (ActorRefCount == 1)
 					{
@@ -795,14 +792,13 @@ void UWorldPartition::OnAssetAdded(const FAssetData& InAssetData)
 {
 	if (ShouldHandleAssetEvent(InAssetData))
 	{
-		TUniquePtr<FWorldPartitionActorDesc> NewActorDesc = GetActorDescriptor(InAssetData);
-		if (NewActorDesc.IsValid())
-		{
-			check(!Actors.Contains(NewActorDesc->GetGuid()));
+		TUniquePtr<FWorldPartitionActorDesc>* NewActorDesc = ActorDescList.Emplace(GetActorDescriptor(InAssetData));
+		check(NewActorDesc->IsValid());
 
-			HashActorDesc(NewActorDesc.Get());
-			Actors.Add(NewActorDesc->GetGuid(), MoveTemp(NewActorDesc));
-		}
+		check(!Actors.Contains((*NewActorDesc)->GetGuid()));
+		Actors.Add((*NewActorDesc)->GetGuid(), NewActorDesc);
+
+		HashActorDesc(NewActorDesc->Get());
 	}
 }
 
@@ -811,13 +807,14 @@ void UWorldPartition::OnAssetRemoved(const FAssetData& InAssetData)
 	if (ShouldHandleAssetEvent(InAssetData))
 	{
 		TUniquePtr<FWorldPartitionActorDesc> NewActorDesc = GetActorDescriptor(InAssetData);
-		if (NewActorDesc.IsValid())
-		{
-			TUniquePtr<FWorldPartitionActorDesc>& ExistingActorDesc = Actors.FindChecked(NewActorDesc->GetGuid());
+		check(NewActorDesc.IsValid());
 
-			UnhashActorDesc(ExistingActorDesc.Get());
-			Actors.Remove(ExistingActorDesc->GetGuid());
-		}
+		TUniquePtr<FWorldPartitionActorDesc>* ExistingActorDesc = Actors.FindChecked(NewActorDesc->GetGuid());
+
+		UnhashActorDesc(ExistingActorDesc->Get());
+
+		Actors.Remove((*ExistingActorDesc)->GetGuid());
+		ExistingActorDesc->Release();
 	}
 }
 
@@ -826,18 +823,18 @@ void UWorldPartition::OnAssetUpdated(const FAssetData& InAssetData)
 	if (ShouldHandleAssetEvent(InAssetData))
 	{
 		TUniquePtr<FWorldPartitionActorDesc> NewActorDesc = GetActorDescriptor(InAssetData);
-		if (NewActorDesc.IsValid())
-		{
-			TUniquePtr<FWorldPartitionActorDesc>& ExistingActorDesc = Actors.FindChecked(NewActorDesc->GetGuid());
+		check(NewActorDesc.IsValid());
+			
+		TUniquePtr<FWorldPartitionActorDesc>* ExistingActorDesc = Actors.FindChecked(NewActorDesc->GetGuid());
 
-			UnhashActorDesc(ExistingActorDesc.Get());
+		UnhashActorDesc(ExistingActorDesc->Get());
 
-			// This is required to support external load references
-			NewActorDesc->SetLoadedRefCount(ExistingActorDesc->GetLoadedRefCount());
+		// This is required to support external load references
+		NewActorDesc->TransferRefCounts(ExistingActorDesc->Get());
 
-			ExistingActorDesc = MoveTemp(NewActorDesc);
-			HashActorDesc(ExistingActorDesc.Get());
-		}
+		*ExistingActorDesc = MoveTemp(NewActorDesc);
+		
+		HashActorDesc(ExistingActorDesc->Get());
 	}
 }
 
@@ -861,6 +858,14 @@ bool UWorldPartition::ShouldHandleAssetEvent(const FAssetData& InAssetData)
 		return false;
 	}
 
+	// Make sure asset contains the required tags
+	static FName NAME_ActorMetaDataClass(TEXT("ActorMetaDataClass"));
+	static FName NAME_ActorMetaData(TEXT("ActorMetaData"));
+	if (!InAssetData.FindTag(NAME_ActorMetaDataClass) || !InAssetData.FindTag(NAME_ActorMetaData))
+	{
+		return false;
+	}
+
 	// Only handle assets that belongs to our level
 	auto RemoveAfterFirstDot = [](const FString& InValue)
 	{
@@ -1085,7 +1090,8 @@ FBox UWorldPartition::GetWorldBounds() const
 	FBox WorldBounds(ForceInit);
 	for (const auto& Pair : Actors)
 	{
-		const FWorldPartitionActorDesc* ActorDesc = Pair.Value.Get();
+		const FWorldPartitionActorDesc* ActorDesc = Pair.Value->Get();
+
 		switch (ActorDesc->GetGridPlacement())
 		{
 			case EActorGridPlacement::Location:
diff --git a/Engine/Source/Runtime/Engine/Private/WorldPartition/WorldPartitionActorDesc.cpp b/Engine/Source/Runtime/Engine/Private/WorldPartition/WorldPartitionActorDesc.cpp
--- a/Engine/Source/Runtime/Engine/Private/WorldPartition/WorldPartitionActorDesc.cpp
+++ b/Engine/Source/Runtime/Engine/Private/WorldPartition/WorldPartitionActorDesc.cpp
@@ -6,6 +6,7 @@
 #include "Misc/HashBuilder.h"
 #include "UObject/LinkerInstancingContext.h"
 #include "UObject/UObjectHash.h"
+#include "UObject/MetaData.h"
 #include "Algo/Transform.h"
 #include "Algo/RemoveIf.h"
 #include "Engine/World.h"
@@ -21,7 +22,8 @@
 uint32 FWorldPartitionActorDesc::GlobalTag = 0;
 
 FWorldPartitionActorDesc::FWorldPartitionActorDesc()
-	: LoadedRefCount(0)
+	: SoftRefCount(0)
+	, HardRefCount(0)
 	, Tag(0)
 {}
 
@@ -118,7 +120,7 @@ void FWorldPartitionActorDesc::SerializeTo(TArray<uint8>& OutData)
 
 void FWorldPartitionActorDesc::TransformInstance(const FString& From, const FString& To, const FTransform& InstanceTransform)
 {
-	check(!LoadedRefCount);
+	check(!HardRefCount);
 
 	ActorPath = *ActorPath.ToString().Replace(*From, *To);
 
@@ -170,39 +172,67 @@ FBox FWorldPartitionActorDesc::GetBounds() const
 
 AActor* FWorldPartitionActorDesc::GetActor() const
 {
-	return FindObject<AActor>(nullptr, *ActorPath.ToString());
+	AActor* Actor = ActorPtr.Get();
+
+	if (Actor)
+	{
+		return Actor;
+	}
+
+	Actor = FindObject<AActor>(nullptr, *ActorPath.ToString());
+
+	if (Actor)
+	{
+		ActorPtr = Actor;
+	}
+
+	return Actor;
 }
 
 AActor* FWorldPartitionActorDesc::Load(const FLinkerInstancingContext* InstancingContext) const
 {
-	UPackage* Package = nullptr;
+	AActor* Actor = GetActor();
 
-	if (InstancingContext)
+	if (!Actor)
 	{
-		FName RemappedPackageName = InstancingContext->Remap(ActorPackage);
-		check(RemappedPackageName != ActorPath);
+		UPackage* Package = nullptr;
 
-		Package = CreatePackage(*RemappedPackageName.ToString());
-	}
+		if (InstancingContext)
+		{
+			FName RemappedPackageName = InstancingContext->Remap(ActorPackage);
+			check(RemappedPackageName != ActorPath);
 
-	if (LoadPackage(Package, *ActorPackage.ToString(), LOAD_None, nullptr, InstancingContext))
-	{
-		return GetActor();
+			Package = CreatePackage(*RemappedPackageName.ToString());
+		}
+
+		Package = LoadPackage(Package, *ActorPackage.ToString(), LOAD_None, nullptr, InstancingContext);
+
+		if (Package)
+		{
+			Actor = GetActor();
+			check(Actor);
+		}
 	}
 
-	return nullptr;
+	return Actor;
 }
 
 void FWorldPartitionActorDesc::Unload()
 {
-	AActor* Actor = GetActor();
-	if (Actor && Actor->IsPackageExternal())
+	if (AActor* Actor = GetActor())
 	{
+		check(Actor->IsPackageExternal());
+
 		ForEachObjectWithPackage(Actor->GetPackage(), [](UObject* Object)
 		{
-			Object->ClearFlags(RF_Public | RF_Standalone);
+			if (Object->HasAnyFlags(RF_Public | RF_Standalone))
+			{
+				CastChecked<UMetaData>(Object)->ClearFlags(RF_Public | RF_Standalone);
+			}
 			return true;
 		}, false);
+
+		ActorPtr = nullptr;
 	}
 }
 #endif
\ No newline at end of file
diff --git a/Engine/Source/Runtime/Engine/Private/WorldPartition/WorldPartitionEditorCell.cpp b/Engine/Source/Runtime/Engine/Private/WorldPartition/WorldPartitionEditorCell.cpp
--- a/Engine/Source/Runtime/Engine/Private/WorldPartition/WorldPartitionEditorCell.cpp
+++ b/Engine/Source/Runtime/Engine/Private/WorldPartition/WorldPartitionEditorCell.cpp
@@ -36,7 +36,7 @@ void UWorldPartitionEditorCell::AddActor(FWorldPartitionActorDesc* InActorDesc)
 			LoadedActors.Add(InActorDesc, &bIsAlreadyInSet);
 			check(!bIsAlreadyInSet);
 
-			const uint32 ActorRefCount = InActorDesc->AddLoadedRefCount();
+			const uint32 ActorRefCount = InActorDesc->IncHardRefCount();
 			UE_LOG(LogWorldPartition, Verbose, TEXT(" ==> Referenced loaded actor %s(%d) [UWorldPartitionEditorCell::AddActor]"), *Actor->GetFullName(), ActorRefCount);
 		}
 	}
@@ -51,7 +51,7 @@ void UWorldPartitionEditorCell::RemoveActor(FWorldPartitionActorDesc* InActorDes
 	{
 		check(bLoaded);
 
-		const uint32 ActorRefCount = InActorDesc->RemoveLoadedRefCount();
+		const uint32 ActorRefCount = InActorDesc->DecHardRefCount();
 		UE_LOG(LogWorldPartition, Verbose, TEXT(" ==> Unreferenced loaded actor %s(%d) [UWorldPartitionEditorCell::RemoveActor]"), *InActorDesc->GetActor()->GetFullName(), ActorRefCount);
 	}
 }
diff --git a/Engine/Source/Runtime/Engine/Private/WorldPartition/WorldPartitionHandle.cpp b/Engine/Source/Runtime/Engine/Private/WorldPartition/WorldPartitionHandle.cpp
new file mode 100644
--- /dev/null
+++ b/Engine/Source/Runtime/Engine/Private/WorldPartition/WorldPartitionHandle.cpp
@@ -0,0 +1,62 @@
+// Copyright Epic Games, Inc. All Rights Reserved.
+
+#include "WorldPartition/WorldPartitionHandle.h"
+
+#if WITH_EDITOR
+void FWorldPartitionHandleImpl::IncRefCount(FWorldPartitionActorDesc* ActorDesc)
+{
+	ActorDesc->IncSoftRefCount();
+}
+
+void FWorldPartitionHandleImpl::DecRefCount(FWorldPartitionActorDesc* ActorDesc)
+{
+	ActorDesc->DecSoftRefCount();
+}
+
+void FWorldPartitionReferenceImpl::IncRefCount(FWorldPartitionActorDesc* ActorDesc)
+{
+	if (ActorDesc->IncHardRefCount() == 1)
+	{
+		AActor* Actor = ActorDesc->Load();
+		Actor->GetLevel()->AddLoadedActor(Actor);
+	}
+}
+
+void FWorldPartitionReferenceImpl::DecRefCount(FWorldPartitionActorDesc* ActorDesc)
+{
+	if (!ActorDesc->DecHardRefCount())
+	{
+		if (AActor* Actor = ActorDesc->GetActor())
+		{
+			Actor->GetLevel()->RemoveLoadedActor(Actor);
+			ActorDesc->Unload();
+		}
+	}
+}
+
+FWorldPartitionReference FWorldPartitionHandleHelpers::ConvertHandleToReference(const FWorldPartitionHandle& Handle)
+{
+	FWorldPartitionReference Reference;
+
+	if (Handle.IsValid())
+	{
+		Reference.ActorDesc = Handle.ActorDesc;
+		Reference.IncRefCount();
+	}
+
+	return Reference;
+}
+
+FWorldPartitionHandle FWorldPartitionHandleHelpers::ConvertReferenceToHandle(const FWorldPartitionReference& Reference)
+{
+	FWorldPartitionHandle Handle;
+
+	if (Reference.IsValid())
+	{
+		Handle.ActorDesc = Reference.ActorDesc;
+		Handle.IncRefCount();
+	}
+
+	return Handle;
+}
+#endif
\ No newline at end of file
diff --git a/Engine/Source/Runtime/Engine/Private/WorldPartition/WorldPartitionRuntimeSpatialHash.cpp b/Engine/Source/Runtime/Engine/Private/WorldPartition/WorldPartitionRuntimeSpatialHash.cpp
--- a/Engine/Source/Runtime/Engine/Private/WorldPartition/WorldPartitionRuntimeSpatialHash.cpp
+++ b/Engine/Source/Runtime/Engine/Private/WorldPartition/WorldPartitionRuntimeSpatialHash.cpp
@@ -8,6 +8,7 @@
 #include "WorldPartition/WorldPartitionActorDesc.h"
 #include "WorldPartition/WorldPartitionActorDescIterator.h"
 #include "WorldPartition/WorldPartitionStreamingPolicy.h"
+#include "WorldPartition/WorldPartitionHandle.h"
 #include "WorldPartition/DataLayer/DataLayerSubsystem.h"
 #include "WorldPartition/DataLayer/WorldDataLayers.h"
 #include "WorldPartition/DataLayer/DataLayer.h"
@@ -30,6 +31,7 @@
 #include "DistanceFieldAtlas.h"
 #include "MeshCardRepresentation.h"
 
+#include "WorldPartition/HLOD/HLODBuilder.h"
 #include "WorldPartition/HLOD/HLODLayer.h"
 #include "WorldPartition/HLOD/HLODActor.h"
 #include "WorldPartition/HLOD/HLODActorDesc.h"
@@ -253,7 +255,7 @@ TArray<FActorCluster> CreateActorClustersImpl(UWorldPartition* WorldPartition, T
 
 	for (const auto& Pair : WorldPartition->Actors)
 	{
-		const FWorldPartitionActorDesc* ActorDesc = Pair.Value.Get();
+		const FWorldPartitionActorDesc* ActorDesc = Pair.Value->Get();
 		EActorGridPlacement GridPlacement = ActorDesc->GetGridPlacement();
 
 		// Check if the actor is loaded (potentially referenced by the level script)
@@ -1578,21 +1580,8 @@ TArray<FGuid> GenerateHLODsForGrid(UWorldPartition* WorldPartition, const FSpati
 				TArray<AActor*> CellActors;
 				for (const FGuid& ActorGuid: GridCellDataChunk->GetActors())
 				{
-					FWorldPartitionActorDesc* ActorDesc = (FWorldPartitionActorDesc*)WorldPartition->GetActorDesc(ActorGuid);
-
-					check(!ActorDesc->GetLoadedRefCount());
-					ActorDesc->AddLoadedRefCount();
-
-					AActor* Actor = ActorDesc->GetActor();
-					
-					if (!Actor)
-					{
-						Actor = ActorDesc->Load();
-						check(Actor);
-					}
-
-					Actor->GetLevel()->AddLoadedActor(Actor);
-
+					FWorldPartitionReference& ActorRef = Context.ActorReferences.Emplace_GetRef(WorldPartition, ActorGuid);
+					FWorldPartitionActorDesc* ActorDesc  = ActorRef.Get();
 					CellActors.Add(ActorDesc->GetActor());
 				}
 
@@ -1651,27 +1640,12 @@ TArray<FGuid> GenerateHLODsForGrid(UWorldPartition* WorldPartition, const FSpati
 					// Update newly created HLOD actors
 					for (AWorldPartitionHLOD* CellHLODActor: NewCellHLODActors)
 					{
-						FWorldPartitionActorDesc* ActorDesc = (FWorldPartitionActorDesc*)WorldPartition->GetActorDesc(CellHLODActor->GetActorGuid());
-					
-						check(!ActorDesc->GetLoadedRefCount());
-						ActorDesc->AddLoadedRefCount();
+						Context.ActorReferences.Emplace(WorldPartition, CellHLODActor->GetActorGuid());
 					}
 				}
 
 				// Unload actors
-				TArray<AActor*> ToUnloadActors = MoveTemp(CellActors);
-				ToUnloadActors.Append(CellHLODActors);
-
-				for (AActor* Actor: ToUnloadActors)
-				{
-					FWorldPartitionActorDesc* ActorDesc = (FWorldPartitionActorDesc*)WorldPartition->GetActorDesc(Actor->GetActorGuid());
-
-					check(ActorDesc->GetLoadedRefCount() == 1);
-					ActorDesc->RemoveLoadedRefCount();
-
-					Actor->GetLevel()->RemoveLoadedActor(Actor);
-					ActorDesc->Unload();
-				}
+				Context.ActorReferences.Empty();
 			}
 		}
 	});
@@ -1901,7 +1875,7 @@ bool UWorldPartitionRuntimeSpatialHash::GenerateHLOD(ISourceControlHelper* Sourc
 	{
 		if (HLODIterator->GetCellHash())
 		{
-			Context.HLODActorDescs.Add(HLODIterator->GetCellHash(), HLODIterator->GetGuid());
+			Context.HLODActorDescs.Add(HLODIterator->GetCellHash(), FWorldPartitionHandle(WorldPartition, HLODIterator->GetGuid()));
 		}
 	}
 
@@ -1980,7 +1954,7 @@ bool UWorldPartitionRuntimeSpatialHash::GenerateHLOD(ISourceControlHelper* Sourc
 	// Destroy all unreferenced HLOD actors
 	for (const auto& HLODActorPair: Context.HLODActorDescs)
 	{
-		FWorldPartitionActorDesc* HLODActorDesc = WorldPartition->GetActorDesc(HLODActorPair.Value);
+		FWorldPartitionActorDesc* HLODActorDesc = HLODActorPair.Value.Get();
 		check(HLODActorDesc);
 
 		const FString ActorPackageFilename = SourceControlHelper->GetFilename(HLODActorDesc->GetActorPackage().ToString());		
diff --git a/Engine/Source/Runtime/Engine/Public/WorldPartition/WorldPartition.h b/Engine/Source/Runtime/Engine/Public/WorldPartition/WorldPartition.h
--- a/Engine/Source/Runtime/Engine/Public/WorldPartition/WorldPartition.h
+++ b/Engine/Source/Runtime/Engine/Public/WorldPartition/WorldPartition.h
@@ -60,15 +60,65 @@ public:
 	virtual bool Delete(UPackage* Package) const =0;
 };
 
-struct ENGINE_API FHLODGenerationContext
+/**
+* Simple chunked array that creates chunks to avoid invalidating elements when adding.
+* This will also never release elements as handles/references can still point to them.
+**/
+class FChunkedActorDescArray
 {
-	TMap<uint64, FGuid> HLODActorDescs;
-	
-	// Everything needed to build the cell hash
-	int64 GridIndexX;
-	int64 GridIndexY;
-	int64 GridIndexZ;
-	FName HLODLayerName;
+	typedef TUniquePtr<FWorldPartitionActorDesc> ElemType;
+
+	enum
+	{
+		NumElementsPerChunk = 1024,
+	};
+
+	inline ElemType* AllocInstance()
+	{
+		const uint32 ChunkIndex = NumElements / NumElementsPerChunk;
+		const uint32 ElemIndex = NumElements % NumElementsPerChunk;
+
+		if (!Chunks.IsValidIndex(ChunkIndex))
+		{
+			Chunks.Add((ElemType*)FMemory::Malloc(sizeof(ElemType) * NumElementsPerChunk));
+		}
+
+		NumElements++;
+		return Chunks[ChunkIndex] + ElemIndex;
+	}
+
+public:
+	inline FChunkedActorDescArray()
+		: NumElements(0)
+	{}
+
+	inline ~FChunkedActorDescArray()
+	{
+		int32 CurElement = 0;
+		for (ElemType* Chunk : Chunks)
+		{
+			const int32 NumElemToDestroy = FMath::Min((int32)NumElementsPerChunk, NumElements - CurElement);
+
+			for (int32 i=0; i<NumElemToDestroy; i++)
+			{
+				Chunk[i].~ElemType();
+			}
+
+			FMemory::Free(Chunk);
+
+			CurElement += NumElementsPerChunk;
+		}
+	}
+
+	template <typename... ArgsType>
+	inline ElemType* Emplace(ArgsType&&... Args)
+	{
+		return new(AllocInstance()) ElemType(Forward<ArgsType>(Args)...);
+	}
+
+private:
+	TArray<ElemType*> Chunks;
+	int32 NumElements;
 };
 #endif
 
@@ -176,7 +226,8 @@ public:
 	UWorldPartitionRuntimeHash* RuntimeHash;
 
 #if WITH_EDITOR
-	TMap<FGuid, TUniquePtr<FWorldPartitionActorDesc>> Actors;
+	FChunkedActorDescArray ActorDescList;
+	TMap<FGuid, TUniquePtr<FWorldPartitionActorDesc>*> Actors;
 	
 	DECLARE_EVENT_TwoParams(UWorldPartition, FWorldPartitionActorRegisteredEvent, AActor&, bool);
 	FWorldPartitionActorRegisteredEvent OnActorRegisteredEvent;
diff --git a/Engine/Source/Runtime/Engine/Public/WorldPartition/WorldPartitionActorDesc.h b/Engine/Source/Runtime/Engine/Public/WorldPartition/WorldPartitionActorDesc.h
--- a/Engine/Source/Runtime/Engine/Public/WorldPartition/WorldPartitionActorDesc.h
+++ b/Engine/Source/Runtime/Engine/Public/WorldPartition/WorldPartitionActorDesc.h
@@ -46,20 +46,36 @@ public:
 	inline FName GetActorPath() const { return ActorPath; }
 	FBox GetBounds() const;
 
-	inline uint32 AddLoadedRefCount() const
+	inline uint32 IncSoftRefCount() const
 	{
-		return ++LoadedRefCount;
+		return ++SoftRefCount;
 	}
 
-	inline uint32 RemoveLoadedRefCount() const
+	inline uint32 DecSoftRefCount() const
 	{
-		check(LoadedRefCount > 0);
-		return --LoadedRefCount;
+		check(SoftRefCount > 0);
+		return --SoftRefCount;
 	}
 
-	inline uint32 GetLoadedRefCount() const
+	inline uint32 GetSoftRefCount() const
 	{
-		return LoadedRefCount;
+		return SoftRefCount;
+	}
+
+	inline uint32 IncHardRefCount() const
+	{
+		return ++HardRefCount;
+	}
+
+	inline uint32 DecHardRefCount() const
+	{
+		check(HardRefCount > 0);
+		return --HardRefCount;
+	}
+
+	inline uint32 GetHardRefCount() const
+	{
+		return HardRefCount;
 	}
 
 	const TArray<FGuid>& GetReferences() const
@@ -83,34 +99,37 @@ public:
 protected:
 	FWorldPartitionActorDesc();
 
-	inline void SetLoadedRefCount(uint32 InLoadedRefCount) const
+	inline void TransferRefCounts(const FWorldPartitionActorDesc* From) const
 	{
-		LoadedRefCount = InLoadedRefCount;
+		SoftRefCount = From->SoftRefCount;
+		HardRefCount = From->HardRefCount;
 	}
 
 	virtual void Serialize(FArchive& Ar);
 
-	FGuid						Guid;
-	FName						Class;
-	FName						ActorPackage;
-	FName						ActorPath;
-	FVector						BoundsLocation;
-	FVector						BoundsExtent;
-	EActorGridPlacement			GridPlacement;
-	FName						RuntimeGrid;
-	bool						bActorIsEditorOnly;
-	bool						bLevelBoundsRelevant;
-	TArray<FName>				DataLayers;
-	TArray<FGuid>				References;
+	FGuid							Guid;
+	FName							Class;
+	FName							ActorPackage;
+	FName							ActorPath;
+	FVector							BoundsLocation;
+	FVector							BoundsExtent;
+	EActorGridPlacement				GridPlacement;
+	FName							RuntimeGrid;
+	bool							bActorIsEditorOnly;
+	bool							bLevelBoundsRelevant;
+	TArray<FName>					DataLayers;
+	TArray<FGuid>					References;
 	
-	mutable uint32				LoadedRefCount;
+	mutable uint32					SoftRefCount;
+	mutable uint32					HardRefCount;
 
 	// Cached values
-	UClass*						ActorClass;
+	UClass*							ActorClass;
+	mutable TWeakObjectPtr<AActor>	ActorPtr;
 
 public:
 	// Tagging
-	mutable uint32				Tag;
-	static uint32				GlobalTag;
+	mutable uint32					Tag;
+	static uint32					GlobalTag;
 #endif
 };
diff --git a/Engine/Source/Runtime/Engine/Public/WorldPartition/WorldPartitionActorDescIterator.h b/Engine/Source/Runtime/Engine/Public/WorldPartition/WorldPartitionActorDescIterator.h
--- a/Engine/Source/Runtime/Engine/Public/WorldPartition/WorldPartitionActorDescIterator.h
+++ b/Engine/Source/Runtime/Engine/Public/WorldPartition/WorldPartitionActorDescIterator.h
@@ -37,7 +37,7 @@ public:
 	 */
 	FORCEINLINE ActorDescType* operator*() const
 	{
-		return StaticCast<ActorDescType*>(ActorsIterator->Value.Get());
+		return StaticCast<ActorDescType*>(ActorsIterator->Value->Get());
 	}
 
 	/**
@@ -47,7 +47,7 @@ public:
 	 */
 	FORCEINLINE ActorDescType* operator->() const
 	{
-		return StaticCast<ActorDescType*>(ActorsIterator->Value.Get());
+		return StaticCast<ActorDescType*>(ActorsIterator->Value->Get());
 	}
 	/**
 	 * Returns whether the iterator has reached the end and no longer points
@@ -62,7 +62,7 @@ public:
 
 private:
 	/**
-	 * Determines whether the iterator cuirrently points to a valid actor desc or not.
+	 * Determines whether the iterator currently points to a valid actor desc or not.
 	 * @return	true
 	 */
 	FORCEINLINE bool ShouldSkip() const
@@ -72,10 +72,10 @@ private:
 			return false;
 		}
 
-		return !ActorsIterator->Value.Get()->GetActorClass()->IsChildOf(ActorType::StaticClass());
+		return !ActorsIterator->Value->Get()->GetActorClass()->IsChildOf(ActorType::StaticClass());
 	}
 
 private:
-	TMap<FGuid, TUniquePtr<FWorldPartitionActorDesc>>::TIterator ActorsIterator;
+	TMap<FGuid, TUniquePtr<FWorldPartitionActorDesc>*>::TIterator ActorsIterator;
 };
 #endif
\ No newline at end of file
diff --git a/Engine/Source/Runtime/Engine/Public/WorldPartition/WorldPartitionHandle.h b/Engine/Source/Runtime/Engine/Public/WorldPartition/WorldPartitionHandle.h
new file mode 100644
--- /dev/null
+++ b/Engine/Source/Runtime/Engine/Public/WorldPartition/WorldPartitionHandle.h
@@ -0,0 +1,133 @@
+// Copyright Epic Games, Inc. All Rights Reserved.
+#pragma once
+
+#include "CoreMinimal.h"
+#include "Templates/UniquePtr.h"
+#include "WorldPartition/WorldPartition.h"
+
+#if WITH_EDITOR
+template <typename Impl>
+class TWorldPartitionHandleBase
+{
+	friend struct FWorldPartitionHandleHelpers;
+
+public:
+	TWorldPartitionHandleBase()
+		: ActorDesc(nullptr)
+	{}
+
+	TWorldPartitionHandleBase(UWorldPartition* WorldPartition, const FGuid& ActorGuid)
+	{
+		if (TUniquePtr<FWorldPartitionActorDesc>** ActorDescPtr = WorldPartition->Actors.Find(ActorGuid))
+		{
+			ActorDesc = *ActorDescPtr;
+
+			if (IsValid())
+			{
+				IncRefCount();
+			}
+		}
+		else
+		{
+			ActorDesc = nullptr;
+		}
+	}
+
+	~TWorldPartitionHandleBase()
+	{
+		if (IsValid())
+		{
+			DecRefCount();
+		}
+	}
+
+	template <typename T>
+	TWorldPartitionHandleBase(const TWorldPartitionHandleBase& Other)
+	{
+		*this = Other;
+	}
+
+	template <typename T>
+	TWorldPartitionHandleBase& operator=(const T& Other)
+	{
+		ActorDesc = Other.ActorDesc;
+
+		if (IsValid())
+		{
+			IncRefCount();
+		}
+
+		return *this;
+	}
+
+	template <typename T>
+	TWorldPartitionHandleBase(T&& Other)
+	{
+		*this = Other;
+	}
+
+	template <typename T>
+	TWorldPartitionHandleBase& operator=(T&& Other)
+	{
+		ActorDesc = Other.ActorDesc;
+
+		if (IsValid())
+		{
+			IncRefCount();
+			Other.DecRefCount();
+			Other.ActorDesc = nullptr;
+		}
+
+		return *this;
+	}
+
+	inline FWorldPartitionActorDesc* operator->() const
+	{
+		return Get();
+	}
+
+	inline bool IsValid() const
+	{
+		return ActorDesc && ActorDesc->IsValid();
+	}
+
+	inline FWorldPartitionActorDesc* Get() const
+	{
+		return IsValid() ? ActorDesc->Get() : nullptr;
+	}
+
+protected:
+	inline void IncRefCount()
+	{
+		Impl::IncRefCount(ActorDesc->Get());
+	}
+
+	inline void DecRefCount()
+	{
+		Impl::DecRefCount(ActorDesc->Get());
+	}
+
+	TUniquePtr<FWorldPartitionActorDesc>* ActorDesc;
+};
+
+struct FWorldPartitionHandleImpl
+{
+	static void IncRefCount(FWorldPartitionActorDesc* ActorDesc);
+	static void DecRefCount(FWorldPartitionActorDesc* ActorDesc);
+};
+
+struct FWorldPartitionReferenceImpl
+{
+	static void IncRefCount(FWorldPartitionActorDesc* ActorDesc);
+	static void DecRefCount(FWorldPartitionActorDesc* ActorDesc);
+};
+
+typedef TWorldPartitionHandleBase<FWorldPartitionHandleImpl> FWorldPartitionHandle;
+typedef TWorldPartitionHandleBase<FWorldPartitionReferenceImpl> FWorldPartitionReference;
+
+struct FWorldPartitionHandleHelpers
+{
+	static FWorldPartitionReference ConvertHandleToReference(const FWorldPartitionHandle& Handle);
+	static FWorldPartitionHandle ConvertReferenceToHandle(const FWorldPartitionReference& Handle);
+};
+#endif
\ No newline at end of file
