Removed a spacing in file "Texture2DArrayResource.cpp"

diff --git a/Engine/Source/Runtime/Engine/Private/Rendering/Texture2DArrayResource .cpp b/Engine/Source/Runtime/Engine/Private/Rendering/Texture2DArrayResource .cpp
deleted file mode 100644
diff --git a/Engine/Source/Runtime/Engine/Private/Rendering/Texture2DArrayResource.cpp b/Engine/Source/Runtime/Engine/Private/Rendering/Texture2DArrayResource.cpp
new file mode 100644
--- /dev/null
+++ b/Engine/Source/Runtime/Engine/Private/Rendering/Texture2DArrayResource.cpp
@@ -0,0 +1,135 @@
+// Copyright Epic Games, Inc. All Rights Reserved.
+
+/*=============================================================================
+	Texture2DArrayResource.cpp: Implementation of FTexture2DArrayResource used  by streamable UTexture2DArray.
+=============================================================================*/
+
+#include "Rendering/Texture2DArrayResource.h"
+#include "Engine/Texture2DArray.h"
+#include "DeviceProfiles/DeviceProfile.h"
+#include "DeviceProfiles/DeviceProfileManager.h"
+#include "RenderUtils.h"
+
+//*****************************************************************************
+//************************* FTexture2DArrayResource ***************************
+//*****************************************************************************
+
+FTexture2DArrayResource::FTexture2DArrayResource(UTexture2DArray* InOwner, const FStreamableRenderResourceState& InState) 
+: FStreamableTextureResource(InOwner, InOwner->PlatformData, InState, false)
+{
+	AddressU = InOwner->AddressX == TA_Wrap ? AM_Wrap : (InOwner->AddressX == TA_Clamp ? AM_Clamp : AM_Mirror);
+	AddressV = InOwner->AddressY == TA_Wrap ? AM_Wrap : (InOwner->AddressY == TA_Clamp ? AM_Clamp : AM_Mirror);
+	AddressW = InOwner->AddressZ == TA_Wrap ? AM_Wrap : (InOwner->AddressZ == TA_Clamp ? AM_Clamp : AM_Mirror);
+
+	TArray<uint64> MipOffsets;
+	MipOffsets.AddZeroed(State.MaxNumLODs);
+
+	uint64 InitialMipDataSize = 0;
+	TArrayView<const FTexture2DMipMap*> MipsView = GetPlatformMipsView();
+	for (int32 MipIdx = State.RequestedFirstLODIdx(); MipIdx < State.MaxNumLODs; ++MipIdx)
+	{
+		const FTexture2DMipMap& Mip = *MipsView[MipIdx];
+		MipOffsets[MipIdx] = InitialMipDataSize;
+		InitialMipDataSize += Mip.BulkData.GetBulkDataSize();
+	}
+	InitialMipData.Reset(new uint8[InitialMipDataSize]);
+
+	// Resize each structure correctly per slice, since we access it per mip after.
+	SliceMipDataViews.AddDefaulted(SizeZ);
+	for (uint32 SliceIdx = 0; SliceIdx < SizeZ; ++SliceIdx)
+	{
+		SliceMipDataViews[SliceIdx].AddDefaulted(State.MaxNumLODs);
+	}
+
+	for (int32 MipIdx = State.RequestedFirstLODIdx(); MipIdx < State.MaxNumLODs; ++MipIdx)
+	{
+		FTexture2DMipMap& Mip = const_cast<FTexture2DMipMap&>(*MipsView[MipIdx]);
+		if (Mip.BulkData.GetBulkDataSize() > 0)
+		{
+			const uint64 SliceMipDataSize = Mip.BulkData.GetBulkDataSize() / SizeZ;
+			const uint8* SrcData = (const uint8*)Mip.BulkData.Lock(LOCK_READ_ONLY);
+
+			for (uint32 SliceIdx = 0; SliceIdx < SizeZ; ++SliceIdx)
+			{
+				TArrayView<uint8>& SliceMipDataView = SliceMipDataViews[SliceIdx][MipIdx];
+				SliceMipDataView = TArrayView<uint8>(InitialMipData.Get() + MipOffsets[MipIdx] + SliceIdx * SliceMipDataSize, SliceMipDataSize);
+				FMemory::Memcpy(SliceMipDataView.GetData(), SrcData + SliceIdx * SliceMipDataSize, SliceMipDataSize);
+			}
+
+			Mip.BulkData.Unlock();
+		}
+		else
+		{
+			UE_LOG(LogTexture, Error, TEXT("Corrupt texture [%s]! Missing bulk data for MipIndex=%d"), *TextureName.ToString(), MipIdx);
+		}
+	}
+}
+
+void FTexture2DArrayResource::CreateTexture()
+{
+	const int32 RequestedFirstLODIdx = State.RequestedFirstLODIdx();
+	TArrayView<const FTexture2DMipMap*> MipsView = GetPlatformMipsView();
+	const FTexture2DMipMap& FirstMip = *MipsView[RequestedFirstLODIdx];
+
+	FRHIResourceCreateInfo CreateInfo;
+	TRefCountPtr<FRHITexture2DArray> TextureArray = RHICreateTexture2DArray(FirstMip.SizeX, FirstMip.SizeY, FirstMip.SizeZ, PixelFormat, State.NumRequestedLODs, 1, CreationFlags, CreateInfo);
+	TextureRHI = TextureArray;
+
+	const int32 BlockSizeX = GPixelFormats[PixelFormat].BlockSizeX;
+	const int32 BlockSizeY = GPixelFormats[PixelFormat].BlockSizeY;
+	const int32 BlockBytes = GPixelFormats[PixelFormat].BlockBytes;
+
+	// Read the initial cached mip levels into the RHI texture.
+
+	for (int32 RHIMipIdx = 0; RHIMipIdx < State.NumRequestedLODs; ++RHIMipIdx)
+	{
+		const int32 MipIdx = RHIMipIdx + RequestedFirstLODIdx;
+		const uint32 MipSizeX = FMath::Max<int32>(SizeX >> MipIdx, 1);
+		const uint32 MipSizeY = FMath::Max<int32>(SizeY >> MipIdx, 1);
+
+		uint32 NumColumns = FMath::DivideAndRoundUp<int32>(MipSizeX, BlockSizeX);
+		uint32 NumRows = FMath::DivideAndRoundUp<int32>(MipSizeY, BlockSizeY);
+		if (PixelFormat == PF_PVRTC2 || PixelFormat == PF_PVRTC4)
+		{
+			// PVRTC has minimum 2 blocks width and height
+			NumColumns = FMath::Max<uint32>(NumColumns, 2);
+			NumRows = FMath::Max<uint32>(NumRows, 2);
+		}
+		const uint32 SrcStride = NumColumns * BlockBytes;
+
+		for (uint32 SliceIdx = 0; SliceIdx < SizeZ; ++SliceIdx)
+		{
+			uint32 DestStride = 0;
+			void* DestData = RHILockTexture2DArray(TextureArray, SliceIdx, RHIMipIdx, RLM_WriteOnly, DestStride, false);
+			if (DestData)
+			{
+				const TArrayView<uint8>& SliceMipDataView = SliceMipDataViews[SliceIdx][MipIdx];
+				CopyTextureData2D(SliceMipDataView.GetData(), DestData, MipSizeY, PixelFormat, SrcStride, DestStride);
+			}
+			RHIUnlockTexture2DArray(TextureArray, SliceIdx, RHIMipIdx, false);
+		}
+	}
+		
+	SliceMipDataViews.Empty();
+	InitialMipData.Release();
+}
+
+void FTexture2DArrayResource::CreatePartiallyResidentTexture()
+{
+	unimplemented();
+	TextureRHI.SafeRelease();
+}
+
+uint64 FTexture2DArrayResource::GetPlatformMipsSize(uint32 NumMips) const
+{
+	if (PlatformData && NumMips > 0)
+	{
+		const FIntPoint MipExtents = CalcMipMapExtent(SizeX, SizeY, PixelFormat, State.LODCountToFirstLODIdx(NumMips));
+		uint32 TextureAlign = 0;
+		return RHICalcTexture2DArrayPlatformSize(MipExtents.X, MipExtents.Y, SizeZ, PixelFormat, NumMips, 1, CreationFlags, FRHIResourceCreateInfo(PlatformData->GetExtData()), TextureAlign);
+	}
+	else
+	{
+		return 0;
+	}
+}
