Make static container thread local to store results from GetFormatVersionMap().

#rnx

diff --git a/Engine/Source/Runtime/Engine/Private/ShaderCompiler/ShaderCompiler.cpp b/Engine/Source/Runtime/Engine/Private/ShaderCompiler/ShaderCompiler.cpp
--- a/Engine/Source/Runtime/Engine/Private/ShaderCompiler/ShaderCompiler.cpp
+++ b/Engine/Source/Runtime/Engine/Private/ShaderCompiler/ShaderCompiler.cpp
@@ -760,7 +760,7 @@ bool FShaderCompileUtilities::DoWriteTasks(const TArray<FShaderCommonCompileJobP
 	int32 InputVersion = ShaderCompileWorkerInputVersion;
 	TransferFile << InputVersion;
 
-	static TMap<FString, uint32> FormatVersionMap;
+	thread_local static TMap<FString, uint32> FormatVersionMap;
 	GetFormatVersionMap(FormatVersionMap);
 
 	TransferFile << FormatVersionMap;
