Exposed CachedAnimData state machine weight info to blueprint

diff --git a/Engine/Source/Runtime/Engine/Private/Animation/CachedAnimDataLibrary.cpp b/Engine/Source/Runtime/Engine/Private/Animation/CachedAnimDataLibrary.cpp
new file mode 100644
--- /dev/null
+++ b/Engine/Source/Runtime/Engine/Private/Animation/CachedAnimDataLibrary.cpp
@@ -0,0 +1,38 @@
+// Copyright Epic Games, Inc. All Rights Reserved.
+
+#include "Animation/CachedAnimDataLibrary.h"
+
+UCachedAnimDataLibrary::UCachedAnimDataLibrary(const FObjectInitializer& ObjectInitializer)
+	: Super(ObjectInitializer)
+{
+}
+
+bool UCachedAnimDataLibrary::StateMachine_IsStateRelevant(UAnimInstance* InAnimInstance, const FCachedAnimStateData& CachedAnimStateData)
+{
+	if (InAnimInstance)
+	{
+		return CachedAnimStateData.IsRelevant(*InAnimInstance);
+	}
+
+	return false;
+}
+
+float UCachedAnimDataLibrary::StateMachine_GetLocalWeight(UAnimInstance* InAnimInstance, const FCachedAnimStateData& CachedAnimStateData)
+{
+	if (InAnimInstance)
+	{
+		return CachedAnimStateData.GetWeight(*InAnimInstance);
+	}
+
+	return 0.0f;
+}
+
+float UCachedAnimDataLibrary::StateMachine_GetGlobalWeight(UAnimInstance* InAnimInstance, const FCachedAnimStateData& CachedAnimStateData)
+{
+	if (InAnimInstance)
+	{
+		return CachedAnimStateData.GetGlobalWeight(*InAnimInstance);
+	}
+
+	return 0.0f;
+}
diff --git a/Engine/Source/Runtime/Engine/Public/Animation/CachedAnimDataLibrary.h b/Engine/Source/Runtime/Engine/Public/Animation/CachedAnimDataLibrary.h
new file mode 100644
--- /dev/null
+++ b/Engine/Source/Runtime/Engine/Public/Animation/CachedAnimDataLibrary.h
@@ -0,0 +1,31 @@
+// Copyright Epic Games, Inc. All Rights Reserved.
+
+#pragma once
+
+#include "CoreMinimal.h"
+#include "Animation/CachedAnimData.h"
+#include "CachedAnimDataLibrary.generated.h"
+
+class UAnimInstance;
+
+/**
+ *	A library of commonly used functionality from the CachedAnimData family, exposed to blueprint.
+ */
+UCLASS(meta = (ScriptName = "CachedAnimDataLibrary"))
+class ENGINE_API UCachedAnimDataLibrary : public UBlueprintFunctionLibrary
+{
+	GENERATED_UCLASS_BODY()
+
+	// Returns whether a state is relevant (specified in the provided FCachedAnimStateData)
+	UFUNCTION(BlueprintPure, Category = "Utilities|Animation|StateMachine")
+	static bool StateMachine_IsStateRelevant(UAnimInstance* InAnimInstance, UPARAM(ref) const FCachedAnimStateData& CachedAnimStateData);
+
+	// Returns the weight of a state, relative to its state machine (specified in the provided FCachedAnimStateData)
+	UFUNCTION(BlueprintPure, Category = "Utilities|Animation|StateMachine")
+	static float StateMachine_GetLocalWeight(UAnimInstance* InAnimInstance, UPARAM(ref) const FCachedAnimStateData& CachedAnimStateData);
+
+	// Returns the weight of a state, relative to the graph (specified in the provided FCachedAnimStateData)
+	UFUNCTION(BlueprintPure, Category = "Utilities|Animation|StateMachine")
+	static float StateMachine_GetGlobalWeight(UAnimInstance* InAnimInstance, UPARAM(ref) const FCachedAnimStateData& CachedAnimStateData);
+};
+
