Add ENGINE_API export to Handle/Ref impl.

diff --git a/Engine/Source/Runtime/Engine/Public/WorldPartition/WorldPartitionHandle.h b/Engine/Source/Runtime/Engine/Public/WorldPartition/WorldPartitionHandle.h
--- a/Engine/Source/Runtime/Engine/Public/WorldPartition/WorldPartitionHandle.h
+++ b/Engine/Source/Runtime/Engine/Public/WorldPartition/WorldPartitionHandle.h
@@ -110,13 +110,13 @@ protected:
 	TUniquePtr<FWorldPartitionActorDesc>* ActorDesc;
 };
 
-struct FWorldPartitionHandleImpl
+struct ENGINE_API FWorldPartitionHandleImpl
 {
 	static void IncRefCount(FWorldPartitionActorDesc* ActorDesc);
 	static void DecRefCount(FWorldPartitionActorDesc* ActorDesc);
 };
 
-struct FWorldPartitionReferenceImpl
+struct ENGINE_API FWorldPartitionReferenceImpl
 {
 	static void IncRefCount(FWorldPartitionActorDesc* ActorDesc);
 	static void DecRefCount(FWorldPartitionActorDesc* ActorDesc);
