Adding support for inertial blending in montages

diff --git a/Engine/Source/Runtime/AnimGraphRuntime/Private/AnimNodes/AnimNode_Slot.cpp b/Engine/Source/Runtime/AnimGraphRuntime/Private/AnimNodes/AnimNode_Slot.cpp
--- a/Engine/Source/Runtime/AnimGraphRuntime/Private/AnimNodes/AnimNode_Slot.cpp
+++ b/Engine/Source/Runtime/AnimGraphRuntime/Private/AnimNodes/AnimNode_Slot.cpp
@@ -3,6 +3,7 @@
 #include "AnimNodes/AnimNode_Slot.h"
 #include "Animation/AnimInstanceProxy.h"
 #include "Animation/AnimTrace.h"
+#include "Animation/AnimNode_Inertialization.h"
 
 /////////////////////////////////////////////////////
 // FAnimNode_Slot
@@ -38,6 +39,21 @@ void FAnimNode_Slot::Update_AnyThread(const FAnimationUpdateContext& Context)
 	// Update cache in AnimInstance.
 	Context.AnimInstanceProxy->UpdateSlotNodeWeight(SlotName, WeightData.SlotNodeWeight, Context.GetFinalBlendWeight());
 
+	float InertializationDuration;
+	if (Context.AnimInstanceProxy->GetSlotInertializationRequest(SlotName, InertializationDuration))
+	{
+		UE::Anim::IInertializationRequester* InertializationRequester = Context.GetMessage<UE::Anim::IInertializationRequester>();
+		if (InertializationRequester)
+		{
+			InertializationRequester->RequestInertialization(InertializationDuration);
+		}
+		else
+		{
+			FAnimNode_Inertialization::LogRequestError(Context, Source);
+		}
+	}
+
+
 	const bool bUpdateSource = bAlwaysUpdateSourcePose || FAnimWeight::IsRelevant(WeightData.SourceWeight);
 	if (bUpdateSource)
 	{
diff --git a/Engine/Source/Runtime/Engine/Classes/Animation/AnimInstance.h b/Engine/Source/Runtime/Engine/Classes/Animation/AnimInstance.h
--- a/Engine/Source/Runtime/Engine/Classes/Animation/AnimInstance.h
+++ b/Engine/Source/Runtime/Engine/Classes/Animation/AnimInstance.h
@@ -840,6 +840,9 @@ protected:
 	/** Map between Active Montages and their FAnimMontageInstance */
 	TMap<class UAnimMontage*, struct FAnimMontageInstance*> ActiveMontagesMap;
 
+	/**  Inertialization requests gathered this frame. Gets reset in UpdateMontageEvaluationData */
+	TMap<FName, float> SlotGroupInertializationRequestMap;
+
 	/** Stop all active montages belonging to 'InGroupName' */
 	void StopAllMontagesByGroupName(FName InGroupName, const FAlphaBlend& BlendOut);
 
@@ -852,6 +855,10 @@ protected:
 	float Montage_PlayInternal(UAnimMontage* MontageToPlay, const FAlphaBlend& BlendIn, float InPlayRate = 1.f, EMontagePlayReturnType ReturnValueType = EMontagePlayReturnType::MontageLength, float InTimeToStartMontageAt = 0.f, bool bStopAllMontages = true);
 
 public:
+
+	/**  Builds an inertialization request from the montage's group and provided duration */
+	void RequestMontageInertialization(const UAnimMontage* Montage, float Duration);
+
 	/** Queue a Montage BlendingOut Event to be triggered. */
 	void QueueMontageBlendingOutEvent(const FQueuedMontageBlendingOutEvent& MontageBlendingOutEvent);
 
diff --git a/Engine/Source/Runtime/Engine/Classes/Animation/AnimMontage.h b/Engine/Source/Runtime/Engine/Classes/Animation/AnimMontage.h
--- a/Engine/Source/Runtime/Engine/Classes/Animation/AnimMontage.h
+++ b/Engine/Source/Runtime/Engine/Classes/Animation/AnimMontage.h
@@ -124,6 +124,15 @@ namespace EAnimNotifyEventType
 	};
 }
 
+UENUM()
+enum class EMontageBlendMode : uint8
+{
+	//Uses standard weight based blend
+	Standard,
+	//Uses inertialization. Requires an inertialization node somewhere in the graph after every slot node used by this montage.
+	Inertialization,
+};
+
 /** AnimNotifies marked as BranchingPoints will create these markers on their Begin/End times.
 	They create stopping points when the Montage is being ticked to dispatch events. */
 USTRUCT()
@@ -581,6 +590,9 @@ class UAnimMontage : public UAnimCompositeBase
 	GENERATED_UCLASS_BODY()
 
 	friend struct FAnimMontageInstance;
+	
+	UPROPERTY(EditAnywhere, Category = BlendOption)
+	EMontageBlendMode BlendMode;
 
 	/** Blend in option. */
 	UPROPERTY(EditAnywhere, Category=BlendOption)
diff --git a/Engine/Source/Runtime/Engine/Private/Animation/AnimInstance.cpp b/Engine/Source/Runtime/Engine/Private/Animation/AnimInstance.cpp
--- a/Engine/Source/Runtime/Engine/Private/Animation/AnimInstance.cpp
+++ b/Engine/Source/Runtime/Engine/Private/Animation/AnimInstance.cpp
@@ -280,6 +280,8 @@ void UAnimInstance::UninitializeAnimation()
 
 	ActiveAnimNotifyState.Reset();
 	NotifyQueue.Reset(SkelMeshComp);
+
+	SlotGroupInertializationRequestMap.Reset();
 }
 
 #if WITH_EDITORONLY_DATA
@@ -1665,6 +1667,16 @@ void UAnimInstance::Montage_Advance(float DeltaSeconds)
 	}
 }
 
+void UAnimInstance::RequestMontageInertialization(const UAnimMontage* Montage, float Duration)
+{
+	if (Montage)
+	{
+		// Adds a new request or overwrites an existing one
+		// We always overwrite with the last request, instead of using the shortest one (differs from AnimNode_Inertialization), because we expect the last montage played/stopped to take precedence
+		SlotGroupInertializationRequestMap.FindOrAdd(Montage->GetGroupName()) = Duration;
+	}
+}
+
 void UAnimInstance::QueueMontageBlendingOutEvent(const FQueuedMontageBlendingOutEvent& MontageBlendingOutEvent)
 {
 	if (bQueueMontageEvents)
@@ -1855,11 +1867,22 @@ float UAnimInstance::Montage_PlayInternal(UAnimMontage* MontageToPlay, const FAl
 	{
 		if (CurrentSkeleton && CurrentSkeleton->IsCompatible(MontageToPlay->GetSkeleton()))
 		{
+			const float BlendDuration = BlendIn.GetBlendTime();
+			// If using interial blend, force an instant blend in
+			const FAlphaBlend EffectiveBlendIn = (MontageToPlay->BlendMode == EMontageBlendMode::Inertialization) ? FAlphaBlend(BlendIn, 0.0f) : BlendIn;
+
+			const FName NewMontageGroupName = MontageToPlay->GetGroupName();
 			if (bStopAllMontages)
 			{
 				// Enforce 'a single montage at once per group' rule
-				FName NewMontageGroupName = MontageToPlay->GetGroupName();
-				StopAllMontagesByGroupName(NewMontageGroupName, BlendIn);
+				StopAllMontagesByGroupName(NewMontageGroupName, EffectiveBlendIn);
+			}
+
+			if (MontageToPlay->BlendMode == EMontageBlendMode::Inertialization)
+			{
+				// Request new inertialization for new montage's group name
+				// The above montage stop may request an inertial blend. We overwrite that here.
+				RequestMontageInertialization(MontageToPlay, BlendDuration);
 			}
 
 			// Enforce 'a single root motion montage at once' rule.
@@ -1868,7 +1891,7 @@ float UAnimInstance::Montage_PlayInternal(UAnimMontage* MontageToPlay, const FAl
 				FAnimMontageInstance* ActiveRootMotionMontageInstance = GetRootMotionMontageInstance();
 				if (ActiveRootMotionMontageInstance)
 				{
-					ActiveRootMotionMontageInstance->Stop(MontageToPlay->BlendIn);
+					ActiveRootMotionMontageInstance->Stop(EffectiveBlendIn);
 				}
 			}
 
@@ -1878,7 +1901,7 @@ float UAnimInstance::Montage_PlayInternal(UAnimMontage* MontageToPlay, const FAl
 			const float MontageLength = MontageToPlay->GetPlayLength();
 
 			NewInstance->Initialize(MontageToPlay);
-			NewInstance->Play(InPlayRate, BlendIn);
+			NewInstance->Play(InPlayRate, EffectiveBlendIn);
 			NewInstance->SetPosition(FMath::Clamp(InTimeToStartMontageAt, 0.f, MontageLength));
 			MontageInstances.Add(NewInstance);
 			ActiveMontagesMap.Add(MontageToPlay, NewInstance);
@@ -3192,6 +3215,12 @@ void UAnimInstance::UpdateMontageEvaluationData()
 				));
 		}
 	}
+
+	Proxy.GetSlotGroupInertializationRequestMap() = SlotGroupInertializationRequestMap;
+
+	// Reset inertialization requests every frame.
+	// If the request is missed by the graph (i.e. the slot node is not relevant), we assume assume what brought it back to relevancy will handle the blend instead.
+	SlotGroupInertializationRequestMap.Reset();
 }
 
 float UAnimInstance::GetInstanceAssetPlayerLength(int32 AssetPlayerIndex)
diff --git a/Engine/Source/Runtime/Engine/Private/Animation/AnimInstanceProxy.cpp b/Engine/Source/Runtime/Engine/Private/Animation/AnimInstanceProxy.cpp
--- a/Engine/Source/Runtime/Engine/Private/Animation/AnimInstanceProxy.cpp
+++ b/Engine/Source/Runtime/Engine/Private/Animation/AnimInstanceProxy.cpp
@@ -343,6 +343,7 @@ void FAnimInstanceProxy::Uninitialize(UAnimInstance* InAnimInstance)
 	DECLARE_SCOPE_HIERARCHICAL_COUNTER_FUNC()
 
 	MontageEvaluationData.Reset();
+	SlotGroupInertializationRequestMap.Reset();
 	DefaultLinkedInstanceInputNode = nullptr;
 	ResetAnimationCurves();
 	MaterialParametersToClear.Reset();
@@ -1087,6 +1088,18 @@ void FAnimInstanceProxy::UpdateSlotNodeWeight(const FName& SlotNodeName, float I
 	}
 }
 
+bool FAnimInstanceProxy::GetSlotInertializationRequest(const FName& SlotName, float& OutDuration)
+{
+	const FName GroupName = Skeleton ? Skeleton->GetSlotGroupName(SlotName) : NAME_None;
+	if (const float* RequestDuration = GetSlotGroupInertializationRequestMap().Find(GroupName))
+	{
+		OutDuration = *RequestDuration;
+		return true;
+	}
+
+	return false;
+}
+
 void FAnimInstanceProxy::ClearSlotNodeWeights()
 {
 	TArray<FMontageActiveSlotTracker>& SlotWeightTracker_Read = SlotWeightTracker[GetSyncGroupReadIndex()];
diff --git a/Engine/Source/Runtime/Engine/Private/Animation/AnimMontage.cpp b/Engine/Source/Runtime/Engine/Private/Animation/AnimMontage.cpp
--- a/Engine/Source/Runtime/Engine/Private/Animation/AnimMontage.cpp
+++ b/Engine/Source/Runtime/Engine/Private/Animation/AnimMontage.cpp
@@ -38,6 +38,8 @@ namespace MontageFNames
 UAnimMontage::UAnimMontage(const FObjectInitializer& ObjectInitializer)
 	: Super(ObjectInitializer)
 {
+	BlendMode = EMontageBlendMode::Standard;
+
 	BlendIn.SetBlendTime(0.25f);
 	BlendOut.SetBlendTime(0.25f);
 	BlendOutTriggerTime = -1.f;
@@ -1387,20 +1389,29 @@ void FAnimMontageInstance::Stop(const FAlphaBlend& InBlendOut, bool bInterrupt)
 	// if it hasn't stopped, stop now
 	if (IsStopped() == false)
 	{
+		// If we are using Inertial Blend, instantly stop the montage.
+		const bool bShouldInertialize = Montage && (Montage->BlendMode == EMontageBlendMode::Inertialization);
+
 		// do not use default Montage->BlendOut 
 		// depending on situation, the BlendOut time can change 
-		InitializeBlend(InBlendOut);
+		InitializeBlend(bShouldInertialize ? (FAlphaBlend(InBlendOut, 0.0f)) : InBlendOut);
 		BlendStartAlpha = Blend.GetAlpha();
 		Blend.SetDesiredValue(0.f);
 		Blend.Update(0.0f);
 
-		if(Montage)
+		if (Montage)
 		{
 			if (UAnimInstance* Inst = AnimInstance.Get())
 			{
 				// Let AnimInstance know we are being stopped.
 				Inst->OnMontageInstanceStopped(*this);
 				Inst->QueueMontageBlendingOutEvent(FQueuedMontageBlendingOutEvent(Montage, bInterrupted, OnMontageBlendingOutStarted));
+
+				if (bShouldInertialize)
+				{
+					// Send the inertial blend request to the anim instance
+					Inst->RequestMontageInertialization(Montage, InBlendOut.GetBlendTime());
+				}
 			}
 		}
 	}
@@ -2295,13 +2306,19 @@ void FAnimMontageInstance::Advance(float DeltaTime, struct FRootMotionMovementPa
 						const float PlayTimeToEnd = MontageSubStepper.GetRemainingPlayTimeToSectionEnd(Position);
 
 						const bool bCustomBlendOutTriggerTime = (Montage->BlendOutTriggerTime >= 0);
+						const bool bUsesInertialization = Montage->BlendMode == EMontageBlendMode::Inertialization;
+
 						const float DefaultBlendOutTime = Montage->BlendOut.GetBlendTime() * DefaultBlendTimeMultiplier;
-						const float BlendOutTriggerTime = bCustomBlendOutTriggerTime ? Montage->BlendOutTriggerTime : DefaultBlendOutTime;
+
+						// If we are using inertialization, we don't need to stop the montage early. In this case, we also don't want to miss out on notifies during the blend out window.
+						const float EffectiveDefaultBlendOutTime = bUsesInertialization ? 0.0f : DefaultBlendOutTime;
+						const float BlendOutTriggerTime = bCustomBlendOutTriggerTime ? Montage->BlendOutTriggerTime : EffectiveDefaultBlendOutTime;
 
 						// ... trigger blend out if within blend out time window.
 						if (PlayTimeToEnd <= FMath::Max<float>(BlendOutTriggerTime, KINDA_SMALL_NUMBER))
 						{
-							const float BlendOutTime = bCustomBlendOutTriggerTime ? DefaultBlendOutTime : PlayTimeToEnd;
+							// If we are using inertialization, use DefaultBlendOutTime, as we no longer crossfade between poses.
+							const float BlendOutTime = (bCustomBlendOutTriggerTime || bUsesInertialization) ? DefaultBlendOutTime : PlayTimeToEnd;
 							Stop(FAlphaBlend(Montage->BlendOut, BlendOutTime), false);
 						}
 					}
diff --git a/Engine/Source/Runtime/Engine/Public/Animation/AnimInstanceProxy.h b/Engine/Source/Runtime/Engine/Public/Animation/AnimInstanceProxy.h
--- a/Engine/Source/Runtime/Engine/Public/Animation/AnimInstanceProxy.h
+++ b/Engine/Source/Runtime/Engine/Public/Animation/AnimInstanceProxy.h
@@ -399,6 +399,8 @@ public:
 	// Allow slot nodes to store off their weight during ticking
 	void UpdateSlotNodeWeight(const FName& SlotNodeName, float InLocalMontageWeight, float InNodeGlobalWeight);
 
+	bool GetSlotInertializationRequest(const FName& SlotName, float& OutDuration);
+
 	/** Register a named slot */
 	void RegisterSlotNodeWithAnimInstance(const FName& SlotNodeName);
 
@@ -704,6 +706,8 @@ protected:
 	/** Access montage array data */
 	TArray<FMontageEvaluationState>& GetMontageEvaluationData() { return MontageEvaluationData; }
 
+	TMap<FName, float>& GetSlotGroupInertializationRequestMap() { return SlotGroupInertializationRequestMap; }
+
 	/** Access montage array data */
 	const TArray<FMontageEvaluationState>& GetMontageEvaluationData() const { return MontageEvaluationData; }
 
@@ -970,6 +974,9 @@ private:
 	/** Copy of UAnimInstance::MontageInstances data used for update & evaluation */
 	TArray<FMontageEvaluationState> MontageEvaluationData;
 
+	// Inertialization request for each slot. Mapped with SlotNameToTrackerIndex. Initialized to a value of 1.0f, which indicates no request for this slot.
+	TMap<FName, float> SlotGroupInertializationRequestMap;
+
 	/** Delegate fired on the game thread before update occurs */
 	TArray<FAnimNode_Base*> GameThreadPreUpdateNodes;
 
