Fix for non-unity build error

diff --git a/Engine/Source/Runtime/Engine/Public/Animation/CachedAnimDataLibrary.h b/Engine/Source/Runtime/Engine/Public/Animation/CachedAnimDataLibrary.h
--- a/Engine/Source/Runtime/Engine/Public/Animation/CachedAnimDataLibrary.h
+++ b/Engine/Source/Runtime/Engine/Public/Animation/CachedAnimDataLibrary.h
@@ -4,6 +4,7 @@
 
 #include "CoreMinimal.h"
 #include "Animation/CachedAnimData.h"
+#include "Kismet/BlueprintFunctionLibrary.h"
 #include "CachedAnimDataLibrary.generated.h"
 
 class UAnimInstance;
