When collecting shader compile results, we also need to collect their success states. This was preventing the shader compiler code from hitting the prompt to recompile code paths.

diff --git a/Engine/Source/Runtime/Engine/Private/ShaderCompiler/ShaderCompiler.cpp b/Engine/Source/Runtime/Engine/Private/ShaderCompiler/ShaderCompiler.cpp
--- a/Engine/Source/Runtime/Engine/Private/ShaderCompiler/ShaderCompiler.cpp
+++ b/Engine/Source/Runtime/Engine/Private/ShaderCompiler/ShaderCompiler.cpp
@@ -3643,6 +3643,7 @@ void FShaderCompilingManager::ProcessAsyncResults(bool bLimitExecutionTime, bool
 						FShaderMapFinalizeResults& FinalizeResults = PendingFinalizeShaderMaps.FindOrAdd(It.Key());
 						FinalizeResults.FinishedJobs.Append(Results->FinishedJobs);
 						Results->FinishedJobs.Reset();
+						FinalizeResults.bAllJobsSucceeded = FinalizeResults.bAllJobsSucceeded && Results->bAllJobsSucceeded;
 					}
 
 					checkf(Results->FinishedJobs.Num() == 0, TEXT("Failed to remove finished jobs, %d remain"), Results->FinishedJobs.Num());
