Added console command wp.Runtime.ToggleDrawDataLayers to display active data layers
Fixed 3D debug display of partitioned level instance (transform was not taken into consideration)

diff --git a/Engine/Source/Runtime/Engine/Private/WorldPartition/DataLayer/DataLayerSubsystem.cpp b/Engine/Source/Runtime/Engine/Private/WorldPartition/DataLayer/DataLayerSubsystem.cpp
--- a/Engine/Source/Runtime/Engine/Private/WorldPartition/DataLayer/DataLayerSubsystem.cpp
+++ b/Engine/Source/Runtime/Engine/Private/WorldPartition/DataLayer/DataLayerSubsystem.cpp
@@ -5,12 +5,20 @@
 #include "WorldPartition/DataLayer/DataLayer.h"
 #include "Engine/Engine.h"
 #include "Engine/World.h"
+#include "Engine/Canvas.h"
+#include "Debug/DebugDrawService.h"
 #if WITH_EDITOR
 #include "Editor.h"
 #include "Modules/ModuleManager.h"
 #include "WorldPartition/DataLayer/IDataLayerEditorModule.h"
 #endif
 
+static int32 GDrawDataLayers = 0;
+static FAutoConsoleCommand CVarDrawDataLayers(
+	TEXT("wp.Runtime.ToggleDrawDataLayers"),
+	TEXT("Toggles debug display of active data layers."),
+	FConsoleCommandDelegate::CreateLambda([] { GDrawDataLayers = !GDrawDataLayers; }));
+
 UDataLayerSubsystem::UDataLayerSubsystem()
 {}
 
@@ -41,6 +49,27 @@ void UDataLayerSubsystem::Initialize(FSubsystemCollectionBase& Collection)
 #endif
 }
 
+void UDataLayerSubsystem::PostInitialize()
+{
+	Super::PostInitialize();
+
+	if (GetWorld()->IsGameWorld() && (GetWorld()->GetNetMode() != NM_DedicatedServer))
+	{
+		DrawHandle = UDebugDrawService::Register(TEXT("Game"), FDebugDrawDelegate::CreateUObject(this, &UDataLayerSubsystem::Draw));
+	}
+}
+
+void UDataLayerSubsystem::Deinitialize()
+{
+	if (DrawHandle.IsValid())
+	{
+		UDebugDrawService::Unregister(DrawHandle);
+		DrawHandle.Reset();
+	}
+
+	Super::Deinitialize();
+}
+
 UDataLayer* UDataLayerSubsystem::GetDataLayerFromLabel(const FName& InDataLayerLabel) const
 {
 	const AWorldDataLayers* WorldDataLayers = AWorldDataLayers::Get(GetWorld());
@@ -122,6 +151,40 @@ bool UDataLayerSubsystem::IsAnyDataLayerActive(const TArray<FName>& InDataLayerN
 	return false;
 }
 
+void UDataLayerSubsystem::Draw(UCanvas* Canvas, class APlayerController* PC)
+{
+	if (!GDrawDataLayers || !Canvas || !Canvas->SceneView)
+	{
+		return;
+	}
+
+	auto DrawText = [](UCanvas* Canvas, const FString& Text, UFont* Font, const FColor& Color, FVector2D& Pos)
+	{
+		float TextWidth, TextHeight;
+		Canvas->StrLen(Font, Text, TextWidth, TextHeight);
+		Canvas->SetDrawColor(Color);
+		Canvas->DrawText(Font, Text, Pos.X, Pos.Y);
+		Pos.Y += TextHeight + 1;
+	};
+
+	if (ActiveDataLayerNames.Num() > 0)
+	{
+		const FVector2D CanvasTopLeftPadding(10.f, 10.f);
+		FVector2D Pos = CanvasTopLeftPadding;
+		FString Title(TEXT("Active Data Layers"));
+		DrawText(Canvas, Title, GEngine->GetSmallFont(), FColor::Yellow, Pos);
+		UFont* DataLayerFont = GEngine->GetTinyFont();
+		for (const FName& DataLayerName : ActiveDataLayerNames)
+		{
+			if (UDataLayer* DataLayer = GetDataLayerFromName(DataLayerName))
+			{
+				const FString Text = DataLayer->GetDataLayerLabel().ToString();
+				DrawText(Canvas, Text, DataLayerFont, FColor::White, Pos);
+			}
+		}
+	}
+}
+
 FAutoConsoleCommand UDataLayerSubsystem::ToggleDataLayerActivation(
 	TEXT("wp.Runtime.ToggleDataLayerActivation"),
 	TEXT("Toggles a DataLayer. Args [DataLayerLabel]"),
diff --git a/Engine/Source/Runtime/Engine/Private/WorldPartition/WorldPartitionRuntimeSpatialHash.cpp b/Engine/Source/Runtime/Engine/Private/WorldPartition/WorldPartitionRuntimeSpatialHash.cpp
--- a/Engine/Source/Runtime/Engine/Private/WorldPartition/WorldPartitionRuntimeSpatialHash.cpp
+++ b/Engine/Source/Runtime/Engine/Private/WorldPartition/WorldPartitionRuntimeSpatialHash.cpp
@@ -839,7 +839,7 @@ void FSpatialHashStreamingGrid::GetAlwaysLoadedCells(const UDataLayerSubsystem*
 	}
 }
 
-void FSpatialHashStreamingGrid::Draw3D(UWorld* World, const TArray<FWorldPartitionStreamingSource>& Sources) const
+void FSpatialHashStreamingGrid::Draw3D(UWorld* World, const TArray<FWorldPartitionStreamingSource>& Sources, const FTransform& Transform) const
 {
 	const FSquare2DGridHelper& Helper = GetGridHelper();
 	int32 MinGridLevel = FMath::Clamp<int32>(GShowRuntimeSpatialHashGridLevel, 0, GridLevels.Num() - 1);
@@ -882,13 +882,16 @@ void FSpatialHashStreamingGrid::Draw3D(UWorld* World, const TArray<FWorldPartiti
 				}
 				FColor CellColor = Cell ? Cell->GetDebugColor().ToFColor(false).WithAlpha(16) : FColor(0, 0, 0, 16);
 				FVector BoundsOrigin(CellWorldBounds.GetCenter(), Z);
-				DrawDebugSolidBox(World, BoundsOrigin, BoundsExtent, CellColor, false, -1.f, 255);
-				DrawDebugBox(World, BoundsOrigin, BoundsExtent, CellColor.WithAlpha(255), false, -1.f, 255, 10.f);
+				FBox Box = FBox::BuildAABB(BoundsOrigin, BoundsExtent);
+				DrawDebugSolidBox(World, Box, CellColor, Transform, false, -1.f, 255);
+				FVector Pos = Transform.TransformPosition(BoundsOrigin);
+				DrawDebugBox(World, Pos, BoundsExtent, Transform.GetRotation(), CellColor.WithAlpha(255), false, -1.f, 255, 10.f);
 			});
 		}
 
 		// Draw Loading Ranges
 		FVector SphereLocation(FVector2D(Source.Location), Z);
+		SphereLocation = Transform.TransformPosition(SphereLocation);
 		DrawDebugSphere(World, SphereLocation, Radius, 32, FColor::White, false, -1.f, 0, 20.f);
 	}
 }
@@ -907,8 +910,8 @@ void FSpatialHashStreamingGrid::Draw2D(UCanvas* Canvas, const TArray<FWorldParti
 			Axis.LineThickness = 3;
 			{
 				Axis.SetColor(FLinearColor::Green);
-				FVector2D LineStart = WorldToScreen(FVector2D(-163840.f, 0.f));
-				FVector2D LineEnd = WorldToScreen(FVector2D(163840.f, 0.f));
+				FVector2D LineStart = WorldToScreen(FVector2D(-1638400.f, 0.f));
+				FVector2D LineEnd = WorldToScreen(FVector2D(1638400.f, 0.f));
 				LineStart.X = FMath::Clamp(LineStart.X, GridScreenBounds.Min.X, GridScreenBounds.Max.X);
 				LineStart.Y = FMath::Clamp(LineStart.Y, GridScreenBounds.Min.Y, GridScreenBounds.Max.Y);
 				LineEnd.X = FMath::Clamp(LineEnd.X, GridScreenBounds.Min.X, GridScreenBounds.Max.X);
@@ -917,8 +920,8 @@ void FSpatialHashStreamingGrid::Draw2D(UCanvas* Canvas, const TArray<FWorldParti
 			}
 			{
 				Axis.SetColor(FLinearColor::Red);
-				FVector2D LineStart = WorldToScreen(FVector2D(0.f, -163840.f));
-				FVector2D LineEnd = WorldToScreen(FVector2D(0.f, 163840.f));
+				FVector2D LineStart = WorldToScreen(FVector2D(0.f, -1638400.f));
+				FVector2D LineEnd = WorldToScreen(FVector2D(0.f, 1638400.f));
 				LineStart.X = FMath::Clamp(LineStart.X, GridScreenBounds.Min.X, GridScreenBounds.Max.X);
 				LineStart.Y = FMath::Clamp(LineStart.Y, GridScreenBounds.Min.Y, GridScreenBounds.Max.Y);
 				LineEnd.X = FMath::Clamp(LineEnd.X, GridScreenBounds.Min.X, GridScreenBounds.Max.X);
@@ -2278,15 +2281,18 @@ void UWorldPartitionRuntimeSpatialHash::Draw3D(const TArray<FWorldPartitionStrea
 {
 	UWorld* World = GetWorld();
 
+	UWorldPartition* WorldPartition = GetOuterUWorldPartition();
+	FTransform Transform = WorldPartition->GetInstanceTransform();
+
 	if (StreamingGrids.IsValidIndex(GShowRuntimeSpatialHashGridIndex))
 	{
-		StreamingGrids[GShowRuntimeSpatialHashGridIndex].Draw3D(World, Sources);
+		StreamingGrids[GShowRuntimeSpatialHashGridIndex].Draw3D(World, Sources, Transform);
 	}
 	else
 	{
 		for (const FSpatialHashStreamingGrid& StreamingGrid : StreamingGrids)
 		{
-			StreamingGrid.Draw3D(World, Sources);
+			StreamingGrid.Draw3D(World, Sources, Transform);
 		}
 	}
 }
\ No newline at end of file
diff --git a/Engine/Source/Runtime/Engine/Private/WorldPartition/WorldPartitionSubsystem.cpp b/Engine/Source/Runtime/Engine/Private/WorldPartition/WorldPartitionSubsystem.cpp
--- a/Engine/Source/Runtime/Engine/Private/WorldPartition/WorldPartitionSubsystem.cpp
+++ b/Engine/Source/Runtime/Engine/Private/WorldPartition/WorldPartitionSubsystem.cpp
@@ -159,10 +159,11 @@ void UWorldPartitionSubsystem::DrawRuntimeHash2D(UCanvas* Canvas, class APlayerC
 		return;
 	}
 
+	const float MaxScreenRatio = 0.75f;
 	const FVector2D CanvasTopLeftPadding(10.f, 10.f);
 	const FVector2D CanvasBottomRightPadding(10.f, 10.f);
 	const FVector2D CanvasMinimumSize(100.f, 100.f);
-	const FVector2D CanvasMaxScreenSize = FVector2D::Max(FVector2D(Canvas->ClipX, Canvas->ClipY) - CanvasBottomRightPadding - CanvasTopLeftPadding, CanvasMinimumSize);
+	const FVector2D CanvasMaxScreenSize = FVector2D::Max(MaxScreenRatio*FVector2D(Canvas->ClipX, Canvas->ClipY) - CanvasBottomRightPadding - CanvasTopLeftPadding, CanvasMinimumSize);
 
 	FVector2D TotalFootprint(ForceInitToZero);
 	for (UWorldPartition* Partition : RegisteredWorldPartitions)
diff --git a/Engine/Source/Runtime/Engine/Public/WorldPartition/DataLayer/DataLayerSubsystem.h b/Engine/Source/Runtime/Engine/Public/WorldPartition/DataLayer/DataLayerSubsystem.h
--- a/Engine/Source/Runtime/Engine/Public/WorldPartition/DataLayer/DataLayerSubsystem.h
+++ b/Engine/Source/Runtime/Engine/Public/WorldPartition/DataLayer/DataLayerSubsystem.h
@@ -22,8 +22,13 @@ public:
 	//~ Begin USubsystem Interface.
 	virtual bool ShouldCreateSubsystem(UObject* Outer) const override;
 	virtual void Initialize(FSubsystemCollectionBase& Collection) override;
+	virtual void Deinitialize() override;
 	//~ End USubsystem Interface.
 
+	//~ Begin UWorldSubsystem Interface.
+	virtual void PostInitialize() override;
+	//~ End UWorldSubsystem Interface.
+
 	//~ Begin Blueprint callable functions
 	UFUNCTION(BlueprintCallable, Category = DataLayers)
 	void ActivateDataLayer(const FActorDataLayer& InDataLayer, bool bInActivate);
@@ -46,9 +51,13 @@ public:
 	bool IsDataLayerActiveByName(const FName& InDataLayerName) const;
 	bool IsAnyDataLayerActive(const TArray<FName>& InDataLayerNames) const;
 
+	void Draw(class UCanvas* Canvas, class APlayerController* PC);
+
 private:
 	/** Console command used to toggle activation of a DataLayer */
 	static class FAutoConsoleCommand ToggleDataLayerActivation;
 
 	TSet<FName> ActiveDataLayerNames;
+
+	FDelegateHandle	DrawHandle;
 };
\ No newline at end of file
diff --git a/Engine/Source/Runtime/Engine/Public/WorldPartition/WorldPartitionRuntimeSpatialHash.h b/Engine/Source/Runtime/Engine/Public/WorldPartition/WorldPartitionRuntimeSpatialHash.h
--- a/Engine/Source/Runtime/Engine/Public/WorldPartition/WorldPartitionRuntimeSpatialHash.h
+++ b/Engine/Source/Runtime/Engine/Public/WorldPartition/WorldPartitionRuntimeSpatialHash.h
@@ -119,7 +119,7 @@ struct FSpatialHashStreamingGrid
 	void GetCells(const TArray<FWorldPartitionStreamingSource>& Sources, const class UDataLayerSubsystem* DataLayerSubsystem, TSet<const UWorldPartitionRuntimeCell*>& Cells) const;
 	void GetAlwaysLoadedCells(const UDataLayerSubsystem* DataLayerSubsystem, TSet<const UWorldPartitionRuntimeCell*>& Cells) const;
 	void Draw2D(UCanvas* Canvas, const TArray<FWorldPartitionStreamingSource>& Sources, const FBox& Region, const FBox2D& GridScreenBounds, TFunctionRef<FVector2D(const FVector2D&)> WorldToScreen) const;
-	void Draw3D(UWorld* World, const TArray<FWorldPartitionStreamingSource>& Sources) const;
+	void Draw3D(UWorld* World, const TArray<FWorldPartitionStreamingSource>& Sources, const FTransform& Transform) const;
 
 private:
 	mutable TMap<uint32, FVector> CachedLineTrace;
