Remove exposed parts of the Blueprint nativization feature (deprecated), and decouple it from the editor/cooking phase.

diff --git a/Engine/Plugins/MovieScene/LevelSequenceEditor/Source/LevelSequenceEditor/Private/LevelSequenceEditorModule.cpp b/Engine/Plugins/MovieScene/LevelSequenceEditor/Source/LevelSequenceEditor/Private/LevelSequenceEditorModule.cpp
--- a/Engine/Plugins/MovieScene/LevelSequenceEditor/Source/LevelSequenceEditor/Private/LevelSequenceEditorModule.cpp
+++ b/Engine/Plugins/MovieScene/LevelSequenceEditor/Source/LevelSequenceEditor/Private/LevelSequenceEditorModule.cpp
@@ -85,15 +85,6 @@ public:
 				// Only have a blueprint if it contains the BlueprintPathWithinPackage tag
 				return InAssetData.TagsAndValues.Contains(FBlueprintTags::BlueprintPathWithinPackage);
 			}
-
-			virtual bool SupportsNativization(const UObject* InAsset, const UBlueprint* InBlueprint, FText* OutReason) const
-			{
-				if (OutReason)
-				{
-					*OutReason = LOCTEXT("NativizationError", "Level Sequences do not support nativization.");
-				}
-				return false;
-			}
 		};
 
 		FBlueprintAssetHandler::Get().RegisterHandler<FLevelSequenceAssetBlueprintHandler>(ULevelSequence::StaticClass()->GetFName());
diff --git a/Engine/Source/Developer/BlueprintNativeCodeGen/Private/BlueprintNativeCodeGenModule.cpp b/Engine/Source/Developer/BlueprintNativeCodeGen/Private/BlueprintNativeCodeGenModule.cpp
--- a/Engine/Source/Developer/BlueprintNativeCodeGen/Private/BlueprintNativeCodeGenModule.cpp
+++ b/Engine/Source/Developer/BlueprintNativeCodeGen/Private/BlueprintNativeCodeGenModule.cpp
@@ -816,6 +816,14 @@ EReplacementResult FBlueprintNativeCodeGenModule::IsTargetedForReplacement(const
 
 EReplacementResult FBlueprintNativeCodeGenModule::IsTargetedForReplacement(const UObject* Object, const FCompilerNativizationOptions& NativizationOptions) const
 {
+#if 1 
+
+	// @todo (wip) - Removing Blueprint nativization as a feature.
+	// For now, don't convert/replace anything, even if the project is configured for nativization.
+	return EReplacementResult::DontReplace;
+
+#else
+
 	if (Object == nullptr)
 	{
 		return EReplacementResult::DontReplace;
@@ -1097,6 +1105,8 @@ EReplacementResult FBlueprintNativeCodeGenModule::IsTargetedForReplacement(const
 	}
 
 	return Result;
+
+#endif
 }
 
 void FBlueprintNativeCodeGenModule::ReplaceAsset(const UObject* InAsset, const FCompilerNativizationOptions& NativizationOptions) const
@@ -1243,7 +1253,7 @@ void FBlueprintNativeCodeGenModule::FillPlatformNativizationDetails(const ITarge
 	Details.CompilerNativizationOptions.PlatformName = Details.PlatformName;
 	Details.CompilerNativizationOptions.ClientOnlyPlatform = Platform->IsClientOnly();
 	Details.CompilerNativizationOptions.ServerOnlyPlatform = Platform->IsServerOnly();
-	Details.CompilerNativizationOptions.bExcludeMonolithicHeaders = GetDefault<UProjectPackagingSettings>()->bExcludeMonolithicEngineHeadersInNativizedCode;
+	Details.CompilerNativizationOptions.bExcludeMonolithicHeaders = false;
 
 	auto GatherExcludedStuff = [&](const TCHAR* KeyForExcludedModules, const TCHAR* KeyForExcludedPaths, const TCHAR* KeyForExcludedAssets)
 	{
diff --git a/Engine/Source/Developer/DeveloperToolSettings/Classes/Settings/ProjectPackagingSettings.h b/Engine/Source/Developer/DeveloperToolSettings/Classes/Settings/ProjectPackagingSettings.h
--- a/Engine/Source/Developer/DeveloperToolSettings/Classes/Settings/ProjectPackagingSettings.h
+++ b/Engine/Source/Developer/DeveloperToolSettings/Classes/Settings/ProjectPackagingSettings.h
@@ -164,19 +164,23 @@ public:
 	bool IncludeDebugFiles;
 
 	/** If enabled, then the project's Blueprint assets (including structs and enums) will be intermediately converted into C++ and used in the packaged project (in place of the .uasset files).*/
-	UPROPERTY(config, EditAnywhere, Category = Blueprints)
+	UE_DEPRECATED(5.0, "Blueprint Nativization has been removed as a supported feature. This setting is no longer exposed for editing and will eventually be removed.")
+	UPROPERTY(config)
 	EProjectPackagingBlueprintNativizationMethod BlueprintNativizationMethod;
 
 	/** List of Blueprints to include for nativization when using the exclusive method. */
-	UPROPERTY(config, EditAnywhere, AdvancedDisplay, Category = Blueprints, meta = (DisplayName = "List of Blueprint assets to nativize", RelativeToGameContentDir, LongPackageName))
+	UE_DEPRECATED(5.0, "Blueprint Nativization has been removed as a supported feature. This setting is no longer exposed for editing and will eventually be removed.")
+	UPROPERTY(config)
 	TArray<FFilePath> NativizeBlueprintAssets;
 
 	/** If enabled, the nativized assets code plugin will be added to the Visual Studio solution if it exists when regenerating the game project. Intended primarily to assist with debugging the target platform after cooking with nativization turned on. */
-	UPROPERTY(config, EditAnywhere, AdvancedDisplay, Category = Blueprints)
+	UE_DEPRECATED(5.0, "Blueprint Nativization has been removed as a supported feature. This setting is no longer exposed for editing and will eventually be removed.")
+	UPROPERTY(config)
 	bool bIncludeNativizedAssetsInProjectGeneration;
 
 	/** Whether or not to exclude monolithic engine headers (e.g. Engine.h) in the generated code when nativizing Blueprint assets. This may improve C++ compiler performance if your game code does not depend on monolithic engine headers to build. */
-	UPROPERTY(config, EditAnywhere, AdvancedDisplay, Category = Blueprints)
+	UE_DEPRECATED(5.0, "Blueprint Nativization has been removed as a supported feature. This setting is no longer exposed for editing and will eventually be removed.")
+	UPROPERTY(config)
 	bool bExcludeMonolithicEngineHeadersInNativizedCode;
 
 	/** If enabled, all content will be put into a one or more .pak files instead of many individual files (default = enabled). */
@@ -494,13 +498,16 @@ public:
 	virtual bool CanEditChange( const FProperty* InProperty ) const override;
 
 	/** Adds the given Blueprint asset to the exclusive nativization list. */
-	bool AddBlueprintAssetToNativizationList(const class UBlueprint* InBlueprint);
+	UE_DEPRECATED(5.0, "Blueprint Nativization has been removed as a supported feature. This API will eventually be removed.")
+	bool AddBlueprintAssetToNativizationList(const class UBlueprint* InBlueprint) { return false; }
 
 	/** Removes the given Blueprint asset from the exclusive nativization list. */
-	bool RemoveBlueprintAssetFromNativizationList(const class UBlueprint* InBlueprint);
+	UE_DEPRECATED(5.0, "Blueprint Nativization has been removed as a supported feature. This API will eventually be removed.")
+	bool RemoveBlueprintAssetFromNativizationList(const class UBlueprint* InBlueprint) { return false; }
 
 	/** Determines if the specified Blueprint is already saved for exclusive nativization. */
-	bool IsBlueprintAssetInNativizationList(const class UBlueprint* InBlueprint) const { return FindBlueprintInNativizationList(InBlueprint) >= 0; }
+	UE_DEPRECATED(5.0, "Blueprint Nativization has been removed as a supported feature. This API will eventually be removed.")
+	bool IsBlueprintAssetInNativizationList(const class UBlueprint* InBlueprint) const { return false; }
 #endif
 
 	/** Gets a list of all valid packaging configurations for the current project */
@@ -510,9 +517,6 @@ public:
 	const FTargetInfo* GetBuildTargetInfo() const;
 
 private:
-	/** Returns the index of the specified Blueprint in the exclusive nativization list (otherwise INDEX_NONE) */
-	int32 FindBlueprintInNativizationList(const UBlueprint* InBlueprint) const;
-
 	/** Fix up cooking paths after they've been edited or laoded */
 	void FixCookingPaths();
 };
diff --git a/Engine/Source/Developer/DeveloperToolSettings/Private/DeveloperToolSettingsModule.cpp b/Engine/Source/Developer/DeveloperToolSettings/Private/DeveloperToolSettingsModule.cpp
--- a/Engine/Source/Developer/DeveloperToolSettings/Private/DeveloperToolSettingsModule.cpp
+++ b/Engine/Source/Developer/DeveloperToolSettings/Private/DeveloperToolSettingsModule.cpp
@@ -4,7 +4,4 @@
 #include "Modules/ModuleManager.h"
 
 
-FDeveloperToolSettingsDelegates::FOnNativeBlueprintsSettingChanged FDeveloperToolSettingsDelegates::OnNativeBlueprintsSettingChanged;
-
-
 IMPLEMENT_MODULE(FDefaultModuleImpl, DeveloperToolSettings);
diff --git a/Engine/Source/Developer/DeveloperToolSettings/Private/ProjectPackagingSettings.cpp b/Engine/Source/Developer/DeveloperToolSettings/Private/ProjectPackagingSettings.cpp
--- a/Engine/Source/Developer/DeveloperToolSettings/Private/ProjectPackagingSettings.cpp
+++ b/Engine/Source/Developer/DeveloperToolSettings/Private/ProjectPackagingSettings.cpp
@@ -33,9 +33,6 @@ void UProjectPackagingSettings::PostInitProperties()
 	// Build code projects by default
 	Build = EProjectPackagingBuild::IfProjectHasCode;
 
-	// Cache the current set of Blueprint assets selected for nativization.
-	CachedNativizeBlueprintAssets = NativizeBlueprintAssets;
-
 	FixCookingPaths();
 
 	Super::PostInitProperties();
@@ -157,156 +154,12 @@ void UProjectPackagingSettings::PostEditChangeProperty(FPropertyChangedEvent& Pr
 			}
 		}
 	}
-	else if (Name == FName((TEXT("NativizeBlueprintAssets"))))
-	{
-		int32 AssetIndex;
-		if (NativizeBlueprintAssets.Num() > 0)
-		{
-			for (AssetIndex = 0; AssetIndex < NativizeBlueprintAssets.Num(); ++AssetIndex)
-			{
-				const FString& PackageName = NativizeBlueprintAssets[AssetIndex].FilePath;
-				if (AssetIndex >= CachedNativizeBlueprintAssets.Num())
-				{
-					// A new entry was added; toggle the exclusive flag on the corresponding Blueprint asset (if loaded).
-					FDeveloperToolSettingsDelegates::OnNativeBlueprintsSettingChanged.Broadcast(PackageName, true);
-
-					// Add an entry to the end of the cached list.
-					CachedNativizeBlueprintAssets.Add(NativizeBlueprintAssets[AssetIndex]);
-				}
-				else if (!PackageName.Equals(CachedNativizeBlueprintAssets[AssetIndex].FilePath))
-				{
-					if (NativizeBlueprintAssets.Num() < CachedNativizeBlueprintAssets.Num())
-					{
-						// An entry was removed; toggle the exclusive flag on the corresponding Blueprint asset (if loaded).
-						FDeveloperToolSettingsDelegates::OnNativeBlueprintsSettingChanged.Broadcast(CachedNativizeBlueprintAssets[AssetIndex].FilePath, false);
-
-						// Remove this entry from the cached list.
-						CachedNativizeBlueprintAssets.RemoveAt(AssetIndex);
-					}
-					else if (NativizeBlueprintAssets.Num() > CachedNativizeBlueprintAssets.Num())
-					{
-						// A new entry was inserted; toggle the exclusive flag on the corresponding Blueprint asset (if loaded).
-						FDeveloperToolSettingsDelegates::OnNativeBlueprintsSettingChanged.Broadcast(PackageName, true);
-
-						// Insert the new entry into the cached list.
-						CachedNativizeBlueprintAssets.Insert(NativizeBlueprintAssets[AssetIndex], AssetIndex);
-					}
-					else
-					{
-						// An entry was changed; toggle the exclusive flag on the corresponding Blueprint assets (if loaded).
-						FDeveloperToolSettingsDelegates::OnNativeBlueprintsSettingChanged.Broadcast(CachedNativizeBlueprintAssets[AssetIndex].FilePath, false);
-						FDeveloperToolSettingsDelegates::OnNativeBlueprintsSettingChanged.Broadcast(PackageName, true);
-
-						// Update the cached entry.
-						CachedNativizeBlueprintAssets[AssetIndex].FilePath = PackageName;
-					}
-				}
-			}
-
-			if (CachedNativizeBlueprintAssets.Num() > NativizeBlueprintAssets.Num())
-			{
-				// Removed entries at the end of the list; toggle the exclusive flag on the corresponding Blueprint asset(s) (if loaded).
-				for (AssetIndex = NativizeBlueprintAssets.Num(); AssetIndex < CachedNativizeBlueprintAssets.Num(); ++AssetIndex)
-				{
-					FDeveloperToolSettingsDelegates::OnNativeBlueprintsSettingChanged.Broadcast(CachedNativizeBlueprintAssets[AssetIndex].FilePath, false);
-				}
-
-				// Remove entries from the end of the cached list.
-				CachedNativizeBlueprintAssets.RemoveAt(NativizeBlueprintAssets.Num(), CachedNativizeBlueprintAssets.Num() - NativizeBlueprintAssets.Num());
-			}
-		}
-		else if (CachedNativizeBlueprintAssets.Num() > 0)
-		{
-			// Removed all entries; toggle the exclusive flag on the corresponding Blueprint asset(s) (if loaded).
-			for (AssetIndex = 0; AssetIndex < CachedNativizeBlueprintAssets.Num(); ++AssetIndex)
-			{
-				FDeveloperToolSettingsDelegates::OnNativeBlueprintsSettingChanged.Broadcast(CachedNativizeBlueprintAssets[AssetIndex].FilePath, false);
-			}
-
-			// Clear the cached list.
-			CachedNativizeBlueprintAssets.Empty();
-		}
-	}
 }
 
 bool UProjectPackagingSettings::CanEditChange(const FProperty* InProperty) const
 {
-	if (InProperty->GetFName() == FName(TEXT("NativizeBlueprintAssets")))
-	{
-		return BlueprintNativizationMethod == EProjectPackagingBlueprintNativizationMethod::Exclusive;
-	}
-
 	return Super::CanEditChange(InProperty);
 }
-
-
-bool UProjectPackagingSettings::AddBlueprintAssetToNativizationList(const class UBlueprint* InBlueprint)
-{
-	if (InBlueprint)
-	{
-		const FString PackageName = ((UObject*)InBlueprint)->GetOutermost()->GetName();
-
-		// Make sure it's not already in the exclusive list. This can happen if the user previously added this asset in the Project Settings editor.
-		const bool bFound = IsBlueprintAssetInNativizationList(InBlueprint);
-		if (!bFound)
-		{
-			// Add this Blueprint asset to the exclusive list.
-			FFilePath FileInfo;
-			FileInfo.FilePath = PackageName;
-			NativizeBlueprintAssets.Add(FileInfo);
-
-			// Also add it to the mirrored list for tracking edits.
-			CachedNativizeBlueprintAssets.Add(FileInfo);
-
-			return true;
-		}
-	}
-
-	return false;
-}
-
-bool UProjectPackagingSettings::RemoveBlueprintAssetFromNativizationList(const class UBlueprint* InBlueprint)
-{
-	if (InBlueprint)
-	{
-		const FString PackageName = ((UObject*)InBlueprint)->GetOutermost()->GetName();
-
-		int32 AssetIndex = FindBlueprintInNativizationList(InBlueprint);
-		if (AssetIndex >= 0)
-		{
-			// Note: Intentionally not using RemoveAtSwap() here, so that the order is preserved.
-			NativizeBlueprintAssets.RemoveAt(AssetIndex);
-
-			// Also remove it from the mirrored list (for tracking edits).
-			CachedNativizeBlueprintAssets.RemoveAt(AssetIndex);
-
-			return true;
-		}
-	}
-
-	return false;
-}
-
-int32 UProjectPackagingSettings::FindBlueprintInNativizationList(const class UBlueprint* InBlueprint) const
-{
-	int32 ListIndex = INDEX_NONE;
-	if (InBlueprint)
-	{
-		const FString PackageName = ((UObject*)InBlueprint)->GetOutermost()->GetName();
-
-		for (int32 AssetIndex = 0; AssetIndex < NativizeBlueprintAssets.Num(); ++AssetIndex)
-		{
-			if (NativizeBlueprintAssets[AssetIndex].FilePath.Equals(PackageName, ESearchCase::IgnoreCase))
-			{
-				ListIndex = AssetIndex;
-				break;
-			}
-		}
-	}
-	return ListIndex;
-}
-
-
 #endif
 
 TArray<EProjectPackagingBuildConfigurations> UProjectPackagingSettings::GetValidPackageConfigurations()
diff --git a/Engine/Source/Developer/DeveloperToolSettings/Public/DeveloperToolSettingsDelegates.h b/Engine/Source/Developer/DeveloperToolSettings/Public/DeveloperToolSettingsDelegates.h
--- a/Engine/Source/Developer/DeveloperToolSettings/Public/DeveloperToolSettingsDelegates.h
+++ b/Engine/Source/Developer/DeveloperToolSettings/Public/DeveloperToolSettingsDelegates.h
@@ -10,8 +10,6 @@
   **/
 struct DEVELOPERTOOLSETTINGS_API FDeveloperToolSettingsDelegates
 {
-	/** Sent after a nativization setting changes */
-	DECLARE_MULTICAST_DELEGATE_TwoParams(FOnNativeBlueprintsSettingChanged, const FString& PackageName, bool bSelect);
-	static FOnNativeBlueprintsSettingChanged OnNativeBlueprintsSettingChanged;
+	
 };
 
diff --git a/Engine/Source/Developer/LauncherServices/Private/Launcher/LauncherWorker.cpp b/Engine/Source/Developer/LauncherServices/Private/Launcher/LauncherWorker.cpp
--- a/Engine/Source/Developer/LauncherServices/Private/Launcher/LauncherWorker.cpp
+++ b/Engine/Source/Developer/LauncherServices/Private/Launcher/LauncherWorker.cpp
@@ -528,15 +528,6 @@ FString FLauncherWorker::CreateUATCommand( const ILauncherProfileRef& InProfile,
 
 	bool bIsBuilding = InProfile->ShouldBuild();
 
-	// Override the Blueprint nativization method for anything other than "cook by the book" mode. Nativized assets
-	// won't get regenerated otherwise, and we don't want UBT to include generated code assets from a previous cook.
-	// Also disable Blueprint nativization if the profile is not configured to also build code. Otherwise nativized
-	// assets generated at cook time will not be linked into the game's executable prior to stage/deployment phases.
-	if (InProfile->GetCookMode() != ELauncherProfileCookModes::ByTheBook || !bIsBuilding)
-	{
-		UATCommand += TEXT(" -ini:Game:[/Script/UnrealEd.ProjectPackagingSettings]:BlueprintNativizationMethod=Disabled");
-	}
-
 	// build
 	if (bIsBuilding)
 	{
diff --git a/Engine/Source/Editor/Kismet/Private/BlueprintDetailsCustomization.cpp b/Engine/Source/Editor/Kismet/Private/BlueprintDetailsCustomization.cpp
--- a/Engine/Source/Editor/Kismet/Private/BlueprintDetailsCustomization.cpp
+++ b/Engine/Source/Editor/Kismet/Private/BlueprintDetailsCustomization.cpp
@@ -5547,147 +5547,6 @@ void FBlueprintGlobalOptionsDetails::CustomizeDetails(IDetailLayoutBuilder& Deta
 					.ToolTipText( this, &FBlueprintGlobalOptionsDetails::GetDeprecatedTooltip )
 				];
 		}
-
-		IDetailCategoryBuilder& PkgCategory = DetailLayout.EditCategory("Packaging", LOCTEXT("BlueprintPackagingCategory", "Packaging"));
-		PkgCategory.AddCustomRow(LOCTEXT("NativizeLabel", "Nativize"))
-			.NameContent()
-			[
-				SNew(STextBlock)
-				.Text(LOCTEXT("NativizeLabel", "Nativize"))
-				.Font(IDetailLayoutBuilder::GetDetailFont())
-			]
-			.ValueContent()
-			[
-				SNew(SCheckBox)
-				.IsEnabled(this, &FBlueprintGlobalOptionsDetails::IsNativizeEnabled)
-				.IsChecked(this, &FBlueprintGlobalOptionsDetails::GetNativizeState)
-				.OnCheckStateChanged(this, &FBlueprintGlobalOptionsDetails::OnNativizeToggled)
-				.ToolTipText(this, &FBlueprintGlobalOptionsDetails::GetNativizeTooltip)
-			];
-	}
-}
-
-bool FBlueprintGlobalOptionsDetails::IsNativizeEnabled() const
-{
-	bool bIsEnabled = false;
-	if (UBlueprint* Blueprint = GetBlueprintObj())
-	{
-		bIsEnabled = Blueprint->SupportsNativization() && !FBlueprintEditorUtils::ShouldNativizeImplicitly(Blueprint);
-	}
-	return bIsEnabled;
-}
-
-ECheckBoxState FBlueprintGlobalOptionsDetails::GetNativizeState() const
-{
-	ECheckBoxState CheckboxState = ECheckBoxState::Undetermined;
-	if (UBlueprint* Blueprint = GetBlueprintObj())
-	{
-		if (FBlueprintEditorUtils::ShouldNativizeImplicitly(Blueprint))
-		{
-			CheckboxState = ECheckBoxState::Checked;
-		}
-		else
-		{
-			switch (Blueprint->NativizationFlag)
-			{
-			case EBlueprintNativizationFlag::Disabled:
-				CheckboxState = ECheckBoxState::Unchecked;
-				break;
-
-			case EBlueprintNativizationFlag::ExplicitlyEnabled:
-				CheckboxState = ECheckBoxState::Checked;
-				break;
-
-			case EBlueprintNativizationFlag::Dependency:
-			default:
-				// leave "Undetermined"
-				break;
-			}
-		}
-	}
-	return CheckboxState;	
-}
-
-FText FBlueprintGlobalOptionsDetails::GetNativizeTooltip() const
-{
-	UBlueprint* Blueprint = GetBlueprintObj();
-
-	if (Blueprint)
-	{
-		if (FBlueprintEditorUtils::ShouldNativizeImplicitly(Blueprint))
-		{
-			return LOCTEXT("NativizeImplicitlyTooltip", "This Blueprint must be nativized because it overrides one or more BlueprintCallable functions inherited from a parent Blueprint class that has also been flagged for nativization.");
-		}
-
-		FText Reason;
-		if (!Blueprint->SupportsNativization(&Reason))
-		{
-			return Reason.IsEmpty() ? LOCTEXT("NativizeDisabledTooltip", "This blueprint does not support nativization.") : Reason;
-		}
-
-		if (Blueprint->NativizationFlag == EBlueprintNativizationFlag::Dependency)
-		{
-			return LOCTEXT("NativizeAsDependencyTooltip", "This Blueprint has been flagged to nativize as a dependency needed by another Blueprint. This will be applied once that Blueprint is saved.");
-		}
-	}
-
-	return LOCTEXT("NativizeTooltip", "When exclusive nativization is enabled, then this asset will be nativized. NOTE: All super classes must be also nativized.");
-}
-
-void FBlueprintGlobalOptionsDetails::OnNativizeToggled(ECheckBoxState NewState) const
-{
-	if (UBlueprint* Blueprint = GetBlueprintObj())
-	{
-		if (NewState == ECheckBoxState::Checked)
-		{
-			Blueprint->NativizationFlag = EBlueprintNativizationFlag::ExplicitlyEnabled;
-
-			TArray<UClass*> NativizationDependencies;
-			FBlueprintEditorUtils::FindNativizationDependencies(Blueprint, NativizationDependencies);
-
-			int32 bDependenciesFlagged = 0;
-			// tag all dependencies as needing nativization
-			for (int32 DependencyIndex = 0; DependencyIndex < NativizationDependencies.Num(); ++DependencyIndex)
-			{
-				UClass* Dependency = NativizationDependencies[DependencyIndex];
-				if (UBlueprint* DependentBp = UBlueprint::GetBlueprintFromClass(Dependency))
-				{
-					if (DependentBp->NativizationFlag == EBlueprintNativizationFlag::Disabled)
-					{
-						DependentBp->NativizationFlag = EBlueprintNativizationFlag::Dependency;
-						++bDependenciesFlagged;
-					}
-					// recursively tag dependencies up the chain...
-					// relying on the fact that this only adds to the array via AddUnique()
-					FBlueprintEditorUtils::FindNativizationDependencies(DependentBp, NativizationDependencies);
-				}
-			}
-
-
-			if (bDependenciesFlagged > 0)
-			{
-				FNotificationInfo Warning(LOCTEXT("DependenciesMarkedForNativization", "Flagged extra (required dependency) Blueprints for nativization."));
-				Warning.ExpireDuration = 5.0f;
-				Warning.bFireAndForget = true;
-				Warning.Image = FCoreStyle::Get().GetBrush(TEXT("MessageLog.Warning"));
-				FSlateNotificationManager::Get().AddNotification(Warning);
-			}
-		}
-		else
-		{
-			Blueprint->NativizationFlag = EBlueprintNativizationFlag::Disabled;
-		}
-
-		// don't need to alter (dirty) compilation state, just the package's save state (since we save this setting to a config on save)
-// 		FProperty* NativizeProperty = UBlueprint::StaticClass()->FindPropertyByName(GET_MEMBER_NAME_CHECKED(UBlueprint, NativizationFlag));
-// 		if (ensure(NativizeProperty != nullptr))
-// 		{
-// 			FPropertyChangedEvent PropertyChangedEvent(NativizeProperty);
-// 			PropertyChangedEvent.ChangeType = EPropertyChangeType::ValueSet;
-// 
-// 			FBlueprintEditorUtils::MarkBlueprintAsModified(Blueprint, PropertyChangedEvent);
-// 		}
-		Blueprint->MarkPackageDirty();
 	}
 }
 
diff --git a/Engine/Source/Editor/Kismet/Private/BlueprintDetailsCustomization.h b/Engine/Source/Editor/Kismet/Private/BlueprintDetailsCustomization.h
--- a/Engine/Source/Editor/Kismet/Private/BlueprintDetailsCustomization.h
+++ b/Engine/Source/Editor/Kismet/Private/BlueprintDetailsCustomization.h
@@ -777,18 +777,6 @@ protected:
 	/** Returns the tooltip explaining deprecation */
 	FText GetDeprecatedTooltip() const;
 
-	/** Disabled in level and macro Blueprints */
-	bool IsNativizeEnabled() const;
-
-	/** Returns the check box state (undefined if the Blueprint is a dependency that will get added as part of another Blueprint) */
-	ECheckBoxState GetNativizeState() const;
-
-	/** Depending on the property's state, returns a tooltip describing the Blueprint nativize setting */
-	FText GetNativizeTooltip() const;
-
-	/** Flags the current Blueprint for nativization (as well as any dependencies that are required) */
-	void OnNativizeToggled(ECheckBoxState NewState) const;
-
 private:
 	/** Weak reference to the Blueprint editor */
 	TWeakPtr<FBlueprintEditor> BlueprintEditorPtr;
diff --git a/Engine/Source/Editor/Kismet/Private/BlueprintEditor.cpp b/Engine/Source/Editor/Kismet/Private/BlueprintEditor.cpp
--- a/Engine/Source/Editor/Kismet/Private/BlueprintEditor.cpp
+++ b/Engine/Source/Editor/Kismet/Private/BlueprintEditor.cpp
@@ -2868,10 +2868,6 @@ void FBlueprintEditor::CreateDefaultCommands()
 		FExecuteAction::CreateSP(this, &FBlueprintEditor::CreateMergeToolTab),
 		FCanExecuteAction());
 
-	ToolkitCommands->MapAction(FBlueprintEditorCommands::Get().GenerateNativeCode,
-		FExecuteAction::CreateSP(this, &FBlueprintEditor::OpenNativeCodeGenerationTool),
-		FCanExecuteAction::CreateSP(this, &FBlueprintEditor::CanGenerateNativeCode));
-
 	ToolkitCommands->MapAction(FBlueprintEditorCommands::Get().GenerateSearchIndex,
 		FExecuteAction::CreateSP(this, &FBlueprintEditor::OnGenerateSearchIndexForDebugging),
 		FCanExecuteAction());
@@ -2922,21 +2918,6 @@ void FBlueprintEditor::CreateDefaultCommands()
 	);
 }
 
-void FBlueprintEditor::OpenNativeCodeGenerationTool()
-{
-	UBlueprint* Blueprint = GetBlueprintObj();
-	if (Blueprint)
-	{
-		FNativeCodeGenerationTool::Open(*Blueprint, SharedThis(this));
-	}
-}
-
-bool FBlueprintEditor::CanGenerateNativeCode() const
-{
-	UBlueprint* Blueprint = GetBlueprintObj();
-	return Blueprint && FNativeCodeGenerationTool::CanGenerate(*Blueprint);
-}
-
 void FBlueprintEditor::OnGenerateSearchIndexForDebugging()
 {
 	UBlueprint* Blueprint = GetBlueprintObj();
@@ -3095,25 +3076,6 @@ void FBlueprintEditor::ReparentBlueprint_NewParentChosen(UClass* ChosenClass)
 			// Ensure that the Blueprint is up-to-date (valid SCS etc.) after compiling (new parent class)
 			EnsureBlueprintIsUpToDate(BlueprintObj);
 
-			if (BlueprintObj->NativizationFlag != EBlueprintNativizationFlag::Disabled)
-			{
-				UBlueprint* ParentBlueprint = UBlueprint::GetBlueprintFromClass(ChosenClass);
-				if (ParentBlueprint && ParentBlueprint->NativizationFlag == EBlueprintNativizationFlag::Disabled)
-				{
-					ParentBlueprint->NativizationFlag = EBlueprintNativizationFlag::Dependency;
-
-					FNotificationInfo Warning(FText::Format(
-						LOCTEXT("InterfaceFlaggedForNativization", "{0} flagged for nativization (as a required dependency)."),
-						FText::FromName(ParentBlueprint->GetFName())
-						)
-					);
-					Warning.ExpireDuration = 5.0f;
-					Warning.bFireAndForget = true;
-					Warning.Image = FCoreStyle::Get().GetBrush(TEXT("MessageLog.Warning"));
-					FSlateNotificationManager::Get().AddNotification(Warning);
-				}
-			}
-
 			if (SCSEditor.IsValid())
 			{
 				SCSEditor->UpdateTree();
diff --git a/Engine/Source/Editor/Kismet/Private/BlueprintEditorCommands.cpp b/Engine/Source/Editor/Kismet/Private/BlueprintEditorCommands.cpp
--- a/Engine/Source/Editor/Kismet/Private/BlueprintEditorCommands.cpp
+++ b/Engine/Source/Editor/Kismet/Private/BlueprintEditorCommands.cpp
@@ -69,7 +69,6 @@ void FBlueprintEditorCommands::RegisterCommands()
 
 	// Development commands
 	UI_COMMAND( SaveIntermediateBuildProducts, "Save Intermediate Build Products", "Should the compiler save intermediate build products for debugging.", EUserInterfaceActionType::ToggleButton, FInputChord() );
-	UI_COMMAND(GenerateNativeCode, "Generate Native Code", "Generate C++ code from the blueprint", EUserInterfaceActionType::Button, FInputChord());
 	UI_COMMAND(GenerateSearchIndex, "Generate Search Index", "Generate the search index for this blueprint.", EUserInterfaceActionType::Button, FInputChord());
 	UI_COMMAND(DumpCachedIndexData, "Dump Cached Index Data", "Dump the cached index data for this blueprint.", EUserInterfaceActionType::Button, FInputChord());
 	UI_COMMAND(ShowActionMenuItemSignatures, "Show Action Menu Item Signatures", "If enabled, tooltips on action menu items will show the associated action's signature id (can be used to setup custom favorites menus).", EUserInterfaceActionType::ToggleButton, FInputChord());
diff --git a/Engine/Source/Editor/Kismet/Private/BlueprintEditorCommands.h b/Engine/Source/Editor/Kismet/Private/BlueprintEditorCommands.h
--- a/Engine/Source/Editor/Kismet/Private/BlueprintEditorCommands.h
+++ b/Engine/Source/Editor/Kismet/Private/BlueprintEditorCommands.h
@@ -70,7 +70,6 @@ public:
 
 	// Development commands
 	TSharedPtr< FUICommandInfo > SaveIntermediateBuildProducts;
-	TSharedPtr< FUICommandInfo > GenerateNativeCode;
 	TSharedPtr< FUICommandInfo > GenerateSearchIndex;
 	TSharedPtr< FUICommandInfo > DumpCachedIndexData;
 	TSharedPtr< FUICommandInfo > ShowActionMenuItemSignatures;
diff --git a/Engine/Source/Editor/Kismet/Private/SBlueprintDiff.cpp b/Engine/Source/Editor/Kismet/Private/SBlueprintDiff.cpp
--- a/Engine/Source/Editor/Kismet/Private/SBlueprintDiff.cpp
+++ b/Engine/Source/Editor/Kismet/Private/SBlueprintDiff.cpp
@@ -383,18 +383,6 @@ public:
 				Children.Push(Entry);
 				OutRealDifferences.Push(Entry);
 			}
-
-			if (OldBlueprint->SupportsNativization() != NewBlueprint->SupportsNativization())
-			{
-				FText DiffText = FText::Format(LOCTEXT("NativizationChanged", "Nativization changed from {0} to {1}"), FText::AsNumber(OldBlueprint->SupportsNativization()), FText::AsNumber(NewBlueprint->SupportsNativization()));
-
-				TSharedPtr<FBlueprintDifferenceTreeEntry> Entry = MakeShared<FBlueprintDifferenceTreeEntry>(
-					SelectionCallback,
-					FGenerateDiffEntryWidget::CreateStatic(&GenerateSimpleDiffWidget, DiffText));
-
-				Children.Push(Entry);
-				OutRealDifferences.Push(Entry);
-			}
 		}
 
 		const bool bHasDifferences = Children.Num() != 0;
diff --git a/Engine/Source/Editor/Kismet/Private/SBlueprintEditorToolbar.cpp b/Engine/Source/Editor/Kismet/Private/SBlueprintEditorToolbar.cpp
--- a/Engine/Source/Editor/Kismet/Private/SBlueprintEditorToolbar.cpp
+++ b/Engine/Source/Editor/Kismet/Private/SBlueprintEditorToolbar.cpp
@@ -112,11 +112,6 @@ void FKismet2Menu::FillDeveloperMenu(UToolMenu* InMenu)
 		Section.AddMenuEntry( FBlueprintEditorCommands::Get().SaveIntermediateBuildProducts );
 	}
 
-	{
-		FToolMenuSection& Section = InMenu->AddSection("GenerateNativeCode", LOCTEXT("Cpp", "C++"));
-		Section.AddMenuEntry(FBlueprintEditorCommands::Get().GenerateNativeCode);
-	}
-
 	if (FFindInBlueprintSearchManager::Get().ShouldEnableDeveloperMenuTools())
 	{
 		FToolMenuSection& Section = InMenu->AddSection("FileDeveloperSearchTools", LOCTEXT("SearchToolsHeading", "Search Tools"));
diff --git a/Engine/Source/Editor/Kismet/Public/BlueprintAssetHandler.h b/Engine/Source/Editor/Kismet/Public/BlueprintAssetHandler.h
--- a/Engine/Source/Editor/Kismet/Public/BlueprintAssetHandler.h
+++ b/Engine/Source/Editor/Kismet/Public/BlueprintAssetHandler.h
@@ -48,7 +48,8 @@ public:
 	 *@param OutReason       (Optional) An optional failure text to set
 	 *@return true if the specified asset supports nativization, false otherwise
 	 */
-	virtual bool SupportsNativization(const UObject* InAsset, const UBlueprint* InBlueprint, FText* OutReason) const { return true; }
+	UE_DEPRECATED(5.0, "Blueprint Nativization has been removed as a supported feature. This API will eventually be removed.")
+	virtual bool SupportsNativization(const UObject* InAsset, const UBlueprint* InBlueprint, FText* OutReason) const { return false; }
 };
 
 
diff --git a/Engine/Source/Editor/Kismet/Public/BlueprintEditor.h b/Engine/Source/Editor/Kismet/Public/BlueprintEditor.h
--- a/Engine/Source/Editor/Kismet/Public/BlueprintEditor.h
+++ b/Engine/Source/Editor/Kismet/Public/BlueprintEditor.h
@@ -589,12 +589,6 @@ public:
 	/** Closes the merge tool, rather than simply hiding it */
 	void CloseMergeTool();
 
-	/** Opens a native code generation modal window */
-	void OpenNativeCodeGenerationTool();
-
-	/** Can generate native code for current blueprint */
-	bool CanGenerateNativeCode() const;
-
 	/** Dumps the current blueprint search index to a JSON file for debugging purposes */
 	void OnGenerateSearchIndexForDebugging();
 
diff --git a/Engine/Source/Editor/TurnkeySupport/Private/TurnkeySupportModule.cpp b/Engine/Source/Editor/TurnkeySupport/Private/TurnkeySupportModule.cpp
--- a/Engine/Source/Editor/TurnkeySupport/Private/TurnkeySupportModule.cpp
+++ b/Engine/Source/Editor/TurnkeySupport/Private/TurnkeySupportModule.cpp
@@ -95,7 +95,6 @@ protected:
 	static bool ShouldBuildProject(UProjectPackagingSettings* PackagingSettings, const ITargetPlatform* TargetPlatform)
 	{
 		const UProjectPackagingSettings::FConfigurationInfo& ConfigurationInfo = UProjectPackagingSettings::ConfigurationInfo[(int)PackagingSettings->BuildConfiguration];
-		bool bAssetNativizationEnabled = (PackagingSettings->BlueprintNativizationMethod != EProjectPackagingBlueprintNativizationMethod::Disabled);
 
 		// Get the target to build
 		const FTargetInfo* Target = PackagingSettings->GetBuildTargetInfo();
@@ -113,7 +112,7 @@ protected:
 		else if (PackagingSettings->Build == EProjectPackagingBuild::IfProjectHasCode)
 		{
 			bBuild = true;
-			if (FApp::GetEngineIsPromotedBuild() && !bAssetNativizationEnabled)
+			if (FApp::GetEngineIsPromotedBuild())
 			{
 				FString BaseDir;
 
@@ -1621,7 +1620,9 @@ if (/*PlatformInfo->DataDrivenPlatformInfo->GetSdkStatus() != ETurnkeyPlatformSd
 
 UProjectPackagingSettings* PackagingSettings = Cast<UProjectPackagingSettings>(UProjectPackagingSettings::StaticClass()->GetDefaultObject());
 const UProjectPackagingSettings::FConfigurationInfo& ConfigurationInfo = UProjectPackagingSettings::ConfigurationInfo[(int)PackagingSettings->BuildConfiguration];
-bool bAssetNativizationEnabled = (PackagingSettings->BlueprintNativizationMethod != EProjectPackagingBlueprintNativizationMethod::Disabled);
+
+// Note: Blueprint nativization has been removed as a feature.
+const bool bAssetNativizationEnabled = false;
 
 const ITargetPlatform* const Platform = GetTargetPlatformManager()->FindTargetPlatform(PlatformInfo->Name.ToString());
 {
diff --git a/Engine/Source/Editor/UnrealEd/Private/CookOnTheFlyServer.cpp b/Engine/Source/Editor/UnrealEd/Private/CookOnTheFlyServer.cpp
--- a/Engine/Source/Editor/UnrealEd/Private/CookOnTheFlyServer.cpp
+++ b/Engine/Source/Editor/UnrealEd/Private/CookOnTheFlyServer.cpp
@@ -87,7 +87,6 @@
 #include "AssetRegistryModule.h"
 #include "AssetRegistryState.h"
 #include "CookerSettings.h"
-#include "BlueprintNativeCodeGenModule.h"
 
 #include "GameDelegates.h"
 #include "IPAddress.h"
@@ -4019,11 +4018,6 @@ void UCookOnTheFlyServer::SaveCookedPackage(UE::Cook::FPackageData& PackageData,
 					}
 
 					GIsCookerLoadingPackage = false;
-					{
-						UE_SCOPED_HIERARCHICAL_COOKTIMER(ConvertingBlueprints);
-						IBlueprintNativeCodeGenModule::Get().Convert(Package, Result.Result, *(Target->PlatformName()));
-					}
-
 					++this->StatSavedPackageCount;
 
 					// If package was actually saved check with asset manager to make sure it wasn't excluded for being a development or never cook package. We do this after Editor Only filtering
@@ -6433,20 +6427,6 @@ void UCookOnTheFlyServer::CookByTheBookFinished()
 	FString LibraryName = !IsCookingDLC() ? FApp::GetProjectName() : CookByTheBookOptions->DlcName;
 
 	{
-		if (IBlueprintNativeCodeGenModule::IsNativeCodeGenModuleLoaded())
-		{
-			UE_SCOPED_HIERARCHICAL_COOKTIMER(GeneratingBlueprintAssets)
-			IBlueprintNativeCodeGenModule& CodeGenModule = IBlueprintNativeCodeGenModule::Get();
-
-			CodeGenModule.GenerateFullyConvertedClasses(); // While generating fully converted classes the list of necessary stubs is created.
-			CodeGenModule.GenerateStubs();
-
-			CodeGenModule.FinalizeManifest();
-
-			// Unload the module as we only need it while cooking. This will also clear the current module's state in order to allow a new cooker pass to function properly.
-			FModuleManager::Get().UnloadModule(CodeGenModule.GetModuleName());
-		}
-
 		// Save modified asset registry with all streaming chunk info generated during cook
 		const FString& SandboxRegistryFilename = GetSandboxAssetRegistryFilename();
 
@@ -7279,21 +7259,6 @@ void UCookOnTheFlyServer::StartCookByTheBook( const FCookByTheBookStartupOptions
 		StartSavingEDLCookInfoForVerification();
 	}
 
-	// Note: Nativization only works with "cook by the book" mode and not from within the current editor process.
-	if (CurrentCookMode == ECookMode::CookByTheBook
-		&& PackagingSettings->BlueprintNativizationMethod != EProjectPackagingBlueprintNativizationMethod::Disabled)
-	{
-		FNativeCodeGenInitData CodeGenData;
-		for (const ITargetPlatform* Entry : CookByTheBookStartupOptions.TargetPlatforms)
-		{
-			FPlatformNativizationDetails PlatformNativizationDetails;
-			IBlueprintNativeCodeGenModule::Get().FillPlatformNativizationDetails(Entry, PlatformNativizationDetails);
-			CodeGenData.CodegenTargets.Push(PlatformNativizationDetails);
-		}
-		CodeGenData.ManifestIdentifier = -1;
-		IBlueprintNativeCodeGenModule::InitializeModule(CodeGenData);
-	}
-
 	{
 		if (CookByTheBookOptions->bGenerateDependenciesForMaps)
 		{
diff --git a/Engine/Source/Editor/UnrealEd/Private/Kismet2/BlueprintEditorUtils.cpp b/Engine/Source/Editor/UnrealEd/Private/Kismet2/BlueprintEditorUtils.cpp
--- a/Engine/Source/Editor/UnrealEd/Private/Kismet2/BlueprintEditorUtils.cpp
+++ b/Engine/Source/Editor/UnrealEd/Private/Kismet2/BlueprintEditorUtils.cpp
@@ -43,7 +43,6 @@
 #include "EdMode.h"
 #include "Dialogs/Dialogs.h"
 #include "UnrealEdGlobals.h"
-#include "Settings/ProjectPackagingSettings.h"
 #include "Matinee/MatineeActor.h"
 #include "Engine/LevelScriptBlueprint.h"
 #include "UObject/BlueprintsObjectVersion.h"
@@ -5998,12 +5997,6 @@ void FBlueprintEditorUtils::MarkBlueprintChildrenAsModified(UBlueprint* InBluepr
 
 //////////////////////////////////////////////////////////////////////////
 
-void FBlueprintEditorUtils::FindNativizationDependencies(UBlueprint* Blueprint, TArray<UClass*>& NativizeDependenciesOut)
-{
-	FBlueprintEditorUtils::FindImplementedInterfaces(Blueprint, /*bGetAllInterfaces =*/false, NativizeDependenciesOut);
-	NativizeDependenciesOut.AddUnique(Blueprint->ParentClass);
-}
-
 /** Shared function for posting notification toasts */
 static void ShowNotification(const FText& Message, EMessageSeverity::Type Severity)
 {
@@ -6048,115 +6041,6 @@ static void ShowNotification(const FText& Message, EMessageSeverity::Type Severi
 	}
 }
 
-bool FBlueprintEditorUtils::PropagateNativizationSetting(UBlueprint* Blueprint)
-{
-	bool bSettingsChanged = false;
-	UProjectPackagingSettings* PackagingSettings = GetMutableDefault<UProjectPackagingSettings>();
-
-	switch (Blueprint->NativizationFlag)
-	{
-	case EBlueprintNativizationFlag::Disabled:
-		bSettingsChanged |= PackagingSettings->RemoveBlueprintAssetFromNativizationList(Blueprint);
-		break;
-
-	case EBlueprintNativizationFlag::ExplicitlyEnabled:
-		{
-			bSettingsChanged |= PackagingSettings->AddBlueprintAssetToNativizationList(Blueprint);
-
-			TArray<UClass*> NativizationDependencies;
-			FindNativizationDependencies(Blueprint, NativizationDependencies);
-
-			bool bAddedDependencies = false;
-
-			for (UClass* Dependency : NativizationDependencies)
-			{
-				if (UBlueprint* DependencyBp = UBlueprint::GetBlueprintFromClass(Dependency))
-				{
-					// if the user hasn't manually altered the setting (chosen 
-					// for themselves), then let's apply the auto-setting
-					if (DependencyBp->NativizationFlag == EBlueprintNativizationFlag::Dependency)
-					{
-						DependencyBp->NativizationFlag = EBlueprintNativizationFlag::ExplicitlyEnabled;
-						// recurse and propagate this setting to dependencies once removed
-						bAddedDependencies |= PropagateNativizationSetting(DependencyBp);
-					}
-					else if (DependencyBp->NativizationFlag == EBlueprintNativizationFlag::ExplicitlyEnabled &&
-						!PackagingSettings->IsBlueprintAssetInNativizationList(DependencyBp))
-					{
-						bAddedDependencies |= PropagateNativizationSetting(DependencyBp);
-						// this is a hairy case, because the user could have changes pending to the DependencyBp
-						// that they will end up discarding... is it their intension to discard the "nativize"  
-						// setting as well? was it set before or after this dependent? maybe they set it before,
-						// and want to discard the change, but didn't realize it was a dependency?
-						// here we'll favor correctness, and save it to the config now
-					}
-				}
-			}
-			bSettingsChanged |= bAddedDependencies;
-			if (bAddedDependencies)
-			{
-				ShowNotification(LOCTEXT("DependenciesSavedForNativization", "Saved extra (required dependency) Blueprints for nativization."), EMessageSeverity::Warning);
-			}
-		}
-		break;
-
-	default:
-	case EBlueprintNativizationFlag::Dependency:
-		// the Blueprint which set this flag is responsible for applying this change
-		break;
-	}
-
-	return bSettingsChanged;
-}
-
-bool FBlueprintEditorUtils::ShouldNativizeImplicitly(const UBlueprint* Blueprint)
-{
-	if (Blueprint)
-	{
-		TArray<UK2Node_Event*> AllEventNodes;
-		FBlueprintEditorUtils::GetAllNodesOfClass<UK2Node_Event>(Blueprint, AllEventNodes);
-
-		// Add all events overridden by this Blueprint.
-		TArray<FName> CheckFunctionNames;
-		for (const UK2Node_Event* EventNode : AllEventNodes)
-		{
-			if (EventNode->bOverrideFunction)
-			{
-				CheckFunctionNames.Add(EventNode->EventReference.GetMemberName());
-			}
-		}
-
-		// Add all function graphs implemented by this Blueprint.
-		for (const UEdGraph* FunctionGraph : Blueprint->FunctionGraphs)
-		{
-			CheckFunctionNames.Add(FunctionGraph->GetFName());
-		}
-
-		// Check each overridable/callable function defined by all ancestors to see if any names match an implementation found in this Blueprint.
-		UClass* ParentClass = Blueprint->SkeletonGeneratedClass ? Blueprint->SkeletonGeneratedClass->GetSuperClass() : *Blueprint->ParentClass;
-		for (TFieldIterator<UFunction> FunctionIt(ParentClass, EFieldIteratorFlags::IncludeSuper); FunctionIt; ++FunctionIt)
-		{
-			const UFunction* Function = *FunctionIt;
-			if (UEdGraphSchema_K2::CanKismetOverrideFunction(Function) && UEdGraphSchema_K2::CanUserKismetCallFunction(Function) && CheckFunctionNames.Contains(Function->GetFName()))
-			{
-				// This Blueprint overrides a callable event/function. If the function is defined in a parent BP that is flagged for nativization, OR if
-				// the parent BP has itself been implicitly flagged for nativization, then this Blueprint will also be implicitly flagged for nativization.
-				// Currently, any calls to such a function within a nativized parent hierarchy are not able to invoke an override in a non-nativized child,
-				// so the current solution is to implicitly force the child BP to also be nativized along with its parent hierarchy in this particular case.
-				const UClass* SignatureClass = CastChecked<UClass>(Function->GetOuter());
-				const UBlueprint* ParentBP = UBlueprint::GetBlueprintFromClass(SignatureClass);
-				if (ParentBP != nullptr
-					&& (ParentBP->NativizationFlag == EBlueprintNativizationFlag::ExplicitlyEnabled || ShouldNativizeImplicitly(ParentBP)))
-				{
-					return true;
-				}
-			}
-		}
-	}
-
-	return false;
-}
-
 //////////////////////////////////////////////////////////////////////////
 // Interfaces
 
@@ -6271,21 +6155,6 @@ bool FBlueprintEditorUtils::ImplementNewInterface(UBlueprint* Blueprint, const F
 	{
 		Blueprint->ImplementedInterfaces.Add(NewInterface);
 		FBlueprintEditorUtils::MarkBlueprintAsStructurallyModified(Blueprint);
-
-		if (Blueprint->NativizationFlag != EBlueprintNativizationFlag::Disabled)
-		{
-			UBlueprint* InterfaceBlueprint = UBlueprint::GetBlueprintFromClass(InterfaceClass);
-			if (InterfaceBlueprint && InterfaceBlueprint->NativizationFlag == EBlueprintNativizationFlag::Disabled)
-			{
-				InterfaceBlueprint->NativizationFlag = EBlueprintNativizationFlag::Dependency;
-				ShowNotification(FText::Format(
-					LOCTEXT("InterfaceFlaggedForNativization", "{0} flagged for nativization (as a required dependency)."),
-					FText::FromName(InterfaceBlueprint->GetFName())
-					),
-					EMessageSeverity::Warning
-				);
-			}
-		}
 	}
 	return bAllFunctionsAdded;
 }
diff --git a/Engine/Source/Editor/UnrealEd/Private/UnrealEdMisc.cpp b/Engine/Source/Editor/UnrealEd/Private/UnrealEdMisc.cpp
--- a/Engine/Source/Editor/UnrealEd/Private/UnrealEdMisc.cpp
+++ b/Engine/Source/Editor/UnrealEd/Private/UnrealEdMisc.cpp
@@ -265,8 +265,6 @@ void FUnrealEdMisc::OnInit()
 	FCoreDelegates::PreModal.AddRaw(this, &FUnrealEdMisc::OnEditorPreModal);
 	FCoreDelegates::PostModal.AddRaw(this, &FUnrealEdMisc::OnEditorPostModal);
 
-	FDeveloperToolSettingsDelegates::OnNativeBlueprintsSettingChanged.AddRaw(this, &FUnrealEdMisc::OnNativizeBlueprintsSettingChanged);
-
 	// Register the play world commands
 	FPlayWorldCommands::Register();
 	FPlayWorldCommands::BindGlobalPlayWorldCommands();
@@ -1224,26 +1222,6 @@ void FUnrealEdMisc::OnEditorPostModal()
 	}
 }
 
-void FUnrealEdMisc::OnNativizeBlueprintsSettingChanged(const FString& PackageName, bool bSelect)
-{
-	if (!PackageName.IsEmpty())
-	{
-		// This should only apply to loaded packages. Any unloaded packages defer setting the transient flag to when they're loaded.
-		if (UPackage* Package = FindPackage(nullptr, *PackageName))
-		{
-			// Find the Blueprint asset within the package.
-			if (UBlueprint* Blueprint = FindObject<UBlueprint>(Package, *FPaths::GetBaseFilename(PackageName)))
-			{
-				// We're toggling the transient flag on or off.
-				if ((Blueprint->NativizationFlag == EBlueprintNativizationFlag::ExplicitlyEnabled) != bSelect)
-				{
-					Blueprint->NativizationFlag = bSelect ? EBlueprintNativizationFlag::ExplicitlyEnabled : EBlueprintNativizationFlag::Disabled;
-				}
-			}
-		}
-	}
-}
-
 void FUnrealEdMisc::OnActiveTabChanged(TSharedPtr<SDockTab> PreviouslyActive, TSharedPtr<SDockTab> NewlyActivated)
 {
 	OnUserActivityTabChanged(NewlyActivated);
diff --git a/Engine/Source/Editor/UnrealEd/Public/Kismet2/BlueprintEditorUtils.h b/Engine/Source/Editor/UnrealEd/Public/Kismet2/BlueprintEditorUtils.h
--- a/Engine/Source/Editor/UnrealEd/Public/Kismet2/BlueprintEditorUtils.h
+++ b/Engine/Source/Editor/UnrealEd/Public/Kismet2/BlueprintEditorUtils.h
@@ -1352,13 +1352,16 @@ public:
 	static UActorComponent* FindUCSComponentTemplate(const FComponentKey& ComponentKey, const FName& TemplateName);
 
 	/** Takes the Blueprint's NativizedFlag property and applies it to the authoritative config (does the same for flagged dependencies) */
-	static bool PropagateNativizationSetting(UBlueprint* Blueprint);
+	UE_DEPRECATED(5.0, "Blueprint Nativization has been removed as a supported feature. This API will eventually be removed.")
+	static bool PropagateNativizationSetting(UBlueprint* Blueprint) { return false; }
 
 	/** Retrieves all dependencies that need to be nativized for this to work as a nativized Blueprint */
-	static void FindNativizationDependencies(UBlueprint* Blueprint, TArray<UClass*>& NativizeDependenciesOut);
+	UE_DEPRECATED(5.0, "Blueprint Nativization has been removed as a supported feature. This API will eventually be removed.")
+	static void FindNativizationDependencies(UBlueprint* Blueprint, TArray<UClass*>& NativizeDependenciesOut) {}
 
 	/** Returns whether or not the given Blueprint should be nativized implicitly, regardless of whether or not the user has explicitly enabled it */
-	static bool ShouldNativizeImplicitly(const UBlueprint* Blueprint);
+	UE_DEPRECATED(5.0, "Blueprint Nativization has been removed as a supported feature. This API will eventually be removed.")
+	static bool ShouldNativizeImplicitly(const UBlueprint* Blueprint) { return false; }
 
 	//////////////////////////////////////////////////////////////////////////
 	// Interface
diff --git a/Engine/Source/Editor/UnrealEd/Public/UnrealEdMisc.h b/Engine/Source/Editor/UnrealEd/Public/UnrealEdMisc.h
--- a/Engine/Source/Editor/UnrealEd/Public/UnrealEdMisc.h
+++ b/Engine/Source/Editor/UnrealEd/Public/UnrealEdMisc.h
@@ -248,7 +248,8 @@ public:
 	void OnEditorPostModal();
 
 	/** BP nativization settings */
-	void OnNativizeBlueprintsSettingChanged(const FString& PackageName, bool bSelect);
+	UE_DEPRECATED(5.0, "Blueprint Nativization has been removed as a supported feature. This delegate method will eventually be removed.")
+	void OnNativizeBlueprintsSettingChanged(const FString& PackageName, bool bSelect) {}
 
 	/** Called from tab manager when the tab changes */
 	void OnActiveTabChanged(TSharedPtr<SDockTab> PreviouslyActive, TSharedPtr<SDockTab> NewlyActivated);
diff --git a/Engine/Source/Programs/AutomationTool/AutomationUtils/ProjectParams.cs b/Engine/Source/Programs/AutomationTool/AutomationUtils/ProjectParams.cs
--- a/Engine/Source/Programs/AutomationTool/AutomationUtils/ProjectParams.cs
+++ b/Engine/Source/Programs/AutomationTool/AutomationUtils/ProjectParams.cs
@@ -364,7 +364,6 @@ namespace AutomationTool
 			this.TitleID = InParams.TitleID;
 			this.bTreatNonShippingBinariesAsDebugFiles = InParams.bTreatNonShippingBinariesAsDebugFiles;
 			this.bUseExtraFlavor = InParams.bUseExtraFlavor;
-			this.RunAssetNativization = InParams.RunAssetNativization;
 			this.AdditionalPackageOptions = InParams.AdditionalPackageOptions;
 		}
 
@@ -502,7 +501,6 @@ namespace AutomationTool
 			bool? IterativeDeploy = null,
 			bool? FastCook = null,
 			bool? IgnoreCookErrors = null,
-            bool? RunAssetNativization = null,
 			bool? CodeSign = null,
 			bool? TreatNonShippingBinariesAsDebugFiles = null,
 			bool? UseExtraFlavor = null,
@@ -794,18 +792,6 @@ namespace AutomationTool
 			this.FastCook = GetParamValueIfNotSpecified(Command, FastCook, this.FastCook, "FastCook");
 			this.IgnoreCookErrors = GetParamValueIfNotSpecified(Command, IgnoreCookErrors, this.IgnoreCookErrors, "IgnoreCookErrors");
 
-            // Determine whether or not we're going to nativize Blueprint assets at cook time.
-            this.RunAssetNativization = false;
-            ConfigHierarchy GameIni = ConfigCache.ReadHierarchy(ConfigHierarchyType.Game, RawProjectPath.Directory, HostPlatform.Current.HostEditorPlatform);
-            if (GameIni != null)
-            {
-                string BlueprintNativizationMethod;
-                if (GameIni.TryGetValue("/Script/UnrealEd.ProjectPackagingSettings", "BlueprintNativizationMethod", out BlueprintNativizationMethod))
-                {
-                    this.RunAssetNativization = !string.IsNullOrEmpty(BlueprintNativizationMethod) && BlueprintNativizationMethod != "Disabled";
-                }
-            }
-
             string DeviceString = ParseParamValueIfNotSpecified(Command, Device, "device", String.Empty).Trim(new char[] { '\"' });
             if(DeviceString == "")
             {
@@ -1377,7 +1363,11 @@ namespace AutomationTool
         /// <summary>
         /// Determines if Blueprint assets should be substituted with auto-generated code.
         /// </summary>
-        public bool RunAssetNativization;
+		[Obsolete("The RunAssetNativization property has been deprecated. This feature is no longer supported.")]
+        public bool RunAssetNativization
+		{
+			get { return false; }
+		}
 
 		/// <summary>
 		/// Keeps track of any '-ini:type:[section]:value' arguments on the command line. These will override cached config settings for the current process, and can be passed along to other tools.
@@ -2048,7 +2038,9 @@ namespace AutomationTool
 			}
 
 			List<UnrealTargetPlatform> ClientTargetPlatformTypes = ClientTargetPlatforms.ConvertAll(x => x.Type).Distinct().ToList();
-			var Properties = ProjectUtils.GetProjectProperties(RawProjectPath, ClientTargetPlatformTypes, ClientConfigsToBuild, RunAssetNativization);
+			// @todo (wip) - Removing Blueprint nativization as a feature.
+			bool bRunAssetNativization = false;// this.RunAssetNativization;
+			var Properties = ProjectUtils.GetProjectProperties(RawProjectPath, ClientTargetPlatformTypes, ClientConfigsToBuild, bRunAssetNativization);
 
 			bIsCodeBasedProject = Properties.bIsCodeBasedProject;
 			DetectedTargets = Properties.Targets;
@@ -2904,7 +2896,6 @@ namespace AutomationTool
 				CommandUtils.LogLog("Stage={0}", Stage);
 				CommandUtils.LogLog("bTreatNonShippingBinariesAsDebugFiles={0}", bTreatNonShippingBinariesAsDebugFiles);
 				CommandUtils.LogLog("bUseExtraFlavor={0}", bUseExtraFlavor);
-				CommandUtils.LogLog("NativizeAssets={0}", RunAssetNativization);
                 CommandUtils.LogLog("StageDirectoryParam={0}", StageDirectoryParam);
 				CommandUtils.LogLog("AdditionalPackageOptions={0}", AdditionalPackageOptions);
 				CommandUtils.LogLog("Project Params **************");
diff --git a/Engine/Source/Programs/AutomationTool/Scripts/BuildCookRun.Automation.cs b/Engine/Source/Programs/AutomationTool/Scripts/BuildCookRun.Automation.cs
--- a/Engine/Source/Programs/AutomationTool/Scripts/BuildCookRun.Automation.cs
+++ b/Engine/Source/Programs/AutomationTool/Scripts/BuildCookRun.Automation.cs
@@ -194,25 +194,14 @@ public class BuildCookRun : BuildCommand
 
 	protected void DoBuildCookRun(ProjectParams Params)
 	{
-		const ProjectBuildTargets ClientTargets = ProjectBuildTargets.ClientCooked | ProjectBuildTargets.ServerCooked;
-        bool bGenerateNativeScripts = Params.RunAssetNativization;
 		int WorkingCL = -1;
 		if (P4Enabled && GlobalCommandLine.Submit && AllowSubmit)
 		{
 			WorkingCL = P4.CreateChange(P4Env.Client, String.Format("{0} build from changelist {1}", Params.ShortProjectName, P4Env.Changelist));
 		}
 
-        Project.Build(this, Params, WorkingCL, bGenerateNativeScripts ? (ProjectBuildTargets.All & ~ClientTargets) : ProjectBuildTargets.All);
+        Project.Build(this, Params, WorkingCL, ProjectBuildTargets.All);
 		Project.Cook(Params);
-        if (bGenerateNativeScripts)
-        {
-            // crash reporter is built along with client targets, so we need to 
-            // include that target flag here as well - note: that its not folded
-            // into ClientTargets because the editor needs its own CrashReporter 
-            // as well (which would be built above)
-            Project.Build(this, Params, WorkingCL, ClientTargets | ProjectBuildTargets.CrashReporter);
-        }
-        
 		Project.CopyBuildToStagingDirectory(Params);
 		Project.Package(Params, WorkingCL);
 		Project.Archive(Params);
diff --git a/Engine/Source/Programs/UnrealBuildTool/Configuration/TargetRules.cs b/Engine/Source/Programs/UnrealBuildTool/Configuration/TargetRules.cs
--- a/Engine/Source/Programs/UnrealBuildTool/Configuration/TargetRules.cs
+++ b/Engine/Source/Programs/UnrealBuildTool/Configuration/TargetRules.cs
@@ -1651,58 +1651,6 @@ namespace UnrealBuildTool
 			}
 		}
 
-		/// <summary>
-		/// Checks whether nativization is enabled for this target, and determine the path for the nativized plugin
-		/// </summary>
-		/// <returns>The nativized plugin file, or null if nativization is not enabled</returns>
-		internal FileReference GetNativizedPlugin()
-		{
-			if (ProjectFile != null && (Type == TargetType.Game || Type == TargetType.Client || Type == TargetType.Server))
-			{
-				// Read the config files for this project
-				ConfigHierarchy Config = ConfigCache.ReadHierarchy(ConfigHierarchyType.Game, ProjectFile.Directory, BuildHostPlatform.Current.Platform);
-				if (Config != null)
-				{
-					// Determine whether or not the user has enabled nativization of Blueprint assets at cook time (default is 'Disabled')
-					string NativizationMethod;
-					if (Config.TryGetValue("/Script/UnrealEd.ProjectPackagingSettings", "BlueprintNativizationMethod", out NativizationMethod) && NativizationMethod != "Disabled")
-					{
-						string PlatformName;
-						if (Platform == UnrealTargetPlatform.Win32 || Platform == UnrealTargetPlatform.Win64)
-						{
-							PlatformName = "Windows";
-						}
-						else
-						{
-							PlatformName = Platform.ToString();
-						}
-
-						// Temp fix to force platforms that only support "Game" configurations at cook time to the correct path.
-						string ProjectTargetType;
-						if (Platform == UnrealTargetPlatform.Win32 || Platform == UnrealTargetPlatform.Win64 || Platform == UnrealTargetPlatform.Linux || Platform == UnrealTargetPlatform.Mac)
-						{
-							ProjectTargetType = Type.ToString();
-						}
-						else
-						{
-							ProjectTargetType = "Game";
-						}
-
-						FileReference PluginFile = FileReference.Combine(ProjectFile.Directory, "Intermediate", "Plugins", "NativizedAssets", PlatformName, ProjectTargetType, "NativizedAssets.uplugin");
-						if (FileReference.Exists(PluginFile))
-						{
-							return PluginFile;
-						}
-						else
-						{
-							Log.TraceWarning("{0} is configured for nativization, but is missing the generated code plugin at \"{1}\". Make sure to cook {2} data before attempting to build the {3} target. If data was cooked with nativization enabled, this can also mean there were no Blueprint assets that required conversion, in which case this warning can be safely ignored.", Name, PluginFile.FullName, Type.ToString(), Platform.ToString());
-						}
-					}
-				}
-			}
-			return null;
-		}
-
 		/// <summary>
 		/// Gets a list of platforms that this target supports
 		/// </summary>
diff --git a/Engine/Source/Programs/UnrealBuildTool/Configuration/UEBuildTarget.cs b/Engine/Source/Programs/UnrealBuildTool/Configuration/UEBuildTarget.cs
--- a/Engine/Source/Programs/UnrealBuildTool/Configuration/UEBuildTarget.cs
+++ b/Engine/Source/Programs/UnrealBuildTool/Configuration/UEBuildTarget.cs
@@ -768,13 +768,6 @@ namespace UnrealBuildTool
 				RulesObject.bDisableLinking = true;
 			}
 
-			// Include generated code plugin if not building an editor target and project is configured for nativization
-			FileReference NativizedPluginFile = RulesObject.GetNativizedPlugin();
-			if(NativizedPluginFile != null)
-			{
-				RulesAssembly = RulesCompiler.CreatePluginRulesAssembly(NativizedPluginFile, bSkipRulesCompile, RulesAssembly, false);
-			}
-
 			// Generate a build target from this rules module
 			UEBuildTarget Target;
 			using(Timeline.ScopeEvent("UEBuildTarget constructor"))
diff --git a/Engine/Source/Programs/UnrealBuildTool/ProjectFiles/ProjectFileGenerator.cs b/Engine/Source/Programs/UnrealBuildTool/ProjectFiles/ProjectFileGenerator.cs
--- a/Engine/Source/Programs/UnrealBuildTool/ProjectFiles/ProjectFileGenerator.cs
+++ b/Engine/Source/Programs/UnrealBuildTool/ProjectFiles/ProjectFileGenerator.cs
@@ -1299,33 +1299,6 @@ namespace UnrealBuildTool
 			}
 		}
 
-        /// <summary>
-        /// Finds any additional plugin files.
-        /// </summary>
-        /// <returns>List of additional plugin files</returns>
-		private List<FileReference> DiscoverExtraPlugins(List<FileReference> AllGameProjects)
-		{
-			List<FileReference> AddedPlugins = new List<FileReference>();
-
-            foreach (FileReference GameProject in AllGameProjects)
-			{
-                // Check the user preference to see if they'd like to include nativized assets as a generated project.
-                bool bIncludeNativizedAssets = false;
-                ConfigHierarchy Config = ConfigCache.ReadHierarchy(ConfigHierarchyType.Game, GameProject.Directory, BuildHostPlatform.Current.Platform);
-                if (Config != null)
-                {
-                    Config.TryGetValue("/Script/UnrealEd.ProjectPackagingSettings", "bIncludeNativizedAssetsInProjectGeneration", out bIncludeNativizedAssets);
-                }
-
-                // Note: Whether or not we include nativized assets here has no bearing on whether or not they actually get built.
-                if (bIncludeNativizedAssets)
-                {
-                    AddedPlugins.AddRange(Plugins.EnumeratePlugins(DirectoryReference.Combine(GameProject.Directory, "Intermediate", "Plugins")).Where(x => x.GetFileNameWithoutExtension() == "NativizedAssets"));
-                }
-			}
-			return AddedPlugins;
-		}
-
 		/// <summary>
 		/// Finds all module files (filtering by platform)
 		/// </summary>
@@ -1335,7 +1308,7 @@ namespace UnrealBuildTool
 			List<FileReference> AllModuleFiles = new List<FileReference>();
 
 			// Locate all modules (*.Build.cs files)
-			List<FileReference> FoundModuleFiles = RulesCompiler.FindAllRulesSourceFiles( RulesCompiler.RulesFileType.Module, GameFolders: AllGameProjects.Select(x => x.Directory).ToList(), ForeignPlugins: DiscoverExtraPlugins(AllGameProjects), AdditionalSearchPaths:null );
+			List<FileReference> FoundModuleFiles = RulesCompiler.FindAllRulesSourceFiles( RulesCompiler.RulesFileType.Module, GameFolders: AllGameProjects.Select(x => x.Directory).ToList(), ForeignPlugins:null, AdditionalSearchPaths:null );
 			foreach( FileReference BuildFileName in FoundModuleFiles )
 			{
 				AllModuleFiles.Add( BuildFileName );
@@ -1379,7 +1352,7 @@ namespace UnrealBuildTool
 			List<string> UnsupportedPlatformNameStrings = Utils.MakeListOfUnsupportedPlatforms(SupportedPlatforms, bIncludeUnbuildablePlatforms:true);
 
 			// Locate all targets (*.Target.cs files)
-			List<FileReference> FoundTargetFiles = RulesCompiler.FindAllRulesSourceFiles( RulesCompiler.RulesFileType.Target, AllGameProjects.Select(x => x.Directory).ToList(), ForeignPlugins: DiscoverExtraPlugins(AllGameProjects), AdditionalSearchPaths:null, bIncludeTempTargets: bIncludeTempTargets );
+			List<FileReference> FoundTargetFiles = RulesCompiler.FindAllRulesSourceFiles( RulesCompiler.RulesFileType.Target, AllGameProjects.Select(x => x.Directory).ToList(), ForeignPlugins:null, AdditionalSearchPaths:null, bIncludeTempTargets: bIncludeTempTargets );
 			foreach( FileReference CurTargetFile in FoundTargetFiles )
 			{
 				string CleanTargetFileName = Utils.CleanDirectorySeparators( CurTargetFile.FullName );
diff --git a/Engine/Source/Programs/UnrealBuildTool/ToolChain/RemoteMac.cs b/Engine/Source/Programs/UnrealBuildTool/ToolChain/RemoteMac.cs
--- a/Engine/Source/Programs/UnrealBuildTool/ToolChain/RemoteMac.cs
+++ b/Engine/Source/Programs/UnrealBuildTool/ToolChain/RemoteMac.cs
@@ -446,13 +446,6 @@ namespace UnrealBuildTool
 				Mappings.Add(new RemoteMapping(RemoteLogFile.Directory, GetRemotePath(RemoteLogFile.Directory)));
 			}
 
-			// Check if we need to enable a nativized plugin, and compile the assembly for that if we do
-			FileReference NativizedPluginFile = Rules.GetNativizedPlugin();
-			if(NativizedPluginFile != null)
-			{
-				RulesAssembly = RulesCompiler.CreatePluginRulesAssembly(NativizedPluginFile, false, RulesAssembly, false);
-			}
-
 			// Path to the local manifest file. This has to be translated from the remote format after the build is complete.
 			List<FileReference> LocalManifestFiles = new List<FileReference>();
 
diff --git a/Engine/Source/Runtime/CoreUObject/Private/UObject/SavePackage/SavePackageUtilities.cpp b/Engine/Source/Runtime/CoreUObject/Private/UObject/SavePackage/SavePackageUtilities.cpp
--- a/Engine/Source/Runtime/CoreUObject/Private/UObject/SavePackage/SavePackageUtilities.cpp
+++ b/Engine/Source/Runtime/CoreUObject/Private/UObject/SavePackage/SavePackageUtilities.cpp
@@ -360,25 +360,6 @@ void ConditionallyExcludeObjectForTarget(UObject* Obj, EObjectMark ExcludedObjec
 	UObject* ObjOuter = Obj->GetOuter();
 	UClass* ObjClass = Obj->GetClass();
 
-	// if TargetPlatorm != nullptr then we are cooking
-	if (TargetPlatform)
-	{
-		// Check for nativization replacement
-		if (const IBlueprintNativeCodeGenCore* Coordinator = IBlueprintNativeCodeGenCore::Get())
-		{
-			const FCompilerNativizationOptions& NativizationOptions = Coordinator->GetNativizationOptionsForPlatform(TargetPlatform);
-			FName UnusedName;
-			if (UClass* ReplacedClass = Coordinator->FindReplacedClassForObject(Obj, NativizationOptions))
-			{
-				ObjClass = ReplacedClass;
-			}
-			if (UObject* ReplacedOuter = Coordinator->FindReplacedNameAndOuter(Obj, /*out*/UnusedName, NativizationOptions))
-			{
-				ObjOuter = ReplacedOuter;
-			}
-		}
-	}
-
 	EObjectMark NewMarks = CurrentMarks;
 
 	// Recurse into parents, then compute inherited marks
diff --git a/Engine/Source/Runtime/Engine/Classes/Engine/Blueprint.h b/Engine/Source/Runtime/Engine/Classes/Engine/Blueprint.h
--- a/Engine/Source/Runtime/Engine/Classes/Engine/Blueprint.h
+++ b/Engine/Source/Runtime/Engine/Classes/Engine/Blueprint.h
@@ -393,7 +393,7 @@ struct FBPEditorBookmarkNode
 };
 
 UENUM()
-enum class EBlueprintNativizationFlag : uint8
+enum class UE_DEPRECATED(5.0, "Blueprint Nativization has been removed as a supported feature. This type will eventually be removed.") EBlueprintNativizationFlag : uint8
 {
 	Disabled,
 	Dependency, // conditionally enabled (set from sub-class as a dependency)
@@ -506,15 +506,13 @@ class ENGINE_API UBlueprint : public UBlueprintCore
 	UPROPERTY()
 	mutable uint8 bDuplicatingReadOnly:1;
 
-private:
-	/** Deprecated properties. */
-	UPROPERTY()
-	uint8 bNativize_DEPRECATED:1;
-
 public:
 	/** When exclusive nativization is enabled, then this asset will be nativized. All super classes must be also nativized. */
+	UE_DEPRECATED(5.0, "Blueprint Nativization has been removed as a supported feature. This property will eventually be removed.")
+PRAGMA_DISABLE_DEPRECATION_WARNINGS
 	UPROPERTY(transient)
 	EBlueprintNativizationFlag NativizationFlag;
+PRAGMA_ENABLE_DEPRECATION_WARNINGS
 
 	/** The mode that will be used when compiling this class. */
 	UPROPERTY(EditAnywhere, Category=ClassOptions, AdvancedDisplay)
@@ -811,7 +809,8 @@ public:
 	/**
 	 * Check whether this blueprint can be nativized or not
 	 */
-	virtual bool SupportsNativization(FText* OutReason = nullptr) const;
+	UE_DEPRECATED(5.0, "Blueprint Nativization has been removed as a supported feature. This API will eventually be removed.")
+	virtual bool SupportsNativization(FText* OutReason = nullptr) const { return false; }
 
 private:
 
diff --git a/Engine/Source/Runtime/Engine/Private/Blueprint.cpp b/Engine/Source/Runtime/Engine/Private/Blueprint.cpp
--- a/Engine/Source/Runtime/Engine/Private/Blueprint.cpp
+++ b/Engine/Source/Runtime/Engine/Private/Blueprint.cpp
@@ -485,51 +485,6 @@ void UBlueprint::Serialize(FArchive& Ar)
 			FBlueprintEditorUtils::FixupVariableDescription(this, Variable);
 		}
 	}
-
-	if (Ar.IsPersistent())
-	{
-		bool bSettingsChanged = false;
-		const FString PackageName = GetOutermost()->GetName();
-		UProjectPackagingSettings* PackagingSettings = GetMutableDefault<UProjectPackagingSettings>();
-
-		if (Ar.IsLoading())
-		{
-#if WITH_EDITORONLY_DATA
-			if (bNativize_DEPRECATED)
-			{
-				// Migrate to the new transient flag.
-				bNativize_DEPRECATED = false;
-
-				NativizationFlag = EBlueprintNativizationFlag::ExplicitlyEnabled;
-				// Add this Blueprint asset to the exclusive list in the Project Settings (in case it doesn't exist).
-				bSettingsChanged |= PackagingSettings->AddBlueprintAssetToNativizationList(this);
-			}
-			else
-#endif
-			{
-				// Cache whether or not this Blueprint asset was selected for exclusive nativization in the Project Settings.
-				for (int AssetIndex = 0; AssetIndex < PackagingSettings->NativizeBlueprintAssets.Num(); ++AssetIndex)
-				{
-					if (PackagingSettings->NativizeBlueprintAssets[AssetIndex].FilePath.Equals(PackageName, ESearchCase::IgnoreCase))
-					{
-						NativizationFlag = EBlueprintNativizationFlag::ExplicitlyEnabled;
-						break;
-					}
-				}
-			}
-		}
-		else if (Ar.IsSaving())
-		{
-			bSettingsChanged |= FBlueprintEditorUtils::PropagateNativizationSetting(this);
-		}
-
-		if (bSettingsChanged)
-		{
-			// Update cached config settings and save.
-			PackagingSettings->SaveConfig();
-			PackagingSettings->UpdateDefaultConfigFile();
-		}
-	}
 #endif // WITH_EDITORONLY_DATA
 }
 
@@ -840,44 +795,6 @@ UClass* UBlueprint::GetBlueprintClass() const
 	return UBlueprintGeneratedClass::StaticClass();
 }
 
-bool UBlueprint::SupportsNativization(FText* OutReason) const
-{
-	// Previously commented out in FBlueprintNativeCodeGenModule::IsTargetedForReplacement - should 'const' blueprints be nativized??
-	//BPTYPE_Const,		// What is a "const" Blueprint?
-	if (BlueprintType == BPTYPE_MacroLibrary)
-	{
-		if (OutReason)
-		{
-			*OutReason = NSLOCTEXT("Blueprint", "MacroLibraryNativizationReason", "Macro Libraries cannot be nativized.");
-		}
-		return false;
-	}
-	else if (BlueprintType == BPTYPE_LevelScript)
-	{
-		if (OutReason)
-		{
-			*OutReason = NSLOCTEXT("Blueprint", "LevelScriptNativizationReason", "Level Blueprints cannot be nativized.");
-		}
-		return false;
-	}
-	else if (!GetOuter()->IsA<UPackage>())
-	{
-		// If this blueprint is not an asset itself, check whether the asset supports nativization
-		UObject* Asset = GetOuter();
-		while (Asset && !Asset->GetOuter()->IsA<UPackage>())
-		{
-			Asset = Asset->GetOuter();
-		}
-
-		const IBlueprintAssetHandler* Handler = Asset ? FBlueprintAssetHandler::Get().FindHandler(Asset->GetClass()) : nullptr;
-		if (Handler && !Handler->SupportsNativization(Asset, this, OutReason))
-		{
-			return false;
-		}
-	}
-	return true;
-}
-
 void UBlueprint::SetObjectBeingDebugged(UObject* NewObject)
 {
 	// Unregister the old object (even if PendingKill)
@@ -1309,84 +1226,6 @@ void UBlueprint::BeginCacheForCookedPlatformData(const ITargetPlatform *TargetPl
 			return false;
 		};
 
-		// If nativization is enabled and this Blueprint class will NOT be nativized, we need to determine if any of its parent Blueprints will be nativized and flag it for the runtime code.
-		// Note: Currently, this flag is set on Actor-based Blueprint classes only. If it's ever needed for non-Actor-based Blueprint classes at runtime, then this needs to be updated to match.
-		const IBlueprintNativeCodeGenCore* NativeCodeGenCore = IBlueprintNativeCodeGenCore::Get();
-		if (GeneratedClass != nullptr && NativeCodeGenCore != nullptr)
-		{
-			ensure(TargetPlatform);
-			const FCompilerNativizationOptions& NativizationOptions = NativeCodeGenCore->GetNativizationOptionsForPlatform(TargetPlatform);
-			TArray<const UBlueprintGeneratedClass*> ParentBPClassStack;
-			UBlueprintGeneratedClass::GetGeneratedClassesHierarchy(GeneratedClass->GetSuperClass(), ParentBPClassStack);
-			for (const UBlueprintGeneratedClass *ParentBPClass : ParentBPClassStack)
-			{
-				if (NativeCodeGenCore->IsTargetedForReplacement(ParentBPClass, NativizationOptions) == EReplacementResult::ReplaceCompletely)
-				{
-					if (UBlueprintGeneratedClass* BPGC = CastChecked<UBlueprintGeneratedClass>(*GeneratedClass))
-					{
-						// Flag that this BP class will have a nativized parent class.
-						BPGC->bHasNativizedParent = true;
-
-						// Cache the IsTargetedForReplacement() result for the parent BP class that we know to be nativized.
-						TMap<const UClass*, bool> ParentBPClassNativizationResultMap;
-						ParentBPClassNativizationResultMap.Add(ParentBPClass, true);
-
-						// Cook all overridden SCS component node templates inherited from parent BP classes that will be nativized.
-						if (UInheritableComponentHandler* TargetInheritableComponentHandler = BPGC->GetInheritableComponentHandler())
-						{
-							for (auto RecordIt = TargetInheritableComponentHandler->CreateRecordIterator(); RecordIt; ++RecordIt)
-							{
-								// Only generate cooked data if the target platform supports the template class type.
-								if (ShouldCookBlueprintComponentTemplate(RecordIt->ComponentTemplate))
-								{
-									// Get the original class that we're overriding a template from.
-									const UClass* ComponentTemplateOwnerClass = RecordIt->ComponentKey.GetComponentOwner();
-
-									// Check to see if we've already checked this class for nativization; if not, then cache the result of a new query and return the result.
-									bool bIsOwnerClassTargetedForReplacement = false;
-									bool* bCachedResult = ParentBPClassNativizationResultMap.Find(ComponentTemplateOwnerClass);
-									if (bCachedResult)
-									{
-										bIsOwnerClassTargetedForReplacement = *bCachedResult;
-									}
-									else
-									{
-										bool bResult = (NativeCodeGenCore->IsTargetedForReplacement(RecordIt->ComponentKey.GetComponentOwner(), NativizationOptions) == EReplacementResult::ReplaceCompletely);
-										bIsOwnerClassTargetedForReplacement = ParentBPClassNativizationResultMap.Add(ComponentTemplateOwnerClass, bResult);
-									}
-
-									if (bIsOwnerClassTargetedForReplacement)
-									{
-										// EDL is required, because we need to enforce a preload dependency on the CDO (see UBlueprintGeneratedClass::GetDefaultObjectPreloadDependencies). This is
-										// difficult to support in the non-EDL case, because we have to enforce the dependency at runtime, which can lead to unpredictable results in a cooked build.
-										if (IsEventDrivenLoaderEnabledInCookedBuilds())
-										{
-											// Use the template's archetype for the delta serialization here; remaining properties will have already been set via native subobject instancing at runtime.
-											constexpr bool bUseTemplateArchetype = true;
-											FBlueprintEditorUtils::BuildComponentInstancingData(RecordIt->ComponentTemplate, RecordIt->CookedComponentInstancingData, bUseTemplateArchetype);
-											++NumCookedComponents;
-										}
-										else
-										{
-											UE_LOG(LogBlueprint, Error, TEXT("%s overrides component \'%s\' inherited from %s, which will be converted to C++. This requires Event-Driven Loading (EDL) to be enabled; otherwise, %s must be excluded from Blueprint nativization."),
-												*GetName(),
-												*RecordIt->ComponentKey.GetSCSVariableName().ToString(),
-												*ComponentTemplateOwnerClass->GetName(),
-												*ComponentTemplateOwnerClass->GetName()
-											);
-										}
-									}
-								}
-							}
-						}
-					}
-					
-					// All remaining antecedent classes should be native or nativized; no need to continue.
-					break;
-				}
-			}
-		}
-
 		auto ShouldCookBlueprintComponentTemplateData = [](UBlueprintGeneratedClass* InBPGClass) -> bool
 		{
 			// Check to see if we should cook component data for the given class type.
