Changed montage BlendOut time to behave the same when using inertial blending + typos

diff --git a/Engine/Source/Runtime/Engine/Classes/Animation/AnimMontage.h b/Engine/Source/Runtime/Engine/Classes/Animation/AnimMontage.h
--- a/Engine/Source/Runtime/Engine/Classes/Animation/AnimMontage.h
+++ b/Engine/Source/Runtime/Engine/Classes/Animation/AnimMontage.h
@@ -129,7 +129,7 @@ enum class EMontageBlendMode : uint8
 {
 	//Uses standard weight based blend
 	Standard,
-	//Uses inertialization. Requires an inertialization node somewhere in the graph after every slot node used by this montage.
+	//Uses inertialization. Requires an inertialization node somewhere in the graph after any slot node used by this montage.
 	Inertialization,
 };
 
diff --git a/Engine/Source/Runtime/Engine/Private/Animation/AnimInstance.cpp b/Engine/Source/Runtime/Engine/Private/Animation/AnimInstance.cpp
--- a/Engine/Source/Runtime/Engine/Private/Animation/AnimInstance.cpp
+++ b/Engine/Source/Runtime/Engine/Private/Animation/AnimInstance.cpp
@@ -1868,7 +1868,7 @@ float UAnimInstance::Montage_PlayInternal(UAnimMontage* MontageToPlay, const FAl
 		if (CurrentSkeleton && CurrentSkeleton->IsCompatible(MontageToPlay->GetSkeleton()))
 		{
 			const float BlendDuration = BlendIn.GetBlendTime();
-			// If using interial blend, force an instant blend in
+			// If using inertial blend, force an instant blend in
 			const FAlphaBlend EffectiveBlendIn = (MontageToPlay->BlendMode == EMontageBlendMode::Inertialization) ? FAlphaBlend(BlendIn, 0.0f) : BlendIn;
 
 			const FName NewMontageGroupName = MontageToPlay->GetGroupName();
@@ -3219,7 +3219,7 @@ void UAnimInstance::UpdateMontageEvaluationData()
 	Proxy.GetSlotGroupInertializationRequestMap() = SlotGroupInertializationRequestMap;
 
 	// Reset inertialization requests every frame.
-	// If the request is missed by the graph (i.e. the slot node is not relevant), we assume assume what brought it back to relevancy will handle the blend instead.
+	// If the request is missed by the graph (i.e. the slot node is not relevant), we assume what brought it back to relevancy will handle the blend instead.
 	SlotGroupInertializationRequestMap.Reset();
 }
 
diff --git a/Engine/Source/Runtime/Engine/Private/Animation/AnimMontage.cpp b/Engine/Source/Runtime/Engine/Private/Animation/AnimMontage.cpp
--- a/Engine/Source/Runtime/Engine/Private/Animation/AnimMontage.cpp
+++ b/Engine/Source/Runtime/Engine/Private/Animation/AnimMontage.cpp
@@ -2306,19 +2306,13 @@ void FAnimMontageInstance::Advance(float DeltaTime, struct FRootMotionMovementPa
 						const float PlayTimeToEnd = MontageSubStepper.GetRemainingPlayTimeToSectionEnd(Position);
 
 						const bool bCustomBlendOutTriggerTime = (Montage->BlendOutTriggerTime >= 0);
-						const bool bUsesInertialization = Montage->BlendMode == EMontageBlendMode::Inertialization;
-
 						const float DefaultBlendOutTime = Montage->BlendOut.GetBlendTime() * DefaultBlendTimeMultiplier;
-
-						// If we are using inertialization, we don't need to stop the montage early. In this case, we also don't want to miss out on notifies during the blend out window.
-						const float EffectiveDefaultBlendOutTime = bUsesInertialization ? 0.0f : DefaultBlendOutTime;
-						const float BlendOutTriggerTime = bCustomBlendOutTriggerTime ? Montage->BlendOutTriggerTime : EffectiveDefaultBlendOutTime;
+						const float BlendOutTriggerTime = bCustomBlendOutTriggerTime ? Montage->BlendOutTriggerTime : DefaultBlendOutTime;
 
 						// ... trigger blend out if within blend out time window.
 						if (PlayTimeToEnd <= FMath::Max<float>(BlendOutTriggerTime, KINDA_SMALL_NUMBER))
 						{
-							// If we are using inertialization, use DefaultBlendOutTime, as we no longer crossfade between poses.
-							const float BlendOutTime = (bCustomBlendOutTriggerTime || bUsesInertialization) ? DefaultBlendOutTime : PlayTimeToEnd;
+							const float BlendOutTime = bCustomBlendOutTriggerTime ? DefaultBlendOutTime : PlayTimeToEnd;
 							Stop(FAlphaBlend(Montage->BlendOut, BlendOutTime), false);
 						}
 					}
