Fix outstanding issues with ICH override template object archetype lookups during serialization of child Blueprint class assets.

Change summary:
- Reordered FComponentOverrideRecord struct fields so that the key is serialized before the template object. Archetype lookups on the template object are now dependent on the key being valid as we now route these through the record table.
- ICH override template object serialization is now deferred in non-EDL builds (i.e. editor) as part of loading the Blueprint class asset. This ensures that the ICH override record key is valid in older assets before we let the linker instance the export object for the ICH override template object on load; per the above, this depends on a successful archetype lookup for its instantiation.
- Modified the signature of UClass::FindArchetype() to include an additional ObjectFlags input argument.
- Revised GetArchetypeFromRequiredInfoImpl() to now route ICH override template object archetype lookups through the owning class in place of the "old way" of doing a brute force subobject search through the outer class chain, which would sometimes fail and fall back to incorrectly return the archetype class's default object, resulting in data loss issues on load.
- Revised UBlueprintGeneratedClass::FindArchetype() to route ICH override template object archetype lookups through the parent class hierarchy, using the ICH override record's key to retrieve the original template as the basis for finding a match. This ensures that lookups won't fail due to things like SCS variable name changes in the base class and also that we now can properly handle archetype lookups for ICH records that override an inherited default scene root.

diff --git a/Engine/Source/Editor/UnrealEd/Private/Editor.cpp b/Engine/Source/Editor/UnrealEd/Private/Editor.cpp
--- a/Engine/Source/Editor/UnrealEd/Private/Editor.cpp
+++ b/Engine/Source/Editor/UnrealEd/Private/Editor.cpp
@@ -1083,7 +1083,7 @@ namespace EditorUtilities
 					// If the target actor's class is a child of our owner and we're both archetypes, then we're actually looking for an overridden version of ourselves
 					if (BPGC && TargetActor->GetClass()->IsChildOf(BPGC))
 					{
-						TargetComponent = Cast<UActorComponent>(TargetActor->GetClass()->FindArchetype(SourceComponent->GetClass(), SourceComponent->GetFName()));
+						TargetComponent = Cast<UActorComponent>(TargetActor->GetClass()->FindArchetype(SourceComponent->GetClass(), SourceComponent->GetFName(), SourceComponent->GetFlags()));
 
 						// If it is us, then we're done, we don't need to find this
 						if (TargetComponent == SourceComponent)
diff --git a/Engine/Source/Runtime/CoreUObject/Private/Blueprint/BlueprintSupport.cpp b/Engine/Source/Runtime/CoreUObject/Private/Blueprint/BlueprintSupport.cpp
--- a/Engine/Source/Runtime/CoreUObject/Private/Blueprint/BlueprintSupport.cpp
+++ b/Engine/Source/Runtime/CoreUObject/Private/Blueprint/BlueprintSupport.cpp
@@ -146,13 +146,23 @@ void FBlueprintSupport::SetClassReparentingFPtr(FClassReparentingFPtr Ptr)
 	ClassReparentingFPtr = Ptr;
 }
 
-bool FBlueprintSupport::IsDeferredDependencyPlaceholder(UObject* LoadedObj)
+bool FBlueprintSupport::IsDeferredDependencyPlaceholder(const UObject* LoadedObj)
 {
 	return LoadedObj && ( LoadedObj->IsA<ULinkerPlaceholderClass>() ||
 		LoadedObj->IsA<ULinkerPlaceholderFunction>() ||
 		LoadedObj->IsA<ULinkerPlaceholderExportObject>() );
 }
 
+FName FBlueprintSupport::GetDeferredExportNameForPlaceholderObject(const UObject* LoadedObj)
+{
+	if (const ULinkerPlaceholderExportObject* PlaceholderObject = Cast<ULinkerPlaceholderExportObject>(LoadedObj))
+	{
+		return PlaceholderObject->ExportName;
+	}
+
+	return NAME_None;
+}
+
 void FBlueprintSupport::RegisterDeferredDependenciesInStruct(const UStruct* Struct, void* StructData)
 {
 #if USE_CIRCULAR_DEPENDENCY_LOAD_DEFERRING
@@ -1181,6 +1191,7 @@ bool FLinkerLoad::DeferExportCreation(const int32 Index, UObject* Outer)
 		ULinkerPlaceholderExportObject* Placeholder = NewObject<ULinkerPlaceholderExportObject>(Outer, PlaceholderType, PlaceholderName, RF_Public | RF_Transient);
 		Placeholder->SetLinker(this, Index, false);
 		Placeholder->PackageIndex = FPackageIndex::FromExport(Index);
+		Placeholder->ExportName = Export.ObjectName;
 		
 		Export.Object = Placeholder;
 
@@ -1189,45 +1200,65 @@ bool FLinkerLoad::DeferExportCreation(const int32 Index, UObject* Outer)
 
 		return true;
 	}
+
+	UClass* BlueprintClass = LoadClass;
 	
-	if (LoadClass->HasAnyClassFlags(CLASS_Native))
+	// Always defer the creation of ICH templates so that the override record key can be serialized first. A valid key is required for ICH archetype lookups to succeed.
+	const bool bIsInheritableComponentTemplate = !!(Export.ObjectFlags & RF_InheritableComponentTemplate);
+	if(bIsInheritableComponentTemplate)
 	{
-		return false;
+		// Don't defer again if we're currently resolving this as a deferred export.
+		if(FResolvingExportTracker::Get().IsLinkerExportBeingResolved(this, Index))
+		{
+			return false;
+		}
+
+		// Redirect to the template's outer, which should always be a Blueprint class.
+		// Note: We cannot cast directly to a Blueprint class object here as that type is not visible to this module, so we check the class flags instead.
+		BlueprintClass = CastChecked<UClass>(Outer);
+		check(BlueprintClass->HasAllClassFlags(CLASS_CompiledFromBlueprint));
 	}
+	else
+	{
+		if (LoadClass->HasAnyClassFlags(CLASS_Native))
+		{
+			return false;
+		}
 
-	ULinkerPlaceholderClass* AsPlaceholderClass = Cast<ULinkerPlaceholderClass>(LoadClass);
-	bool const bIsPlaceholderClass = (AsPlaceholderClass != nullptr);
+		ULinkerPlaceholderClass* AsPlaceholderClass = Cast<ULinkerPlaceholderClass>(LoadClass);
+		bool const bIsPlaceholderClass = (AsPlaceholderClass != nullptr);
 
-	FLinkerLoad* ClassLinker = LoadClass->GetLinker();
-	if ( !bIsPlaceholderClass 
-		&& ((ClassLinker == nullptr) || !ClassLinker->IsBlueprintFinalizationPending())
-		&& (!LoadClass->ClassDefaultObject || LoadClass->ClassDefaultObject->HasAnyFlags(RF_LoadCompleted) || !LoadClass->ClassDefaultObject->HasAnyFlags(RF_WasLoaded)) )
-	{
-		return false;
-	}
+		FLinkerLoad* ClassLinker = LoadClass->GetLinker();
+		if (!bIsPlaceholderClass
+			&& ((ClassLinker == nullptr) || !ClassLinker->IsBlueprintFinalizationPending())
+			&& (!LoadClass->ClassDefaultObject || LoadClass->ClassDefaultObject->HasAnyFlags(RF_LoadCompleted) || !LoadClass->ClassDefaultObject->HasAnyFlags(RF_WasLoaded)))
+		{
+			return false;
+		}
 
-	bool const bIsLoadingExportClass = (LoadFlags & LOAD_DeferDependencyLoads) ||
-		IsBlueprintFinalizationPending();
-	// if we're not in the process of "loading/finalizing" this package's 
-	// Blueprint class, then we're either running this before the linker has got 
-	// to that class, or we're finished and in the midst of regenerating that 
-	// class... either way, we don't have to defer the export (as long as we 
-	// make sure the export's class is fully regenerated... presumably it is in 
-	// the midst of doing so somewhere up the callstack)
-	if (!bIsLoadingExportClass || (LoadFlags & LOAD_ResolvingDeferredExports) != 0 )
-	{
-		DEFERRED_DEPENDENCY_CHECK(!IsExportBeingResolved(Index));
-		FScopedResolvingExportTracker ReentranceGuard(this, Index);
+		bool const bIsLoadingExportClass = (LoadFlags & LOAD_DeferDependencyLoads) ||
+			IsBlueprintFinalizationPending();
+		// if we're not in the process of "loading/finalizing" this package's 
+		// Blueprint class, then we're either running this before the linker has got 
+		// to that class, or we're finished and in the midst of regenerating that 
+		// class... either way, we don't have to defer the export (as long as we 
+		// make sure the export's class is fully regenerated... presumably it is in 
+		// the midst of doing so somewhere up the callstack)
+		if (!bIsLoadingExportClass || (LoadFlags & LOAD_ResolvingDeferredExports) != 0)
+		{
+			DEFERRED_DEPENDENCY_CHECK(!IsExportBeingResolved(Index));
+			FScopedResolvingExportTracker ReentranceGuard(this, Index);
 
-		// we want to be very careful, since we haven't filled in the export yet,
-		// we could get stuck in a recursive loop here (force-finalizing the 
-		// class here ends us back 
-		ForceRegenerateClass(LoadClass);
-		return false;
+			// we want to be very careful, since we haven't filled in the export yet,
+			// we could get stuck in a recursive loop here (force-finalizing the 
+			// class here ends us back 
+			ForceRegenerateClass(LoadClass);
+			return false;
+		}
 	}
 	
 	UPackage* PlaceholderOuter = LinkerRoot;
-	UClass*   PlaceholderType  = ULinkerPlaceholderExportObject::StaticClass();
+	UClass*   PlaceholderType  = ULinkerPlaceholderExportObject::StaticClass();	
 
 	FString ClassName = LoadClass->GetName();
 	//ClassName.RemoveFromEnd("_C");	
@@ -1236,8 +1267,9 @@ bool FLinkerLoad::DeferExportCreation(const int32 Index, UObject* Outer)
 
 	ULinkerPlaceholderExportObject* Placeholder = NewObject<ULinkerPlaceholderExportObject>(PlaceholderOuter, PlaceholderType, PlaceholderName, RF_Public | RF_Transient);
 	Placeholder->PackageIndex = FPackageIndex::FromExport(Index);
+	Placeholder->ExportName = Export.ObjectName;
 	Placeholder->SetLinker(this, Index, false);
-	FResolvingExportTracker::Get().AddLinkerPlaceholderObject(LoadClass, Placeholder);
+	FResolvingExportTracker::Get().AddLinkerPlaceholderObject(BlueprintClass, Placeholder);
 
 	Export.Object = Placeholder;
 #endif // USE_CIRCULAR_DEPENDENCY_LOAD_DEFERRING
@@ -2116,6 +2148,13 @@ void FLinkerLoad::ResolvePlaceholder(ULinkerPlaceholderExportObject* Placeholder
 {
 	int32 ExportIndex = Placeholder->PackageIndex.ToExport();
 
+	// If not already flagged, signal that this export is being resolved.
+	TUniquePtr<FScopedResolvingExportTracker> ScopedResolvingExportTracker;
+	if(!FResolvingExportTracker::Get().IsLinkerExportBeingResolved(this, ExportIndex))
+	{
+		ScopedResolvingExportTracker = MakeUnique<FScopedResolvingExportTracker>(this, ExportIndex);
+	}
+
 	Placeholder->SetLinker(nullptr, INDEX_NONE);
 
 	FObjectExport& Export = ExportMap[ExportIndex];
diff --git a/Engine/Source/Runtime/CoreUObject/Private/UObject/Class.cpp b/Engine/Source/Runtime/CoreUObject/Private/UObject/Class.cpp
--- a/Engine/Source/Runtime/CoreUObject/Private/UObject/Class.cpp
+++ b/Engine/Source/Runtime/CoreUObject/Private/UObject/Class.cpp
@@ -6190,7 +6190,7 @@ void UDynamicClass::PurgeClass(bool bRecompilingOnLoad)
 	AnimClassImplementation = nullptr;
 }
 
-UObject* UDynamicClass::FindArchetype(const UClass* ArchetypeClass, const FName ArchetypeName) const
+UObject* UDynamicClass::FindArchetype(const UClass* ArchetypeClass, const FName ArchetypeName, EObjectFlags InstanceFlags) const
 {
 	UObject* Archetype = static_cast<UObject*>(FindObjectWithOuter(this, ArchetypeClass, ArchetypeName));
 	if (!Archetype)
@@ -6208,7 +6208,7 @@ UObject* UDynamicClass::FindArchetype(const UClass* ArchetypeClass, const FName
 	}
 	const UClass* SuperClass = GetSuperClass();
 	return Archetype ? Archetype :
-		(SuperClass ? SuperClass->FindArchetype(ArchetypeClass, ArchetypeName) : nullptr);
+		(SuperClass ? SuperClass->FindArchetype(ArchetypeClass, ArchetypeName, InstanceFlags) : nullptr);
 }
 
 void UDynamicClass::SetupObjectInitializer(FObjectInitializer& ObjectInitializer) const
diff --git a/Engine/Source/Runtime/CoreUObject/Private/UObject/LinkerLoad.cpp b/Engine/Source/Runtime/CoreUObject/Private/UObject/LinkerLoad.cpp
--- a/Engine/Source/Runtime/CoreUObject/Private/UObject/LinkerLoad.cpp
+++ b/Engine/Source/Runtime/CoreUObject/Private/UObject/LinkerLoad.cpp
@@ -4552,7 +4552,10 @@ UObject* FLinkerLoad::CreateExport( int32 Index )
 			ThisParent = LinkerRoot;
 		}
 
-		if ( !LoadClass->HasAnyClassFlags(CLASS_Intrinsic) || Cast<ULinkerPlaceholderExportObject>(ThisParent))
+		// Always route ICH override template exports through the object creation deferral logic, regardless of type.
+		const bool bIsInheritableComponentTemplate = !!(Export.ObjectFlags & RF_InheritableComponentTemplate);
+
+		if (bIsInheritableComponentTemplate || !LoadClass->HasAnyClassFlags(CLASS_Intrinsic) || Cast<ULinkerPlaceholderExportObject>(ThisParent))
 		{
 #if USE_CIRCULAR_DEPENDENCY_LOAD_DEFERRING
 			if (LoadClass->HasAnyFlags(RF_NeedLoad))
diff --git a/Engine/Source/Runtime/CoreUObject/Private/UObject/LinkerPlaceholderExportObject.h b/Engine/Source/Runtime/CoreUObject/Private/UObject/LinkerPlaceholderExportObject.h
--- a/Engine/Source/Runtime/CoreUObject/Private/UObject/LinkerPlaceholderExportObject.h
+++ b/Engine/Source/Runtime/CoreUObject/Private/UObject/LinkerPlaceholderExportObject.h
@@ -27,4 +27,7 @@ public:
 	// FLinkerPlaceholderBase interface
 	virtual UObject* GetPlaceholderAsUObject() override { return (UObject*)(this); }
 	// End of FLinkerPlaceholderBase interface
+
+	/** Keeps track of the original export object resource name that this instance is being used in place of. */
+	FName ExportName;
 }; 
diff --git a/Engine/Source/Runtime/CoreUObject/Private/UObject/UObjectArchetype.cpp b/Engine/Source/Runtime/CoreUObject/Private/UObject/UObjectArchetype.cpp
--- a/Engine/Source/Runtime/CoreUObject/Private/UObject/UObjectArchetype.cpp
+++ b/Engine/Source/Runtime/CoreUObject/Private/UObject/UObjectArchetype.cpp
@@ -54,8 +54,8 @@ static FUObjectAnnotationDense<FArchetypeInfo, true> ArchetypeAnnotation;
 
 UObject* GetArchetypeFromRequiredInfoImpl(const UClass* Class, const UObject* Outer, FName Name, EObjectFlags ObjectFlags, bool bUseUpToDateClass)
 {
-	UObject* Result = NULL;
-	const bool bIsCDO = !!(ObjectFlags&RF_ClassDefaultObject);
+	UObject* Result = nullptr;
+	const bool bIsCDO = !!(ObjectFlags & RF_ClassDefaultObject);
 	if (bIsCDO)
 	{
 		Result = bUseUpToDateClass ? Class->GetAuthoritativeClass()->GetArchetypeForCDO() : Class->GetArchetypeForCDO();
@@ -68,44 +68,35 @@ UObject* GetArchetypeFromRequiredInfoImpl(const UClass* Class, const UObject* Ou
 			// Get a lock on the UObject hash tables for the duration of the GetArchetype operation
 			FScopedUObjectHashTablesLock HashTablesLock;
 
-			UObject* ArchetypeToSearch = nullptr;
-#if UE_CACHE_ARCHETYPE
-			ArchetypeToSearch = Outer->GetArchetype();
-#if UE_VERIFY_CACHED_ARCHETYPE
+			const UObject* ArchetypeToSearch;
+			UObject* FoundArchetype = nullptr;
+			if (!!(ObjectFlags & RF_InheritableComponentTemplate) && Outer->IsA<UClass>())
 			{
-				UObject* VerifyArchetype = GetArchetypeFromRequiredInfoImpl(Outer->GetClass(), Outer->GetOuter(), Outer->GetFName(), Outer->GetFlags(), bUseUpToDateClass);
-				checkf(ArchetypeToSearch == VerifyArchetype, TEXT("Cached archetype mismatch, expected: %s, cached: %s"), *GetFullNameSafe(VerifyArchetype), *GetFullNameSafe(ArchetypeToSearch));
+				// In this case, we are looking for the archetype of a non-native component template, which is owned by an outer UClass object. This
+				// differs from the "normal" path in that we don't search in the outer class object's archetype (which is a UClass derivative). Instead,
+				// we set the search object to the outer class's own default object (CDO). This way we'll defer the search to the virtual class API below.
+				const UClass* OuterClass = static_cast<const UClass*>(Outer);
+				ArchetypeToSearch = bUseUpToDateClass ? OuterClass->GetAuthoritativeClass()->GetDefaultObject() : OuterClass->GetDefaultObject();
 			}
+			else
+			{
+#if UE_CACHE_ARCHETYPE
+				ArchetypeToSearch = Outer->GetArchetype();
+#if UE_VERIFY_CACHED_ARCHETYPE
+				{
+					UObject* VerifyArchetype = GetArchetypeFromRequiredInfoImpl(Outer->GetClass(), Outer->GetOuter(), Outer->GetFName(), Outer->GetFlags(), bUseUpToDateClass);
+					checkf(ArchetypeToSearch == VerifyArchetype, TEXT("Cached archetype mismatch, expected: %s, cached: %s"), *GetFullNameSafe(VerifyArchetype), *GetFullNameSafe(ArchetypeToSearch));
+				}
 #endif // UE_VERIFY_CACHED_ARCHETYPE
 #else
-			ArchetypeToSearch = GetArchetypeFromRequiredInfoImpl(Outer->GetClass(), Outer->GetOuter(), Outer->GetFName(), Outer->GetFlags(), bUseUpToDateClass);
+				ArchetypeToSearch = GetArchetypeFromRequiredInfoImpl(Outer->GetClass(), Outer->GetOuter(), Outer->GetFName(), Outer->GetFlags(), bUseUpToDateClass);
 #endif // UE_CACHE_ARCHETYPE
-			UObject* MyArchetype = static_cast<UObject*>(FindObjectWithOuter(ArchetypeToSearch, Class, Name));
-			if (MyArchetype)
-			{
-				Result = MyArchetype; // found that my outers archetype had a matching component, that must be my archetype
+				FoundArchetype = static_cast<UObject*>(FindObjectWithOuter(ArchetypeToSearch, Class, Name));
 			}
-			else if (!!(ObjectFlags&RF_InheritableComponentTemplate) && Outer->IsA<UClass>())
+
+			if (FoundArchetype)
 			{
-				const UClass* OuterSuperClass = static_cast<const UClass*>(Outer)->GetSuperClass();
-				for (const UClass* SuperClassArchetype = bUseUpToDateClass && OuterSuperClass ? OuterSuperClass->GetAuthoritativeClass() : OuterSuperClass;
-					SuperClassArchetype && SuperClassArchetype->HasAllClassFlags(CLASS_CompiledFromBlueprint);
-					SuperClassArchetype = SuperClassArchetype->GetSuperClass())
-				{
-					if (GEventDrivenLoaderEnabled && EVENT_DRIVEN_ASYNC_LOAD_ACTIVE_AT_RUNTIME)
-					{
-						if (SuperClassArchetype->HasAnyFlags(RF_NeedLoad))
-						{
-							UE_LOG(LogClass, Fatal, TEXT("%s had RF_NeedLoad when searching supers for an archetype of %s in %s"), *GetFullNameSafe(ArchetypeToSearch), *GetFullNameSafe(Class), *GetFullNameSafe(Outer));
-						}
-					}
-					Result = static_cast<UObject*>(FindObjectWithOuter(SuperClassArchetype, Class, Name));
-					// We can have invalid archetypes halfway through the hierarchy, keep looking if it's pending kill or transient
-					if (Result && !Result->IsPendingKill() && !Result->HasAnyFlags(RF_Transient))
-					{
-						break;
-					}
-				}
+				Result = FoundArchetype; // found that my outers archetype had a matching component, that must be my archetype
 			}
 			else
 			{
@@ -117,7 +108,7 @@ UObject* GetArchetypeFromRequiredInfoImpl(const UClass* Class, const UObject* Ou
 					}
 				}
 
-				Result = ArchetypeToSearch->GetClass()->FindArchetype(Class, Name);
+				Result = ArchetypeToSearch->GetClass()->FindArchetype(Class, Name, ObjectFlags);
 			}
 		}
 
diff --git a/Engine/Source/Runtime/CoreUObject/Public/Blueprint/BlueprintSupport.h b/Engine/Source/Runtime/CoreUObject/Public/Blueprint/BlueprintSupport.h
--- a/Engine/Source/Runtime/CoreUObject/Public/Blueprint/BlueprintSupport.h
+++ b/Engine/Source/Runtime/CoreUObject/Public/Blueprint/BlueprintSupport.h
@@ -95,7 +95,10 @@ struct FBlueprintSupport
 	COREUOBJECT_API static void SetClassReparentingFPtr(FClassReparentingFPtr Ptr);
 
 	/** Tells if the specified object is one of the many flavors of FLinkerPlaceholderBase that we have. */
-	COREUOBJECT_API static bool IsDeferredDependencyPlaceholder(UObject* LoadedObj);
+	COREUOBJECT_API static bool IsDeferredDependencyPlaceholder(const UObject* LoadedObj);
+
+	/** Returns the original export name if its creation was deferred on load by the linker. Returns NAME_None if the given object is not an export placeholder object. */
+	COREUOBJECT_API static FName GetDeferredExportNameForPlaceholderObject(const UObject* LoadedObj);
 
 	/** Registers any object properties in this struct with the deferred dependency system */
 	COREUOBJECT_API static void RegisterDeferredDependenciesInStruct(const UStruct* Struct, void* StructData);
diff --git a/Engine/Source/Runtime/CoreUObject/Public/UObject/Class.h b/Engine/Source/Runtime/CoreUObject/Public/UObject/Class.h
--- a/Engine/Source/Runtime/CoreUObject/Public/UObject/Class.h
+++ b/Engine/Source/Runtime/CoreUObject/Public/UObject/Class.h
@@ -3161,7 +3161,7 @@ public:
 	virtual bool HasProperty(FProperty* InProperty) const;
 
 	/** Finds the object that is used as the parent object when serializing properties, overridden for blueprints */
-	virtual UObject* FindArchetype(const UClass* ArchetypeClass, const FName ArchetypeName) const { return nullptr; }
+	virtual UObject* FindArchetype(const UClass* ArchetypeClass, const FName ArchetypeName, EObjectFlags InstanceFlags = RF_NoFlags) const { return nullptr; }
 
 	/** Returns archetype object for CDO */
 	virtual UObject* GetArchetypeForCDO() const;
@@ -3270,7 +3270,7 @@ public:
 	// UClass interface
 	virtual UObject* CreateDefaultObject();
 	virtual void PurgeClass(bool bRecompilingOnLoad) override;
-	virtual UObject* FindArchetype(const UClass* ArchetypeClass, const FName ArchetypeName) const override;
+	virtual UObject* FindArchetype(const UClass* ArchetypeClass, const FName ArchetypeName, EObjectFlags InstanceFlags) const override;
 	virtual void SetupObjectInitializer(FObjectInitializer& ObjectInitializer) const override;
 
 	/** Find a struct property, called from generated code */
diff --git a/Engine/Source/Runtime/Engine/Classes/Engine/BlueprintGeneratedClass.h b/Engine/Source/Runtime/Engine/Classes/Engine/BlueprintGeneratedClass.h
--- a/Engine/Source/Runtime/Engine/Classes/Engine/BlueprintGeneratedClass.h
+++ b/Engine/Source/Runtime/Engine/Classes/Engine/BlueprintGeneratedClass.h
@@ -765,7 +765,7 @@ public:
 	virtual void PurgeClass(bool bRecompilingOnLoad) override;
 	virtual void Bind() override;
 	virtual void GetDefaultObjectPreloadDependencies(TArray<UObject*>& OutDeps) override;
-	virtual UObject* FindArchetype(const UClass* ArchetypeClass, const FName ArchetypeName) const override;
+	virtual UObject* FindArchetype(const UClass* ArchetypeClass, const FName ArchetypeName, EObjectFlags InstanceFlags) const override;
 
 	virtual void InitPropertiesFromCustomList(uint8* DataPtr, const uint8* DefaultDataPtr) override;
 	virtual void SetupObjectInitializer(FObjectInitializer& ObjectInitializer) const override;
diff --git a/Engine/Source/Runtime/Engine/Classes/Engine/InheritableComponentHandler.h b/Engine/Source/Runtime/Engine/Classes/Engine/InheritableComponentHandler.h
--- a/Engine/Source/Runtime/Engine/Classes/Engine/InheritableComponentHandler.h
+++ b/Engine/Source/Runtime/Engine/Classes/Engine/InheritableComponentHandler.h
@@ -70,14 +70,19 @@ struct FComponentOverrideRecord
 {
 	GENERATED_USTRUCT_BODY()
 
+	// Keep this ordered first! We rely on it to have been serialized for archetype lookup before we can load the template below.
+	// Note: This ordering only applies to Blueprint assets saved after the ordering was changed and/or cooked assets. To deal with
+	// uncooked assets saved prior to the ordering change, the editor's linker will defer all ICH template object loads (i.e. any non-NULL
+	// object referenced by the 'ComponentTemplate' property below) until the Blueprint class regeneration phase, prior to compile-on-load.
+	// ==> @see FLinkerLoad::DeferExportCreation(), GetArchetypeFromRequiredInfoImpl() and UBlueprintGeneratedClass::FindArchetype().
 	UPROPERTY()
-	UClass* ComponentClass;
+	FComponentKey ComponentKey;
 
 	UPROPERTY()
-	UActorComponent* ComponentTemplate;
+	UClass* ComponentClass;
 
 	UPROPERTY()
-	FComponentKey ComponentKey;
+	UActorComponent* ComponentTemplate;
 
 	UPROPERTY()
 	FBlueprintCookedComponentInstancingData CookedComponentInstancingData;
@@ -93,11 +98,13 @@ class ENGINE_API UInheritableComponentHandler : public UObject
 {
 	GENERATED_BODY()
 
-private:
+public:
 
 	/* Template name prefix for SCS DefaultSceneRootNode overrides */
 	static const FString SCSDefaultSceneRootOverrideNamePrefix;
 
+private:
+
 #if WITH_EDITOR
 
 	bool IsRecordValid(const FComponentOverrideRecord& Record) const;
diff --git a/Engine/Source/Runtime/Engine/Private/BlueprintGeneratedClass.cpp b/Engine/Source/Runtime/Engine/Private/BlueprintGeneratedClass.cpp
--- a/Engine/Source/Runtime/Engine/Private/BlueprintGeneratedClass.cpp
+++ b/Engine/Source/Runtime/Engine/Private/BlueprintGeneratedClass.cpp
@@ -755,10 +755,13 @@ UInheritableComponentHandler* UBlueprintGeneratedClass::GetInheritableComponentH
 }
 
 
-UObject* UBlueprintGeneratedClass::FindArchetype(const UClass* ArchetypeClass, const FName ArchetypeName) const
+UObject* UBlueprintGeneratedClass::FindArchetype(const UClass* ArchetypeClass, const FName ArchetypeName, EObjectFlags InstanceFlags) const
 {
 	UObject* Archetype = nullptr;
 
+	// Determine if we're searching for the archetype of an inherited component template override.
+	const bool bIsFromInheritedComponentTemplate = !!(InstanceFlags & RF_InheritableComponentTemplate);
+
 	// There are some rogue LevelScriptActors that still have a SimpleConstructionScript
 	// and since preloading the SCS of a script in a world package is bad news, we need to filter them out
 	if (SimpleConstructionScript && !IsChildOf<ALevelScriptActor>())
@@ -797,19 +800,22 @@ UObject* UBlueprintGeneratedClass::FindArchetype(const UClass* ArchetypeClass, c
 		UBlueprintGeneratedClass* Class = const_cast<UBlueprintGeneratedClass*>(this);
 		while (Class)
 		{
-			USimpleConstructionScript* ClassSCS = Class->SimpleConstructionScript;
-			USCS_Node* SCSNode = nullptr;
-			if (ClassSCS)
+			const USCS_Node* SCSNode = nullptr;
+			if (!bIsFromInheritedComponentTemplate)
 			{
-				if (ClassSCS->HasAnyFlags(RF_NeedLoad))
+				USimpleConstructionScript* ClassSCS = Class->SimpleConstructionScript;
+				if (ClassSCS)
 				{
-					ClassSCS->PreloadChain();
-				}
+					if (ClassSCS->HasAnyFlags(RF_NeedLoad))
+					{
+						ClassSCS->PreloadChain();
+					}
 
-				// We keep the index name here rather than the base name, in order to avoid potential
-				// collisions between an SCS variable name and an existing AddComponent node template.
-				// This is because old AddComponent node templates were based on the class display name.
-				SCSNode = ClassSCS->FindSCSNode(ArchetypeName);
+					// We keep the index name here rather than the base name, in order to avoid potential
+					// collisions between an SCS variable name and an existing AddComponent node template.
+					// This is because old AddComponent node templates were based on the class display name.
+					SCSNode = ClassSCS->FindSCSNode(ArchetypeName);
+				}
 			}
 
 			if (SCSNode)
@@ -836,6 +842,7 @@ UObject* UBlueprintGeneratedClass::FindArchetype(const UClass* ArchetypeClass, c
 				// This would find either an SCS component template override (for which the archetype
 				// name will match the SCS variable name), or an old AddComponent node template override
 				// (for which the archetype name will match the override record's component template name).
+				// It would also find archetypes of ICH template objects (which are themselves archetypes).
 				FComponentKey ComponentKey = ICH->FindKey(ArchetypeName);
 				if (!ComponentKey.IsValid() && ArchetypeName != ArchetypeBaseName)
 				{
@@ -856,7 +863,73 @@ UObject* UBlueprintGeneratedClass::FindArchetype(const UClass* ArchetypeClass, c
 				// Avoid searching for an invalid key.
 				if (ComponentKey.IsValid())
 				{
-					Archetype = ICH->GetOverridenComponentTemplate(ComponentKey);
+					if (bIsFromInheritedComponentTemplate)
+					{
+						FName OriginalTemplateName = ArchetypeName;
+						bool bIsDefaultSceneRootNodeTemplateOverride = false;
+
+						if(const UActorComponent* OriginalTemplate = ComponentKey.GetOriginalTemplate())
+						{
+							// Since the original template may have been renamed, we treat it as authoritative during ICH template archetype lookup.
+							// Note: ICH override template names will be fixed up during ICH serialization after its super class is loaded/compiled.
+							OriginalTemplateName = OriginalTemplate->GetFName();
+
+							// Get the owner of the original template, which should always be a valid Blueprint class.
+							const UBlueprintGeneratedClass* OriginalTemplateOwnerClass = CastChecked<UBlueprintGeneratedClass>(OriginalTemplate->GetOuter());
+							if(OriginalTemplateOwnerClass->SimpleConstructionScript)
+							{
+								// Determine if this template overrides the original template's owner class's default scene root node.
+								if (const USCS_Node* OriginalTemplateOwnerClass_DefaultSceneRootNode = OriginalTemplateOwnerClass->SimpleConstructionScript->GetDefaultSceneRootNode())
+								{
+									if (OriginalTemplateOwnerClass_DefaultSceneRootNode->ComponentTemplate == OriginalTemplate)
+									{
+										bIsDefaultSceneRootNodeTemplateOverride = true;
+									}
+								}
+							}
+						}
+
+						// Look for the first valid archetype in the parent class hierarchy that matches the archetype name.
+						do
+						{
+							Class = Cast<UBlueprintGeneratedClass>(Class->GetSuperClass());
+							if (Class)
+							{
+								if (GEventDrivenLoaderEnabled && EVENT_DRIVEN_ASYNC_LOAD_ACTIVE_AT_RUNTIME)
+								{
+									if (Class->HasAnyFlags(RF_NeedLoad))
+									{
+										UE_LOG(LogClass, Fatal, TEXT("%s had RF_NeedLoad when searching for an archetype of %s named %s"), *GetFullNameSafe(Class), *GetFullNameSafe(ArchetypeClass), *ArchetypeName.ToString());
+									}
+								}
+
+								// Note: The template name differs from the SCS variable name; template override archetypes use the template name.
+								FName SuperClassArchetypeName = OriginalTemplateName;
+
+								// We need special handling for a default scene root override template - these are explicitly named to avoid collisions
+								// with every Actor-based Blueprint class's own default scene root node (which persist but is "hidden" in a child class).
+								if (bIsDefaultSceneRootNodeTemplateOverride && ComponentKey.GetComponentOwner() != Class)
+								{
+									// We're not searching in the original template owner's class, so use the "special" default scene root override template name.
+									SuperClassArchetypeName = FName(UInheritableComponentHandler::SCSDefaultSceneRootOverrideNamePrefix + SuperClassArchetypeName.ToString());
+								}
+
+								// Note: This may also traverse the class hierarchy, but it will either return a valid match or an invalid/unnecessary
+								// template at the current level, in which case we need to move up one level here and try again (per the condition below).
+								Archetype = Class->FindArchetype(ArchetypeClass, SuperClassArchetypeName, InstanceFlags & ~RF_InheritableComponentTemplate);
+							}
+							else
+							{
+								// No valid match was found. We've searched the full inheritance hierarchy at this point, so just return NULL.
+								return nullptr;
+							}
+						}
+						while (Archetype && (Archetype->IsPendingKill() || Archetype->HasAnyFlags(RF_Transient)));	// Note: Anything marked for GC or transient (temporary) should not match as an archetype.
+					}
+					else
+					{
+						Archetype = ICH->GetOverridenComponentTemplate(ComponentKey);
+					}
 
 					if (GEventDrivenLoaderEnabled && EVENT_DRIVEN_ASYNC_LOAD_ACTIVE_AT_RUNTIME)
 					{
@@ -866,6 +939,12 @@ UObject* UBlueprintGeneratedClass::FindArchetype(const UClass* ArchetypeClass, c
 						}
 					}
 				}
+				else if (bIsFromInheritedComponentTemplate)
+				{
+					// We don't have a valid key, which means the override is missing, invalid or unnecessary.
+					// In that case, we drop to the logic below to try and find a match within the super class.
+					Class = CastChecked<UBlueprintGeneratedClass>(Class->GetSuperClass());
+				}
 			}
 
 			if (Archetype == nullptr)
diff --git a/Engine/Source/Runtime/Engine/Private/InheritableComponentHandler.cpp b/Engine/Source/Runtime/Engine/Private/InheritableComponentHandler.cpp
--- a/Engine/Source/Runtime/Engine/Private/InheritableComponentHandler.cpp
+++ b/Engine/Source/Runtime/Engine/Private/InheritableComponentHandler.cpp
@@ -47,6 +47,7 @@ void UInheritableComponentHandler::PostLoad()
 				
 				// Fix up component template name on load, if it doesn't match the original template name. Otherwise, archetype lookups will fail for this template.
 				// For example, this can occur after a component variable rename in a parent BP class, but before a child BP class with an override template is loaded.
+				// Note: If the key maps to an SCS node, the node's variable GUID will be used for the lookup instead of the name below (that's only used for UCS keys).
 				if (UActorComponent* OriginalTemplate = Record.ComponentKey.GetOriginalTemplate(Record.ComponentTemplate->GetFName()))
 				{
 					FString ExpectedTemplateName = OriginalTemplate->GetName();
@@ -104,9 +105,11 @@ UActorComponent* UInheritableComponentHandler::CreateOverridenComponentTemplate(
 	FName NewComponentTemplateName = BestArchetype->GetFName();
 	if (USCS_Node* SCSNode = Key.FindSCSNode())
 	{
+		const USCS_Node* DefaultSceneRootNode = SCSNode->GetSCS()->GetDefaultSceneRootNode();
+
 		// If this template will override an inherited DefaultSceneRoot node from a parent class's SCS, adjust the template name so that we don't reallocate our owner class's SCS DefaultSceneRoot node template.
 		// Note: This is currently the only case where a child class can have both an SCS node template and an override template associated with the same variable name, that is not considered to be a collision.
-		if (SCSNode == SCSNode->GetSCS()->GetDefaultSceneRootNode())
+		if (SCSNode == DefaultSceneRootNode && BestArchetype == DefaultSceneRootNode->ComponentTemplate)
 		{
 			NewComponentTemplateName = FName(*(SCSDefaultSceneRootOverrideNamePrefix + BestArchetype->GetName()));
 		}
@@ -467,9 +470,36 @@ void UInheritableComponentHandler::PreloadAll()
 
 FComponentKey UInheritableComponentHandler::FindKey(const FName VariableName) const
 {
+	auto MatchesComponentTemplateLambda = [&VariableName](const UObject* InTemplate) -> bool
+	{
+		if (InTemplate)
+		{
+			// AddComponent templates and ICH archetype lookups will match on the actual template name.
+			FName TemplateName = InTemplate->GetFName();
+
+			// ICH templates are deferred on load in order to allow the outer ICH record to be fully serialized first, as ICH template archetype lookup
+			// relies on a valid override record, and we need to be able to find the archetype before we create the export object on load. In that case,
+			// the template object will be a "placeholder" that's created by the linker at load time, in which case we need to redirect the template name
+			// to the original name that's recorded in the package export map, because the actual placeholder object name will not match the name of the
+			// ICH and/or SCS template that it's based on. See FLinkerLoad::CreateExport() and FLinkerLoad::DeferExportCreation() for the relevant code.
+			if (!InTemplate->HasAnyFlags(RF_WasLoaded) && FBlueprintSupport::IsDeferredDependencyPlaceholder(InTemplate))
+			{
+				const FName DeferredExportName = FBlueprintSupport::GetDeferredExportNameForPlaceholderObject(InTemplate);
+				if (!DeferredExportName.IsNone())
+				{
+					TemplateName = DeferredExportName;
+				}
+			}
+			
+			return TemplateName == VariableName;
+		}
+
+		return false;
+	};
+
 	for (const FComponentOverrideRecord& Record : Records)
 	{
-		if (Record.ComponentKey.GetSCSVariableName() == VariableName || (Record.ComponentTemplate && Record.ComponentTemplate->GetFName() == VariableName))
+		if (Record.ComponentKey.GetSCSVariableName() == VariableName || MatchesComponentTemplateLambda(Record.ComponentTemplate))
 		{
 			return Record.ComponentKey;
 		}
