Fix Clang compiler warnings about deprecated member variable usage in an implicit move ctor.

#rnx

diff --git a/Engine/Source/Runtime/Engine/Public/SceneView.h b/Engine/Source/Runtime/Engine/Public/SceneView.h
--- a/Engine/Source/Runtime/Engine/Public/SceneView.h
+++ b/Engine/Source/Runtime/Engine/Public/SceneView.h
@@ -1178,9 +1178,6 @@ public:
 	/** True if we should draw translucent objects when rendering hit proxies */
 	bool bAllowTranslucentPrimitivesInHitProxy;
 
-	/** BitArray representing the visibility state of the various sprite categories in the editor for this view */
-	UE_DEPRECATED(5.0, "Please use GetSpriteCategoryVisibility() instead.")
-	TBitArray<> SpriteCategoryVisibility;
 	/** Selection color for the editor (used by post processing) */
 	FLinearColor SelectionOutlineColor;
 	/** Selection color for use in the editor with inactive primitives */
@@ -1428,6 +1425,13 @@ protected:
 #if WITH_EDITOR
 	FGetSpriteCategoryVisibility GetSpriteCategoryVisibilityDelegate;
 #endif
+
+private:
+#if WITH_EDITOR
+	/** BitArray representing the visibility state of the various sprite categories in the editor for this view */
+	// Deprecated - Use GetSpriteCategoryVisibility() instead.
+	TBitArray<> SpriteCategoryVisibility;
+#endif
 };
 
 //////////////////////////////////////////////////////////////////////////
