Strata now correctly handles tangent to world space normal/trangent transform.

diff --git a/Engine/Shaders/Private/BasePassPixelShader.usf b/Engine/Shaders/Private/BasePassPixelShader.usf
--- a/Engine/Shaders/Private/BasePassPixelShader.usf
+++ b/Engine/Shaders/Private/BasePassPixelShader.usf
@@ -1485,7 +1485,7 @@ void FPixelShaderInOut_MainPS(
 	GBuffer.Specular = 0;
 	GBuffer.Metallic = 0;
 	GBuffer.ShadingModelID = SHADINGMODELID_STRATA;
-	GBuffer.WorldNormal = MaterialParameters.WorldNormal;
+	GBuffer.WorldNormal = 0; // Normal will be overriden by the top layer normal
 
 	uint2 PixelPos = In.SvPosition.xy;
 	uint PixelStrataDataByteOffset = GetStrataPixelDataByteOffset(PixelPos, ResolvedView.BufferSizeAndInvSize.x, OpaqueBasePass.Strata.MaxBytesPerPixel);
@@ -1494,7 +1494,7 @@ void FPixelShaderInOut_MainPS(
 	PackStrataOut(
 		OpaqueBasePass.Strata.MaterialLobesBufferUAV, PixelStrataDataByteOffset, 
 		StrataPixelHeader, Strata,
-		VertexNormal, MaterialParameters.CameraVector, GBuffer);
+		MaterialParameters.CameraVector, GBuffer);
 
 	// Since we don't carry SSS on/off through the ShadingModel value (DefaultLit/SubSurfaceProfile), we update subsurface toggle, based on the MFP radius value
 	bSubsurfaceEnable = GBuffer.CustomData.a > 0;
diff --git a/Engine/Shaders/Private/HairStrands/HairStrandsVisibilityPS.usf b/Engine/Shaders/Private/HairStrands/HairStrandsVisibilityPS.usf
--- a/Engine/Shaders/Private/HairStrands/HairStrandsVisibilityPS.usf
+++ b/Engine/Shaders/Private/HairStrands/HairStrandsVisibilityPS.usf
@@ -134,7 +134,7 @@ void MainVisibility(
 		float3 TranslatedWorldPosition = SvPositionToResolvedTranslatedWorld(SvPosition);
 		CalcMaterialParametersEx(MaterialParameters, PixelMaterialInputs, SvPosition, ScreenPosition, bIsFrontFace, TranslatedWorldPosition, TranslatedWorldPosition);
 
-#if PROJECT_STRATA && MATERIAL_IS_STRATA && MATERIAL_SHADINGMODEL_HAIR
+#if PROJECT_STRATA && MATERIAL_IS_STRATA
 		FStrataBSDF HairBSDF = PixelMaterialInputs.FrontMaterial.Layers[0].BSDFs[0];
 
 		BaseColor		 = HAIR_BASECOLOR(HairBSDF);
@@ -144,7 +144,7 @@ void MainVisibility(
 		Backlit			 = HAIR_BACKLIT(HairBSDF);
 		float MaterialAO = GetMaterialAmbientOcclusion(PixelMaterialInputs);
 
-		Tangent = mul(normalize(GetMaterialNormalRaw(PixelMaterialInputs)), MaterialParameters.SharedNormals.Normals[BSDF_GETNORMALID(HairBSDF)]);
+		Tangent = MaterialParameters.SharedNormals.Normals[BSDF_GETNORMALID(HairBSDF)];	// Strata normals are all world space already
 #else
 		BaseColor		 = GetMaterialBaseColor(PixelMaterialInputs);
 		Roughness		 = GetMaterialRoughness(PixelMaterialInputs);
diff --git a/Engine/Shaders/Private/Strata/Strata.ush b/Engine/Shaders/Private/Strata/Strata.ush
--- a/Engine/Shaders/Private/Strata/Strata.ush
+++ b/Engine/Shaders/Private/Strata/Strata.ush
@@ -424,8 +424,8 @@ struct FSharedNormals
 {
 	uint NormalCount;
 	uint NormalTypes;
-	float3 Normals[STRATA_MAX_SHARED_NORMAL_REGISTERS];
-	float3 Tangents[STRATA_MAX_SHARED_NORMAL_REGISTERS];
+	float3 Normals[STRATA_MAX_SHARED_NORMAL_REGISTERS];	// once registered, normals are always world space
+	float3 Tangents[STRATA_MAX_SHARED_NORMAL_REGISTERS];// idem for tangents
 };
 
 FSharedNormals StrataInitialiseSharedNormals()
@@ -989,7 +989,6 @@ void PackStrataOut(
 	uint ByteOffset, 
 	FStrataPixelHeader StrataPixelHeader,
 	FStrataData Strata,
-	float3 VertexNormal,
 	float3 V,
 	inout FGBufferData GBuffer)
 {
diff --git a/Engine/Source/Runtime/Engine/Private/Materials/HLSLMaterialTranslator.cpp b/Engine/Source/Runtime/Engine/Private/Materials/HLSLMaterialTranslator.cpp
--- a/Engine/Source/Runtime/Engine/Private/Materials/HLSLMaterialTranslator.cpp
+++ b/Engine/Source/Runtime/Engine/Private/Materials/HLSLMaterialTranslator.cpp
@@ -6680,6 +6680,11 @@ int32 FHLSLMaterialTranslator::TransformPosition(EMaterialCommonBasis SourceCoor
 	return TransformBase(SourceCoordBasis, DestCoordBasis, A, 1);
 }
 
+int32 FHLSLMaterialTranslator::TransformNormalFromRequestedBasisToWorld(int32 NormalCodeChunk)
+{
+	return Material->IsTangentSpaceNormal() ? TransformBase(MCB_Tangent, MCB_World, NormalCodeChunk, 0) : NormalCodeChunk;
+}
+
 int32 FHLSLMaterialTranslator::DynamicParameter(FLinearColor& DefaultValue, uint32 ParameterIndex)
 {
 	if (ShaderFrequency != SF_Vertex && ShaderFrequency != SF_Pixel && ShaderFrequency != SF_Compute)
@@ -7544,7 +7549,7 @@ int32 FHLSLMaterialTranslator::StrataVolumeBSDF(int32 Albedo, int32 Extinction,
 		*GetParameterCode(Anisotropy),
 		*GetParameterCode(Thickness),
 		SharedNormalIndex,
-		*GetParameterCode(VertexNormal())
+		*GetParameterCode(Normal)
 	);
 }
 
diff --git a/Engine/Source/Runtime/Engine/Private/Materials/HLSLMaterialTranslator.h b/Engine/Source/Runtime/Engine/Private/Materials/HLSLMaterialTranslator.h
--- a/Engine/Source/Runtime/Engine/Private/Materials/HLSLMaterialTranslator.h
+++ b/Engine/Source/Runtime/Engine/Private/Materials/HLSLMaterialTranslator.h
@@ -713,6 +713,7 @@ protected:
 	
 	virtual int32 TransformVector(EMaterialCommonBasis SourceCoordBasis, EMaterialCommonBasis DestCoordBasis, int32 A) override;
 	virtual int32 TransformPosition(EMaterialCommonBasis SourceCoordBasis, EMaterialCommonBasis DestCoordBasis, int32 A) override;
+	virtual int32 TransformNormalFromRequestedBasisToWorld(int32 NormalCodeChunk) override;
 	virtual int32 DynamicParameter(FLinearColor& DefaultValue, uint32 ParameterIndex = 0) override;
 	virtual int32 LightmapUVs() override;
 	virtual int32 PrecomputedAOMask() override;
diff --git a/Engine/Source/Runtime/Engine/Private/Materials/MaterialExpressions.cpp b/Engine/Source/Runtime/Engine/Private/Materials/MaterialExpressions.cpp
--- a/Engine/Source/Runtime/Engine/Private/Materials/MaterialExpressions.cpp
+++ b/Engine/Source/Runtime/Engine/Private/Materials/MaterialExpressions.cpp
@@ -19196,6 +19196,50 @@ static int32 CompileWithDefaultFloat3(class FMaterialCompiler* Compiler, FExpres
 	}
 	return CodeChunk == INDEX_NONE ? Compiler->Constant3(X, Y, Z) : CodeChunk;
 }
+static int32 CompileWithDefaultNormalWS(class FMaterialCompiler* Compiler, FExpressionInput& Input, bool* bDefaultIsUsed = nullptr)
+{
+	if (Input.GetTracedInput().Expression != nullptr)
+	{
+		int32 NormalCodeChunk = Input.Compile(Compiler);
+
+		if (NormalCodeChunk == INDEX_NONE)
+		{
+			if (bDefaultIsUsed)
+			{
+				*bDefaultIsUsed = true;
+			}
+			return Compiler->VertexNormal();
+		}
+		return Compiler->TransformNormalFromRequestedBasisToWorld(NormalCodeChunk);
+	}
+	if (bDefaultIsUsed)
+	{
+		*bDefaultIsUsed = true;
+	}
+	return Compiler->VertexNormal();
+}
+static int32 CompileWithDefaultTangentWS(class FMaterialCompiler* Compiler, FExpressionInput& Input, bool* bDefaultIsUsed = nullptr)
+{
+	if (Input.GetTracedInput().Expression != nullptr)
+	{
+		int32 TangentCodeChunk = Input.Compile(Compiler);
+
+		if (TangentCodeChunk == INDEX_NONE)
+		{
+			if (bDefaultIsUsed)
+			{
+				*bDefaultIsUsed = true;
+			}
+			return Compiler->VertexTangent();
+		}
+		return Compiler->TransformNormalFromRequestedBasisToWorld(TangentCodeChunk);
+	}
+	if (bDefaultIsUsed)
+	{
+		*bDefaultIsUsed = true;
+	}
+	return Compiler->VertexTangent();
+}
 
 #endif // WITH_EDITOR
 
@@ -19216,7 +19260,7 @@ UMaterialExpressionStrataDiffuseBSDF::UMaterialExpressionStrataDiffuseBSDF(const
 #if WITH_EDITOR
 int32 UMaterialExpressionStrataDiffuseBSDF::Compile(class FMaterialCompiler* Compiler, int32 OutputIndex)
 {
-	int32 NormalCodeChunk = CompileWithDefaultCodeChunk(Compiler, Normal, Compiler->VertexNormal());
+	int32 NormalCodeChunk = CompileWithDefaultNormalWS(Compiler, Normal);
 	uint8 SharedNormalIndex = StrataCompilationInfoCreateSharedNormal(Compiler, NormalCodeChunk);
 
 	int32 SSSProfileCodeChunk = INDEX_NONE;
@@ -19350,8 +19394,8 @@ UMaterialExpressionStrataDielectricBSDF::UMaterialExpressionStrataDielectricBSDF
 #if WITH_EDITOR
 int32 UMaterialExpressionStrataDielectricBSDF::Compile(class FMaterialCompiler* Compiler, int32 OutputIndex)
 {
-	int32 NormalCodeChunk = CompileWithDefaultCodeChunk(Compiler, Normal, Compiler->VertexNormal());
-	int32 TangentCodeChunk = CompileWithDefaultCodeChunk(Compiler, Tangent, Compiler->VertexTangent());
+	int32 NormalCodeChunk = CompileWithDefaultNormalWS(Compiler, Normal);
+	int32 TangentCodeChunk = CompileWithDefaultTangentWS(Compiler, Tangent);
 	uint8 SharedNormalIndex = StrataCompilationInfoCreateSharedNormal(Compiler, NormalCodeChunk, TangentCodeChunk);
 
 	int32 RoughnessXCodeChunk = CompileWithDefaultFloat1(Compiler, RoughnessX, 0.0f);
@@ -19439,8 +19483,8 @@ UMaterialExpressionStrataConductorBSDF::UMaterialExpressionStrataConductorBSDF(c
 #if WITH_EDITOR
 int32 UMaterialExpressionStrataConductorBSDF::Compile(class FMaterialCompiler* Compiler, int32 OutputIndex)
 {
-	int32 NormalCodeChunk = CompileWithDefaultCodeChunk(Compiler, Normal, Compiler->VertexNormal());
-	int32 TangentCodeChunk = CompileWithDefaultCodeChunk(Compiler, Tangent, Compiler->VertexTangent());
+	int32 NormalCodeChunk = CompileWithDefaultNormalWS(Compiler, Normal);
+	int32 TangentCodeChunk = CompileWithDefaultTangentWS(Compiler, Tangent);
 	uint8 SharedNormalIndex = StrataCompilationInfoCreateSharedNormal(Compiler, NormalCodeChunk, TangentCodeChunk);
 
 	int32 RoughnessXCodeChunk = CompileWithDefaultFloat1(Compiler, RoughnessX, 0.0f);
@@ -19608,7 +19652,7 @@ UMaterialExpressionStrataSheenBSDF::UMaterialExpressionStrataSheenBSDF(const FOb
 #if WITH_EDITOR
 int32 UMaterialExpressionStrataSheenBSDF::Compile(class FMaterialCompiler* Compiler, int32 OutputIndex)
 {
-	int32 NormalCodeChunk = CompileWithDefaultCodeChunk(Compiler, Normal, Compiler->VertexNormal());
+	int32 NormalCodeChunk = CompileWithDefaultNormalWS(Compiler, Normal);
 	uint8 SharedNormalIndex = StrataCompilationInfoCreateSharedNormal(Compiler, NormalCodeChunk);
 
 	int32 OutputCodeChunk = Compiler->StrataSheenBSDF(
@@ -19821,7 +19865,7 @@ UMaterialExpressionStrataHairBSDF::UMaterialExpressionStrataHairBSDF(const FObje
 int32 UMaterialExpressionStrataHairBSDF::Compile(class FMaterialCompiler* Compiler, int32 OutputIndex)
 {
 	// For hair, the shared normal index in fact represent the tangent
-	int32 TangentCodeChunk = CompileWithDefaultCodeChunk(Compiler, Tangent, Compiler->VertexTangent());
+	int32 TangentCodeChunk = CompileWithDefaultTangentWS(Compiler, Tangent);
 	uint8 SharedNormalIndex = StrataCompilationInfoCreateSharedNormal(Compiler, TangentCodeChunk);
 
 	int32 OutputCodeChunk = Compiler->StrataHairBSDF(
diff --git a/Engine/Source/Runtime/Engine/Public/MaterialCompiler.h b/Engine/Source/Runtime/Engine/Public/MaterialCompiler.h
--- a/Engine/Source/Runtime/Engine/Public/MaterialCompiler.h
+++ b/Engine/Source/Runtime/Engine/Public/MaterialCompiler.h
@@ -315,6 +315,7 @@ public:
 	virtual int32 AppendVector(int32 A,int32 B) = 0;
 	virtual int32 TransformVector(EMaterialCommonBasis SourceCoordBasis, EMaterialCommonBasis DestCoordBasis, int32 A) = 0;
 	virtual int32 TransformPosition(EMaterialCommonBasis SourceCoordBasis, EMaterialCommonBasis DestCoordBasis, int32 A) = 0;
+	virtual int32 TransformNormalFromRequestedBasisToWorld(int32 NormalCodeChunk) = 0;
 
 	virtual int32 DynamicParameter(FLinearColor& DefaultValue, uint32 ParameterIndex = 0) = 0;
 	virtual int32 LightmapUVs() = 0;
@@ -621,6 +622,7 @@ public:
 	{
 		return Compiler->TransformPosition(SourceCoordBasis, DestCoordBasis, A);
 	}
+	virtual int32 TransformNormalFromRequestedBasisToWorld(int32 NormalCodeChunk) override { return Compiler->TransformNormalFromRequestedBasisToWorld(NormalCodeChunk); }
 
 	virtual int32 DynamicParameter(FLinearColor& DefaultValue, uint32 ParameterIndex = 0) override { return Compiler->DynamicParameter(DefaultValue, ParameterIndex); }
 	virtual int32 LightmapUVs() override { return Compiler->LightmapUVs(); }
