Fixed Strata node to be use the correct float1 type.

diff --git a/Engine/Source/Runtime/Engine/Private/Materials/MaterialExpressions.cpp b/Engine/Source/Runtime/Engine/Private/Materials/MaterialExpressions.cpp
--- a/Engine/Source/Runtime/Engine/Private/Materials/MaterialExpressions.cpp
+++ b/Engine/Source/Runtime/Engine/Private/Materials/MaterialExpressions.cpp
@@ -19436,10 +19436,10 @@ uint32 UMaterialExpressionStrataDielectricBSDF::GetInputType(int32 InputIndex)
 		return MCT_Float3;
 		break;
 	case 2:
-		return MCT_Float;
+		return MCT_Float1;
 		break;
 	case 3:
-		return MCT_Float;
+		return MCT_Float1;
 		break;
 	case 4:
 		return MCT_Float3;
@@ -19525,10 +19525,10 @@ uint32 UMaterialExpressionStrataConductorBSDF::GetInputType(int32 InputIndex)
 		return MCT_Float3;
 		break;
 	case 2:
-		return MCT_Float;
+		return MCT_Float1;
 		break;
 	case 3:
-		return MCT_Float;
+		return MCT_Float1;
 		break;
 	case 4:
 		return MCT_Float3;
@@ -19683,7 +19683,7 @@ uint32 UMaterialExpressionStrataSheenBSDF::GetInputType(int32 InputIndex)
 		return MCT_Float3;
 		break;
 	case 1:
-		return MCT_Float;
+		return MCT_Float1;
 		break;	
 	case 2:
 		return MCT_Float3;
@@ -19760,7 +19760,7 @@ uint32 UMaterialExpressionStrataVolumetricFogCloudBSDF::GetInputType(int32 Input
 		return MCT_Float3;
 		break;
 	case 3:
-		return MCT_Float;
+		return MCT_Float1;
 		break;
 	}
 
@@ -19901,16 +19901,16 @@ uint32 UMaterialExpressionStrataHairBSDF::GetInputType(int32 InputIndex)
 		return MCT_Float3;
 		break;
 	case 1:
-		return MCT_Float;
+		return MCT_Float1;
 		break;
 	case 2:
-		return MCT_Float;
+		return MCT_Float1;
 		break;
 	case 3:
-		return MCT_Float;
+		return MCT_Float1;
 		break;
 	case 4:
-		return MCT_Float;
+		return MCT_Float1;
 		break;
 	case 5:
 		return MCT_Float3;
@@ -20221,7 +20221,7 @@ uint32 UMaterialExpressionStrataMultiply::GetInputType(int32 InputIndex)
 	{
 		return MCT_Strata;
 	}
-	return MCT_Float;
+	return MCT_Float1;
 }
 
 bool UMaterialExpressionStrataMultiply::IsResultStrataMaterial(int32 OutputIndex)
@@ -20378,10 +20378,10 @@ uint32 UMaterialExpressionStrataAnisotropyToRoughness::GetInputType(int32 InputI
 	switch (InputIndex)
 	{
 	case 0:
-		return MCT_Float;
+		return MCT_Float1;
 		break;
 	case 1:
-		return MCT_Float;
+		return MCT_Float1;
 		break;
 	}
 
