Fixed a bug in how FVisualLogger.bForceUniqueLogNames got implemented that resutled the feature always being enabled regardless of the flag's value.

diff --git a/Engine/Source/Runtime/Engine/Private/VisualLogger/VisualLogger.cpp b/Engine/Source/Runtime/Engine/Private/VisualLogger/VisualLogger.cpp
--- a/Engine/Source/Runtime/Engine/Private/VisualLogger/VisualLogger.cpp
+++ b/Engine/Source/Runtime/Engine/Private/VisualLogger/VisualLogger.cpp
@@ -182,7 +182,10 @@ FVisualLogEntry* FVisualLogger::GetEntryToWrite(const UObject* Object, float Tim
 	{
 		// It's first and only one usage of LogOwner as regular object to get names. We assume once that LogOwner is correct here and only here.
 		CurrentEntry = &CurrentEntryPerObject.Add(LogOwner);
-		ObjectToNameMap.Add(LogOwner, FName(*FString::Printf(TEXT("%s [%d]"), *LogOwner->GetName(), LogOwner->GetUniqueID())));
+		const FName LogName(bForceUniqueLogNames  
+			? *FString::Printf(TEXT("%s [%d]"), *LogOwner->GetName(), LogOwner->GetUniqueID()) 
+			: *FString::Printf(TEXT("%s"), *LogOwner->GetName()));
+		ObjectToNameMap.Add(LogOwner, LogName);
 		ObjectToClassNameMap.Add(LogOwner, *(LogOwner->GetClass()->GetName()));
 		ObjectToPointerMap.Add(LogOwner, LogOwner);
 		ObjectToWorldMap.Add(LogOwner, World);
@@ -354,6 +357,8 @@ void FVisualLogger::NavigationDataDump(const UObject* Object, const FName& Categ
 
 FVisualLogger::FVisualLogger()
 {
+	bForceUniqueLogNames = true;
+
 	BlockAllCategories(false);
 	AddDevice(&FVisualLoggerBinaryFileDevice::Get());
 	SetIsRecording(GEngine ? !!GEngine->bEnableVisualLogRecordingOnStart : false);
