Use TChunkedArray instead of custom solution for WorldPartition actor descriptors list.

diff --git a/Engine/Source/Runtime/Engine/Private/WorldPartition/WorldPartition.cpp b/Engine/Source/Runtime/Engine/Private/WorldPartition/WorldPartition.cpp
--- a/Engine/Source/Runtime/Engine/Private/WorldPartition/WorldPartition.cpp
+++ b/Engine/Source/Runtime/Engine/Private/WorldPartition/WorldPartition.cpp
@@ -237,7 +237,7 @@ void UWorldPartition::Initialize(UWorld* InWorld, const FTransform& InTransform)
 
 		for (const FAssetData& Asset : Assets)
 		{
-			TUniquePtr<FWorldPartitionActorDesc>* NewActorDesc = ActorDescList.Emplace(GetActorDescriptor(Asset));
+			TUniquePtr<FWorldPartitionActorDesc>* NewActorDesc = new(ActorDescList) TUniquePtr<FWorldPartitionActorDesc>(GetActorDescriptor(Asset));
 			check(NewActorDesc->IsValid());
 
 			if (bIsInstanced)
@@ -792,7 +792,7 @@ void UWorldPartition::OnAssetAdded(const FAssetData& InAssetData)
 {
 	if (ShouldHandleAssetEvent(InAssetData))
 	{
-		TUniquePtr<FWorldPartitionActorDesc>* NewActorDesc = ActorDescList.Emplace(GetActorDescriptor(InAssetData));
+		TUniquePtr<FWorldPartitionActorDesc>* NewActorDesc = new(ActorDescList) TUniquePtr<FWorldPartitionActorDesc>(GetActorDescriptor(InAssetData));
 		check(NewActorDesc->IsValid());
 
 		check(!Actors.Contains((*NewActorDesc)->GetGuid()));
diff --git a/Engine/Source/Runtime/Engine/Public/WorldPartition/WorldPartition.h b/Engine/Source/Runtime/Engine/Public/WorldPartition/WorldPartition.h
--- a/Engine/Source/Runtime/Engine/Public/WorldPartition/WorldPartition.h
+++ b/Engine/Source/Runtime/Engine/Public/WorldPartition/WorldPartition.h
@@ -59,67 +59,6 @@ public:
 	virtual bool Delete(const FString& PackageName) const =0;
 	virtual bool Delete(UPackage* Package) const =0;
 };
-
-/**
-* Simple chunked array that creates chunks to avoid invalidating elements when adding.
-* This will also never release elements as handles/references can still point to them.
-**/
-class FChunkedActorDescArray
-{
-	typedef TUniquePtr<FWorldPartitionActorDesc> ElemType;
-
-	enum
-	{
-		NumElementsPerChunk = 1024,
-	};
-
-	inline ElemType* AllocInstance()
-	{
-		const uint32 ChunkIndex = NumElements / NumElementsPerChunk;
-		const uint32 ElemIndex = NumElements % NumElementsPerChunk;
-
-		if (!Chunks.IsValidIndex(ChunkIndex))
-		{
-			Chunks.Add((ElemType*)FMemory::Malloc(sizeof(ElemType) * NumElementsPerChunk));
-		}
-
-		NumElements++;
-		return Chunks[ChunkIndex] + ElemIndex;
-	}
-
-public:
-	inline FChunkedActorDescArray()
-		: NumElements(0)
-	{}
-
-	inline ~FChunkedActorDescArray()
-	{
-		int32 CurElement = 0;
-		for (ElemType* Chunk : Chunks)
-		{
-			const int32 NumElemToDestroy = FMath::Min((int32)NumElementsPerChunk, NumElements - CurElement);
-
-			for (int32 i=0; i<NumElemToDestroy; i++)
-			{
-				Chunk[i].~ElemType();
-			}
-
-			FMemory::Free(Chunk);
-
-			CurElement += NumElementsPerChunk;
-		}
-	}
-
-	template <typename... ArgsType>
-	inline ElemType* Emplace(ArgsType&&... Args)
-	{
-		return new(AllocInstance()) ElemType(Forward<ArgsType>(Args)...);
-	}
-
-private:
-	TArray<ElemType*> Chunks;
-	int32 NumElements;
-};
 #endif
 
 UCLASS(AutoExpandCategories=(WorldPartition))
@@ -226,7 +165,7 @@ public:
 	UWorldPartitionRuntimeHash* RuntimeHash;
 
 #if WITH_EDITOR
-	FChunkedActorDescArray ActorDescList;
+	TChunkedArray<TUniquePtr<FWorldPartitionActorDesc>> ActorDescList;
 	TMap<FGuid, TUniquePtr<FWorldPartitionActorDesc>*> Actors;
 	
 	DECLARE_EVENT_TwoParams(UWorldPartition, FWorldPartitionActorRegisteredEvent, AActor&, bool);
