Add GetTypeHash / operator== to Handle

diff --git a/Engine/Source/Runtime/Engine/Public/WorldPartition/WorldPartitionHandle.h b/Engine/Source/Runtime/Engine/Public/WorldPartition/WorldPartitionHandle.h
--- a/Engine/Source/Runtime/Engine/Public/WorldPartition/WorldPartitionHandle.h
+++ b/Engine/Source/Runtime/Engine/Public/WorldPartition/WorldPartitionHandle.h
@@ -95,6 +95,22 @@ public:
 	{
 		return IsValid() ? ActorDesc->Get() : nullptr;
 	}
+		
+	friend uint32 GetTypeHash(const TWorldPartitionHandleBase<Impl>& HandleBase)
+	{
+		return ::PointerHash(HandleBase.ActorDesc);
+	}
+
+	friend bool operator==(const TWorldPartitionHandleBase<Impl>& Lhs, const TWorldPartitionHandleBase<Impl>& Rhs)
+	{
+		return Lhs.ActorDesc == Rhs.ActorDesc;
+	}
+
+	friend bool operator!=(const TWorldPartitionHandleBase<Impl>& Lhs, const TWorldPartitionHandleBase<Impl>& Rhs)
+	{
+		return !(Lhs == Rhs);
+	}
+
 
 protected:
 	inline void IncRefCount()
@@ -125,7 +141,7 @@ struct ENGINE_API FWorldPartitionReferenceImpl
 typedef TWorldPartitionHandleBase<FWorldPartitionHandleImpl> FWorldPartitionHandle;
 typedef TWorldPartitionHandleBase<FWorldPartitionReferenceImpl> FWorldPartitionReference;
 
-struct FWorldPartitionHandleHelpers
+struct ENGINE_API FWorldPartitionHandleHelpers
 {
 	static FWorldPartitionReference ConvertHandleToReference(const FWorldPartitionHandle& Handle);
 	static FWorldPartitionHandle ConvertReferenceToHandle(const FWorldPartitionReference& Handle);
