Add OnRegister/OnUnregister methods to ActorDesc to allow sub-class ActorDesc to do some processing when they are added or removed from a WP world

diff --git a/Engine/Source/Runtime/Engine/Private/WorldPartition/WorldPartition.cpp b/Engine/Source/Runtime/Engine/Private/WorldPartition/WorldPartition.cpp
--- a/Engine/Source/Runtime/Engine/Private/WorldPartition/WorldPartition.cpp
+++ b/Engine/Source/Runtime/Engine/Private/WorldPartition/WorldPartition.cpp
@@ -799,6 +799,8 @@ void UWorldPartition::OnAssetAdded(const FAssetData& InAssetData)
 		Actors.Add((*NewActorDesc)->GetGuid(), NewActorDesc);
 
 		HashActorDesc(NewActorDesc->Get());
+
+		(*NewActorDesc)->OnRegister(this);
 	}
 }
 
@@ -813,6 +815,8 @@ void UWorldPartition::OnAssetRemoved(const FAssetData& InAssetData)
 
 		UnhashActorDesc(ExistingActorDesc->Get());
 
+		(*ExistingActorDesc)->OnUnregister(this);
+
 		Actors.Remove((*ExistingActorDesc)->GetGuid());
 		ExistingActorDesc->Release();
 	}
diff --git a/Engine/Source/Runtime/Engine/Public/WorldPartition/WorldPartitionActorDesc.h b/Engine/Source/Runtime/Engine/Public/WorldPartition/WorldPartitionActorDesc.h
--- a/Engine/Source/Runtime/Engine/Public/WorldPartition/WorldPartitionActorDesc.h
+++ b/Engine/Source/Runtime/Engine/Public/WorldPartition/WorldPartitionActorDesc.h
@@ -106,6 +106,8 @@ protected:
 	}
 
 	virtual void Serialize(FArchive& Ar);
+	virtual void OnRegister(UWorldPartition* WorldPartition) {}
+	virtual void OnUnregister(UWorldPartition* WorldPartition) {}
 
 	FGuid							Guid;
 	FName							Class;
