WorldPartition: Cooking an AlwaysLoaded streaming cell now transfers its actors to the world's PersistentLevel. This ensures always loaded actors to be loaded on time and fixes the bug of always loaded actors of a Partitioned Level instance not properly loaded.

diff --git a/Engine/Source/Runtime/Engine/Classes/Engine/LevelStreaming.h b/Engine/Source/Runtime/Engine/Classes/Engine/LevelStreaming.h
--- a/Engine/Source/Runtime/Engine/Classes/Engine/LevelStreaming.h
+++ b/Engine/Source/Runtime/Engine/Classes/Engine/LevelStreaming.h
@@ -233,12 +233,6 @@ public:
 	uint8 bDrawOnLevelStatusMap : 1;
 
 #if WITH_EDITORONLY_DATA
-	/** If True, IsValidStreamingLevel will return true. This is necessary when cooking/saving a World with StreamingLevels
-	 *  referencing generated packages (done during cook process) that doesn't exist yet on disk (i.e. the workspace
-	 *  domain differs from the target domain).
-	 */
-	uint8 bForceIsValidStreamingLevel : 1;
-
 	/** Deprecated level color used for visualization. */
 	UPROPERTY()
 	FColor DrawColor_DEPRECATED;
diff --git a/Engine/Source/Runtime/Engine/Private/LevelStreaming.cpp b/Engine/Source/Runtime/Engine/Private/LevelStreaming.cpp
--- a/Engine/Source/Runtime/Engine/Private/LevelStreaming.cpp
+++ b/Engine/Source/Runtime/Engine/Private/LevelStreaming.cpp
@@ -277,7 +277,6 @@ ULevelStreaming::ULevelStreaming(const FObjectInitializer& ObjectInitializer)
 {
 # if WITH_EDITORONLY_DATA
 	bShouldBeVisibleInEditor = true;
-	bForceIsValidStreamingLevel = false;
 #endif
 	LevelColor = FLinearColor::White;
 	LevelTransform = FTransform::Identity;
@@ -1739,16 +1738,11 @@ ALevelScriptActor* ULevelStreaming::GetLevelScriptActor()
 
 bool ULevelStreaming::IsValidStreamingLevel() const
 {
-#if WITH_EDITOR
-	if (!bForceIsValidStreamingLevel)
-#endif
+	const bool PIESession = GetWorld()->WorldType == EWorldType::PIE || GetOutermost()->HasAnyPackageFlags(PKG_PlayInEditor);
+	if (!PIESession && !WorldAsset.IsNull())
 	{
-		const bool PIESession = GetWorld()->WorldType == EWorldType::PIE || GetOutermost()->HasAnyPackageFlags(PKG_PlayInEditor);
-		if (!PIESession && !WorldAsset.IsNull())
-		{
-			const FString WorldPackageName = GetWorldAssetPackageName();
-			return FPackageName::DoesPackageExist(WorldPackageName);
-		}
+		const FString WorldPackageName = GetWorldAssetPackageName();
+		return FPackageName::DoesPackageExist(WorldPackageName);
 	}
 	return true;
 }
diff --git a/Engine/Source/Runtime/Engine/Private/WorldPartition/WorldPartitionLevelHelper.cpp b/Engine/Source/Runtime/Engine/Private/WorldPartition/WorldPartitionLevelHelper.cpp
--- a/Engine/Source/Runtime/Engine/Private/WorldPartition/WorldPartitionLevelHelper.cpp
+++ b/Engine/Source/Runtime/Engine/Private/WorldPartition/WorldPartitionLevelHelper.cpp
@@ -30,36 +30,33 @@ UWorld::InitializationValues FWorldPartitionLevelHelper::GetWorldInitializationV
 }
 
 /**
- * Creates and populates a Level used in World Partition
+ * Moves external actors into the given level
  */
-bool FWorldPartitionLevelHelper::CreateAndFillLevelForRuntimeCell(const UWorld* InWorld, const FString& InWorldAssetName, UPackage* InPackage, const TArray<FWorldPartitionRuntimeCellObjectMapping>& InChildPackages)
+void FWorldPartitionLevelHelper::MoveExternalActorsToLevel(const TArray<FWorldPartitionRuntimeCellObjectMapping>& InChildPackages, ULevel* InLevel)
 {
-	check(IsRunningCommandlet());
-	check(InWorld && !InWorld->IsGameWorld());
-	check(InPackage);
-
-	// Create streaming cell Level package
-	ULevel* NewLevel = FWorldPartitionLevelHelper::CreateEmptyLevelForRuntimeCell(InWorld, InWorldAssetName, InPackage);
-	UPackage* NewLevelPackage = NewLevel->GetPackage();
-	check(NewLevelPackage == InPackage);
-	UWorld* NewWorld = UWorld::FindWorldInPackage(NewLevelPackage);
-
+	check(InLevel);
+	UPackage* LevelPackage = InLevel->GetPackage();
+	
 	// Move all actors to Cell level
 	for (const FWorldPartitionRuntimeCellObjectMapping& PackageObjectMapping : InChildPackages)
 	{
 		AActor* Actor = FindObject<AActor>(nullptr, *PackageObjectMapping.Path.ToString());
 		if (ensure(Actor))
 		{
+			const bool bSameOuter = (InLevel == Actor->GetOuter());
 			Actor->SetPackageExternal(false, false);
-			Actor->Rename(nullptr, NewLevel);
-			check(Actor->GetPackage() == NewLevelPackage);
+			Actor->Rename(nullptr, InLevel);
+			check(Actor->GetPackage() == LevelPackage);
+			if (bSameOuter && !InLevel->Actors.Contains(Actor))
+			{
+				InLevel->AddLoadedActor(Actor);
+			}
 		}
 		else
 		{
 			UE_LOG(LogEngine, Warning, TEXT("Can't find actor %s."), *PackageObjectMapping.Path.ToString());
 		}
 	}
-	return true;
 }
 
 /**
diff --git a/Engine/Source/Runtime/Engine/Private/WorldPartition/WorldPartitionRuntimeLevelStreamingCell.cpp b/Engine/Source/Runtime/Engine/Private/WorldPartition/WorldPartitionRuntimeLevelStreamingCell.cpp
--- a/Engine/Source/Runtime/Engine/Private/WorldPartition/WorldPartitionRuntimeLevelStreamingCell.cpp
+++ b/Engine/Source/Runtime/Engine/Private/WorldPartition/WorldPartitionRuntimeLevelStreamingCell.cpp
@@ -67,7 +67,7 @@ ULevelStreaming* UWorldPartitionRuntimeLevelStreamingCell::CreateLevelStreaming(
 		UWorld* OwningWorld = WorldPartition->GetWorld();
 		
 		const FName LevelStreamingName = FName(*FString::Printf(TEXT("WorldPartitionLevelStreaming_%s"), *GetName()));
-		const UClass* LevelStreamingClass = (IsRunningCommandlet() && IsAlwaysLoaded() && (OuterWorld == OwningWorld)) ? ULevelStreamingAlwaysLoaded::StaticClass() : UWorldPartitionLevelStreamingDynamic::StaticClass();
+		const UClass* LevelStreamingClass = UWorldPartitionLevelStreamingDynamic::StaticClass();
 		
 		// When called by Commandlet (PopulateGeneratedPackageForCook), LevelStreaming's outer is set to Cell/WorldPartition's outer to prevent warnings when saving Cell Levels (Warning: Obj in another map). 
 		// At runtime, LevelStreaming's outer will be properly set to the main world (see UWorldPartitionRuntimeLevelStreamingCell::Activate).
@@ -96,42 +96,56 @@ ULevelStreaming* UWorldPartitionRuntimeLevelStreamingCell::CreateLevelStreaming(
 	return nullptr;
 }
 
+void UWorldPartitionRuntimeLevelStreamingCell::LoadActorsForCook()
+{
+	// Load cell Actors
+	UWorldPartition* WorldPartition = GetOuterUWorldPartition();
+	for (const FGuid& ActorGuid : ActorsToLoadForCook)
+	{
+		const FWorldPartitionActorDesc* ActorDesc = WorldPartition->GetActorDesc(ActorGuid);
+		AActor* Actor = ActorDesc->GetActor();
+		if (!Actor)
+		{
+			Actor = ActorDesc->Load();
+			ensure(Actor);
+		}
+	}
+}
+
+void UWorldPartitionRuntimeLevelStreamingCell::MoveAlwaysLoadedContentToPersistentLevel()
+{
+	check(IsAlwaysLoaded());
+	if (GetActorCount() > 0)
+	{
+		LoadActorsForCook();
+
+		UWorld* OuterWorld = GetOuterUWorldPartition()->GetTypedOuter<UWorld>();
+		FWorldPartitionLevelHelper::MoveExternalActorsToLevel(Packages, OuterWorld->PersistentLevel);
+	}
+}
+
 bool UWorldPartitionRuntimeLevelStreamingCell::PopulateGeneratedPackageForCook(UPackage* InPackage, const FString& InPackageCookName)
 {
+	check(!IsAlwaysLoaded());
 	if (!InPackage || InPackageCookName.IsEmpty())
 	{
 		return false;
 	}
 
-	if (ULevelStreaming* NewLevelStreaming = CreateLevelStreaming(InPackageCookName))
+	if (GetActorCount() > 0)
 	{
 		UWorldPartition* WorldPartition = GetOuterUWorldPartition();
 		UWorld* OuterWorld = WorldPartition->GetTypedOuter<UWorld>();
-
+		ULevelStreaming* NewLevelStreaming = CreateLevelStreaming(InPackageCookName);
+		check(NewLevelStreaming)
+		
 		// Load cell Actors
-		for (const FGuid& ActorGuid : ActorsToLoadForCook)
-		{
-			const FWorldPartitionActorDesc* ActorDesc = WorldPartition->GetActorDesc(ActorGuid);
-			AActor* Actor = ActorDesc->GetActor();
-			if (!Actor)
-			{
-				Actor = ActorDesc->Load();
-				ensure(Actor);
-			}
-		}
+		LoadActorsForCook();
 
 		LevelStreaming = Cast<UWorldPartitionLevelStreamingDynamic>(NewLevelStreaming);
-		if (!FWorldPartitionLevelHelper::CreateAndFillLevelForRuntimeCell(OuterWorld, NewLevelStreaming->GetWorldAsset().ToString(), InPackage, Packages))
-		{
-			return false;
-		}
-
-		// Always loaded streaming levels are added to world's streaming level to make sure that they get streamed in as soon as world is loaded.
-		if (NewLevelStreaming->IsA<ULevelStreamingAlwaysLoaded>())
-		{
-			OuterWorld->AddStreamingLevel(NewLevelStreaming);
-			NewLevelStreaming->bForceIsValidStreamingLevel = true;
-		}
+		ULevel* NewLevel = FWorldPartitionLevelHelper::CreateEmptyLevelForRuntimeCell(OuterWorld, NewLevelStreaming->GetWorldAsset().ToString(), InPackage);
+		check(NewLevel->GetPackage() == InPackage);
+		FWorldPartitionLevelHelper::MoveExternalActorsToLevel(Packages, NewLevel);
 	}
 	return true;
 }
diff --git a/Engine/Source/Runtime/Engine/Private/WorldPartition/WorldPartitionRuntimeSpatialHash.cpp b/Engine/Source/Runtime/Engine/Private/WorldPartition/WorldPartitionRuntimeSpatialHash.cpp
--- a/Engine/Source/Runtime/Engine/Private/WorldPartition/WorldPartitionRuntimeSpatialHash.cpp
+++ b/Engine/Source/Runtime/Engine/Private/WorldPartition/WorldPartitionRuntimeSpatialHash.cpp
@@ -1442,18 +1442,26 @@ bool UWorldPartitionRuntimeSpatialHash::CreateStreamingGrid(const FSpatialHashRu
 
 					if (StreamingCell->GetActorCount())
 					{
-						if (!OutPackagesToGenerate)
+						// Always loaded cell actors are transfered to World's Persistent Level
+						if (StreamingCell->IsAlwaysLoaded())
 						{
-							UE_LOG(LogWorldPartitionRuntimeSpatialHash, Error, TEXT("Error creating runtime streaming cells for cook, OutPackagesToGenerate is null."));
-							return false;
+							StreamingCell->MoveAlwaysLoadedContentToPersistentLevel();
 						}
+						else
+						{
+							if (!OutPackagesToGenerate)
+							{
+								UE_LOG(LogWorldPartitionRuntimeSpatialHash, Error, TEXT("Error creating runtime streaming cells for cook, OutPackagesToGenerate is null."));
+								return false;
+							}
 
-						const FString PackageRelativePath = StreamingCell->GetPackageNameToCreate();
-						check(!PackageRelativePath.IsEmpty());
-						OutPackagesToGenerate->Add(PackageRelativePath);
+							const FString PackageRelativePath = StreamingCell->GetPackageNameToCreate();
+							check(!PackageRelativePath.IsEmpty());
+							OutPackagesToGenerate->Add(PackageRelativePath);
 
-						// Map relative package to StreamingCell for PopulateGeneratedPackageForCook/FinalizeGeneratedPackageForCook
-						PackagesToGenerateForCook.Add(PackageRelativePath, StreamingCell);
+							// Map relative package to StreamingCell for PopulateGeneratedPackageForCook/FinalizeGeneratedPackageForCook
+							PackagesToGenerateForCook.Add(PackageRelativePath, StreamingCell);
+						}
 					}
 				}
 			}
diff --git a/Engine/Source/Runtime/Engine/Public/WorldPartition/WorldPartitionLevelHelper.h b/Engine/Source/Runtime/Engine/Public/WorldPartition/WorldPartitionLevelHelper.h
--- a/Engine/Source/Runtime/Engine/Public/WorldPartition/WorldPartitionLevelHelper.h
+++ b/Engine/Source/Runtime/Engine/Public/WorldPartition/WorldPartitionLevelHelper.h
@@ -17,8 +17,8 @@
 class FWorldPartitionLevelHelper
 {
 public:
-	static bool CreateAndFillLevelForRuntimeCell(const UWorld* InWorld, const FString& InWorldAssetName, UPackage* InPackage, const TArray<FWorldPartitionRuntimeCellObjectMapping>& InChildPackages);
 	static ULevel* CreateEmptyLevelForRuntimeCell(const UWorld* InWorld, const FString& InWorldAssetName, UPackage* DestPackage = nullptr);
+	static void MoveExternalActorsToLevel(const TArray<FWorldPartitionRuntimeCellObjectMapping>& InChildPackages, ULevel* InLevel);
 private:
 	static UWorld::InitializationValues GetWorldInitializationValues();
 };
diff --git a/Engine/Source/Runtime/Engine/Public/WorldPartition/WorldPartitionRuntimeCell.h b/Engine/Source/Runtime/Engine/Public/WorldPartition/WorldPartitionRuntimeCell.h
--- a/Engine/Source/Runtime/Engine/Public/WorldPartition/WorldPartitionRuntimeCell.h
+++ b/Engine/Source/Runtime/Engine/Public/WorldPartition/WorldPartitionRuntimeCell.h
@@ -65,6 +65,7 @@ class UWorldPartitionRuntimeCell : public UObject
 
 	// Cook methods
 	virtual bool PopulateGeneratedPackageForCook(UPackage* InPackage, const FString& InPackageCookName) PURE_VIRTUAL(UWorldPartitionRuntimeCell::PopulateGeneratedPackageForCook, return false;);
+	virtual void MoveAlwaysLoadedContentToPersistentLevel() PURE_VIRTUAL(UWorldPartitionRuntimeCell::MoveAlwaysLoadedContentToPersistentLevel);
 	virtual void FinalizeGeneratedPackageForCook() {}
 	virtual FString GetPackageNameToCreate() const PURE_VIRTUAL(UWorldPartitionRuntimeCell::GetPackageNameToCreate, return FString(""););
 #endif
diff --git a/Engine/Source/Runtime/Engine/Public/WorldPartition/WorldPartitionRuntimeLevelStreamingCell.h b/Engine/Source/Runtime/Engine/Public/WorldPartition/WorldPartitionRuntimeLevelStreamingCell.h
--- a/Engine/Source/Runtime/Engine/Public/WorldPartition/WorldPartitionRuntimeLevelStreamingCell.h
+++ b/Engine/Source/Runtime/Engine/Public/WorldPartition/WorldPartitionRuntimeLevelStreamingCell.h
@@ -38,6 +38,8 @@ private:
 	EStreamingStatus GetLevelStreamingStatus() const;
 
 #if WITH_EDITOR
+	void LoadActorsForCook();
+	void MoveAlwaysLoadedContentToPersistentLevel();
 	class ULevelStreaming* CreateLevelStreaming(const FString& InPackageName = FString()) const;
 #endif
 
