Added the possibility for some actors to force a hard reference from the level serialization to make sure they get loaded before all always loaded actors in World Partition.

diff --git a/Engine/Source/Runtime/Engine/Classes/GameFramework/Actor.h b/Engine/Source/Runtime/Engine/Classes/GameFramework/Actor.h
--- a/Engine/Source/Runtime/Engine/Classes/GameFramework/Actor.h
+++ b/Engine/Source/Runtime/Engine/Classes/GameFramework/Actor.h
@@ -1804,6 +1804,9 @@ public:
 	/** Does this actor supports external packaging? */
 	virtual bool SupportsExternalPackaging() const;
 
+	/** Force the level to keep a hard reference on that external actor. */
+	virtual bool ForceExternalActorLevelReference() const { return false; }
+
 	/** Internal struct used to store information about an actor's components during reconstruction */
 	struct FActorRootComponentReconstructionData
 	{
diff --git a/Engine/Source/Runtime/Engine/Private/Level.cpp b/Engine/Source/Runtime/Engine/Private/Level.cpp
--- a/Engine/Source/Runtime/Engine/Private/Level.cpp
+++ b/Engine/Source/Runtime/Engine/Private/Level.cpp
@@ -443,7 +443,7 @@ void ULevel::Serialize( FArchive& Ar )
 				// Don't filter out external actors if duplicating the world to get the actors properly duplicated.
 				if (!(Ar.GetPortFlags() & PPF_Duplicate))
 				{
-					return false;
+					return Actor->ForceExternalActorLevelReference();
 				}
 				// When the world is partitioned and duplicating for PIE, we don't want to duplicate any external actors as they will be streamed by the runtime grids.
 				else if (bIsPartitioned && (Ar.GetPortFlags() & PPF_DuplicateForPIE))
diff --git a/Engine/Source/Runtime/Engine/Private/LevelActor.cpp b/Engine/Source/Runtime/Engine/Private/LevelActor.cpp
--- a/Engine/Source/Runtime/Engine/Private/LevelActor.cpp
+++ b/Engine/Source/Runtime/Engine/Private/LevelActor.cpp
@@ -715,6 +715,11 @@ AActor* UWorld::SpawnActor( UClass* Class, FTransform const* UserTransformPtr, c
 	if (ExternalPackage)
 	{
 		ExternalPackage->MarkPackageDirty();
+
+		if (Actor->ForceExternalActorLevelReference())
+		{
+			LevelToSpawnIn->MarkPackageDirty();
+		}
 	}
 
 	if (Actor->IsPendingKill() && !SpawnParameters.bNoFail)
diff --git a/Engine/Source/Runtime/Engine/Public/WorldPartition/DataLayer/WorldDataLayers.h b/Engine/Source/Runtime/Engine/Public/WorldPartition/DataLayer/WorldDataLayers.h
--- a/Engine/Source/Runtime/Engine/Public/WorldPartition/DataLayer/WorldDataLayers.h
+++ b/Engine/Source/Runtime/Engine/Public/WorldPartition/DataLayer/WorldDataLayers.h
@@ -27,6 +27,10 @@ public:
 	bool RemoveDataLayer(UDataLayer* InDataLayer);
 	bool RemoveDataLayers(const TArray<UDataLayer*>& InDataLayers);
 	FName GenerateUniqueDataLayerLabel(const FName& InDataLayerLabel) const;
+
+	//~ Begin AActor Interface
+	virtual bool ForceExternalActorLevelReference() const override { return true; }
+	//~ End AActor Interface
 #endif
 	
 	bool ContainsDataLayer(const UDataLayer* InDataLayer) const;
