fix non-editor build

diff --git a/Engine/Source/Runtime/Engine/Private/LevelActor.cpp b/Engine/Source/Runtime/Engine/Private/LevelActor.cpp
--- a/Engine/Source/Runtime/Engine/Private/LevelActor.cpp
+++ b/Engine/Source/Runtime/Engine/Private/LevelActor.cpp
@@ -716,10 +716,12 @@ AActor* UWorld::SpawnActor( UClass* Class, FTransform const* UserTransformPtr, c
 	{
 		ExternalPackage->MarkPackageDirty();
 
+#if WITH_EDITOR
 		if (Actor->ForceExternalActorLevelReference())
 		{
 			LevelToSpawnIn->MarkPackageDirty();
 		}
+#endif
 	}
 
 	if (Actor->IsPendingKill() && !SpawnParameters.bNoFail)
