Fixed uninitialized object uproperty in FCompositeReroute struct

diff --git a/Engine/Source/Runtime/Engine/Classes/Materials/MaterialExpressionPinBase.h b/Engine/Source/Runtime/Engine/Classes/Materials/MaterialExpressionPinBase.h
--- a/Engine/Source/Runtime/Engine/Classes/Materials/MaterialExpressionPinBase.h
+++ b/Engine/Source/Runtime/Engine/Classes/Materials/MaterialExpressionPinBase.h
@@ -39,6 +39,18 @@ struct FCompositeReroute
 
 	UPROPERTY(meta = (NoResetToDefault))
 	UMaterialExpressionReroute* Expression;
+
+	FCompositeReroute()
+		: Name(NAME_None)
+		, Expression(nullptr)
+	{
+	}
+
+	FCompositeReroute(FName InName, UMaterialExpressionReroute* InExpression)
+		: Name(InName)
+		, Expression(InExpression)
+	{
+	}
 };
 
 UCLASS(MinimalAPI)
