- Move AudioEditorSubsystem to AudioWidgetSubsystem and respective Interfaces to Engine so they can be used at runtime
- Fix ImplementsInterface helper
- Various clean-up & reparenting of classes to make them more generic/flexible for future types

diff --git a/Engine/Source/Editor/AudioEditor/Classes/SoundSubmixGraph/SoundSubmixGraphNode.h b/Engine/Source/Editor/AudioEditor/Classes/SoundSubmixGraph/SoundSubmixGraphNode.h
--- a/Engine/Source/Editor/AudioEditor/Classes/SoundSubmixGraph/SoundSubmixGraphNode.h
+++ b/Engine/Source/Editor/AudioEditor/Classes/SoundSubmixGraph/SoundSubmixGraphNode.h
@@ -3,33 +3,24 @@
 #pragma once
 
 #include "CoreMinimal.h"
+
+#include "AudioDevice.h"
+#include "Audio/AudioWidgetSubsystem.h"
 #include "Blueprint/UserWidget.h"
 #include "UObject/ObjectMacros.h"
 #include "EdGraph/EdGraphNode.h"
 #include "SGraphNode.h"
 #include "Sound/SoundSubmix.h"
-#include "AudioDevice.h"
 #include "DSP/EnvelopeFollower.h"
 #include "DSP/MultithreadedPatching.h"
-#include "AudioEditorSubsystem.h"
 #include "SoundSubmixGraphNode.generated.h"
 
+
+// Forward Declarations
 class USoundSubmixBase;
 
-UINTERFACE(Blueprintable)
-class AUDIOEDITOR_API USubmixNodeWidgetInterface : public UInterface
-{
-	GENERATED_BODY()
-};
 
-class AUDIOEDITOR_API ISubmixNodeWidgetInterface
-{
-	GENERATED_BODY()
 
-public:
-	UFUNCTION(BlueprintImplementableEvent)
-	void OnSubmixNodeConstructed(USoundSubmixBase* SoundSubmix);
-};
 
 class SSubmixGraphNode : public SGraphNode
 {
@@ -43,9 +34,6 @@ public:
 
 	void Construct(const FArguments& InArgs, UEdGraphNode* InGraphNode);
 
-	virtual void UpdateGraphNode() override;
-	virtual void Tick(const FGeometry& AllottedGeometry, const double InCurrentTime, const float InDeltaTime) override;
-
 	TSharedRef<SWidget> CreateNodeContentArea() override;
 
 private:
diff --git a/Engine/Source/Editor/AudioEditor/Private/AssetTypeActions/AssetTypeActions_SoundEffectPreset.cpp b/Engine/Source/Editor/AudioEditor/Private/AssetTypeActions/AssetTypeActions_SoundEffectPreset.cpp
--- a/Engine/Source/Editor/AudioEditor/Private/AssetTypeActions/AssetTypeActions_SoundEffectPreset.cpp
+++ b/Engine/Source/Editor/AudioEditor/Private/AssetTypeActions/AssetTypeActions_SoundEffectPreset.cpp
@@ -1,7 +1,7 @@
 // Copyright Epic Games, Inc. All Rights Reserved.
 
 #include "AssetTypeActions_SoundEffectPreset.h"
-#include "AudioEditorSubsystem.h"
+#include "Audio/AudioWidgetSubsystem.h"
 #include "Developer/AssetTools/Public/IAssetTools.h"
 #include "Developer/AssetTools/Public/AssetToolsModule.h"
 #include "Editors/SoundEffectPresetEditor.h"
@@ -84,14 +84,35 @@ void FAssetTypeActions_SoundEffectPreset::OpenAssetEditor(const TArray<UObject*>
 	EToolkitMode::Type Mode = ToolkitHost.IsValid() ? EToolkitMode::WorldCentric : EToolkitMode::Standalone;
 
 	TArray<UObject*> Objects = InObjects;
-
-	if (UAudioEditorSubsystem* AudioEditorSubsytem = GEditor->GetEditorSubsystem<UAudioEditorSubsystem>())
+	if (UWorld* World = GEditor ? GEditor->GetEditorWorldContext().World() : nullptr)
 	{
-		for (UObject* Object : InObjects)
+		if (UAudioWidgetSubsystem* WidgetSubsystem = GEngine ? GEngine->GetEngineSubsystem<UAudioWidgetSubsystem>() : nullptr)
 		{
-			if (USoundEffectPreset* Preset = Cast<USoundEffectPreset>(Object))
+			for (UObject* Object : InObjects)
 			{
-				TArray<UUserWidget*> UserWidgets = AudioEditorSubsytem->CreateUserWidgets(USoundEffectPresetViewInterface::StaticClass(), Object->GetClass());
+				USoundEffectPreset* Preset = Cast<USoundEffectPreset>(Object);
+				if (!Preset)
+				{
+					continue;
+				}
+
+				auto FilterFunction = [InPresetClass = Object->GetClass()](UUserWidget* UserWidget)
+				{
+					TSubclassOf<USoundEffectPreset> PresetClass = ISoundEffectPresetWidgetInterface::Execute_GetClass(UserWidget);
+					while (PresetClass)
+					{
+						if (PresetClass == InPresetClass)
+						{
+							return true;
+						}
+
+						PresetClass = PresetClass->GetSuperClass();
+					}
+
+					return false;
+				};
+
+				TArray<UUserWidget*> UserWidgets = WidgetSubsystem->CreateUserWidgets(*World, USoundEffectPresetWidgetInterface::StaticClass(), FilterFunction);
 				if (!UserWidgets.IsEmpty())
 				{
 					TSharedRef<FSoundEffectPresetEditor> PresetEditor = MakeShared<FSoundEffectPresetEditor>();
diff --git a/Engine/Source/Editor/AudioEditor/Private/AudioEditorSubsystem.cpp b/Engine/Source/Editor/AudioEditor/Private/AudioEditorSubsystem.cpp
deleted file mode 100644
diff --git a/Engine/Source/Editor/AudioEditor/Private/Editors/SoundEffectPresetEditor.cpp b/Engine/Source/Editor/AudioEditor/Private/Editors/SoundEffectPresetEditor.cpp
--- a/Engine/Source/Editor/AudioEditor/Private/Editors/SoundEffectPresetEditor.cpp
+++ b/Engine/Source/Editor/AudioEditor/Private/Editors/SoundEffectPresetEditor.cpp
@@ -154,6 +154,17 @@ void FSoundEffectPresetEditor::Init(const EToolkitMode::Type Mode, const TShared
 		bUseSmallIcons);
 }
 
+bool FSoundEffectPresetEditor::CloseWindow()
+{
+	if (FAssetEditorToolkit::CloseWindow())
+	{
+		UserWidgets.Reset();
+		return true;
+	}
+
+	return false;
+}
+
 FName FSoundEffectPresetEditor::GetEditorName() const
 {
 	return "Preset Editor";
@@ -183,7 +194,7 @@ void FSoundEffectPresetEditor::InitPresetWidgets(const TArray<UUserWidget*>& InW
 			if (Widget)
 			{
 				UserWidgets.Add(TStrongObjectPtr<UUserWidget>(Widget));
-				IAudioWidgetInterface::Execute_OnConstructed(Widget, Cast<UObject>(SoundEffectPreset.Get()));
+				ISoundEffectPresetWidgetInterface::Execute_OnConstructed(Widget, SoundEffectPreset.Get());
 			}
 		}
 	}
@@ -213,7 +224,7 @@ void FSoundEffectPresetEditor::NotifyPostChange(const FPropertyChangedEvent& Pro
 		for (TStrongObjectPtr<UUserWidget>& UserWidget : UserWidgets)
 		{
 			const FName PropertyName = PropertyThatChanged->GetFName();
-			IAudioWidgetInterface::Execute_OnPropertyChanged(UserWidget.Get(), SoundEffectPreset.Get(), PropertyName);
+			ISoundEffectPresetWidgetInterface::Execute_OnPropertyChanged(UserWidget.Get(), SoundEffectPreset.Get(), PropertyName);
 		}
 	}
 }
@@ -230,7 +241,7 @@ void FSoundEffectPresetEditor::NotifyPostChange(const FPropertyChangedEvent& Pro
 				if (FProperty* Property = Node->GetValue())
 				{
 					const FName PropertyName = Property->GetFName();
-					IAudioWidgetInterface::Execute_OnPropertyChanged(UserWidget.Get(), SoundEffectPreset.Get(), PropertyName);
+					ISoundEffectPresetWidgetInterface::Execute_OnPropertyChanged(UserWidget.Get(), SoundEffectPreset.Get(), PropertyName);
 				}
 				Node = Node->GetNextNode();
 			}
@@ -252,7 +263,7 @@ TSharedRef<SDockTab> FSoundEffectPresetEditor::SpawnTab_Properties(const FSpawnT
 
 TSharedRef<SDockTab> FSoundEffectPresetEditor::SpawnTab_UserWidgetEditor(const FSpawnTabArgs& Args, int32 WidgetIndex)
 {
-	FName IconBrushName = IAudioWidgetInterface::Execute_GetIconBrushName(UserWidgets[WidgetIndex].Get());
+	FName IconBrushName = ISoundEffectPresetWidgetInterface::Execute_GetIconBrushName(UserWidgets[WidgetIndex].Get());
 	if (IconBrushName == FName())
 	{
 		IconBrushName = "GenericEditor.Tabs.Properties";
@@ -273,7 +284,7 @@ TSharedRef<SDockTab> FSoundEffectPresetEditor::SpawnTab_UserWidgetEditor(const F
 			];
 	}
 
-	const FText CustomLabel = IAudioWidgetInterface::Execute_GetEditorName(UserWidgets[WidgetIndex].Get());
+	const FText CustomLabel = ISoundEffectPresetWidgetInterface::Execute_GetEditorName(UserWidgets[WidgetIndex].Get());
 	if (!CustomLabel.IsEmpty())
 	{
 		Label = CustomLabel;
diff --git a/Engine/Source/Editor/AudioEditor/Private/Editors/SoundEffectPresetEditor.h b/Engine/Source/Editor/AudioEditor/Private/Editors/SoundEffectPresetEditor.h
--- a/Engine/Source/Editor/AudioEditor/Private/Editors/SoundEffectPresetEditor.h
+++ b/Engine/Source/Editor/AudioEditor/Private/Editors/SoundEffectPresetEditor.h
@@ -3,7 +3,7 @@
 
 #include "CoreMinimal.h"
 
-#include "AudioEditorSubsystem.h"
+#include "Audio/AudioWidgetSubsystem.h"
 #include "Blueprint/UserWidget.h"
 #include "EditorUndoClient.h"
 #include "Engine/DeveloperSettings.h"
@@ -18,8 +18,6 @@
 #include "UObject/ObjectMacros.h"
 #include "Widgets/SWidget.h"
 
-#include "SoundEffectPresetEditor.generated.h"
-
 
 // Forward Declarations
 class FCurveEditor;
@@ -31,17 +29,6 @@ class UCurveBase;
 class USoundEffectPreset;
 class UWidgetBlueprint;
 
-UINTERFACE(Blueprintable)
-class AUDIOEDITOR_API USoundEffectPresetViewInterface : public UAudioWidgetInterface
-{
-	GENERATED_BODY()
-};
-
-class AUDIOEDITOR_API ISoundEffectPresetViewInterface : public IAudioWidgetInterface
-{
-	GENERATED_BODY()
-};
-
 
 class FSoundEffectPresetEditor : public FAssetEditorToolkit, public FNotifyHook, public FEditorUndoClient
 {
@@ -52,6 +39,7 @@ public:
 	void Init(const EToolkitMode::Type Mode, const TSharedPtr<IToolkitHost>& InitToolkitHost, USoundEffectPreset* PresetToEdit, const TArray<UUserWidget*>& InWidgetBlueprints);
 
 	/** FAssetEditorToolkit interface */
+	virtual bool CloseWindow() override;
 	virtual FName GetEditorName() const override;
 	virtual FName GetToolkitFName() const override;
 	virtual FText GetBaseToolkitName() const override;
diff --git a/Engine/Source/Editor/AudioEditor/Private/SoundSubmixGraphNode.cpp b/Engine/Source/Editor/AudioEditor/Private/SoundSubmixGraphNode.cpp
--- a/Engine/Source/Editor/AudioEditor/Private/SoundSubmixGraphNode.cpp
+++ b/Engine/Source/Editor/AudioEditor/Private/SoundSubmixGraphNode.cpp
@@ -1,19 +1,22 @@
 // Copyright Epic Games, Inc. All Rights Reserved.
 
 #include "SoundSubmixGraph/SoundSubmixGraphNode.h"
+
+#include "AudioDeviceManager.h"
+#include "AudioEditorSettings.h"
+#include "Audio/AudioWidgetSubsystem.h"
 #include "Sound/SoundSubmix.h"
-#include "SoundSubmixGraph/SoundSubmixGraphSchema.h"
-#include "SoundSubmixGraph/SoundSubmixGraph.h"
-#include "SoundSubmixEditor.h"
 #include "SoundSubmixDefaultColorPalette.h"
+#include "SoundSubmixEditor.h"
+#include "SoundSubmixGraph/SoundSubmixGraph.h"
+#include "SoundSubmixGraph/SoundSubmixGraphSchema.h"
 #include "Subsystems/AssetEditorSubsystem.h"
-#include "AudioDeviceManager.h"
-#include "AudioEditorSettings.h"
 #include "Widgets/Input/SSlider.h"
-#include "AudioEditorSubsystem.h"
+
 
 #define LOCTEXT_NAMESPACE "SoundSubmixGraphNode"
 
+
 void SSubmixGraphNode::Construct(const FArguments& InArgs, UEdGraphNode* InGraphNode)
 {
 	SubmixBase = InArgs._SubmixBase;
@@ -21,21 +24,12 @@ void SSubmixGraphNode::Construct(const FArguments& InArgs, UEdGraphNode* InGraph
 
 	if (SubmixNodeUserWidget.IsValid())
 	{
-		ISubmixNodeWidgetInterface::Execute_OnSubmixNodeConstructed(SubmixNodeUserWidget.Get(), SubmixBase.Get());
+		ISoundSubmixWidgetInterface::Execute_OnConstructed(SubmixNodeUserWidget.Get(), SubmixBase.Get());
 	}
 	GraphNode = InGraphNode;
 	UpdateGraphNode();
 }
 
-void SSubmixGraphNode::UpdateGraphNode()
-{
-	SGraphNode::UpdateGraphNode();
-}
-
-void SSubmixGraphNode::Tick(const FGeometry& AllottedGeometry, const double InCurrentTime, const float InDeltaTime)
-{
-}
-
 TSharedRef<SWidget> SSubmixGraphNode::CreateNodeContentArea()
 {	
 	if (SubmixNodeUserWidget.IsValid())
@@ -219,25 +213,31 @@ FText USoundSubmixGraphNode::GetNodeTitle(ENodeTitleType::Type TitleType) const
 
 TSharedPtr<SGraphNode> USoundSubmixGraphNode::CreateVisualWidget()
 {
-	if (USoundSubmixBase* SubmixBase = Cast<USoundSubmixBase>(SoundSubmix))
+	if (UWorld* World = GEditor ? GEditor->GetEditorWorldContext().World() : nullptr)
 	{
-		if (UAudioEditorSubsystem* AudioEditorSubsytem = GEditor->GetEditorSubsystem<UAudioEditorSubsystem>())
+		if (USoundSubmixBase* SubmixBase = Cast<USoundSubmixBase>(SoundSubmix))
 		{
-			TArray<UUserWidget*> UserWidgets = AudioEditorSubsytem->CreateUserWidgets(USubmixNodeWidgetInterface::StaticClass());
-			if (!UserWidgets.IsEmpty())
+			if (UAudioWidgetSubsystem* AudioWidgetSubsystem = GEngine ? GEngine->GetEngineSubsystem<UAudioWidgetSubsystem>() : nullptr)
 			{
-				SubmixNodeUserWidget = UserWidgets[0];
+				TArray<UUserWidget*> UserWidgets = AudioWidgetSubsystem->CreateUserWidgets(*World, USoundSubmixWidgetInterface::StaticClass());
+				if (!UserWidgets.IsEmpty())
+				{
+					// For now, only supports single widget. Gallery system to be implemented to support
+					// showing multiple widgets and/or cycling node widgets.
+					SubmixNodeUserWidget = UserWidgets[0];
+				}
 			}
-		}
 
-		const UAudioSubmixEditorSettings* Settings = GetDefault<UAudioSubmixEditorSettings>();
-		check(Settings);
+			const UAudioSubmixEditorSettings* Settings = GetDefault<UAudioSubmixEditorSettings>();
+			check(Settings);
 
-		// Pass the owning submix and the user widgets to the graph node
-		return SNew(SSubmixGraphNode, this)
-			.SubmixBase(SubmixBase)
-			.SubmixNodeUserWidget(SubmixNodeUserWidget);
+			// Pass the owning submix and the user widgets to the graph node
+			return SNew(SSubmixGraphNode, this)
+				.SubmixBase(SubmixBase)
+				.SubmixNodeUserWidget(SubmixNodeUserWidget);
+		}
 	}
+
 	return nullptr;
 }
 
diff --git a/Engine/Source/Editor/AudioEditor/Public/AudioEditorSubsystem.h b/Engine/Source/Editor/AudioEditor/Public/AudioEditorSubsystem.h
deleted file mode 100644
diff --git a/Engine/Source/Runtime/Engine/Private/Audio/AudioWidgetSubsystem.cpp b/Engine/Source/Runtime/Engine/Private/Audio/AudioWidgetSubsystem.cpp
new file mode 100644
--- /dev/null
+++ b/Engine/Source/Runtime/Engine/Private/Audio/AudioWidgetSubsystem.cpp
@@ -0,0 +1,128 @@
+// Copyright Epic Games, Inc. All Rights Reserved.
+#include "Audio/AudioWidgetSubsystem.h"
+
+#include "AssetRegistry/AssetRegistryModule.h"
+#include "Audio/SoundEffectPresetWidgetInterface.h"
+#include "Blueprint/BlueprintSupport.h"
+#include "Blueprint/UserWidget.h"
+#include "Engine/World.h"
+#include "Interfaces/IPluginManager.h"
+#include "Internationalization/Text.h"
+#include "Misc/CoreDelegates.h"
+#include "Misc/FileHelper.h"
+#include "Modules/ModuleManager.h"
+#include "UObject/CoreRedirects.h"
+
+
+#define LOCTEXT_NAMESPACE "AudioWidgetSubsystem"
+
+
+TArray<UUserWidget*> UAudioWidgetSubsystem::CreateUserWidgets(UWorld& InWorld, TSubclassOf<UInterface> InWidgetInterface, TFunction<bool(UUserWidget*)> FilterFunction) const
+{
+	TArray<UUserWidget*> UserWidgets;
+	UClass* InterfaceClass = InWidgetInterface ? InWidgetInterface.Get() : UInterface::StaticClass();
+	if (!ensure(InterfaceClass))
+	{
+		return UserWidgets;
+	}
+
+	const TArray<FAssetData> AssetData = GetBlueprintAssetData(InterfaceClass);
+	for (const FAssetData& AssetEntry : AssetData)
+	{
+		UObject* Asset = AssetEntry.GetAsset();
+		if (!Asset)
+		{
+			continue;
+		}
+
+		UBlueprint* Blueprint = CastChecked<UBlueprint>(Asset);
+		UClass* GeneratedClass = Blueprint->GeneratedClass;
+		if (!GeneratedClass)
+		{
+			continue;
+		}
+
+		if (UUserWidget * UserWidget = CreateWidget<UUserWidget>(&InWorld, GeneratedClass))
+		{
+			if (!FilterFunction)
+			{
+				UserWidgets.Add(UserWidget);
+			}
+			else if (FilterFunction(UserWidget))
+			{
+				UserWidgets.Add(UserWidget);
+			}
+		}
+	}
+
+	return UserWidgets;
+}
+
+const TArray<FAssetData> UAudioWidgetSubsystem::GetBlueprintAssetData(UClass* InInterfaceClass)
+{
+	FAssetRegistryModule& AssetRegistryModule = FModuleManager::LoadModuleChecked<FAssetRegistryModule>("AssetRegistry");
+	TArray<FAssetData> AssetData;
+
+	FARFilter Filter;
+	Filter.ClassNames.Add(UBlueprint::StaticClass()->GetFName());
+	Filter.ClassNames.Add("WidgetBlueprint");
+	Filter.TagsAndValues.Add(FBlueprintTags::ImplementedInterfaces);
+
+	AssetRegistryModule.Get().EnumerateAssets(Filter, [FoundAssets = &AssetData, InInterfaceClass](const FAssetData& InAssetData)
+	{
+		if (ImplementsInterface(InAssetData, InInterfaceClass))
+		{
+			FoundAssets->Add(InAssetData);
+		}
+
+		return true;
+	});
+
+	return AssetData;
+}
+
+bool UAudioWidgetSubsystem::ImplementsInterface(const FAssetData& InAssetData, UClass* InInterfaceClass)
+{
+	const FString ImplementedInterfaces = InAssetData.GetTagValueRef<FString>(FBlueprintTags::ImplementedInterfaces);
+	if (!ImplementedInterfaces.IsEmpty())
+	{
+		FString CurrentString = *ImplementedInterfaces;
+
+		static const FString InterfaceIdentifier = TEXT("Interface=");
+		int32 FirstIndex = CurrentString.Find(InterfaceIdentifier);
+		if (FirstIndex > 0)
+		{
+			CurrentString.RightChopInline(FirstIndex);
+		}
+
+		TArray<FString> Tokens;
+		CurrentString.ParseIntoArray(Tokens, *InterfaceIdentifier);
+		for (FString& Token : Tokens)
+		{
+			FirstIndex = Token.Find(TEXT("'"));
+			if (FirstIndex > 0)
+			{
+				Token.RightChopInline(FirstIndex + 1);
+			}
+
+			Token.RemoveFromStart(TEXT("\""));
+
+			FirstIndex = Token.Find(TEXT("'"));
+			if (FirstIndex > 0)
+			{
+				Token.LeftChopInline(Token.Len() - FirstIndex);
+			}
+
+			Token.RemoveFromEnd(TEXT("\""));
+
+			const FCoreRedirectObjectName ResolvedInterfaceName = FCoreRedirects::GetRedirectedName(ECoreRedirectFlags::Type_Class, FCoreRedirectObjectName(Token));
+			if (InInterfaceClass->GetFName() == ResolvedInterfaceName.ObjectName)
+			{
+				return true;
+			}
+		}
+	}
+
+	return false;
+}
+#undef LOCTEXT_NAMESPACE // AudioWidgetSubsystem
\ No newline at end of file
diff --git a/Engine/Source/Runtime/Engine/Public/Audio/AudioPanelWidgetInterface.h b/Engine/Source/Runtime/Engine/Public/Audio/AudioPanelWidgetInterface.h
new file mode 100644
--- /dev/null
+++ b/Engine/Source/Runtime/Engine/Public/Audio/AudioPanelWidgetInterface.h
@@ -0,0 +1,29 @@
+// Copyright Epic Games, Inc. All Rights Reserved.
+#pragma once
+
+#include "CoreMinimal.h"
+#include "UObject/Class.h"
+#include "UObject/Interface.h"
+#include "UObject/ObjectMacros.h"
+#include "UObject/UObjectGlobals.h"
+
+#include "AudioPanelWidgetInterface.generated.h"
+
+
+UINTERFACE()
+class ENGINE_API UAudioPanelWidgetInterface : public UInterface
+{
+	GENERATED_BODY()
+};
+
+class ENGINE_API IAudioPanelWidgetInterface
+{
+	GENERATED_BODY()
+
+public:
+	UFUNCTION(BlueprintImplementableEvent)
+	FText GetEditorName();
+
+	UFUNCTION(BlueprintImplementableEvent)
+	FName GetIconBrushName();
+};
\ No newline at end of file
diff --git a/Engine/Source/Runtime/Engine/Public/Audio/AudioWidgetSubsystem.h b/Engine/Source/Runtime/Engine/Public/Audio/AudioWidgetSubsystem.h
new file mode 100644
--- /dev/null
+++ b/Engine/Source/Runtime/Engine/Public/Audio/AudioWidgetSubsystem.h
@@ -0,0 +1,36 @@
+// Copyright Epic Games, Inc. All Rights Reserved.
+#pragma once
+
+#include "CoreMinimal.h"
+
+#include "AssetRegistry/AssetData.h"
+#include "Audio/SoundEffectPresetWidgetInterface.h"
+#include "Audio/SoundSubmixWidgetInterface.h"
+#include "Blueprint/UserWidget.h"
+#include "Sound/SoundEffectPreset.h"
+#include "Sound/SoundSubmix.h"
+#include "Subsystems/EngineSubsystem.h"
+#include "Templates/Function.h"
+#include "UObject/Class.h"
+#include "UObject/ObjectMacros.h"
+#include "UObject/UObjectGlobals.h"
+
+#include "AudioWidgetSubsystem.generated.h"
+
+
+// Forward Declarations
+class UWorld;
+
+UCLASS()
+class ENGINE_API UAudioWidgetSubsystem : public UEngineSubsystem
+{
+	GENERATED_BODY()
+
+private:
+	static const TArray<FAssetData> GetBlueprintAssetData(UClass* InInterfaceClass);
+	static bool ImplementsInterface(const FAssetData& InAssetData, UClass* InInterfaceClass);
+
+public:
+	// Returns user widgets that implement a widget interface of the given subclass and optionally, filters and returns only those which pass the provided filter function.
+	TArray<UUserWidget*> CreateUserWidgets(UWorld& InWorld, TSubclassOf<UInterface> InWidgetInterface, TFunction<bool(UUserWidget*)> FilterFunction = TFunction<bool(UUserWidget*)>()) const;
+};
diff --git a/Engine/Source/Runtime/Engine/Public/Audio/SoundEffectPresetWidgetInterface.h b/Engine/Source/Runtime/Engine/Public/Audio/SoundEffectPresetWidgetInterface.h
new file mode 100644
--- /dev/null
+++ b/Engine/Source/Runtime/Engine/Public/Audio/SoundEffectPresetWidgetInterface.h
@@ -0,0 +1,37 @@
+// Copyright Epic Games, Inc. All Rights Reserved.
+#pragma once
+
+#include "CoreMinimal.h"
+
+#include "AudioPanelWidgetInterface.h"
+#include "UObject/Class.h"
+#include "UObject/Interface.h"
+#include "UObject/ObjectMacros.h"
+#include "UObject/UObjectGlobals.h"
+
+#include "SoundEffectPresetWidgetInterface.generated.h"
+
+
+UINTERFACE(Blueprintable)
+class ENGINE_API USoundEffectPresetWidgetInterface : public UAudioPanelWidgetInterface
+{
+	GENERATED_BODY()
+};
+
+class ENGINE_API ISoundEffectPresetWidgetInterface : public IAudioPanelWidgetInterface
+{
+	GENERATED_BODY()
+
+public:
+	// Returns the class of Preset the widget supports
+	UFUNCTION(BlueprintImplementableEvent)
+	TSubclassOf<USoundEffectPreset> GetClass();
+
+	// Called when the preset widget is constructed
+	UFUNCTION(BlueprintImplementableEvent, meta = (DisplayName = "On SoundEffectPreset Widget Constructed"))
+	void OnConstructed(USoundEffectPreset* Preset);
+
+	// Called when the preset object is changed
+	UFUNCTION(BlueprintImplementableEvent)
+	void OnPropertyChanged(USoundEffectPreset* Preset, FName PropertyName);
+};
\ No newline at end of file
diff --git a/Engine/Source/Runtime/Engine/Public/Audio/SoundSubmixWidgetInterface.h b/Engine/Source/Runtime/Engine/Public/Audio/SoundSubmixWidgetInterface.h
new file mode 100644
--- /dev/null
+++ b/Engine/Source/Runtime/Engine/Public/Audio/SoundSubmixWidgetInterface.h
@@ -0,0 +1,26 @@
+// Copyright Epic Games, Inc. All Rights Reserved.
+#pragma once
+
+#include "CoreMinimal.h"
+#include "UObject/Class.h"
+#include "UObject/Interface.h"
+#include "UObject/ObjectMacros.h"
+#include "UObject/UObjectGlobals.h"
+
+#include "SoundSubmixWidgetInterface.generated.h"
+
+
+UINTERFACE(Blueprintable)
+class ENGINE_API USoundSubmixWidgetInterface : public UInterface
+{
+	GENERATED_BODY()
+};
+
+class ENGINE_API ISoundSubmixWidgetInterface
+{
+	GENERATED_BODY()
+
+public:
+	UFUNCTION(BlueprintImplementableEvent, meta = (DisplayName = "On SoundSubmix Widget Constructed"))
+	void OnConstructed(USoundSubmixBase* SoundSubmix);
+};
\ No newline at end of file
