Added  FTaskTagScope() to address an "oldValue == newValue" ensure.

#rnx

diff --git a/Engine/Source/Runtime/Engine/Private/Particles/ParticleComponents.cpp b/Engine/Source/Runtime/Engine/Private/Particles/ParticleComponents.cpp
--- a/Engine/Source/Runtime/Engine/Private/Particles/ParticleComponents.cpp
+++ b/Engine/Source/Runtime/Engine/Private/Particles/ParticleComponents.cpp
@@ -4781,6 +4781,7 @@ public:
 
 	void DoTask(ENamedThreads::Type CurrentThread, const FGraphEventRef& MyCompletionGraphEvent)
 	{
+		FTaskTagScope Scope(ETaskTag::EParallelRenderingThread);
 		Target->ComputeTickComponent_Concurrent();
 #if !WITH_EDITOR  // otherwise this is queued by the calling code because we need to be able to block and wait on it
 		{
