Fix for strata material with null material to not be setup for compilation with strata.
This is needed because the mateiral attribute always setup assign a codechunk to all its input and we need to detect the case where strata input is default.
This is only temporary while strata becomes just the defacto material.
Removed a condition in MaterialTemplate.ush that is no longer needed.

diff --git a/Engine/Shaders/Private/MaterialTemplate.ush b/Engine/Shaders/Private/MaterialTemplate.ush
--- a/Engine/Shaders/Private/MaterialTemplate.ush
+++ b/Engine/Shaders/Private/MaterialTemplate.ush
@@ -21,20 +21,7 @@
 
 #define TEMPLATE_USES_STRATA (PROJECT_STRATA && MATERIAL_IS_STRATA)
 
-#if TEMPLATE_USES_STRATA
 #include "/Engine/Private/Strata/Strata.ush"
-#else
- // This is a simple way to not have strata included in vertex shader for instance while allowing generated code to compile.
-struct FStrataData
-{
-	uint a;
-};
-FStrataData GetInitialisedStrataData()
-{
-	FStrataData StrataData = (FStrataData)0;
-	return StrataData;
-}
-#endif
 
 #if USES_SPEEDTREE
 	#include "/Engine/Private/SpeedTreeCommon.ush"
diff --git a/Engine/Source/Runtime/Engine/Private/Materials/HLSLMaterialTranslator.cpp b/Engine/Source/Runtime/Engine/Private/Materials/HLSLMaterialTranslator.cpp
--- a/Engine/Source/Runtime/Engine/Private/Materials/HLSLMaterialTranslator.cpp
+++ b/Engine/Source/Runtime/Engine/Private/Materials/HLSLMaterialTranslator.cpp
@@ -107,6 +107,7 @@ FHLSLMaterialTranslator::FHLSLMaterialTranslator(FMaterial* InMaterial,
 ,	bAllowCodeChunkGeneration(true)
 ,	bUsesPerInstanceCustomData(false)
 ,	bUsesAnisotropy(false)
+,	bMaterialIsStrata(false)
 ,	AllocatedUserTexCoords()
 ,	AllocatedUserVertexTexCoords()
 ,	DynamicParticleParameterMask(0)
@@ -1117,6 +1118,7 @@ bool FHLSLMaterialTranslator::Translate()
 
 		if (bStrataFrontMaterialIsValid)
 		{
+			bMaterialIsStrata = true;
 			const FStrataMaterialCompilationInfo& StrataCompilationInfo = CodeChunkToStrataCompilationInfoMap[Chunk[MP_FrontMaterial]];
 
 			// VolumetricFogCloud must be used in isolation
@@ -1558,12 +1560,11 @@ void FHLSLMaterialTranslator::GetMaterialEnvironment(EShaderPlatform InPlatform,
 		}
 	}
 
-	const bool bStrataMaterial = Material->IsStrataMaterial();
-	OutEnvironment.SetDefine(TEXT("MATERIAL_IS_STRATA"), bStrataMaterial ? TEXT("1") : TEXT("0"));
+	OutEnvironment.SetDefine(TEXT("MATERIAL_IS_STRATA"), bMaterialIsStrata ? TEXT("1") : TEXT("0"));
 
 	// STRATA_TODO do not request DualSourceBlending if gray scale transmittance is selected.
 	// bMaterialRequestsDualSourceBlending this base on limited set of blend mode: Opaque, Masked, TransmittanceCoverage, TransmittanceColored;
-	bMaterialRequestsDualSourceBlending |= bStrataMaterial;
+	bMaterialRequestsDualSourceBlending |= bMaterialIsStrata;
 
 	// if duals source blending (colored transmittance) is not supported on a platform, it will fall back to standard alpha blending (grey scale transmittance)
 	OutEnvironment.SetDefine(TEXT("DUAL_SOURCE_COLOR_BLENDING_ENABLED"), bMaterialRequestsDualSourceBlending && Material->IsDualBlendingEnabled(Platform) ? TEXT("1") : TEXT("0"));
diff --git a/Engine/Source/Runtime/Engine/Private/Materials/HLSLMaterialTranslator.h b/Engine/Source/Runtime/Engine/Private/Materials/HLSLMaterialTranslator.h
--- a/Engine/Source/Runtime/Engine/Private/Materials/HLSLMaterialTranslator.h
+++ b/Engine/Source/Runtime/Engine/Private/Materials/HLSLMaterialTranslator.h
@@ -297,6 +297,12 @@ protected:
 	/** True if this material write anisotropy material property */
 	uint32 bUsesAnisotropy : 1;
 
+	/** 
+	 * True if the material is detected as a strata material at compile time.
+	 * This is decoupled from runtime FMaterialResource::IsStrataMaterial but practically fine since this is only temporary until Strata is the main shading system. Only really used at runtime for translucency dual source blending.
+	 */
+	uint32 bMaterialIsStrata : 1; // 
+
 	/** Tracks the texture coordinates used by this material. */
 	TBitArray<> AllocatedUserTexCoords;
 	/** Tracks the texture coordinates used by the vertex shader in this material. */
diff --git a/Engine/Source/Runtime/Engine/Private/Materials/MaterialExpressions.cpp b/Engine/Source/Runtime/Engine/Private/Materials/MaterialExpressions.cpp
--- a/Engine/Source/Runtime/Engine/Private/Materials/MaterialExpressions.cpp
+++ b/Engine/Source/Runtime/Engine/Private/Materials/MaterialExpressions.cpp
@@ -6361,6 +6361,10 @@ uint32 UMaterialExpressionSetMaterialAttributes::GetInputType(int32 InputIndex)
 		{
 			InputType = MCT_ShadingModel;
 		}
+		else if (InputType == MCT_Strata)
+		{
+			InputType = MCT_Strata;
+		}
 		else
 		{
 			InputType = MCT_Float3;
