Don't take ResourceState.AssetLODBias into account when computing the cummulutive mip sizes for streaming

diff --git a/Engine/Source/Runtime/Engine/Private/Streaming/StreamingTexture.cpp b/Engine/Source/Runtime/Engine/Private/Streaming/StreamingTexture.cpp
--- a/Engine/Source/Runtime/Engine/Private/Streaming/StreamingTexture.cpp
+++ b/Engine/Source/Runtime/Engine/Private/Streaming/StreamingTexture.cpp
@@ -107,7 +107,7 @@ void FStreamingRenderAsset::UpdateStaticData(const FRenderAssetStreamingSettings
 
 		for (int32 LODIndex = 0; LODIndex < ResourceState.MaxNumLODs; ++LODIndex)
 		{
-			CumulativeLODSizes[LODIndex] = RenderAsset->CalcCumulativeLODSize(LODIndex + ResourceState.AssetLODBias + 1);
+			CumulativeLODSizes[LODIndex] = RenderAsset->CalcCumulativeLODSize(LODIndex + 1);
 		}
 
 		// If there are optional mips
