Fix for master reverb asset not working.

Porting 14885691 from 4.26.1

diff --git a/Engine/Source/Runtime/Engine/Private/AudioEffect.cpp b/Engine/Source/Runtime/Engine/Private/AudioEffect.cpp
--- a/Engine/Source/Runtime/Engine/Private/AudioEffect.cpp
+++ b/Engine/Source/Runtime/Engine/Private/AudioEffect.cpp
@@ -361,7 +361,7 @@ FAudioEffectsManager::FAudioEffectsManager( FAudioDevice* InDevice )
 	, CurrentEQMix(nullptr)
 	, bReverbActive(false)
 	, bEQActive(false)
-	, bReverbChanged(true) // Initially true to update from default reverb setting
+	, bReverbChanged(false)
 	, bEQChanged(false)
 {
 	InitAudioEffects();
@@ -377,13 +377,6 @@ void FAudioEffectsManager::ResetInterpolation()
  */
 void FAudioEffectsManager::InitAudioEffects( void )
 {
-	// Clear out the default reverb settings
-	FReverbSettings ReverbSettings;
-	ReverbSettings.ReverbEffect = NULL;
-	ReverbSettings.Volume = 0.0f;
-	ReverbSettings.FadeTime = 0.1f;
-	SetReverbSettings( ReverbSettings );
-
 	FMemory::Memzero(&PrevReverbEffect, sizeof(PrevReverbEffect));
 
 	ClearMixSettings();
