IWYU fix for unity build

diff --git a/Engine/Source/Runtime/Engine/Public/Audio/SoundEffectPresetWidgetInterface.h b/Engine/Source/Runtime/Engine/Public/Audio/SoundEffectPresetWidgetInterface.h
--- a/Engine/Source/Runtime/Engine/Public/Audio/SoundEffectPresetWidgetInterface.h
+++ b/Engine/Source/Runtime/Engine/Public/Audio/SoundEffectPresetWidgetInterface.h
@@ -4,6 +4,7 @@
 #include "CoreMinimal.h"
 
 #include "AudioPanelWidgetInterface.h"
+#include "Sound/SoundEffectPreset.h"
 #include "UObject/Class.h"
 #include "UObject/Interface.h"
 #include "UObject/ObjectMacros.h"
