Adding logging to shader autogen.

diff --git a/Engine/Source/Runtime/Engine/Private/ShaderCompiler/ShaderGenerationUtil.cpp b/Engine/Source/Runtime/Engine/Private/ShaderCompiler/ShaderGenerationUtil.cpp
--- a/Engine/Source/Runtime/Engine/Private/ShaderCompiler/ShaderGenerationUtil.cpp
+++ b/Engine/Source/Runtime/Engine/Private/ShaderCompiler/ShaderGenerationUtil.cpp
@@ -1993,6 +1993,9 @@ FGBufferInfo FShaderCompileUtilities::FetchGBufferInfoAndWriteAutogen(EShaderPla
 		OutputFileData += TEXT("#endif\n");
 		OutputFileData += TEXT("\n");
 
+
+		UE_LOG(LogShaderCompilers, Display, TEXT("Compiling shader autogen file: %s"), *AutogenHeaderFilename);
+
 		// If this file already exists, and the text file is the same as what we are planning to write, then don't write it. In that case,
 		// we don't need to update it (since it didn't change) and it leaves the timestamp unchanged in case you need to open the file manually
 		// for debugging.
@@ -2002,6 +2005,7 @@ FGBufferInfo FShaderCompileUtilities::FetchGBufferInfoAndWriteAutogen(EShaderPla
 			bool bFileExists = FFileHelper::LoadFileToString(PrevFileData,*AutogenHeaderFilename);
 			if (bFileExists && PrevFileData == OutputFileData)
 			{
+				UE_LOG(LogShaderCompilers, Display, TEXT("Autogen file is unchanged, skipping write."));
 				bWriteNeeded = false;
 			}
 		}
@@ -2009,10 +2013,27 @@ FGBufferInfo FShaderCompileUtilities::FetchGBufferInfoAndWriteAutogen(EShaderPla
 		if (bWriteNeeded)
 		{
 			FPaths::MakeStandardFilename(AutogenHeaderFilenameTemp);
-			FFileHelper::SaveStringToFile(OutputFileData, *AutogenHeaderFilenameTemp, FFileHelper::EEncodingOptions::ForceAnsi);
+			bool bSaveOk = FFileHelper::SaveStringToFile(OutputFileData, *AutogenHeaderFilenameTemp, FFileHelper::EEncodingOptions::ForceAnsi);
+			if (!bSaveOk)
+			{
+				UE_LOG(LogShaderCompilers, Display, TEXT("Failed to save shader autogen file: %s"), *AutogenHeaderFilename);
+			}
+
+			bool bDeleteOk = FPlatformFileManager::Get().GetPlatformFile().DeleteFile(*AutogenHeaderFilename);
+			if (!bDeleteOk)
+			{
+				// note that it will fail to delete when the file does not exist, which is acceptable. so in that case
+				// write to the log, but don't error out
+				UE_LOG(LogShaderCompilers, Display, TEXT("Failed to delete old shader autogen file: %s"), *AutogenHeaderFilename);
+			}
+
+			bool bRenameOk = FPlatformFileManager::Get().GetPlatformFile().MoveFile(*AutogenHeaderFilename, *AutogenHeaderFilenameTemp);
+			if (!bRenameOk)
+			{
+				UE_LOG(LogShaderCompilers, Display, TEXT("Failed to rename shader autogen file: %s"), *AutogenHeaderFilename);
+			}
 
-			FPlatformFileManager::Get().GetPlatformFile().DeleteFile(*AutogenHeaderFilename);
-			FPlatformFileManager::Get().GetPlatformFile().MoveFile(*AutogenHeaderFilename, *AutogenHeaderFilenameTemp);
+			UE_LOG(LogShaderCompilers, Display, TEXT("Shader autogen file written: %s"), *AutogenHeaderFilename);
 		}
 	}
 
