Use a dumb string instead of smart soft object path to prevent the cooker from annexing references to assets that are only used for Editor automated tests.

This approach is preferred over the functional alternative of #if WITH_EDITORONLY_DATA as this is more simple and permutes the build less across all targets and configurations.

Restore references to test assets used by [REDACTED_PERSON] trace gathering.

diff --git a/Engine/Source/Runtime/Engine/Classes/Tests/AutomationTestSettings.h b/Engine/Source/Runtime/Engine/Classes/Tests/AutomationTestSettings.h
--- a/Engine/Source/Runtime/Engine/Classes/Tests/AutomationTestSettings.h
+++ b/Engine/Source/Runtime/Engine/Classes/Tests/AutomationTestSettings.h
@@ -371,7 +371,7 @@ public:
 	 * Asset to test for open in automation process
 	 */
 	UPROPERTY(EditAnywhere, config, Category="Open Asset Tests")
-	TArray<FSoftObjectPath> AssetsToOpen;
+	TArray<FString> AssetsToOpen;
 
 	/**
 	 * Maps to PIE during the PIE test
