-Added Blend Profile and Blend Mode overrides on Montage_Play/Stop
-Added BlendMode In and Out
-Exposed blend parameters in AnimMontage to blueprint, for ease of use w/ new setting overrides

diff --git a/Engine/Source/Runtime/Engine/Classes/Animation/AnimInstance.h b/Engine/Source/Runtime/Engine/Classes/Animation/AnimInstance.h
--- a/Engine/Source/Runtime/Engine/Classes/Animation/AnimInstance.h
+++ b/Engine/Source/Runtime/Engine/Classes/Animation/AnimInstance.h
@@ -602,15 +602,23 @@ public:
 	UFUNCTION(BlueprintCallable, Category = "Montage")
 	float Montage_PlayWithBlendIn(UAnimMontage* MontageToPlay, const FAlphaBlendArgs& BlendIn, float InPlayRate = 1.f, EMontagePlayReturnType ReturnValueType = EMontagePlayReturnType::MontageLength, float InTimeToStartMontageAt=0.f, bool bStopAllMontages = true);
 
+	/** Plays an animation montage. Same as Montage_Play, but you can overwrite all of the montage's default blend in settings. */
+	UFUNCTION(BlueprintCallable, Category = "Montage")
+	float Montage_PlayWithBlendSettings(UAnimMontage* MontageToPlay, const FMontageBlendSettings& BlendInSettings, float InPlayRate = 1.f, EMontagePlayReturnType ReturnValueType = EMontagePlayReturnType::MontageLength, float InTimeToStartMontageAt=0.f, bool bStopAllMontages = true);
+
 	/** Stops the animation montage. If reference is NULL, it will stop ALL active montages. */
 	/** Stopped montages will blend out using their montage asset's BlendOut, with InBlendOutTime as the BlendTime */
 	UFUNCTION(BlueprintCallable, Category = "Montage")
 	void Montage_Stop(float InBlendOutTime, const UAnimMontage* Montage = NULL);
 
-	/** Same as Montage_Stop, but all BlendOut settings are provided instead of using the ones on the montage asset*/
+	/** Same as Montage_Stop. Uses values from the AlphaBlendArgs. Other settings come from the montage asset*/
 	UFUNCTION(BlueprintCallable, Category = "Montage")
 	void Montage_StopWithBlendOut(const FAlphaBlendArgs& BlendOut, const UAnimMontage* Montage = nullptr);
 
+	/** Same as Montage_Stop, but all blend settings are provided instead of using the ones on the montage asset*/
+	UFUNCTION(BlueprintCallable, Category = "Montage")
+	void Montage_StopWithBlendSettings(const FMontageBlendSettings& BlendOutSettings, const UAnimMontage* Montage = nullptr);
+
 	/** Stops all active montages belonging to a group. */
 	UFUNCTION(BlueprintCallable, Category = "Montage")
 	void Montage_StopGroupByName(float InBlendOutTime, FName GroupName);
@@ -843,16 +851,20 @@ protected:
 	/**  Inertialization requests gathered this frame. Gets reset in UpdateMontageEvaluationData */
 	TMap<FName, float> SlotGroupInertializationRequestMap;
 
-	/** Stop all active montages belonging to 'InGroupName' */
+	/* StopAllMontagesByGroupName needs a BlendMode and BlendProfile to function properly if using non-default ones in your montages. If you want default BlendMode/BlendProfiles, you need to update the calling code to do so. */
+	UE_DEPRECATED(5.0, "Use StopAllMontagesByGroupName with other signature.")
 	void StopAllMontagesByGroupName(FName InGroupName, const FAlphaBlend& BlendOut);
 
+	/** Stop all active montages belonging to 'InGroupName' */
+	void StopAllMontagesByGroupName(FName InGroupName, const FMontageBlendSettings& BlendOutSettings);
+
 	/** Update weight of montages  **/
 	virtual void Montage_UpdateWeight(float DeltaSeconds);
 	/** Advance montages **/
 	virtual void Montage_Advance(float DeltaSeconds);
 
-	void Montage_StopInternal(TFunctionRef<FAlphaBlend(const FAnimMontageInstance*)> AlphaBlendSelectorFunction, const UAnimMontage* Montage = nullptr);
-	float Montage_PlayInternal(UAnimMontage* MontageToPlay, const FAlphaBlend& BlendIn, float InPlayRate = 1.f, EMontagePlayReturnType ReturnValueType = EMontagePlayReturnType::MontageLength, float InTimeToStartMontageAt = 0.f, bool bStopAllMontages = true);
+	void Montage_StopInternal(TFunctionRef<FMontageBlendSettings(const FAnimMontageInstance*)> AlphaBlendSelectorFunction, const UAnimMontage* Montage = nullptr);
+	float Montage_PlayInternal(UAnimMontage* MontageToPlay, const FMontageBlendSettings& BlendInSettings, float InPlayRate = 1.f, EMontagePlayReturnType ReturnValueType = EMontagePlayReturnType::MontageLength, float InTimeToStartMontageAt = 0.f, bool bStopAllMontages = true);
 
 public:
 
diff --git a/Engine/Source/Runtime/Engine/Classes/Animation/AnimMontage.h b/Engine/Source/Runtime/Engine/Classes/Animation/AnimMontage.h
--- a/Engine/Source/Runtime/Engine/Classes/Animation/AnimMontage.h
+++ b/Engine/Source/Runtime/Engine/Classes/Animation/AnimMontage.h
@@ -290,6 +290,29 @@ private:
 	float GetCurrSectionStartPosition_Target() const;
 };
 
+/**
+* Montage blend settings. Can be used to overwrite default Montage settings on Play/Stop
+*/
+USTRUCT(BlueprintType)
+struct ENGINE_API FMontageBlendSettings
+{
+	GENERATED_BODY()
+
+	FMontageBlendSettings();
+
+	/** Blend Profile to use for this blend */
+	UPROPERTY(EditAnywhere, BlueprintReadWrite, Category = "Blend", meta = (DisplayAfter = "Blend"))
+	UBlendProfile* BlendProfile;
+
+	/** AlphaBlend options (time, curve, etc.) */
+	UPROPERTY(EditAnywhere, BlueprintReadWrite, Category = "Blend", meta = (DisplayAfter = "BlendMode"))
+	FAlphaBlendArgs Blend;
+
+	/** Type of blend mode (Standard vs Inertial) */
+	UPROPERTY(EditAnywhere, BlueprintReadWrite, Category = "Blend")
+	EMontageBlendMode BlendMode;
+};
+
 USTRUCT()
 struct ENGINE_API FAnimMontageInstance
 {
@@ -457,12 +480,16 @@ public:
 
 	//~ Begin montage instance Interfaces
 
-	// Fills BlendIn from the associated montage 
+	// Blend in with the supplied play rate. Other blend settings will come from the Montage asset.
 	void Play(float InPlayRate = 1.f);
-	// Use provided BlendIn
-	void Play(float InPlayRate, const FAlphaBlend& BlendIn);
+	// Blend in with the supplied blend settings
+	void Play(float InPlayRate, const FMontageBlendSettings& BlendInSettings);
 
+	// Blend out with the supplied FAlphaBlend. Other blend settings will come from the Montage asset.
 	void Stop(const FAlphaBlend& InBlendOut, bool bInterrupt=true);
+	// Blend out with the supplied blend settings
+	void Stop(const FMontageBlendSettings& InBlendOutSettings, bool bInterrupt=true);
+
 	void Pause();
 	void Initialize(class UAnimMontage * InMontage);
 
@@ -591,8 +618,11 @@ class UAnimMontage : public UAnimCompositeBase
 
 	friend struct FAnimMontageInstance;
 	
-	UPROPERTY(EditAnywhere, Category = BlendOption)
-	EMontageBlendMode BlendMode;
+	UPROPERTY(EditAnywhere, BlueprintReadOnly, Category = BlendOption)
+	EMontageBlendMode BlendModeIn;
+
+	UPROPERTY(EditAnywhere, BlueprintReadOnly, Category = BlendOption)
+	EMontageBlendMode BlendModeOut;
 
 	/** Blend in option. */
 	UPROPERTY(EditAnywhere, Category=BlendOption)
@@ -614,6 +644,15 @@ class UAnimMontage : public UAnimCompositeBase
 	UPROPERTY(EditAnywhere, Category = BlendOption)
 	float BlendOutTriggerTime;
 
+	UFUNCTION(BlueprintPure, Category = "Montage")
+	FAlphaBlendArgs GetBlendInArgs() const { return FAlphaBlendArgs(BlendIn); }
+
+	UFUNCTION(BlueprintPure, Category = "Montage")
+	FAlphaBlendArgs GetBlendOutArgs() const { return FAlphaBlendArgs(BlendOut); }
+
+	UFUNCTION(BlueprintPure, Category = "Montage")
+	float GetDefaultBlendInTime() const { return BlendIn.GetBlendTime(); }
+
 	UFUNCTION(BlueprintCallable, Category = "Montage")
 	float GetDefaultBlendOutTime() const { return BlendOut.GetBlendTime(); }
 
@@ -653,11 +692,11 @@ class UAnimMontage : public UAnimCompositeBase
 	bool bEnableAutoBlendOut;
 
 	/** The blend profile to use. */
-	UPROPERTY(EditAnywhere, Category = BlendOption, meta = (UseAsBlendProfile = true))
+	UPROPERTY(EditAnywhere, BlueprintReadOnly, Category = BlendOption, meta = (UseAsBlendProfile = true))
 	UBlendProfile* BlendProfileIn;
 
 	/** The blend profile to use. */
-	UPROPERTY(EditAnywhere, Category = BlendOption, meta = (UseAsBlendProfile = true))
+	UPROPERTY(EditAnywhere, BlueprintReadOnly, Category = BlendOption, meta = (UseAsBlendProfile = true))
 	UBlendProfile* BlendProfileOut;
 
 	/** Root Bone will be locked to that position when extracting root motion. DEPRECATED in 4.5 root motion is controlled by anim sequences **/
diff --git a/Engine/Source/Runtime/Engine/Private/AlphaBlend.cpp b/Engine/Source/Runtime/Engine/Private/AlphaBlend.cpp
--- a/Engine/Source/Runtime/Engine/Private/AlphaBlend.cpp
+++ b/Engine/Source/Runtime/Engine/Private/AlphaBlend.cpp
@@ -234,3 +234,11 @@ FAlphaBlendArgs::FAlphaBlendArgs()
 {
 
 }
+
+FAlphaBlendArgs::FAlphaBlendArgs(const struct FAlphaBlend& InAlphaBlend)
+	: CustomCurve(InAlphaBlend.GetCustomCurve())
+	, BlendTime(InAlphaBlend.GetBlendTime())
+	, BlendOption(InAlphaBlend.GetBlendOption())
+{
+
+}
diff --git a/Engine/Source/Runtime/Engine/Private/Animation/AnimInstance.cpp b/Engine/Source/Runtime/Engine/Private/Animation/AnimInstance.cpp
--- a/Engine/Source/Runtime/Engine/Private/Animation/AnimInstance.cpp
+++ b/Engine/Source/Runtime/Engine/Private/Animation/AnimInstance.cpp
@@ -1859,7 +1859,7 @@ bool UAnimInstance::IsPlayingSlotAnimation(const UAnimSequenceBase* Asset, FName
 	return false;
 }
 
-float UAnimInstance::Montage_PlayInternal(UAnimMontage* MontageToPlay, const FAlphaBlend& BlendIn, float InPlayRate /*= 1.f*/, EMontagePlayReturnType ReturnValueType /*= EMontagePlayReturnType::MontageLength*/, float InTimeToStartMontageAt /*= 0.f*/, bool bStopAllMontages /*= true*/)
+float UAnimInstance::Montage_PlayInternal(UAnimMontage* MontageToPlay, const FMontageBlendSettings& BlendInSettings, float InPlayRate /*= 1.f*/, EMontagePlayReturnType ReturnValueType /*= EMontagePlayReturnType::MontageLength*/, float InTimeToStartMontageAt /*= 0.f*/, bool bStopAllMontages /*= true*/)
 {
 	LLM_SCOPE(ELLMTag::Animation);
 
@@ -1867,22 +1867,11 @@ float UAnimInstance::Montage_PlayInternal(UAnimMontage* MontageToPlay, const FAl
 	{
 		if (CurrentSkeleton && CurrentSkeleton->IsCompatible(MontageToPlay->GetSkeleton()))
 		{
-			const float BlendDuration = BlendIn.GetBlendTime();
-			// If using inertial blend, force an instant blend in
-			const FAlphaBlend EffectiveBlendIn = (MontageToPlay->BlendMode == EMontageBlendMode::Inertialization) ? FAlphaBlend(BlendIn, 0.0f) : BlendIn;
-
 			const FName NewMontageGroupName = MontageToPlay->GetGroupName();
 			if (bStopAllMontages)
 			{
 				// Enforce 'a single montage at once per group' rule
-				StopAllMontagesByGroupName(NewMontageGroupName, EffectiveBlendIn);
-			}
-
-			if (MontageToPlay->BlendMode == EMontageBlendMode::Inertialization)
-			{
-				// Request new inertialization for new montage's group name
-				// The above montage stop may request an inertial blend. We overwrite that here.
-				RequestMontageInertialization(MontageToPlay, BlendDuration);
+				StopAllMontagesByGroupName(NewMontageGroupName, BlendInSettings);
 			}
 
 			// Enforce 'a single root motion montage at once' rule.
@@ -1891,7 +1880,7 @@ float UAnimInstance::Montage_PlayInternal(UAnimMontage* MontageToPlay, const FAl
 				FAnimMontageInstance* ActiveRootMotionMontageInstance = GetRootMotionMontageInstance();
 				if (ActiveRootMotionMontageInstance)
 				{
-					ActiveRootMotionMontageInstance->Stop(EffectiveBlendIn);
+					ActiveRootMotionMontageInstance->Stop(BlendInSettings);
 				}
 			}
 
@@ -1901,7 +1890,7 @@ float UAnimInstance::Montage_PlayInternal(UAnimMontage* MontageToPlay, const FAl
 			const float MontageLength = MontageToPlay->GetPlayLength();
 
 			NewInstance->Initialize(MontageToPlay);
-			NewInstance->Play(InPlayRate, EffectiveBlendIn);
+			NewInstance->Play(InPlayRate, BlendInSettings);
 			NewInstance->SetPosition(FMath::Clamp(InTimeToStartMontageAt, 0.f, MontageLength));
 			MontageInstances.Add(NewInstance);
 			ActiveMontagesMap.Add(MontageToPlay, NewInstance);
@@ -1932,16 +1921,37 @@ float UAnimInstance::Montage_PlayInternal(UAnimMontage* MontageToPlay, const FAl
 /** Play a Montage. Returns Length of Montage in seconds. Returns 0.f if failed to play. */
 float UAnimInstance::Montage_Play(UAnimMontage* MontageToPlay, float InPlayRate/*= 1.f*/, EMontagePlayReturnType ReturnValueType, float InTimeToStartMontageAt, bool bStopAllMontages /*= true*/)
 {
-	FAlphaBlend BlendIn = MontageToPlay ? MontageToPlay->BlendIn : FAlphaBlend();
-	return Montage_PlayInternal(MontageToPlay, BlendIn, InPlayRate, ReturnValueType, InTimeToStartMontageAt, bStopAllMontages);
+	FMontageBlendSettings BlendSettings;
+	if (MontageToPlay)
+	{
+		BlendSettings.Blend = MontageToPlay->BlendIn;
+		BlendSettings.BlendMode = MontageToPlay->BlendModeIn;
+		BlendSettings.BlendProfile = MontageToPlay->BlendProfileIn;
+	}
+
+	return Montage_PlayInternal(MontageToPlay, BlendSettings, InPlayRate, ReturnValueType, InTimeToStartMontageAt, bStopAllMontages);
 }
 
 float UAnimInstance::Montage_PlayWithBlendIn(UAnimMontage* MontageToPlay, const FAlphaBlendArgs& BlendIn, float InPlayRate /*= 1.f*/, EMontagePlayReturnType ReturnValueType /*= EMontagePlayReturnType::MontageLength*/, float InTimeToStartMontageAt/*=0.f*/, bool bStopAllMontages /*= true*/)
 {
-	return Montage_PlayInternal(MontageToPlay, FAlphaBlend(BlendIn), InPlayRate, ReturnValueType, InTimeToStartMontageAt, bStopAllMontages);
+	FMontageBlendSettings BlendSettings;
+	BlendSettings.Blend = BlendIn;
+
+	if (MontageToPlay)
+	{
+		BlendSettings.BlendMode = MontageToPlay->BlendModeIn;
+		BlendSettings.BlendProfile = MontageToPlay->BlendProfileIn;
+	}
+
+	return Montage_PlayInternal(MontageToPlay, BlendSettings, InPlayRate, ReturnValueType, InTimeToStartMontageAt, bStopAllMontages);
 }
 
-void UAnimInstance::Montage_StopInternal(TFunctionRef<FAlphaBlend(const FAnimMontageInstance*)> AlphaBlendSelectorFunction, const UAnimMontage* Montage /*= nullptr*/)
+float UAnimInstance::Montage_PlayWithBlendSettings(UAnimMontage* MontageToPlay, const FMontageBlendSettings& BlendInSettings, float InPlayRate /*= 1.f*/, EMontagePlayReturnType ReturnValueType /*= EMontagePlayReturnType::MontageLength*/, float InTimeToStartMontageAt/*=0.f*/, bool bStopAllMontages /*= true*/)
+{
+	return Montage_PlayInternal(MontageToPlay, BlendInSettings, InPlayRate, ReturnValueType, InTimeToStartMontageAt, bStopAllMontages);
+}
+
+void UAnimInstance::Montage_StopInternal(TFunctionRef<FMontageBlendSettings(const FAnimMontageInstance*)> AlphaBlendSelectorFunction, const UAnimMontage* Montage /*= nullptr*/)
 {
 	if (Montage)
 	{
@@ -1969,8 +1979,17 @@ void UAnimInstance::Montage_Stop(float InBlendOutTime, const UAnimMontage* Monta
 {
 	auto AlphaBlendFromInstanceAndInBlendOutTime = [InBlendOutTime](const FAnimMontageInstance* InMontageInstance)
 	{
-		check(InMontageInstance->Montage);
-		return FAlphaBlend(InMontageInstance->Montage->BlendOut, InBlendOutTime);
+		FMontageBlendSettings BlendOutSettings;
+		if (const UAnimMontage* InstanceMontage = InMontageInstance->Montage)
+		{
+			//Grab all settings from the montage, except BlendTime
+			BlendOutSettings.Blend = InstanceMontage->BlendOut;
+			BlendOutSettings.Blend.BlendTime = InBlendOutTime;
+			BlendOutSettings.BlendMode = InstanceMontage->BlendModeOut;
+			BlendOutSettings.BlendProfile = InstanceMontage->BlendProfileOut;
+		}
+
+		return BlendOutSettings;
 	};
 
 	Montage_StopInternal(AlphaBlendFromInstanceAndInBlendOutTime, Montage);
@@ -1978,10 +1997,28 @@ void UAnimInstance::Montage_Stop(float InBlendOutTime, const UAnimMontage* Monta
 
 void UAnimInstance::Montage_StopWithBlendOut(const FAlphaBlendArgs& BlendOutArgs, const UAnimMontage* Montage /*= NULL*/)
 {
-	FAlphaBlend BlendOut = FAlphaBlend(BlendOutArgs);
-	auto AlphaBlendPassthrough = [BlendOut](const FAnimMontageInstance* InMontageInstance)
+	auto AlphaBlendPassthrough = [BlendOutArgs](const FAnimMontageInstance* InMontageInstance)
 	{
-		return BlendOut;
+		FMontageBlendSettings BlendOutSettings;
+		if (const UAnimMontage* InstanceMontage = InMontageInstance->Montage)
+		{
+			//Grab all settings from the montage, except the FAlphaBlend
+			BlendOutSettings.Blend = BlendOutArgs;
+			BlendOutSettings.BlendMode = InstanceMontage->BlendModeOut;
+			BlendOutSettings.BlendProfile = InstanceMontage->BlendProfileOut;
+		}
+
+		return BlendOutSettings;
+	};
+
+	Montage_StopInternal(AlphaBlendPassthrough, Montage);
+}
+
+void UAnimInstance::Montage_StopWithBlendSettings(const FMontageBlendSettings& BlendOutSettings, const UAnimMontage* Montage /*= nullptr*/)
+{
+	auto AlphaBlendPassthrough = [BlendOutSettings](const FAnimMontageInstance* InMontageInstance)
+	{
+		return BlendOutSettings;
 	};
 
 	Montage_StopInternal(AlphaBlendPassthrough, Montage);
@@ -2491,13 +2528,20 @@ void UAnimInstance::StopAllMontages(float BlendOut)
 }
 
 void UAnimInstance::StopAllMontagesByGroupName(FName InGroupName, const FAlphaBlend& BlendOut)
+{
+	FMontageBlendSettings BlendOutSettings;
+	BlendOutSettings.Blend = BlendOut;
+	StopAllMontagesByGroupName(InGroupName, BlendOutSettings);
+}
+
+void UAnimInstance::StopAllMontagesByGroupName(FName InGroupName, const FMontageBlendSettings& BlendOutSettings)
 {
 	for (int32 InstanceIndex = MontageInstances.Num() - 1; InstanceIndex >= 0; InstanceIndex--)
 	{
 		FAnimMontageInstance* MontageInstance = MontageInstances[InstanceIndex];
 		if (MontageInstance && MontageInstance->Montage && (MontageInstance->Montage->GetGroupName() == InGroupName))
 		{
-			MontageInstances[InstanceIndex]->Stop(BlendOut, true);
+			MontageInstances[InstanceIndex]->Stop(BlendOutSettings, true);
 		}
 	}
 }
diff --git a/Engine/Source/Runtime/Engine/Private/Animation/AnimMontage.cpp b/Engine/Source/Runtime/Engine/Private/Animation/AnimMontage.cpp
--- a/Engine/Source/Runtime/Engine/Private/Animation/AnimMontage.cpp
+++ b/Engine/Source/Runtime/Engine/Private/Animation/AnimMontage.cpp
@@ -38,7 +38,8 @@ namespace MontageFNames
 UAnimMontage::UAnimMontage(const FObjectInitializer& ObjectInitializer)
 	: Super(ObjectInitializer)
 {
-	BlendMode = EMontageBlendMode::Standard;
+	BlendModeIn = EMontageBlendMode::Standard;
+	BlendModeOut = EMontageBlendMode::Standard;
 
 	BlendIn.SetBlendTime(0.25f);
 	BlendOut.SetBlendTime(0.25f);
@@ -1342,10 +1343,20 @@ FAnimMontageInstance::FAnimMontageInstance(UAnimInstance * InAnimInstance)
 
 void FAnimMontageInstance::Play(float InPlayRate)
 {
-	Play(InPlayRate, Montage->BlendIn);
+	FMontageBlendSettings BlendInSettings;
+
+	// Fill settings from our Montage asset
+	if (Montage)
+	{
+		BlendInSettings.Blend = Montage->BlendIn;
+		BlendInSettings.BlendMode = Montage->BlendModeIn;
+		BlendInSettings.BlendProfile = Montage->BlendProfileIn;
+	}
+
+	Play(InPlayRate, BlendInSettings);
 }
 
-void FAnimMontageInstance::Play(float InPlayRate, const FAlphaBlend& BlendIn)
+void FAnimMontageInstance::Play(float InPlayRate, const FMontageBlendSettings& BlendInSettings)
 {
 	bPlaying = true;
 	PlayRate = InPlayRate;
@@ -1353,15 +1364,28 @@ void FAnimMontageInstance::Play(float InPlayRate, const FAlphaBlend& BlendIn)
 	// if this doesn't exist, nothing works
 	check(Montage);
 	
+	// Inertialization
+	FAlphaBlendArgs BlendInArgs = BlendInSettings.Blend;
+	if (AnimInstance.IsValid() && BlendInSettings.BlendMode == EMontageBlendMode::Inertialization)
+	{
+		const float InertialBlendDuration = BlendInArgs.BlendTime;
+		// Request new inertialization for new montage's group name
+		// If there is an existing inertialization request, we overwrite that here.
+		AnimInstance->RequestMontageInertialization(Montage, InertialBlendDuration);
+
+		// When using inertialization, we need to instantly blend in.
+		BlendInArgs.BlendTime = 0.0f;
+	}
+
 	// set blend option
 	float CurrentWeight = Blend.GetBlendedValue();
-	InitializeBlend(BlendIn);	
+	InitializeBlend(FAlphaBlend(BlendInArgs));
 	BlendStartAlpha = Blend.GetAlpha();
-	Blend.SetBlendTime(BlendIn.GetBlendTime() * DefaultBlendTimeMultiplier);
+	Blend.SetBlendTime(BlendInArgs.BlendTime * DefaultBlendTimeMultiplier);
 	Blend.SetValueRange(CurrentWeight, 1.f);
 	bEnableAutoBlendOut = Montage->bEnableAutoBlendOut;
 
-	ActiveBlendProfile = Montage->BlendProfileIn;
+	ActiveBlendProfile = BlendInSettings.BlendProfile;
 }
 
 void FAnimMontageInstance::InitializeBlend(const FAlphaBlend& InAlphaBlend)
@@ -1371,7 +1395,7 @@ void FAnimMontageInstance::InitializeBlend(const FAlphaBlend& InAlphaBlend)
 	Blend.SetBlendTime(InAlphaBlend.GetBlendTime());
 }
 
-void FAnimMontageInstance::Stop(const FAlphaBlend& InBlendOut, bool bInterrupt)
+void FAnimMontageInstance::Stop(const FMontageBlendSettings& InBlendOutSettings, bool bInterrupt)
 {
 	if (Montage)
 	{
@@ -1389,16 +1413,21 @@ void FAnimMontageInstance::Stop(const FAlphaBlend& InBlendOut, bool bInterrupt)
 	// if it hasn't stopped, stop now
 	if (IsStopped() == false)
 	{
-		// If we are using Inertial Blend, instantly stop the montage.
-		const bool bShouldInertialize = Montage && (Montage->BlendMode == EMontageBlendMode::Inertialization);
+		// If we are using Inertial Blend, blend time should be 0 to instantly stop the montage.
+		FAlphaBlendArgs BlendOutArgs = InBlendOutSettings.Blend;
+		const bool bShouldInertialize = InBlendOutSettings.BlendMode == EMontageBlendMode::Inertialization;
+		BlendOutArgs.BlendTime = bShouldInertialize ? 0.0f : BlendOutArgs.BlendTime;
 
 		// do not use default Montage->BlendOut 
-		// depending on situation, the BlendOut time can change 
-		InitializeBlend(bShouldInertialize ? (FAlphaBlend(InBlendOut, 0.0f)) : InBlendOut);
+		// depending on situation, the BlendOut time can change
+		InitializeBlend(FAlphaBlend(BlendOutArgs));
 		BlendStartAlpha = Blend.GetAlpha();
 		Blend.SetDesiredValue(0.f);
 		Blend.Update(0.0f);
 
+		// Only change the active blend profile if the montage isn't stopped. This is to prevent pops on a sudden blend profile switch
+		ActiveBlendProfile = InBlendOutSettings.BlendProfile;
+
 		if (Montage)
 		{
 			if (UAnimInstance* Inst = AnimInstance.Get())
@@ -1410,7 +1439,7 @@ void FAnimMontageInstance::Stop(const FAlphaBlend& InBlendOut, bool bInterrupt)
 				if (bShouldInertialize)
 				{
 					// Send the inertial blend request to the anim instance
-					Inst->RequestMontageInertialization(Montage, InBlendOut.GetBlendTime());
+					Inst->RequestMontageInertialization(Montage, InBlendOutSettings.Blend.BlendTime);
 				}
 			}
 		}
@@ -1420,12 +1449,12 @@ void FAnimMontageInstance::Stop(const FAlphaBlend& InBlendOut, bool bInterrupt)
 		// it is already stopped, but new montage blendtime is shorter than what 
 		// I'm blending out, that means this needs to readjust blendtime
 		// that way we don't accumulate old longer blendtime for newer montage to play
-		if (InBlendOut.GetBlendTime() < Blend.GetBlendTime())
+		if (InBlendOutSettings.Blend.BlendTime < Blend.GetBlendTime())
 		{
 			// I don't know if also using inBlendOut is better than
 			// currently set up blend option, but it might be worse to switch between 
 			// blending out, but it is possible options in the future
-			Blend.SetBlendTime(InBlendOut.GetBlendTime());
+			Blend.SetBlendTime(InBlendOutSettings.Blend.BlendTime);
 			BlendStartAlpha = Blend.GetAlpha();
 			// have to call this again to restart blending with new blend time
 			// we don't change blend options
@@ -1444,12 +1473,26 @@ void FAnimMontageInstance::Stop(const FAlphaBlend& InBlendOut, bool bInterrupt)
 
 	if (Montage != nullptr)
 	{
-		ActiveBlendProfile = Montage->BlendProfileOut;
 		UE_LOG(LogAnimMontage, Verbose, TEXT("Montage.Stop After: AnimMontage: %s,  (DesiredWeight:%0.2f, Weight:%0.2f)"),
 			*Montage->GetName(), GetDesiredWeight(), GetWeight());
 	}
 }
 
+void FAnimMontageInstance::Stop(const FAlphaBlend& InBlendOut, bool bInterrupt/*=true*/)
+{
+	FMontageBlendSettings BlendOutSettings;
+	BlendOutSettings.Blend = InBlendOut;
+
+	// Fill our other settings from the montage asset
+	if (Montage)
+	{
+		BlendOutSettings.BlendMode = Montage->BlendModeOut;
+		BlendOutSettings.BlendProfile = Montage->BlendProfileOut;
+	}
+
+	Stop(BlendOutSettings, bInterrupt);
+}
+
 void FAnimMontageInstance::Pause()
 {
 	bPlaying = false;
@@ -3035,3 +3078,11 @@ void UAnimMontage::BakeTimeStretchCurve()
 	TimeStretchCurve.BakeFromFloatCurve(*TimeStretchFloatCurve, SequenceLength);
 	PRAGMA_ENABLE_DEPRECATION_WARNINGS
 }
+
+FMontageBlendSettings::FMontageBlendSettings()
+	: BlendProfile(nullptr)
+	, Blend(FAlphaBlendArgs())
+	, BlendMode(EMontageBlendMode::Standard)
+{
+
+}
diff --git a/Engine/Source/Runtime/Engine/Public/AlphaBlend.h b/Engine/Source/Runtime/Engine/Public/AlphaBlend.h
--- a/Engine/Source/Runtime/Engine/Public/AlphaBlend.h
+++ b/Engine/Source/Runtime/Engine/Public/AlphaBlend.h
@@ -44,7 +44,7 @@ enum class EAlphaBlendOption : uint8
 };
 
 /**
- * Alpha Blend construction arguments. Used only for creation of an AlphaBlend.
+ * Alpha Blend construction arguments. Used for creation of an AlphaBlend.
  */
 USTRUCT(BlueprintType)
 struct ENGINE_API FAlphaBlendArgs
@@ -52,6 +52,7 @@ struct ENGINE_API FAlphaBlendArgs
 	GENERATED_BODY()
 
 	FAlphaBlendArgs();
+	FAlphaBlendArgs(const struct FAlphaBlend& InAlphaBlend);
 
 	/** If you're using Custom BlendOption, you can specify curve */
 	UPROPERTY(EditAnywhere, BlueprintReadWrite, Category = "Blend", meta=(DisplayAfter="BlendOption"))
