World Partition Actor Descriptors Handles
- Handle/Reference are now interchangeable types, that can be compared against or converted to.
- Added unit-tests that covers most of the standard Handle/Reference supported operations.

diff --git a/Engine/Source/Runtime/Engine/Private/Tests/WorldPartition/WorldPartitionHandleTests.cpp b/Engine/Source/Runtime/Engine/Private/Tests/WorldPartition/WorldPartitionHandleTests.cpp
new file mode 100644
--- /dev/null
+++ b/Engine/Source/Runtime/Engine/Private/Tests/WorldPartition/WorldPartitionHandleTests.cpp
@@ -0,0 +1,145 @@
+// Copyright Epic Games, Inc. All Rights Reserved.
+
+#include "Misc/AutomationTest.h"
+#include "WorldPartition/WorldPartitionHandle.h"
+#include "WorldPartition/WorldPartitionActorDesc.h"
+
+#if WITH_DEV_AUTOMATION_TESTS
+
+namespace WorldPartitionTests
+{
+	#define TEST_NAME_ROOT "System.Engine.WorldPartition"
+	constexpr const uint32 TestFlags = EAutomationTestFlags::ApplicationContextMask | EAutomationTestFlags::EngineFilter;
+
+	IMPLEMENT_SIMPLE_AUTOMATION_TEST(FWorldPartitionHandleTest, TEST_NAME_ROOT ".Handle", TestFlags)
+	bool FWorldPartitionHandleTest::RunTest(const FString& Parameters)
+	{
+#if WITH_EDITOR
+		TUniquePtr<FWorldPartitionActorDesc> ActorDescPtr1(new FWorldPartitionActorDesc());
+		TUniquePtr<FWorldPartitionActorDesc> ActorDescPtr2(new FWorldPartitionActorDesc());
+
+		FWorldPartitionHandle Handle(&ActorDescPtr1);
+		FWorldPartitionReference Reference(&ActorDescPtr2);
+
+		TestTrue(TEXT("Handle/Reference equality"), Handle != Reference);
+		TestTrue(TEXT("Reference/Handle equality"), Reference != Handle);
+		
+		TestTrue(TEXT("Handle soft refcount"), Handle->GetSoftRefCount() == 1);
+		TestTrue(TEXT("Handle hard refcount"), Handle->GetHardRefCount() == 0);
+		TestTrue(TEXT("Reference soft refcount"), Reference->GetSoftRefCount() == 0);
+		TestTrue(TEXT("Reference hard refcount"), Reference->GetHardRefCount() == 1);
+
+		// Conversions
+		{
+			FWorldPartitionHandle HandleToReference = Reference;
+			TestTrue(TEXT("Handle/Reference equality"), HandleToReference == Reference);
+			TestTrue(TEXT("Reference/Handle equality"), Reference == HandleToReference);
+			TestTrue(TEXT("Reference soft refcount"), Reference->GetSoftRefCount() == 1);
+			TestTrue(TEXT("Reference hard refcount"), Reference->GetHardRefCount() == 1);
+
+			FWorldPartitionReference ReferenceToHandle = Handle;
+			TestTrue(TEXT("Handle/Reference equality"), ReferenceToHandle == Handle);
+			TestTrue(TEXT("Handle/Reference equality"), Handle == ReferenceToHandle);
+			TestTrue(TEXT("Reference soft refcount"), Reference->GetSoftRefCount() == 1);
+			TestTrue(TEXT("Reference hard refcount"), Reference->GetHardRefCount() == 1);
+		}
+
+		TestTrue(TEXT("Handle soft refcount"), Handle->GetSoftRefCount() == 1);
+		TestTrue(TEXT("Handle hard refcount"), Handle->GetHardRefCount() == 0);
+		TestTrue(TEXT("Reference soft refcount"), Reference->GetSoftRefCount() == 0);
+		TestTrue(TEXT("Reference hard refcount"), Reference->GetHardRefCount() == 1);
+
+		Handle = Handle;
+		TestTrue(TEXT("Handle soft refcount"), Handle->GetSoftRefCount() == 1);
+		TestTrue(TEXT("Handle hard refcount"), Handle->GetHardRefCount() == 0);
+
+		Reference = Reference;
+		TestTrue(TEXT("Reference soft refcount"), Reference->GetSoftRefCount() == 0);
+		TestTrue(TEXT("Reference hard refcount"), Reference->GetHardRefCount() == 1);
+
+		// inplace new test
+		{
+			uint8 Buffer[sizeof(FWorldPartitionHandle)];
+			FWorldPartitionHandle* HandlePtr = new (Buffer) FWorldPartitionHandle(Reference);
+
+			TestTrue(TEXT("Handle array soft refcount"), Reference->GetSoftRefCount() == 1);
+			TestTrue(TEXT("Handle array hard refcount"), Reference->GetHardRefCount() == 1);
+
+			HandlePtr->~FWorldPartitionHandle();
+
+			TestTrue(TEXT("Handle array soft refcount"), Reference->GetSoftRefCount() == 0);
+			TestTrue(TEXT("Handle array hard refcount"), Reference->GetHardRefCount() == 1);
+		}
+
+		// TArray test
+		{
+			TArray<FWorldPartitionHandle> HandleList;
+			HandleList.Add(Handle);
+
+			TestTrue(TEXT("Handle array soft refcount"), Handle->GetSoftRefCount() == 2);
+			TestTrue(TEXT("Handle array hard refcount"), Handle->GetHardRefCount() == 0);
+
+			FWorldPartitionReference ReferenceToHandle = Handle;
+			TestTrue(TEXT("Handle/Reference equality"), ReferenceToHandle == Handle);
+			TestTrue(TEXT("Handle/Reference equality"), Handle == ReferenceToHandle);
+			TestTrue(TEXT("Handle soft refcount"), Handle->GetSoftRefCount() == 2);
+			TestTrue(TEXT("Handle hard refcount"), Handle->GetHardRefCount() == 1);
+
+			TestTrue(TEXT("Handle array contains handle"), HandleList.Contains(Handle));
+			TestTrue(TEXT("Handle array contains reference"), HandleList.Contains(ReferenceToHandle));
+
+			HandleList.Add(Reference);
+			
+			TestTrue(TEXT("Handle array contains reference"), HandleList.Contains(Reference));
+			TestTrue(TEXT("Handle array soft refcount"), Reference->GetSoftRefCount() == 1);
+			TestTrue(TEXT("Handle array hard refcount"), Reference->GetHardRefCount() == 1);
+			
+			HandleList.Remove(Handle);
+			TestTrue(TEXT("Handle array soft refcount"), Handle->GetSoftRefCount() == 1);
+			TestTrue(TEXT("Handle array hard refcount"), Handle->GetHardRefCount() == 1);
+
+			HandleList.Remove(Reference);
+			TestTrue(TEXT("Handle array soft refcount"), Reference->GetSoftRefCount() == 0);
+			TestTrue(TEXT("Handle array hard refcount"), Reference->GetHardRefCount() == 1);
+		}
+
+		TestTrue(TEXT("Handle soft refcount"), Handle->GetSoftRefCount() == 1);
+		TestTrue(TEXT("Handle hard refcount"), Handle->GetHardRefCount() == 0);
+		TestTrue(TEXT("Reference soft refcount"), Reference->GetSoftRefCount() == 0);
+		TestTrue(TEXT("Reference hard refcount"), Reference->GetHardRefCount() == 1);
+
+		// TSet test
+		{
+			TSet<FWorldPartitionHandle> HandleSet;
+			HandleSet.Add(Handle);
+
+			TestTrue(TEXT("Handle set soft refcount"), Handle->GetSoftRefCount() == 2);
+			TestTrue(TEXT("Handle set hard refcount"), Handle->GetHardRefCount() == 0);
+
+			FWorldPartitionReference ReferenceToHandle = Handle;
+			TestTrue(TEXT("Handle/Reference equality"), ReferenceToHandle == Handle);
+			TestTrue(TEXT("Handle/Reference equality"), Handle == ReferenceToHandle);
+			TestTrue(TEXT("Handle soft refcount"), Handle->GetSoftRefCount() == 2);
+			TestTrue(TEXT("Handle hard refcount"), Handle->GetHardRefCount() == 1);
+
+			TestTrue(TEXT("Handle set contains handle"), HandleSet.Contains(Handle));
+			TestTrue(TEXT("Handle set contains reference"), HandleSet.Contains(ReferenceToHandle));
+
+			HandleSet.Add(Reference);
+			TestTrue(TEXT("Handle set contains reference"), HandleSet.Contains(Reference));
+
+			TestTrue(TEXT("Reference soft refcount"), Reference->GetSoftRefCount() == 1);
+			TestTrue(TEXT("Reference hard refcount"), Reference->GetHardRefCount() == 1);
+		}
+
+		TestTrue(TEXT("Handle soft refcount"), Handle->GetSoftRefCount() == 1);
+		TestTrue(TEXT("Handle hard refcount"), Handle->GetHardRefCount() == 0);
+		TestTrue(TEXT("Reference soft refcount"), Reference->GetSoftRefCount() == 0);
+		TestTrue(TEXT("Reference hard refcount"), Reference->GetHardRefCount() == 1);
+#endif
+
+		return true;
+	}
+}
+
+#endif 
diff --git a/Engine/Source/Runtime/Engine/Private/WorldPartition/WorldPartitionHandle.cpp b/Engine/Source/Runtime/Engine/Private/WorldPartition/WorldPartitionHandle.cpp
--- a/Engine/Source/Runtime/Engine/Private/WorldPartition/WorldPartitionHandle.cpp
+++ b/Engine/Source/Runtime/Engine/Private/WorldPartition/WorldPartitionHandle.cpp
@@ -1,8 +1,19 @@
 // Copyright Epic Games, Inc. All Rights Reserved.
 
 #include "WorldPartition/WorldPartitionHandle.h"
+#include "WorldPArtition/WorldPartition.h"
 
 #if WITH_EDITOR
+TUniquePtr<FWorldPartitionActorDesc>* FWorldPartitionHandleBase::GetActorDesc(UWorldPartition* WorldPartition, const FGuid& ActorGuid)
+{
+	if (TUniquePtr<FWorldPartitionActorDesc>** ActorDescPtr = WorldPartition->Actors.Find(ActorGuid))
+	{
+		return *ActorDescPtr;
+	}
+
+	return nullptr;
+}
+
 void FWorldPartitionHandleImpl::IncRefCount(FWorldPartitionActorDesc* ActorDesc)
 {
 	ActorDesc->IncSoftRefCount();
@@ -18,7 +29,12 @@ void FWorldPartitionReferenceImpl::IncRefCount(FWorldPartitionActorDesc* ActorDe
 	if (ActorDesc->IncHardRefCount() == 1)
 	{
 		AActor* Actor = ActorDesc->Load();
-		Actor->GetLevel()->AddLoadedActor(Actor);
+		check(Actor || GIsAutomationTesting);
+
+		if (Actor)
+		{
+			Actor->GetLevel()->AddLoadedActor(Actor);
+		}
 	}
 }
 
diff --git a/Engine/Source/Runtime/Engine/Public/WorldPartition/WorldPartitionActorDesc.h b/Engine/Source/Runtime/Engine/Public/WorldPartition/WorldPartitionActorDesc.h
--- a/Engine/Source/Runtime/Engine/Public/WorldPartition/WorldPartitionActorDesc.h
+++ b/Engine/Source/Runtime/Engine/Public/WorldPartition/WorldPartitionActorDesc.h
@@ -20,6 +20,13 @@ struct FWorldPartitionActorDescInitData
 };
 #endif
 
+#if WITH_DEV_AUTOMATION_TESTS
+namespace WorldPartitionTests
+{
+	class FWorldPartitionHandleTest;
+}
+#endif
+
 /**
  * Represents a potentially unloaded actor (editor-only)
  */
@@ -29,6 +36,10 @@ class ENGINE_API FWorldPartitionActorDesc
 	friend class AActor;
 	friend class UWorldPartition;
 
+#if WITH_DEV_AUTOMATION_TESTS
+	friend class WorldPartitionTests::FWorldPartitionHandleTest;
+#endif
+
 public:
 	virtual ~FWorldPartitionActorDesc() {}
 
diff --git a/Engine/Source/Runtime/Engine/Public/WorldPartition/WorldPartitionHandle.h b/Engine/Source/Runtime/Engine/Public/WorldPartition/WorldPartitionHandle.h
--- a/Engine/Source/Runtime/Engine/Public/WorldPartition/WorldPartitionHandle.h
+++ b/Engine/Source/Runtime/Engine/Public/WorldPartition/WorldPartitionHandle.h
@@ -3,11 +3,19 @@
 
 #include "CoreMinimal.h"
 #include "Templates/UniquePtr.h"
-#include "WorldPartition/WorldPartition.h"
 
 #if WITH_EDITOR
+class UWorldPartition;
+class FWorldPartitionActorDesc;
+
+class ENGINE_API FWorldPartitionHandleBase
+{
+public:
+	static TUniquePtr<FWorldPartitionActorDesc>* GetActorDesc(UWorldPartition* WorldPartition, const FGuid& ActorGuid);
+};
+
 template <typename Impl>
-class TWorldPartitionHandleBase
+class ENGINE_API TWorldPartitionHandleBase : public FWorldPartitionHandleBase
 {
 	friend struct FWorldPartitionHandleHelpers;
 
@@ -16,40 +24,73 @@ public:
 		: ActorDesc(nullptr)
 	{}
 
-	TWorldPartitionHandleBase(UWorldPartition* WorldPartition, const FGuid& ActorGuid)
+	TWorldPartitionHandleBase(TUniquePtr<FWorldPartitionActorDesc>* InActorDesc)
 	{
-		if (TUniquePtr<FWorldPartitionActorDesc>** ActorDescPtr = WorldPartition->Actors.Find(ActorGuid))
-		{
-			ActorDesc = *ActorDescPtr;
+		ActorDesc = InActorDesc;
 
-			if (IsValid())
-			{
-				IncRefCount();
-			}
-		}
-		else
+		if (IsValid())
 		{
-			ActorDesc = nullptr;
+			IncRefCount();
 		}
 	}
 
-	~TWorldPartitionHandleBase()
+	TWorldPartitionHandleBase(UWorldPartition* WorldPartition, const FGuid& ActorGuid)
 	{
+		ActorDesc = GetActorDesc(WorldPartition, ActorGuid);
+
 		if (IsValid())
 		{
-			DecRefCount();
+			IncRefCount();
 		}
 	}
 
-	template <typename T>
 	TWorldPartitionHandleBase(const TWorldPartitionHandleBase& Other)
+		: ActorDesc(nullptr)
+	{
+		*this = Other;
+	}
+
+	TWorldPartitionHandleBase(TWorldPartitionHandleBase&& Other)
+		: ActorDesc(nullptr)
+	{
+		*this = Other;
+	}
+
+	// Conversions
+	template <typename T>
+	TWorldPartitionHandleBase<Impl>(const TWorldPartitionHandleBase<T>& Other)
+		: ActorDesc(nullptr)
 	{
 		*this = Other;
 	}
 
 	template <typename T>
-	TWorldPartitionHandleBase& operator=(const T& Other)
+	TWorldPartitionHandleBase<Impl>(TWorldPartitionHandleBase<T>&& Other)
+		: ActorDesc(nullptr)
+	{
+		*this = Other;
+	}
+
+	~TWorldPartitionHandleBase()
+	{
+		if (IsValid())
+		{
+			DecRefCount();
+		}
+	}
+
+	TWorldPartitionHandleBase& operator=(const TWorldPartitionHandleBase& Other)
 	{
+		if ((void*)this == (void*)&Other)
+		{
+			return *this;
+		}
+
+		if (IsValid())
+		{
+			DecRefCount();
+		}
+
 		ActorDesc = Other.ActorDesc;
 
 		if (IsValid())
@@ -60,15 +101,67 @@ public:
 		return *this;
 	}
 
+	TWorldPartitionHandleBase<Impl>& operator=(TWorldPartitionHandleBase&& Other)
+	{
+		if ((void*)this == (void*)&Other)
+		{
+			return *this;
+		}
+
+		if (IsValid())
+		{
+			DecRefCount();
+		}
+
+		ActorDesc = Other.ActorDesc;
+
+		if (IsValid())
+		{
+			IncRefCount();
+			Other.DecRefCount();
+			Other.ActorDesc = nullptr;
+		}
+
+		return *this;
+	}
+
+	// Conversions
 	template <typename T>
-	TWorldPartitionHandleBase(T&& Other)
+	TWorldPartitionHandleBase<Impl>& operator=(const TWorldPartitionHandleBase<T>& Other)
 	{
-		*this = Other;
+		if ((void*)this == (void*)&Other)
+		{
+			return *this;
+		}
+
+		if (IsValid())
+		{
+			DecRefCount();
+		}
+
+		ActorDesc = Other.ActorDesc;
+
+		if (IsValid())
+		{
+			IncRefCount();
+		}
+
+		return *this;
 	}
 
 	template <typename T>
-	TWorldPartitionHandleBase& operator=(T&& Other)
+	TWorldPartitionHandleBase<Impl>& operator=(TWorldPartitionHandleBase<T>&& Other)
 	{
+		if ((void*)this == (void*)&Other)
+		{
+			return *this;
+		}
+
+		if (IsValid())
+		{
+			DecRefCount();
+		}
+
 		ActorDesc = Other.ActorDesc;
 
 		if (IsValid())
@@ -86,6 +179,11 @@ public:
 		return Get();
 	}
 
+	inline FWorldPartitionActorDesc* operator*() const
+	{
+		return Get();
+	}
+
 	inline bool IsValid() const
 	{
 		return ActorDesc && ActorDesc->IsValid();
@@ -101,16 +199,28 @@ public:
 		return ::PointerHash(HandleBase.ActorDesc);
 	}
 
-	friend bool operator==(const TWorldPartitionHandleBase<Impl>& Lhs, const TWorldPartitionHandleBase<Impl>& Rhs)
+	inline bool operator==(const TWorldPartitionHandleBase& Other) const
 	{
-		return Lhs.ActorDesc == Rhs.ActorDesc;
+		return ActorDesc == Other.ActorDesc;
 	}
 
-	friend bool operator!=(const TWorldPartitionHandleBase<Impl>& Lhs, const TWorldPartitionHandleBase<Impl>& Rhs)
+	inline bool operator!=(const TWorldPartitionHandleBase& Other) const
 	{
-		return !(Lhs == Rhs);
+		return !(*this == Other);
 	}
 
+	// Conversions
+	template <typename T>
+	inline bool operator==(const TWorldPartitionHandleBase<T>& Other) const
+	{
+		return Get() == *Other;
+	}
+
+	template <typename T>
+	inline bool operator!=(const TWorldPartitionHandleBase<T>& Other) const
+	{
+		return !(*this == Other);
+	}
 
 protected:
 	inline void IncRefCount()
@@ -123,6 +233,7 @@ protected:
 		Impl::DecRefCount(ActorDesc->Get());
 	}
 
+public:
 	TUniquePtr<FWorldPartitionActorDesc>* ActorDesc;
 };
 
