Undefine TEST_NAME_ROOT to fix unity builds.

diff --git a/Engine/Source/Runtime/Engine/Private/Tests/WorldPartition/WorldPartitionHandleTests.cpp b/Engine/Source/Runtime/Engine/Private/Tests/WorldPartition/WorldPartitionHandleTests.cpp
--- a/Engine/Source/Runtime/Engine/Private/Tests/WorldPartition/WorldPartitionHandleTests.cpp
+++ b/Engine/Source/Runtime/Engine/Private/Tests/WorldPartition/WorldPartitionHandleTests.cpp
@@ -6,9 +6,10 @@
 
 #if WITH_DEV_AUTOMATION_TESTS
 
+#define TEST_NAME_ROOT "System.Engine.WorldPartition"
+
 namespace WorldPartitionTests
 {
-	#define TEST_NAME_ROOT "System.Engine.WorldPartition"
 	constexpr const uint32 TestFlags = EAutomationTestFlags::ApplicationContextMask | EAutomationTestFlags::EngineFilter;
 
 	IMPLEMENT_SIMPLE_AUTOMATION_TEST(FWorldPartitionHandleTest, TEST_NAME_ROOT ".Handle", TestFlags)
@@ -137,9 +138,10 @@ namespace WorldPartitionTests
 		TestTrue(TEXT("Reference soft refcount"), Reference->GetSoftRefCount() == 0);
 		TestTrue(TEXT("Reference hard refcount"), Reference->GetHardRefCount() == 1);
 #endif
-
 		return true;
 	}
 }
 
+#undef TEST_NAME_ROOT
+
 #endif 
