Fix incremental  build.

diff --git a/Engine/Source/Runtime/Engine/Private/Tests/WorldPartition/WorldPartitionHandleTests.cpp b/Engine/Source/Runtime/Engine/Private/Tests/WorldPartition/WorldPartitionHandleTests.cpp
--- a/Engine/Source/Runtime/Engine/Private/Tests/WorldPartition/WorldPartitionHandleTests.cpp
+++ b/Engine/Source/Runtime/Engine/Private/Tests/WorldPartition/WorldPartitionHandleTests.cpp
@@ -50,14 +50,6 @@ namespace WorldPartitionTests
 		TestTrue(TEXT("Reference soft refcount"), Reference->GetSoftRefCount() == 0);
 		TestTrue(TEXT("Reference hard refcount"), Reference->GetHardRefCount() == 1);
 
-		Handle = Handle;
-		TestTrue(TEXT("Handle soft refcount"), Handle->GetSoftRefCount() == 1);
-		TestTrue(TEXT("Handle hard refcount"), Handle->GetHardRefCount() == 0);
-
-		Reference = Reference;
-		TestTrue(TEXT("Reference soft refcount"), Reference->GetSoftRefCount() == 0);
-		TestTrue(TEXT("Reference hard refcount"), Reference->GetHardRefCount() == 1);
-
 		// inplace new test
 		{
 			uint8 Buffer[sizeof(FWorldPartitionHandle)];
