Fix include case mismatch.

diff --git a/Engine/Source/Runtime/Engine/Private/WorldPartition/WorldPartitionHandle.cpp b/Engine/Source/Runtime/Engine/Private/WorldPartition/WorldPartitionHandle.cpp
--- a/Engine/Source/Runtime/Engine/Private/WorldPartition/WorldPartitionHandle.cpp
+++ b/Engine/Source/Runtime/Engine/Private/WorldPartition/WorldPartitionHandle.cpp
@@ -1,7 +1,7 @@
 // Copyright Epic Games, Inc. All Rights Reserved.
 
 #include "WorldPartition/WorldPartitionHandle.h"
-#include "WorldPArtition/WorldPartition.h"
+#include "WorldPartition/WorldPartition.h"
 
 #if WITH_EDITOR
 TUniquePtr<FWorldPartitionActorDesc>* FWorldPartitionHandleBase::GetActorDesc(UWorldPartition* WorldPartition, const FGuid& ActorGuid)
