Force inline functions in TWorldPartitionHandleBase to fix link errors in debug builds.

diff --git a/Engine/Source/Runtime/Engine/Public/WorldPartition/WorldPartitionHandle.h b/Engine/Source/Runtime/Engine/Public/WorldPartition/WorldPartitionHandle.h
--- a/Engine/Source/Runtime/Engine/Public/WorldPartition/WorldPartitionHandle.h
+++ b/Engine/Source/Runtime/Engine/Public/WorldPartition/WorldPartitionHandle.h
@@ -20,11 +20,11 @@ class ENGINE_API TWorldPartitionHandleBase : public FWorldPartitionHandleBase
 	friend struct FWorldPartitionHandleHelpers;
 
 public:
-	TWorldPartitionHandleBase()
+	FORCEINLINE TWorldPartitionHandleBase()
 		: ActorDesc(nullptr)
 	{}
 
-	TWorldPartitionHandleBase(TUniquePtr<FWorldPartitionActorDesc>* InActorDesc)
+	FORCEINLINE TWorldPartitionHandleBase(TUniquePtr<FWorldPartitionActorDesc>* InActorDesc)
 	{
 		ActorDesc = InActorDesc;
 
@@ -34,7 +34,7 @@ public:
 		}
 	}
 
-	TWorldPartitionHandleBase(UWorldPartition* WorldPartition, const FGuid& ActorGuid)
+	FORCEINLINE TWorldPartitionHandleBase(UWorldPartition* WorldPartition, const FGuid& ActorGuid)
 	{
 		ActorDesc = GetActorDesc(WorldPartition, ActorGuid);
 
@@ -44,13 +44,13 @@ public:
 		}
 	}
 
-	TWorldPartitionHandleBase(const TWorldPartitionHandleBase& Other)
+	FORCEINLINE TWorldPartitionHandleBase(const TWorldPartitionHandleBase& Other)
 		: ActorDesc(nullptr)
 	{
 		*this = Other;
 	}
 
-	TWorldPartitionHandleBase(TWorldPartitionHandleBase&& Other)
+	FORCEINLINE TWorldPartitionHandleBase(TWorldPartitionHandleBase&& Other)
 		: ActorDesc(nullptr)
 	{
 		*this = Other;
@@ -58,20 +58,20 @@ public:
 
 	// Conversions
 	template <typename T>
-	TWorldPartitionHandleBase<Impl>(const TWorldPartitionHandleBase<T>& Other)
+	FORCEINLINE TWorldPartitionHandleBase<Impl>(const TWorldPartitionHandleBase<T>& Other)
 		: ActorDesc(nullptr)
 	{
 		*this = Other;
 	}
 
 	template <typename T>
-	TWorldPartitionHandleBase<Impl>(TWorldPartitionHandleBase<T>&& Other)
+	FORCEINLINE TWorldPartitionHandleBase<Impl>(TWorldPartitionHandleBase<T>&& Other)
 		: ActorDesc(nullptr)
 	{
 		*this = Other;
 	}
 
-	~TWorldPartitionHandleBase()
+	FORCEINLINE ~TWorldPartitionHandleBase()
 	{
 		if (IsValid())
 		{
@@ -79,7 +79,7 @@ public:
 		}
 	}
 
-	TWorldPartitionHandleBase& operator=(const TWorldPartitionHandleBase& Other)
+	FORCEINLINE TWorldPartitionHandleBase& operator=(const TWorldPartitionHandleBase& Other)
 	{
 		if ((void*)this == (void*)&Other)
 		{
@@ -101,7 +101,7 @@ public:
 		return *this;
 	}
 
-	TWorldPartitionHandleBase<Impl>& operator=(TWorldPartitionHandleBase&& Other)
+	FORCEINLINE TWorldPartitionHandleBase<Impl>& operator=(TWorldPartitionHandleBase&& Other)
 	{
 		if ((void*)this == (void*)&Other)
 		{
@@ -127,7 +127,7 @@ public:
 
 	// Conversions
 	template <typename T>
-	TWorldPartitionHandleBase<Impl>& operator=(const TWorldPartitionHandleBase<T>& Other)
+	FORCEINLINE TWorldPartitionHandleBase<Impl>& operator=(const TWorldPartitionHandleBase<T>& Other)
 	{
 		if ((void*)this == (void*)&Other)
 		{
@@ -150,7 +150,7 @@ public:
 	}
 
 	template <typename T>
-	TWorldPartitionHandleBase<Impl>& operator=(TWorldPartitionHandleBase<T>&& Other)
+	FORCEINLINE TWorldPartitionHandleBase<Impl>& operator=(TWorldPartitionHandleBase<T>&& Other)
 	{
 		if ((void*)this == (void*)&Other)
 		{
@@ -174,61 +174,61 @@ public:
 		return *this;
 	}
 
-	inline FWorldPartitionActorDesc* operator->() const
+	FORCEINLINE FWorldPartitionActorDesc* operator->() const
 	{
 		return Get();
 	}
 
-	inline FWorldPartitionActorDesc* operator*() const
+	FORCEINLINE FWorldPartitionActorDesc* operator*() const
 	{
 		return Get();
 	}
 
-	inline bool IsValid() const
+	FORCEINLINE bool IsValid() const
 	{
 		return ActorDesc && ActorDesc->IsValid();
 	}
 
-	inline FWorldPartitionActorDesc* Get() const
+	FORCEINLINE FWorldPartitionActorDesc* Get() const
 	{
 		return IsValid() ? ActorDesc->Get() : nullptr;
 	}
 		
-	friend uint32 GetTypeHash(const TWorldPartitionHandleBase<Impl>& HandleBase)
+	friend FORCEINLINE uint32 GetTypeHash(const TWorldPartitionHandleBase<Impl>& HandleBase)
 	{
 		return ::PointerHash(HandleBase.ActorDesc);
 	}
 
-	inline bool operator==(const TWorldPartitionHandleBase& Other) const
+	FORCEINLINE bool operator==(const TWorldPartitionHandleBase& Other) const
 	{
 		return ActorDesc == Other.ActorDesc;
 	}
 
-	inline bool operator!=(const TWorldPartitionHandleBase& Other) const
+	FORCEINLINE bool operator!=(const TWorldPartitionHandleBase& Other) const
 	{
 		return !(*this == Other);
 	}
 
 	// Conversions
 	template <typename T>
-	inline bool operator==(const TWorldPartitionHandleBase<T>& Other) const
+	FORCEINLINE bool operator==(const TWorldPartitionHandleBase<T>& Other) const
 	{
 		return Get() == *Other;
 	}
 
 	template <typename T>
-	inline bool operator!=(const TWorldPartitionHandleBase<T>& Other) const
+	FORCEINLINE bool operator!=(const TWorldPartitionHandleBase<T>& Other) const
 	{
 		return !(*this == Other);
 	}
 
 protected:
-	inline void IncRefCount()
+	FORCEINLINE void IncRefCount()
 	{
 		Impl::IncRefCount(ActorDesc->Get());
 	}
 
-	inline void DecRefCount()
+	FORCEINLINE void DecRefCount()
 	{
 		Impl::DecRefCount(ActorDesc->Get());
 	}
