Wait for staticmesh async compilation to end before starting a HLOD merge process, as they access their data.

diff --git a/Engine/Source/Runtime/Engine/Private/WorldPartition/WorldPartitionRuntimeSpatialHash.cpp b/Engine/Source/Runtime/Engine/Private/WorldPartition/WorldPartitionRuntimeSpatialHash.cpp
--- a/Engine/Source/Runtime/Engine/Private/WorldPartition/WorldPartitionRuntimeSpatialHash.cpp
+++ b/Engine/Source/Runtime/Engine/Private/WorldPartition/WorldPartitionRuntimeSpatialHash.cpp
@@ -49,6 +49,8 @@
 #include "Engine/WorldComposition.h"
 #include "LevelUtils.h"
 
+#include "StaticMeshCompiler.h"
+
 extern UNREALED_API class UEditorEngine* GEditor;
 #endif //WITH_EDITOR
 
@@ -1596,6 +1598,8 @@ TArray<FGuid> GenerateHLODsForGrid(UWorldPartition* WorldPartition, const FSpati
 					CellActors.Add(ActorDesc->GetActor());
 				}
 
+				FStaticMeshCompilingManager::Get().FinishAllCompilation();
+
 				if (HasExceededMaxMemory())
 				{
 					DoCollectGarbage();
