LevelInstance: Edit resets trans buffer

diff --git a/Engine/Source/Runtime/Engine/Private/LevelInstance/LevelInstanceSubsystem.cpp b/Engine/Source/Runtime/Engine/Private/LevelInstance/LevelInstanceSubsystem.cpp
--- a/Engine/Source/Runtime/Engine/Private/LevelInstance/LevelInstanceSubsystem.cpp
+++ b/Engine/Source/Runtime/Engine/Private/LevelInstance/LevelInstanceSubsystem.cpp
@@ -1286,6 +1286,9 @@ void ULevelInstanceSubsystem::EditLevelInstance(ALevelInstance* LevelInstanceAct
 			}
 		}
 	}
+
+	// Edit can't be undone
+	GEditor->ResetTransaction(LOCTEXT("LevelInstanceEditResetTrans", "Edit Level Instance"));
 }
 
 void ULevelInstanceSubsystem::CommitChildrenLevelInstances(ALevelInstance* LevelInstanceActor)
