World Partition actor descriptor handles: convert existing code to use automatic handle/reference conversions and remove custom conversion code.

diff --git a/Engine/Source/Runtime/Engine/Private/WorldPartition/HLOD/HLODBuilder.cpp b/Engine/Source/Runtime/Engine/Private/WorldPartition/HLOD/HLODBuilder.cpp
--- a/Engine/Source/Runtime/Engine/Private/WorldPartition/HLOD/HLODBuilder.cpp
+++ b/Engine/Source/Runtime/Engine/Private/WorldPartition/HLOD/HLODBuilder.cpp
@@ -44,7 +44,7 @@ public:
 		FWorldPartitionHandle HLODActorHandle;
 		if (Context->HLODActorDescs.RemoveAndCopyValue(CellHash, HLODActorHandle))
 		{
-			HLODActorRefIndex = Context->ActorReferences.Add(FWorldPartitionHandleHelpers::ConvertHandleToReference(HLODActorHandle));
+			HLODActorRefIndex = Context->ActorReferences.Add(HLODActorHandle);
 			HLODActor = CastChecked<AWorldPartitionHLOD>(HLODActorHandle->GetActor());
 		}
 
diff --git a/Engine/Source/Runtime/Engine/Private/WorldPartition/WorldPartitionHandle.cpp b/Engine/Source/Runtime/Engine/Private/WorldPartition/WorldPartitionHandle.cpp
--- a/Engine/Source/Runtime/Engine/Private/WorldPartition/WorldPartitionHandle.cpp
+++ b/Engine/Source/Runtime/Engine/Private/WorldPartition/WorldPartitionHandle.cpp
@@ -49,30 +49,4 @@ void FWorldPartitionReferenceImpl::DecRefCount(FWorldPartitionActorDesc* ActorDe
 		}
 	}
 }
-
-FWorldPartitionReference FWorldPartitionHandleHelpers::ConvertHandleToReference(const FWorldPartitionHandle& Handle)
-{
-	FWorldPartitionReference Reference;
-
-	if (Handle.IsValid())
-	{
-		Reference.ActorDesc = Handle.ActorDesc;
-		Reference.IncRefCount();
-	}
-
-	return Reference;
-}
-
-FWorldPartitionHandle FWorldPartitionHandleHelpers::ConvertReferenceToHandle(const FWorldPartitionReference& Reference)
-{
-	FWorldPartitionHandle Handle;
-
-	if (Reference.IsValid())
-	{
-		Handle.ActorDesc = Reference.ActorDesc;
-		Handle.IncRefCount();
-	}
-
-	return Handle;
-}
 #endif
\ No newline at end of file
diff --git a/Engine/Source/Runtime/Engine/Public/WorldPartition/WorldPartitionHandle.h b/Engine/Source/Runtime/Engine/Public/WorldPartition/WorldPartitionHandle.h
--- a/Engine/Source/Runtime/Engine/Public/WorldPartition/WorldPartitionHandle.h
+++ b/Engine/Source/Runtime/Engine/Public/WorldPartition/WorldPartitionHandle.h
@@ -251,10 +251,4 @@ struct ENGINE_API FWorldPartitionReferenceImpl
 
 typedef TWorldPartitionHandleBase<FWorldPartitionHandleImpl> FWorldPartitionHandle;
 typedef TWorldPartitionHandleBase<FWorldPartitionReferenceImpl> FWorldPartitionReference;
-
-struct ENGINE_API FWorldPartitionHandleHelpers
-{
-	static FWorldPartitionReference ConvertHandleToReference(const FWorldPartitionHandle& Handle);
-	static FWorldPartitionHandle ConvertReferenceToHandle(const FWorldPartitionReference& Handle);
-};
 #endif
\ No newline at end of file
diff --git a/Engine/Source/Runtime/Landscape/Private/LandscapeEdit.cpp b/Engine/Source/Runtime/Landscape/Private/LandscapeEdit.cpp
--- a/Engine/Source/Runtime/Landscape/Private/LandscapeEdit.cpp
+++ b/Engine/Source/Runtime/Landscape/Private/LandscapeEdit.cpp
@@ -5039,7 +5039,7 @@ void ALandscapeStreamingProxy::PostRegisterAllComponents()
 				{
 					if (SplineHandle.IsValid() && Bounds.IntersectXY(SplineHandle->GetBounds()))
 					{
-						ActorDescReferences.Add(FWorldPartitionHandleHelpers::ConvertHandleToReference(SplineHandle));
+						ActorDescReferences.Add(SplineHandle);
 					}
 				}
 			}
