Make ADirectionalLight always loaded.

diff --git a/Engine/Source/Runtime/Engine/Classes/Engine/DirectionalLight.h b/Engine/Source/Runtime/Engine/Classes/Engine/DirectionalLight.h
--- a/Engine/Source/Runtime/Engine/Classes/Engine/DirectionalLight.h
+++ b/Engine/Source/Runtime/Engine/Classes/Engine/DirectionalLight.h
@@ -37,6 +37,7 @@ public:
 #if WITH_EDITOR
 	virtual void LoadedFromAnotherClass(const FName& OldClassName) override;
 	virtual void PostEditChangeProperty(FPropertyChangedEvent& PropertyChangedEvent) override;
+	virtual EActorGridPlacement GetDefaultGridPlacement() const override { return EActorGridPlacement::AlwaysLoaded; }
 #endif
 
 public:
