Don't try to unload actors when releasing actor descriptor handles during static exit.

diff --git a/Engine/Source/Runtime/Engine/Private/WorldPartition/WorldPartitionHandle.cpp b/Engine/Source/Runtime/Engine/Private/WorldPartition/WorldPartitionHandle.cpp
--- a/Engine/Source/Runtime/Engine/Private/WorldPartition/WorldPartitionHandle.cpp
+++ b/Engine/Source/Runtime/Engine/Private/WorldPartition/WorldPartitionHandle.cpp
@@ -30,6 +30,7 @@ void FWorldPartitionReferenceImpl::IncRefCount(FWorldPartitionActorDesc* ActorDe
 	{
 		AActor* Actor = ActorDesc->Load();
 		check(Actor || GIsAutomationTesting);
+		check(!IsEngineExitRequested());
 
 		if (Actor)
 		{
@@ -42,10 +43,13 @@ void FWorldPartitionReferenceImpl::DecRefCount(FWorldPartitionActorDesc* ActorDe
 {
 	if (!ActorDesc->DecHardRefCount())
 	{
-		if (AActor* Actor = ActorDesc->GetActor())
+		if (!IsEngineExitRequested())
 		{
-			Actor->GetLevel()->RemoveLoadedActor(Actor);
-			ActorDesc->Unload();
+			if (AActor* Actor = ActorDesc->GetActor())
+			{
+				Actor->GetLevel()->RemoveLoadedActor(Actor);
+				ActorDesc->Unload();
+			}
 		}
 	}
 }
