Fix build error

diff --git a/Engine/Source/Runtime/Engine/Public/Audio/SoundEffectPresetWidgetInterface.h b/Engine/Source/Runtime/Engine/Public/Audio/SoundEffectPresetWidgetInterface.h
--- a/Engine/Source/Runtime/Engine/Public/Audio/SoundEffectPresetWidgetInterface.h
+++ b/Engine/Source/Runtime/Engine/Public/Audio/SoundEffectPresetWidgetInterface.h
@@ -9,6 +9,7 @@
 #include "UObject/Interface.h"
 #include "UObject/ObjectMacros.h"
 #include "UObject/UObjectGlobals.h"
+#include "Templates/SubclassOf.h"
 
 #include "SoundEffectPresetWidgetInterface.generated.h"
 
