HLOD actors with no cell hash (generated with an older version) were never deleted

diff --git a/Engine/Source/Runtime/Engine/Private/WorldPartition/WorldPartitionRuntimeSpatialHash.cpp b/Engine/Source/Runtime/Engine/Private/WorldPartition/WorldPartitionRuntimeSpatialHash.cpp
--- a/Engine/Source/Runtime/Engine/Private/WorldPartition/WorldPartitionRuntimeSpatialHash.cpp
+++ b/Engine/Source/Runtime/Engine/Private/WorldPartition/WorldPartitionRuntimeSpatialHash.cpp
@@ -26,6 +26,7 @@
 #include "Misc/Parse.h"
 #include "Misc/ScopedSlowTask.h"
 #include "Algo/Find.h"
+#include "Algo/ForEach.h"
 #include "EngineUtils.h"
 
 #include "DistanceFieldAtlas.h"
@@ -1886,11 +1887,18 @@ bool UWorldPartitionRuntimeSpatialHash::GenerateHLOD(ISourceControlHelper* Sourc
 
 	// Create HLODs generation context
 	FHLODGenerationContext Context;
+	TArray<FWorldPartitionHandle> InvalidHLODActors;
 	for (TWorldPartitionActorDescIterator<AWorldPartitionHLOD, FHLODActorDesc> HLODIterator(WorldPartition); HLODIterator; ++HLODIterator)
 	{
+		FWorldPartitionHandle HLODActorHandle(WorldPartition, HLODIterator->GetGuid());
+
 		if (HLODIterator->GetCellHash())
 		{
-			Context.HLODActorDescs.Add(HLODIterator->GetCellHash(), FWorldPartitionHandle(WorldPartition, HLODIterator->GetGuid()));
+			Context.HLODActorDescs.Add(HLODIterator->GetCellHash(), HLODActorHandle);
+		}
+		else
+		{
+			InvalidHLODActors.Add(HLODActorHandle);
 		}
 	}
 
@@ -1966,16 +1974,24 @@ bool UWorldPartitionRuntimeSpatialHash::GenerateHLOD(ISourceControlHelper* Sourc
 		GenerateHLODs(HLODGrids[HLODGridName], GridsDepth[HLODGridName] + 1, HLODActorClusters);
 	}
 
-	// Destroy all unreferenced HLOD actors
-	for (const auto& HLODActorPair: Context.HLODActorDescs)
+	auto DeleteHLODActor = [&SourceControlHelper](FWorldPartitionHandle ActorHandle)
 	{
-		FWorldPartitionActorDesc* HLODActorDesc = HLODActorPair.Value.Get();
+		FWorldPartitionActorDesc* HLODActorDesc = ActorHandle.Get();
 		check(HLODActorDesc);
 
-		const FString ActorPackageFilename = SourceControlHelper->GetFilename(HLODActorDesc->GetActorPackage().ToString());		
+		const FString ActorPackageFilename = SourceControlHelper->GetFilename(HLODActorDesc->GetActorPackage().ToString());
 		SourceControlHelper->Delete(ActorPackageFilename);
+	};
+
+	// Destroy all unreferenced HLOD actors
+	for (const auto& HLODActorPair: Context.HLODActorDescs)
+	{
+		DeleteHLODActor(HLODActorPair.Value);
 	}
 
+	// Destroy all invalid HLOD actors
+	Algo::ForEach(InvalidHLODActors, DeleteHLODActor);
+
 	// Create/destroy HLOD grid actors
 	UpdateHLODGridsActors(GetWorld(), HLODGrids);
 
