Fixed static analysis warning.

diff --git a/Engine/Source/Runtime/Engine/Private/DataChannel.cpp b/Engine/Source/Runtime/Engine/Private/DataChannel.cpp
--- a/Engine/Source/Runtime/Engine/Private/DataChannel.cpp
+++ b/Engine/Source/Runtime/Engine/Private/DataChannel.cpp
@@ -3531,6 +3531,8 @@ int32 UActorChannel::WriteContentBlockPayload( UObject* Obj, FNetBitWriter &Bunc
 
 UObject* UActorChannel::ReadContentBlockHeader( FInBunch & Bunch, bool& bObjectDeleted, bool& bOutHasRepLayout )
 {
+	CA_ASSUME(Connection != nullptr);
+
 	const bool IsServer = Connection->Driver->IsServer();
 	bObjectDeleted = false;
 
