Added  FWorldPartitionPinRef, that acts as a FWorldPartitionSoftRef when there's not hard reference count to the actor descriptor, and like a FWorldPartitionHardRef when there is. This is useful when you want to keep an actor loaded during special operations, and don't trigger loading when the actor is not loaded. Mainly intended to be on stack, as a scoped operation. Also added unit-tests that covers this new type.

diff --git a/Engine/Source/Runtime/Engine/Private/Tests/WorldPartition/WorldPartitionHandleTests.cpp b/Engine/Source/Runtime/Engine/Private/Tests/WorldPartition/WorldPartitionHandleTests.cpp
--- a/Engine/Source/Runtime/Engine/Private/Tests/WorldPartition/WorldPartitionHandleTests.cpp
+++ b/Engine/Source/Runtime/Engine/Private/Tests/WorldPartition/WorldPartitionHandleTests.cpp
@@ -12,15 +12,34 @@ namespace WorldPartitionTests
 {
 	constexpr const uint32 TestFlags = EAutomationTestFlags::ApplicationContextMask | EAutomationTestFlags::EngineFilter;
 
-	IMPLEMENT_SIMPLE_AUTOMATION_TEST(FWorldPartitionHandleTest, TEST_NAME_ROOT ".Handle", TestFlags)
-	bool FWorldPartitionHandleTest::RunTest(const FString& Parameters)
+	IMPLEMENT_SIMPLE_AUTOMATION_TEST(FWorldPartitionSoftRefTest, TEST_NAME_ROOT ".Handle", TestFlags)
+	bool FWorldPartitionSoftRefTest::RunTest(const FString& Parameters)
 	{
 #if WITH_EDITOR
+		TestTrue(TEXT("SoftRef sizeof"), sizeof(FWorldPartitionSoftRef) == sizeof(void*));
+		TestTrue(TEXT("HardRef sizeof"), sizeof(FWorldPartitionHardRef) == sizeof(void*));
+
 		TUniquePtr<FWorldPartitionActorDesc> ActorDescPtr1(new FWorldPartitionActorDesc());
 		TUniquePtr<FWorldPartitionActorDesc> ActorDescPtr2(new FWorldPartitionActorDesc());
 
-		FWorldPartitionHandle Handle(&ActorDescPtr1);
-		FWorldPartitionReference Reference(&ActorDescPtr2);
+		FWorldPartitionSoftRef Handle(&ActorDescPtr1);
+		FWorldPartitionHardRef Reference(&ActorDescPtr2);
+
+		// Auto reference tests
+		{
+			TestTrue(TEXT("Handle soft refcount"), Handle->GetSoftRefCount() == 1);
+			TestTrue(TEXT("Handle hard refcount"), Handle->GetHardRefCount() == 0);
+			TestTrue(TEXT("Reference soft refcount"), Reference->GetSoftRefCount() == 0);
+			TestTrue(TEXT("Reference hard refcount"), Reference->GetHardRefCount() == 1);
+
+			FWorldPartitionPinRef AutoRefToHandle(Handle);
+			FWorldPartitionPinRef AutoRefToReference(Reference);
+
+			TestTrue(TEXT("Pin to Handle soft refcount"), Handle->GetSoftRefCount() == 2);
+			TestTrue(TEXT("Pin to Handle hard refcount"), Handle->GetHardRefCount() == 0);
+			TestTrue(TEXT("Pin to Reference soft refcount"), Reference->GetSoftRefCount() == 0);
+			TestTrue(TEXT("Pin to Reference hard refcount"), Reference->GetHardRefCount() == 2);
+		}
 
 		TestTrue(TEXT("Handle/Reference equality"), Handle != Reference);
 		TestTrue(TEXT("Reference/Handle equality"), Reference != Handle);
@@ -32,13 +51,13 @@ namespace WorldPartitionTests
 
 		// Conversions
 		{
-			FWorldPartitionHandle HandleToReference = Reference;
+			FWorldPartitionSoftRef HandleToReference = Reference;
 			TestTrue(TEXT("Handle/Reference equality"), HandleToReference == Reference);
 			TestTrue(TEXT("Reference/Handle equality"), Reference == HandleToReference);
 			TestTrue(TEXT("Reference soft refcount"), Reference->GetSoftRefCount() == 1);
 			TestTrue(TEXT("Reference hard refcount"), Reference->GetHardRefCount() == 1);
 
-			FWorldPartitionReference ReferenceToHandle = Handle;
+			FWorldPartitionHardRef ReferenceToHandle = Handle;
 			TestTrue(TEXT("Handle/Reference equality"), ReferenceToHandle == Handle);
 			TestTrue(TEXT("Handle/Reference equality"), Handle == ReferenceToHandle);
 			TestTrue(TEXT("Reference soft refcount"), Reference->GetSoftRefCount() == 1);
@@ -52,13 +71,13 @@ namespace WorldPartitionTests
 
 		// inplace new test
 		{
-			uint8 Buffer[sizeof(FWorldPartitionHandle)];
-			FWorldPartitionHandle* HandlePtr = new (Buffer) FWorldPartitionHandle(Reference);
+			uint8 Buffer[sizeof(FWorldPartitionSoftRef)];
+			FWorldPartitionSoftRef* HandlePtr = new (Buffer) FWorldPartitionSoftRef(Reference);
 
 			TestTrue(TEXT("Handle array soft refcount"), Reference->GetSoftRefCount() == 1);
 			TestTrue(TEXT("Handle array hard refcount"), Reference->GetHardRefCount() == 1);
 
-			HandlePtr->~FWorldPartitionHandle();
+			HandlePtr->~FWorldPartitionSoftRef();
 
 			TestTrue(TEXT("Handle array soft refcount"), Reference->GetSoftRefCount() == 0);
 			TestTrue(TEXT("Handle array hard refcount"), Reference->GetHardRefCount() == 1);
@@ -66,13 +85,13 @@ namespace WorldPartitionTests
 
 		// TArray test
 		{
-			TArray<FWorldPartitionHandle> HandleList;
+			TArray<FWorldPartitionSoftRef> HandleList;
 			HandleList.Add(Handle);
 
 			TestTrue(TEXT("Handle array soft refcount"), Handle->GetSoftRefCount() == 2);
 			TestTrue(TEXT("Handle array hard refcount"), Handle->GetHardRefCount() == 0);
 
-			FWorldPartitionReference ReferenceToHandle = Handle;
+			FWorldPartitionHardRef ReferenceToHandle = Handle;
 			TestTrue(TEXT("Handle/Reference equality"), ReferenceToHandle == Handle);
 			TestTrue(TEXT("Handle/Reference equality"), Handle == ReferenceToHandle);
 			TestTrue(TEXT("Handle soft refcount"), Handle->GetSoftRefCount() == 2);
@@ -103,13 +122,13 @@ namespace WorldPartitionTests
 
 		// TSet test
 		{
-			TSet<FWorldPartitionHandle> HandleSet;
+			TSet<FWorldPartitionSoftRef> HandleSet;
 			HandleSet.Add(Handle);
 
 			TestTrue(TEXT("Handle set soft refcount"), Handle->GetSoftRefCount() == 2);
 			TestTrue(TEXT("Handle set hard refcount"), Handle->GetHardRefCount() == 0);
 
-			FWorldPartitionReference ReferenceToHandle = Handle;
+			FWorldPartitionHardRef ReferenceToHandle = Handle;
 			TestTrue(TEXT("Handle/Reference equality"), ReferenceToHandle == Handle);
 			TestTrue(TEXT("Handle/Reference equality"), Handle == ReferenceToHandle);
 			TestTrue(TEXT("Handle soft refcount"), Handle->GetSoftRefCount() == 2);
diff --git a/Engine/Source/Runtime/Engine/Private/WorldPartition/HLOD/HLODBuilder.cpp b/Engine/Source/Runtime/Engine/Private/WorldPartition/HLOD/HLODBuilder.cpp
--- a/Engine/Source/Runtime/Engine/Private/WorldPartition/HLOD/HLODBuilder.cpp
+++ b/Engine/Source/Runtime/Engine/Private/WorldPartition/HLOD/HLODBuilder.cpp
@@ -41,7 +41,7 @@ public:
 		uint64 CellHash = FHLODActorDesc::ComputeCellHash(HLODLayer->GetName(), Context->GridIndexX, Context->GridIndexY, Context->GridIndexZ);
 
 		int32 HLODActorRefIndex = INDEX_NONE;
-		FWorldPartitionHandle HLODActorHandle;
+		FWorldPartitionSoftRef HLODActorHandle;
 		if (Context->HLODActorDescs.RemoveAndCopyValue(CellHash, HLODActorHandle))
 		{
 			HLODActorRefIndex = Context->ActorReferences.Add(HLODActorHandle);
@@ -75,7 +75,7 @@ public:
 		{
 			if (HLODActorRefIndex != INDEX_NONE)
 			{
-				Context->HLODActorDescs.Add(CellHash, FWorldPartitionHandle(WorldPartition, HLODActor->GetActorGuid()));
+				Context->HLODActorDescs.Add(CellHash, FWorldPartitionSoftRef(WorldPartition, HLODActor->GetActorGuid()));
 				Context->ActorReferences.RemoveAtSwap(HLODActorRefIndex);
 			}
 			else
diff --git a/Engine/Source/Runtime/Engine/Private/WorldPartition/HLOD/HLODBuilder.h b/Engine/Source/Runtime/Engine/Private/WorldPartition/HLOD/HLODBuilder.h
--- a/Engine/Source/Runtime/Engine/Private/WorldPartition/HLOD/HLODBuilder.h
+++ b/Engine/Source/Runtime/Engine/Private/WorldPartition/HLOD/HLODBuilder.h
@@ -15,8 +15,8 @@ class AWorldPartitionHLOD;
 
 struct ENGINE_API FHLODGenerationContext
 {
-	TMap<uint64, FWorldPartitionHandle> HLODActorDescs;
-	TArray<FWorldPartitionReference> ActorReferences;
+	TMap<uint64, FWorldPartitionSoftRef> HLODActorDescs;
+	TArray<FWorldPartitionHardRef> ActorReferences;
 	
 	// Everything needed to build the cell hash
 	int64 GridIndexX;
diff --git a/Engine/Source/Runtime/Engine/Private/WorldPartition/Landscape/LandscapeActorDesc.cpp b/Engine/Source/Runtime/Engine/Private/WorldPartition/Landscape/LandscapeActorDesc.cpp
--- a/Engine/Source/Runtime/Engine/Private/WorldPartition/Landscape/LandscapeActorDesc.cpp
+++ b/Engine/Source/Runtime/Engine/Private/WorldPartition/Landscape/LandscapeActorDesc.cpp
@@ -36,7 +36,7 @@ void FLandscapeActorDesc::OnRegister(UWorldPartition* WorldPartition)
 
 	if (ULandscapeInfo* LandscapeInfo = ULandscapeInfo::Find(WorldPartition->GetWorld(), GridGuid))
 	{
-		FWorldPartitionHandle Handle(WorldPartition, GetGuid());
+		FWorldPartitionSoftRef Handle(WorldPartition, GetGuid());
 		check(Handle.IsValid());
 		LandscapeInfo->ProxyHandles.Add(MoveTemp(Handle));
 	}
@@ -48,7 +48,7 @@ void FLandscapeActorDesc::OnUnregister(UWorldPartition* WorldPartition)
 
 	if (ULandscapeInfo* LandscapeInfo = ULandscapeInfo::Find(WorldPartition->GetWorld(), GridGuid))
 	{
-		FWorldPartitionHandle Handle(WorldPartition, GetGuid());
+		FWorldPartitionSoftRef Handle(WorldPartition, GetGuid());
 		check(Handle.IsValid());
 		LandscapeInfo->ProxyHandles.Remove(Handle);
 	}
diff --git a/Engine/Source/Runtime/Engine/Private/WorldPartition/Landscape/LandscapeSplineActorDesc.cpp b/Engine/Source/Runtime/Engine/Private/WorldPartition/Landscape/LandscapeSplineActorDesc.cpp
--- a/Engine/Source/Runtime/Engine/Private/WorldPartition/Landscape/LandscapeSplineActorDesc.cpp
+++ b/Engine/Source/Runtime/Engine/Private/WorldPartition/Landscape/LandscapeSplineActorDesc.cpp
@@ -37,7 +37,7 @@ void FLandscapeSplineActorDesc::OnRegister(UWorldPartition* WorldPartition)
 
 	if (ULandscapeInfo* LandscapeInfo = ULandscapeInfo::Find(WorldPartition->GetWorld(), LandscapeGuid))
 	{
-		FWorldPartitionHandle Handle(WorldPartition, GetGuid());
+		FWorldPartitionSoftRef Handle(WorldPartition, GetGuid());
 		check(Handle.IsValid());
 		LandscapeInfo->SplineHandles.Add(MoveTemp(Handle));
 	}
@@ -49,7 +49,7 @@ void FLandscapeSplineActorDesc::OnUnregister(UWorldPartition* WorldPartition)
 
 	if (ULandscapeInfo* LandscapeInfo = ULandscapeInfo::Find(WorldPartition->GetWorld(), LandscapeGuid))
 	{
-		FWorldPartitionHandle Handle(WorldPartition, GetGuid());
+		FWorldPartitionSoftRef Handle(WorldPartition, GetGuid());
 		check(Handle.IsValid());
 		LandscapeInfo->SplineHandles.Remove(Handle);
 	}
diff --git a/Engine/Source/Runtime/Engine/Private/WorldPartition/WorldPartitionHandle.cpp b/Engine/Source/Runtime/Engine/Private/WorldPartition/WorldPartitionHandle.cpp
--- a/Engine/Source/Runtime/Engine/Private/WorldPartition/WorldPartitionHandle.cpp
+++ b/Engine/Source/Runtime/Engine/Private/WorldPartition/WorldPartitionHandle.cpp
@@ -4,7 +4,7 @@
 #include "WorldPartition/WorldPartition.h"
 
 #if WITH_EDITOR
-TUniquePtr<FWorldPartitionActorDesc>* FWorldPartitionHandleBase::GetActorDesc(UWorldPartition* WorldPartition, const FGuid& ActorGuid)
+TUniquePtr<FWorldPartitionActorDesc>* FWorldPartitionSoftRefUtils::GetActorDesc(UWorldPartition* WorldPartition, const FGuid& ActorGuid)
 {
 	if (TUniquePtr<FWorldPartitionActorDesc>** ActorDescPtr = WorldPartition->Actors.Find(ActorGuid))
 	{
@@ -14,17 +14,17 @@ TUniquePtr<FWorldPartitionActorDesc>* FWorldPartitionHandleBase::GetActorDesc(UW
 	return nullptr;
 }
 
-void FWorldPartitionHandleImpl::IncRefCount(FWorldPartitionActorDesc* ActorDesc)
+void FWorldPartitionSoftRefImpl::IncRefCount(FWorldPartitionActorDesc* ActorDesc)
 {
 	ActorDesc->IncSoftRefCount();
 }
 
-void FWorldPartitionHandleImpl::DecRefCount(FWorldPartitionActorDesc* ActorDesc)
+void FWorldPartitionSoftRefImpl::DecRefCount(FWorldPartitionActorDesc* ActorDesc)
 {
 	ActorDesc->DecSoftRefCount();
 }
 
-void FWorldPartitionReferenceImpl::IncRefCount(FWorldPartitionActorDesc* ActorDesc)
+void FWorldPartitionHardRefImpl::IncRefCount(FWorldPartitionActorDesc* ActorDesc)
 {
 	if (ActorDesc->IncHardRefCount() == 1)
 	{
@@ -39,7 +39,7 @@ void FWorldPartitionReferenceImpl::IncRefCount(FWorldPartitionActorDesc* ActorDe
 	}
 }
 
-void FWorldPartitionReferenceImpl::DecRefCount(FWorldPartitionActorDesc* ActorDesc)
+void FWorldPartitionHardRefImpl::DecRefCount(FWorldPartitionActorDesc* ActorDesc)
 {
 	if (!ActorDesc->DecHardRefCount())
 	{
@@ -53,4 +53,30 @@ void FWorldPartitionReferenceImpl::DecRefCount(FWorldPartitionActorDesc* ActorDe
 		}
 	}
 }
+
+void FWorldPartitionPinRefImpl::IncRefCount(FWorldPartitionActorDesc* ActorDesc)
+{
+	bIsReference = ActorDesc->GetHardRefCount() > 0;
+
+	if (bIsReference)
+	{
+		FWorldPartitionHardRefImpl::IncRefCount(ActorDesc);
+	}
+	else
+	{
+		FWorldPartitionSoftRefImpl::IncRefCount(ActorDesc);
+	}
+}
+
+void FWorldPartitionPinRefImpl::DecRefCount(FWorldPartitionActorDesc* ActorDesc)
+{
+	if (bIsReference)
+	{
+		FWorldPartitionHardRefImpl::DecRefCount(ActorDesc);
+	}
+	else
+	{
+		FWorldPartitionSoftRefImpl::DecRefCount(ActorDesc);
+	}
+}
 #endif
\ No newline at end of file
diff --git a/Engine/Source/Runtime/Engine/Private/WorldPartition/WorldPartitionRuntimeSpatialHash.cpp b/Engine/Source/Runtime/Engine/Private/WorldPartition/WorldPartitionRuntimeSpatialHash.cpp
--- a/Engine/Source/Runtime/Engine/Private/WorldPartition/WorldPartitionRuntimeSpatialHash.cpp
+++ b/Engine/Source/Runtime/Engine/Private/WorldPartition/WorldPartitionRuntimeSpatialHash.cpp
@@ -1594,7 +1594,7 @@ TArray<FGuid> GenerateHLODsForGrid(UWorldPartition* WorldPartition, const FSpati
 				TArray<AActor*> CellActors;
 				for (const FGuid& ActorGuid: GridCellDataChunk->GetActors())
 				{
-					FWorldPartitionReference& ActorRef = Context.ActorReferences.Emplace_GetRef(WorldPartition, ActorGuid);
+					FWorldPartitionHardRef& ActorRef = Context.ActorReferences.Emplace_GetRef(WorldPartition, ActorGuid);
 					FWorldPartitionActorDesc* ActorDesc  = ActorRef.Get();
 					CellActors.Add(ActorDesc->GetActor());
 				}
@@ -1887,10 +1887,10 @@ bool UWorldPartitionRuntimeSpatialHash::GenerateHLOD(ISourceControlHelper* Sourc
 
 	// Create HLODs generation context
 	FHLODGenerationContext Context;
-	TArray<FWorldPartitionHandle> InvalidHLODActors;
+	TArray<FWorldPartitionSoftRef> InvalidHLODActors;
 	for (TWorldPartitionActorDescIterator<AWorldPartitionHLOD, FHLODActorDesc> HLODIterator(WorldPartition); HLODIterator; ++HLODIterator)
 	{
-		FWorldPartitionHandle HLODActorHandle(WorldPartition, HLODIterator->GetGuid());
+		FWorldPartitionSoftRef HLODActorHandle(WorldPartition, HLODIterator->GetGuid());
 
 		if (HLODIterator->GetCellHash())
 		{
@@ -1974,7 +1974,7 @@ bool UWorldPartitionRuntimeSpatialHash::GenerateHLOD(ISourceControlHelper* Sourc
 		GenerateHLODs(HLODGrids[HLODGridName], GridsDepth[HLODGridName] + 1, HLODActorClusters);
 	}
 
-	auto DeleteHLODActor = [&SourceControlHelper](FWorldPartitionHandle ActorHandle)
+	auto DeleteHLODActor = [&SourceControlHelper](FWorldPartitionSoftRef ActorHandle)
 	{
 		FWorldPartitionActorDesc* HLODActorDesc = ActorHandle.Get();
 		check(HLODActorDesc);
diff --git a/Engine/Source/Runtime/Engine/Public/WorldPartition/WorldPartitionActorDesc.h b/Engine/Source/Runtime/Engine/Public/WorldPartition/WorldPartitionActorDesc.h
--- a/Engine/Source/Runtime/Engine/Public/WorldPartition/WorldPartitionActorDesc.h
+++ b/Engine/Source/Runtime/Engine/Public/WorldPartition/WorldPartitionActorDesc.h
@@ -23,7 +23,7 @@ struct FWorldPartitionActorDescInitData
 #if WITH_DEV_AUTOMATION_TESTS
 namespace WorldPartitionTests
 {
-	class FWorldPartitionHandleTest;
+	class FWorldPartitionSoftRefTest;
 }
 #endif
 
@@ -37,7 +37,7 @@ class ENGINE_API FWorldPartitionActorDesc
 	friend class UWorldPartition;
 
 #if WITH_DEV_AUTOMATION_TESTS
-	friend class WorldPartitionTests::FWorldPartitionHandleTest;
+	friend class WorldPartitionTests::FWorldPartitionSoftRefTest;
 #endif
 
 public:
diff --git a/Engine/Source/Runtime/Engine/Public/WorldPartition/WorldPartitionHandle.h b/Engine/Source/Runtime/Engine/Public/WorldPartition/WorldPartitionHandle.h
--- a/Engine/Source/Runtime/Engine/Public/WorldPartition/WorldPartitionHandle.h
+++ b/Engine/Source/Runtime/Engine/Public/WorldPartition/WorldPartitionHandle.h
@@ -8,23 +8,22 @@
 class UWorldPartition;
 class FWorldPartitionActorDesc;
 
-class ENGINE_API FWorldPartitionHandleBase
+struct ENGINE_API FWorldPartitionSoftRefUtils
 {
-public:
 	static TUniquePtr<FWorldPartitionActorDesc>* GetActorDesc(UWorldPartition* WorldPartition, const FGuid& ActorGuid);
 };
 
 template <typename Impl>
-class ENGINE_API TWorldPartitionHandleBase : public FWorldPartitionHandleBase
+class ENGINE_API TWorldPartitionHandle : protected Impl
 {
-	friend struct FWorldPartitionHandleHelpers;
+	friend struct FWorldPartitionSoftRefHelpers;
 
 public:
-	FORCEINLINE TWorldPartitionHandleBase()
+	FORCEINLINE TWorldPartitionHandle()
 		: ActorDesc(nullptr)
 	{}
 
-	FORCEINLINE TWorldPartitionHandleBase(TUniquePtr<FWorldPartitionActorDesc>* InActorDesc)
+	FORCEINLINE TWorldPartitionHandle(TUniquePtr<FWorldPartitionActorDesc>* InActorDesc)
 	{
 		ActorDesc = InActorDesc;
 
@@ -34,9 +33,9 @@ public:
 		}
 	}
 
-	FORCEINLINE TWorldPartitionHandleBase(UWorldPartition* WorldPartition, const FGuid& ActorGuid)
+	FORCEINLINE TWorldPartitionHandle(UWorldPartition* WorldPartition, const FGuid& ActorGuid)
 	{
-		ActorDesc = GetActorDesc(WorldPartition, ActorGuid);
+		ActorDesc = FWorldPartitionSoftRefUtils::GetActorDesc(WorldPartition, ActorGuid);
 
 		if (IsValid())
 		{
@@ -44,13 +43,13 @@ public:
 		}
 	}
 
-	FORCEINLINE TWorldPartitionHandleBase(const TWorldPartitionHandleBase& Other)
+	FORCEINLINE TWorldPartitionHandle(const TWorldPartitionHandle& Other)
 		: ActorDesc(nullptr)
 	{
 		*this = Other;
 	}
 
-	FORCEINLINE TWorldPartitionHandleBase(TWorldPartitionHandleBase&& Other)
+	FORCEINLINE TWorldPartitionHandle(TWorldPartitionHandle&& Other)
 		: ActorDesc(nullptr)
 	{
 		*this = Other;
@@ -58,20 +57,20 @@ public:
 
 	// Conversions
 	template <typename T>
-	FORCEINLINE TWorldPartitionHandleBase<Impl>(const TWorldPartitionHandleBase<T>& Other)
+	FORCEINLINE TWorldPartitionHandle<Impl>(const TWorldPartitionHandle<T>& Other)
 		: ActorDesc(nullptr)
 	{
 		*this = Other;
 	}
 
 	template <typename T>
-	FORCEINLINE TWorldPartitionHandleBase<Impl>(TWorldPartitionHandleBase<T>&& Other)
+	FORCEINLINE TWorldPartitionHandle<Impl>(TWorldPartitionHandle<T>&& Other)
 		: ActorDesc(nullptr)
 	{
 		*this = Other;
 	}
 
-	FORCEINLINE ~TWorldPartitionHandleBase()
+	FORCEINLINE ~TWorldPartitionHandle()
 	{
 		if (IsValid())
 		{
@@ -79,7 +78,7 @@ public:
 		}
 	}
 
-	FORCEINLINE TWorldPartitionHandleBase& operator=(const TWorldPartitionHandleBase& Other)
+	FORCEINLINE TWorldPartitionHandle& operator=(const TWorldPartitionHandle& Other)
 	{
 		if ((void*)this == (void*)&Other)
 		{
@@ -101,7 +100,7 @@ public:
 		return *this;
 	}
 
-	FORCEINLINE TWorldPartitionHandleBase<Impl>& operator=(TWorldPartitionHandleBase&& Other)
+	FORCEINLINE TWorldPartitionHandle<Impl>& operator=(TWorldPartitionHandle&& Other)
 	{
 		if ((void*)this == (void*)&Other)
 		{
@@ -127,7 +126,7 @@ public:
 
 	// Conversions
 	template <typename T>
-	FORCEINLINE TWorldPartitionHandleBase<Impl>& operator=(const TWorldPartitionHandleBase<T>& Other)
+	FORCEINLINE TWorldPartitionHandle<Impl>& operator=(const TWorldPartitionHandle<T>& Other)
 	{
 		if ((void*)this == (void*)&Other)
 		{
@@ -150,7 +149,7 @@ public:
 	}
 
 	template <typename T>
-	FORCEINLINE TWorldPartitionHandleBase<Impl>& operator=(TWorldPartitionHandleBase<T>&& Other)
+	FORCEINLINE TWorldPartitionHandle<Impl>& operator=(TWorldPartitionHandle<T>&& Other)
 	{
 		if ((void*)this == (void*)&Other)
 		{
@@ -194,30 +193,30 @@ public:
 		return IsValid() ? ActorDesc->Get() : nullptr;
 	}
 		
-	friend FORCEINLINE uint32 GetTypeHash(const TWorldPartitionHandleBase<Impl>& HandleBase)
+	friend FORCEINLINE uint32 GetTypeHash(const TWorldPartitionHandle<Impl>& HandleBase)
 	{
 		return ::PointerHash(HandleBase.ActorDesc);
 	}
 
-	FORCEINLINE bool operator==(const TWorldPartitionHandleBase& Other) const
+	FORCEINLINE bool operator==(const TWorldPartitionHandle& Other) const
 	{
 		return ActorDesc == Other.ActorDesc;
 	}
 
-	FORCEINLINE bool operator!=(const TWorldPartitionHandleBase& Other) const
+	FORCEINLINE bool operator!=(const TWorldPartitionHandle& Other) const
 	{
 		return !(*this == Other);
 	}
 
 	// Conversions
 	template <typename T>
-	FORCEINLINE bool operator==(const TWorldPartitionHandleBase<T>& Other) const
+	FORCEINLINE bool operator==(const TWorldPartitionHandle<T>& Other) const
 	{
 		return Get() == *Other;
 	}
 
 	template <typename T>
-	FORCEINLINE bool operator!=(const TWorldPartitionHandleBase<T>& Other) const
+	FORCEINLINE bool operator!=(const TWorldPartitionHandle<T>& Other) const
 	{
 		return !(*this == Other);
 	}
@@ -237,18 +236,44 @@ public:
 	TUniquePtr<FWorldPartitionActorDesc>* ActorDesc;
 };
 
-struct ENGINE_API FWorldPartitionHandleImpl
+struct ENGINE_API FWorldPartitionSoftRefImpl
 {
 	static void IncRefCount(FWorldPartitionActorDesc* ActorDesc);
 	static void DecRefCount(FWorldPartitionActorDesc* ActorDesc);
 };
 
-struct ENGINE_API FWorldPartitionReferenceImpl
+struct ENGINE_API FWorldPartitionHardRefImpl
 {
 	static void IncRefCount(FWorldPartitionActorDesc* ActorDesc);
 	static void DecRefCount(FWorldPartitionActorDesc* ActorDesc);
 };
 
-typedef TWorldPartitionHandleBase<FWorldPartitionHandleImpl> FWorldPartitionHandle;
-typedef TWorldPartitionHandleBase<FWorldPartitionReferenceImpl> FWorldPartitionReference;
+struct ENGINE_API FWorldPartitionPinRefImpl
+{
+	FWorldPartitionPinRefImpl() : bIsReference(false) {}
+	void IncRefCount(FWorldPartitionActorDesc* ActorDesc);
+	void DecRefCount(FWorldPartitionActorDesc* ActorDesc);
+	bool bIsReference;
+};
+
+/**
+ * FWorldPartitionSoftRef will increment/decrement the soft reference count on the actor descriptor.
+ * This won't trigger any loading, but will prevent cleanup of the actor descriptor when destroying an
+ * actor in the editor.
+ */
+typedef TWorldPartitionHandle<FWorldPartitionSoftRefImpl> FWorldPartitionSoftRef;
+
+/**
+ * FWorldPartitionHardRef will increment/decrement the hard reference count on the actor descriptor.
+ * This will trigger actor loading/unloading when the hard reference counts gets to one/zero.
+ */
+typedef TWorldPartitionHandle<FWorldPartitionHardRefImpl> FWorldPartitionHardRef;
+
+/**
+ * FWorldPartitionPinRef will act as a softref when there's not hard reference count to the actor descriptor,
+ * and like a hardref when there is. This is useful when you want to keep an actor loaded during special
+ * operations, and don't trigger loading when the actor is not loaded. Mainly intended to be on stack, as a
+ * scoped operation.
+ */
+typedef TWorldPartitionHandle<FWorldPartitionPinRefImpl> FWorldPartitionPinRef;
 #endif
\ No newline at end of file
diff --git a/Engine/Source/Runtime/Landscape/Classes/LandscapeInfo.h b/Engine/Source/Runtime/Landscape/Classes/LandscapeInfo.h
--- a/Engine/Source/Runtime/Landscape/Classes/LandscapeInfo.h
+++ b/Engine/Source/Runtime/Landscape/Classes/LandscapeInfo.h
@@ -152,9 +152,9 @@ private:
 
 	// World Partition handles to actors with the same LandscapeGuid
 	friend class FLandscapeActorDesc;
-	TSet<FWorldPartitionHandle> ProxyHandles;
+	TSet<FWorldPartitionSoftRef> ProxyHandles;
 	friend class FLandscapeSplineActorDesc;
-	TSet<FWorldPartitionHandle> SplineHandles;
+	TSet<FWorldPartitionSoftRef> SplineHandles;
 #endif
 
 	TSet<ULandscapeComponent*> SelectedComponents;
@@ -293,7 +293,7 @@ public:
 	LANDSCAPE_API uint32 GetGridSize(uint32 InGridSizeInComponents) const;
 
 	/** Returns LandscapeSplineActor WorldPartition handles */
-	const TSet<FWorldPartitionHandle>& GetSplineHandles() const { return SplineHandles; }
+	const TSet<FWorldPartitionSoftRef>& GetSplineHandles() const { return SplineHandles; }
 #endif
 	LANDSCAPE_API static ULandscapeInfo* Find(UWorld* InWorld, const FGuid& LandscapeGuid);
 	static ULandscapeInfo* FindOrCreate(UWorld* InWorld, const FGuid& LandscapeGuid);
diff --git a/Engine/Source/Runtime/Landscape/Classes/LandscapeStreamingProxy.h b/Engine/Source/Runtime/Landscape/Classes/LandscapeStreamingProxy.h
--- a/Engine/Source/Runtime/Landscape/Classes/LandscapeStreamingProxy.h
+++ b/Engine/Source/Runtime/Landscape/Classes/LandscapeStreamingProxy.h
@@ -28,7 +28,7 @@ public:
 
 #if WITH_EDITORONLY_DATA
 	/** hard refs to actors that need to be loaded when this proxy is loaded */
-	TSet<FWorldPartitionReference> ActorDescReferences;
+	TSet<FWorldPartitionHardRef> ActorDescReferences;
 #endif
 
 	//~ Begin UObject Interface
diff --git a/Engine/Source/Runtime/Landscape/Private/Landscape.cpp b/Engine/Source/Runtime/Landscape/Private/Landscape.cpp
--- a/Engine/Source/Runtime/Landscape/Private/Landscape.cpp
+++ b/Engine/Source/Runtime/Landscape/Private/Landscape.cpp
@@ -3123,11 +3123,11 @@ void ULandscapeInfo::Initialize(UWorld* InWorld, const FGuid& InLandscapeGuid)
 		{
 			if (ActorDescIterator->GetActorClass()->IsChildOf(ALandscapeStreamingProxy::StaticClass()))
 			{
-				ProxyHandles.Add(FWorldPartitionHandle(WorldPartition, ActorDescIterator->GetGuid()));
+				ProxyHandles.Add(FWorldPartitionSoftRef(WorldPartition, ActorDescIterator->GetGuid()));
 			}
 			else if (ActorDescIterator->GetActorClass()->IsChildOf(ALandscapeSplineActor::StaticClass()))
 			{
-				SplineHandles.Emplace(FWorldPartitionHandle(WorldPartition, ActorDescIterator->GetGuid()));
+				SplineHandles.Emplace(FWorldPartitionSoftRef(WorldPartition, ActorDescIterator->GetGuid()));
 			}
 
 		}
diff --git a/Engine/Source/Runtime/Landscape/Private/LandscapeEdit.cpp b/Engine/Source/Runtime/Landscape/Private/LandscapeEdit.cpp
--- a/Engine/Source/Runtime/Landscape/Private/LandscapeEdit.cpp
+++ b/Engine/Source/Runtime/Landscape/Private/LandscapeEdit.cpp
@@ -2251,7 +2251,7 @@ bool ULandscapeInfo::HasUnloadedComponentsInRegion(int32 X1, int32 Y1, int32 X2,
 	UActorPartitionSubsystem::FCellCoord MinCoord = UActorPartitionSubsystem::FCellCoord::GetCellCoord(FIntPoint(ComponentIndexX1 * ComponentSizeQuads, ComponentIndexY1 * ComponentSizeQuads), World->PersistentLevel, LandscapeActor->GridSize);
 	UActorPartitionSubsystem::FCellCoord MaxCoord = UActorPartitionSubsystem::FCellCoord::GetCellCoord(FIntPoint(ComponentIndexX2 * ComponentSizeQuads, ComponentIndexY2 * ComponentSizeQuads), World->PersistentLevel, LandscapeActor->GridSize);
 
-	for (const FWorldPartitionHandle& Handle : ProxyHandles)
+	for (const FWorldPartitionSoftRef& Handle : ProxyHandles)
 	{
 		if (FLandscapeActorDesc* LandscapeActorDesc = (FLandscapeActorDesc*)Handle.Get())
 		{
@@ -5029,13 +5029,13 @@ void ALandscapeStreamingProxy::PostRegisterAllComponents()
 		check(LandscapeInfo);
 		if (GEditor && !GetWorld()->IsGameWorld())
 		{
-			const TSet<FWorldPartitionHandle>& SplineHandles = LandscapeInfo->GetSplineHandles();
+			const TSet<FWorldPartitionSoftRef>& SplineHandles = LandscapeInfo->GetSplineHandles();
 			if (SplineHandles.Num())
 			{
 				FVector ActorLocation = GetActorLocation();
 				FBox Bounds(ActorLocation, ActorLocation + (GridSize * LandscapeInfo->DrawScale));
 				// Get a reference to Spline Actors that have intersecting bounds with us
-				for (const FWorldPartitionHandle& SplineHandle : LandscapeInfo->GetSplineHandles())
+				for (const FWorldPartitionSoftRef& SplineHandle : LandscapeInfo->GetSplineHandles())
 				{
 					if (SplineHandle.IsValid() && Bounds.IntersectXY(SplineHandle->GetBounds()))
 					{
