Add check to silence static analysis

diff --git a/Engine/Source/Runtime/Engine/Private/TextureDerivedDataTask.cpp b/Engine/Source/Runtime/Engine/Private/TextureDerivedDataTask.cpp
--- a/Engine/Source/Runtime/Engine/Private/TextureDerivedDataTask.cpp
+++ b/Engine/Source/Runtime/Engine/Private/TextureDerivedDataTask.cpp
@@ -492,6 +492,7 @@ void FTextureCacheDerivedDataWorker::DoWork()
 
 		if (bSucceeded && bForVirtualTextureStreamingBuild && CVarVTValidateCompressionOnLoad.GetValueOnAnyThread())
 		{
+			check(DerivedData->VTData);
 			bSucceeded = DerivedData->VTData->ValidateCompression(Texture.GetPathName());
 			if (!bSucceeded)
 			{
