Fix mac build.

diff --git a/Engine/Source/Runtime/Engine/Private/WorldPartition/WorldPartitionHandle.cpp b/Engine/Source/Runtime/Engine/Private/WorldPartition/WorldPartitionHandle.cpp
--- a/Engine/Source/Runtime/Engine/Private/WorldPartition/WorldPartitionHandle.cpp
+++ b/Engine/Source/Runtime/Engine/Private/WorldPartition/WorldPartitionHandle.cpp
@@ -4,7 +4,7 @@
 #include "WorldPartition/WorldPartition.h"
 
 #if WITH_EDITOR
-TUniquePtr<FWorldPartitionActorDesc>* FWorldPartitionSoftRefUtils::GetActorDesc(UWorldPartition* WorldPartition, const FGuid& ActorGuid)
+TUniquePtr<FWorldPartitionActorDesc>* FWorldPartitionHandleUtils::GetActorDesc(UWorldPartition* WorldPartition, const FGuid& ActorGuid)
 {
 	if (TUniquePtr<FWorldPartitionActorDesc>** ActorDescPtr = WorldPartition->Actors.Find(ActorGuid))
 	{
@@ -14,6 +14,11 @@ TUniquePtr<FWorldPartitionActorDesc>* FWorldPartitionSoftRefUtils::GetActorDesc(
 	return nullptr;
 }
 
+bool FWorldPartitionHandleUtils::IsActorDescLoaded(FWorldPartitionActorDesc* ActorDesc)
+{
+	return !!ActorDesc->GetActor();
+}
+
 void FWorldPartitionSoftRefImpl::IncRefCount(FWorldPartitionActorDesc* ActorDesc)
 {
 	ActorDesc->IncSoftRefCount();
diff --git a/Engine/Source/Runtime/Engine/Public/WorldPartition/WorldPartitionHandle.h b/Engine/Source/Runtime/Engine/Public/WorldPartition/WorldPartitionHandle.h
--- a/Engine/Source/Runtime/Engine/Public/WorldPartition/WorldPartitionHandle.h
+++ b/Engine/Source/Runtime/Engine/Public/WorldPartition/WorldPartitionHandle.h
@@ -8,9 +8,10 @@
 class UWorldPartition;
 class FWorldPartitionActorDesc;
 
-struct ENGINE_API FWorldPartitionSoftRefUtils
+struct ENGINE_API FWorldPartitionHandleUtils
 {
 	static TUniquePtr<FWorldPartitionActorDesc>* GetActorDesc(UWorldPartition* WorldPartition, const FGuid& ActorGuid);
+	static bool IsActorDescLoaded(FWorldPartitionActorDesc* ActorDesc);
 };
 
 template <typename Impl>
@@ -35,7 +36,7 @@ public:
 
 	FORCEINLINE TWorldPartitionHandle(UWorldPartition* WorldPartition, const FGuid& ActorGuid)
 	{
-		ActorDesc = FWorldPartitionSoftRefUtils::GetActorDesc(WorldPartition, ActorGuid);
+		ActorDesc = FWorldPartitionHandleUtils::GetActorDesc(WorldPartition, ActorGuid);
 
 		if (IsValid())
 		{
@@ -190,7 +191,7 @@ public:
 
 	FORCEINLINE bool IsLoaded() const
 	{
-		return IsValid() && (*ActorDesc)->GetActor();
+		return IsValid() && FWorldPartitionHandleUtils::IsActorDescLoaded(ActorDesc->Get());
 	}
 
 	FORCEINLINE FWorldPartitionActorDesc* Get() const
