Added Insights Trace for heavy function in editor use-cases

diff --git a/Engine/Source/Runtime/Engine/Private/PhysicsEngine/Experimental/PhysInterface_Chaos.cpp b/Engine/Source/Runtime/Engine/Private/PhysicsEngine/Experimental/PhysInterface_Chaos.cpp
--- a/Engine/Source/Runtime/Engine/Private/PhysicsEngine/Experimental/PhysInterface_Chaos.cpp
+++ b/Engine/Source/Runtime/Engine/Private/PhysicsEngine/Experimental/PhysInterface_Chaos.cpp
@@ -813,6 +813,7 @@ const FBodyInstance* FPhysInterface_Chaos::ShapeToOriginalBodyInstance(const FBo
 
 void FPhysInterface_Chaos::AddGeometry(FPhysicsActorHandle& InActor, const FGeometryAddParams& InParams, TArray<FPhysicsShapeHandle>* OutOptShapes)
 {
+	TRACE_CPUPROFILER_EVENT_SCOPE(FPhysInterface_Chaos::AddGeometry);
 	LLM_SCOPE(ELLMTag::ChaosGeometry);
 	TArray<TUniquePtr<Chaos::FImplicitObject>> Geoms;
 	Chaos::FShapesArray Shapes;
diff --git a/Engine/Source/Runtime/Experimental/Chaos/Public/Chaos/AABBTree.h b/Engine/Source/Runtime/Experimental/Chaos/Public/Chaos/AABBTree.h
--- a/Engine/Source/Runtime/Experimental/Chaos/Public/Chaos/AABBTree.h
+++ b/Engine/Source/Runtime/Experimental/Chaos/Public/Chaos/AABBTree.h
@@ -651,6 +651,7 @@ private:
 
 	void ReoptimizeTree()
 	{
+		TRACE_CPUPROFILER_EVENT_SCOPE(TAABBTree::ReoptimizeTree);
 		TArray<FElement> AllElements;
 
 		SIZE_T ReserveCount = DirtyElements.Num() + GlobalPayloads.Num();
diff --git a/Engine/Source/Runtime/PhysicsCore/Private/ChaosScene.cpp b/Engine/Source/Runtime/PhysicsCore/Private/ChaosScene.cpp
--- a/Engine/Source/Runtime/PhysicsCore/Private/ChaosScene.cpp
+++ b/Engine/Source/Runtime/PhysicsCore/Private/ChaosScene.cpp
@@ -262,6 +262,8 @@ void FChaosScene::UpdateActorsInAccelerationStructure(const TArrayView<FPhysicsA
 void FChaosScene::AddActorsToScene_AssumesLocked(TArray<FPhysicsActorHandle>& InHandles,const bool bImmediate)
 {
 #if WITH_CHAOS
+	TRACE_CPUPROFILER_EVENT_SCOPE(FChaosScene::AddActorsToScene_AssumesLocked)
+
 	Chaos::FPhysicsSolver* Solver = GetSolver();
 	Chaos::ISpatialAcceleration<Chaos::TAccelerationStructureHandle<float,3>,float,3>* SpatialAcceleration = GetSpacialAcceleration();
 	for(FPhysicsActorHandle& Handle : InHandles)
