Fixed move operators for actor descriptor handles and added unit-tests for move operations.

diff --git a/Engine/Source/Runtime/Engine/Private/Tests/WorldPartition/WorldPartitionHandleTests.cpp b/Engine/Source/Runtime/Engine/Private/Tests/WorldPartition/WorldPartitionHandleTests.cpp
--- a/Engine/Source/Runtime/Engine/Private/Tests/WorldPartition/WorldPartitionHandleTests.cpp
+++ b/Engine/Source/Runtime/Engine/Private/Tests/WorldPartition/WorldPartitionHandleTests.cpp
@@ -148,6 +148,73 @@ namespace WorldPartitionTests
 		TestTrue(TEXT("Handle hard refcount"), Handle->GetHardRefCount() == 0);
 		TestTrue(TEXT("Reference soft refcount"), Reference->GetSoftRefCount() == 0);
 		TestTrue(TEXT("Reference hard refcount"), Reference->GetHardRefCount() == 1);
+
+		// Move tests
+		{
+			// Handle move
+			TestTrue(TEXT("Handle soft refcount"), Handle->GetSoftRefCount() == 1);
+			TestTrue(TEXT("Handle hard refcount"), Handle->GetHardRefCount() == 0);
+			{
+				FWorldPartitionSoftRef HandleCopy(MoveTemp(Handle));
+				TestTrue(TEXT("Handle move src not valid"), !Handle.IsValid());
+				TestTrue(TEXT("Handle move dst valid"), HandleCopy.IsValid());
+				TestTrue(TEXT("Handle soft refcount"), HandleCopy->GetSoftRefCount() == 1);
+				TestTrue(TEXT("Handle hard refcount"), HandleCopy->GetHardRefCount() == 0);
+
+				Handle = MoveTemp(HandleCopy);
+				TestTrue(TEXT("Handle move src not valid"), !HandleCopy.IsValid());
+				TestTrue(TEXT("Handle move dst valid"), Handle.IsValid());
+				TestTrue(TEXT("Handle soft refcount"), Handle->GetSoftRefCount() == 1);
+				TestTrue(TEXT("Handle hard refcount"), Handle->GetHardRefCount() == 0);
+			}
+
+			// Reference move
+			TestTrue(TEXT("Reference soft refcount"), Reference->GetSoftRefCount() == 0);
+			TestTrue(TEXT("Reference hard refcount"), Reference->GetHardRefCount() == 1);
+			{
+				FWorldPartitionHardRef ReferenceCopy(MoveTemp(Reference));
+				TestTrue(TEXT("Reference move src not valid"), !Reference.IsValid());
+				TestTrue(TEXT("Reference move dst valid"), ReferenceCopy.IsValid());
+				TestTrue(TEXT("Reference soft refcount"), ReferenceCopy->GetSoftRefCount() == 0);
+				TestTrue(TEXT("Reference hard refcount"), ReferenceCopy->GetHardRefCount() == 1);
+
+				Reference = MoveTemp(ReferenceCopy);
+				TestTrue(TEXT("Reference move src not valid"), !ReferenceCopy.IsValid());
+				TestTrue(TEXT("Reference move dst valid"), Reference.IsValid());
+				TestTrue(TEXT("Reference soft refcount"), Reference->GetSoftRefCount() == 0);
+				TestTrue(TEXT("Reference hard refcount"), Reference->GetHardRefCount() == 1);
+			}
+
+			// Handle reference move
+			{
+				FWorldPartitionSoftRef HandleFromReference(MoveTemp(Reference));
+				TestTrue(TEXT("Handle move src not valid"), !Reference.IsValid());
+				TestTrue(TEXT("Handle move dst valid"), HandleFromReference.IsValid());
+				TestTrue(TEXT("Handle soft refcount"), HandleFromReference->GetSoftRefCount() == 1);
+				TestTrue(TEXT("Handle hard refcount"), HandleFromReference->GetHardRefCount() == 0);
+
+				Reference = MoveTemp(HandleFromReference);
+				TestTrue(TEXT("Handle move src not valid"), !HandleFromReference.IsValid());
+				TestTrue(TEXT("Handle move dst valid"), Reference.IsValid());
+				TestTrue(TEXT("Handle soft refcount"), Reference->GetSoftRefCount() == 0);
+				TestTrue(TEXT("Handle hard refcount"), Reference->GetHardRefCount() == 1);
+			}
+
+			// Reference handle move
+			{
+				FWorldPartitionHardRef ReferenceFromHandle(MoveTemp(Handle));
+				TestTrue(TEXT("Reference move src not valid"), !Handle.IsValid());
+				TestTrue(TEXT("Reference move dst valid"), ReferenceFromHandle.IsValid());
+				TestTrue(TEXT("Reference soft refcount"), ReferenceFromHandle->GetSoftRefCount() == 0);
+				TestTrue(TEXT("Reference hard refcount"), ReferenceFromHandle->GetHardRefCount() == 1);
+
+				Handle = MoveTemp(ReferenceFromHandle);
+				TestTrue(TEXT("Reference move src not valid"), !ReferenceFromHandle.IsValid());
+				TestTrue(TEXT("Reference move dst valid"), Handle.IsValid());
+				TestTrue(TEXT("Reference soft refcount"), Handle->GetSoftRefCount() == 1);
+				TestTrue(TEXT("Reference hard refcount"), Handle->GetHardRefCount() == 0);
+			}
+		}
 #endif
 		return true;
 	}
diff --git a/Engine/Source/Runtime/Engine/Public/WorldPartition/WorldPartitionHandle.h b/Engine/Source/Runtime/Engine/Public/WorldPartition/WorldPartitionHandle.h
--- a/Engine/Source/Runtime/Engine/Public/WorldPartition/WorldPartitionHandle.h
+++ b/Engine/Source/Runtime/Engine/Public/WorldPartition/WorldPartitionHandle.h
@@ -53,7 +53,7 @@ public:
 	FORCEINLINE TWorldPartitionHandle(TWorldPartitionHandle&& Other)
 		: ActorDesc(nullptr)
 	{
-		*this = Other;
+		*this = MoveTemp(Other);
 	}
 
 	// Conversions
@@ -68,7 +68,7 @@ public:
 	FORCEINLINE TWorldPartitionHandle<Impl>(TWorldPartitionHandle<T>&& Other)
 		: ActorDesc(nullptr)
 	{
-		*this = Other;
+		*this = MoveTemp(Other);
 	}
 
 	FORCEINLINE ~TWorldPartitionHandle()
@@ -81,21 +81,19 @@ public:
 
 	FORCEINLINE TWorldPartitionHandle& operator=(const TWorldPartitionHandle& Other)
 	{
-		if ((void*)this == (void*)&Other)
-		{
-			return *this;
-		}
-
-		if (IsValid())
+		if (this != &Other)
 		{
-			DecRefCount();
-		}
+			if (IsValid())
+			{
+				DecRefCount();
+			}
 
-		ActorDesc = Other.ActorDesc;
+			ActorDesc = Other.ActorDesc;
 
-		if (IsValid())
-		{
-			IncRefCount();
+			if (IsValid())
+			{
+				IncRefCount();
+			}
 		}
 
 		return *this;
@@ -103,22 +101,14 @@ public:
 
 	FORCEINLINE TWorldPartitionHandle<Impl>& operator=(TWorldPartitionHandle&& Other)
 	{
-		if ((void*)this == (void*)&Other)
-		{
-			return *this;
-		}
-
-		if (IsValid())
+		if (this != &Other)
 		{
-			DecRefCount();
-		}
-
-		ActorDesc = Other.ActorDesc;
+			if (IsValid())
+			{
+				DecRefCount();
+			}
 
-		if (IsValid())
-		{
-			IncRefCount();
-			Other.DecRefCount();
+			ActorDesc = Other.ActorDesc;
 			Other.ActorDesc = nullptr;
 		}
 
@@ -129,11 +119,6 @@ public:
 	template <typename T>
 	FORCEINLINE TWorldPartitionHandle<Impl>& operator=(const TWorldPartitionHandle<T>& Other)
 	{
-		if ((void*)this == (void*)&Other)
-		{
-			return *this;
-		}
-
 		if (IsValid())
 		{
 			DecRefCount();
@@ -152,11 +137,6 @@ public:
 	template <typename T>
 	FORCEINLINE TWorldPartitionHandle<Impl>& operator=(TWorldPartitionHandle<T>&& Other)
 	{
-		if ((void*)this == (void*)&Other)
-		{
-			return *this;
-		}
-
 		if (IsValid())
 		{
 			DecRefCount();
@@ -227,7 +207,7 @@ public:
 		return !(*this == Other);
 	}
 
-protected:
+public:
 	FORCEINLINE void IncRefCount()
 	{
 		Impl::IncRefCount(ActorDesc->Get());
@@ -238,7 +218,6 @@ protected:
 		Impl::DecRefCount(ActorDesc->Get());
 	}
 
-public:
 	TUniquePtr<FWorldPartitionActorDesc>* ActorDesc;
 };
 
