Replaced FWorldPartitionPinRef by FWorldPartitionHandlePinRefScope, which is a scoped struct that will hold a reference only if the actor was already loaded.
This removes the possibility to accidentally convert it to FWorldPartitionHandle/FWorldPartitionReference and make it clear that it's meant to be used as a scoped construct.

diff --git a/Engine/Source/Runtime/Engine/Private/Tests/WorldPartition/WorldPartitionHandleTests.cpp b/Engine/Source/Runtime/Engine/Private/Tests/WorldPartition/WorldPartitionHandleTests.cpp
--- a/Engine/Source/Runtime/Engine/Private/Tests/WorldPartition/WorldPartitionHandleTests.cpp
+++ b/Engine/Source/Runtime/Engine/Private/Tests/WorldPartition/WorldPartitionHandleTests.cpp
@@ -16,27 +16,30 @@ namespace WorldPartitionTests
 	bool FWorldPartitionSoftRefTest::RunTest(const FString& Parameters)
 	{
 #if WITH_EDITOR
-		TestTrue(TEXT("SoftRef sizeof"), sizeof(FWorldPartitionSoftRef) == sizeof(void*));
-		TestTrue(TEXT("HardRef sizeof"), sizeof(FWorldPartitionHardRef) == sizeof(void*));
+		TestTrue(TEXT("SoftRef sizeof"), sizeof(FWorldPartitionHandle) == sizeof(void*));
+		TestTrue(TEXT("HardRef sizeof"), sizeof(FWorldPartitionReference) == sizeof(void*));
 
 		TUniquePtr<FWorldPartitionActorDesc> ActorDescPtr1(new FWorldPartitionActorDesc());
 		TUniquePtr<FWorldPartitionActorDesc> ActorDescPtr2(new FWorldPartitionActorDesc());
 
-		FWorldPartitionSoftRef Handle(&ActorDescPtr1);
-		FWorldPartitionHardRef Reference(&ActorDescPtr2);
+		FWorldPartitionHandle Handle(&ActorDescPtr1);
+		FWorldPartitionReference Reference(&ActorDescPtr2);
 
-		// Auto reference tests
-		{
-			TestTrue(TEXT("Handle soft refcount"), Handle->GetSoftRefCount() == 1);
-			TestTrue(TEXT("Handle hard refcount"), Handle->GetHardRefCount() == 0);
-			TestTrue(TEXT("Reference soft refcount"), Reference->GetSoftRefCount() == 0);
-			TestTrue(TEXT("Reference hard refcount"), Reference->GetHardRefCount() == 1);
-
-			FWorldPartitionPinRef AutoRefToHandle(Handle);
-			FWorldPartitionPinRef AutoRefToReference(Reference);
+		TestTrue(TEXT("Handle soft refcount"), Handle->GetSoftRefCount() == 1);
+		TestTrue(TEXT("Handle hard refcount"), Handle->GetHardRefCount() == 0);
+		TestTrue(TEXT("Reference soft refcount"), Reference->GetSoftRefCount() == 0);
+		TestTrue(TEXT("Reference hard refcount"), Reference->GetHardRefCount() == 1);
 
-			TestTrue(TEXT("Pin to Handle soft refcount"), Handle->GetSoftRefCount() == 2);
+		// Pin handle scope test
+		{
+			FWorldPartitionHandlePinRefScope PinRefScopeHandle(Handle);
+			TestTrue(TEXT("Pin to Handle soft refcount"), Handle->GetSoftRefCount() == 1);
 			TestTrue(TEXT("Pin to Handle hard refcount"), Handle->GetHardRefCount() == 0);
+		}
+
+		// Pin reference scope test
+		{
+			FWorldPartitionHandlePinRefScope PinRefScopeReference(Reference);
 			TestTrue(TEXT("Pin to Reference soft refcount"), Reference->GetSoftRefCount() == 0);
 			TestTrue(TEXT("Pin to Reference hard refcount"), Reference->GetHardRefCount() == 2);
 		}
@@ -51,13 +54,13 @@ namespace WorldPartitionTests
 
 		// Conversions
 		{
-			FWorldPartitionSoftRef HandleToReference = Reference;
+			FWorldPartitionHandle HandleToReference = Reference;
 			TestTrue(TEXT("Handle/Reference equality"), HandleToReference == Reference);
 			TestTrue(TEXT("Reference/Handle equality"), Reference == HandleToReference);
 			TestTrue(TEXT("Reference soft refcount"), Reference->GetSoftRefCount() == 1);
 			TestTrue(TEXT("Reference hard refcount"), Reference->GetHardRefCount() == 1);
 
-			FWorldPartitionHardRef ReferenceToHandle = Handle;
+			FWorldPartitionReference ReferenceToHandle = Handle;
 			TestTrue(TEXT("Handle/Reference equality"), ReferenceToHandle == Handle);
 			TestTrue(TEXT("Handle/Reference equality"), Handle == ReferenceToHandle);
 			TestTrue(TEXT("Reference soft refcount"), Reference->GetSoftRefCount() == 1);
@@ -71,13 +74,13 @@ namespace WorldPartitionTests
 
 		// inplace new test
 		{
-			uint8 Buffer[sizeof(FWorldPartitionSoftRef)];
-			FWorldPartitionSoftRef* HandlePtr = new (Buffer) FWorldPartitionSoftRef(Reference);
+			uint8 Buffer[sizeof(FWorldPartitionHandle)];
+			FWorldPartitionHandle* HandlePtr = new (Buffer) FWorldPartitionHandle(Reference);
 
 			TestTrue(TEXT("Handle array soft refcount"), Reference->GetSoftRefCount() == 1);
 			TestTrue(TEXT("Handle array hard refcount"), Reference->GetHardRefCount() == 1);
 
-			HandlePtr->~FWorldPartitionSoftRef();
+			HandlePtr->~FWorldPartitionHandle();
 
 			TestTrue(TEXT("Handle array soft refcount"), Reference->GetSoftRefCount() == 0);
 			TestTrue(TEXT("Handle array hard refcount"), Reference->GetHardRefCount() == 1);
@@ -85,13 +88,13 @@ namespace WorldPartitionTests
 
 		// TArray test
 		{
-			TArray<FWorldPartitionSoftRef> HandleList;
+			TArray<FWorldPartitionHandle> HandleList;
 			HandleList.Add(Handle);
 
 			TestTrue(TEXT("Handle array soft refcount"), Handle->GetSoftRefCount() == 2);
 			TestTrue(TEXT("Handle array hard refcount"), Handle->GetHardRefCount() == 0);
 
-			FWorldPartitionHardRef ReferenceToHandle = Handle;
+			FWorldPartitionReference ReferenceToHandle = Handle;
 			TestTrue(TEXT("Handle/Reference equality"), ReferenceToHandle == Handle);
 			TestTrue(TEXT("Handle/Reference equality"), Handle == ReferenceToHandle);
 			TestTrue(TEXT("Handle soft refcount"), Handle->GetSoftRefCount() == 2);
@@ -122,13 +125,13 @@ namespace WorldPartitionTests
 
 		// TSet test
 		{
-			TSet<FWorldPartitionSoftRef> HandleSet;
+			TSet<FWorldPartitionHandle> HandleSet;
 			HandleSet.Add(Handle);
 
 			TestTrue(TEXT("Handle set soft refcount"), Handle->GetSoftRefCount() == 2);
 			TestTrue(TEXT("Handle set hard refcount"), Handle->GetHardRefCount() == 0);
 
-			FWorldPartitionHardRef ReferenceToHandle = Handle;
+			FWorldPartitionReference ReferenceToHandle = Handle;
 			TestTrue(TEXT("Handle/Reference equality"), ReferenceToHandle == Handle);
 			TestTrue(TEXT("Handle/Reference equality"), Handle == ReferenceToHandle);
 			TestTrue(TEXT("Handle soft refcount"), Handle->GetSoftRefCount() == 2);
@@ -155,7 +158,7 @@ namespace WorldPartitionTests
 			TestTrue(TEXT("Handle soft refcount"), Handle->GetSoftRefCount() == 1);
 			TestTrue(TEXT("Handle hard refcount"), Handle->GetHardRefCount() == 0);
 			{
-				FWorldPartitionSoftRef HandleCopy(MoveTemp(Handle));
+				FWorldPartitionHandle HandleCopy(MoveTemp(Handle));
 				TestTrue(TEXT("Handle move src not valid"), !Handle.IsValid());
 				TestTrue(TEXT("Handle move dst valid"), HandleCopy.IsValid());
 				TestTrue(TEXT("Handle soft refcount"), HandleCopy->GetSoftRefCount() == 1);
@@ -172,7 +175,7 @@ namespace WorldPartitionTests
 			TestTrue(TEXT("Reference soft refcount"), Reference->GetSoftRefCount() == 0);
 			TestTrue(TEXT("Reference hard refcount"), Reference->GetHardRefCount() == 1);
 			{
-				FWorldPartitionHardRef ReferenceCopy(MoveTemp(Reference));
+				FWorldPartitionReference ReferenceCopy(MoveTemp(Reference));
 				TestTrue(TEXT("Reference move src not valid"), !Reference.IsValid());
 				TestTrue(TEXT("Reference move dst valid"), ReferenceCopy.IsValid());
 				TestTrue(TEXT("Reference soft refcount"), ReferenceCopy->GetSoftRefCount() == 0);
@@ -187,7 +190,7 @@ namespace WorldPartitionTests
 
 			// Handle reference move
 			{
-				FWorldPartitionSoftRef HandleFromReference(MoveTemp(Reference));
+				FWorldPartitionHandle HandleFromReference(MoveTemp(Reference));
 				TestTrue(TEXT("Handle move src not valid"), !Reference.IsValid());
 				TestTrue(TEXT("Handle move dst valid"), HandleFromReference.IsValid());
 				TestTrue(TEXT("Handle soft refcount"), HandleFromReference->GetSoftRefCount() == 1);
@@ -202,7 +205,7 @@ namespace WorldPartitionTests
 
 			// Reference handle move
 			{
-				FWorldPartitionHardRef ReferenceFromHandle(MoveTemp(Handle));
+				FWorldPartitionReference ReferenceFromHandle(MoveTemp(Handle));
 				TestTrue(TEXT("Reference move src not valid"), !Handle.IsValid());
 				TestTrue(TEXT("Reference move dst valid"), ReferenceFromHandle.IsValid());
 				TestTrue(TEXT("Reference soft refcount"), ReferenceFromHandle->GetSoftRefCount() == 0);
diff --git a/Engine/Source/Runtime/Engine/Private/WorldPartition/HLOD/HLODBuilder.cpp b/Engine/Source/Runtime/Engine/Private/WorldPartition/HLOD/HLODBuilder.cpp
--- a/Engine/Source/Runtime/Engine/Private/WorldPartition/HLOD/HLODBuilder.cpp
+++ b/Engine/Source/Runtime/Engine/Private/WorldPartition/HLOD/HLODBuilder.cpp
@@ -41,7 +41,7 @@ public:
 		uint64 CellHash = FHLODActorDesc::ComputeCellHash(HLODLayer->GetName(), Context->GridIndexX, Context->GridIndexY, Context->GridIndexZ);
 
 		int32 HLODActorRefIndex = INDEX_NONE;
-		FWorldPartitionSoftRef HLODActorHandle;
+		FWorldPartitionHandle HLODActorHandle;
 		if (Context->HLODActorDescs.RemoveAndCopyValue(CellHash, HLODActorHandle))
 		{
 			HLODActorRefIndex = Context->ActorReferences.Add(HLODActorHandle);
@@ -75,7 +75,7 @@ public:
 		{
 			if (HLODActorRefIndex != INDEX_NONE)
 			{
-				Context->HLODActorDescs.Add(CellHash, FWorldPartitionSoftRef(WorldPartition, HLODActor->GetActorGuid()));
+				Context->HLODActorDescs.Add(CellHash, FWorldPartitionHandle(WorldPartition, HLODActor->GetActorGuid()));
 				Context->ActorReferences.RemoveAtSwap(HLODActorRefIndex);
 			}
 			else
diff --git a/Engine/Source/Runtime/Engine/Private/WorldPartition/HLOD/HLODBuilder.h b/Engine/Source/Runtime/Engine/Private/WorldPartition/HLOD/HLODBuilder.h
--- a/Engine/Source/Runtime/Engine/Private/WorldPartition/HLOD/HLODBuilder.h
+++ b/Engine/Source/Runtime/Engine/Private/WorldPartition/HLOD/HLODBuilder.h
@@ -15,8 +15,8 @@ class AWorldPartitionHLOD;
 
 struct ENGINE_API FHLODGenerationContext
 {
-	TMap<uint64, FWorldPartitionSoftRef> HLODActorDescs;
-	TArray<FWorldPartitionHardRef> ActorReferences;
+	TMap<uint64, FWorldPartitionHandle> HLODActorDescs;
+	TArray<FWorldPartitionReference> ActorReferences;
 	
 	// Everything needed to build the cell hash
 	int64 GridIndexX;
diff --git a/Engine/Source/Runtime/Engine/Private/WorldPartition/Landscape/LandscapeActorDesc.cpp b/Engine/Source/Runtime/Engine/Private/WorldPartition/Landscape/LandscapeActorDesc.cpp
--- a/Engine/Source/Runtime/Engine/Private/WorldPartition/Landscape/LandscapeActorDesc.cpp
+++ b/Engine/Source/Runtime/Engine/Private/WorldPartition/Landscape/LandscapeActorDesc.cpp
@@ -36,7 +36,7 @@ void FLandscapeActorDesc::OnRegister(UWorldPartition* WorldPartition)
 
 	if (ULandscapeInfo* LandscapeInfo = ULandscapeInfo::Find(WorldPartition->GetWorld(), GridGuid))
 	{
-		FWorldPartitionSoftRef Handle(WorldPartition, GetGuid());
+		FWorldPartitionHandle Handle(WorldPartition, GetGuid());
 		check(Handle.IsValid());
 		LandscapeInfo->ProxyHandles.Add(MoveTemp(Handle));
 	}
@@ -48,7 +48,7 @@ void FLandscapeActorDesc::OnUnregister(UWorldPartition* WorldPartition)
 
 	if (ULandscapeInfo* LandscapeInfo = ULandscapeInfo::Find(WorldPartition->GetWorld(), GridGuid))
 	{
-		FWorldPartitionSoftRef Handle(WorldPartition, GetGuid());
+		FWorldPartitionHandle Handle(WorldPartition, GetGuid());
 		check(Handle.IsValid());
 		LandscapeInfo->ProxyHandles.Remove(Handle);
 	}
diff --git a/Engine/Source/Runtime/Engine/Private/WorldPartition/Landscape/LandscapeSplineActorDesc.cpp b/Engine/Source/Runtime/Engine/Private/WorldPartition/Landscape/LandscapeSplineActorDesc.cpp
--- a/Engine/Source/Runtime/Engine/Private/WorldPartition/Landscape/LandscapeSplineActorDesc.cpp
+++ b/Engine/Source/Runtime/Engine/Private/WorldPartition/Landscape/LandscapeSplineActorDesc.cpp
@@ -37,7 +37,7 @@ void FLandscapeSplineActorDesc::OnRegister(UWorldPartition* WorldPartition)
 
 	if (ULandscapeInfo* LandscapeInfo = ULandscapeInfo::Find(WorldPartition->GetWorld(), LandscapeGuid))
 	{
-		FWorldPartitionSoftRef Handle(WorldPartition, GetGuid());
+		FWorldPartitionHandle Handle(WorldPartition, GetGuid());
 		check(Handle.IsValid());
 		LandscapeInfo->SplineHandles.Add(MoveTemp(Handle));
 	}
@@ -49,7 +49,7 @@ void FLandscapeSplineActorDesc::OnUnregister(UWorldPartition* WorldPartition)
 
 	if (ULandscapeInfo* LandscapeInfo = ULandscapeInfo::Find(WorldPartition->GetWorld(), LandscapeGuid))
 	{
-		FWorldPartitionSoftRef Handle(WorldPartition, GetGuid());
+		FWorldPartitionHandle Handle(WorldPartition, GetGuid());
 		check(Handle.IsValid());
 		LandscapeInfo->SplineHandles.Remove(Handle);
 	}
diff --git a/Engine/Source/Runtime/Engine/Private/WorldPartition/WorldPartitionHandle.cpp b/Engine/Source/Runtime/Engine/Private/WorldPartition/WorldPartitionHandle.cpp
--- a/Engine/Source/Runtime/Engine/Private/WorldPartition/WorldPartitionHandle.cpp
+++ b/Engine/Source/Runtime/Engine/Private/WorldPartition/WorldPartitionHandle.cpp
@@ -16,6 +16,13 @@ TUniquePtr<FWorldPartitionActorDesc>* FWorldPartitionHandleUtils::GetActorDesc(U
 
 bool FWorldPartitionHandleUtils::IsActorDescLoaded(FWorldPartitionActorDesc* ActorDesc)
 {
+#if WITH_DEV_AUTOMATION_TESTS
+	if (GIsAutomationTesting)
+	{
+		return ActorDesc->GetHardRefCount() > 0;
+	}
+#endif
+
 	return !!ActorDesc->GetActor();
 }
 
@@ -58,30 +65,4 @@ void FWorldPartitionHardRefImpl::DecRefCount(FWorldPartitionActorDesc* ActorDesc
 		}
 	}
 }
-
-void FWorldPartitionPinRefImpl::IncRefCount(FWorldPartitionActorDesc* ActorDesc)
-{
-	bIsReference = ActorDesc->GetHardRefCount() > 0;
-
-	if (bIsReference)
-	{
-		FWorldPartitionHardRefImpl::IncRefCount(ActorDesc);
-	}
-	else
-	{
-		FWorldPartitionSoftRefImpl::IncRefCount(ActorDesc);
-	}
-}
-
-void FWorldPartitionPinRefImpl::DecRefCount(FWorldPartitionActorDesc* ActorDesc)
-{
-	if (bIsReference)
-	{
-		FWorldPartitionHardRefImpl::DecRefCount(ActorDesc);
-	}
-	else
-	{
-		FWorldPartitionSoftRefImpl::DecRefCount(ActorDesc);
-	}
-}
 #endif
\ No newline at end of file
diff --git a/Engine/Source/Runtime/Engine/Private/WorldPartition/WorldPartitionRuntimeSpatialHash.cpp b/Engine/Source/Runtime/Engine/Private/WorldPartition/WorldPartitionRuntimeSpatialHash.cpp
--- a/Engine/Source/Runtime/Engine/Private/WorldPartition/WorldPartitionRuntimeSpatialHash.cpp
+++ b/Engine/Source/Runtime/Engine/Private/WorldPartition/WorldPartitionRuntimeSpatialHash.cpp
@@ -1594,7 +1594,7 @@ TArray<FGuid> GenerateHLODsForGrid(UWorldPartition* WorldPartition, const FSpati
 				TArray<AActor*> CellActors;
 				for (const FGuid& ActorGuid: GridCellDataChunk->GetActors())
 				{
-					FWorldPartitionHardRef& ActorRef = Context.ActorReferences.Emplace_GetRef(WorldPartition, ActorGuid);
+					FWorldPartitionReference& ActorRef = Context.ActorReferences.Emplace_GetRef(WorldPartition, ActorGuid);
 					FWorldPartitionActorDesc* ActorDesc  = ActorRef.Get();
 					CellActors.Add(ActorDesc->GetActor());
 				}
@@ -1887,10 +1887,10 @@ bool UWorldPartitionRuntimeSpatialHash::GenerateHLOD(ISourceControlHelper* Sourc
 
 	// Create HLODs generation context
 	FHLODGenerationContext Context;
-	TArray<FWorldPartitionSoftRef> InvalidHLODActors;
+	TArray<FWorldPartitionHandle> InvalidHLODActors;
 	for (TWorldPartitionActorDescIterator<AWorldPartitionHLOD, FHLODActorDesc> HLODIterator(WorldPartition); HLODIterator; ++HLODIterator)
 	{
-		FWorldPartitionSoftRef HLODActorHandle(WorldPartition, HLODIterator->GetGuid());
+		FWorldPartitionHandle HLODActorHandle(WorldPartition, HLODIterator->GetGuid());
 
 		if (HLODIterator->GetCellHash())
 		{
@@ -1974,7 +1974,7 @@ bool UWorldPartitionRuntimeSpatialHash::GenerateHLOD(ISourceControlHelper* Sourc
 		GenerateHLODs(HLODGrids[HLODGridName], GridsDepth[HLODGridName] + 1, HLODActorClusters);
 	}
 
-	auto DeleteHLODActor = [&SourceControlHelper](FWorldPartitionSoftRef ActorHandle)
+	auto DeleteHLODActor = [&SourceControlHelper](FWorldPartitionHandle ActorHandle)
 	{
 		FWorldPartitionActorDesc* HLODActorDesc = ActorHandle.Get();
 		check(HLODActorDesc);
diff --git a/Engine/Source/Runtime/Engine/Public/WorldPartition/WorldPartitionHandle.h b/Engine/Source/Runtime/Engine/Public/WorldPartition/WorldPartitionHandle.h
--- a/Engine/Source/Runtime/Engine/Public/WorldPartition/WorldPartitionHandle.h
+++ b/Engine/Source/Runtime/Engine/Public/WorldPartition/WorldPartitionHandle.h
@@ -15,19 +15,16 @@ struct ENGINE_API FWorldPartitionHandleUtils
 };
 
 template <typename Impl>
-class ENGINE_API TWorldPartitionHandle : protected Impl
+class ENGINE_API TWorldPartitionHandle
 {
-	friend struct FWorldPartitionSoftRefHelpers;
-
 public:
 	FORCEINLINE TWorldPartitionHandle()
 		: ActorDesc(nullptr)
 	{}
 
 	FORCEINLINE TWorldPartitionHandle(TUniquePtr<FWorldPartitionActorDesc>* InActorDesc)
+		: ActorDesc(InActorDesc)
 	{
-		ActorDesc = InActorDesc;
-
 		if (IsValid())
 		{
 			IncRefCount();
@@ -35,9 +32,8 @@ public:
 	}
 
 	FORCEINLINE TWorldPartitionHandle(UWorldPartition* WorldPartition, const FGuid& ActorGuid)
+		: ActorDesc(FWorldPartitionHandleUtils::GetActorDesc(WorldPartition, ActorGuid))
 	{
-		ActorDesc = FWorldPartitionHandleUtils::GetActorDesc(WorldPartition, ActorGuid);
-
 		if (IsValid())
 		{
 			IncRefCount();
@@ -233,32 +229,33 @@ struct ENGINE_API FWorldPartitionHardRefImpl
 	static void DecRefCount(FWorldPartitionActorDesc* ActorDesc);
 };
 
-struct ENGINE_API FWorldPartitionPinRefImpl
-{
-	FWorldPartitionPinRefImpl() : bIsReference(false) {}
-	void IncRefCount(FWorldPartitionActorDesc* ActorDesc);
-	void DecRefCount(FWorldPartitionActorDesc* ActorDesc);
-	bool bIsReference;
-};
-
 /**
- * FWorldPartitionSoftRef will increment/decrement the soft reference count on the actor descriptor.
+ * FWorldPartitionHandle will increment/decrement the soft reference count on the actor descriptor.
  * This won't trigger any loading, but will prevent cleanup of the actor descriptor when destroying an
  * actor in the editor.
  */
-typedef TWorldPartitionHandle<FWorldPartitionSoftRefImpl> FWorldPartitionSoftRef;
+typedef TWorldPartitionHandle<FWorldPartitionSoftRefImpl> FWorldPartitionHandle;
 
 /**
- * FWorldPartitionHardRef will increment/decrement the hard reference count on the actor descriptor.
+ * FWorldPartitionReference will increment/decrement the hard reference count on the actor descriptor.
  * This will trigger actor loading/unloading when the hard reference counts gets to one/zero.
  */
-typedef TWorldPartitionHandle<FWorldPartitionHardRefImpl> FWorldPartitionHardRef;
+typedef TWorldPartitionHandle<FWorldPartitionHardRefImpl> FWorldPartitionReference;
 
 /**
- * FWorldPartitionPinRef will act as a softref when there's not hard reference count to the actor descriptor,
- * and like a hardref when there is. This is useful when you want to keep an actor loaded during special
- * operations, and don't trigger loading when the actor is not loaded. Mainly intended to be on stack, as a
- * scoped operation.
+ * FWorldPartitionHandlePinRefScope will keep a reference if the actor is already loaded. This is useful 
+ * when you want to keep an actor loaded during special operations, but don't trigger loading when the 
+ * actor is not loaded. Intended to be on stack, as a scoped operation.
  */
-typedef TWorldPartitionHandle<FWorldPartitionPinRefImpl> FWorldPartitionPinRef;
+struct FWorldPartitionHandlePinRefScope
+{
+	FWorldPartitionHandlePinRefScope(const FWorldPartitionHandle& Handle)
+	{
+		if (Handle.IsLoaded())
+		{
+			Reference = Handle;
+		}
+	}
+	FWorldPartitionReference Reference;
+};
 #endif
\ No newline at end of file
diff --git a/Engine/Source/Runtime/Landscape/Classes/LandscapeInfo.h b/Engine/Source/Runtime/Landscape/Classes/LandscapeInfo.h
--- a/Engine/Source/Runtime/Landscape/Classes/LandscapeInfo.h
+++ b/Engine/Source/Runtime/Landscape/Classes/LandscapeInfo.h
@@ -152,9 +152,9 @@ private:
 
 	// World Partition handles to actors with the same LandscapeGuid
 	friend class FLandscapeActorDesc;
-	TSet<FWorldPartitionSoftRef> ProxyHandles;
+	TSet<FWorldPartitionHandle> ProxyHandles;
 	friend class FLandscapeSplineActorDesc;
-	TSet<FWorldPartitionSoftRef> SplineHandles;
+	TSet<FWorldPartitionHandle> SplineHandles;
 #endif
 
 	TSet<ULandscapeComponent*> SelectedComponents;
@@ -293,7 +293,7 @@ public:
 	LANDSCAPE_API uint32 GetGridSize(uint32 InGridSizeInComponents) const;
 
 	/** Returns LandscapeSplineActor WorldPartition handles */
-	const TSet<FWorldPartitionSoftRef>& GetSplineHandles() const { return SplineHandles; }
+	const TSet<FWorldPartitionHandle>& GetSplineHandles() const { return SplineHandles; }
 #endif
 	LANDSCAPE_API static ULandscapeInfo* Find(UWorld* InWorld, const FGuid& LandscapeGuid);
 	static ULandscapeInfo* FindOrCreate(UWorld* InWorld, const FGuid& LandscapeGuid);
diff --git a/Engine/Source/Runtime/Landscape/Classes/LandscapeStreamingProxy.h b/Engine/Source/Runtime/Landscape/Classes/LandscapeStreamingProxy.h
--- a/Engine/Source/Runtime/Landscape/Classes/LandscapeStreamingProxy.h
+++ b/Engine/Source/Runtime/Landscape/Classes/LandscapeStreamingProxy.h
@@ -28,7 +28,7 @@ public:
 
 #if WITH_EDITORONLY_DATA
 	/** hard refs to actors that need to be loaded when this proxy is loaded */
-	TSet<FWorldPartitionHardRef> ActorDescReferences;
+	TSet<FWorldPartitionReference> ActorDescReferences;
 #endif
 
 	//~ Begin UObject Interface
diff --git a/Engine/Source/Runtime/Landscape/Private/Landscape.cpp b/Engine/Source/Runtime/Landscape/Private/Landscape.cpp
--- a/Engine/Source/Runtime/Landscape/Private/Landscape.cpp
+++ b/Engine/Source/Runtime/Landscape/Private/Landscape.cpp
@@ -3123,11 +3123,11 @@ void ULandscapeInfo::Initialize(UWorld* InWorld, const FGuid& InLandscapeGuid)
 		{
 			if (ActorDescIterator->GetActorClass()->IsChildOf(ALandscapeStreamingProxy::StaticClass()))
 			{
-				ProxyHandles.Add(FWorldPartitionSoftRef(WorldPartition, ActorDescIterator->GetGuid()));
+				ProxyHandles.Add(FWorldPartitionHandle(WorldPartition, ActorDescIterator->GetGuid()));
 			}
 			else if (ActorDescIterator->GetActorClass()->IsChildOf(ALandscapeSplineActor::StaticClass()))
 			{
-				SplineHandles.Emplace(FWorldPartitionSoftRef(WorldPartition, ActorDescIterator->GetGuid()));
+				SplineHandles.Emplace(FWorldPartitionHandle(WorldPartition, ActorDescIterator->GetGuid()));
 			}
 
 		}
diff --git a/Engine/Source/Runtime/Landscape/Private/LandscapeEdit.cpp b/Engine/Source/Runtime/Landscape/Private/LandscapeEdit.cpp
--- a/Engine/Source/Runtime/Landscape/Private/LandscapeEdit.cpp
+++ b/Engine/Source/Runtime/Landscape/Private/LandscapeEdit.cpp
@@ -2251,7 +2251,7 @@ bool ULandscapeInfo::HasUnloadedComponentsInRegion(int32 X1, int32 Y1, int32 X2,
 	UActorPartitionSubsystem::FCellCoord MinCoord = UActorPartitionSubsystem::FCellCoord::GetCellCoord(FIntPoint(ComponentIndexX1 * ComponentSizeQuads, ComponentIndexY1 * ComponentSizeQuads), World->PersistentLevel, LandscapeActor->GridSize);
 	UActorPartitionSubsystem::FCellCoord MaxCoord = UActorPartitionSubsystem::FCellCoord::GetCellCoord(FIntPoint(ComponentIndexX2 * ComponentSizeQuads, ComponentIndexY2 * ComponentSizeQuads), World->PersistentLevel, LandscapeActor->GridSize);
 
-	for (const FWorldPartitionSoftRef& Handle : ProxyHandles)
+	for (const FWorldPartitionHandle& Handle : ProxyHandles)
 	{
 		if (FLandscapeActorDesc* LandscapeActorDesc = (FLandscapeActorDesc*)Handle.Get())
 		{
@@ -5029,13 +5029,13 @@ void ALandscapeStreamingProxy::PostRegisterAllComponents()
 		check(LandscapeInfo);
 		if (GEditor && !GetWorld()->IsGameWorld())
 		{
-			const TSet<FWorldPartitionSoftRef>& SplineHandles = LandscapeInfo->GetSplineHandles();
+			const TSet<FWorldPartitionHandle>& SplineHandles = LandscapeInfo->GetSplineHandles();
 			if (SplineHandles.Num())
 			{
 				FVector ActorLocation = GetActorLocation();
 				FBox Bounds(ActorLocation, ActorLocation + (GridSize * LandscapeInfo->DrawScale));
 				// Get a reference to Spline Actors that have intersecting bounds with us
-				for (const FWorldPartitionSoftRef& SplineHandle : LandscapeInfo->GetSplineHandles())
+				for (const FWorldPartitionHandle& SplineHandle : LandscapeInfo->GetSplineHandles())
 				{
 					if (SplineHandle.IsValid() && Bounds.IntersectXY(SplineHandle->GetBounds()))
 					{
