World partition editor hash and editor cells now uses actor descriptor handles and references to handle cells loading.
This greatly streamline the loading/unloading code by hiding all the loading specific code into the actor descriptors.

diff --git a/Engine/Source/Runtime/Engine/Private/Level.cpp b/Engine/Source/Runtime/Engine/Private/Level.cpp
--- a/Engine/Source/Runtime/Engine/Private/Level.cpp
+++ b/Engine/Source/Runtime/Engine/Private/Level.cpp
@@ -56,6 +56,7 @@ Level.cpp: Level-related functions
 #include "AssetData.h"
 #include "Misc/ArchiveMD5.h"
 #endif
+#include "WorldPartition/WorldPartition.h"
 #include "Engine/LevelStreaming.h"
 #include "LevelUtils.h"
 #include "Components/ModelComponent.h"
@@ -376,6 +377,11 @@ void ULevel::CleanupLevel()
 			}, false);
 		}
 	}
+
+	if (UWorldPartition* WorldPartition = GetWorldPartition())
+	{
+		WorldPartition->CleanupWorldPartition();
+	}
 }
 
 void ULevel::PostInitProperties()
@@ -632,6 +638,7 @@ void ULevel::RemoveLoadedActor(AActor* Actor)
 {
 	check(Actor);
 	check(Actor->GetLevel() == this);
+	check(!Actor->IsPendingKill());
 
 	Actor->UnregisterAllComponents();
 	Actor->RegisterAllActorTickFunctions(false, true);	
diff --git a/Engine/Source/Runtime/Engine/Private/WorldPartition/Landscape/LandscapeActorDesc.cpp b/Engine/Source/Runtime/Engine/Private/WorldPartition/Landscape/LandscapeActorDesc.cpp
--- a/Engine/Source/Runtime/Engine/Private/WorldPartition/Landscape/LandscapeActorDesc.cpp
+++ b/Engine/Source/Runtime/Engine/Private/WorldPartition/Landscape/LandscapeActorDesc.cpp
@@ -30,9 +30,9 @@ void FLandscapeActorDesc::Unload()
 	FPartitionActorDesc::Unload();
 }
 
-void FLandscapeActorDesc::OnRegister(UWorldPartition* WorldPartition)
+void FLandscapeActorDesc::OnRegister()
 {
-	FPartitionActorDesc::OnRegister(WorldPartition);
+	FPartitionActorDesc::OnRegister();
 
 	if (ULandscapeInfo* LandscapeInfo = ULandscapeInfo::Find(WorldPartition->GetWorld(), GridGuid))
 	{
@@ -42,9 +42,9 @@ void FLandscapeActorDesc::OnRegister(UWorldPartition* WorldPartition)
 	}
 }
 
-void FLandscapeActorDesc::OnUnregister(UWorldPartition* WorldPartition)
+void FLandscapeActorDesc::OnUnregister()
 {
-	FPartitionActorDesc::OnUnregister(WorldPartition);
+	FPartitionActorDesc::OnUnregister();
 
 	if (ULandscapeInfo* LandscapeInfo = ULandscapeInfo::Find(WorldPartition->GetWorld(), GridGuid))
 	{
diff --git a/Engine/Source/Runtime/Engine/Private/WorldPartition/Landscape/LandscapeSplineActorDesc.cpp b/Engine/Source/Runtime/Engine/Private/WorldPartition/Landscape/LandscapeSplineActorDesc.cpp
--- a/Engine/Source/Runtime/Engine/Private/WorldPartition/Landscape/LandscapeSplineActorDesc.cpp
+++ b/Engine/Source/Runtime/Engine/Private/WorldPartition/Landscape/LandscapeSplineActorDesc.cpp
@@ -31,9 +31,9 @@ void FLandscapeSplineActorDesc::Serialize(FArchive& Ar)
 	}
 }
 
-void FLandscapeSplineActorDesc::OnRegister(UWorldPartition* WorldPartition)
+void FLandscapeSplineActorDesc::OnRegister()
 {
-	FWorldPartitionActorDesc::OnRegister(WorldPartition);
+	FWorldPartitionActorDesc::OnRegister();
 
 	if (ULandscapeInfo* LandscapeInfo = ULandscapeInfo::Find(WorldPartition->GetWorld(), LandscapeGuid))
 	{
@@ -43,9 +43,9 @@ void FLandscapeSplineActorDesc::OnRegister(UWorldPartition* WorldPartition)
 	}
 }
 
-void FLandscapeSplineActorDesc::OnUnregister(UWorldPartition* WorldPartition)
+void FLandscapeSplineActorDesc::OnUnregister()
 {
-	FWorldPartitionActorDesc::OnUnregister(WorldPartition);
+	FWorldPartitionActorDesc::OnUnregister();
 
 	if (ULandscapeInfo* LandscapeInfo = ULandscapeInfo::Find(WorldPartition->GetWorld(), LandscapeGuid))
 	{
diff --git a/Engine/Source/Runtime/Engine/Private/WorldPartition/WorldPartition.cpp b/Engine/Source/Runtime/Engine/Private/WorldPartition/WorldPartition.cpp
--- a/Engine/Source/Runtime/Engine/Private/WorldPartition/WorldPartition.cpp
+++ b/Engine/Source/Runtime/Engine/Private/WorldPartition/WorldPartition.cpp
@@ -17,15 +17,11 @@
 #include "Engine/LevelStreaming.h"
 #include "GameFramework/WorldSettings.h"
 #include "ProfilingDebugging/ScopedTimers.h"
-#include "LevelUtils.h"
 
 #if WITH_EDITOR
 #include "Editor.h"
-#include "Editor/GroupActor.h"
-#include "EditorLevelUtils.h"
 #include "FileHelpers.h"
 #include "HAL/FileManager.h"
-#include "Misc/Base64.h"
 #include "Misc/ScopedSlowTask.h"
 #include "Misc/ScopeExit.h"
 #include "ScopedTransaction.h"
@@ -329,7 +325,7 @@ void UWorldPartition::Uninitialize()
 		{
 			CleanupForPIE();
 		}
-		else if (IsMainWorldPartition())
+		else
 		{
 			if (!IsEngineExitRequested())
 			{
@@ -358,6 +354,11 @@ void UWorldPartition::Uninitialize()
 	}
 }
 
+void UWorldPartition::CleanupWorldPartition()
+{
+	Uninitialize();
+}
+
 bool UWorldPartition::IsInitialized() const
 {
 	return InitState == EWorldPartitionInitState::Initialized;
@@ -401,17 +402,6 @@ void UWorldPartition::UnregisterDelegates()
 	}
 }
 
-void UWorldPartition::ApplyActorTransform(AActor* InActor, const FTransform& InTransform)
-{
-	if (!InTransform.Equals(FTransform::Identity))
-	{
-		FLevelUtils::FApplyLevelTransformParams TransformParams(InActor->GetLevel(), InTransform);
-		TransformParams.Actor = InActor;
-		TransformParams.bDoPostEditMove = true;
-		FLevelUtils::ApplyLevelTransform(TransformParams);
-	}
-}
-
 void UWorldPartition::ForEachIntersectingActorDesc(const FBox& Box, TSubclassOf<AActor> ActorClass, TFunctionRef<bool(const FWorldPartitionActorDesc*)> Predicate) const
 {
 	EditorHash->ForEachIntersectingActor(Box, [&ActorClass, Predicate](const FWorldPartitionActorDesc* ActorDesc)
@@ -481,7 +471,6 @@ bool UWorldPartition::RefreshLoadedEditorCells()
 {
 	auto GetCellsToRefresh = [this](TArray<UWorldPartitionEditorCell*>& CellsToRefresh)
 	{
-		CellsToRefresh.Add(EditorHash->GetAlwaysLoadedCell());
 		EditorHash->ForEachCell([&CellsToRefresh](UWorldPartitionEditorCell* Cell)
 		{
 			if (Cell->bLoaded)
@@ -520,11 +509,11 @@ bool UWorldPartition::UpdateEditorCells(TFunctionRef<bool(TArray<UWorldPartition
 
 	for (UWorldPartitionEditorCell* Cell : CellsToProcess)
 	{
-		for (FWorldPartitionActorDesc* ActorDesc : Cell->LoadedActors)
+		for (const FWorldPartitionReference& ActorDesc : Cell->LoadedActors)
 		{
-			if (!bIsCellShouldBeLoaded || !ShouldActorBeLoaded(ActorDesc))
+			if (!bIsCellShouldBeLoaded || !ShouldActorBeLoaded(*ActorDesc))
 			{
-				UnloadCount.FindOrAdd(ActorDesc, 0)++;
+				UnloadCount.FindOrAdd(*ActorDesc, 0)++;
 			}
 		}
 	}
@@ -636,81 +625,36 @@ void UWorldPartition::UpdateLoadingEditorCell(UWorldPartitionEditorCell* Cell, b
 
 	Cell->Modify(false);
 
-	auto UnloadActor = [this](FWorldPartitionActorDesc* ActorDesc, AActor* Actor)
-	{
-		check(Actor);
-
-		const uint32 ActorRefCount = ActorDesc->DecHardRefCount();
-		UE_LOG(LogWorldPartition, Verbose, TEXT(" ==> Unreferenced loaded actor %s(%d) [UWorldPartition::UpdateLoadingEditorCell]"), *Actor->GetFullName(), ActorRefCount);
-
-		if (!ActorRefCount)
-		{
-			UnregisterActor(Actor);
-		}
-	};
+	bool bPotentiallyUnloadedActors = false;
 
 	if (!bShouldBeLoaded)
 	{
-		for (FWorldPartitionActorDesc* ActorDesc: Cell->LoadedActors.Array())
-		{
-			// Here we test for null Actor only for when UpdateLoadingEditorCell is called by UWorldPartition::Uninitialize() when exiting the editor (see StaticExit() / GExitPurge).
-			AActor* Actor = ActorDesc->GetActor();
-			check(Actor || GExitPurge);
-			if (Actor)
-			{
-				UnloadActor(ActorDesc, Actor);
-			}
-		}
-
 		Cell->LoadedActors.Empty();
+		bPotentiallyUnloadedActors = true;
 	}
 	else
 	{
 		TGuardValue<bool> IsEditorLoadingPackageGuard(GIsEditorLoadingPackage, true);
 
 		TArray<FWorldPartitionActorDesc*> ActorsToUnload;
-		for (FWorldPartitionActorDesc* ActorDesc: Cell->Actors.Array())
+		for (FWorldPartitionHandle& ActorHandle: Cell->Actors.Array())
 		{
-			AActor* Actor = ActorDesc->GetActor();
+			AActor* Actor = ActorHandle->GetActor();
 
 			// The actor could be either loaded but with no cells loaded (in the case of a reference from another actor, for example)
 			// or directly referenced by a loaded cell.
-			check(Actor || !ActorDesc->GetHardRefCount());
+			check(Actor || !ActorHandle->GetHardRefCount());
 
 			// Filter actor against DataLayers
-			if (ShouldActorBeLoaded(ActorDesc))
+			if (ShouldActorBeLoaded(*ActorHandle))
 			{
 				bool bIsAlreadyInLoadedActors = false;
-				Cell->LoadedActors.Add(ActorDesc, &bIsAlreadyInLoadedActors);
-
-				if (bIsAlreadyInLoadedActors)
-				{
-					// We already hold a reference to this actor
-					check(Actor);
-					check(ActorDesc->GetHardRefCount());
-					UE_LOG(LogWorldPartition, Verbose, TEXT(" ==> Skipped already loaded actor %s"), *Actor->GetFullName());
-				}
-				else
-				{
-					const uint32 ActorRefCount = ActorDesc->IncHardRefCount();
-
-					if (ActorRefCount == 1)
-					{
-						// Register actor
-						Actor = RegisterActor(ActorDesc);
-					}
-					else
-					{
-						check(Actor);
-						UE_LOG(LogWorldPartition, Verbose, TEXT(" ==> Referenced unloaded actor %s(%d)"), *Actor->GetFullName(), ActorRefCount);
-					}
-
-					check(Actor && Actor->GetActorGuid() == ActorDesc->GetGuid());
-				}
+				Cell->LoadedActors.Add(ActorHandle, &bIsAlreadyInLoadedActors);
 			}
-			else if (Cell->LoadedActors.Remove(ActorDesc))
+			else
 			{
-				UnloadActor(ActorDesc, Actor);
+				Cell->LoadedActors.Remove(ActorHandle);
+				bPotentiallyUnloadedActors = true;
 			}
 		}
 	}
@@ -728,66 +672,14 @@ void UWorldPartition::UpdateLoadingEditorCell(UWorldPartitionEditorCell* Cell, b
 			EditorHash->OnCellUnloaded(Cell);
 		}
 	}
-}
-
-#endif
-
-AActor* UWorldPartition::RegisterActor(FWorldPartitionActorDesc* ActorDesc)
-{
-#if WITH_EDITOR
-	check(ActorDesc);
-	AActor* Actor = ActorDesc->GetActor();
-
-	if (!Actor)
-	{
-		Actor = ActorDesc->Load(InstancingContext.IsInstanced() ? &InstancingContext : nullptr);
-		UE_LOG(LogWorldPartition, Verbose, TEXT(" ==> Loaded %s"), *Actor->GetFullName());
-	}
-
-	check(Actor);
-	check(Actor->IsPackageExternal());
-	check(ActorDesc->GetActor() == Actor);
-
-	ApplyActorTransform(Actor, InstanceTransform);
-
-	TGuardValue<ITransaction*> TransGuard(GUndo, nullptr);
-	Actor->GetLevel()->AddLoadedActor(Actor);
-
-	OnActorRegisteredEvent.Broadcast(*Actor, true);
-
-	UE_LOG(LogWorldPartition, Verbose, TEXT(" ==> Registered loaded actor %s"), *Actor->GetFullName());
-
-	return Actor;
-#else
-	return nullptr;
-#endif	
-}
-
-void UWorldPartition::UnregisterActor(AActor* Actor)
-{
-#if WITH_EDITOR
-	FWorldPartitionActorDesc* ActorDesc = GetActorDesc(Actor->GetActorGuid());
-	check(ActorDesc);
 
-	if (!Actor->IsPendingKill())
+	if (bPotentiallyUnloadedActors)
 	{
-		OnActorRegisteredEvent.Broadcast(*Actor, false);
-
-		Actor->GetLevel()->RemoveLoadedActor(Actor);
+		bForceGarbageCollection = true;
+		bForceGarbageCollectionPurge = true;
 	}
-
-	ActorDesc->Unload();
-
-	ApplyActorTransform(Actor, InstanceTransform.Inverse());
-
-	bForceGarbageCollection = true;
-	bForceGarbageCollectionPurge = true;
-
-	UE_LOG(LogWorldPartition, Verbose, TEXT(" ==> Unregistered loaded actor %s"), *Actor->GetFullName());
-#endif
 }
 
-#if WITH_EDITOR
 void UWorldPartition::OnAssetAdded(const FAssetData& InAssetData)
 {
 	if (ShouldHandleAssetEvent(InAssetData))
@@ -800,7 +692,7 @@ void UWorldPartition::OnAssetAdded(const FAssetData& InAssetData)
 
 		HashActorDesc(NewActorDesc->Get());
 
-		(*NewActorDesc)->OnRegister(this);
+		(*NewActorDesc)->OnRegister();
 	}
 }
 
@@ -815,7 +707,7 @@ void UWorldPartition::OnAssetRemoved(const FAssetData& InAssetData)
 
 		UnhashActorDesc(ExistingActorDesc->Get());
 
-		(*ExistingActorDesc)->OnUnregister(this);
+		(*ExistingActorDesc)->OnUnregister();
 
 		Actors.Remove((*ExistingActorDesc)->GetGuid());
 		ExistingActorDesc->Release();
@@ -828,12 +720,16 @@ void UWorldPartition::OnAssetUpdated(const FAssetData& InAssetData)
 	{
 		TUniquePtr<FWorldPartitionActorDesc> NewActorDesc = GetActorDescriptor(InAssetData);
 		check(NewActorDesc.IsValid());
-			
+
 		TUniquePtr<FWorldPartitionActorDesc>* ExistingActorDesc = Actors.FindChecked(NewActorDesc->GetGuid());
 
+		// Pin the actor handle on the actor to prevent unloading it when unhashing
+		FWorldPartitionHandle ExistingActorHandle(ExistingActorDesc);
+		FWorldPartitionHandlePinRefScope ExistingActorHandlePin(ExistingActorHandle);
+
 		UnhashActorDesc(ExistingActorDesc->Get());
 
-		// This is required to support external load references
+		// Transfer any reference count from external sources
 		NewActorDesc->TransferRefCounts(ExistingActorDesc->Get());
 
 		*ExistingActorDesc = MoveTemp(NewActorDesc);
@@ -903,7 +799,7 @@ TUniquePtr<FWorldPartitionActorDesc> UWorldPartition::GetActorDescriptor(const F
 			FBase64::Decode(ActorMetaDataStr, ActorDescInitData.SerializedData);
 
 			TUniquePtr<FWorldPartitionActorDesc> NewActorDesc(AActor::CreateClassActorDesc(ActorDescInitData.NativeClass));
-			NewActorDesc->Init(ActorDescInitData);
+			NewActorDesc->Init(this, ActorDescInitData);
 			return NewActorDesc;
 		}
 	}
@@ -915,14 +811,16 @@ void UWorldPartition::HashActorDesc(FWorldPartitionActorDesc* ActorDesc)
 {
 	check(ActorDesc);
 	check(EditorHash);
-	EditorHash->HashActor(ActorDesc);
+	FWorldPartitionHandle ActorHandle(this, ActorDesc->GetGuid());
+	EditorHash->HashActor(ActorHandle);
 }
 
 void UWorldPartition::UnhashActorDesc(FWorldPartitionActorDesc* ActorDesc)
 {
 	check(ActorDesc);
 	check(EditorHash);
-	EditorHash->UnhashActor(ActorDesc);
+	FWorldPartitionHandle ActorHandle(this, ActorDesc->GetGuid());
+	EditorHash->UnhashActor(ActorHandle);
 }
 #endif
 
@@ -1060,7 +958,7 @@ void UWorldPartition::DumpActorDescs(const FString& Path)
 			
 		ForEachActorDesc(AActor::StaticClass(), [LogFile, &LineEntry](const FWorldPartitionActorDesc* ActorDesc)
 		{
-				LineEntry = FString::Printf(
+			LineEntry = FString::Printf(
 				TEXT("%s, %s, %s, %.2f, %.2f, %.2f, %.2f, %.2f, %.2f") LINE_TERMINATOR, 
 				*ActorDesc->GetGuid().ToString(), 
 				*ActorDesc->GetClass().ToString(), 
diff --git a/Engine/Source/Runtime/Engine/Private/WorldPartition/WorldPartitionActorDesc.cpp b/Engine/Source/Runtime/Engine/Private/WorldPartition/WorldPartitionActorDesc.cpp
--- a/Engine/Source/Runtime/Engine/Private/WorldPartition/WorldPartitionActorDesc.cpp
+++ b/Engine/Source/Runtime/Engine/Private/WorldPartition/WorldPartitionActorDesc.cpp
@@ -3,6 +3,7 @@
 #include "WorldPartition/WorldPartitionActorDesc.h"
 
 #if WITH_EDITOR
+#include "LevelUtils.h"
 #include "Misc/HashBuilder.h"
 #include "UObject/LinkerInstancingContext.h"
 #include "UObject/UObjectHash.h"
@@ -24,6 +25,7 @@ uint32 FWorldPartitionActorDesc::GlobalTag = 0;
 FWorldPartitionActorDesc::FWorldPartitionActorDesc()
 	: SoftRefCount(0)
 	, HardRefCount(0)
+	, WorldPartition(nullptr)
 	, Tag(0)
 {}
 
@@ -38,6 +40,8 @@ void FWorldPartitionActorDesc::Init(const AActor* InActor)
 	ActorClass = GetParentNativeClass(InActor->GetClass());
 	Class = ActorClass->GetFName();
 
+	ActorPtr = (AActor*)InActor;
+
 	FVector NewBoundsLocation;
 	FVector NewBoundsExtent;
 	InActor->GetActorLocationBounds(/*bOnlyCollidingComponents*/false, NewBoundsLocation, NewBoundsExtent, /*bIncludeFromChildActors*/true);
@@ -72,9 +76,12 @@ void FWorldPartitionActorDesc::Init(const AActor* InActor)
 			References.Add(ActorReference->GetActorGuid());
 		}
 	}
+
+	WorldPartition = InActor->GetLevel()->GetWorldPartition();
+	check(WorldPartition);
 }
 
-void FWorldPartitionActorDesc::Init(const FWorldPartitionActorDescInitData& DescData)
+void FWorldPartitionActorDesc::Init(UWorldPartition* InWorldPartition, const FWorldPartitionActorDescInitData& DescData)
 {
 	ActorPackage = DescData.PackageName;
 	ActorPath = DescData.ActorPath;
@@ -98,6 +105,8 @@ void FWorldPartitionActorDesc::Init(const FWorldPartitionActorDescInitData& Desc
 	{
 		GridPlacement = DefaultGridPlacement;
 	}
+
+	WorldPartition = InWorldPartition;
 }
 
 void FWorldPartitionActorDesc::SerializeTo(TArray<uint8>& OutData)
@@ -135,6 +144,21 @@ void FWorldPartitionActorDesc::TransformInstance(const FString& From, const FStr
 	}
 }
 
+void FWorldPartitionActorDesc::ApplyActorTransform(const FTransform& InTransform)
+{
+	AActor* Actor = GetActor();
+	check(Actor);
+
+	if (!InTransform.Equals(FTransform::Identity))
+	{
+		FLevelUtils::FApplyLevelTransformParams TransformParams(Actor->GetLevel(), InTransform);
+		TransformParams.Actor = Actor;
+		TransformParams.bDoPostEditMove = true;
+		FLevelUtils::ApplyLevelTransform(TransformParams);
+	}
+}
+
+
 FString FWorldPartitionActorDesc::ToString() const
 {
 	return FString::Printf(TEXT("Guid:%s Class:%s Name:%s"), *Guid.ToString(), *Class.ToString(), *FPaths::GetExtension(ActorPath.ToString()));
@@ -189,12 +213,15 @@ AActor* FWorldPartitionActorDesc::GetActor() const
 	return Actor;
 }
 
-AActor* FWorldPartitionActorDesc::Load(const FLinkerInstancingContext* InstancingContext) const
+AActor* FWorldPartitionActorDesc::Load() const
 {
 	AActor* Actor = GetActor();
 
 	if (!Actor)
 	{
+		check(WorldPartition || GIsAutomationTesting);
+		const FLinkerInstancingContext* InstancingContext = WorldPartition ? &WorldPartition->InstancingContext : nullptr;
+
 		UPackage* Package = nullptr;
 
 		if (InstancingContext)
@@ -235,4 +262,25 @@ void FWorldPartitionActorDesc::Unload()
 		ActorPtr = nullptr;
 	}
 }
+
+void FWorldPartitionActorDesc::RegisterActor()
+{
+	AActor* Actor = GetActor();
+	check(Actor);
+
+	ApplyActorTransform(WorldPartition->InstanceTransform);
+	Actor->GetLevel()->AddLoadedActor(Actor);
+}
+
+void FWorldPartitionActorDesc::UnregisterActor()
+{
+	AActor* Actor = GetActor();
+	check(Actor);
+
+	if (!Actor->IsPendingKill())
+	{
+		Actor->GetLevel()->RemoveLoadedActor(Actor);
+		ApplyActorTransform(WorldPartition->InstanceTransform.Inverse());
+	}
+}
 #endif
\ No newline at end of file
diff --git a/Engine/Source/Runtime/Engine/Private/WorldPartition/WorldPartitionEditorCell.cpp b/Engine/Source/Runtime/Engine/Private/WorldPartition/WorldPartitionEditorCell.cpp
--- a/Engine/Source/Runtime/Engine/Private/WorldPartition/WorldPartitionEditorCell.cpp
+++ b/Engine/Source/Runtime/Engine/Private/WorldPartition/WorldPartitionEditorCell.cpp
@@ -20,39 +20,34 @@ void UWorldPartitionEditorCell::Serialize(FArchive& Ar)
 }
 
 #if WITH_EDITOR
-void UWorldPartitionEditorCell::AddActor(FWorldPartitionActorDesc* InActorDesc)
+void UWorldPartitionEditorCell::AddActor(const FWorldPartitionHandle& ActorHandle)
 {
-	check(InActorDesc);
+	check(ActorHandle.IsValid());
+	
 	bool bIsAlreadyInSet = false;
-	Actors.Add(InActorDesc, &bIsAlreadyInSet);
+	Actors.Add(ActorHandle, &bIsAlreadyInSet);
+	check(!bIsAlreadyInSet);
 
-	if (!bIsAlreadyInSet && bLoaded)
+	if (ActorHandle.IsLoaded())
 	{
-		if (AActor* Actor = InActorDesc->GetActor())
-		{
-			check(!Actor->IsChildActor());
-
-			bIsAlreadyInSet = false;
-			LoadedActors.Add(InActorDesc, &bIsAlreadyInSet);
-			check(!bIsAlreadyInSet);
-
-			const uint32 ActorRefCount = InActorDesc->IncHardRefCount();
-			UE_LOG(LogWorldPartition, Verbose, TEXT(" ==> Referenced loaded actor %s(%d) [UWorldPartitionEditorCell::AddActor]"), *Actor->GetFullName(), ActorRefCount);
-		}
+		bIsAlreadyInSet = false;
+		LoadedActors.Add(ActorHandle, &bIsAlreadyInSet);
+		check(!bIsAlreadyInSet);
 	}
 }
 
-void UWorldPartitionEditorCell::RemoveActor(FWorldPartitionActorDesc* InActorDesc)
+void UWorldPartitionEditorCell::RemoveActor(const FWorldPartitionHandle& ActorHandle)
 {
-	check(InActorDesc);
-	verify(Actors.Remove(InActorDesc));
+	check(ActorHandle.IsValid());
+	verify(Actors.Remove(ActorHandle));
 
-	if (LoadedActors.Remove(InActorDesc))
+	if (ActorHandle.IsLoaded())
 	{
-		check(bLoaded);
-
-		const uint32 ActorRefCount = InActorDesc->DecHardRefCount();
-		UE_LOG(LogWorldPartition, Verbose, TEXT(" ==> Unreferenced loaded actor %s(%d) [UWorldPartitionEditorCell::RemoveActor]"), *InActorDesc->GetActor()->GetFullName(), ActorRefCount);
+		verify(LoadedActors.Remove(ActorHandle));
+	}
+	else
+	{
+		check(!LoadedActors.Contains(ActorHandle));
 	}
 }
 #endif
diff --git a/Engine/Source/Runtime/Engine/Private/WorldPartition/WorldPartitionEditorSpatialHash.cpp b/Engine/Source/Runtime/Engine/Private/WorldPartition/WorldPartitionEditorSpatialHash.cpp
--- a/Engine/Source/Runtime/Engine/Private/WorldPartition/WorldPartitionEditorSpatialHash.cpp
+++ b/Engine/Source/Runtime/Engine/Private/WorldPartition/WorldPartitionEditorSpatialHash.cpp
@@ -99,16 +99,16 @@ void UWorldPartitionEditorSpatialHash::Tick(float DeltaSeconds)
 }
 
 // In the editor, actors are always using their bounds for grid placement, which makes more sense from a user standpoint.
-FBox UWorldPartitionEditorSpatialHash::GetActorBounds(FWorldPartitionActorDesc* InActorDesc) const
+FBox UWorldPartitionEditorSpatialHash::GetActorBounds(const FWorldPartitionHandle& InActorHandle) const
 {
 	FBox ActorBounds;
 
-	switch(InActorDesc->GetGridPlacement())
+	switch(InActorHandle->GetGridPlacement())
 	{
 	case EActorGridPlacement::Location:
 	case EActorGridPlacement::Bounds:
 	{
-		ActorBounds = InActorDesc->GetBounds();
+		ActorBounds = InActorHandle->GetBounds();
 		break;
 	}
 	}
@@ -117,30 +117,30 @@ FBox UWorldPartitionEditorSpatialHash::GetActorBounds(FWorldPartitionActorDesc*
 	return ActorBounds;
 }
 
-bool UWorldPartitionEditorSpatialHash::IsActorAlwaysLoaded(FWorldPartitionActorDesc* InActorDesc) const
+bool UWorldPartitionEditorSpatialHash::IsActorAlwaysLoaded(const FWorldPartitionHandle& InActorHandle) const
 {
-	if (InActorDesc->GetGridPlacement() == EActorGridPlacement::AlwaysLoaded)
+	if (InActorHandle->GetGridPlacement() == EActorGridPlacement::AlwaysLoaded)
 	{
 		return true;
 	}
 
 	// If an actor covers more that 4 levels in the octree (which means 32K cells), treat it as always loaded
-	const FBox ActorBounds = GetActorBounds(InActorDesc);
+	const FBox ActorBounds = GetActorBounds(InActorHandle);
 	const int32 ActorLevel = GetLevelForBox(ActorBounds);
 	return (ActorLevel > 4);
 }
 
-void UWorldPartitionEditorSpatialHash::HashActor(FWorldPartitionActorDesc* InActorDesc)
+void UWorldPartitionEditorSpatialHash::HashActor(FWorldPartitionHandle& InActorHandle)
 {
-	check(InActorDesc);
+	check(InActorHandle.IsValid());
 
-	if (IsActorAlwaysLoaded(InActorDesc))
+	if (IsActorAlwaysLoaded(InActorHandle))
 	{
-		AlwaysLoadedCell->AddActor(InActorDesc);
+		AlwaysLoadedCell->AddActor(InActorHandle);
 	}
 	else
 	{
-		const FBox ActorBounds = GetActorBounds(InActorDesc);
+		const FBox ActorBounds = GetActorBounds(InActorHandle);
 		const int32 CurrentLevel = GetLevelForBox(Bounds);
 
 		ForEachIntersectingCells(ActorBounds, 0, [&](const FCellCoord& CellCoord)
@@ -180,7 +180,7 @@ void UWorldPartitionEditorSpatialHash::HashActor(FWorldPartitionActorDesc* InAct
 			}
 
 			check(EditorCell);
-			EditorCell->AddActor(InActorDesc);
+			EditorCell->AddActor(InActorHandle);
 		});
 
 		const int32 NewLevel = GetLevelForBox(Bounds);
@@ -239,24 +239,24 @@ void UWorldPartitionEditorSpatialHash::HashActor(FWorldPartitionActorDesc* InAct
 	}
 }
 
-void UWorldPartitionEditorSpatialHash::UnhashActor(FWorldPartitionActorDesc* InActorDesc)
+void UWorldPartitionEditorSpatialHash::UnhashActor(FWorldPartitionHandle& InActorHandle)
 {
-	check(InActorDesc);
+	check(InActorHandle.IsValid());
 
-	if (IsActorAlwaysLoaded(InActorDesc))
+	if (IsActorAlwaysLoaded(InActorHandle))
 	{
-		AlwaysLoadedCell->RemoveActor(InActorDesc);
+		AlwaysLoadedCell->RemoveActor(InActorHandle);
 	}
 	else
 	{
-		const FBox ActorBounds = GetActorBounds(InActorDesc);
+		const FBox ActorBounds = GetActorBounds(InActorHandle);
 		const int32 CurrentLevel = GetLevelForBox(Bounds);
 
 		ForEachIntersectingCells(ActorBounds, 0, [&](const FCellCoord& CellCoord)
 		{
 			UWorldPartitionEditorCell* EditorCell = HashCells.FindChecked(CellCoord);
 
-			EditorCell->RemoveActor(InActorDesc);
+			EditorCell->RemoveActor(InActorHandle);
 
 			if (!EditorCell->Actors.Num())
 			{
@@ -345,6 +345,11 @@ void UWorldPartitionEditorSpatialHash::OnCellLoaded(const UWorldPartitionEditorC
 
 void UWorldPartitionEditorSpatialHash::OnCellUnloaded(const UWorldPartitionEditorCell* Cell)
 {
+	if (Cell == AlwaysLoadedCell)
+	{
+		return;
+	}
+
 	check(!Cell->bLoaded);
 
 	const FCellCoord CellCoord = GetCellCoords(Cell->Bounds.GetCenter(), 0);
@@ -376,13 +381,13 @@ int32 UWorldPartitionEditorSpatialHash::ForEachIntersectingActor(const FBox& Box
 
 	ForEachIntersectingCell(Box, [&](UWorldPartitionEditorCell* EditorCell)
 	{
-		for(FWorldPartitionActorDesc* ActorDesc: EditorCell->Actors)
+		for(FWorldPartitionHandle& ActorDesc: EditorCell->Actors)
 		{
 			if (ActorDesc->Tag != FWorldPartitionActorDesc::GlobalTag)
 			{
 				if (Box.Intersect(ActorDesc->GetBounds()))
 				{
-					InOperation(ActorDesc);
+					InOperation(*ActorDesc);
 					NumIntersecting++;
 				}
 
@@ -391,13 +396,13 @@ int32 UWorldPartitionEditorSpatialHash::ForEachIntersectingActor(const FBox& Box
 		}
 	});
 
-	for(FWorldPartitionActorDesc* ActorDesc: AlwaysLoadedCell->Actors)
+	for(FWorldPartitionHandle& ActorDesc: AlwaysLoadedCell->Actors)
 	{
 		check(ActorDesc->Tag != FWorldPartitionActorDesc::GlobalTag);
 
 		if (Box.Intersect(ActorDesc->GetBounds()))
 		{
-			InOperation(ActorDesc);
+			InOperation(*ActorDesc);
 			NumIntersecting++;
 		}
 
@@ -525,6 +530,7 @@ int32 UWorldPartitionEditorSpatialHash::ForEachIntersectingUnloadedRegion(const
 
 int32 UWorldPartitionEditorSpatialHash::ForEachCell(TFunctionRef<void(UWorldPartitionEditorCell*)> InOperation)
 {
+	InOperation(AlwaysLoadedCell);
 	for (UWorldPartitionEditorCell* Cell: Cells)
 	{
 		InOperation(Cell);
diff --git a/Engine/Source/Runtime/Engine/Private/WorldPartition/WorldPartitionHandle.cpp b/Engine/Source/Runtime/Engine/Private/WorldPartition/WorldPartitionHandle.cpp
--- a/Engine/Source/Runtime/Engine/Private/WorldPartition/WorldPartitionHandle.cpp
+++ b/Engine/Source/Runtime/Engine/Private/WorldPartition/WorldPartitionHandle.cpp
@@ -1,6 +1,7 @@
 // Copyright Epic Games, Inc. All Rights Reserved.
 
 #include "WorldPartition/WorldPartitionHandle.h"
+#include "WorldPartition/WorldPartitionActorDesc.h"
 #include "WorldPartition/WorldPartition.h"
 
 #if WITH_EDITOR
@@ -26,17 +27,17 @@ bool FWorldPartitionHandleUtils::IsActorDescLoaded(FWorldPartitionActorDesc* Act
 	return !!ActorDesc->GetActor();
 }
 
-void FWorldPartitionSoftRefImpl::IncRefCount(FWorldPartitionActorDesc* ActorDesc)
+void FWorldPartitionHandleImpl::IncRefCount(FWorldPartitionActorDesc* ActorDesc)
 {
 	ActorDesc->IncSoftRefCount();
 }
 
-void FWorldPartitionSoftRefImpl::DecRefCount(FWorldPartitionActorDesc* ActorDesc)
+void FWorldPartitionHandleImpl::DecRefCount(FWorldPartitionActorDesc* ActorDesc)
 {
 	ActorDesc->DecSoftRefCount();
 }
 
-void FWorldPartitionHardRefImpl::IncRefCount(FWorldPartitionActorDesc* ActorDesc)
+void FWorldPartitionReferenceImpl::IncRefCount(FWorldPartitionActorDesc* ActorDesc)
 {
 	if (ActorDesc->IncHardRefCount() == 1)
 	{
@@ -46,20 +47,21 @@ void FWorldPartitionHardRefImpl::IncRefCount(FWorldPartitionActorDesc* ActorDesc
 
 		if (Actor)
 		{
-			Actor->GetLevel()->AddLoadedActor(Actor);
+			ActorDesc->RegisterActor();
 		}
 	}
 }
 
-void FWorldPartitionHardRefImpl::DecRefCount(FWorldPartitionActorDesc* ActorDesc)
+void FWorldPartitionReferenceImpl::DecRefCount(FWorldPartitionActorDesc* ActorDesc)
 {
 	if (!ActorDesc->DecHardRefCount())
 	{
 		if (!IsEngineExitRequested())
 		{
+			// We can still hold a reference to an actor that was manually deleted in the editor
 			if (AActor* Actor = ActorDesc->GetActor())
 			{
-				Actor->GetLevel()->RemoveLoadedActor(Actor);
+				ActorDesc->UnregisterActor();
 				ActorDesc->Unload();
 			}
 		}
diff --git a/Engine/Source/Runtime/Engine/Public/WorldPartition/Landscape/LandscapeActorDesc.h b/Engine/Source/Runtime/Engine/Public/WorldPartition/Landscape/LandscapeActorDesc.h
--- a/Engine/Source/Runtime/Engine/Public/WorldPartition/Landscape/LandscapeActorDesc.h
+++ b/Engine/Source/Runtime/Engine/Public/WorldPartition/Landscape/LandscapeActorDesc.h
@@ -17,7 +17,7 @@ public:
 	virtual void Unload() override;
 
 protected:
-	virtual void OnRegister(UWorldPartition* WorldPartition) override;
-	virtual void OnUnregister(UWorldPartition* WorldPartition) override;
+	virtual void OnRegister() override;
+	virtual void OnUnregister() override;
 #endif
 };
diff --git a/Engine/Source/Runtime/Engine/Public/WorldPartition/Landscape/LandscapeSplineActorDesc.h b/Engine/Source/Runtime/Engine/Public/WorldPartition/Landscape/LandscapeSplineActorDesc.h
--- a/Engine/Source/Runtime/Engine/Public/WorldPartition/Landscape/LandscapeSplineActorDesc.h
+++ b/Engine/Source/Runtime/Engine/Public/WorldPartition/Landscape/LandscapeSplineActorDesc.h
@@ -15,8 +15,8 @@ class ENGINE_API FLandscapeSplineActorDesc : public FWorldPartitionActorDesc
 protected:
 	virtual void Init(const AActor* InActor) override;
 	virtual void Serialize(FArchive& Ar) override;
-	virtual void OnRegister(UWorldPartition* WorldPartition) override;
-	virtual void OnUnregister(UWorldPartition* WorldPartition) override;
+	virtual void OnRegister() override;
+	virtual void OnUnregister() override;
 
 	FGuid LandscapeGuid;
 #endif
diff --git a/Engine/Source/Runtime/Engine/Public/WorldPartition/WorldPartition.h b/Engine/Source/Runtime/Engine/Public/WorldPartition/WorldPartition.h
--- a/Engine/Source/Runtime/Engine/Public/WorldPartition/WorldPartition.h
+++ b/Engine/Source/Runtime/Engine/Public/WorldPartition/WorldPartition.h
@@ -66,6 +66,7 @@ class ENGINE_API UWorldPartition final : public UObject
 {
 	GENERATED_UCLASS_BODY()
 
+	friend class FWorldPartitionActorDesc;
 	friend class UWorldPartitionEditorCell;
 	friend class FWorldPartitionEditorModule;
 	friend class FUnrealEdMisc;
@@ -91,11 +92,11 @@ private:
 
 	bool ShouldHandleAssetEvent(const FAssetData& InAssetData);
 	TUniquePtr<FWorldPartitionActorDesc> GetActorDescriptor(const FAssetData& InAssetData);
-
-	void ApplyActorTransform(AActor* InActor, const FTransform& InTransform);
 #endif
 
+	//~ Begin UObject Interface
 	virtual void Serialize(FArchive& Ar) override;
+	//~ End UObject Interface
 
 	static void WorldPartitionOnLevelRemovedFromWorld(class ULevel* Level, UWorld* InWorld);
 
@@ -139,6 +140,9 @@ public:
 	void Initialize(UWorld* World, const FTransform& InTransform);
 	bool IsInitialized() const;
 	void Uninitialize();
+
+	void CleanupWorldPartition();
+
 	const FTransform& GetInstanceTransform() const { return InstanceTransform; }
 
 	void Tick(float DeltaSeconds);
@@ -168,9 +172,6 @@ public:
 	TChunkedArray<TUniquePtr<FWorldPartitionActorDesc>> ActorDescList;
 	TMap<FGuid, TUniquePtr<FWorldPartitionActorDesc>*> Actors;
 	
-	DECLARE_EVENT_TwoParams(UWorldPartition, FWorldPartitionActorRegisteredEvent, AActor&, bool);
-	FWorldPartitionActorRegisteredEvent OnActorRegisteredEvent;
-
 	bool bIgnoreAssetRegistryEvents;
 	bool bForceGarbageCollection;
 	bool bForceGarbageCollectionPurge;
@@ -193,10 +194,6 @@ private:
 	FLinkerInstancingContext InstancingContext;
 #endif
 
-	// Actor registration
-	AActor* RegisterActor(FWorldPartitionActorDesc* ActorDesc);
-	void UnregisterActor(AActor* Actor);
-
 	bool IsMainWorldPartition() const;
 
 	// Delegates registration
diff --git a/Engine/Source/Runtime/Engine/Public/WorldPartition/WorldPartitionActorDesc.h b/Engine/Source/Runtime/Engine/Public/WorldPartition/WorldPartitionActorDesc.h
--- a/Engine/Source/Runtime/Engine/Public/WorldPartition/WorldPartitionActorDesc.h
+++ b/Engine/Source/Runtime/Engine/Public/WorldPartition/WorldPartitionActorDesc.h
@@ -35,6 +35,9 @@ class ENGINE_API FWorldPartitionActorDesc
 #if WITH_EDITOR
 	friend class AActor;
 	friend class UWorldPartition;
+	friend struct FWorldPartitionHandleImpl;
+	friend struct FWorldPartitionReferenceImpl;
+	friend struct FWorldPartitionHandleUtils;
 
 #if WITH_DEV_AUTOMATION_TESTS
 	friend class WorldPartitionTests::FWorldPartitionSoftRefTest;
@@ -57,6 +60,7 @@ public:
 	inline FName GetActorPath() const { return ActorPath; }
 	FBox GetBounds() const;
 
+protected:
 	inline uint32 IncSoftRefCount() const
 	{
 		return ++SoftRefCount;
@@ -89,6 +93,7 @@ public:
 		return HardRefCount;
 	}
 
+public:
 	const TArray<FGuid>& GetReferences() const
 	{
 		return References;
@@ -97,19 +102,23 @@ public:
 	FString ToString() const;
 
 	AActor* GetActor() const;
-	AActor* Load(const FLinkerInstancingContext* InstancingContext = nullptr) const;
+	AActor* Load() const;
 	virtual void Unload();
 
+	void RegisterActor();
+	void UnregisterActor();
+
 	virtual void Init(const AActor* InActor);
-	void Init(const FWorldPartitionActorDescInitData& DescData);
+	void Init(UWorldPartition* InWorldPartition, const FWorldPartitionActorDescInitData& DescData);
 
 	void SerializeTo(TArray<uint8>& OutData);
 
-	void TransformInstance(const FString& From, const FString& To, const FTransform& Transform);
-
 protected:
 	FWorldPartitionActorDesc();
 
+	void TransformInstance(const FString& From, const FString& To, const FTransform& Transform);
+	void ApplyActorTransform(const FTransform& InTransform);
+
 	inline void TransferRefCounts(const FWorldPartitionActorDesc* From) const
 	{
 		SoftRefCount = From->SoftRefCount;
@@ -117,9 +126,11 @@ protected:
 	}
 
 	virtual void Serialize(FArchive& Ar);
-	virtual void OnRegister(UWorldPartition* WorldPartition) {}
-	virtual void OnUnregister(UWorldPartition* WorldPartition) {}
 
+	virtual void OnRegister() {}
+	virtual void OnUnregister() {}
+
+	// Persistent
 	FGuid							Guid;
 	FName							Class;
 	FName							ActorPackage;
@@ -133,12 +144,12 @@ protected:
 	TArray<FName>					DataLayers;
 	TArray<FGuid>					References;
 	
+	// Transient
 	mutable uint32					SoftRefCount;
 	mutable uint32					HardRefCount;
-
-	// Cached values
 	UClass*							ActorClass;
 	mutable TWeakObjectPtr<AActor>	ActorPtr;
+	UWorldPartition*				WorldPartition;
 
 public:
 	// Tagging
diff --git a/Engine/Source/Runtime/Engine/Public/WorldPartition/WorldPartitionEditorCell.h b/Engine/Source/Runtime/Engine/Public/WorldPartition/WorldPartitionEditorCell.h
--- a/Engine/Source/Runtime/Engine/Public/WorldPartition/WorldPartitionEditorCell.h
+++ b/Engine/Source/Runtime/Engine/Public/WorldPartition/WorldPartitionEditorCell.h
@@ -2,6 +2,7 @@
 #pragma once
 
 #include "CoreMinimal.h"
+#include "WorldPartition/WorldPartitionHandle.h"
 #include "WorldPartitionEditorCell.generated.h"
 
 class FWorldPartitionActorDesc;
@@ -19,17 +20,15 @@ public:
 	virtual void Serialize(FArchive& Ar) override;
 
 #if WITH_EDITOR
-	UWorldPartition* GetWorldPartition();
-
-	void AddActor(FWorldPartitionActorDesc* InActorDesc);
-	void RemoveActor(FWorldPartitionActorDesc* InActorDesc);
+	void AddActor(const FWorldPartitionHandle& ActorHandle);
+	void RemoveActor(const FWorldPartitionHandle& ActorHandle);
 
 	FBox						Bounds;
 
 	/** Tells if the cell was manually loaded in the editor */
 	bool						bLoaded : 1;
 
-	TSet<FWorldPartitionActorDesc*> Actors;
-	TSet<FWorldPartitionActorDesc*> LoadedActors;
+	TSet<FWorldPartitionHandle> Actors;
+	TSet<FWorldPartitionReference> LoadedActors;
 #endif
 };
\ No newline at end of file
diff --git a/Engine/Source/Runtime/Engine/Public/WorldPartition/WorldPartitionEditorHash.h b/Engine/Source/Runtime/Engine/Public/WorldPartition/WorldPartitionEditorHash.h
--- a/Engine/Source/Runtime/Engine/Public/WorldPartition/WorldPartitionEditorHash.h
+++ b/Engine/Source/Runtime/Engine/Public/WorldPartition/WorldPartitionEditorHash.h
@@ -5,6 +5,7 @@
 #include "UObject/ObjectMacros.h"
 #include "UObject/Object.h"
 #include "WorldPartition/WorldPartitionActorDesc.h"
+#include "WorldPartition/WorldPartitionHandle.h"
 #include "WorldPartitionEditorHash.generated.h"
 
 class UWorldPartitionEditorCell;
@@ -21,8 +22,8 @@ public:
 	virtual FName GetWorldPartitionEditorName() PURE_VIRTUAL(UWorldPartitionEditorHash::GetWorldPartitionEditorName, return FName(NAME_None););
 	virtual void Tick(float DeltaSeconds) PURE_VIRTUAL(UWorldPartitionEditorHash::Tick, return;);
 
-	virtual void HashActor(FWorldPartitionActorDesc* InActorDesc) PURE_VIRTUAL(UWorldPartitionEditorHash::HashActor, ;);
-	virtual void UnhashActor(FWorldPartitionActorDesc* InActorDesc) PURE_VIRTUAL(UWorldPartitionEditorHash::UnhashActor, ;);
+	virtual void HashActor(FWorldPartitionHandle& InActorHandle) PURE_VIRTUAL(UWorldPartitionEditorHash::HashActor, ;);
+	virtual void UnhashActor(FWorldPartitionHandle& InActorHandle) PURE_VIRTUAL(UWorldPartitionEditorHash::UnhashActor, ;);
 
 	virtual void OnCellLoaded(const UWorldPartitionEditorCell* Cell) PURE_VIRTUAL(UWorldPartitionEditorHash::OnCellLoaded, ;);
 	virtual void OnCellUnloaded(const UWorldPartitionEditorCell* Cell) PURE_VIRTUAL(UWorldPartitionEditorHash::OnCellLoaded, ;);
diff --git a/Engine/Source/Runtime/Engine/Public/WorldPartition/WorldPartitionEditorSpatialHash.h b/Engine/Source/Runtime/Engine/Public/WorldPartition/WorldPartitionEditorSpatialHash.h
--- a/Engine/Source/Runtime/Engine/Public/WorldPartition/WorldPartitionEditorSpatialHash.h
+++ b/Engine/Source/Runtime/Engine/Public/WorldPartition/WorldPartitionEditorSpatialHash.h
@@ -212,8 +212,8 @@ public:
 	virtual FName GetWorldPartitionEditorName() override;
 	virtual void Tick(float DeltaSeconds) override;
 
-	virtual void HashActor(FWorldPartitionActorDesc* InActorDesc) override;
-	virtual void UnhashActor(FWorldPartitionActorDesc* InActorDesc) override;
+	virtual void HashActor(FWorldPartitionHandle& InActorHandle) override;
+	virtual void UnhashActor(FWorldPartitionHandle& InActorHandle) override;
 
 	virtual void OnCellLoaded(const UWorldPartitionEditorCell* Cell) override;
 	virtual void OnCellUnloaded(const UWorldPartitionEditorCell* Cell) override;
@@ -230,8 +230,8 @@ private:
 	int32 ForEachIntersectingUnloadedRegion(const FBox& Box, TFunctionRef<void(const FCellCoord&)> InOperation);
 	int32 ForEachIntersectingUnloadedRegionInner(const FBox& Box, const FCellCoord& CellCoord, TFunctionRef<void(const FCellCoord&)> InOperation);
 
-	FBox GetActorBounds(FWorldPartitionActorDesc* InActorDesc) const;
-	bool IsActorAlwaysLoaded(FWorldPartitionActorDesc* InActorDesc) const;
+	FBox GetActorBounds(const FWorldPartitionHandle& InActorHandle) const;
+	bool IsActorAlwaysLoaded(const FWorldPartitionHandle& InActorHandle) const;
 	int32 ForEachIntersectingCellInner(const FBox& Box, const FCellCoord& CellCoord, TFunctionRef<void(UWorldPartitionEditorCell*)> InOperation);
 
 	UPROPERTY(Config)
diff --git a/Engine/Source/Runtime/Engine/Public/WorldPartition/WorldPartitionHandle.h b/Engine/Source/Runtime/Engine/Public/WorldPartition/WorldPartitionHandle.h
--- a/Engine/Source/Runtime/Engine/Public/WorldPartition/WorldPartitionHandle.h
+++ b/Engine/Source/Runtime/Engine/Public/WorldPartition/WorldPartitionHandle.h
@@ -217,13 +217,13 @@ public:
 	TUniquePtr<FWorldPartitionActorDesc>* ActorDesc;
 };
 
-struct ENGINE_API FWorldPartitionSoftRefImpl
+struct ENGINE_API FWorldPartitionHandleImpl
 {
 	static void IncRefCount(FWorldPartitionActorDesc* ActorDesc);
 	static void DecRefCount(FWorldPartitionActorDesc* ActorDesc);
 };
 
-struct ENGINE_API FWorldPartitionHardRefImpl
+struct ENGINE_API FWorldPartitionReferenceImpl
 {
 	static void IncRefCount(FWorldPartitionActorDesc* ActorDesc);
 	static void DecRefCount(FWorldPartitionActorDesc* ActorDesc);
@@ -234,13 +234,13 @@ struct ENGINE_API FWorldPartitionHardRefImpl
  * This won't trigger any loading, but will prevent cleanup of the actor descriptor when destroying an
  * actor in the editor.
  */
-typedef TWorldPartitionHandle<FWorldPartitionSoftRefImpl> FWorldPartitionHandle;
+typedef TWorldPartitionHandle<FWorldPartitionHandleImpl> FWorldPartitionHandle;
 
 /**
  * FWorldPartitionReference will increment/decrement the hard reference count on the actor descriptor.
  * This will trigger actor loading/unloading when the hard reference counts gets to one/zero.
  */
-typedef TWorldPartitionHandle<FWorldPartitionHardRefImpl> FWorldPartitionReference;
+typedef TWorldPartitionHandle<FWorldPartitionReferenceImpl> FWorldPartitionReference;
 
 /**
  * FWorldPartitionHandlePinRefScope will keep a reference if the actor is already loaded. This is useful 
