Fixed crash in async DF generation when StaticMesh RenderDatat doesn't have a valid DF

diff --git a/Engine/Source/Runtime/Engine/Private/DistanceFieldAtlas.cpp b/Engine/Source/Runtime/Engine/Private/DistanceFieldAtlas.cpp
--- a/Engine/Source/Runtime/Engine/Private/DistanceFieldAtlas.cpp
+++ b/Engine/Source/Runtime/Engine/Private/DistanceFieldAtlas.cpp
@@ -1367,10 +1367,13 @@ void FDistanceFieldAsyncQueue::ProcessAsyncTasks(bool bLimitExecutionTime)
 				}
 			}
 
-			OldVolumeData->VolumeTexture.Release();
+			if (OldVolumeData)
+			{
+				OldVolumeData->VolumeTexture.Release();
 
-			// Rendering thread may still be referencing the old one, use the deferred cleanup interface to delete it next frame when it is safe
-			BeginCleanup(OldVolumeData);
+				// Rendering thread may still be referencing the old one, use the deferred cleanup interface to delete it next frame when it is safe
+				BeginCleanup(OldVolumeData);
+			}
 
 			{
 				TArray<uint8> DerivedData;
