Fixed strata SingleLayerWater node tooltip.

diff --git a/Engine/Source/Runtime/Engine/Classes/Materials/MaterialExpressionStrata.h b/Engine/Source/Runtime/Engine/Classes/Materials/MaterialExpressionStrata.h
--- a/Engine/Source/Runtime/Engine/Classes/Materials/MaterialExpressionStrata.h
+++ b/Engine/Source/Runtime/Engine/Classes/Materials/MaterialExpressionStrata.h
@@ -217,7 +217,7 @@ class UMaterialExpressionStrataVolumeBSDF : public UMaterialExpression
 	FExpressionInput Extinction;
 
 	/**
-	 * Anisotropy (type = float, unitless)
+	 * Anisotropy of the volume with values lower than 0 representing back-scattering, equal 0 representing isotropic scattering and greater than 0 representing forward scattering. (type = float, unit = unitless, defaults to 0)
 	 */
 	UPROPERTY()
 	FExpressionInput Anisotropy;
@@ -410,19 +410,19 @@ class UMaterialExpressionStrataSingleLayerWaterBSDF : public UMaterialExpression
 	GENERATED_UCLASS_BODY()
 
 	/**
-	 * Surface base color resulting from single and multiple scattering combined. (type = float3, unit = unitless, defaults to black)
+	 * Surface base color. (type = float3, unit = unitless, defaults to black)
 	 */
 	UPROPERTY()
 	FExpressionInput BaseColor;
 
 	/**
-	 * TODO. (type = float3, unit = unitless, defaults to black)
+	 * whether the surface represents a dielectric (such as plastic) or a conductor (such as metal). (type = float, unit = unitless, defaults to 0 = dielectric)
 	 */
 	UPROPERTY()
 	FExpressionInput Metallic;
 
 	/**
-	 * Specular (type = float, unit = unitless, defaults to 0.5)
+	 * Specular amount (type = float, unit = unitless, defaults to 0.5)
 	 */
 	UPROPERTY()
 	FExpressionInput Specular;
@@ -434,13 +434,13 @@ class UMaterialExpressionStrataSingleLayerWaterBSDF : public UMaterialExpression
 	FExpressionInput Roughness;
 
 	/**
-	 * Tangent (type = float3, unit = unitless, defaults to +Z vector)
+	 * The normal of the surface (type = float3, unit = unitless, defaults to +Z vector)
 	 */
 	UPROPERTY()
 	FExpressionInput Normal;
 
 	/**
-	 * Emissive (type = float3, unit = luminance, defaults to 0.0)
+	 * Emissive RGB (type = float3, unit = luminance, defaults to 0.0)
 	 */
 	UPROPERTY()
 	FExpressionInput Emissive;
@@ -452,25 +452,25 @@ class UMaterialExpressionStrataSingleLayerWaterBSDF : public UMaterialExpression
 	FExpressionInput TopMaterialOpacity;
 
 	/**
-	 * TODO. (type = float3, unit = unitless, defaults to black)
+	 * Albedo (type = float3, unit = unitless, defaults to black)
 	 */
 	UPROPERTY()
 	FExpressionInput WaterAlbedo;
 
 	/**
-	 * TODO. (type = float3, unit = 1/cm, defaults to black)
+	 * The rate at which light is absorbed or out-scattered by the medium. Mean Free Path = 1 / Extinction. (type = float3, unit = 1/cm, default = 0)
 	 */
 	UPROPERTY()
 	FExpressionInput WaterExtinction;
 
 	/**
-	 * TODO. (type = float3, unit = unitless, defaults to 0)
+	 * Anisotropy of the volume with values lower than 0 representing back-scattering, equal 0 representing isotropic scattering and greater than 0 representing forward scattering. (type = float, unit = unitless, defaults to 0)
 	 */
 	UPROPERTY()
 	FExpressionInput WaterPhaseG;
 
 	/**
-	 * TODO. (type = float3, unit = unitless, defaults to white)
+	 * A scale to apply on the scene color behind the water surface. It can be used to approximate caustics for instance. (type = float3, unit = unitless, defaults to 1)
 	 */
 	UPROPERTY()
 	FExpressionInput ColorScaleBehindWater;
