Restoring functionality that would compile High or higher priority shader compiles (like Global shaders) in single jobs instead of batching them together.
Restoring log level that would provide log feedback on global shader compiles by default. This matches the legacy UE4 behavior from before the shader compiler refactors.

diff --git a/Engine/Source/Runtime/Engine/Private/ShaderCompiler/ShaderCompiler.cpp b/Engine/Source/Runtime/Engine/Private/ShaderCompiler/ShaderCompiler.cpp
--- a/Engine/Source/Runtime/Engine/Private/ShaderCompiler/ShaderCompiler.cpp
+++ b/Engine/Source/Runtime/Engine/Private/ShaderCompiler/ShaderCompiler.cpp
@@ -1534,10 +1534,10 @@ int32 FShaderCompileThreadRunnable::PullTasksFromQueue()
 					NumPendingJobs = Manager->AllJobs.GetNumPendingJobs((EShaderCompileJobPriority)PriorityIndex);
 					if (NumPendingJobs > 0)
 					{
-						UE_LOG(LogShaderCompilers, Verbose, TEXT("Worker (%d/%d): shaders left to compile %i"), WorkerIndex + 1, WorkerInfos.Num(), NumPendingJobs);
+						UE_LOG(LogShaderCompilers, Display, TEXT("Worker (%d/%d): shaders left to compile %i"), WorkerIndex + 1, WorkerInfos.Num(), NumPendingJobs);
 
 						int32 MaxNumJobs = 1;
-						//if (PriorityIndex < (int32)EShaderCompileJobPriority::ForceLocal)
+						if (PriorityIndex < (int32)EShaderCompileJobPriority::High)
 						{
 							MaxNumJobs = FMath::Min3(NumJobsPerWorker, NumPendingJobs, Manager->MaxShaderJobBatchSize);
 						}
