#BlendProfile - Check PoseIndex before use to fix crash happening when using a blend profile with a montage when current mesh doesn't have all the bones present in the skeleton/blend profile.

diff --git a/Engine/Source/Runtime/Engine/Private/Animation/BlendProfile.cpp b/Engine/Source/Runtime/Engine/Private/Animation/BlendProfile.cpp
--- a/Engine/Source/Runtime/Engine/Private/Animation/BlendProfile.cpp
+++ b/Engine/Source/Runtime/Engine/Private/Animation/BlendProfile.cpp
@@ -191,7 +191,10 @@ void UBlendProfile::FillBoneScalesArray(TArray<float>& OutBoneBlendProfileFactor
 	{
 		const int32 SkeletonBoneIndex = ProfileEntries[Index].BoneReference.BoneIndex;
 		const FCompactPoseBoneIndex PoseBoneIndex = BoneContainer.GetCompactPoseIndexFromSkeletonIndex(SkeletonBoneIndex);
-		OutBoneBlendProfileFactors[PoseBoneIndex.GetInt()] = GetEntryBlendScale(Index);
+		if (PoseBoneIndex.IsValid())
+		{
+			OutBoneBlendProfileFactors[PoseBoneIndex.GetInt()] = GetEntryBlendScale(Index);
+		}
 	}
 }
 
