Modify Engine Pre/Post Render delegates to accept an FRDGBuilder.
Mark old delegates for deprecation.

diff --git a/Engine/Source/Runtime/Engine/Classes/Engine/Engine.h b/Engine/Source/Runtime/Engine/Classes/Engine/Engine.h
--- a/Engine/Source/Runtime/Engine/Classes/Engine/Engine.h
+++ b/Engine/Source/Runtime/Engine/Classes/Engine/Engine.h
@@ -625,9 +625,10 @@ enum class EFrameHitchType : uint8;
 
 DECLARE_MULTICAST_DELEGATE_TwoParams(FEngineHitchDetectedDelegate, EFrameHitchType /*HitchType*/, float /*HitchDurationInSeconds*/);
 
-
 DECLARE_MULTICAST_DELEGATE(FPreRenderDelegate);
+DECLARE_MULTICAST_DELEGATE_OneParam(FPreRenderDelegateEx, class FRDGBuilder&);
 DECLARE_MULTICAST_DELEGATE(FPostRenderDelegate);
+DECLARE_MULTICAST_DELEGATE_OneParam(FPostRenderDelegateEx, class FRDGBuilder&);
 
 /**
  * Abstract base class of all Engine classes, responsible for management of systems critical to editor or game systems.
@@ -1551,13 +1552,21 @@ public:
 	void RegisterEndStreamingPauseRenderingDelegate( FEndStreamingPauseDelegate* InDelegate );
 	FEndStreamingPauseDelegate* EndStreamingPauseDelegate;
 
-
+private:
 	/** Delegate called just prior to rendering. */
 	FPreRenderDelegate PreRenderDelegate;
-	FPreRenderDelegate& GetPreRenderDelegate() { return PreRenderDelegate; }
+	FPreRenderDelegateEx PreRenderDelegateEx;
 	/** Delegate called just after to rendering. */
 	FPostRenderDelegate PostRenderDelegate;
+	FPostRenderDelegateEx PostRenderDelegateEx;
+
+public:
+	UE_DEPRECATED(5.0, "Please use GetPreRenderDelegateEx().")
+	FPreRenderDelegate& GetPreRenderDelegate() { return PreRenderDelegate; }
+	FPreRenderDelegateEx& GetPreRenderDelegateEx() { return PreRenderDelegateEx; }
+	UE_DEPRECATED(5.0, "Please use GetPostRenderDelegateEx().")
 	FPostRenderDelegate& GetPostRenderDelegate() { return PostRenderDelegate; }
+	FPostRenderDelegateEx& GetPostRenderDelegateEx() { return PostRenderDelegateEx; }
 
 	/** 
 	 * Error message event relating to server travel failures 
diff --git a/Engine/Source/Runtime/Engine/Private/GameEngine.cpp b/Engine/Source/Runtime/Engine/Private/GameEngine.cpp
--- a/Engine/Source/Runtime/Engine/Private/GameEngine.cpp
+++ b/Engine/Source/Runtime/Engine/Private/GameEngine.cpp
@@ -783,6 +783,12 @@ UEngine::UEngine(const FObjectInitializer& ObjectInitializer)
 		}
 	}
 	#endif
+
+	// Relay to deprecated delegates
+	PRAGMA_DISABLE_DEPRECATION_WARNINGS
+	PreRenderDelegateEx.AddLambda([this](FRDGBuilder&) { PreRenderDelegate.Broadcast(); });
+	PostRenderDelegateEx.AddLambda([this](FRDGBuilder&){ PostRenderDelegate.Broadcast(); });
+	PRAGMA_ENABLE_DEPRECATION_WARNINGS
 }
 
 
diff --git a/Engine/Source/Runtime/Renderer/Private/DeferredShadingRenderer.cpp b/Engine/Source/Runtime/Renderer/Private/DeferredShadingRenderer.cpp
--- a/Engine/Source/Runtime/Renderer/Private/DeferredShadingRenderer.cpp
+++ b/Engine/Source/Runtime/Renderer/Private/DeferredShadingRenderer.cpp
@@ -1774,9 +1774,9 @@ void FDeferredShadingSceneRenderer::Render(FRDGBuilder& GraphBuilder)
 	const bool bIsOcclusionTesting = DoOcclusionQueries(FeatureLevel) && (!ViewFamily.EngineShowFlags.Wireframe || bIsViewFrozen || bHasViewParent);
 	const bool bNeedsPrePass = NeedsPrePass(this);
 
-	// Dynamic vertex and index buffers need to be committed before rendering.
-	GEngine->GetPreRenderDelegate().Broadcast();
+	GEngine->GetPreRenderDelegateEx().Broadcast(GraphBuilder);
 
+	// Dynamic vertex and index buffers need to be committed before rendering.
 	{
 		SCOPE_CYCLE_COUNTER(STAT_FDeferredShadingSceneRenderer_FGlobalDynamicVertexBuffer_Commit);
 		DynamicIndexBufferForInitViews.Commit();
@@ -2850,10 +2850,10 @@ void FDeferredShadingSceneRenderer::Render(FRDGBuilder& GraphBuilder)
 		ShaderDrawDebug::EndView(Views[ViewIndex]);
 	}
 
+	GEngine->GetPostRenderDelegateEx().Broadcast(GraphBuilder);
+
 	AddPass(GraphBuilder, [this, &SceneContext](FRHICommandListImmediate& InRHICmdList)
 	{
-		GEngine->GetPostRenderDelegate().Broadcast();
-
 		SceneContext.AdjustGBufferRefCount(InRHICmdList, -1);
 
 #if RHI_RAYTRACING
diff --git a/Engine/Source/Runtime/Renderer/Private/MobileShadingRenderer.cpp b/Engine/Source/Runtime/Renderer/Private/MobileShadingRenderer.cpp
--- a/Engine/Source/Runtime/Renderer/Private/MobileShadingRenderer.cpp
+++ b/Engine/Source/Runtime/Renderer/Private/MobileShadingRenderer.cpp
@@ -590,7 +590,7 @@ void FMobileSceneRenderer::Render(FRDGBuilder& GraphBuilder)
 		FRHICommandListExecutor::GetImmediateCommandList().ImmediateFlush(EImmediateFlushType::FlushRHIThreadFlushResources);
 	}
 
-	GEngine->GetPreRenderDelegate().Broadcast();
+	GEngine->GetPreRenderDelegateEx().Broadcast(GraphBuilder);
 
 	// Global dynamic buffers need to be committed before rendering.
 	DynamicIndexBuffer.Commit();
@@ -795,10 +795,7 @@ void FMobileSceneRenderer::Render(FRDGBuilder& GraphBuilder)
 		}
 	}
 
-	AddPass(GraphBuilder, [this](FRHICommandListImmediate&)
-	{
-		GEngine->GetPostRenderDelegate().Broadcast();
-	});
+	GEngine->GetPostRenderDelegateEx().Broadcast(GraphBuilder);
 
 	GraphBuilder.SetCommandListStat(GET_STATID(STAT_CLMM_SceneEnd));
 
diff --git a/Engine/Source/Runtime/Renderer/Private/SceneHitProxyRendering.cpp b/Engine/Source/Runtime/Renderer/Private/SceneHitProxyRendering.cpp
--- a/Engine/Source/Runtime/Renderer/Private/SceneHitProxyRendering.cpp
+++ b/Engine/Source/Runtime/Renderer/Private/SceneHitProxyRendering.cpp
@@ -562,7 +562,7 @@ void FMobileSceneRenderer::RenderHitProxies(FRDGBuilder& GraphBuilder)
 	// Find the visible primitives.
 	InitViews(GraphBuilder);
 
-	GEngine->GetPreRenderDelegate().Broadcast();
+	GEngine->GetPreRenderDelegateEx().Broadcast(GraphBuilder);
 
 	// Global dynamic buffers need to be committed before rendering.
 	DynamicIndexBuffer.Commit();
@@ -573,7 +573,7 @@ void FMobileSceneRenderer::RenderHitProxies(FRDGBuilder& GraphBuilder)
 
 	::DoRenderHitProxies(GraphBuilder, this, HitProxyTexture, HitProxyDepthTexture, NaniteRasterResults);
 
-	GEngine->GetPostRenderDelegate().Broadcast();
+	GEngine->GetPostRenderDelegateEx().Broadcast(GraphBuilder);
 #endif
 
 	if (bDeferredShading)
@@ -642,7 +642,7 @@ void FDeferredShadingSceneRenderer::RenderHitProxies(FRDGBuilder& GraphBuilder)
 		Nanite::GStreamingManager.EndAsyncUpdate(GraphBuilder);
 	}
 
-	GEngine->GetPreRenderDelegate().Broadcast();
+	GEngine->GetPreRenderDelegateEx().Broadcast(GraphBuilder);
 
 	// Global dynamic buffers need to be committed before rendering.
 	DynamicIndexBufferForInitViews.Commit();
@@ -710,7 +710,7 @@ void FDeferredShadingSceneRenderer::RenderHitProxies(FRDGBuilder& GraphBuilder)
 		ShaderPrint::EndView(Views[ViewIndex]);
 	}
 
-	GEngine->GetPostRenderDelegate().Broadcast();
+	GEngine->GetPostRenderDelegateEx().Broadcast(GraphBuilder);
 
 #endif
 }
