Remove high resolution screenshot manually setting the scene texture extent. Does not actually have an affect.

#rnx

diff --git a/Engine/Source/Runtime/Engine/Private/UnrealClient.cpp b/Engine/Source/Runtime/Engine/Private/UnrealClient.cpp
--- a/Engine/Source/Runtime/Engine/Private/UnrealClient.cpp
+++ b/Engine/Source/Runtime/Engine/Private/UnrealClient.cpp
@@ -1243,8 +1243,6 @@ void FViewport::HighResScreenshot()
 	// cleared out before we use it below
 	const FString CachedScreenshotName = FScreenshotRequest::GetFilename();
 
-	FIntPoint RestoreSize(SizeX, SizeY);
-
 	FDummyViewport* DummyViewport = new FDummyViewport(ViewportClient);
 
 	DummyViewport->SizeX = (GScreenshotResolutionX > 0) ? GScreenshotResolutionX : SizeX;
@@ -1313,10 +1311,9 @@ void FViewport::HighResScreenshot()
 	}
 
 	ENQUEUE_RENDER_COMMAND(EndDrawingCommand)(
-		[DummyViewport, RestoreSize](FRHICommandListImmediate& RHICmdList)
+		[DummyViewport](FRHICommandListImmediate& RHICmdList)
 		{
 			DummyViewport->EndRenderFrame(RHICmdList, false, false);
-			GetRendererModule().SceneRenderTargetsSetBufferSize(RestoreSize.X, RestoreSize.Y);
 		});
 
 	BeginReleaseResource(DummyViewport);
