Fix crash compiling debug viewmodes for landscape materials (UE-104818).

diff --git a/Engine/Source/Runtime/Engine/Private/DebugViewModeHelpers.cpp b/Engine/Source/Runtime/Engine/Private/DebugViewModeHelpers.cpp
--- a/Engine/Source/Runtime/Engine/Private/DebugViewModeHelpers.cpp
+++ b/Engine/Source/Runtime/Engine/Private/DebugViewModeHelpers.cpp
@@ -277,6 +277,19 @@ bool CompileDebugViewModeShaders(EDebugViewShaderMode ShaderMode, EMaterialQuali
 		return false;
 	}
 
+	// remove materials incompatible with debug view modes (e.g. landscape materials can only be compiled with the landscape VF)
+	for (TSet<UMaterialInterface*>::TIterator It(Materials); It; ++It)
+	{
+		UMaterialInterface* MaterialInterface = *It;
+		check(MaterialInterface); // checked for null in GetTextureStreamingBuildMaterials
+		const FMaterial* Material = MaterialInterface->GetMaterialResource(FeatureLevel, QualityLevel);
+
+		if (Material && (Material->GetMaterialDomain() != MD_Surface || Material->IsUsedWithLandscape()))
+		{
+			It.RemoveCurrent();
+		}
+	}
+
 	TSet<UMaterialInterface*> PendingMaterials = Materials;
 	while (PendingMaterials.Num() > 0)
 	{
