Make the IStreamingManager interface return true for HasShutdown() if it was never initialized in the first place. Stops the system being initialized during shutdown if it was never created in the first place

diff --git a/Engine/Source/Runtime/Engine/Private/ContentStreaming.cpp b/Engine/Source/Runtime/Engine/Private/ContentStreaming.cpp
--- a/Engine/Source/Runtime/Engine/Private/ContentStreaming.cpp
+++ b/Engine/Source/Runtime/Engine/Private/ContentStreaming.cpp
@@ -392,7 +392,7 @@ void IStreamingManager::Shutdown()
 
 bool IStreamingManager::HasShutdown()
 {
-	return StreamingManagerCollection == (FStreamingManagerCollection*)-1;
+	return (StreamingManagerCollection == nullptr) || (StreamingManagerCollection == (FStreamingManagerCollection*)-1);
 }
 
 /**
