Animation sync node and anim graph attributes

Refactored tick record sync into utility structure - FAnimSync. This improves the odd API around tick records, better encapsulating functionality and leaving less for the caller to get wrong. This will also eventually allow this to be refactored out into a scriptable pipeline stage.
Added sync node and scoped sync message for new 'graph based sync'. Nodes that subscribe to graph-based-sync determine their sync group based on the scope that they are in.
Removed 4.26-style sync scopes - pushed all syncing up to the main anim instance. Linked anim instances no longer sync their own tick records.
To make graph based sync more useful, surfaced graph attributes and their visualizations as labels on pins and parallel wires to visualize flow.
This involves statically determining the attribute flow of the graph at compile time. Added a new compiler handler to deal with this new debug data.
Updated a lot of nodes to specify their attributes so graph flow can be correctly visualized.
Added tracing of attributes and sync records and visualization of traced records when debugging the anim graph.

diff --git a/Engine/Plugins/Animation/GameplayInsights/Source/GameplayInsights/Private/AnimNodesTrack.cpp b/Engine/Plugins/Animation/GameplayInsights/Source/GameplayInsights/Private/AnimNodesTrack.cpp
--- a/Engine/Plugins/Animation/GameplayInsights/Source/GameplayInsights/Private/AnimNodesTrack.cpp
+++ b/Engine/Plugins/Animation/GameplayInsights/Source/GameplayInsights/Private/AnimNodesTrack.cpp
@@ -304,88 +304,126 @@ void FAnimNodesTrack::UpdateDebugData(const TraceServices::FFrame& InFrame)
 {
 	if(InstanceClass.LoadSynchronous())
 	{
-		const FAnimationProvider* AnimationProvider = SharedData.GetAnalysisSession().ReadProvider<FAnimationProvider>(FAnimationProvider::ProviderName);
-		const FGameplayProvider* GameplayProvider = SharedData.GetAnalysisSession().ReadProvider<FGameplayProvider>(FGameplayProvider::ProviderName);
-
-		if(AnimationProvider && GameplayProvider)
+		if(UAnimBlueprint* AnimBlueprint = Cast<UAnimBlueprint>(InstanceClass.Get()->ClassGeneratedBy))
 		{
-			TraceServices::FAnalysisSessionReadScope SessionReadScope(SharedData.GetAnalysisSession());
-
-			AnimationProvider->ReadAnimGraphTimeline(GetGameplayTrack().GetObjectId(), [this, AnimationProvider, GameplayProvider, InFrame](const FAnimationProvider::AnimGraphTimeline& InGraphTimeline)
+			if(AnimInstance && AnimBlueprint->IsObjectBeingDebugged(AnimInstance))
 			{
-				FAnimBlueprintDebugData& DebugData = InstanceClass.Get()->GetAnimBlueprintDebugData();
-				DebugData.ResetNodeVisitSites();
+				const FAnimationProvider* AnimationProvider = SharedData.GetAnalysisSession().ReadProvider<FAnimationProvider>(FAnimationProvider::ProviderName);
+				const FGameplayProvider* GameplayProvider = SharedData.GetAnalysisSession().ReadProvider<FGameplayProvider>(FGameplayProvider::ProviderName);
 
-				InGraphTimeline.EnumerateEvents(InFrame.StartTime, InFrame.EndTime, [this, AnimationProvider, GameplayProvider, &DebugData](double InGraphStartTime, double InGraphEndTime, uint32 InDepth, const FAnimGraphMessage& InMessage)
+				if(AnimationProvider && GameplayProvider)
 				{
-					// Check for an update phase (which contains weights)
-					if(InMessage.Phase == EAnimGraphPhase::Update)
+					TraceServices::FAnalysisSessionReadScope SessionReadScope(SharedData.GetAnalysisSession());
+
+					const int32 NodeCount = InstanceClass.Get()->GetAnimNodeProperties().Num();
+
+					AnimationProvider->ReadAnimGraphTimeline(GetGameplayTrack().GetObjectId(), [this, AnimationProvider, GameplayProvider, InFrame, NodeCount](const FAnimationProvider::AnimGraphTimeline& InGraphTimeline)
 					{
-						// Basic verification - check node count is the same
-						// @TODO: could add some form of node hash/CRC to the class to improve this
-						if(InMessage.NodeCount == InstanceClass.Get()->GetAnimNodeProperties().Num())
+						FAnimBlueprintDebugData& DebugData = InstanceClass.Get()->GetAnimBlueprintDebugData();
+						DebugData.ResetNodeVisitSites();
+
+						InGraphTimeline.EnumerateEvents(InFrame.StartTime, InFrame.EndTime, [this, AnimationProvider, GameplayProvider, &DebugData, NodeCount](double InGraphStartTime, double InGraphEndTime, uint32 InDepth, const FAnimGraphMessage& InMessage)
 						{
-							AnimationProvider->ReadAnimNodesTimeline(GetGameplayTrack().GetObjectId(), [InGraphStartTime, InGraphEndTime, &DebugData](const FAnimationProvider::AnimNodesTimeline& InNodesTimeline)
+							// Basic verification - check node count is the same
+							// @TODO: could add some form of node hash/CRC to the class to improve this
+							if(InMessage.NodeCount == NodeCount)
 							{
-								InNodesTimeline.EnumerateEvents(InGraphStartTime, InGraphEndTime, [&DebugData](double InStartTime, double InEndTime, uint32 InDepth, const FAnimNodeMessage& InMessage)
+								// Check for an update phase (which contains weights)
+								if(InMessage.Phase == EAnimGraphPhase::Update)
 								{
-									if(InMessage.Phase == EAnimGraphPhase::Update)
+									AnimationProvider->ReadAnimNodesTimeline(GetGameplayTrack().GetObjectId(), [InGraphStartTime, InGraphEndTime, &DebugData](const FAnimationProvider::AnimNodesTimeline& InNodesTimeline)
 									{
-										DebugData.RecordNodeVisit(InMessage.NodeId, InMessage.PreviousNodeId, InMessage.Weight);
-									}
-									return TraceServices::EEventEnumerate::Continue;
-								});
-							});
+										InNodesTimeline.EnumerateEvents(InGraphStartTime, InGraphEndTime, [&DebugData](double InStartTime, double InEndTime, uint32 InDepth, const FAnimNodeMessage& InMessage)
+										{
+											DebugData.RecordNodeVisit(InMessage.NodeId, InMessage.PreviousNodeId, InMessage.Weight);
+											return TraceServices::EEventEnumerate::Continue;
+										});
+									});
 
-							AnimationProvider->ReadStateMachinesTimeline(GetGameplayTrack().GetObjectId(), [InGraphStartTime, InGraphEndTime, &DebugData](const FAnimationProvider::StateMachinesTimeline& InStateMachinesTimeline)
-							{
-								InStateMachinesTimeline.EnumerateEvents(InGraphStartTime, InGraphEndTime, [&DebugData](double InStartTime, double InEndTime, uint32 InDepth, const FAnimStateMachineMessage& InMessage)
-								{
-									DebugData.RecordStateData(InMessage.StateMachineIndex, InMessage.StateIndex, InMessage.StateWeight, InMessage.ElapsedTime);
-									return TraceServices::EEventEnumerate::Continue;
-								});
-							});
+									AnimationProvider->ReadStateMachinesTimeline(GetGameplayTrack().GetObjectId(), [InGraphStartTime, InGraphEndTime, &DebugData](const FAnimationProvider::StateMachinesTimeline& InStateMachinesTimeline)
+									{
+										InStateMachinesTimeline.EnumerateEvents(InGraphStartTime, InGraphEndTime, [&DebugData](double InStartTime, double InEndTime, uint32 InDepth, const FAnimStateMachineMessage& InMessage)
+										{
+											DebugData.RecordStateData(InMessage.StateMachineIndex, InMessage.StateIndex, InMessage.StateWeight, InMessage.ElapsedTime);
+											return TraceServices::EEventEnumerate::Continue;
+										});
+									});
 
-							AnimationProvider->ReadAnimNodeValuesTimeline(GetGameplayTrack().GetObjectId(), [InGraphStartTime, InGraphEndTime, AnimationProvider, &DebugData](const FAnimationProvider::AnimNodeValuesTimeline& InNodeValuesTimeline)
-							{
-								InNodeValuesTimeline.EnumerateEvents(InGraphStartTime, InGraphEndTime, [AnimationProvider, &DebugData](double InStartTime, double InEndTime, uint32 InDepth, const FAnimNodeValueMessage& InMessage)
-								{
-									FText Text = AnimationProvider->FormatNodeKeyValue(InMessage);
-									DebugData.RecordNodeValue(InMessage.NodeId, Text.ToString());
-									return TraceServices::EEventEnumerate::Continue;
-								});
-							});
+									AnimationProvider->ReadAnimSequencePlayersTimeline(GetGameplayTrack().GetObjectId(), [InGraphStartTime, InGraphEndTime, GameplayProvider, &DebugData](const FAnimationProvider::AnimSequencePlayersTimeline& InSequencePlayersTimeline)
+									{
+										InSequencePlayersTimeline.EnumerateEvents(InGraphStartTime, InGraphEndTime, [&DebugData](double InStartTime, double InEndTime, uint32 InDepth, const FAnimSequencePlayerMessage& InMessage)
+										{
+											DebugData.RecordSequencePlayer(InMessage.NodeId, InMessage.Position, InMessage.Length, InMessage.FrameCounter);
+											return TraceServices::EEventEnumerate::Continue;
+										});
+									});
 
-							AnimationProvider->ReadAnimSequencePlayersTimeline(GetGameplayTrack().GetObjectId(), [InGraphStartTime, InGraphEndTime, GameplayProvider, &DebugData](const FAnimationProvider::AnimSequencePlayersTimeline& InSequencePlayersTimeline)
-							{
-								InSequencePlayersTimeline.EnumerateEvents(InGraphStartTime, InGraphEndTime, [&DebugData](double InStartTime, double InEndTime, uint32 InDepth, const FAnimSequencePlayerMessage& InMessage)
-								{
-									DebugData.RecordSequencePlayer(InMessage.NodeId, InMessage.Position, InMessage.Length, InMessage.FrameCounter);
-									return TraceServices::EEventEnumerate::Continue;
-								});
-							});
+									AnimationProvider->ReadAnimBlendSpacePlayersTimeline(GetGameplayTrack().GetObjectId(), [InGraphStartTime, InGraphEndTime, GameplayProvider, &DebugData](const FAnimationProvider::BlendSpacePlayersTimeline& InBlendSpacePlayersTimeline)
+									{
+										InBlendSpacePlayersTimeline.EnumerateEvents(InGraphStartTime, InGraphEndTime, [GameplayProvider, &DebugData](double InStartTime, double InEndTime, uint32 InDepth, const FBlendSpacePlayerMessage& InMessage)
+										{
+											UBlendSpaceBase* BlendSpaceBase = nullptr;
+											const FObjectInfo* BlendSpaceInfo = GameplayProvider->FindObjectInfo(InMessage.BlendSpaceId);
+											if(BlendSpaceInfo)
+											{
+												BlendSpaceBase = TSoftObjectPtr<UBlendSpaceBase>(FSoftObjectPath(BlendSpaceInfo->PathName)).LoadSynchronous();
+											}
 
-							AnimationProvider->ReadAnimBlendSpacePlayersTimeline(GetGameplayTrack().GetObjectId(), [InGraphStartTime, InGraphEndTime, GameplayProvider, &DebugData](const FAnimationProvider::BlendSpacePlayersTimeline& InBlendSpacePlayersTimeline)
-							{
-								InBlendSpacePlayersTimeline.EnumerateEvents(InGraphStartTime, InGraphEndTime, [GameplayProvider, &DebugData](double InStartTime, double InEndTime, uint32 InDepth, const FBlendSpacePlayerMessage& InMessage)
+											DebugData.RecordBlendSpacePlayer(InMessage.NodeId, BlendSpaceBase, InMessage.PositionX, InMessage.PositionY, InMessage.PositionZ);
+											return TraceServices::EEventEnumerate::Continue;
+										});
+									});
+
+									AnimationProvider->ReadAnimSyncTimeline(GetGameplayTrack().GetObjectId(), [InGraphStartTime, InGraphEndTime, AnimationProvider, &DebugData](const FAnimationProvider::AnimSyncTimeline& InAnimSyncTimeline)
+									{
+										InAnimSyncTimeline.EnumerateEvents(InGraphStartTime, InGraphEndTime, [AnimationProvider, &DebugData](double InStartTime, double InEndTime, uint32 InDepth, const FAnimSyncMessage& InMessage)
+										{
+											const TCHAR* GroupName = AnimationProvider->GetName(InMessage.GroupNameId);
+											if(GroupName)
+											{
+												DebugData.RecordNodeSync(InMessage.SourceNodeId, FName(GroupName));
+											}
+									
+											return TraceServices::EEventEnumerate::Continue;
+										});
+									});
+								}
+
+								// Some traces come from both update and evaluate phases
+								if(InMessage.Phase == EAnimGraphPhase::Update || InMessage.Phase == EAnimGraphPhase::Evaluate)
 								{
-									UBlendSpaceBase* BlendSpaceBase = nullptr;
-									const FObjectInfo* BlendSpaceInfo = GameplayProvider->FindObjectInfo(InMessage.BlendSpaceId);
-									if(BlendSpaceInfo)
+									AnimationProvider->ReadAnimAttributesTimeline(GetGameplayTrack().GetObjectId(), [InGraphStartTime, InGraphEndTime, AnimationProvider, &DebugData](const FAnimationProvider::AnimAttributeTimeline& InAnimAttributeTimeline)
 									{
-										BlendSpaceBase = TSoftObjectPtr<UBlendSpaceBase>(FSoftObjectPath(BlendSpaceInfo->PathName)).LoadSynchronous();
-									}
+										InAnimAttributeTimeline.EnumerateEvents(InGraphStartTime, InGraphEndTime, [AnimationProvider, &DebugData](double InStartTime, double InEndTime, uint32 InDepth, const FAnimAttributeMessage& InMessage)
+										{
+											const TCHAR* AttributeName = AnimationProvider->GetName(InMessage.AttributeNameId);
+											if(AttributeName)
+											{
+												DebugData.RecordNodeAttribute(InMessage.TargetNodeId, InMessage.SourceNodeId, FName(AttributeName));
+											}
+									
+											return TraceServices::EEventEnumerate::Continue;
+										});
+									});
+								}
 
-									DebugData.RecordBlendSpacePlayer(InMessage.NodeId, BlendSpaceBase, InMessage.PositionX, InMessage.PositionY, InMessage.PositionZ);
-									return TraceServices::EEventEnumerate::Continue;
+								// Anim node values can come from all phases
+								AnimationProvider->ReadAnimNodeValuesTimeline(GetGameplayTrack().GetObjectId(), [InGraphStartTime, InGraphEndTime, AnimationProvider, &DebugData](const FAnimationProvider::AnimNodeValuesTimeline& InNodeValuesTimeline)
+								{
+									InNodeValuesTimeline.EnumerateEvents(InGraphStartTime, InGraphEndTime, [AnimationProvider, &DebugData](double InStartTime, double InEndTime, uint32 InDepth, const FAnimNodeValueMessage& InMessage)
+									{
+										FText Text = AnimationProvider->FormatNodeKeyValue(InMessage);
+										DebugData.RecordNodeValue(InMessage.NodeId, Text.ToString());
+										return TraceServices::EEventEnumerate::Continue;
+									});
 								});
-							});
-						}
-					}
-					return TraceServices::EEventEnumerate::Continue;
-				});
-			});
-		}
+							}
+							return TraceServices::EEventEnumerate::Continue;
+						});
+					});
+				}
+			}
+		}	
 	}
 }
 
diff --git a/Engine/Plugins/Animation/GameplayInsights/Source/GameplayInsights/Private/AnimationAnalyzer.cpp b/Engine/Plugins/Animation/GameplayInsights/Source/GameplayInsights/Private/AnimationAnalyzer.cpp
--- a/Engine/Plugins/Animation/GameplayInsights/Source/GameplayInsights/Private/AnimationAnalyzer.cpp
+++ b/Engine/Plugins/Animation/GameplayInsights/Source/GameplayInsights/Private/AnimationAnalyzer.cpp
@@ -24,6 +24,7 @@ void FAnimationAnalyzer::OnAnalysisBegin(const FOnAnalysisContext& Context)
 	Builder.RouteEvent(RouteId_SkeletalMeshFrame, "Animation", "SkeletalMeshFrame");
 	Builder.RouteEvent(RouteId_AnimGraph, "Animation", "AnimGraph");
 	Builder.RouteEvent(RouteId_AnimNodeStart, "Animation", "AnimNodeStart");
+	Builder.RouteEvent(RouteId_AnimNodeAttribute, "Animation", "AnimNodeAttribute");
 	Builder.RouteEvent(RouteId_AnimNodeValueBool, "Animation", "AnimNodeValueBool");
 	Builder.RouteEvent(RouteId_AnimNodeValueInt, "Animation", "AnimNodeValueInt");
 	Builder.RouteEvent(RouteId_AnimNodeValueFloat, "Animation", "AnimNodeValueFloat");
@@ -39,6 +40,7 @@ void FAnimationAnalyzer::OnAnalysisBegin(const FOnAnalysisContext& Context)
 	Builder.RouteEvent(RouteId_Notify, "Animation", "Notify");
 	Builder.RouteEvent(RouteId_SyncMarker, "Animation", "SyncMarker");
 	Builder.RouteEvent(RouteId_Montage, "Animation", "Montage");
+	Builder.RouteEvent(RouteId_Sync, "Animation", "Sync");
 }
 
 bool FAnimationAnalyzer::OnEvent(uint16 RouteId, EStyle Style, const FOnEventContext& Context)
@@ -181,6 +183,17 @@ bool FAnimationAnalyzer::OnEvent(uint16 RouteId, EStyle Style, const FOnEventCon
 		AnimationProvider.AppendAnimNodeStart(AnimInstanceId, Context.EventTime.AsSeconds(StartCycle), FrameCounter, NodeId, PreviousNodeId, Weight, RootMotionWeight, TargetNodeName, Phase);
 		break;
 	}
+	case RouteId_AnimNodeAttribute:
+	{
+		uint64 Cycle = EventData.GetValue<uint64>("Cycle");
+		uint64 SourceAnimInstanceId = EventData.GetValue<uint64>("SourceAnimInstanceId");
+		uint64 TargetAnimInstanceId = EventData.GetValue<uint64>("TargetAnimInstanceId");
+		int32 SourceNodeId = EventData.GetValue<int32>("SourceNodeId");
+		int32 TargetNodeId = EventData.GetValue<int32>("TargetNodeId");
+		uint32 NameId = EventData.GetValue<uint32>("NameId");
+		AnimationProvider.AppendAnimGraphAttribute(SourceAnimInstanceId, TargetAnimInstanceId, Context.EventTime.AsSeconds(Cycle), SourceNodeId, TargetNodeId, NameId);
+		break;
+	}
 	case RouteId_AnimNodeValueBool:
 	case RouteId_AnimNodeValueInt:
 	case RouteId_AnimNodeValueFloat:
@@ -318,6 +331,15 @@ bool FAnimationAnalyzer::OnEvent(uint16 RouteId, EStyle Style, const FOnEventCon
 		AnimationProvider.AppendMontage(AnimInstanceId, Context.EventTime.AsSeconds(Cycle), MontageId, CurrentSectionNameId, NextSectionNameId, Weight, DesiredWeight, FrameCounter);
 		break;
 	}
+	case RouteId_Sync:
+	{
+		uint64 Cycle = EventData.GetValue<uint64>("Cycle");
+		uint64 AnimInstanceId = EventData.GetValue<uint64>("AnimInstanceId");
+		int32 SourceNodeId = EventData.GetValue<int32>("SourceNodeId");
+		uint32 GroupNameId = EventData.GetValue<uint32>("GroupNameId");
+		AnimationProvider.AppendSync(AnimInstanceId, Context.EventTime.AsSeconds(Cycle), SourceNodeId, GroupNameId);
+		break;
+	}
 	}
 
 	return true;
diff --git a/Engine/Plugins/Animation/GameplayInsights/Source/GameplayInsights/Private/AnimationAnalyzer.h b/Engine/Plugins/Animation/GameplayInsights/Source/GameplayInsights/Private/AnimationAnalyzer.h
--- a/Engine/Plugins/Animation/GameplayInsights/Source/GameplayInsights/Private/AnimationAnalyzer.h
+++ b/Engine/Plugins/Animation/GameplayInsights/Source/GameplayInsights/Private/AnimationAnalyzer.h
@@ -27,6 +27,7 @@ private:
 		RouteId_SkeletalMeshFrame,
 		RouteId_AnimGraph,
 		RouteId_AnimNodeStart,
+		RouteId_AnimNodeAttribute,
 		RouteId_AnimNodeValueBool,
 		RouteId_AnimNodeValueInt,
 		RouteId_AnimNodeValueFloat,
@@ -42,6 +43,7 @@ private:
 		RouteId_Notify,
 		RouteId_SyncMarker,
 		RouteId_Montage,
+		RouteId_Sync,
 	};
 
 	TraceServices::IAnalysisSession& Session;
diff --git a/Engine/Plugins/Animation/GameplayInsights/Source/GameplayInsights/Private/AnimationProvider.cpp b/Engine/Plugins/Animation/GameplayInsights/Source/GameplayInsights/Private/AnimationProvider.cpp
--- a/Engine/Plugins/Animation/GameplayInsights/Source/GameplayInsights/Private/AnimationProvider.cpp
+++ b/Engine/Plugins/Animation/GameplayInsights/Source/GameplayInsights/Private/AnimationProvider.cpp
@@ -176,6 +176,23 @@ bool FAnimationProvider::ReadAnimNodeValuesTimeline(uint64 InObjectId, TFunction
 	return false;
 }
 
+bool FAnimationProvider::ReadAnimAttributesTimeline(uint64 InObjectId, TFunctionRef<void(const AnimAttributeTimeline&)> Callback) const
+{
+	Session.ReadAccessCheck();
+
+	const uint32* IndexPtr = ObjectIdToAnimAttributeTimelines.Find(InObjectId);
+	if(IndexPtr != nullptr)
+	{
+		if (*IndexPtr < uint32(AnimAttributeTimelines.Num()))
+		{
+			Callback(*AnimAttributeTimelines[*IndexPtr]);
+			return true;
+		}
+	}
+
+	return false;
+}
+
 bool FAnimationProvider::ReadAnimSequencePlayersTimeline(uint64 InObjectId, TFunctionRef<void(const AnimSequencePlayersTimeline&)> Callback) const
 {
 	Session.ReadAccessCheck();
@@ -295,6 +312,23 @@ void FAnimationProvider::EnumerateMontageIds(uint64 InObjectId, TFunctionRef<voi
 	}
 }
 
+bool FAnimationProvider::ReadAnimSyncTimeline(uint64 InObjectId, TFunctionRef<void(const AnimSyncTimeline&)> Callback) const
+{
+	Session.ReadAccessCheck();
+
+	const uint32* IndexPtr = ObjectIdToAnimSyncTimelines.Find(InObjectId);
+	if(IndexPtr != nullptr)
+	{
+		if (*IndexPtr < uint32(AnimSyncTimelines.Num()))
+		{
+			Callback(*AnimSyncTimelines[*IndexPtr]);
+			return true;
+		}
+	}
+
+	return false;
+}
+
 const FSkeletalMeshInfo* FAnimationProvider::FindSkeletalMeshInfo(uint64 InObjectId) const
 {
 	Session.ReadAccessCheck();
@@ -850,6 +884,81 @@ void FAnimationProvider::AppendAnimNodeValue(uint64 InAnimInstanceId, double InT
 	Session.UpdateDurationSeconds(InTime);
 }
 
+void FAnimationProvider::AppendAnimGraphAttribute(uint64 InSourceAnimInstanceId, uint64 InTargetAnimInstanceId, double InTime, int32 InSourceNodeId, int32 InTargetNodeId, uint32 InAttributeNameId)
+{
+	Session.WriteAccessCheck();
+
+	bHasAnyData = true;
+
+	if(InSourceAnimInstanceId == InTargetAnimInstanceId)
+	{
+		TSharedPtr<TraceServices::TPointTimeline<FAnimAttributeMessage>> Timeline;
+		uint32* IndexPtr = ObjectIdToAnimAttributeTimelines.Find(InSourceAnimInstanceId);
+		if(IndexPtr != nullptr)
+		{
+			Timeline = AnimAttributeTimelines[*IndexPtr];
+		}
+		else
+		{
+			Timeline = MakeShared<TraceServices::TPointTimeline<FAnimAttributeMessage>>(Session.GetLinearAllocator());
+			ObjectIdToAnimAttributeTimelines.Add(InSourceAnimInstanceId, AnimAttributeTimelines.Num());
+			AnimAttributeTimelines.Add(Timeline.ToSharedRef());
+		}
+
+		FAnimAttributeMessage Message;
+		Message.SourceNodeId = InSourceNodeId;
+		Message.TargetNodeId = InTargetNodeId;
+		Message.AttributeNameId = InAttributeNameId;
+
+		Timeline->AppendEvent(InTime, Message);
+	}
+	else
+	{
+		// If we are using two different anim instances, we need to append to two different timelines
+		TSharedPtr<TraceServices::TPointTimeline<FAnimAttributeMessage>> SourceTimeline;
+		uint32* SourceIndexPtr = ObjectIdToAnimAttributeTimelines.Find(InSourceAnimInstanceId);
+		if(SourceIndexPtr != nullptr)
+		{
+			SourceTimeline = AnimAttributeTimelines[*SourceIndexPtr];
+		}
+		else
+		{
+			SourceTimeline = MakeShared<TraceServices::TPointTimeline<FAnimAttributeMessage>>(Session.GetLinearAllocator());
+			ObjectIdToAnimAttributeTimelines.Add(InSourceAnimInstanceId, AnimAttributeTimelines.Num());
+			AnimAttributeTimelines.Add(SourceTimeline.ToSharedRef());
+		}
+
+		TSharedPtr<TraceServices::TPointTimeline<FAnimAttributeMessage>> TargetTimeline;
+		uint32* TargetIndexPtr = ObjectIdToAnimAttributeTimelines.Find(InTargetAnimInstanceId);
+		if(TargetIndexPtr != nullptr)
+		{
+			TargetTimeline = AnimAttributeTimelines[*TargetIndexPtr];
+		}
+		else
+		{
+			TargetTimeline = MakeShared<TraceServices::TPointTimeline<FAnimAttributeMessage>>(Session.GetLinearAllocator());
+			ObjectIdToAnimAttributeTimelines.Add(InTargetAnimInstanceId, AnimAttributeTimelines.Num());
+			AnimAttributeTimelines.Add(TargetTimeline.ToSharedRef());
+		}
+
+		FAnimAttributeMessage TargetMessage;
+		TargetMessage.SourceNodeId = INDEX_NONE;
+		TargetMessage.TargetNodeId = InTargetNodeId;
+		TargetMessage.AttributeNameId = InAttributeNameId;
+
+		TargetTimeline->AppendEvent(InTime, TargetMessage);
+
+		FAnimAttributeMessage SourceMessage;
+		SourceMessage.SourceNodeId = InSourceNodeId;
+		SourceMessage.TargetNodeId = INDEX_NONE;
+		SourceMessage.AttributeNameId = InAttributeNameId;
+
+		SourceTimeline->AppendEvent(InTime, SourceMessage);
+	}
+
+	Session.UpdateDurationSeconds(InTime);
+}
+
 void FAnimationProvider::AppendStateMachineState(uint64 InAnimInstanceId, double InTime, int32 InNodeId, int32 InStateMachineIndex, int32 InStateIndex, float InStateWeight, float InElapsedTime)
 {
 	Session.WriteAccessCheck();
@@ -1005,4 +1114,33 @@ void FAnimationProvider::AppendMontage(uint64 InAnimInstanceId, double InTime, u
 	Session.UpdateDurationSeconds(InTime);
 }
 
+void FAnimationProvider::AppendSync(uint64 InAnimInstanceId, double InTime, int32 InSourceNodeId, uint32 InGroupNameId)
+{
+	Session.WriteAccessCheck();
+
+	bHasAnyData = true;
+
+	TSharedPtr<TraceServices::TPointTimeline<FAnimSyncMessage>> Timeline;
+	uint32* IndexPtr = ObjectIdToAnimSyncTimelines.Find(InAnimInstanceId);
+	if(IndexPtr != nullptr)
+	{
+		check(AnimSyncTimelines.IsValidIndex(*IndexPtr));
+		Timeline = AnimSyncTimelines[*IndexPtr];
+	}
+	else
+	{
+		Timeline = MakeShared<TraceServices::TPointTimeline<FAnimSyncMessage>>(Session.GetLinearAllocator());
+		ObjectIdToAnimSyncTimelines.Add(InAnimInstanceId, AnimSyncTimelines.Num());
+		AnimSyncTimelines.Add(Timeline.ToSharedRef());
+	}
+
+	FAnimSyncMessage Message;
+	Message.SourceNodeId = InSourceNodeId;
+	Message.GroupNameId = InGroupNameId;
+
+	Timeline->AppendEvent(InTime, Message);
+
+	Session.UpdateDurationSeconds(InTime);
+}
+
 #undef LOCTEXT_NAMESPACE
diff --git a/Engine/Plugins/Animation/GameplayInsights/Source/GameplayInsights/Private/AnimationProvider.h b/Engine/Plugins/Animation/GameplayInsights/Source/GameplayInsights/Private/AnimationProvider.h
--- a/Engine/Plugins/Animation/GameplayInsights/Source/GameplayInsights/Private/AnimationProvider.h
+++ b/Engine/Plugins/Animation/GameplayInsights/Source/GameplayInsights/Private/AnimationProvider.h
@@ -31,6 +31,7 @@ public:
 	virtual bool ReadAnimGraphTimeline(uint64 InObjectId, TFunctionRef<void(const AnimGraphTimeline&)> Callback) const override;
 	virtual bool ReadAnimNodesTimeline(uint64 InObjectId, TFunctionRef<void(const AnimNodesTimeline&)> Callback) const override;
 	virtual bool ReadAnimNodeValuesTimeline(uint64 InObjectId, TFunctionRef<void(const AnimNodeValuesTimeline&)> Callback) const override;
+	virtual bool ReadAnimAttributesTimeline(uint64 InObjectId, TFunctionRef<void(const AnimAttributeTimeline&)> Callback) const override;
 	virtual bool ReadAnimSequencePlayersTimeline(uint64 InObjectId, TFunctionRef<void(const AnimSequencePlayersTimeline&)> Callback) const override;
 	virtual bool ReadAnimBlendSpacePlayersTimeline(uint64 InObjectId, TFunctionRef<void(const BlendSpacePlayersTimeline&)> Callback) const override;
 	virtual bool ReadStateMachinesTimeline(uint64 InObjectId, TFunctionRef<void(const StateMachinesTimeline&)> Callback) const override;
@@ -38,6 +39,7 @@ public:
 	virtual void EnumerateNotifyStateTimelines(uint64 InObjectId, TFunctionRef<void(uint64, const AnimNotifyTimeline&)> Callback) const override;
 	virtual bool ReadMontageTimeline(uint64 InObjectId, TFunctionRef<void(const AnimMontageTimeline&)> Callback) const override;
 	virtual void EnumerateMontageIds(uint64 InObjectId, TFunctionRef<void(uint64)> Callback) const override;
+	virtual bool ReadAnimSyncTimeline(uint64 InObjectId, TFunctionRef<void(const AnimSyncTimeline&)> Callback) const override;
 	virtual const FSkeletalMeshInfo* FindSkeletalMeshInfo(uint64 InObjectId) const override;
 	virtual const TCHAR* GetName(uint32 InId) const override;
 	virtual FText FormatNodeKeyValue(const FAnimNodeValueMessage& InMessage) const override;
@@ -81,6 +83,9 @@ public:
 	void AppendAnimNodeValueObject(uint64 InAnimInstanceId, double InTime, uint16 InFrameCounter, int32 InNodeId, const TCHAR* InKey, uint64 InValue);
 	void AppendAnimNodeValueClass(uint64 InAnimInstanceId, double InTime, uint16 InFrameCounter, int32 InNodeId, const TCHAR* InKey, uint64 InValue);
 
+	/** Add anim graph attribute */
+	void AppendAnimGraphAttribute(uint64 InSourceAnimInstanceId, uint64 InTargetAnimInstanceId, double InTime, int32 InSourceNodeId, int32 InTargetNodeId, uint32 InAttributeNameId);
+
 	/** Add a sequence player record */
 	void AppendAnimSequencePlayer(uint64 InAnimInstanceId, double InTime, int32 InNodeId, float InPosition, float InLength, uint16 InFrameCounter);
 
@@ -96,6 +101,9 @@ public:
 	/** Append montage data */
 	void AppendMontage(uint64 InAnimInstanceId, double InTime, uint64 InMontageId, uint32 InCurrentSectionNameId, uint32 InNextSectionNameId, float InWeight, float InDesiredWeight, uint16 InFrameCounter);
 
+	/** Append sync data */
+	void AppendSync(uint64 InAnimInstanceId, double InTime, int32 InSourceNodeId, uint32 InGroupNameId);
+
 private:
 	/** Add anim node values helper */
 	void AppendAnimNodeValue(uint64 InAnimInstanceId, double InTime, uint16 InFrameCounter, int32 InNodeId, const TCHAR* InKey, FAnimNodeValueMessage& InMessage);
@@ -120,6 +128,8 @@ private:
 	TMap<uint64, uint32> ObjectIdToAnimNotifyStateTimelines;
 	TMap<uint64, uint32> ObjectIdToAnimNotifyTimelines;
 	TMap<uint64, uint32> ObjectIdToAnimMontageTimelines;
+	TMap<uint64, uint32> ObjectIdToAnimAttributeTimelines;
+	TMap<uint64, uint32> ObjectIdToAnimSyncTimelines;
 
 	/** All the skeletal mesh info we have seen, grow only for stable indices */
 	TArray<FSkeletalMeshInfo> SkeletalMeshInfos;
@@ -170,6 +180,8 @@ private:
 	TraceServices::TPagedArray<FTransform> SkeletalMeshPoseTransforms;
 	TraceServices::TPagedArray<FSkeletalMeshNamedCurve> SkeletalMeshCurves;
 	TraceServices::TPagedArray<int32> SkeletalMeshParentIndices;
+	TArray<TSharedRef<TraceServices::TPointTimeline<FAnimAttributeMessage>>> AnimAttributeTimelines;
+	TArray<TSharedRef<TraceServices::TPointTimeline<FAnimSyncMessage>>> AnimSyncTimelines;
 
 	// Flag to indicate if any data is present
 	bool bHasAnyData;
diff --git a/Engine/Plugins/Animation/GameplayInsights/Source/GameplayInsights/Public/IAnimationProvider.h b/Engine/Plugins/Animation/GameplayInsights/Source/GameplayInsights/Public/IAnimationProvider.h
--- a/Engine/Plugins/Animation/GameplayInsights/Source/GameplayInsights/Public/IAnimationProvider.h
+++ b/Engine/Plugins/Animation/GameplayInsights/Source/GameplayInsights/Public/IAnimationProvider.h
@@ -227,6 +227,19 @@ struct FAnimMontageMessage
 	uint16 FrameCounter = 0;
 };
 
+struct FAnimAttributeMessage
+{
+	int32 SourceNodeId = 0;
+	int32 TargetNodeId = 0;
+	uint32 AttributeNameId = 0;
+};
+
+struct FAnimSyncMessage
+{
+	int32 SourceNodeId = 0;
+	uint32 GroupNameId = 0;
+};
+
 class IAnimationProvider : public TraceServices::IProvider
 {
 public:
@@ -240,6 +253,8 @@ public:
 	typedef TraceServices::ITimeline<FBlendSpacePlayerMessage> BlendSpacePlayersTimeline;
 	typedef TraceServices::ITimeline<FAnimNotifyMessage> AnimNotifyTimeline;
 	typedef TraceServices::ITimeline<FAnimMontageMessage> AnimMontageTimeline;
+	typedef TraceServices::ITimeline<FAnimAttributeMessage> AnimAttributeTimeline;
+	typedef TraceServices::ITimeline<FAnimSyncMessage> AnimSyncTimeline;
 
 	virtual bool ReadSkeletalMeshPoseTimeline(uint64 InObjectId, TFunctionRef<void(const SkeletalMeshPoseTimeline&, bool)> Callback) const = 0;
 	virtual void GetSkeletalMeshComponentSpacePose(const FSkeletalMeshPoseMessage& InMessage, const FSkeletalMeshInfo& InMeshInfo, FTransform& OutComponentToWorld, TArray<FTransform>& OutTransforms) const = 0;
@@ -250,6 +265,7 @@ public:
 	virtual bool ReadAnimGraphTimeline(uint64 InObjectId, TFunctionRef<void(const AnimGraphTimeline&)> Callback) const = 0;
 	virtual bool ReadAnimNodesTimeline(uint64 InObjectId, TFunctionRef<void(const AnimNodesTimeline&)> Callback) const = 0;
 	virtual bool ReadAnimNodeValuesTimeline(uint64 InObjectId, TFunctionRef<void(const AnimNodeValuesTimeline&)> Callback) const = 0;
+	virtual bool ReadAnimAttributesTimeline(uint64 InObjectId, TFunctionRef<void(const AnimAttributeTimeline&)> Callback) const = 0;
 	virtual bool ReadAnimSequencePlayersTimeline(uint64 InObjectId, TFunctionRef<void(const AnimSequencePlayersTimeline&)> Callback) const = 0;
 	virtual bool ReadAnimBlendSpacePlayersTimeline(uint64 InObjectId, TFunctionRef<void(const BlendSpacePlayersTimeline&)> Callback) const = 0;
 	virtual bool ReadStateMachinesTimeline(uint64 InObjectId, TFunctionRef<void(const StateMachinesTimeline&)> Callback) const = 0;
@@ -257,6 +273,7 @@ public:
 	virtual void EnumerateNotifyStateTimelines(uint64 InObjectId, TFunctionRef<void(uint64, const AnimNotifyTimeline&)> Callback) const = 0;
 	virtual bool ReadMontageTimeline(uint64 InObjectId, TFunctionRef<void(const AnimMontageTimeline&)> Callback) const = 0;
 	virtual void EnumerateMontageIds(uint64 InObjectId, TFunctionRef<void(uint64)> Callback) const = 0;
+	virtual bool ReadAnimSyncTimeline(uint64 InObjectId, TFunctionRef<void(const AnimSyncTimeline&)> Callback) const = 0;
 	virtual const FSkeletalMeshInfo* FindSkeletalMeshInfo(uint64 InObjectId) const = 0;
 	virtual const TCHAR* GetName(uint32 InId) const = 0;
 	virtual FText FormatNodeKeyValue(const FAnimNodeValueMessage& InMessage) const = 0;
diff --git a/Engine/Plugins/Animation/LiveLink/Source/LiveLinkGraphNode/Private/AnimGraphNode_LiveLinkPose.cpp b/Engine/Plugins/Animation/LiveLink/Source/LiveLinkGraphNode/Private/AnimGraphNode_LiveLinkPose.cpp
--- a/Engine/Plugins/Animation/LiveLink/Source/LiveLinkGraphNode/Private/AnimGraphNode_LiveLinkPose.cpp
+++ b/Engine/Plugins/Animation/LiveLink/Source/LiveLinkGraphNode/Private/AnimGraphNode_LiveLinkPose.cpp
@@ -1,8 +1,8 @@
 // Copyright Epic Games, Inc. All Rights Reserved.
 
 #include "AnimGraphNode_LiveLinkPose.h"
-
 #include "EdGraph/EdGraphSchema.h"
+#include "Animation/AnimAttributes.h"
 
 #define LOCTEXT_NAMESPACE "LiveLinkAnimNode"
 
@@ -51,4 +51,10 @@ void UAnimGraphNode_LiveLinkPose::ConvertDeprecatedNode(UEdGraph* Graph, bool bO
 	}
 }
 
+void UAnimGraphNode_LiveLinkPose::GetOutputLinkAttributes(FNodeAttributeArray& OutAttributes) const
+{
+	OutAttributes.Add(UE::Anim::FAttributes::Curves);
+	OutAttributes.Add(UE::Anim::FAttributes::Attributes);
+}
+
 #undef LOCTEXT_NAMESPACE
\ No newline at end of file
diff --git a/Engine/Plugins/Animation/LiveLink/Source/LiveLinkGraphNode/Private/AnimGraphNode_LiveLinkPose.h b/Engine/Plugins/Animation/LiveLink/Source/LiveLinkGraphNode/Private/AnimGraphNode_LiveLinkPose.h
--- a/Engine/Plugins/Animation/LiveLink/Source/LiveLinkGraphNode/Private/AnimGraphNode_LiveLinkPose.h
+++ b/Engine/Plugins/Animation/LiveLink/Source/LiveLinkGraphNode/Private/AnimGraphNode_LiveLinkPose.h
@@ -24,6 +24,7 @@ public:
 
 	//~ Begin UK2Node interface
 	virtual void ConvertDeprecatedNode(UEdGraph* Graph, bool bOnlySafeChanges) override;
+	virtual void GetOutputLinkAttributes(FNodeAttributeArray& OutAttributes) const override;
 	//~ End UK2Node interface
 
 public:
diff --git a/Engine/Source/Editor/AnimGraph/Private/AnimBlueprintCompilationContext.cpp b/Engine/Source/Editor/AnimGraph/Private/AnimBlueprintCompilationContext.cpp
--- a/Engine/Source/Editor/AnimGraph/Private/AnimBlueprintCompilationContext.cpp
+++ b/Engine/Source/Editor/AnimGraph/Private/AnimBlueprintCompilationContext.cpp
@@ -83,6 +83,16 @@ const TMap<UAnimGraphNode_Base*, FProperty*>& FAnimBlueprintCompilationContext::
 	return CompilerContext->AllocatedAnimNodes;
 }
 
+void FAnimBlueprintCompilationContext::AddAttributesToNodeImpl(UAnimGraphNode_Base* InNode, TArrayView<const FName> InAttributes) const
+{
+	CompilerContext->AddAttributesToNode(InNode, InAttributes);
+}
+
+TArrayView<const FName> FAnimBlueprintCompilationContext::GetAttributesFromNodeImpl(UAnimGraphNode_Base* InNode) const
+{
+	return CompilerContext->GetAttributesFromNode(InNode);
+}
+
 IAnimBlueprintCompilerHandler* FAnimBlueprintCompilationContext::GetHandlerInternal(FName InName) const
 {
 	return CompilerContext->AnimBlueprintCompilerHandlerCollection.GetHandler<IAnimBlueprintCompilerHandler>(InName);
diff --git a/Engine/Source/Editor/AnimGraph/Private/AnimBlueprintCompilationContext.h b/Engine/Source/Editor/AnimGraph/Private/AnimBlueprintCompilationContext.h
--- a/Engine/Source/Editor/AnimGraph/Private/AnimBlueprintCompilationContext.h
+++ b/Engine/Source/Editor/AnimGraph/Private/AnimBlueprintCompilationContext.h
@@ -35,6 +35,8 @@ private:
 	virtual const TMap<UAnimGraphNode_Base*, UAnimGraphNode_Base*>& GetSourceNodeToProcessedNodeMapImpl() const override;
 	virtual const TMap<int32, FProperty*>& GetAllocatedPropertiesByIndexImpl() const override;
 	virtual const TMap<UAnimGraphNode_Base*, FProperty*>& GetAllocatedPropertiesByNodeImpl() const override;
+	virtual void AddAttributesToNodeImpl(UAnimGraphNode_Base* InNode, TArrayView<const FName> InAttributes) const override;
+	virtual TArrayView<const FName> GetAttributesFromNodeImpl(UAnimGraphNode_Base* InNode) const override;
 	virtual IAnimBlueprintCompilerHandler* GetHandlerInternal(FName InName) const override;
 	virtual FKismetCompilerContext* GetKismetCompiler() const override;
 
diff --git a/Engine/Source/Editor/AnimGraph/Private/AnimBlueprintCompiler.cpp b/Engine/Source/Editor/AnimGraph/Private/AnimBlueprintCompiler.cpp
--- a/Engine/Source/Editor/AnimGraph/Private/AnimBlueprintCompiler.cpp
+++ b/Engine/Source/Editor/AnimGraph/Private/AnimBlueprintCompiler.cpp
@@ -1251,6 +1251,33 @@ void FAnimBlueprintCompilerContext::SetCalculatedMetaDataAndFlags(UFunction* Fun
 	}
 }
 
+void FAnimBlueprintCompilerContext::AddAttributesToNode(UAnimGraphNode_Base* InNode, TArrayView<const FName> InAttributes) const
+{
+	if(UAnimGraphNode_Base* OriginalNode = CastChecked<UAnimGraphNode_Base>(MessageLog.FindSourceObject(InNode)))
+	{
+		TArray<FName>& AttributeSet = NewAnimBlueprintClass->GetAnimBlueprintDebugData().NodeAttributes.FindOrAdd(OriginalNode);
+		AttributeSet.Reserve(AttributeSet.Num() + InAttributes.Num());
+
+		for(const FName& Attribute : InAttributes)
+		{
+			AttributeSet.AddUnique(Attribute);
+		}
+	}
+}
+
+TArrayView<const FName> FAnimBlueprintCompilerContext::GetAttributesFromNode(UAnimGraphNode_Base* InNode) const
+{
+	if(UAnimGraphNode_Base* OriginalNode = CastChecked<UAnimGraphNode_Base>(MessageLog.FindSourceObject(InNode)))
+	{
+		if(const TArray<FName>* AttributeSetPtr = NewAnimBlueprintClass->GetAnimBlueprintDebugData().NodeAttributes.Find(OriginalNode))
+		{
+			return MakeArrayView(*AttributeSetPtr);
+		}
+	}
+
+	return TArrayView<const FName>();
+}
+
 //////////////////////////////////////////////////////////////////////////
 
 #undef LOCTEXT_NAMESPACE
\ No newline at end of file
diff --git a/Engine/Source/Editor/AnimGraph/Private/AnimBlueprintCompiler.h b/Engine/Source/Editor/AnimGraph/Private/AnimBlueprintCompiler.h
--- a/Engine/Source/Editor/AnimGraph/Private/AnimBlueprintCompiler.h
+++ b/Engine/Source/Editor/AnimGraph/Private/AnimBlueprintCompiler.h
@@ -10,6 +10,7 @@
 #include "AnimBlueprintCompilerHandlerCollection.h"
 #include "IAnimBlueprintCompilerHandlerCollection.h"
 #include "IAnimBlueprintCompilerCreationContext.h"
+#include "Containers/ArrayView.h"
 
 class UAnimationGraphSchema;
 class UAnimGraphNode_SaveCachedPose;
@@ -168,5 +169,11 @@ private:
 
 	// Expands split pins for a graph
 	void ExpandSplitPins(UEdGraph* InGraph);
+
+	// Add the specified compiled-in attribute uniquely to the specified node
+	void AddAttributesToNode(UAnimGraphNode_Base* InNode, TArrayView<const FName> InAttributes) const;
+
+	// Get the current compiled-in attributes uniquely assigned to the specified node
+	TArrayView<const FName> GetAttributesFromNode(UAnimGraphNode_Base* InNode) const;
 };
 
diff --git a/Engine/Source/Editor/AnimGraph/Private/AnimBlueprintCompilerHandler_Attributes.cpp b/Engine/Source/Editor/AnimGraph/Private/AnimBlueprintCompilerHandler_Attributes.cpp
new file mode 100644
--- /dev/null
+++ b/Engine/Source/Editor/AnimGraph/Private/AnimBlueprintCompilerHandler_Attributes.cpp
@@ -0,0 +1,248 @@
+// Copyright Epic Games, Inc. All Rights Reserved.
+
+#include "AnimBlueprintCompilerHandler_Attributes.h"
+#include "IAnimBlueprintCompilerCreationContext.h"
+#include "IAnimBlueprintCompilationContext.h"
+#include "AnimationGraphSchema.h"
+#include "AnimGraphNode_Base.h"
+#include "AnimGraphNode_SaveCachedPose.h"
+#include "AnimGraphNode_UseCachedPose.h"
+#include "AnimGraphNode_StateResult.h"
+#include "AnimationStateGraph.h"
+#include "AnimationStateMachineGraph.h"
+#include "AnimGraphNode_StateMachineBase.h"
+#include "AnimStateNode.h"
+#include "AnimGraphNode_Root.h"
+#include "AnimGraphAttributes.h"
+#include "AnimGraphNode_TransitionResult.h"
+#include "AnimGraphNode_CustomTransitionResult.h"
+#include "AnimationCustomTransitionGraph.h"
+#include "AnimStateTransitionNode.h"
+
+FAnimBlueprintCompilerHandler_Attributes::FAnimBlueprintCompilerHandler_Attributes(IAnimBlueprintCompilerCreationContext& InCreationContext)
+{
+	InCreationContext.OnPreProcessAnimationNodes().AddRaw(this, &FAnimBlueprintCompilerHandler_Attributes::PreProcessAnimationNodes);
+}
+
+void FAnimBlueprintCompilerHandler_Attributes::PreProcessAnimationNodes(TArrayView<UAnimGraphNode_Base*> InAnimNodes, IAnimBlueprintCompilationContext& InCompilationContext, IAnimBlueprintGeneratedClassCompiledData& OutCompiledData)
+{	
+	using FNodeAttributeMap = TMap<UAnimGraphNode_Base*, UAnimGraphNode_Base::FNodeAttributeArray>;
+	using FCachedPoseMap = TMultiMap<UAnimGraphNode_SaveCachedPose*, UAnimGraphNode_UseCachedPose*>;
+
+	FNodeAttributeMap AttributeInputNodes;
+	FNodeAttributeMap AttributeOutputNodes;
+	FCachedPoseMap SaveCachedPoseMap;
+	UAnimGraphNode_Base::FNodeAttributeArray Attributes;
+
+	for(UAnimGraphNode_Base* Node : InAnimNodes)
+	{
+		// Establish links between save/used cached pose nodes
+		if(UAnimGraphNode_UseCachedPose* UseCachedPoseNode = Cast<UAnimGraphNode_UseCachedPose>(Node))
+		{
+			if(ensure(UseCachedPoseNode->SaveCachedPoseNode.IsValid()))
+			{
+				SaveCachedPoseMap.Add(UseCachedPoseNode->SaveCachedPoseNode.Get(), UseCachedPoseNode);
+			}
+		}
+		
+		// Get I/O attributes from all nodes
+		Node->GetInputLinkAttributes(Attributes);
+		if(Attributes.Num() > 0)
+		{
+			AttributeInputNodes.Add(Node, Attributes);
+		}
+
+		Attributes.Reset();
+
+		Node->GetOutputLinkAttributes(Attributes);
+		if(Attributes.Num() > 0)
+		{
+			AttributeOutputNodes.Add(Node, Attributes);
+		}
+
+		Attributes.Reset();
+	}
+
+	// Utility struct to check whether nodes with outputs reach nodes with inputs further towards the root
+	struct FCheckOutputNodes
+	{
+	private:
+		IAnimBlueprintCompilationContext& CompilationContext;
+		const UAnimGraphNode_Base::FNodeAttributeArray& BaseAttributes;
+		const FNodeAttributeMap& AttributeInputNodes;
+		const FCachedPoseMap& SaveCachedPoseMap;
+		UAnimGraphNode_Base* OutputNode;
+		TSet<UEdGraphNode*> VisitedNodes;
+
+	public:
+		FCheckOutputNodes(IAnimBlueprintCompilationContext& InCompilationContext, const UAnimGraphNode_Base::FNodeAttributeArray& InAttributes, const FNodeAttributeMap& InAttributeInputNodes, const FCachedPoseMap& InSaveCachedPoseMap, UAnimGraphNode_Base* InNode)
+			: CompilationContext(InCompilationContext)
+			, BaseAttributes(InAttributes)
+			, AttributeInputNodes(InAttributeInputNodes)
+			, SaveCachedPoseMap(InSaveCachedPoseMap)
+			, OutputNode(InNode)
+		{
+			UAnimGraphNode_Base::FNodeAttributeArray AbsorbedAttributes = TraverseNodes_Recursive(OutputNode, BaseAttributes);
+
+			if(AbsorbedAttributes.Num() > 0)
+			{
+				// Push only absorbed attributes to the source node's internal set
+				CompilationContext.AddAttributesToNode(OutputNode, AbsorbedAttributes);
+			}
+		}
+
+	private:
+		UAnimGraphNode_Base::FNodeAttributeArray TraverseNodes_Recursive_PerNode(UAnimGraphNode_Base* InLinkedNode, const UAnimGraphNode_Base::FNodeAttributeArray& InAttributes)
+		{
+			UAnimGraphNode_Base::FNodeAttributeArray AbsorbedAttributes;
+
+			if (InLinkedNode)
+			{
+				if(!VisitedNodes.Contains(InLinkedNode))
+				{
+					// See if this node absorbs this attribute
+					const UAnimGraphNode_Base::FNodeAttributeArray* AbsorbedAttributesPtr = AttributeInputNodes.Find(InLinkedNode);
+
+					auto HasAttribute = [AbsorbedAttributesPtr, InLinkedNode](FName InAttribute)
+					{
+						return (AbsorbedAttributesPtr && AbsorbedAttributesPtr->Contains(InAttribute));
+					};
+
+					if(InLinkedNode->IsA<UAnimGraphNode_Root>() || InLinkedNode->IsA<UAnimGraphNode_TransitionResult>() || InAttributes.ContainsByPredicate(HasAttribute))
+					{
+						UAnimGraphNode_Base::FNodeAttributeArray ReducedAttributes;
+
+						if(InLinkedNode->IsA<UAnimGraphNode_Root>() || InLinkedNode->IsA<UAnimGraphNode_TransitionResult>())
+						{
+							const UAnimGraphAttributes* AnimGraphAttributes = GetDefault<UAnimGraphAttributes>();
+
+							auto IsAttributeBlendable = [AnimGraphAttributes](const FName& InAttribute)
+							{
+								const FAnimGraphAttributeDesc* Desc = AnimGraphAttributes->FindAttributeDesc(InAttribute);
+								return Desc && Desc->Blend == EAnimGraphAttributeBlend::Blendable;
+							};
+
+							// Absorb all blendables at the root
+							for(const FName& Attribute : InAttributes)
+							{
+								if(IsAttributeBlendable(Attribute))
+								{
+									AbsorbedAttributes.Add(Attribute);
+								}
+								else
+								{
+									ReducedAttributes.Add(Attribute);
+								}
+							}
+						}
+						else
+						{
+							// Reduce the set of attributes that we are using in this traversal
+							for(const FName& Attribute : InAttributes)
+							{
+								if(HasAttribute(Attribute))
+								{
+									AbsorbedAttributes.Add(Attribute);
+								}
+								else
+								{
+									ReducedAttributes.Add(Attribute);
+								}
+							}
+						}
+
+						if(ReducedAttributes.Num() > 0)
+						{
+							UAnimGraphNode_Base::FNodeAttributeArray AbsorbedAttributesBelow = TraverseNodes_Recursive(InLinkedNode, ReducedAttributes);
+							AbsorbedAttributes.Append(AbsorbedAttributesBelow);
+						}
+					}
+					else
+					{
+						UAnimGraphNode_Base::FNodeAttributeArray AbsorbedAttributesBelow = TraverseNodes_Recursive(InLinkedNode, InAttributes);
+						AbsorbedAttributes.Append(AbsorbedAttributesBelow);
+					}
+				}
+				else
+				{
+					// already visited by another branch of a cached pose, see if any attributes got absorbed the last time we took this branch
+					TArrayView<const FName> PreviouslyAbsorbedAttributes = CompilationContext.GetAttributesFromNode(InLinkedNode);
+					AbsorbedAttributes.Append(PreviouslyAbsorbedAttributes.GetData(), PreviouslyAbsorbedAttributes.Num());
+				}
+			}
+
+			// Post-recursion, we add any pass through attributes that got absorbed
+			CompilationContext.AddAttributesToNode(InLinkedNode, AbsorbedAttributes);
+
+			return AbsorbedAttributes;
+		}
+
+		UAnimGraphNode_Base::FNodeAttributeArray TraverseNodes_Recursive(UEdGraphNode* InNode, const UAnimGraphNode_Base::FNodeAttributeArray& InAttributes)
+		{
+			VisitedNodes.Add(InNode);
+
+			for (UEdGraphPin* Pin : InNode->Pins)
+			{
+				if (UAnimationGraphSchema::IsPosePin(Pin->PinType) && Pin->Direction == EGPD_Output)
+				{
+					// Traverse pins
+					for (UEdGraphPin* LinkedPin : Pin->LinkedTo)
+					{
+						UEdGraphNode* OwningNode = LinkedPin->GetOwningNode();
+						if(UAnimGraphNode_Base* LinkedNode = Cast<UAnimGraphNode_Base>(OwningNode))
+						{
+							return TraverseNodes_Recursive_PerNode(LinkedNode, InAttributes);
+						}
+						else if(OwningNode != nullptr)
+						{
+							// Its a pose link, but not on an anim node, likely a knot
+							return TraverseNodes_Recursive(OwningNode, InAttributes);
+						}
+					}
+				}
+			}
+
+			// Traverse saved cached pose->all use cached pose nodes
+			if(UAnimGraphNode_SaveCachedPose* SaveCachedPoseNode = Cast<UAnimGraphNode_SaveCachedPose>(InNode))
+			{
+				TArray<UAnimGraphNode_UseCachedPose*> UseCachedPoseNodes;
+				SaveCachedPoseMap.MultiFind(SaveCachedPoseNode, UseCachedPoseNodes);
+				for(int32 UseCachedPoseIndex = 0; UseCachedPoseIndex < UseCachedPoseNodes.Num(); ++UseCachedPoseIndex)
+				{
+					UAnimGraphNode_UseCachedPose* UsedCachedPoseNode = UseCachedPoseNodes[UseCachedPoseIndex];
+					UAnimGraphNode_Base::FNodeAttributeArray AbsorbedAttributesBelow = TraverseNodes_Recursive_PerNode(UsedCachedPoseNode, InAttributes);
+					if(UseCachedPoseIndex == UseCachedPoseNodes.Num() - 1)
+					{
+						return AbsorbedAttributesBelow;
+					}
+				}
+			}
+			// Traverse out of custom transitions
+			else if(UAnimGraphNode_CustomTransitionResult* TransitionResultNode = Cast<UAnimGraphNode_CustomTransitionResult>(InNode))
+			{
+				UAnimationCustomTransitionGraph* TransitionGraph = CastChecked<UAnimationCustomTransitionGraph>(TransitionResultNode->GetGraph());
+				UAnimStateTransitionNode* TransitionNode = CastChecked<UAnimStateTransitionNode>(TransitionGraph->GetOuter());
+				UAnimationStateMachineGraph* StateMachineGraph = CastChecked<UAnimationStateMachineGraph>(TransitionNode->GetOuter());
+				UAnimGraphNode_StateMachineBase* StateMachineNode = CastChecked<UAnimGraphNode_StateMachineBase>(StateMachineGraph->GetOuter());
+				return TraverseNodes_Recursive_PerNode(StateMachineNode, InAttributes);
+			}
+			// Traverse out of state machines
+			else if(UAnimGraphNode_StateResult* StateResultNode = Cast<UAnimGraphNode_StateResult>(InNode))
+			{
+				UAnimationStateGraph* StateGraph = CastChecked<UAnimationStateGraph>(StateResultNode->GetGraph());
+				UAnimStateNode* StateNode = CastChecked<UAnimStateNode>(StateGraph->GetOuter());
+				UAnimationStateMachineGraph* StateMachineGraph = CastChecked<UAnimationStateMachineGraph>(StateNode->GetOuter());
+				UAnimGraphNode_StateMachineBase* StateMachineNode = CastChecked<UAnimGraphNode_StateMachineBase>(StateMachineGraph->GetOuter());
+				return TraverseNodes_Recursive_PerNode(StateMachineNode, InAttributes);
+			}
+
+			return UAnimGraphNode_Base::FNodeAttributeArray();
+		}
+	};
+
+	// Color connected nodes root-wise from output nodes
+	for(auto& NodeAttributesPair : AttributeOutputNodes)
+	{
+		FCheckOutputNodes Checker(InCompilationContext, NodeAttributesPair.Value, AttributeInputNodes, SaveCachedPoseMap, NodeAttributesPair.Key);
+	}
+}
diff --git a/Engine/Source/Editor/AnimGraph/Private/AnimBlueprintCompilerHandler_Attributes.h b/Engine/Source/Editor/AnimGraph/Private/AnimBlueprintCompilerHandler_Attributes.h
new file mode 100644
--- /dev/null
+++ b/Engine/Source/Editor/AnimGraph/Private/AnimBlueprintCompilerHandler_Attributes.h
@@ -0,0 +1,22 @@
+// Copyright Epic Games, Inc. All Rights Reserved.
+
+#pragma once
+
+#include "CoreMinimal.h"
+#include "AnimGraphNode_Base.h"
+#include "IAnimBlueprintCompilerHandler.h"
+
+class UAnimGraphNode_SaveCachedPose;
+class IAnimBlueprintCompilerCreationContext;
+class IAnimBlueprintCompilationContext;
+class IAnimBlueprintGeneratedClassCompiledData;
+
+// Handler to propogate attributes from outputs to inputs and to build a static debug record of their path through the graph
+class FAnimBlueprintCompilerHandler_Attributes : public IAnimBlueprintCompilerHandler
+{
+public:
+	FAnimBlueprintCompilerHandler_Attributes(IAnimBlueprintCompilerCreationContext& InCreationContext);
+
+private:
+	void PreProcessAnimationNodes(TArrayView<UAnimGraphNode_Base*> InAnimNodes, IAnimBlueprintCompilationContext& InCompilationContext, IAnimBlueprintGeneratedClassCompiledData& OutCompiledData);
+};
\ No newline at end of file
diff --git a/Engine/Source/Editor/AnimGraph/Private/AnimGraphAttributes.cpp b/Engine/Source/Editor/AnimGraph/Private/AnimGraphAttributes.cpp
new file mode 100644
--- /dev/null
+++ b/Engine/Source/Editor/AnimGraph/Private/AnimGraphAttributes.cpp
@@ -0,0 +1,48 @@
+// Copyright Epic Games, Inc. All Rights Reserved.
+
+#include "AnimGraphAttributes.h"
+
+void UAnimGraphAttributes::Register(const FAnimGraphAttributeDesc& InDesc)
+{
+	if(AttributesMap.Find(InDesc.Name) == nullptr)
+	{
+		FAnimGraphAttributeDesc& NewDesc = Attributes.Add_GetRef(InDesc);
+		if(NewDesc.SortOrder == INDEX_NONE)
+		{
+			NewDesc.SortOrder = Attributes.Num() - 1;
+		}
+		AttributesMap.Add(InDesc.Name, Attributes.Num() - 1);
+	}
+}
+
+const FAnimGraphAttributeDesc* UAnimGraphAttributes::FindAttributeDesc(FName InName) const
+{
+	if(const int32* IndexPtr = AttributesMap.Find(InName))
+	{
+		return &Attributes[*IndexPtr];
+	}
+	
+	return nullptr;
+}
+
+void UAnimGraphAttributes::ForEachAttribute(TFunctionRef<void(const FAnimGraphAttributeDesc&)> InFunction) const
+{
+	for(const FAnimGraphAttributeDesc& Desc : Attributes)
+	{
+		InFunction(Desc);
+	}
+}
+
+void UAnimGraphAttributes::PostEditChangeProperty(FPropertyChangedEvent& PropertyChangedEvent)
+{
+	Super::PostEditChangeProperty(PropertyChangedEvent);
+
+	// Rebuild map
+	AttributesMap.Reset();
+
+	for(int32 DescIndex = 0; DescIndex < Attributes.Num(); ++DescIndex)
+	{
+		const FAnimGraphAttributeDesc& Desc = Attributes[DescIndex];
+		AttributesMap.Add(Desc.Name, DescIndex);
+	}
+}
\ No newline at end of file
diff --git a/Engine/Source/Editor/AnimGraph/Private/AnimGraphModule.cpp b/Engine/Source/Editor/AnimGraph/Private/AnimGraphModule.cpp
--- a/Engine/Source/Editor/AnimGraph/Private/AnimGraphModule.cpp
+++ b/Engine/Source/Editor/AnimGraph/Private/AnimGraphModule.cpp
@@ -28,6 +28,11 @@
 #include "AnimBlueprintCompilerHandler_CachedPose.h"
 #include "AnimBlueprintCompilerHandler_LinkedAnimGraph.h"
 #include "AnimBlueprintCompilerHandler_StateMachine.h"
+#include "AnimBlueprintCompilerHandler_Attributes.h"
+#include "Animation/AnimAttributes.h"
+#include "AnimGraphAttributes.h"
+#include "Animation/AnimSync.h"
+#include "Animation/AnimNode_Inertialization.h"
 
 IMPLEMENT_MODULE(FAnimGraphModule, AnimGraph);
 
@@ -63,6 +68,11 @@ void FAnimGraphModule::StartupModule()
 		return MakeUnique<FAnimBlueprintCompilerHandler_StateMachine>(InCreationContext);
 	});
 
+	IAnimBlueprintCompilerHandlerCollection::RegisterHandler("AnimBlueprintCompilerHandler_Attributes", [](IAnimBlueprintCompilerCreationContext& InCreationContext)
+	{
+		return MakeUnique<FAnimBlueprintCompilerHandler_Attributes>(InCreationContext);
+	});
+
 	// Register the editor modes
 	FEditorModeRegistry::Get().RegisterMode<FAnimNodeEditMode>(AnimNodeEditModes::AnimNode, LOCTEXT("AnimNodeEditMode", "Anim Node"), FSlateIcon(), false);
 	FEditorModeRegistry::Get().RegisterMode<FTwoBoneIKEditMode>(AnimNodeEditModes::TwoBoneIK, LOCTEXT("TwoBoneIKEditMode", "2-Bone IK"), FSlateIcon(), false);
@@ -80,6 +90,65 @@ void FAnimGraphModule::StartupModule()
 
 	PropertyModule.RegisterCustomPropertyTypeLayout("AnimBlueprintFunctionPinInfo", FOnGetPropertyTypeCustomizationInstance::CreateStatic(&FAnimBlueprintFunctionPinInfoDetails::MakeInstance));
 
+	// Register built-in graph attributes
+	UAnimGraphAttributes* AnimGraphAttributes = GetMutableDefault<UAnimGraphAttributes>();
+	AnimGraphAttributes->LoadConfig();
+
+	AnimGraphAttributes->Register(
+		FAnimGraphAttributeDesc(
+			UE::Anim::FAttributes::Pose,
+			EAnimGraphAttributeBlend::Blendable,
+			FSlateBrush(),
+			LOCTEXT("Pose", "Pose"),
+			LOCTEXT("PoseToolTip", "Pose: Transforms for all bones in the current LOD, blendable"),
+			FEditorStyle::GetSlateColor("AnimGraph.Attribute.Pose.Color"),
+			EAnimGraphAttributesDisplayMode::HideOnPins)
+	);
+
+	AnimGraphAttributes->Register(
+		FAnimGraphAttributeDesc(
+			UE::Anim::FAttributes::Curves,
+			EAnimGraphAttributeBlend::Blendable,
+			*FEditorStyle::GetBrush("AnimGraph.Attribute.Curves.Icon"),
+			LOCTEXT("Curves", "Curves"),
+			LOCTEXT("CurvesToolTip", "Curves: Floating point curve attributes, blendable"),
+			FEditorStyle::GetSlateColor("AnimGraph.Attribute.Curves.Color"),
+			EAnimGraphAttributesDisplayMode::Automatic)
+	);
+
+	AnimGraphAttributes->Register(
+		FAnimGraphAttributeDesc(
+			UE::Anim::FAttributes::Attributes,
+			EAnimGraphAttributeBlend::Blendable,
+			*FEditorStyle::GetBrush("AnimGraph.Attribute.Attributes.Icon"),
+			LOCTEXT("Attributes", "Attributes"),
+			LOCTEXT("AttributesToolTip", "Attributes: Misc. custom attributes, blendable"),
+			FEditorStyle::GetSlateColor("AnimGraph.Attribute.Attributes.Color"),
+			EAnimGraphAttributesDisplayMode::Automatic)
+	);
+
+	AnimGraphAttributes->Register(
+		FAnimGraphAttributeDesc(
+			UE::Anim::FAnimSync::Attribute,
+			EAnimGraphAttributeBlend::NonBlendable,
+			*FEditorStyle::Get().GetBrush("AnimGraph.Attribute.Sync.Icon"),
+			LOCTEXT("Sync", "Sync"),
+			LOCTEXT("SyncToolTip", "Sync: Synchronization between the playback of different assets present in the graph, not blendable"),
+			FEditorStyle::Get().GetSlateColor("AnimGraph.Attribute.Sync.Color"),
+			EAnimGraphAttributesDisplayMode::Automatic)
+	);
+
+	AnimGraphAttributes->Register(
+		FAnimGraphAttributeDesc(
+			UE::Anim::IInertializationRequester::Attribute,
+			EAnimGraphAttributeBlend::NonBlendable,
+			*FEditorStyle::Get().GetBrush("AnimGraph.Attribute.InertialBlending.Icon"),
+			LOCTEXT("InertialBlending", "Inertialize"),
+			LOCTEXT("InertialBlendingToolTip", "Inertial Blending: Used to blend between poses preserving bone velocities, not blendable"),
+			FEditorStyle::Get().GetSlateColor("AnimGraph.Attribute.InertialBlending.Color"),
+			EAnimGraphAttributesDisplayMode::Automatic)
+	);
+
 	// Register BP-editor function customization once the Kismet module is loaded.
 	if (FModuleManager::Get().IsModuleLoaded("Kismet"))
 	{
@@ -105,6 +174,7 @@ void FAnimGraphModule::ShutdownModule()
 	IAnimBlueprintCompilerHandlerCollection::UnregisterHandler("AnimBlueprintCompilerHandler_CachedPose");
 	IAnimBlueprintCompilerHandlerCollection::UnregisterHandler("AnimBlueprintCompilerHandler_LinkedAnimGraph");
 	IAnimBlueprintCompilerHandlerCollection::UnregisterHandler("AnimBlueprintCompilerHandler_StateMachine");
+	IAnimBlueprintCompilerHandlerCollection::UnregisterHandler("AnimBlueprintCompilerHandler_Attributes");
 
 	// Unregister the editor modes
 	FEditorModeRegistry::Get().UnregisterMode(AnimNodeEditModes::CCDIK);
diff --git a/Engine/Source/Editor/AnimGraph/Private/AnimGraphNode_AssetPlayerBase.cpp b/Engine/Source/Editor/AnimGraph/Private/AnimGraphNode_AssetPlayerBase.cpp
--- a/Engine/Source/Editor/AnimGraph/Private/AnimGraphNode_AssetPlayerBase.cpp
+++ b/Engine/Source/Editor/AnimGraph/Private/AnimGraphNode_AssetPlayerBase.cpp
@@ -21,6 +21,40 @@
 #include "AnimBlueprintCompilerHandler_Base.h"
 #include "IAnimBlueprintGeneratedClassCompiledData.h"
 #include "IAnimBlueprintCompilationContext.h"
+#include "Animation/AnimSync.h"
+#include "Animation/AnimAttributes.h"
+#include "UObject/UE5MainStreamObjectVersion.h"
+
+#define LOCTEXT_NAMESPACE "UAnimGraphNode_AssetPlayerBase"
+
+void UAnimGraphNode_AssetPlayerBase::Serialize(FArchive& Ar)
+{
+	Super::Serialize(Ar);
+
+	Ar.UsingCustomVersion(FUE5MainStreamObjectVersion::GUID);
+
+	if(Ar.CustomVer(FUE5MainStreamObjectVersion::GUID) < FUE5MainStreamObjectVersion::AnimSyncGroupsExplicitSyncMethod)
+	{
+		if(SyncGroup.GroupName != NAME_None)
+		{
+			SyncGroup.Method = EAnimSyncMethod::SyncGroup;
+		}
+	}
+}
+
+void UAnimGraphNode_AssetPlayerBase::PostEditChangeProperty(FPropertyChangedEvent& PropertyChangedEvent)
+{
+	Super::PostEditChangeProperty(PropertyChangedEvent);
+
+	if(PropertyChangedEvent.GetPropertyName() == GET_MEMBER_NAME_CHECKED(FAnimationGroupReference, Method))
+	{
+		if(SyncGroup.Method != EAnimSyncMethod::SyncGroup)
+		{
+			SyncGroup.GroupName = NAME_None;
+			SyncGroup.GroupRole = EAnimGroupRole::CanBeLeader;
+		}
+	}
+}
 
 void UAnimGraphNode_AssetPlayerBase::PinConnectionListChanged(UEdGraphPin* Pin)
 {
@@ -97,6 +131,26 @@ void UAnimGraphNode_AssetPlayerBase::OnProcessDuringCompilation(IAnimBlueprintCo
 	}
 }
 
+void UAnimGraphNode_AssetPlayerBase::ValidateAnimNodeDuringCompilation(USkeleton* ForSkeleton, FCompilerResultsLog& MessageLog)
+{
+	Super::ValidateAnimNodeDuringCompilation(ForSkeleton, MessageLog);
+
+	if(SyncGroup.Method == EAnimSyncMethod::SyncGroup && SyncGroup.GroupName == NAME_None)
+	{
+		MessageLog.Error(*LOCTEXT("NoSyncGroupSupplied", "Node @@ is set to use sync groups, but no sync group has been supplied").ToString(), this);
+	}
+}
+
+void UAnimGraphNode_AssetPlayerBase::GetOutputLinkAttributes(FNodeAttributeArray& OutAttributes) const
+{
+	OutAttributes.Add(UE::Anim::FAttributes::Curves);
+	OutAttributes.Add(UE::Anim::FAttributes::Attributes);
+	if(SyncGroup.Method == EAnimSyncMethod::Graph)
+	{
+		OutAttributes.Add(UE::Anim::FAnimSync::Attribute);
+	}
+}
+
 UClass* GetNodeClassForAsset(const UClass* AssetClass)
 {
 	UClass* NodeClass = nullptr;
@@ -128,3 +182,5 @@ bool SupportNodeClassForAsset(const UClass* AssetClass, UClass* NodeClass)
 	// See if this node supports this asset type (primary or not)
 	return (NodeCDO->SupportsAssetClass(AssetClass) != EAnimAssetHandlerType::NotSupported);
 }
+
+#undef LOCTEXT_NAMESPACE
\ No newline at end of file
diff --git a/Engine/Source/Editor/AnimGraph/Private/AnimGraphNode_BlendListBase.cpp b/Engine/Source/Editor/AnimGraph/Private/AnimGraphNode_BlendListBase.cpp
--- a/Engine/Source/Editor/AnimGraph/Private/AnimGraphNode_BlendListBase.cpp
+++ b/Engine/Source/Editor/AnimGraph/Private/AnimGraphNode_BlendListBase.cpp
@@ -2,7 +2,7 @@
 
 #include "AnimGraphNode_BlendListBase.h"
 #include "AnimNodes/AnimNode_BlendListBase.h"
-#include "UObject/UE5MainStreamObjectVersion.h"
+#include "Animation/AnimNode_Inertialization.h"
 
 /////////////////////////////////////////////////////
 // UAnimGraphNode_BlendListBase
@@ -106,3 +106,16 @@ void UAnimGraphNode_BlendListBase::ReallocatePinsDuringReconstruction(TArray<UEd
 		RemovedPinArrayIndex = INDEX_NONE;
 	}
 }
+
+void UAnimGraphNode_BlendListBase::GetOutputLinkAttributes(FNodeAttributeArray& OutAttributes) const
+{
+	if(GetFNodeProperty()->Struct->IsChildOf(FAnimNode_BlendListBase::StaticStruct()))
+	{
+		const FAnimNode_BlendListBase* BaseNode = GetFNodeProperty()->ContainerPtrToValuePtr<FAnimNode_BlendListBase>(this);
+
+		if(BaseNode->TransitionType == EBlendListTransitionType::Inertialization)
+		{
+			OutAttributes.Add(UE::Anim::IInertializationRequester::Attribute);
+		}
+	}
+}
\ No newline at end of file
diff --git a/Engine/Source/Editor/AnimGraph/Private/AnimGraphNode_BlendSpaceEvaluator.cpp b/Engine/Source/Editor/AnimGraph/Private/AnimGraphNode_BlendSpaceEvaluator.cpp
--- a/Engine/Source/Editor/AnimGraph/Private/AnimGraphNode_BlendSpaceEvaluator.cpp
+++ b/Engine/Source/Editor/AnimGraph/Private/AnimGraphNode_BlendSpaceEvaluator.cpp
@@ -139,7 +139,7 @@ void UAnimGraphNode_BlendSpaceEvaluator::BakeDataDuringCompilation(class FCompil
 	AnimBlueprint->FindOrAddGroup(SyncGroup.GroupName);
 	Node.GroupName = SyncGroup.GroupName;
 	Node.GroupRole = SyncGroup.GroupRole;
-	Node.GroupScope = SyncGroup.GroupScope;
+	Node.Method = SyncGroup.Method;
 }
 
 void UAnimGraphNode_BlendSpaceEvaluator::GetNodeContextMenuActions(UToolMenu* Menu, UGraphNodeContextMenuContext* Context) const
diff --git a/Engine/Source/Editor/AnimGraph/Private/AnimGraphNode_BlendSpacePlayer.cpp b/Engine/Source/Editor/AnimGraph/Private/AnimGraphNode_BlendSpacePlayer.cpp
--- a/Engine/Source/Editor/AnimGraph/Private/AnimGraphNode_BlendSpacePlayer.cpp
+++ b/Engine/Source/Editor/AnimGraph/Private/AnimGraphNode_BlendSpacePlayer.cpp
@@ -44,12 +44,35 @@ FText UAnimGraphNode_BlendSpacePlayer::GetNodeTitleForBlendSpace(ENodeTitleType:
 		TitleArgs.Add(TEXT("BlendSpaceName"), BlendSpaceName);
 		FText Title = FText::Format(LOCTEXT("BlendSpacePlayerFullTitle", "{BlendSpaceName}\nBlendspace Player"), TitleArgs);
 
-		if ((TitleType == ENodeTitleType::FullTitle) && (SyncGroup.GroupName != NAME_None))
+		if (TitleType == ENodeTitleType::FullTitle)
 		{
 			FFormatNamedArguments Args;
 			Args.Add(TEXT("Title"), Title);
-			Args.Add(TEXT("SyncGroupName"), FText::FromName(SyncGroup.GroupName));
-			Title = FText::Format(LOCTEXT("BlendSpaceNodeGroupSubtitle", "{Title}\nSync group {SyncGroupName}"), Args);
+
+			if(SyncGroup.Method == EAnimSyncMethod::SyncGroup)
+			{
+				Args.Add(TEXT("SyncGroupName"), FText::FromName(SyncGroup.GroupName));
+				Title = FText::Format(LOCTEXT("BlendSpaceNodeGroupSubtitle", "{Title}\nSync group {SyncGroupName}"), Args);
+			}
+			else if(SyncGroup.Method == EAnimSyncMethod::Graph)
+			{
+				Title = FText::Format(LOCTEXT("BlendSpaceNodeGroupSubtitle", "{Title}\nGraph sync group"), Args);
+
+				UObject* ObjectBeingDebugged = GetAnimBlueprint()->GetObjectBeingDebugged();
+				UAnimBlueprintGeneratedClass* GeneratedClass = GetAnimBlueprint()->GetAnimBlueprintGeneratedClass();
+				if (ObjectBeingDebugged && GeneratedClass)
+				{
+					int32 NodeIndex = GeneratedClass->GetNodeIndexFromGuid(NodeGuid);
+					if(NodeIndex != INDEX_NONE)
+					{
+						if(const FName* SyncGroupNamePtr = GeneratedClass->GetAnimBlueprintDebugData().NodeSyncsThisFrame.Find(NodeIndex))
+						{
+							Args.Add(TEXT("SyncGroupName"), FText::FromName(*SyncGroupNamePtr));
+							Title = FText::Format(LOCTEXT("BlendSpaceNodeGraphGroupSubtitle", "{Title}\nGraph sync group {SyncGroupName}"), Args);
+						}
+					}
+				}
+			}
 		}
 		// FText::Format() is slow, so we cache this to save on performance
 		CachedNodeTitles.SetCachedTitle(TitleType, Title, this);
@@ -138,7 +161,7 @@ void UAnimGraphNode_BlendSpacePlayer::BakeDataDuringCompilation(class FCompilerR
 	AnimBlueprint->FindOrAddGroup(SyncGroup.GroupName);
 	Node.GroupName = SyncGroup.GroupName;
 	Node.GroupRole = SyncGroup.GroupRole;
-	Node.GroupScope = SyncGroup.GroupScope;
+	Node.Method = SyncGroup.Method;
 }
 
 void UAnimGraphNode_BlendSpacePlayer::GetNodeContextMenuActions(UToolMenu* Menu, UGraphNodeContextMenuContext* Context) const
diff --git a/Engine/Source/Editor/AnimGraph/Private/AnimGraphNode_CopyPoseFromMesh.cpp b/Engine/Source/Editor/AnimGraph/Private/AnimGraphNode_CopyPoseFromMesh.cpp
--- a/Engine/Source/Editor/AnimGraph/Private/AnimGraphNode_CopyPoseFromMesh.cpp
+++ b/Engine/Source/Editor/AnimGraph/Private/AnimGraphNode_CopyPoseFromMesh.cpp
@@ -1,6 +1,7 @@
 // Copyright Epic Games, Inc. All Rights Reserved.
 
 #include "AnimGraphNode_CopyPoseFromMesh.h"
+#include "Animation/AnimAttributes.h"
 
 /////////////////////////////////////////////////////
 // UAnimGraphNode_CopyPoseFromMeshSkeletalControl
@@ -22,4 +23,10 @@ FText UAnimGraphNode_CopyPoseFromMesh::GetNodeTitle(ENodeTitleType::Type TitleTy
 	return LOCTEXT("CopyPoseFromMesh", "Copy Pose From Mesh");
 }
 
+void UAnimGraphNode_CopyPoseFromMesh::GetOutputLinkAttributes(FNodeAttributeArray& OutAttributes) const
+{
+	OutAttributes.Add(UE::Anim::FAttributes::Curves);
+	OutAttributes.Add(UE::Anim::FAttributes::Attributes);
+}
+
 #undef LOCTEXT_NAMESPACE
diff --git a/Engine/Source/Editor/AnimGraph/Private/AnimGraphNode_CurveSource.cpp b/Engine/Source/Editor/AnimGraph/Private/AnimGraphNode_CurveSource.cpp
--- a/Engine/Source/Editor/AnimGraph/Private/AnimGraphNode_CurveSource.cpp
+++ b/Engine/Source/Editor/AnimGraph/Private/AnimGraphNode_CurveSource.cpp
@@ -1,6 +1,7 @@
 // Copyright Epic Games, Inc. All Rights Reserved.
 
 #include "AnimGraphNode_CurveSource.h"
+#include "Animation/AnimAttributes.h"
 
 #define LOCTEXT_NAMESPACE "ExternalCurve"
 
@@ -28,4 +29,9 @@ FText UAnimGraphNode_CurveSource::GetNodeTitle(ENodeTitleType::Type TitleType) c
 	return LOCTEXT("AnimGraphNode_CurveSource_Title", "Curve Source");
 }
 
+void UAnimGraphNode_CurveSource::GetOutputLinkAttributes(FNodeAttributeArray& OutAttributes) const
+{
+	OutAttributes.Add(UE::Anim::FAttributes::Curves);
+}
+
 #undef LOCTEXT_NAMESPACE
diff --git a/Engine/Source/Editor/AnimGraph/Private/AnimGraphNode_Inertialization.cpp b/Engine/Source/Editor/AnimGraph/Private/AnimGraphNode_Inertialization.cpp
--- a/Engine/Source/Editor/AnimGraph/Private/AnimGraphNode_Inertialization.cpp
+++ b/Engine/Source/Editor/AnimGraph/Private/AnimGraphNode_Inertialization.cpp
@@ -25,5 +25,9 @@ FText UAnimGraphNode_Inertialization::GetMenuCategory() const
 	return LOCTEXT("NodeCategory", "Inertialization");
 }
 
+void UAnimGraphNode_Inertialization::GetInputLinkAttributes(FNodeAttributeArray& OutAttributes) const
+{
+	OutAttributes.Add(UE::Anim::IInertializationRequester::Attribute);
+}
 
 #undef LOCTEXT_NAMESPACE
diff --git a/Engine/Source/Editor/AnimGraph/Private/AnimGraphNode_LinkedAnimGraphBase.cpp b/Engine/Source/Editor/AnimGraph/Private/AnimGraphNode_LinkedAnimGraphBase.cpp
--- a/Engine/Source/Editor/AnimGraph/Private/AnimGraphNode_LinkedAnimGraphBase.cpp
+++ b/Engine/Source/Editor/AnimGraph/Private/AnimGraphNode_LinkedAnimGraphBase.cpp
@@ -16,6 +16,7 @@
 #include "Widgets/Layout/SBox.h"
 #include "UObject/CoreRedirects.h"
 #include "Animation/AnimNode_LinkedAnimGraph.h"
+#include "AnimGraphAttributes.h"
 
 #define LOCTEXT_NAMESPACE "LinkedAnimGraph"
 
@@ -470,4 +471,14 @@ FPoseLinkMappingRecord UAnimGraphNode_LinkedAnimGraphBase::GetLinkIDLocation(con
 	return FPoseLinkMappingRecord::MakeInvalid();
 }
 
+void UAnimGraphNode_LinkedAnimGraphBase::GetOutputLinkAttributes(FNodeAttributeArray& OutAttributes) const
+{
+	// We have the potential to output ALL registered attributes as we can contain any dynamically-linked graph
+	const UAnimGraphAttributes* AnimGraphAttributes = GetDefault<UAnimGraphAttributes>();
+	AnimGraphAttributes->ForEachAttribute([&OutAttributes](const FAnimGraphAttributeDesc& InDesc)
+	{
+		OutAttributes.Add(InDesc.Name);
+	});
+}
+
 #undef LOCTEXT_NAMESPACE
diff --git a/Engine/Source/Editor/AnimGraph/Private/AnimGraphNode_LinkedInputPose.cpp b/Engine/Source/Editor/AnimGraph/Private/AnimGraphNode_LinkedInputPose.cpp
--- a/Engine/Source/Editor/AnimGraph/Private/AnimGraphNode_LinkedInputPose.cpp
+++ b/Engine/Source/Editor/AnimGraph/Private/AnimGraphNode_LinkedInputPose.cpp
@@ -27,6 +27,7 @@
 #include "KismetCompiler.h"
 #include "K2Node_VariableGet.h"
 #include "AnimBlueprintCompiler.h"
+#include "AnimGraphAttributes.h"
 
 #define LOCTEXT_NAMESPACE "LinkedInputPose"
 
@@ -632,4 +633,14 @@ bool UAnimGraphNode_LinkedInputPose::IsCompatibleWithGraph(UEdGraph const* Graph
 	return Graph->GetFName() == UEdGraphSchema_K2::GN_AnimGraph;
 }
 
+void UAnimGraphNode_LinkedInputPose::GetOutputLinkAttributes(FNodeAttributeArray& OutAttributes) const
+{
+	// We have the potential to output ALL registered attributes
+	const UAnimGraphAttributes* AnimGraphAttributes = GetDefault<UAnimGraphAttributes>();
+	AnimGraphAttributes->ForEachAttribute([&OutAttributes](const FAnimGraphAttributeDesc& InDesc)
+	{
+		OutAttributes.Add(InDesc.Name);
+	});
+}
+
 #undef LOCTEXT_NAMESPACE
diff --git a/Engine/Source/Editor/AnimGraph/Private/AnimGraphNode_RandomPlayer.cpp b/Engine/Source/Editor/AnimGraph/Private/AnimGraphNode_RandomPlayer.cpp
--- a/Engine/Source/Editor/AnimGraph/Private/AnimGraphNode_RandomPlayer.cpp
+++ b/Engine/Source/Editor/AnimGraph/Private/AnimGraphNode_RandomPlayer.cpp
@@ -3,6 +3,7 @@
 #include "AnimGraphNode_RandomPlayer.h"
 
 #include "EditorCategoryUtils.h"
+#include "Animation/AnimAttributes.h"
 
 #define LOCTEXT_NAMESPACE "AnimGraphNode_RandomPlayer"
 
@@ -26,4 +27,10 @@ FText UAnimGraphNode_RandomPlayer::GetMenuCategory() const
 	return FEditorCategoryUtils::GetCommonCategory(FCommonEditorCategory::Animation);
 }
 
+void UAnimGraphNode_RandomPlayer::GetOutputLinkAttributes(FNodeAttributeArray& OutAttributes) const
+{
+	OutAttributes.Add(UE::Anim::FAttributes::Curves);
+	OutAttributes.Add(UE::Anim::FAttributes::Attributes);
+}
+
 #undef LOCTEXT_NAMESPACE
diff --git a/Engine/Source/Editor/AnimGraph/Private/AnimGraphNode_SequenceEvaluator.cpp b/Engine/Source/Editor/AnimGraph/Private/AnimGraphNode_SequenceEvaluator.cpp
--- a/Engine/Source/Editor/AnimGraph/Private/AnimGraphNode_SequenceEvaluator.cpp
+++ b/Engine/Source/Editor/AnimGraph/Private/AnimGraphNode_SequenceEvaluator.cpp
@@ -32,7 +32,7 @@ void UAnimGraphNode_SequenceEvaluator::BakeDataDuringCompilation(class FCompiler
 	AnimBlueprint->FindOrAddGroup(SyncGroup.GroupName);
 	Node.GroupName = SyncGroup.GroupName;
 	Node.GroupRole = SyncGroup.GroupRole;
-	Node.GroupScope = SyncGroup.GroupScope;
+	Node.Method = SyncGroup.Method;
 }
 
 void UAnimGraphNode_SequenceEvaluator::GetAllAnimationSequencesReferred(TArray<UAnimationAsset*>& AnimationAssets) const
diff --git a/Engine/Source/Editor/AnimGraph/Private/AnimGraphNode_SequencePlayer.cpp b/Engine/Source/Editor/AnimGraph/Private/AnimGraphNode_SequencePlayer.cpp
--- a/Engine/Source/Editor/AnimGraph/Private/AnimGraphNode_SequencePlayer.cpp
+++ b/Engine/Source/Editor/AnimGraph/Private/AnimGraphNode_SequencePlayer.cpp
@@ -90,11 +90,7 @@ FText UAnimGraphNode_SequencePlayer::GetNodeTitleForSequence(ENodeTitleType::Typ
 	const bool bAdditive = InSequence->IsValidAdditive();
 	const FText BasicTitle = GetTitleGivenAssetInfo(FText::FromName(InSequence->GetFName()), bAdditive);
 
-	if (SyncGroup.GroupName == NAME_None)
-	{
-		return BasicTitle;
-	}
-	else
+	if(SyncGroup.Method == EAnimSyncMethod::SyncGroup)
 	{
 		const FText SyncGroupName = FText::FromName(SyncGroup.GroupName);
 
@@ -111,6 +107,48 @@ FText UAnimGraphNode_SequencePlayer::GetNodeTitleForSequence(ENodeTitleType::Typ
 			return FText::Format(LOCTEXT("SequenceNodeGroupWithSubtitleList", "{Title} (Sync group {SyncGroup})"), Args);
 		}
 	}
+	else if(SyncGroup.Method == EAnimSyncMethod::Graph)
+	{
+		FFormatNamedArguments Args;
+		Args.Add(TEXT("Title"), BasicTitle);
+
+		UObject* ObjectBeingDebugged = GetAnimBlueprint()->GetObjectBeingDebugged();
+		UAnimBlueprintGeneratedClass* GeneratedClass = GetAnimBlueprint()->GetAnimBlueprintGeneratedClass();
+		if (ObjectBeingDebugged && GeneratedClass)
+		{
+			int32 NodeIndex = GeneratedClass->GetNodeIndexFromGuid(NodeGuid);
+			if(NodeIndex != INDEX_NONE)
+			{
+				if(const FName* SyncGroupNamePtr = GeneratedClass->GetAnimBlueprintDebugData().NodeSyncsThisFrame.Find(NodeIndex))
+				{
+					Args.Add(TEXT("SyncGroup"), FText::FromName(*SyncGroupNamePtr));
+
+					if (TitleType == ENodeTitleType::FullTitle)
+					{
+						return FText::Format(LOCTEXT("SequenceNodeGraphSyncDebuggedWithSubtitleFull", "{Title}\nGraph sync group {SyncGroup}"), Args);
+					}
+					else
+					{
+						return FText::Format(LOCTEXT("SequenceNodeGraphSyncDebuggedWithSubtitleList", "{Title} (Graph sync group {SyncGroup})"), Args);
+					}
+				}
+			}
+		}
+
+		// Indicate that the sync is graph-based if we have no debug data
+		if (TitleType == ENodeTitleType::FullTitle)
+		{
+			return FText::Format(LOCTEXT("SequenceNodeGraphSyncWithSubtitleFull", "{Title}\nGraph sync group"), Args);
+		}
+		else
+		{
+			return FText::Format(LOCTEXT("SequenceNodeGraphSyncWithSubtitleList", "{Title} (Graph sync group)"), Args);
+		}
+	}
+	else
+	{
+		return BasicTitle;
+	}
 }
 
 FText UAnimGraphNode_SequencePlayer::GetNodeTitle(ENodeTitleType::Type TitleType) const
@@ -384,7 +422,7 @@ void UAnimGraphNode_SequencePlayer::BakeDataDuringCompilation(class FCompilerRes
 	AnimBlueprint->FindOrAddGroup(SyncGroup.GroupName);
 	Node.GroupName = SyncGroup.GroupName;
 	Node.GroupRole = SyncGroup.GroupRole;
-	Node.GroupScope = SyncGroup.GroupScope;
+	Node.Method = SyncGroup.Method;
 }
 
 void UAnimGraphNode_SequencePlayer::GetAllAnimationSequencesReferred(TArray<UAnimationAsset*>& AnimationAssets) const
diff --git a/Engine/Source/Editor/AnimGraph/Private/AnimGraphNode_StateMachineBase.cpp b/Engine/Source/Editor/AnimGraph/Private/AnimGraphNode_StateMachineBase.cpp
--- a/Engine/Source/Editor/AnimGraph/Private/AnimGraphNode_StateMachineBase.cpp
+++ b/Engine/Source/Editor/AnimGraph/Private/AnimGraphNode_StateMachineBase.cpp
@@ -28,6 +28,7 @@
 #include "AnimBlueprintCompilerHandler_StateMachine.h"
 #include "IAnimBlueprintGeneratedClassCompiledData.h"
 #include "IAnimBlueprintCompilationContext.h"
+#include "Animation/AnimNode_Inertialization.h"
 
 /////////////////////////////////////////////////////
 // FAnimStateMachineNodeNameValidator
@@ -594,4 +595,19 @@ void UAnimGraphNode_StateMachineBase::OnProcessDuringCompilation(IAnimBlueprintC
 	Oven.Validate();
 }
 
+void UAnimGraphNode_StateMachineBase::GetOutputLinkAttributes(FNodeAttributeArray& OutAttributes) const
+{
+	for (const UEdGraphNode* Node : EditorStateMachineGraph->Nodes)
+	{
+		if (const UAnimStateTransitionNode* TransitionNode = Cast<UAnimStateTransitionNode>(Node))
+		{
+			if(TransitionNode->LogicType == ETransitionLogicType::TLT_Inertialization)
+			{
+				OutAttributes.Add(UE::Anim::IInertializationRequester::Attribute);
+				break;
+			}
+		}
+	}
+}
+
 #undef LOCTEXT_NAMESPACE
diff --git a/Engine/Source/Editor/AnimGraph/Private/AnimGraphNode_Sync.cpp b/Engine/Source/Editor/AnimGraph/Private/AnimGraphNode_Sync.cpp
new file mode 100644
--- /dev/null
+++ b/Engine/Source/Editor/AnimGraph/Private/AnimGraphNode_Sync.cpp
@@ -0,0 +1,62 @@
+// Copyright Epic Games, Inc. All Rights Reserved.
+
+#include "AnimGraphNode_Sync.h"
+#include "Kismet2/CompilerResultsLog.h"
+
+#define LOCTEXT_NAMESPACE "AnimGraphNode_Sync"
+
+void UAnimGraphNode_Sync::PostEditChangeProperty(FPropertyChangedEvent& PropertyChangedEvent)
+{
+	Super::PostEditChangeProperty(PropertyChangedEvent);
+
+	if(PropertyChangedEvent.GetPropertyName() == GET_MEMBER_NAME_CHECKED(FAnimNode_Sync, GroupName))
+	{
+		OnNodeTitleChangedEvent().Broadcast();
+	}
+}
+
+FText UAnimGraphNode_Sync::GetNodeTitle(ENodeTitleType::Type TitleType) const
+{
+	if(TitleType == ENodeTitleType::FullTitle)
+	{
+		return FText::Format(LOCTEXT("NodeFullTitle", "Sync\nSync group {0}"), FText::FromName(Node.GetGroupName()));
+	}
+	else
+	{
+		return LOCTEXT("NodeTitle", "Sync");
+	}
+}
+
+FText UAnimGraphNode_Sync::GetTooltipText() const
+{
+	return LOCTEXT("NodeTooltip", "Synchronizes asset players and blendspaces that are children of this node.");
+}
+
+void UAnimGraphNode_Sync::GetInputLinkAttributes(FNodeAttributeArray& OutAttributes) const
+{
+	if(Node.GetGroupName() != NAME_None)
+	{
+		OutAttributes.Add(UE::Anim::FAnimSync::Attribute);
+	}
+}
+
+void UAnimGraphNode_Sync::ValidateAnimNodeDuringCompilation(USkeleton* ForSkeleton, FCompilerResultsLog& MessageLog)
+{
+	Super::ValidateAnimNodeDuringCompilation(ForSkeleton, MessageLog);
+
+	if(Node.GetGroupName() == NAME_None)
+	{
+		MessageLog.Error(*LOCTEXT("NoSyncGroupSupplied", "Node @@ does not specifiy a sync group").ToString(), this);
+	}
+}
+
+void UAnimGraphNode_Sync::BakeDataDuringCompilation(FCompilerResultsLog& MessageLog)
+{
+	if(Node.GetGroupName() != NAME_None)
+	{
+		UAnimBlueprint* AnimBlueprint = GetAnimBlueprint();
+		AnimBlueprint->FindOrAddGroup(Node.GetGroupName());
+	}
+}
+
+#undef LOCTEXT_NAMESPACE
\ No newline at end of file
diff --git a/Engine/Source/Editor/AnimGraph/Private/AnimGraphNode_Sync.h b/Engine/Source/Editor/AnimGraph/Private/AnimGraphNode_Sync.h
new file mode 100644
--- /dev/null
+++ b/Engine/Source/Editor/AnimGraph/Private/AnimGraphNode_Sync.h
@@ -0,0 +1,27 @@
+// Copyright Epic Games, Inc. All Rights Reserved.
+
+#pragma once
+
+#include "CoreMinimal.h"
+#include "AnimNodes/AnimNode_Sync.h"
+#include "AnimGraphNode_Base.h"
+#include "AnimGraphNode_Sync.generated.h"
+
+UCLASS(MinimalAPI)
+class UAnimGraphNode_Sync : public UAnimGraphNode_Base
+{
+	GENERATED_BODY()
+
+	UPROPERTY(EditAnywhere, Category = Settings)
+	FAnimNode_Sync Node;
+
+	// UObject interface
+	virtual void PostEditChangeProperty(FPropertyChangedEvent& PropertyChangedEvent) override;
+
+	// UAnimGraphNode_Base interface
+	virtual FText GetNodeTitle(ENodeTitleType::Type TitleType) const override;
+	virtual FText GetTooltipText() const override;
+	virtual void ValidateAnimNodeDuringCompilation(USkeleton* ForSkeleton, FCompilerResultsLog& MessageLog) override;
+	virtual void BakeDataDuringCompilation(class FCompilerResultsLog& MessageLog) override;
+	virtual void GetInputLinkAttributes(FNodeAttributeArray& OutAttributes) const override;
+};
\ No newline at end of file
diff --git a/Engine/Source/Editor/AnimGraph/Private/AnimGraphNode_TransitionPoseEvaluator.cpp b/Engine/Source/Editor/AnimGraph/Private/AnimGraphNode_TransitionPoseEvaluator.cpp
--- a/Engine/Source/Editor/AnimGraph/Private/AnimGraphNode_TransitionPoseEvaluator.cpp
+++ b/Engine/Source/Editor/AnimGraph/Private/AnimGraphNode_TransitionPoseEvaluator.cpp
@@ -8,6 +8,7 @@
 #include "IDetailPropertyRow.h"
 #include "DetailCategoryBuilder.h"
 #include "DetailLayoutBuilder.h"
+#include "Animation/AnimAttributes.h"
 
 /////////////////////////////////////////////////////
 // UAnimGraphNode_TransitionPoseEvaluator
@@ -99,4 +100,10 @@ EVisibility UAnimGraphNode_TransitionPoseEvaluator::GetCacheFramesVisibility(IDe
 	return EVisibility::Hidden;
 }
 
+void UAnimGraphNode_TransitionPoseEvaluator::GetOutputLinkAttributes(FNodeAttributeArray& OutAttributes) const
+{
+	OutAttributes.Add(UE::Anim::FAttributes::Curves);
+	OutAttributes.Add(UE::Anim::FAttributes::Attributes);
+}
+
 #undef LOCTEXT_NAMESPACE
diff --git a/Engine/Source/Editor/AnimGraph/Public/AnimGraphAttributes.h b/Engine/Source/Editor/AnimGraph/Public/AnimGraphAttributes.h
new file mode 100644
--- /dev/null
+++ b/Engine/Source/Editor/AnimGraph/Public/AnimGraphAttributes.h
@@ -0,0 +1,108 @@
+// Copyright Epic Games, Inc. All Rights Reserved.
+
+#pragma once
+
+#include "CoreMinimal.h"
+#include "AnimGraphAttributes.generated.h"
+
+UENUM()
+enum class EAnimGraphAttributesDisplayMode
+{
+	// Always hide on pins
+	HideOnPins,
+
+	// Always show on pins
+	ShowOnPins,
+
+	// Automatically calculate visibility based on graph connectivity/usage
+	Automatic
+};
+
+UENUM()
+enum class EAnimGraphAttributeBlend
+{
+	// Attribute is subject to weighted blends as it flows through the graph
+	Blendable,
+
+	// Attribute is not blended. Corresponds to messaging between two nodes that are not necessarily directly connected.
+	NonBlendable,
+};
+
+// Description of an attribute
+USTRUCT()
+struct FAnimGraphAttributeDesc
+{
+	GENERATED_BODY()
+
+	FAnimGraphAttributeDesc() = default;
+
+	FAnimGraphAttributeDesc(FName InName, EAnimGraphAttributeBlend InBlend, FSlateBrush InIcon = FSlateBrush(), FText InDisplayName = FText(), FText InToolTipText = FText(), FSlateColor InColor = FSlateColor::UseSubduedForeground(), EAnimGraphAttributesDisplayMode InDisplayMode = EAnimGraphAttributesDisplayMode::Automatic, int32 InSortOrder = INDEX_NONE)
+		: Name(InName)
+		, Icon(InIcon)
+		, DisplayName(InDisplayName.IsEmpty() ? FText::FromName(InName) : InDisplayName)
+		, ToolTipText(InToolTipText.IsEmpty() ? FText::FromName(InName) : InToolTipText)
+		, Color(InColor)
+		, DisplayMode(InDisplayMode)
+		, Blend(InBlend)
+		, SortOrder(InSortOrder)
+	{
+	}
+
+	UPROPERTY(config)
+	FName Name;
+
+	UPROPERTY(config)
+	FSlateBrush Icon;
+
+	UPROPERTY(config)
+	FText DisplayName;
+
+	UPROPERTY(config)
+	FText ToolTipText;
+
+	UPROPERTY(config)
+	FSlateColor Color;
+
+	// How the attribute is displayed in the graph. This overrides any settings per node type.
+	UPROPERTY(config)
+	EAnimGraphAttributesDisplayMode DisplayMode;
+
+	// Blendability of the attribute
+	UPROPERTY(config)
+	EAnimGraphAttributeBlend Blend;
+
+	// How to sort the attribute
+	UPROPERTY(config)
+	int32 SortOrder;
+};
+
+UCLASS(Config = EditorPerProjectUserSettings)
+class ANIMGRAPH_API UAnimGraphAttributes : public UObject
+{
+	GENERATED_BODY()
+
+public:
+	// Register an anim graph attribute
+	void Register(const FAnimGraphAttributeDesc& InDesc);
+
+	// Get a previously registered attribute desc
+	// @param	InName	The name of a previously registered attribute
+	// @return nullptr if no desc is registered with the specified name
+	const FAnimGraphAttributeDesc* FindAttributeDesc(FName InName) const;
+
+	// Iterate over all attributes that are registered
+	// @param	InFunction	FUnction receiving the various attribute descriptors
+	void ForEachAttribute(TFunctionRef<void(const FAnimGraphAttributeDesc&)> InFunction) const;
+
+private:
+	// UObject interface
+	virtual void PostEditChangeProperty(FPropertyChangedEvent& PropertyChangedEvent) override;
+
+private:
+	// All attributes
+	UPROPERTY(config)
+	TArray<FAnimGraphAttributeDesc> Attributes;
+
+	// Cached attribute index map for faster lookup
+	TMap<FName, int32> AttributesMap;
+};
\ No newline at end of file
diff --git a/Engine/Source/Editor/AnimGraph/Public/AnimGraphNode_AssetPlayerBase.h b/Engine/Source/Editor/AnimGraph/Public/AnimGraphNode_AssetPlayerBase.h
--- a/Engine/Source/Editor/AnimGraph/Public/AnimGraphNode_AssetPlayerBase.h
+++ b/Engine/Source/Editor/AnimGraph/Public/AnimGraphNode_AssetPlayerBase.h
@@ -24,12 +24,17 @@ public:
 	UPROPERTY(EditAnywhere, Category = Settings)
 	FAnimationGroupReference SyncGroup;
 
+	/** UObject interface */
+	void Serialize(FArchive& Ar) override;
+	void PostEditChangeProperty(FPropertyChangedEvent& PropertyChangedEvent) override;
+
 	/** UEdGraphNode interface */
 	virtual void PinConnectionListChanged(UEdGraphPin* Pin) override;
 	virtual void PinDefaultValueChanged(UEdGraphPin* Pin) override;
 
 	/** UAnimGraphNode_Base interface */
 	virtual void OnProcessDuringCompilation(IAnimBlueprintCompilationContext& InCompilationContext, IAnimBlueprintGeneratedClassCompiledData& OutCompiledData) override;
-
+	virtual void GetOutputLinkAttributes(FNodeAttributeArray& OutAttributes) const override;
+	virtual void ValidateAnimNodeDuringCompilation(USkeleton* ForSkeleton, FCompilerResultsLog& MessageLog) override;
 	virtual void SetAnimationAsset(UAnimationAsset* Asset) { check(false); /*Base function called*/ }
 };
diff --git a/Engine/Source/Editor/AnimGraph/Public/AnimGraphNode_Base.h b/Engine/Source/Editor/AnimGraph/Public/AnimGraphNode_Base.h
--- a/Engine/Source/Editor/AnimGraph/Public/AnimGraphNode_Base.h
+++ b/Engine/Source/Editor/AnimGraph/Public/AnimGraphNode_Base.h
@@ -242,7 +242,7 @@ class ANIMGRAPH_API UAnimGraphNode_Base : public UK2Node
 	virtual void ValidateAnimNodePostCompile(FCompilerResultsLog& MessageLog, UAnimBlueprintGeneratedClass* CompiledClass, int32 CompiledNodeIndex) {}
 
 	// If using CopyPoseFromMesh, the AnimBlueprint Compiler will cache this off for optimizations. 
-	virtual bool UsingCopyPoseFromMesh() { return false; }
+	virtual bool UsingCopyPoseFromMesh() const { return false; }
 
 	// Gives each visual node a chance to update the node template before it is inserted in the compiled class
 	virtual void BakeDataDuringCompilation(FCompilerResultsLog& MessageLog) {}
@@ -354,6 +354,17 @@ class ANIMGRAPH_API UAnimGraphNode_Base : public UK2Node
 	DECLARE_EVENT(UAnimGraphNode_Base, FOnNodeTitleChangedEvent);
 	FOnNodeTitleChangedEvent& OnNodeTitleChangedEvent() { return NodeTitleChangedEvent; }
 
+	using FNodeAttributeArray = TArray<FName, TInlineAllocator<4>>;
+
+	// Get the named attribute types that this node takes (absorbs) as inputs. Other attributes are assumed to 'pass through' this node.
+	virtual void GetInputLinkAttributes(FNodeAttributeArray& OutAttributes) const {}
+
+	// Get the named attribute types that this node provides as outputs. Other attributes are assumed to 'pass through' this node.
+	virtual void GetOutputLinkAttributes(FNodeAttributeArray& OutAttributes) const {}
+
+	// @return wether to show graph attribute icons on pins for this node.
+	virtual bool ShouldShowAttributesOnPins() const { return true; }
+
 protected:
 	friend class FAnimBlueprintCompilerContext;
 	friend class FAnimGraphNodeDetails;
diff --git a/Engine/Source/Editor/AnimGraph/Public/AnimGraphNode_BlendListBase.h b/Engine/Source/Editor/AnimGraph/Public/AnimGraphNode_BlendListBase.h
--- a/Engine/Source/Editor/AnimGraph/Public/AnimGraphNode_BlendListBase.h
+++ b/Engine/Source/Editor/AnimGraph/Public/AnimGraphNode_BlendListBase.h
@@ -20,6 +20,7 @@ class UAnimGraphNode_BlendListBase : public UAnimGraphNode_Base
 
 	// UAnimGraphNode_Base interface
 	virtual FString GetNodeCategory() const override;
+	virtual void GetOutputLinkAttributes(FNodeAttributeArray& OutAttributes) const override;
 	// End of UAnimGraphNode_Base interface
 
 protected:
diff --git a/Engine/Source/Editor/AnimGraph/Public/AnimGraphNode_CopyPoseFromMesh.h b/Engine/Source/Editor/AnimGraph/Public/AnimGraphNode_CopyPoseFromMesh.h
--- a/Engine/Source/Editor/AnimGraph/Public/AnimGraphNode_CopyPoseFromMesh.h
+++ b/Engine/Source/Editor/AnimGraph/Public/AnimGraphNode_CopyPoseFromMesh.h
@@ -23,6 +23,7 @@ public:
 	// End of UEdGraphNode interface
 
 	// UAnimGraphNode_Base interface
-	virtual bool UsingCopyPoseFromMesh() { return true; }
+	virtual bool UsingCopyPoseFromMesh() const override { return true; }
+	virtual void GetOutputLinkAttributes(FNodeAttributeArray& OutAttributes) const override;
 	// End of UAnimGraphNode_Base interface
 };
diff --git a/Engine/Source/Editor/AnimGraph/Public/AnimGraphNode_CurveSource.h b/Engine/Source/Editor/AnimGraph/Public/AnimGraphNode_CurveSource.h
--- a/Engine/Source/Editor/AnimGraph/Public/AnimGraphNode_CurveSource.h
+++ b/Engine/Source/Editor/AnimGraph/Public/AnimGraphNode_CurveSource.h
@@ -25,5 +25,6 @@ public:
 
 	// UAnimGraphNode_Base interface
 	virtual FString GetNodeCategory() const override;
+	virtual void GetOutputLinkAttributes(FNodeAttributeArray& OutAttributes) const override;
 	// End of UAnimGraphNode_Base interface
 };
diff --git a/Engine/Source/Editor/AnimGraph/Public/AnimGraphNode_Inertialization.h b/Engine/Source/Editor/AnimGraph/Public/AnimGraphNode_Inertialization.h
--- a/Engine/Source/Editor/AnimGraph/Public/AnimGraphNode_Inertialization.h
+++ b/Engine/Source/Editor/AnimGraph/Public/AnimGraphNode_Inertialization.h
@@ -21,4 +21,5 @@ class UAnimGraphNode_Inertialization : public UAnimGraphNode_Base
 	virtual FText GetTooltipText() const override;
 	virtual FText GetNodeTitle(ENodeTitleType::Type TitleType) const override;
 	virtual FText GetMenuCategory() const override;
+	virtual void GetInputLinkAttributes(FNodeAttributeArray& OutAttributes) const override;
 };
diff --git a/Engine/Source/Editor/AnimGraph/Public/AnimGraphNode_LinkedAnimGraphBase.h b/Engine/Source/Editor/AnimGraph/Public/AnimGraphNode_LinkedAnimGraphBase.h
--- a/Engine/Source/Editor/AnimGraph/Public/AnimGraphNode_LinkedAnimGraphBase.h
+++ b/Engine/Source/Editor/AnimGraph/Public/AnimGraphNode_LinkedAnimGraphBase.h
@@ -34,6 +34,8 @@ public:
 
 	// UAnimGraphNode_Base interface
 	virtual FPoseLinkMappingRecord GetLinkIDLocation(const UScriptStruct* NodeType, UEdGraphPin* SourcePin) override;
+	virtual void GetOutputLinkAttributes(FNodeAttributeArray& OutAttributes) const override;
+	virtual bool ShouldShowAttributesOnPins() const override { return false; }
 
 	// UAnimGraphNode_CustomProperty interface
 	virtual bool IsStructuralProperty(FProperty* InProperty) const override;
diff --git a/Engine/Source/Editor/AnimGraph/Public/AnimGraphNode_LinkedInputPose.h b/Engine/Source/Editor/AnimGraph/Public/AnimGraphNode_LinkedInputPose.h
--- a/Engine/Source/Editor/AnimGraph/Public/AnimGraphNode_LinkedInputPose.h
+++ b/Engine/Source/Editor/AnimGraph/Public/AnimGraphNode_LinkedInputPose.h
@@ -85,6 +85,8 @@ public:
 
 	/** UAnimGraphNode_Base interface */
 	virtual void CustomizeDetails(IDetailLayoutBuilder& DetailBuilder) override;
+	virtual void GetOutputLinkAttributes(FNodeAttributeArray& OutAttributes) const override;
+	virtual bool ShouldShowAttributesOnPins() const override { return false; }
 
 	/** Make a name widget for this linked input pose node */
 	TSharedRef<SWidget> MakeNameWidget(IDetailLayoutBuilder& DetailBuilder);
diff --git a/Engine/Source/Editor/AnimGraph/Public/AnimGraphNode_RandomPlayer.h b/Engine/Source/Editor/AnimGraph/Public/AnimGraphNode_RandomPlayer.h
--- a/Engine/Source/Editor/AnimGraph/Public/AnimGraphNode_RandomPlayer.h
+++ b/Engine/Source/Editor/AnimGraph/Public/AnimGraphNode_RandomPlayer.h
@@ -24,4 +24,8 @@ public:
 	virtual FText GetNodeTitle(ENodeTitleType::Type TitleType) const override;
 	virtual FText GetMenuCategory() const override;
 	// End of UEdGraphNode interface
+
+	// UAnimGraphNode_Base interface
+	virtual void GetOutputLinkAttributes(FNodeAttributeArray& OutAttributes) const override;
+	// End of UAnimGraphNode_Base interface
 };
diff --git a/Engine/Source/Editor/AnimGraph/Public/AnimGraphNode_StateMachineBase.h b/Engine/Source/Editor/AnimGraph/Public/AnimGraphNode_StateMachineBase.h
--- a/Engine/Source/Editor/AnimGraph/Public/AnimGraphNode_StateMachineBase.h
+++ b/Engine/Source/Editor/AnimGraph/Public/AnimGraphNode_StateMachineBase.h
@@ -37,6 +37,7 @@ class ANIMGRAPH_API UAnimGraphNode_StateMachineBase : public UAnimGraphNode_Base
 	// UAnimGraphNode_Base interface
 	virtual FString GetNodeCategory() const override;
 	virtual void OnProcessDuringCompilation(IAnimBlueprintCompilationContext& InCompilationContext, IAnimBlueprintGeneratedClassCompiledData& OutCompiledData) override;
+	virtual void GetOutputLinkAttributes(FNodeAttributeArray& OutAttributes) const override;
 	// End of UAnimGraphNode_Base interface
 
 	//  @return the name of this state machine
diff --git a/Engine/Source/Editor/AnimGraph/Public/AnimGraphNode_TransitionPoseEvaluator.h b/Engine/Source/Editor/AnimGraph/Public/AnimGraphNode_TransitionPoseEvaluator.h
--- a/Engine/Source/Editor/AnimGraph/Public/AnimGraphNode_TransitionPoseEvaluator.h
+++ b/Engine/Source/Editor/AnimGraph/Public/AnimGraphNode_TransitionPoseEvaluator.h
@@ -32,6 +32,7 @@ class UAnimGraphNode_TransitionPoseEvaluator : public UAnimGraphNode_Base
 	// UAnimGraphNode_Base interface
 	virtual FString GetNodeCategory() const override;
 	virtual void CustomizeDetails(IDetailLayoutBuilder& DetailBuilder) override;
+	virtual void GetOutputLinkAttributes(FNodeAttributeArray& OutAttributes) const override;
 	// End of UAnimGraphNode_Base interface
 
 	// UK2Node interface.
diff --git a/Engine/Source/Editor/AnimGraph/Public/IAnimBlueprintCompilationContext.h b/Engine/Source/Editor/AnimGraph/Public/IAnimBlueprintCompilationContext.h
--- a/Engine/Source/Editor/AnimGraph/Public/IAnimBlueprintCompilationContext.h
+++ b/Engine/Source/Editor/AnimGraph/Public/IAnimBlueprintCompilationContext.h
@@ -4,6 +4,7 @@
 
 #include "CoreMinimal.h"
 #include "KismetCompiler.h"
+#include "Containers/ArrayView.h"
 
 class UEdGraphNode;
 class UEdGraph;
@@ -102,6 +103,12 @@ public:
 	// Map of anim node indices to node properties
 	const TMap<UAnimGraphNode_Base*, FProperty*>& GetAllocatedPropertiesByNode() const { return GetAllocatedPropertiesByNodeImpl(); }
 
+	// Add the specified compiled-in attribute uniquely to the specified node
+	void AddAttributesToNode(UAnimGraphNode_Base* InNode, TArrayView<const FName> InAttributes) const { AddAttributesToNodeImpl(InNode, InAttributes); }
+
+	// Get the current compiled-in attributes uniquely assigned to the specified node
+	TArrayView<const FName> GetAttributesFromNode(UAnimGraphNode_Base* InNode) const { return GetAttributesFromNodeImpl(InNode); }
+
 protected:
 	// Adds a pose link mapping record
 	virtual void AddPoseLinkMappingRecordImpl(const FPoseLinkMappingRecord& InRecord) = 0;
@@ -148,6 +155,12 @@ protected:
 	// Map of anim node indices to node properties
 	virtual const TMap<UAnimGraphNode_Base*, FProperty*>& GetAllocatedPropertiesByNodeImpl() const = 0;
 
+	// Add the specified compiled-in attribute uniquely to the specified node
+	virtual void AddAttributesToNodeImpl(UAnimGraphNode_Base* InNode, TArrayView<const FName> InAttributes) const = 0;
+
+	// Get the current compiled-in attributes uniquely assigned to the specified node
+	virtual TArrayView<const FName> GetAttributesFromNodeImpl(UAnimGraphNode_Base* InNode) const = 0;
+
 	// GetHandler helper function
 	virtual IAnimBlueprintCompilerHandler* GetHandlerInternal(FName InName) const = 0;
 
diff --git a/Engine/Source/Editor/AnimationBlueprintEditor/Private/AnimGraphConnectionDrawingPolicy.cpp b/Engine/Source/Editor/AnimationBlueprintEditor/Private/AnimGraphConnectionDrawingPolicy.cpp
--- a/Engine/Source/Editor/AnimationBlueprintEditor/Private/AnimGraphConnectionDrawingPolicy.cpp
+++ b/Engine/Source/Editor/AnimationBlueprintEditor/Private/AnimGraphConnectionDrawingPolicy.cpp
@@ -6,6 +6,9 @@
 #include "Kismet2/BlueprintEditorUtils.h"
 #include "AnimationGraphSchema.h"
 #include "AnimGraphNode_Base.h"
+#include "AnimGraphAttributes.h"
+#include "Algo/Sort.h"
+#include "AnimationPins/SGraphPinPose.h"
 
 /////////////////////////////////////////////////////
 // FAnimGraphConnectionDrawingPolicy
@@ -26,53 +29,81 @@ void FAnimGraphConnectionDrawingPolicy::BuildExecutionRoadmap()
 {
 	if(UAnimBlueprint* TargetBP = Cast<UAnimBlueprint>(FBlueprintEditorUtils::FindBlueprintForGraph(GraphObj)))
 	{
-	UAnimBlueprintGeneratedClass* AnimBlueprintClass = (UAnimBlueprintGeneratedClass*)(*(TargetBP->GeneratedClass));
+		UAnimBlueprintGeneratedClass* AnimBlueprintClass = (UAnimBlueprintGeneratedClass*)(*(TargetBP->GeneratedClass));
 
-	if (TargetBP->GetObjectBeingDebugged() == NULL)
-	{
-		return;
-	}
-
-	TMap<FProperty*, UObject*> PropertySourceMap;
-	AnimBlueprintClass->GetDebugData().GenerateReversePropertyMap(/*out*/ PropertySourceMap);
+		if (TargetBP->GetObjectBeingDebugged() == NULL)
+		{
+			return;
+		}
 
-	FAnimBlueprintDebugData& DebugInfo = AnimBlueprintClass->GetAnimBlueprintDebugData();
-	for (auto VisitIt = DebugInfo.UpdatedNodesThisFrame.CreateIterator(); VisitIt; ++VisitIt)
-	{
-		const FAnimBlueprintDebugData::FNodeVisit& VisitRecord = *VisitIt;
+		TMap<FProperty*, UObject*> PropertySourceMap;
+		AnimBlueprintClass->GetDebugData().GenerateReversePropertyMap(/*out*/ PropertySourceMap);
 
-		if ((VisitRecord.SourceID >= 0) && (VisitRecord.SourceID < AnimBlueprintClass->GetAnimNodeProperties().Num()) && (VisitRecord.TargetID >= 0) && (VisitRecord.TargetID < AnimBlueprintClass->AnimNodeProperties.Num()))
+		FAnimBlueprintDebugData& DebugInfo = AnimBlueprintClass->GetAnimBlueprintDebugData();
+		for (auto VisitIt = DebugInfo.UpdatedNodesThisFrame.CreateIterator(); VisitIt; ++VisitIt)
 		{
-			if (UAnimGraphNode_Base* SourceNode = Cast<UAnimGraphNode_Base>(PropertySourceMap.FindRef(AnimBlueprintClass->GetAnimNodeProperties()[VisitRecord.SourceID])))
+			const FAnimBlueprintDebugData::FNodeVisit& VisitRecord = *VisitIt;
+
+			if ((VisitRecord.SourceID >= 0) && (VisitRecord.SourceID < AnimBlueprintClass->GetAnimNodeProperties().Num()) && (VisitRecord.TargetID >= 0) && (VisitRecord.TargetID < AnimBlueprintClass->AnimNodeProperties.Num()))
 			{
-				if (UAnimGraphNode_Base* TargetNode = Cast<UAnimGraphNode_Base>(PropertySourceMap.FindRef(AnimBlueprintClass->GetAnimNodeProperties()[VisitRecord.TargetID])))
+				if (UAnimGraphNode_Base* SourceNode = Cast<UAnimGraphNode_Base>(PropertySourceMap.FindRef(AnimBlueprintClass->GetAnimNodeProperties()[VisitRecord.SourceID])))
 				{
-					UEdGraphPin* PoseNet = NULL;
-
-					UAnimationGraphSchema const* AnimSchema = GetDefault<UAnimationGraphSchema>();
-					for (int32 PinIndex = 0; PinIndex < TargetNode->Pins.Num(); ++PinIndex)
+					if (UAnimGraphNode_Base* TargetNode = Cast<UAnimGraphNode_Base>(PropertySourceMap.FindRef(AnimBlueprintClass->GetAnimNodeProperties()[VisitRecord.TargetID])))
 					{
-						UEdGraphPin* Pin = TargetNode->Pins[PinIndex];
-						check(Pin);
-						if (AnimSchema->IsPosePin(Pin->PinType) && (Pin->Direction == EGPD_Output))
+						UEdGraphPin* PoseNet = NULL;
+
+						UAnimationGraphSchema const* AnimSchema = GetDefault<UAnimationGraphSchema>();
+						for (int32 PinIndex = 0; PinIndex < TargetNode->Pins.Num(); ++PinIndex)
 						{
-							PoseNet = Pin;
-							break;
+							UEdGraphPin* Pin = TargetNode->Pins[PinIndex];
+							check(Pin);
+							if (AnimSchema->IsPosePin(Pin->PinType) && (Pin->Direction == EGPD_Output))
+							{
+								PoseNet = Pin;
+								break;
+							}
 						}
-					}
 
-					if (PoseNet != NULL)
-					{
-						//@TODO: Extend the rendering code to allow using the recorded blend weight instead of faked exec times
-						FExecPairingMap& Predecessors = PredecessorPins.FindOrAdd((UEdGraphNode*)SourceNode);
-						FTimePair& Timings = Predecessors.FindOrAdd(PoseNet);
-						Timings.PredExecTime = 0.0;
-						Timings.ThisExecTime = FMath::Clamp(VisitRecord.Weight, 0.f, 1.f);
+						if (PoseNet != NULL)
+						{
+							//@TODO: Extend the rendering code to allow using the recorded blend weight instead of faked exec times
+							FExecPairingMap& Predecessors = PredecessorPins.FindOrAdd((UEdGraphNode*)SourceNode);
+							FTimePair& Timings = Predecessors.FindOrAdd(PoseNet);
+							Timings.PredExecTime = 0.0;
+							Timings.ThisExecTime = FMath::Clamp(VisitRecord.Weight, 0.f, 1.f);
+						}
 					}
 				}
 			}
 		}
 	}
+}
+
+void FAnimGraphConnectionDrawingPolicy::BuildPinToPinWidgetMap(TMap<TSharedRef<SWidget>, FArrangedWidget>& InPinGeometries)
+{
+	FKismetConnectionDrawingPolicy::BuildPinToPinWidgetMap(InPinGeometries);
+
+	// Cache additional attributes
+	PinAttributes.Reset();
+	bool bFoundPanelZoom = false;
+	PanelZoom = 1.0f;
+
+	for(const TPair<UEdGraphPin*, TSharedPtr<SGraphPin>>& PinWidgetPair : PinToPinWidgetMap)
+	{
+		if(PinWidgetPair.Key->Direction == EGPD_Output && UAnimationGraphSchema::IsPosePin(PinWidgetPair.Key->PinType) && PinWidgetPair.Key->GetOwningNode()->IsA<UAnimGraphNode_Base>())
+		{
+			// Pose pins are assumed to be SGraphPinPose widgets here
+			check(PinWidgetPair.Value->GetType() == TEXT("SGraphPinPose"));
+			TSharedPtr<SGraphPinPose> PosePin = StaticCastSharedPtr<SGraphPinPose>(PinWidgetPair.Value);
+
+			PinAttributes.Add(PinWidgetPair.Key, PosePin->GetAttributeInfo());
+
+			if(!bFoundPanelZoom)
+			{
+				PanelZoom = PosePin->GetZoomAmount();
+				bFoundPanelZoom = true;
+			}
+		}
 	}
 }
 
@@ -90,3 +121,143 @@ void FAnimGraphConnectionDrawingPolicy::DetermineStyleOfExecWire(float& Thicknes
 	bDrawBubbles = true;
 }
 
+void FAnimGraphConnectionDrawingPolicy::DrawConnection(int32 LayerId, const FVector2D& Start, const FVector2D& End, const FConnectionParams& Params)
+{
+	bool bCompositeWire = false;
+
+	// Pose pins display attribute links
+	if(Params.AssociatedPin1 && Params.AssociatedPin2)
+	{
+		if(UAnimationGraphSchema::IsPosePin(Params.AssociatedPin1->PinType))
+		{
+			UAnimGraphNode_Base* Node1 = Cast<UAnimGraphNode_Base>(Params.AssociatedPin1->GetOwningNode());
+			UAnimGraphNode_Base* Node2 = Cast<UAnimGraphNode_Base>(Params.AssociatedPin2->GetOwningNode());
+			if(Node1 && Node2)
+			{
+				const TArrayView<const SGraphPinPose::FAttributeInfo>& AdditionalAttributeInfo = PinAttributes.FindRef(Params.AssociatedPin1);
+
+				if(AdditionalAttributeInfo.Num() > 0)
+				{
+					const float MaxAttributeWireThickness = 3.0f;
+					const float MinAttributeWireThickness = 1.0f;
+					const float MaxWireGap = 2.0f;
+					const float MinWireGap = 0.5f;
+
+					// 0.375f is the zoom level before the 'low LOD' cutoff
+					const float ZoomLevelAlpha = FMath::GetMappedRangeValueClamped(TRange<float>(0.375f, 1.0f), TRange<float>(0.0f, 1.0f), PanelZoom);
+					const float AttributeWireThickness = FMath::Lerp(MinAttributeWireThickness, MaxAttributeWireThickness, ZoomLevelAlpha);
+					const float WireGap = FMath::Lerp(MinWireGap, MaxWireGap, ZoomLevelAlpha);
+
+					const FVector2D& P0 = Start;
+					const FVector2D& P1 = End;
+
+					const FVector2D SplineTangent = ComputeSplineTangent(P0, P1);
+					const FVector2D P0Tangent = (Params.StartDirection == EGPD_Output) ? SplineTangent : -SplineTangent;
+					const FVector2D P1Tangent = (Params.EndDirection == EGPD_Input) ? SplineTangent : -SplineTangent;
+
+					bCompositeWire = true;
+
+					float TotalThickness = Params.WireThickness;
+
+					static TArray<float> CachedWireThicknesses;
+					check(CachedWireThicknesses.Num() == 0);	// Cant be called recursively or on multiple threads
+					CachedWireThicknesses.SetNumZeroed(AdditionalAttributeInfo.Num());
+
+					for(int32 AttributeInfoIndex = 0; AttributeInfoIndex < AdditionalAttributeInfo.Num(); ++AttributeInfoIndex)
+					{
+						const SGraphPinPose::FAttributeInfo& AttributeInfo = AdditionalAttributeInfo[AttributeInfoIndex];
+
+						float WireThickness = 0.0f;
+						switch(AttributeInfo.Blend)
+						{
+						case EAnimGraphAttributeBlend::Blendable:
+						{
+							if(UAnimBlueprint* AnimBlueprint = Cast<UAnimBlueprint>(FBlueprintEditorUtils::FindBlueprintForGraph(GraphObj)))
+							{
+								UAnimBlueprintGeneratedClass* AnimBlueprintClass = (UAnimBlueprintGeneratedClass*)(*(AnimBlueprint->GeneratedClass));
+								int32 SourceNodeId = AnimBlueprintClass->GetNodeIndexFromGuid(Node1->NodeGuid);
+								int32 TargetNodeId = AnimBlueprintClass->GetNodeIndexFromGuid(Node2->NodeGuid);
+								if(SourceNodeId != INDEX_NONE && TargetNodeId != INDEX_NONE)
+								{
+									const TArray<FAnimBlueprintDebugData::FAttributeRecord>* LinkAttributes = AnimBlueprintClass->GetAnimBlueprintDebugData().NodeOutputAttributesThisFrame.Find(SourceNodeId);
+									const bool bAttributeUsedInLink = LinkAttributes && LinkAttributes->ContainsByPredicate(
+										[&AttributeInfo, TargetNodeId](const FAnimBlueprintDebugData::FAttributeRecord& InRecord)
+										{
+											return InRecord.Attribute == AttributeInfo.Attribute && InRecord.OtherNode == TargetNodeId; 
+										});
+
+
+									WireThickness = bAttributeUsedInLink ? AttributeWireThickness : 0.0f;
+								}
+							}
+							break;
+						}
+						case EAnimGraphAttributeBlend::NonBlendable:
+							WireThickness = AttributeWireThickness;
+							break;
+						}
+
+						CachedWireThicknesses[AttributeInfoIndex] = WireThickness;
+						TotalThickness += WireThickness != 0.0f ? (WireThickness + WireGap) : 0.0f;
+					}
+
+					const float InitialOffset = TotalThickness * 0.5f;
+					FVector2D SubWireP0 = P0;
+					SubWireP0.Y += InitialOffset;
+					FVector2D SubWireP1 = P1;
+					SubWireP1.Y += InitialOffset;
+
+					// Draw in reverse order to get pose wires appearing on top
+					for(int32 AttributeInfoIndex = AdditionalAttributeInfo.Num() - 1; AttributeInfoIndex >= 0; --AttributeInfoIndex)
+					{
+						float Thickness = CachedWireThicknesses[AttributeInfoIndex];
+						if(Thickness > 0.0f)
+						{
+							const SGraphPinPose::FAttributeInfo& AttributeInfo = AdditionalAttributeInfo[AttributeInfoIndex];
+							FLinearColor Color = AttributeInfo.Color;
+
+							if(HoveredPins.Num() > 0)
+							{
+								ApplyHoverDeemphasis(Params.AssociatedPin1, Params.AssociatedPin2, Thickness, Color);
+							}
+
+							SubWireP0.Y -= Thickness + WireGap;
+							SubWireP1.Y -= Thickness + WireGap;
+
+							// Draw the spline itself
+							FSlateDrawElement::MakeDrawSpaceSpline(
+								DrawElementsList,
+								LayerId,
+								SubWireP0, P0Tangent,
+								SubWireP1, P1Tangent,
+								Thickness,
+								ESlateDrawEffect::None,
+								Color
+							);
+						}
+					}
+
+					SubWireP0.Y -= Params.WireThickness + WireGap;
+					SubWireP1.Y -= Params.WireThickness + WireGap;
+
+					FKismetConnectionDrawingPolicy::DrawConnection(LayerId, SubWireP0, SubWireP1, Params);
+
+					CachedWireThicknesses.Reset();
+				}
+			}
+		}
+	}
+
+	if(!bCompositeWire)
+	{
+		FKismetConnectionDrawingPolicy::DrawConnection(LayerId, Start, End, Params);
+	}
+}
+
+void FAnimGraphConnectionDrawingPolicy::ApplyHoverDeemphasis(UEdGraphPin* OutputPin, UEdGraphPin* InputPin, float& Thickness, FLinearColor& WireColor)
+{
+	// Remove the thickness increase on hover
+	float OriginalThickness = Thickness;
+	FKismetConnectionDrawingPolicy::ApplyHoverDeemphasis(OutputPin, InputPin, Thickness, WireColor);
+	Thickness = OriginalThickness;
+}
\ No newline at end of file
diff --git a/Engine/Source/Editor/AnimationBlueprintEditor/Private/AnimGraphConnectionDrawingPolicy.h b/Engine/Source/Editor/AnimationBlueprintEditor/Private/AnimGraphConnectionDrawingPolicy.h
--- a/Engine/Source/Editor/AnimationBlueprintEditor/Private/AnimGraphConnectionDrawingPolicy.h
+++ b/Engine/Source/Editor/AnimationBlueprintEditor/Private/AnimGraphConnectionDrawingPolicy.h
@@ -4,6 +4,7 @@
 
 #include "CoreMinimal.h"
 #include "BlueprintConnectionDrawingPolicy.h"
+#include "AnimationPins/SGraphPinPose.h"
 
 class FSlateWindowElementList;
 class UEdGraph;
@@ -22,6 +23,19 @@ public:
 	virtual bool TreatWireAsExecutionPin(UEdGraphPin* InputPin, UEdGraphPin* OutputPin) const override;
 	virtual void BuildExecutionRoadmap() override;
 	virtual void DetermineStyleOfExecWire(float& Thickness, FLinearColor& WireColor, bool& bDrawBubbles, const FTimePair& Times) override;
+	virtual void DrawConnection(int32 LayerId, const FVector2D& Start, const FVector2D& End, const FConnectionParams& Params) override;
+	virtual void BuildPinToPinWidgetMap(TMap<TSharedRef<SWidget>, FArrangedWidget>& InPinGeometries) override;
+	virtual void ApplyHoverDeemphasis(UEdGraphPin* OutputPin, UEdGraphPin* InputPin, /*inout*/ float& Thickness, /*inout*/ FLinearColor& WireColor) override;
 	// End of FKismetConnectionDrawingPolicy interface
+
+private:
+	// Map to cached attribute array on the node
+	TMap<UEdGraphPin*, TArrayView<const SGraphPinPose::FAttributeInfo>> PinAttributes;
+	
+	// Handle to compilation delegate
+	FDelegateHandle OnBlueprintCompiledHandle;
+
+	// Zoom level of the current panel
+	float PanelZoom;
 };
 
diff --git a/Engine/Source/Editor/AnimationBlueprintEditor/Private/AnimationNodes/SGraphNodeBlendSpacePlayer.cpp b/Engine/Source/Editor/AnimationBlueprintEditor/Private/AnimationNodes/SGraphNodeBlendSpacePlayer.cpp
--- a/Engine/Source/Editor/AnimationBlueprintEditor/Private/AnimationNodes/SGraphNodeBlendSpacePlayer.cpp
+++ b/Engine/Source/Editor/AnimationBlueprintEditor/Private/AnimationNodes/SGraphNodeBlendSpacePlayer.cpp
@@ -9,8 +9,7 @@
 #include "PersonaModule.h"
 #include "Widgets/Layout/SBox.h"
 #include "Modules/ModuleManager.h"
-
-PRAGMA_DISABLE_OPTIMIZATION
+#include "AnimGraphNode_BlendSpacePlayer.h"
 
 void SGraphNodeBlendSpacePlayer::Construct(const FArguments& InArgs, UAnimGraphNode_Base* InNode)
 {
@@ -20,11 +19,14 @@ void SGraphNodeBlendSpacePlayer::Construct(const FArguments& InArgs, UAnimGraphN
 
 	this->UpdateGraphNode();
 
+	CachedSyncGroupName = NAME_None;
+
 	SAnimationGraphNode::Construct(SAnimationGraphNode::FArguments(), InNode);
 
 	RegisterActiveTimer(1.0f / 60.0f, FWidgetActiveTimerDelegate::CreateLambda([this](double InCurrentTime, float InDeltaTime)
 	{
 		GetBlendSpaceInfo(CachedBlendSpace, CachedPosition);
+		UpdateGraphSyncLabel();
 		return EActiveTimerReturnType::Continue;
 	}));
 }
@@ -113,4 +115,39 @@ bool SGraphNodeBlendSpacePlayer::GetBlendSpaceInfo(TWeakObjectPtr<const UBlendSp
 	return false;
 }
 
-PRAGMA_ENABLE_OPTIMIZATION
\ No newline at end of file
+
+void SGraphNodeBlendSpacePlayer::UpdateGraphSyncLabel()
+{
+	if (UAnimGraphNode_BlendSpacePlayer* VisualBlendSpacePlayer = Cast<UAnimGraphNode_BlendSpacePlayer>(GraphNode))
+	{
+		FName CurrentSyncGroupName = NAME_None;
+
+		if (UAnimBlueprint* AnimBlueprint = Cast<UAnimBlueprint>(FBlueprintEditorUtils::FindBlueprintForNode(GraphNode)))
+		{
+			if(UAnimBlueprintGeneratedClass* GeneratedClass = AnimBlueprint->GetAnimBlueprintGeneratedClass())
+			{
+				if (UObject* ActiveObject = AnimBlueprint->GetObjectBeingDebugged())
+				{
+					if(VisualBlendSpacePlayer->SyncGroup.Method == EAnimSyncMethod::Graph)
+					{
+						int32 NodeIndex = GeneratedClass->GetNodeIndexFromGuid(VisualBlendSpacePlayer->NodeGuid);
+						if(NodeIndex != INDEX_NONE)
+						{
+							if(const FName* SyncGroupNamePtr = GeneratedClass->GetAnimBlueprintDebugData().NodeSyncsThisFrame.Find(NodeIndex))
+							{
+								CurrentSyncGroupName = *SyncGroupNamePtr;
+							}
+						}
+					}
+				}
+			}
+		}
+
+		if(CachedSyncGroupName != CurrentSyncGroupName)
+		{
+			// Invalidate the node title so we can dynamically display the sync group gleaned from the graph
+			VisualBlendSpacePlayer->OnNodeTitleChangedEvent().Broadcast();
+			CachedSyncGroupName = CurrentSyncGroupName;
+		}
+	}
+}
\ No newline at end of file
diff --git a/Engine/Source/Editor/AnimationBlueprintEditor/Private/AnimationNodes/SGraphNodeBlendSpacePlayer.h b/Engine/Source/Editor/AnimationBlueprintEditor/Private/AnimationNodes/SGraphNodeBlendSpacePlayer.h
--- a/Engine/Source/Editor/AnimationBlueprintEditor/Private/AnimationNodes/SGraphNodeBlendSpacePlayer.h
+++ b/Engine/Source/Editor/AnimationBlueprintEditor/Private/AnimationNodes/SGraphNodeBlendSpacePlayer.h
@@ -28,6 +28,12 @@ protected:
 	EVisibility GetBlendSpaceVisibility() const;
 	bool GetBlendSpaceInfo(TWeakObjectPtr<const UBlendSpaceBase>& OutBlendSpace, FVector& OutPosition) const;
 
+	// Invalidates the node's label if we are syncing based on graph context
+	void UpdateGraphSyncLabel();
+
 	TWeakObjectPtr<const UBlendSpaceBase> CachedBlendSpace;
 	FVector CachedPosition;
+
+	// Cached name to display when sync groups are dynamic
+	FName CachedSyncGroupName;
 };
diff --git a/Engine/Source/Editor/AnimationBlueprintEditor/Private/AnimationNodes/SGraphNodeSequencePlayer.cpp b/Engine/Source/Editor/AnimationBlueprintEditor/Private/AnimationNodes/SGraphNodeSequencePlayer.cpp
--- a/Engine/Source/Editor/AnimationBlueprintEditor/Private/AnimationNodes/SGraphNodeSequencePlayer.cpp
+++ b/Engine/Source/Editor/AnimationBlueprintEditor/Private/AnimationNodes/SGraphNodeSequencePlayer.cpp
@@ -20,7 +20,15 @@ void SGraphNodeSequencePlayer::Construct(const FArguments& InArgs, UAnimGraphNod
 
 	this->UpdateGraphNode();
 
+	CachedSyncGroupName = NAME_None;
+
 	SAnimationGraphNode::Construct(SAnimationGraphNode::FArguments(), InNode);
+
+	RegisterActiveTimer(1.0f / 60.0f, FWidgetActiveTimerDelegate::CreateLambda([this](double InCurrentTime, float InDeltaTime)
+	{
+		UpdateGraphSyncLabel();
+		return EActiveTimerReturnType::Continue;
+	}));
 }
 
 void SGraphNodeSequencePlayer::GetNodeInfoPopups(FNodeInfoContext* Context, TArray<FGraphInformationPopupInfo>& Popups) const
@@ -175,3 +183,39 @@ void SGraphNodeSequencePlayer::SetSequencePositionRatio(float NewRatio)
 		}
 	}
 }
+
+void SGraphNodeSequencePlayer::UpdateGraphSyncLabel()
+{
+	if (UAnimGraphNode_SequencePlayer* VisualSequencePlayer = Cast<UAnimGraphNode_SequencePlayer>(GraphNode))
+	{
+		FName CurrentSyncGroupName = NAME_None;
+
+		if (UAnimBlueprint* AnimBlueprint = Cast<UAnimBlueprint>(FBlueprintEditorUtils::FindBlueprintForNode(GraphNode)))
+		{
+			if(UAnimBlueprintGeneratedClass* GeneratedClass = AnimBlueprint->GetAnimBlueprintGeneratedClass())
+			{
+				if (UObject* ActiveObject = AnimBlueprint->GetObjectBeingDebugged())
+				{
+					if(VisualSequencePlayer->SyncGroup.Method == EAnimSyncMethod::Graph)
+					{
+						int32 NodeIndex = GeneratedClass->GetNodeIndexFromGuid(VisualSequencePlayer->NodeGuid);
+						if(NodeIndex != INDEX_NONE)
+						{
+							if(const FName* SyncGroupNamePtr = GeneratedClass->GetAnimBlueprintDebugData().NodeSyncsThisFrame.Find(NodeIndex))
+							{
+								CurrentSyncGroupName = *SyncGroupNamePtr;
+							}
+						}
+					}
+				}
+			}
+		}
+
+		if(CachedSyncGroupName != CurrentSyncGroupName)
+		{
+			// Invalidate the node title so we can dynamically display the sync group gleaned from the graph
+			VisualSequencePlayer->OnNodeTitleChangedEvent().Broadcast();
+			CachedSyncGroupName = CurrentSyncGroupName;
+		}
+	}
+}
\ No newline at end of file
diff --git a/Engine/Source/Editor/AnimationBlueprintEditor/Private/AnimationNodes/SGraphNodeSequencePlayer.h b/Engine/Source/Editor/AnimationBlueprintEditor/Private/AnimationNodes/SGraphNodeSequencePlayer.h
--- a/Engine/Source/Editor/AnimationBlueprintEditor/Private/AnimationNodes/SGraphNodeSequencePlayer.h
+++ b/Engine/Source/Editor/AnimationBlueprintEditor/Private/AnimationNodes/SGraphNodeSequencePlayer.h
@@ -38,4 +38,10 @@ protected:
 	FText GetPositionTooltip() const;
 
 	bool GetSequencePositionInfo(float& Out_Position, float& Out_Length, int32& FrameCount) const;
+
+	// Invalidates the node's label if we are syncing based on graph context
+	void UpdateGraphSyncLabel();
+
+	// Cached name to display when sync groups are dynamic
+	FName CachedSyncGroupName;
 };
diff --git a/Engine/Source/Editor/AnimationBlueprintEditor/Private/AnimationPins/SGraphPinPose.cpp b/Engine/Source/Editor/AnimationBlueprintEditor/Private/AnimationPins/SGraphPinPose.cpp
--- a/Engine/Source/Editor/AnimationBlueprintEditor/Private/AnimationPins/SGraphPinPose.cpp
+++ b/Engine/Source/Editor/AnimationBlueprintEditor/Private/AnimationPins/SGraphPinPose.cpp
@@ -1,6 +1,14 @@
 // Copyright Epic Games, Inc. All Rights Reserved.
 
 #include "AnimationPins/SGraphPinPose.h"
+#include "AnimGraphNode_Base.h"
+#include "Preferences/PersonaOptions.h"
+#include "AnimGraphAttributes.h"
+#include "Kismet2/BlueprintEditorUtils.h"
+#include "Widgets/Layout/SWrapBox.h"
+#include "SGraphPanel.h"
+
+#define LOCTEXT_NAMESPACE "SGraphPinPose"
 
 /////////////////////////////////////////////////////
 // SGraphPinPose
@@ -13,6 +21,8 @@ void SGraphPinPose::Construct(const FArguments& InArgs, UEdGraphPin* InPin)
 	CachedImg_Pin_Connected = FEditorStyle::GetBrush(TEXT("Graph.PosePin.Connected"));
 	CachedImg_Pin_DisconnectedHovered = FEditorStyle::GetBrush(TEXT("Graph.PosePin.DisconnectedHovered"));
 	CachedImg_Pin_Disconnected = FEditorStyle::GetBrush(TEXT("Graph.PosePin.Disconnected"));
+
+	ReconfigureWidgetForAttributes();
 }
 
 const FSlateBrush* SGraphPinPose::GetPinIcon() const
@@ -30,3 +40,333 @@ const FSlateBrush* SGraphPinPose::GetPinIcon() const
 
 	return Brush;
 }
+
+
+enum class EAttributeUsage
+{
+	// Attribute is input/output by the pin, but is not currently connected to a producer/consumer
+	Unused,
+
+	// Attribute is input/output by the pin and is currently connected to a producer/consumer
+	Used,
+
+	// Attribute is not an input/output of the pin but passes through the node
+	Passthrough,
+};
+
+class SAttributeIndicator : public SCompoundWidget
+{
+public:
+	SLATE_BEGIN_ARGS(SAttributeIndicator) {}
+
+	SLATE_ARGUMENT(FName, Attribute)
+
+	SLATE_END_ARGS()
+
+	void Construct(const FArguments& InArgs, UAnimBlueprint* InAnimBlueprint, const FAnimGraphAttributeDesc* InAttributeDesc, const UEdGraphPin* InPin, EAttributeUsage InUsage, UAnimGraphNode_Base* InAnimGraphNode)
+	{
+		AnimBlueprint = InAnimBlueprint;
+		AttributeDesc = InAttributeDesc;
+		Pin = InPin;
+		Usage = InUsage;
+		AnimGraphNode = InAnimGraphNode;
+		ActiveColor = AttributeDesc->Color.GetSpecifiedColor();
+		Color = ActiveColor;
+		BackgroundColor = FLinearColor::Transparent;
+		Value = 0.0f;
+
+		ChildSlot
+		[
+			SNew(SOverlay)
+			.ToolTipText(AttributeDesc->ToolTipText)
+			+SOverlay::Slot()
+			[
+				SNew(SBorder)
+				.Padding(FMargin(2.0f, 1.0f, 5.0f, 1.0f))
+				.BorderImage(FEditorStyle::GetBrush("AnimGraph.Attribute.Border.Solid"))
+				.BorderBackgroundColor_Lambda([this](){ return BackgroundColor; })
+			]
+			+SOverlay::Slot()
+			.Padding(FMargin(2.0f, 1.0f, 2.0f, 1.0f))
+			[
+				SNew(SHorizontalBox)
+				+SHorizontalBox::Slot()
+				.AutoWidth()
+				.VAlign(VAlign_Center)
+				.HAlign(HAlign_Center)
+				[
+					SNew(SImage)
+					.Image(&AttributeDesc->Icon)
+					.ColorAndOpacity_Lambda([this](){ return Color; })
+				]
+			]
+		];
+	}
+
+	virtual void Tick(const FGeometry& AllottedGeometry, const double InCurrentTime, const float InDeltaTime) override
+	{
+		bool bActive = false;
+
+		if((Usage == EAttributeUsage::Used || AttributeDesc->Blend == EAnimGraphAttributeBlend::Blendable) && AnimBlueprint && AnimBlueprint->GetObjectBeingDebugged() != nullptr)
+		{
+			UAnimBlueprintGeneratedClass* AnimBlueprintClass = (UAnimBlueprintGeneratedClass*)(*(AnimBlueprint->GeneratedClass));
+			int32 SourceNodeId = AnimBlueprintClass->GetNodeIndexFromGuid(AnimGraphNode->NodeGuid);
+			if(SourceNodeId != INDEX_NONE) 
+			{
+				const TArray<FAnimBlueprintDebugData::FAttributeRecord>* LinkAttributes;
+				if(Pin->Direction == EGPD_Input)
+				{
+					LinkAttributes = AnimBlueprintClass->GetAnimBlueprintDebugData().NodeInputAttributesThisFrame.Find(SourceNodeId);
+				}
+				else
+				{
+					LinkAttributes = AnimBlueprintClass->GetAnimBlueprintDebugData().NodeOutputAttributesThisFrame.Find(SourceNodeId);
+				}
+
+				bActive = LinkAttributes && LinkAttributes->ContainsByPredicate([this](const FAnimBlueprintDebugData::FAttributeRecord& InRecord){ return InRecord.Attribute == AttributeDesc->Name; });
+			}
+		}
+
+		if(bActive)
+		{
+			Value = 1.0f;
+			Color = FLinearColor::Black;
+			BackgroundColor = ActiveColor;
+		}
+		else
+		{
+			Value = FMath::FInterpTo(Value, 0.0f, InDeltaTime, 4.0f);
+			Color = FMath::Lerp(ActiveColor, FLinearColor::Black, Value);
+			BackgroundColor = FMath::Lerp(FLinearColor::Transparent, ActiveColor, Value);
+		}
+	}
+
+	UAnimBlueprint* AnimBlueprint;
+	const FAnimGraphAttributeDesc* AttributeDesc;
+	const UEdGraphPin* Pin;
+	EAttributeUsage Usage;
+	UAnimGraphNode_Base* AnimGraphNode;
+	FLinearColor ActiveColor;
+	FLinearColor Color;
+	FLinearColor BackgroundColor;
+	float Value;
+};
+
+void SGraphPinPose::ReconfigureWidgetForAttributes()
+{
+	AttributeInfos.Empty();
+
+	UAnimGraphNode_Base* AnimGraphNode = CastChecked<UAnimGraphNode_Base>(GraphPinObj->GetOwningNode());
+	UAnimBlueprint* AnimBlueprint = Cast<UAnimBlueprint>(FBlueprintEditorUtils::FindBlueprintForNode(AnimGraphNode));
+
+	auto AddAttributes = [this, AnimBlueprint, AnimGraphNode](TArrayView<const FName> InUsedAttributes, TArrayView<const FName> InUnusedAttributes, TArrayView<const FName> InPassThroughAttributes)
+	{
+		auto AddAttribute = [this, AnimBlueprint, AnimGraphNode](FName InAttribute, EAttributeUsage InUsage)
+		{
+			const FAnimGraphAttributeDesc* AttributeDesc = GetDefault<UAnimGraphAttributes>()->FindAttributeDesc(InAttribute);
+
+			// Early out if we don't want to display this attribute
+			if(AttributeDesc)
+			{
+				if(AttributeDesc->DisplayMode == EAnimGraphAttributesDisplayMode::HideOnPins)
+				{
+					return;
+				}
+
+				// Only store cached attributes for wires on used pins
+				if(InUsage != EAttributeUsage::Unused)
+				{
+					AttributeInfos.Emplace(InAttribute, AttributeDesc->Color.GetSpecifiedColor(), AttributeDesc->Blend, AttributeDesc->SortOrder);
+
+					// Skip displaying passthrough attributes on pins
+					if(InUsage != EAttributeUsage::Passthrough && AnimGraphNode->ShouldShowAttributesOnPins())
+					{
+						GetLabelAndValue()->AddSlot()
+						.Padding(2.0f, 0.0f, 0.0f, 0.0f)
+						[
+							SNew(SAttributeIndicator, AnimBlueprint, AttributeDesc, GraphPinObj, InUsage, AnimGraphNode)
+						];
+					}
+				}
+			}
+		};
+
+		for(const FName& Attribute : InUsedAttributes)
+		{
+			AddAttribute(Attribute, EAttributeUsage::Used);
+		}
+
+		for(const FName& Attribute : InUnusedAttributes)
+		{
+			AddAttribute(Attribute, EAttributeUsage::Unused);
+		}
+
+		for(const FName& Attribute : InPassThroughAttributes)
+		{
+			AddAttribute(Attribute, EAttributeUsage::Passthrough);
+		}
+	};
+
+	UAnimBlueprintGeneratedClass* AnimBlueprintClass = (UAnimBlueprintGeneratedClass*)(*(AnimBlueprint->GeneratedClass));
+
+	UAnimGraphNode_Base::FNodeAttributeArray PinAttributes;
+	switch(GraphPinObj->Direction)
+	{
+	case EGPD_Input:
+		AnimGraphNode->GetInputLinkAttributes(PinAttributes);
+		break;
+	case EGPD_Output:
+		AnimGraphNode->GetOutputLinkAttributes(PinAttributes);
+		break;
+	}
+
+	UAnimGraphNode_Base::FNodeAttributeArray NodeAttributes;
+	AnimGraphNode->GetInputLinkAttributes(NodeAttributes);
+	AnimGraphNode->GetOutputLinkAttributes(NodeAttributes);
+
+	// Unlinked pins display attributes if they are inputs and the node takes them as inputs
+	if(GraphPinObj->LinkedTo.Num() == 0)
+	{
+		AddAttributes(TArrayView<const FName>(), PinAttributes, TArrayView<const FName>());
+	}
+	else if(GraphPinObj->LinkedTo.Num() > 0)
+	{
+		TArrayView<const FName> CompilerGeneratedAttributes = AnimBlueprintClass->GetAnimBlueprintDebugData().GetNodeAttributes(AnimGraphNode);
+		UAnimGraphNode_Base::FNodeAttributeArray UsedAttributes;
+		UAnimGraphNode_Base::FNodeAttributeArray UnusedAttributes;
+		if(PinAttributes.Num() > 0)
+		{
+			for(const FName& Name : PinAttributes)
+			{
+				if(CompilerGeneratedAttributes.Contains(Name))
+				{
+					UsedAttributes.Add(Name);
+				}
+				else
+				{
+					UnusedAttributes.Add(Name);
+				}
+			}
+		}
+
+		UAnimGraphNode_Base::FNodeAttributeArray PassthroughAttributes;
+		for(const FName& Name : CompilerGeneratedAttributes)
+		{
+			if(!PinAttributes.Contains(Name))
+			{
+				PassthroughAttributes.Add(Name);
+			}
+		}
+
+		AddAttributes(UsedAttributes, UnusedAttributes, PassthroughAttributes);
+	}
+
+	// sort pin attributes
+	Algo::Sort(AttributeInfos, [](const FAttributeInfo& InValue0, const FAttributeInfo& InValue1)
+	{
+		return InValue0.SortOrder < InValue1.SortOrder;
+	});
+
+	// overrides attribute set in the base class Construct()
+	SetToolTipText(MakeAttributeSP(this, &SGraphPinPose::GetAttributeTooltipText));
+}
+
+TArrayView<const SGraphPinPose::FAttributeInfo> SGraphPinPose::GetAttributeInfo() const
+{
+	SGraphNode* MyOwnerNode = OwnerNodePtr.Pin().Get();
+	if (MyOwnerNode && MyOwnerNode->GetOwnerPanel().IsValid())
+	{
+		if(MyOwnerNode->GetOwnerPanel()->GetCurrentLOD() <= EGraphRenderingLOD::LowDetail)
+		{
+			return TArrayView<const FAttributeInfo>();
+		}
+	}
+
+	return AttributeInfos;
+}
+
+float SGraphPinPose::GetZoomAmount() const
+{
+	SGraphNode* MyOwnerNode = OwnerNodePtr.Pin().Get();
+	if (MyOwnerNode && MyOwnerNode->GetOwnerPanel().IsValid())
+	{
+		return MyOwnerNode->GetOwnerPanel()->GetZoomAmount();
+	}
+
+	return 1.0f;
+}
+
+FText SGraphPinPose::GetAttributeTooltipText() const
+{
+	if(GraphPinObj)
+	{
+		UEdGraphPin* OtherPin = GraphPinObj->LinkedTo.Num() > 0 ? GraphPinObj->LinkedTo[0] : nullptr;
+		if(OtherPin)
+		{
+			UAnimGraphNode_Base* Node1 = Cast<UAnimGraphNode_Base>(GraphPinObj->GetOwningNode());
+			UAnimGraphNode_Base* Node2 = Cast<UAnimGraphNode_Base>(OtherPin->GetOwningNode());
+			if(Node1 && Node2)
+			{
+				if(UAnimBlueprint* AnimBlueprint = Cast<UAnimBlueprint>(FBlueprintEditorUtils::FindBlueprintForNode(Node1)))
+				{
+					if(AnimBlueprint->GetObjectBeingDebugged() != nullptr)
+					{
+						UAnimBlueprintGeneratedClass* AnimBlueprintClass = (UAnimBlueprintGeneratedClass*)(*(AnimBlueprint->GeneratedClass));
+						int32 SourceNodeId = AnimBlueprintClass->GetNodeIndexFromGuid(Node1->NodeGuid);
+						int32 TargetNodeId = AnimBlueprintClass->GetNodeIndexFromGuid(Node2->NodeGuid);
+
+						if(AttributeInfos.Num() > 0)
+						{
+							FTextBuilder TextBuilder;
+							TextBuilder.AppendLine(SGraphPin::GetTooltipText());
+
+							bool bAddedAttributeSubtitle = false;
+
+							for(const SGraphPinPose::FAttributeInfo& AttributeInfo : AttributeInfos)
+							{
+
+
+								if(SourceNodeId != INDEX_NONE && TargetNodeId != INDEX_NONE)
+								{
+									const TArray<FAnimBlueprintDebugData::FAttributeRecord>* LinkAttributes;
+
+									if(GraphPinObj->Direction == EGPD_Input)
+									{
+										LinkAttributes = AnimBlueprintClass->GetAnimBlueprintDebugData().NodeInputAttributesThisFrame.Find(SourceNodeId);
+									}
+									else
+									{
+										LinkAttributes = AnimBlueprintClass->GetAnimBlueprintDebugData().NodeOutputAttributesThisFrame.Find(SourceNodeId);
+									}
+
+									const bool bAttributeUsedInLink = LinkAttributes && LinkAttributes->ContainsByPredicate(
+										[&AttributeInfo, TargetNodeId](const FAnimBlueprintDebugData::FAttributeRecord& InRecord)
+										{
+											return InRecord.Attribute == AttributeInfo.Attribute && InRecord.OtherNode == TargetNodeId; 
+										});
+
+									if(bAttributeUsedInLink)
+									{
+										if(!bAddedAttributeSubtitle)
+										{
+											TextBuilder.AppendLine(LOCTEXT("AttributesSubtitle", "Attributes:"));
+											bAddedAttributeSubtitle = true;
+										}
+
+										TextBuilder.AppendLine(AttributeInfo.Attribute);
+									}
+								}
+							}
+
+							return TextBuilder.ToText();
+						}
+					}
+				}
+			}
+		}
+	}
+	
+	return SGraphPin::GetTooltipText();
+}
+
+#undef LOCTEXT_NAMESPACE
\ No newline at end of file
diff --git a/Engine/Source/Editor/AnimationBlueprintEditor/Private/AnimationPins/SGraphPinPose.h b/Engine/Source/Editor/AnimationBlueprintEditor/Private/AnimationPins/SGraphPinPose.h
--- a/Engine/Source/Editor/AnimationBlueprintEditor/Private/AnimationPins/SGraphPinPose.h
+++ b/Engine/Source/Editor/AnimationBlueprintEditor/Private/AnimationPins/SGraphPinPose.h
@@ -6,6 +6,8 @@
 #include "Widgets/DeclarativeSyntaxSupport.h"
 #include "SGraphPin.h"
 
+enum class EAnimGraphAttributeBlend;
+
 class SGraphPinPose : public SGraphPin
 {
 public:
@@ -14,6 +16,34 @@ public:
 
 	void Construct(const FArguments& InArgs, UEdGraphPin* InPin);
 
+	// Struct used by connection drawing to draw attributes
+	struct FAttributeInfo
+	{
+		FAttributeInfo(FName InAttribute, const FLinearColor& InColor, EAnimGraphAttributeBlend InBlend, int32 InSortOrder)
+			: Attribute(InAttribute)
+			, Color(InColor)
+			, Blend(InBlend)
+			, SortOrder(InSortOrder)
+		{}
+
+		FName Attribute;
+		FLinearColor Color;
+		EAnimGraphAttributeBlend Blend;
+		int32 SortOrder;
+	};
+
+	// Get the attribute info used to draw connections. This varies based on LOD level.
+	TArrayView<const FAttributeInfo> GetAttributeInfo() const;
+
+	// Exposes the parent panel's zoom scalar for use when drawing links
+	float GetZoomAmount() const;
+
+private:
+	void ReconfigureWidgetForAttributes();
+
+	// Get tooltip text with attributes included
+	FText GetAttributeTooltipText() const;
+
 protected:
 	//~ Begin SGraphPin Interface
 	virtual const FSlateBrush* GetPinIcon() const override;
@@ -21,4 +51,6 @@ protected:
 
 	mutable const FSlateBrush* CachedImg_Pin_ConnectedHovered;
 	mutable const FSlateBrush* CachedImg_Pin_DisconnectedHovered;
+
+	TArray<FAttributeInfo> AttributeInfos;
 };
diff --git a/Engine/Source/Editor/EditorStyle/Private/StarshipStyle.cpp b/Engine/Source/Editor/EditorStyle/Private/StarshipStyle.cpp
--- a/Engine/Source/Editor/EditorStyle/Private/StarshipStyle.cpp
+++ b/Engine/Source/Editor/EditorStyle/Private/StarshipStyle.cpp
@@ -5661,6 +5661,28 @@ void FStarshipEditorStyle::FStyle::SetupPersonaStyle()
 			.SetMenuBorderPadding(FMargin(0.0f))
 		);
 	}
+
+	// Animation blueprint
+	{
+		Set("AnimGraph.Attribute.Border.Solid", new FSlateRoundedBoxBrush(FStyleColors::White));
+		Set("AnimGraph.Attribute.Border.Outline", new FSlateRoundedBoxBrush(FStyleColors::Transparent, FStyleColors::White, 1.0f));
+		Set("AnimGraph.Attribute.DefaultColor", FStyleColors::AccentGray);
+
+		FSlateColor AttributeTextColor = FStyleColors::AccentBlack;
+		Set("AnimGraph.Attribute.TextColor", AttributeTextColor);
+		Set("AnimGraph.Attribute.Text", FTextBlockStyle(NormalText)
+			.SetFont(DEFAULT_FONT("Regular", FStarshipCoreStyle::SmallTextSize))
+			.SetColorAndOpacity(AttributeTextColor));
+
+		Set("AnimGraph.Attribute.Curves.Icon", new IMAGE_BRUSH("GenericCurveEditor/Icons/GenericCurveEditor_48x", Icon10x10));
+		Set("AnimGraph.Attribute.Curves.Color", FStyleColors::AccentPurple);
+		Set("AnimGraph.Attribute.Attributes.Icon", new IMAGE_BRUSH("Icons/pillarray_16x", Icon10x10));
+		Set("AnimGraph.Attribute.Attributes.Color", FStyleColors::AccentYellow);
+		Set("AnimGraph.Attribute.Sync.Icon", new IMAGE_BRUSH("Icons/icon_tab_TimecodeProvider_16x", Icon10x10));
+		Set("AnimGraph.Attribute.Sync.Color", FStyleColors::AccentBlue);
+		Set("AnimGraph.Attribute.InertialBlending.Icon", new IMAGE_BRUSH("Icons/icon_tab_TimecodeProvider_16x", Icon10x10));
+		Set("AnimGraph.Attribute.InertialBlending.Color", FStyleColors::AccentOrange);
+	}	
 #endif // WITH_EDITOR
 }
 
diff --git a/Engine/Source/Runtime/AnimGraphRuntime/Private/AnimNodes/AnimNode_BlendListBase.cpp b/Engine/Source/Runtime/AnimGraphRuntime/Private/AnimNodes/AnimNode_BlendListBase.cpp
--- a/Engine/Source/Runtime/AnimGraphRuntime/Private/AnimNodes/AnimNode_BlendListBase.cpp
+++ b/Engine/Source/Runtime/AnimGraphRuntime/Private/AnimNodes/AnimNode_BlendListBase.cpp
@@ -126,6 +126,7 @@ void FAnimNode_BlendListBase::Update_AnyThread(const FAnimationUpdateContext& Co
 				if (InertializationRequester)
 				{
 					InertializationRequester->RequestInertialization(BlendTime[ChildIndex]);
+					InertializationRequester->AddDebugRecord(*Context.AnimInstanceProxy, Context.GetCurrentNodeId());
 				}
 				else
 				{
diff --git a/Engine/Source/Runtime/AnimGraphRuntime/Private/AnimNodes/AnimNode_BlendSpacePlayer.cpp b/Engine/Source/Runtime/AnimGraphRuntime/Private/AnimNodes/AnimNode_BlendSpacePlayer.cpp
--- a/Engine/Source/Runtime/AnimGraphRuntime/Private/AnimNodes/AnimNode_BlendSpacePlayer.cpp
+++ b/Engine/Source/Runtime/AnimGraphRuntime/Private/AnimNodes/AnimNode_BlendSpacePlayer.cpp
@@ -5,6 +5,8 @@
 #include "Animation/AnimSequence.h"
 #include "Animation/AnimInstanceProxy.h"
 #include "AnimGraphRuntimeTrace.h"
+#include "Animation/AnimSync.h"
+#include "Animation/AnimSyncScope.h"
 
 /////////////////////////////////////////////////////
 // FAnimNode_BlendSpacePlayer
@@ -81,27 +83,22 @@ void FAnimNode_BlendSpacePlayer::UpdateAssetPlayer(const FAnimationUpdateContext
 void FAnimNode_BlendSpacePlayer::UpdateInternal(const FAnimationUpdateContext& Context)
 {
 	DECLARE_SCOPE_HIERARCHICAL_COUNTER_ANIMNODE(UpdateInternal)
-	if ((BlendSpace != NULL) && (Context.AnimInstanceProxy->IsSkeletonCompatible(BlendSpace->GetSkeleton())))
-	{
-		// Create a tick record and fill it out
-		FAnimGroupInstance* SyncGroup;
-		FAnimTickRecord& TickRecord = Context.AnimInstanceProxy->CreateUninitializedTickRecordInScope(/*out*/ SyncGroup, GroupName, GroupScope);
-
-		const FVector BlendInput(X, Y, Z);
-	
+	if ((BlendSpace != nullptr) && (Context.AnimInstanceProxy->IsSkeletonCompatible(BlendSpace->GetSkeleton())))
+	{	
 		if (PreviousBlendSpace != BlendSpace)
 		{
 			Reinitialize(bResetPlayTimeWhenBlendSpaceChanges);
 		}
 
-		Context.AnimInstanceProxy->MakeBlendSpaceTickRecord(TickRecord, BlendSpace, BlendInput, BlendSampleDataCache, BlendFilter, bLoop, PlayRate, Context.GetFinalBlendWeight(), /*inout*/ InternalTimeAccumulator, MarkerTickRecord);
+		const FVector BlendInput(X, Y, Z);
 
-		// Update the sync group if it exists
-		if (SyncGroup != NULL)
-		{
-			SyncGroup->TestTickRecordForLeadership(GroupRole);
-		}
+		// Create a tick record and push into the closest scope
+		UE::Anim::FAnimSyncGroupScope& SyncScope = Context.GetMessageChecked<UE::Anim::FAnimSyncGroupScope>();
+
+		FAnimTickRecord TickRecord(BlendSpace, BlendInput, BlendSampleDataCache, BlendFilter, bLoop, PlayRate, Context.GetFinalBlendWeight(), /*inout*/ InternalTimeAccumulator, MarkerTickRecord);
+		UE::Anim::FAnimSyncParams SyncParams(GroupName, GroupRole, Method);
 
+		SyncScope.AddTickRecord(TickRecord, SyncParams, UE::Anim::FAnimSyncDebugInfo(Context));
 
 		TRACE_ANIM_TICK_RECORD(Context, TickRecord);
 
@@ -124,7 +121,7 @@ void FAnimNode_BlendSpacePlayer::UpdateInternal(const FAnimationUpdateContext& C
 void FAnimNode_BlendSpacePlayer::Evaluate_AnyThread(FPoseContext& Output)
 {
 	DECLARE_SCOPE_HIERARCHICAL_COUNTER_ANIMNODE(Evaluate_AnyThread)
-	if ((BlendSpace != NULL) && (Output.AnimInstanceProxy->IsSkeletonCompatible(BlendSpace->GetSkeleton())))
+	if ((BlendSpace != nullptr) && (Output.AnimInstanceProxy->IsSkeletonCompatible(BlendSpace->GetSkeleton())))
 	{
 		FAnimationPoseData AnimationPoseData(Output);
 		BlendSpace->GetAnimationPose(BlendSampleDataCache, AnimationPoseData);
@@ -202,7 +199,7 @@ void FAnimNode_BlendSpacePlayer::Reinitialize(bool bResetTime)
 		}
 	}
 
-	if (BlendSpace != NULL)
+	if (BlendSpace != nullptr)
 	{
 		BlendSpace->InitializeFilter(&BlendFilter);
 	}
diff --git a/Engine/Source/Runtime/AnimGraphRuntime/Private/AnimNodes/AnimNode_RandomPlayer.cpp b/Engine/Source/Runtime/AnimGraphRuntime/Private/AnimNodes/AnimNode_RandomPlayer.cpp
--- a/Engine/Source/Runtime/AnimGraphRuntime/Private/AnimNodes/AnimNode_RandomPlayer.cpp
+++ b/Engine/Source/Runtime/AnimGraphRuntime/Private/AnimNodes/AnimNode_RandomPlayer.cpp
@@ -6,6 +6,7 @@
 #include "Animation/AnimInstanceProxy.h"
 #include "Animation/AnimSequence.h"
 #include "Animation/AnimTrace.h"
+#include "Animation/AnimSyncScope.h"
 
 FAnimNode_RandomPlayer::FAnimNode_RandomPlayer()
     : CurrentPlayDataIndex(0)
@@ -229,22 +230,18 @@ void FAnimNode_RandomPlayer::Update_AnyThread(const FAnimationUpdateContext& Con
 		NextData = &GetPlayData(ERandomDataIndexType::Next);
 	}
 
-	FAnimInstanceProxy* AnimProxy = Context.AnimInstanceProxy;
-	FAnimGroupInstance* SyncGroup;
-	FAnimTickRecord& TickRecord = AnimProxy->CreateUninitializedTickRecord(SyncGroup, NAME_None);
-	AnimProxy->MakeSequenceTickRecord(
-	    TickRecord, CurrentData->Entry->Sequence, true, CurrentData->PlayRate,
-	    CurrentData->BlendWeight, CurrentData->CurrentPlayTime, CurrentData->MarkerTickRecord);
-	
+	FAnimTickRecord TickRecord(CurrentData->Entry->Sequence, true, CurrentData->PlayRate, CurrentData->BlendWeight, CurrentData->CurrentPlayTime, CurrentData->MarkerTickRecord);
+
+	UE::Anim::FAnimSyncGroupScope& SyncScope = Context.GetMessageChecked<UE::Anim::FAnimSyncGroupScope>();
+	SyncScope.AddTickRecord(TickRecord, UE::Anim::FAnimSyncParams(), UE::Anim::FAnimSyncDebugInfo(Context));
+
 	TRACE_ANIM_TICK_RECORD(Context, TickRecord);
 
 	if (FAnimationRuntime::HasWeight(NextData->BlendWeight))
 	{
-		FAnimTickRecord& NextTickRecord = AnimProxy->CreateUninitializedTickRecord(SyncGroup, NAME_None);
-		AnimProxy->MakeSequenceTickRecord(
-		    NextTickRecord, NextData->Entry->Sequence, true, NextData->PlayRate,
-		    NextData->BlendWeight, NextData->CurrentPlayTime, NextData->MarkerTickRecord);
-			
+		FAnimTickRecord NextTickRecord(NextData->Entry->Sequence, true, NextData->PlayRate, NextData->BlendWeight, NextData->CurrentPlayTime, NextData->MarkerTickRecord);
+		SyncScope.AddTickRecord(NextTickRecord, UE::Anim::FAnimSyncParams(), UE::Anim::FAnimSyncDebugInfo(Context));
+
 		TRACE_ANIM_TICK_RECORD(Context, NextTickRecord);
 	}
 
diff --git a/Engine/Source/Runtime/AnimGraphRuntime/Private/AnimNodes/AnimNode_Sync.cpp b/Engine/Source/Runtime/AnimGraphRuntime/Private/AnimNodes/AnimNode_Sync.cpp
new file mode 100644
--- /dev/null
+++ b/Engine/Source/Runtime/AnimGraphRuntime/Private/AnimNodes/AnimNode_Sync.cpp
@@ -0,0 +1,33 @@
+// Copyright Epic Games, Inc. All Rights Reserved.
+
+#include "AnimNodes/AnimNode_Sync.h"
+#include "Animation/AnimSyncScope.h"
+
+void FAnimNode_Sync::Initialize_AnyThread(const FAnimationInitializeContext& Context)
+{
+	Source.Initialize(Context);
+}
+
+void FAnimNode_Sync::Update_AnyThread(const FAnimationUpdateContext& Context)
+{
+	UE::Anim::TScopedGraphMessage<UE::Anim::FAnimSyncGroupScope> Message(Context, Context, GroupName, GroupRole);
+
+	Source.Update(Context);
+}
+
+void FAnimNode_Sync::CacheBones_AnyThread(const FAnimationCacheBonesContext& Context)
+{
+	Source.CacheBones(Context);
+}
+
+void FAnimNode_Sync::Evaluate_AnyThread(FPoseContext& Output)
+{
+	Source.Evaluate(Output);
+}
+
+void FAnimNode_Sync::GatherDebugData(FNodeDebugData& DebugData)
+{
+	DebugData.AddDebugItem(DebugData.GetNodeName(this));
+
+	Source.GatherDebugData(DebugData);
+}
diff --git a/Engine/Source/Runtime/AnimGraphRuntime/Public/AnimNodes/AnimNode_Sync.h b/Engine/Source/Runtime/AnimGraphRuntime/Public/AnimNodes/AnimNode_Sync.h
new file mode 100644
--- /dev/null
+++ b/Engine/Source/Runtime/AnimGraphRuntime/Public/AnimNodes/AnimNode_Sync.h
@@ -0,0 +1,40 @@
+// Copyright Epic Games, Inc. All Rights Reserved.
+
+#pragma once
+
+#include "CoreMinimal.h"
+#include "Animation/AnimNodeBase.h"
+#include "Animation/AnimSync.h"
+#include "AnimNode_Sync.generated.h"
+
+USTRUCT(BlueprintInternalUseOnly)
+struct ANIMGRAPHRUNTIME_API FAnimNode_Sync : public FAnimNode_Base
+{
+	GENERATED_BODY()
+
+	friend class UAnimGraphNode_Sync;
+
+public:
+	// Get the group name that we synchronize with
+	FName GetGroupName() const { return GroupName; }
+
+private:
+	UPROPERTY(EditAnywhere, Category = Links)
+	FPoseLink Source;
+
+	// The group name that we synchronize with. All nodes employing sync beyond this in the anim graph will implicitly use this sync group.
+	UPROPERTY(EditAnywhere, Category = Settings)
+	FName GroupName;
+
+	// The role this player can assume within the group
+	UPROPERTY(EditAnywhere, Category = Settings)
+	TEnumAsByte<EAnimGroupRole::Type> GroupRole;
+
+private:
+	// FAnimNode_Base
+	virtual void Initialize_AnyThread(const FAnimationInitializeContext& Context) override;
+	virtual void Update_AnyThread(const FAnimationUpdateContext& Context) override;
+	virtual void CacheBones_AnyThread(const FAnimationCacheBonesContext& Context) override;
+	virtual void Evaluate_AnyThread(FPoseContext& Output) override;
+	virtual void GatherDebugData(FNodeDebugData& DebugData) override;
+};
\ No newline at end of file
diff --git a/Engine/Source/Runtime/Core/Public/UObject/UE5MainStreamObjectVersion.h b/Engine/Source/Runtime/Core/Public/UObject/UE5MainStreamObjectVersion.h
--- a/Engine/Source/Runtime/Core/Public/UObject/UE5MainStreamObjectVersion.h
+++ b/Engine/Source/Runtime/Core/Public/UObject/UE5MainStreamObjectVersion.h
@@ -64,6 +64,9 @@ struct CORE_API FUE5MainStreamObjectVersion
 		// Initial Mantle Serialize Version
 		MantleDbSerialize,
 
+		// Animation sync groups explicitly specify sync method
+		AnimSyncGroupsExplicitSyncMethod,
+
 		// -----<new versions can be added above this line>-------------------------------------------------
 		VersionPlusOne,
 		LatestVersion = VersionPlusOne - 1
diff --git a/Engine/Source/Runtime/Engine/Classes/Animation/AnimBlueprintGeneratedClass.h b/Engine/Source/Runtime/Engine/Classes/Animation/AnimBlueprintGeneratedClass.h
--- a/Engine/Source/Runtime/Engine/Classes/Animation/AnimBlueprintGeneratedClass.h
+++ b/Engine/Source/Runtime/Engine/Classes/Animation/AnimBlueprintGeneratedClass.h
@@ -149,6 +149,9 @@ public:
 	// Map from animation node GUID to property index
 	TMap<FGuid, int32> NodeGuidToIndexMap;
 
+	// Map from animation node to attributes
+	TMap<TWeakObjectPtr<const UAnimGraphNode_Base>, TArray<FName>> NodeAttributes;
+
 	// The debug data for each state machine state
 	TArray<FStateMachineStateDebugData> StateData;	
 	
@@ -173,6 +176,25 @@ public:
 	// History of activated nodes
 	TArray<FNodeVisit> UpdatedNodesThisFrame;
 
+	// Record of attribute transfer between nodes
+	struct FAttributeRecord
+	{
+		FName Attribute;
+		int32 OtherNode;
+
+		FAttributeRecord(int32 InOtherNode, FName InAttribute)
+			: Attribute(InAttribute)
+			, OtherNode(InOtherNode)
+		{}
+	};
+
+	// History of node attributes that are output from and input to nodes
+	TMap<int32, TArray<FAttributeRecord>> NodeInputAttributesThisFrame;
+	TMap<int32, TArray<FAttributeRecord>> NodeOutputAttributesThisFrame;
+
+	// History of node syncs - maps from player node index to graph-determined group name
+	TMap<int32, FName> NodeSyncsThisFrame;
+
 	// Values output by nodes
 	struct FNodeValue
 	{
@@ -257,6 +279,10 @@ public:
 	void ResetNodeVisitSites();
 	void RecordNodeVisit(int32 TargetNodeIndex, int32 SourceNodeIndex, float BlendWeight);
 	void RecordNodeVisitArray(const TArray<FNodeVisit>& Nodes);
+	void RecordNodeAttribute(int32 TargetNodeIndex, int32 SourceNodeIndex, FName InAttribute);
+	void RecordNodeAttributeMaps(const TMap<int32, TArray<FAttributeRecord>>& InInputAttributes, const TMap<int32, TArray<FAttributeRecord>>& InOutputAttributes);
+	void RecordNodeSync(int32 InSourceNodeIndex, FName InSyncGroup);
+	void RecordNodeSyncsArray(const TMap<int32, FName>& InNodeSyncs);
 	void RecordStateData(int32 StateMachineIndex, int32 StateIndex, float Weight, float ElapsedTime);
 	void RecordNodeValue(int32 InNodeID, const FString& InText);
 	void RecordSequencePlayer(int32 InNodeID, float InPosition, float InLength, int32 InFrameCount);
@@ -265,6 +291,8 @@ public:
 	void AddPoseWatch(int32 NodeID, FColor Color);
 	void RemovePoseWatch(int32 NodeID);
 	void UpdatePoseWatchColour(int32 NodeID, FColor Color);
+
+	TArrayView<const FName> GetNodeAttributes(TWeakObjectPtr<UAnimGraphNode_Base> InAnimGraphNode) const;
 #endif
 };
 
@@ -471,8 +499,12 @@ public:
 		return AnimationProperty->ContainerPtrToValuePtr<StructType>((void*)Object);
 	}
 
+	// Gets the property index from the original UAnimGraphNode's GUID. Does not remap to property order.
 	const int32* GetNodePropertyIndexFromGuid(FGuid Guid, EPropertySearchMode::Type SearchMode = EPropertySearchMode::OnlyThis);
 
+	// Gets the remapped property index from the original UAnimGraphNode's GUID. Can be used to index the AnimNodeProperties array.
+	int32 GetNodeIndexFromGuid(FGuid Guid, EPropertySearchMode::Type SearchMode = EPropertySearchMode::OnlyThis);
+
 	const UEdGraphNode* GetVisualNodeFromNodePropertyIndex(int32 PropertyIndex) const;
 #endif
 
diff --git a/Engine/Source/Runtime/Engine/Classes/Animation/AnimNodeBase.h b/Engine/Source/Runtime/Engine/Classes/Animation/AnimNodeBase.h
--- a/Engine/Source/Runtime/Engine/Classes/Animation/AnimNodeBase.h
+++ b/Engine/Source/Runtime/Engine/Classes/Animation/AnimNodeBase.h
@@ -293,6 +293,13 @@ public:
 		: FAnimationBaseContext(InAnimInstanceProxy)
 	{
 	}
+
+	FAnimationCacheBonesContext WithNodeId(int32 InNodeId) const
+	{ 
+		FAnimationCacheBonesContext Result(*this);
+		Result.SetNodeId(InNodeId);
+		return Result; 
+	}
 };
 
 /** Update context passed around during animation tree update */
@@ -346,7 +353,7 @@ public:
 		Result.SharedContext = InSharedContext;
 
 		// This is currently only used in the case of cached poses, where we dont want to preserve the previous node, so clear it here
-		Result.PreviousNodeId = INDEX_NONE;
+	//	Result.PreviousNodeId = INDEX_NONE;
 
 		return Result;
 	}
diff --git a/Engine/Source/Runtime/Engine/Classes/Animation/AnimNode_AssetPlayerBase.h b/Engine/Source/Runtime/Engine/Classes/Animation/AnimNode_AssetPlayerBase.h
--- a/Engine/Source/Runtime/Engine/Classes/Animation/AnimNode_AssetPlayerBase.h
+++ b/Engine/Source/Runtime/Engine/Classes/Animation/AnimNode_AssetPlayerBase.h
@@ -50,15 +50,18 @@ struct ENGINE_API FAnimNode_AssetPlayerBase : public FAnimNode_Base
 #if WITH_EDITORONLY_DATA
 	UPROPERTY()
 	int32 GroupIndex_DEPRECATED;
+
+	UPROPERTY()
+	EAnimSyncGroupScope GroupScope_DEPRECATED;
 #endif
 
 	// The role this player can assume within the group (ignored if GroupIndex is INDEX_NONE)
 	UPROPERTY()
 	TEnumAsByte<EAnimGroupRole::Type> GroupRole;
 
-	// The scope at which marker-based sync is applied
+	// How synchronization is determined
 	UPROPERTY()
-	EAnimSyncGroupScope GroupScope;
+	EAnimSyncMethod Method;
 
 	/** If true, "Relevant anim" nodes that look for the highest weighted animation in a state will ignore
 	 *  this node
diff --git a/Engine/Source/Runtime/Engine/Classes/Animation/AnimNode_Inertialization.h b/Engine/Source/Runtime/Engine/Classes/Animation/AnimNode_Inertialization.h
--- a/Engine/Source/Runtime/Engine/Classes/Animation/AnimNode_Inertialization.h
+++ b/Engine/Source/Runtime/Engine/Classes/Animation/AnimNode_Inertialization.h
@@ -26,9 +26,14 @@ class ENGINE_API IInertializationRequester : public IGraphMessage
 	DECLARE_ANIMGRAPH_MESSAGE(IInertializationRequester);
 
 public:
+	static const FName Attribute;
+
 	// Request to activate inertialization for a duration.
 	// If multiple requests are made on the same inertialization node, the minimum requested time will be used.
 	virtual void RequestInertialization(float InRequestedDuration) = 0;
+
+	// Add a record of this request
+	virtual void AddDebugRecord(const FAnimInstanceProxy& InSourceProxy, int32 InSourceNodeId) = 0;
 };
 
 }}	// namespace UE::Anim
diff --git a/Engine/Source/Runtime/Engine/Classes/Animation/AnimationAsset.h b/Engine/Source/Runtime/Engine/Classes/Animation/AnimationAsset.h
--- a/Engine/Source/Runtime/Engine/Classes/Animation/AnimationAsset.h
+++ b/Engine/Source/Runtime/Engine/Classes/Animation/AnimationAsset.h
@@ -23,6 +23,9 @@ class UAnimMetaData;
 class UAssetMappingTable;
 class UAssetUserData;
 class USkeleton;
+class UAnimSequenceBase;
+class UBlendSpaceBase;
+class UPoseAsset;
 
 namespace MarkerIndexSpecialValues
 {
@@ -284,14 +287,14 @@ struct FAnimTickRecord
 	GENERATED_USTRUCT_BODY()
 
 	UPROPERTY()
-	class UAnimationAsset* SourceAsset;
+	class UAnimationAsset* SourceAsset = nullptr;
 
-	float*  TimeAccumulator;
-	float PlayRateMultiplier;
-	float EffectiveBlendWeight;
-	float RootMotionWeightModifier;
+	float* TimeAccumulator = nullptr;
+	float PlayRateMultiplier = 0.0f;
+	float EffectiveBlendWeight = 0.0f;
+	float RootMotionWeightModifier = 0.0f;
 
-	bool bLooping;
+	bool bLooping = false;
 
 	union
 	{
@@ -313,9 +316,9 @@ struct FAnimTickRecord
 	};
 
 	// marker sync related data
-	FMarkerTickRecord* MarkerTickRecord;
-	bool bCanUseMarkerSync;
-	float LeaderScore;
+	FMarkerTickRecord* MarkerTickRecord = nullptr;
+	bool bCanUseMarkerSync = false;
+	float LeaderScore = 0.0f;
 
 	// Return the root motion weight for this tick record
 	float GetRootMotionWeight() const { return EffectiveBlendWeight * RootMotionWeightModifier; }
@@ -334,6 +337,18 @@ public:
 	{
 	}
 
+	// Create a tick record for an anim sequence
+	ENGINE_API FAnimTickRecord(UAnimSequenceBase* InSequence, bool bInLooping, float InPlayRate, float InFinalBlendWeight, float& InCurrentTime, FMarkerTickRecord& InMarkerTickRecord);
+
+	// Create a tick record for a blendspace
+	ENGINE_API FAnimTickRecord(UBlendSpaceBase* InBlendSpace, const FVector& InBlendInput, TArray<FBlendSampleData>& InBlendSampleDataCache, FBlendFilter& InBlendFilter, bool bInLooping, float InPlayRate, float InFinalBlendWeight, float& InCurrentTime, FMarkerTickRecord& InMarkerTickRecord);
+
+	// Create a tick record for a montage
+	ENGINE_API FAnimTickRecord(UAnimMontage* InMontage, float InCurrentPosition, float InPreviousPosition, float InMoveDelta, float InWeight, TArray<FPassedMarker>& InMarkersPassedThisTick, FMarkerTickRecord& InMarkerTickRecord);
+
+	// Create a tick record for a pose asset
+	ENGINE_API FAnimTickRecord(UPoseAsset* InPoseAsset, float InFinalBlendWeight);
+
 	/** This can be used with the Sort() function on a TArray of FAnimTickRecord to sort from higher leader score */
 	ENGINE_API bool operator <(const FAnimTickRecord& Other) const { return LeaderScore > Other.LeaderScore; }
 };
@@ -438,6 +453,7 @@ namespace EAnimGroupRole
 	};
 }
 
+// Deprecated - do not use
 UENUM()
 enum class EAnimSyncGroupScope : uint8
 {
@@ -448,6 +464,20 @@ enum class EAnimSyncGroupScope : uint8
 	Component,
 };
 
+// How an asset will synchronize with other assets
+UENUM()
+enum class EAnimSyncMethod : uint8
+{
+	// Don't sync ever
+	DoNotSync,
+
+	// Use a named sync group
+	SyncGroup,
+
+	// Use the graph structure to provide a sync group to apply
+	Graph
+};
+
 USTRUCT()
 struct FAnimGroupInstance
 {
@@ -495,10 +525,9 @@ public:
 	// Checks the last tick record in the ActivePlayers array to see if it's a better leader than the current candidate.
 	// This should be called once for each record added to ActivePlayers, after the record is setup.
 	ENGINE_API void TestTickRecordForLeadership(EAnimGroupRole::Type MembershipType);
-	// Checks the last tick record in the ActivePlayers array to see if we have a better leader for montage
-	// This is simple rule for higher weight becomes leader
-	// Only one from same sync group with highest weight will be leader
-	ENGINE_API void TestMontageTickRecordForLeadership();
+
+	UE_DEPRECATED(5.0, "Use TestTickRecordForLeadership, as it now internally supports montages")
+	ENGINE_API void TestMontageTickRecordForLeadership() { TestTickRecordForLeadership(EAnimGroupRole::CanBeLeader); }
 
 	// Called after leader has been ticked and decided
 	ENGINE_API void Finalize(const FAnimGroupInstance* PreviousGroup);
@@ -800,21 +829,22 @@ struct FAnimationGroupReference
 {
 	GENERATED_USTRUCT_BODY()
 	
-	// The name of the group
+	// How an asset will synchronize with other assets
 	UPROPERTY(EditAnywhere, Category=Settings)
+	EAnimSyncMethod Method;
+
+	// The name of the group
+	UPROPERTY(EditAnywhere, Category=Settings, meta = (EditCondition = "Method == EAnimSyncMethod::SyncGroup"))
 	FName GroupName;
 
 	// The type of membership in the group (potential leader, always follower, etc...)
-	UPROPERTY(EditAnywhere, Category=Settings)
+	UPROPERTY(EditAnywhere, Category=Settings, meta = (EditCondition = "Method == EAnimSyncMethod::SyncGroup"))
 	TEnumAsByte<EAnimGroupRole::Type> GroupRole;
 
-	// The scope at which marker-based sync is applied (local, component etc...)
-	UPROPERTY(EditAnywhere, Category=Settings)
-	EAnimSyncGroupScope GroupScope;
-
 	FAnimationGroupReference()
-		: GroupRole(EAnimGroupRole::CanBeLeader)
-		, GroupScope(EAnimSyncGroupScope::Local)
+		: Method(EAnimSyncMethod::DoNotSync)
+		, GroupName(NAME_None)
+		, GroupRole(EAnimGroupRole::CanBeLeader)
 	{
 	}
 };
diff --git a/Engine/Source/Runtime/Engine/Private/Animation/AnimAttributes.cpp b/Engine/Source/Runtime/Engine/Private/Animation/AnimAttributes.cpp
new file mode 100644
--- /dev/null
+++ b/Engine/Source/Runtime/Engine/Private/Animation/AnimAttributes.cpp
@@ -0,0 +1,7 @@
+// Copyright Epic Games, Inc. All Rights Reserved.
+
+#include "Animation/AnimAttributes.h"
+
+const FName UE::Anim::FAttributes::Pose("Pose");
+const FName UE::Anim::FAttributes::Curves("Curves");
+const FName UE::Anim::FAttributes::Attributes("Attributes");
\ No newline at end of file
diff --git a/Engine/Source/Runtime/Engine/Private/Animation/AnimBlueprintGeneratedClass.cpp b/Engine/Source/Runtime/Engine/Private/Animation/AnimBlueprintGeneratedClass.cpp
--- a/Engine/Source/Runtime/Engine/Private/Animation/AnimBlueprintGeneratedClass.cpp
+++ b/Engine/Source/Runtime/Engine/Private/Animation/AnimBlueprintGeneratedClass.cpp
@@ -137,6 +137,9 @@ void FAnimBlueprintDebugData::SetSnapshotIndexByTime(UAnimInstance* Instance, do
 void FAnimBlueprintDebugData::ResetNodeVisitSites()
 {
 	UpdatedNodesThisFrame.Empty(UpdatedNodesThisFrame.Num());
+	NodeInputAttributesThisFrame.Empty(NodeInputAttributesThisFrame.Num());
+	NodeOutputAttributesThisFrame.Empty(NodeOutputAttributesThisFrame.Num());
+	NodeSyncsThisFrame.Empty(NodeSyncsThisFrame.Num());
 	StateData.Empty(StateData.Num());
 	NodeValuesThisFrame.Empty(NodeValuesThisFrame.Num());
 	SequencePlayerRecordsThisFrame.Empty(SequencePlayerRecordsThisFrame.Num());
@@ -153,6 +156,59 @@ void FAnimBlueprintDebugData::RecordNodeVisitArray(const TArray<FNodeVisit>& Nod
 	UpdatedNodesThisFrame.Append(Nodes);
 }
 
+void FAnimBlueprintDebugData::RecordNodeAttribute(int32 TargetNodeIndex, int32 SourceNodeIndex, FName InAttribute)
+{
+	if(TargetNodeIndex != INDEX_NONE)
+	{
+		TArray<FAnimBlueprintDebugData::FAttributeRecord>& InputAttributeRecords = NodeInputAttributesThisFrame.FindOrAdd(TargetNodeIndex);
+		InputAttributeRecords.Emplace(SourceNodeIndex, InAttribute);
+	}
+
+	if(SourceNodeIndex != INDEX_NONE)
+	{
+		TArray<FAnimBlueprintDebugData::FAttributeRecord>& OutputAttributeRecords = NodeOutputAttributesThisFrame.FindOrAdd(SourceNodeIndex);
+		OutputAttributeRecords.Emplace(TargetNodeIndex, InAttribute);
+	}
+}
+
+void FAnimBlueprintDebugData::RecordNodeAttributeMaps(const TMap<int32, TArray<FAttributeRecord>>& InInputAttributes, const TMap<int32, TArray<FAttributeRecord>>& InOutputAttributes)
+{
+	if(NodeInputAttributesThisFrame.Num() == 0)
+	{
+		NodeInputAttributesThisFrame = InInputAttributes;
+	}
+	else
+	{
+		NodeInputAttributesThisFrame.Append(InInputAttributes);
+	}
+
+	if(NodeOutputAttributesThisFrame.Num() == 0)
+	{
+		NodeOutputAttributesThisFrame = InOutputAttributes;
+	}
+	else
+	{
+		NodeOutputAttributesThisFrame.Append(InOutputAttributes);
+	}
+}
+
+void FAnimBlueprintDebugData::RecordNodeSync(int32 InSourceNodeIndex, FName InSyncGroup)
+{
+	NodeSyncsThisFrame.FindOrAdd(InSourceNodeIndex, InSyncGroup);
+}
+
+void FAnimBlueprintDebugData::RecordNodeSyncsArray(const TMap<int32, FName>& InNodeSyncs)
+{	
+	if(NodeSyncsThisFrame.Num() == 0)
+	{
+		NodeSyncsThisFrame = InNodeSyncs;
+	}
+	else
+	{
+		NodeSyncsThisFrame.Append(InNodeSyncs);
+	}
+}
+
 void FAnimBlueprintDebugData::RecordStateData(int32 StateMachineIndex, int32 StateIndex, float Weight, float ElapsedTime)
 {
 	StateData.Emplace(StateMachineIndex, StateIndex, Weight, ElapsedTime);
@@ -216,6 +272,18 @@ void FAnimBlueprintDebugData::UpdatePoseWatchColour(int32 NodeID, FColor Color)
 		}
 	}
 }
+
+TArrayView<const FName> FAnimBlueprintDebugData::GetNodeAttributes(TWeakObjectPtr<UAnimGraphNode_Base> InAnimGraphNode) const
+{
+	const TArray<FName>* Attributes = NodeAttributes.Find(InAnimGraphNode);
+	if(Attributes)
+	{
+		return MakeArrayView(*Attributes);
+	}
+
+	return TArrayView<const FName>();
+}
+
 /////////////////////////////////////////////////////
 // FBinaryObjectWriter
 
@@ -584,6 +652,16 @@ const int32* UAnimBlueprintGeneratedClass::GetNodePropertyIndexFromGuid(FGuid Gu
 	return NULL;
 }
 
+int32 UAnimBlueprintGeneratedClass::GetNodeIndexFromGuid(FGuid Guid, EPropertySearchMode::Type SearchMode)
+{
+	if(const int32* IndexPtr = GetNodePropertyIndexFromGuid(Guid, SearchMode))
+	{
+		return AnimNodeProperties.Num() - 1 - *IndexPtr;
+	}
+
+	return INDEX_NONE;
+}
+
 const UEdGraphNode* UAnimBlueprintGeneratedClass::GetVisualNodeFromNodePropertyIndex(int32 PropertyIndex) const
 {
 	const int32 ReversedPropertyIndex = AnimNodeProperties.Num() - PropertyIndex - 1;
diff --git a/Engine/Source/Runtime/Engine/Private/Animation/AnimInstance.cpp b/Engine/Source/Runtime/Engine/Private/Animation/AnimInstance.cpp
--- a/Engine/Source/Runtime/Engine/Private/Animation/AnimInstance.cpp
+++ b/Engine/Source/Runtime/Engine/Private/Animation/AnimInstance.cpp
@@ -341,18 +341,11 @@ void UAnimInstance::UpdateMontageSyncGroup()
 			if (ensure(GroupNameToUse != NAME_None))
 			{
 				bRecordNeedsResetting = false;
-				FAnimGroupInstance* SyncGroup;
-				FAnimTickRecord& TickRecord = GetProxyOnGameThread<FAnimInstanceProxy>().CreateUninitializedTickRecord(/*out*/ SyncGroup, GroupNameToUse);
-				MakeMontageTickRecord(TickRecord, MontageInstance->Montage, MontageInstance->GetPosition(),
+				FAnimTickRecord TickRecord(MontageInstance->Montage, MontageInstance->GetPosition(),
 					MontageInstance->GetPreviousPosition(), MontageInstance->GetDeltaMoved(), MontageInstance->GetWeight(),
 					MontageInstance->MarkersPassedThisTick, MontageInstance->MarkerTickRecord);
-
-				// Update the sync group if it exists
-				if (SyncGroup != NULL)
-				{
-					// the max count should be 2 as you had older one and you have newer one. After TestMontageTickRecordForLeadership, it should set to be 1
-					SyncGroup->TestMontageTickRecordForLeadership();
-				}
+				UE::Anim::FAnimSyncParams Params(GroupNameToUse);
+				GetProxyOnGameThread<FAnimInstanceProxy>().AddTickRecord(TickRecord, Params);
 
 #if ANIM_TRACE_ENABLED
 				FAnimationUpdateContext UpdateContext(&GetProxyOnGameThread<FAnimInstanceProxy>());
@@ -577,7 +570,7 @@ void UAnimInstance::PostUpdateAnimation()
 
 	// flip read/write index
 	// Do this first, as we'll be reading cached slot weights, and we want this to be up to date for this frame.
-	Proxy.TickSyncGroupWriteIndex();
+	Proxy.FlipBufferWriteIndex();
 
 	Proxy.PostUpdate(this);
 
@@ -659,17 +652,6 @@ void UAnimInstance::DispatchQueuedAnimEvents()
 void UAnimInstance::ParallelUpdateAnimation()
 {
 	GetProxyOnAnyThread<FAnimInstanceProxy>().UpdateAnimation();
-
-	if(GetSkelMeshComponent()->GetAnimInstance() == this)
-	{
-		// If this is the main instance,  Tick asset players for this and any linked instances we have
-		for(UAnimInstance* LinkedInstance : GetSkelMeshComponent()->GetLinkedAnimInstances())
-		{
-			LinkedInstance->GetProxyOnAnyThread<FAnimInstanceProxy>().TickAssetPlayerInstances();
-		}
-	}
-
-	GetProxyOnAnyThread<FAnimInstanceProxy>().TickAssetPlayerInstances();
 }
 
 bool UAnimInstance::NeedsImmediateUpdate(float DeltaSeconds) const
diff --git a/Engine/Source/Runtime/Engine/Private/Animation/AnimInstanceProxy.cpp b/Engine/Source/Runtime/Engine/Private/Animation/AnimInstanceProxy.cpp
--- a/Engine/Source/Runtime/Engine/Private/Animation/AnimInstanceProxy.cpp
+++ b/Engine/Source/Runtime/Engine/Private/Animation/AnimInstanceProxy.cpp
@@ -20,6 +20,8 @@
 #include "Animation/AnimNode_Root.h"
 #include "Animation/AnimNode_LinkedAnimLayer.h"
 #include "Animation/AnimMontage.h"
+#include "Animation/AnimInstance.h"
+#include "Animation/AnimSyncScope.h"
 
 #define DO_ANIMSTAT_PROCESSING(StatName) DEFINE_STAT(STAT_ ## StatName)
 #include "Animation/AnimMTStats.h"
@@ -27,8 +29,6 @@
 
 #define DO_ANIMSTAT_PROCESSING(StatName) DEFINE_STAT(STAT_ ## StatName ## _WorkerThread)
 #include "Animation/AnimMTStats.h"
-#include "Animation/AnimInstance.h"
-
 #undef DO_ANIMSTAT_PROCESSING
 
 #define LOCTEXT_NAMESPACE "AnimInstance"
@@ -70,14 +70,8 @@ void FAnimInstanceProxy::UpdateAnimationNode_WithRoot(const FAnimationUpdateCont
 void FAnimInstanceProxy::AddReferencedObjects(UAnimInstance* InAnimInstance, FReferenceCollector& Collector)
 {
 	DECLARE_SCOPE_HIERARCHICAL_COUNTER_FUNC()
-	for (int32 Index = 0; Index < UE_ARRAY_COUNT(UngroupedActivePlayerArrays); ++Index)
-	{
-		TArray<FAnimTickRecord>& UngroupedPlayers = UngroupedActivePlayerArrays[Index];
-		for (FAnimTickRecord& TickRecord : UngroupedPlayers)
-		{
-			Collector.AddReferencedObject(TickRecord.SourceAsset, InAnimInstance);
-		}
-	}
+	
+	Sync.AddReferencedObjects(InAnimInstance, Collector);
 }
 
 void FAnimInstanceProxy::Initialize(UAnimInstance* InAnimInstance)
@@ -395,20 +389,13 @@ void FAnimInstanceProxy::PreUpdate(UAnimInstance* InAnimInstance, float DeltaSec
 
 	ClearSlotNodeWeights();
 
-	// Reset the player tick list (but keep it presized)
-	TArray<FAnimTickRecord>& UngroupedActivePlayers = UngroupedActivePlayerArrays[GetSyncGroupWriteIndex()];
-	UngroupedActivePlayers.Reset();
+	// Reset the synchronizer
+	Sync.Reset();
 
-	FSyncGroupMap& SyncGroups = SyncGroupMaps[GetSyncGroupWriteIndex()];
-	for (auto& SyncGroupPair : SyncGroups)
-	{
-		SyncGroupPair.Value.Reset();
-	}
-
-	TArray<float>& StateWeights = StateWeightArrays[GetSyncGroupWriteIndex()];
+	TArray<float>& StateWeights = StateWeightArrays[GetBufferWriteIndex()];
 	FMemory::Memset(StateWeights.GetData(), 0, StateWeights.Num() * sizeof(float));
 
-	TArray<float>& MachineWeights = MachineWeightArrays[GetSyncGroupWriteIndex()];
+	TArray<float>& MachineWeights = MachineWeightArrays[GetBufferWriteIndex()];
 	FMemory::Memset(MachineWeights.GetData(), 0, MachineWeights.Num() * sizeof(float));
 
 #if WITH_EDITORONLY_DATA
@@ -499,6 +486,8 @@ void FAnimInstanceProxy::PostUpdate(UAnimInstance* InAnimInstance) const
 	if (FAnimBlueprintDebugData* DebugData = GetAnimBlueprintDebugData())
 	{
 		DebugData->RecordNodeVisitArray(UpdatedNodesThisFrame);
+		DebugData->RecordNodeAttributeMaps(NodeInputAttributesThisFrame, NodeOutputAttributesThisFrame);
+		DebugData->RecordNodeSyncsArray(NodeSyncsThisFrame);
 		DebugData->AnimNodePoseWatch = PoseWatchEntriesForThisFrame;
 	}
 #endif
@@ -609,44 +598,16 @@ void FAnimInstanceProxy::ClearObjects()
 
 FAnimTickRecord& FAnimInstanceProxy::CreateUninitializedTickRecord(FAnimGroupInstance*& OutSyncGroupPtr, FName GroupName)
 {
-	// Find or create the sync group if there is one
-	OutSyncGroupPtr = nullptr;
-	if (GroupName != NAME_None)
-	{
-		FSyncGroupMap& SyncGroupMap = SyncGroupMaps[GetSyncGroupWriteIndex()];
-		OutSyncGroupPtr = &SyncGroupMap.FindOrAdd(GroupName);
-	}
-
-	// Create the record
-	FAnimTickRecord* TickRecord = new ((OutSyncGroupPtr != nullptr) ? OutSyncGroupPtr->ActivePlayers : UngroupedActivePlayerArrays[GetSyncGroupWriteIndex()]) FAnimTickRecord();
-	return *TickRecord;
+PRAGMA_DISABLE_DEPRECATION_WARNINGS
+	return Sync.CreateUninitializedTickRecord(OutSyncGroupPtr, GroupName);
+PRAGMA_ENABLE_DEPRECATION_WARNINGS
 }
 
 FAnimTickRecord& FAnimInstanceProxy::CreateUninitializedTickRecordInScope(FAnimGroupInstance*& OutSyncGroupPtr, FName GroupName, EAnimSyncGroupScope Scope)
 {
-	if (GroupName != NAME_None)
-	{
-		// If we have no main proxy or it is "this", force us to local
-		if(MainInstanceProxy == nullptr || MainInstanceProxy == this)
-		{
-			Scope = EAnimSyncGroupScope::Local;
-		}
-
-		switch(Scope)
-		{
-		default:
-			ensureMsgf(false, TEXT("FAnimInstanceProxy::CreateUninitializedTickRecordInScope: Scope has invalid value %d"), Scope);
-			// Fall through
-
-		case EAnimSyncGroupScope::Local:
-			return CreateUninitializedTickRecord(OutSyncGroupPtr, GroupName);
-		case EAnimSyncGroupScope::Component:
-			// Forward to the main instance to sync with animations there in TickAssetPlayerInstances()
-			return MainInstanceProxy->CreateUninitializedTickRecord(OutSyncGroupPtr, GroupName);
-		}
-	}
-	
-	return CreateUninitializedTickRecord(OutSyncGroupPtr, GroupName);
+PRAGMA_DISABLE_DEPRECATION_WARNINGS
+	return Sync.CreateUninitializedTickRecordInScope(*this, OutSyncGroupPtr, GroupName, Scope);
+PRAGMA_ENABLE_DEPRECATION_WARNINGS
 }
 
 FAnimTickRecord& FAnimInstanceProxy::CreateUninitializedTickRecord(int32 GroupIndex, FAnimGroupInstance*& OutSyncGroupPtr)
@@ -657,7 +618,9 @@ FAnimTickRecord& FAnimInstanceProxy::CreateUninitializedTickRecord(int32 GroupIn
 		SyncGroupName = GetAnimClassInterface()->GetSyncGroupNames()[GroupIndex];
 	}
 
+PRAGMA_DISABLE_DEPRECATION_WARNINGS
 	return CreateUninitializedTickRecord(OutSyncGroupPtr, SyncGroupName);
+PRAGMA_ENABLE_DEPRECATION_WARNINGS
 }
 
 FAnimTickRecord& FAnimInstanceProxy::CreateUninitializedTickRecordInScope(int32 GroupIndex, EAnimSyncGroupScope Scope, FAnimGroupInstance*& OutSyncGroupPtr)
@@ -668,7 +631,9 @@ FAnimTickRecord& FAnimInstanceProxy::CreateUninitializedTickRecordInScope(int32
 		SyncGroupName = GetAnimClassInterface()->GetSyncGroupNames()[GroupIndex];
 	}
 
+PRAGMA_DISABLE_DEPRECATION_WARNINGS
 	return CreateUninitializedTickRecordInScope(OutSyncGroupPtr, SyncGroupName, Scope);
+PRAGMA_ENABLE_DEPRECATION_WARNINGS
 }
 
 void FAnimInstanceProxy::MakeSequenceTickRecord(FAnimTickRecord& TickRecord, class UAnimSequenceBase* Sequence, bool bLooping, float PlayRate, float FinalBlendWeight, float& CurrentTime, FMarkerTickRecord& MarkerTickRecord) const
@@ -695,7 +660,6 @@ void FAnimInstanceProxy::MakeBlendSpaceTickRecord(FAnimTickRecord& TickRecord, c
 	TickRecord.bLooping = bLooping;
 }
 
-/** Helper function: make a tick record for a pose asset*/
 void FAnimInstanceProxy::MakePoseAssetTickRecord(FAnimTickRecord& TickRecord, class UPoseAsset* PoseAsset, float FinalBlendWeight) const
 {
 	TickRecord.SourceAsset = PoseAsset;
@@ -705,7 +669,9 @@ void FAnimInstanceProxy::MakePoseAssetTickRecord(FAnimTickRecord& TickRecord, cl
 void FAnimInstanceProxy::SequenceAdvanceImmediate(UAnimSequenceBase* Sequence, bool bLooping, float PlayRate, float DeltaSeconds, float& CurrentTime, FMarkerTickRecord& MarkerTickRecord)
 {
 	FAnimTickRecord TickRecord;
+PRAGMA_DISABLE_DEPRECATION_WARNINGS
 	MakeSequenceTickRecord(TickRecord, Sequence, bLooping, PlayRate, /*FinalBlendWeight=*/ 1.0f, CurrentTime, MarkerTickRecord);
+PRAGMA_ENABLE_DEPRECATION_WARNINGS
 
 	FAnimAssetTickContext TickContext(DeltaSeconds, RootMotionMode, true);
 	TickRecord.SourceAsset->TickAssetPlayer(TickRecord, NotifyQueue, TickContext);
@@ -714,210 +680,25 @@ void FAnimInstanceProxy::SequenceAdvanceImmediate(UAnimSequenceBase* Sequence, b
 void FAnimInstanceProxy::BlendSpaceAdvanceImmediate(class UBlendSpaceBase* BlendSpace, const FVector& BlendInput, TArray<FBlendSampleData>& BlendSampleDataCache, FBlendFilter& BlendFilter, bool bLooping, float PlayRate, float DeltaSeconds, float& CurrentTime, FMarkerTickRecord& MarkerTickRecord)
 {
 	FAnimTickRecord TickRecord;
+PRAGMA_DISABLE_DEPRECATION_WARNINGS
 	MakeBlendSpaceTickRecord(TickRecord, BlendSpace, BlendInput, BlendSampleDataCache, BlendFilter, bLooping, PlayRate, /*FinalBlendWeight=*/ 1.0f, CurrentTime, MarkerTickRecord);
-	
+PRAGMA_ENABLE_DEPRECATION_WARNINGS
+
 	FAnimAssetTickContext TickContext(DeltaSeconds, RootMotionMode, true);
 	TickRecord.SourceAsset->TickAssetPlayer(TickRecord, NotifyQueue, TickContext);
 }
 
 void FAnimInstanceProxy::TickAssetPlayerInstances()
 {
+PRAGMA_DISABLE_DEPRECATION_WARNINGS
 	TickAssetPlayerInstances(CurrentDeltaSeconds);
+PRAGMA_ENABLE_DEPRECATION_WARNINGS
 }
 
-void FAnimInstanceProxy::TickAssetPlayerInstances(float DeltaSeconds)
+void FAnimInstanceProxy::TickAssetPlayerInstances(float InDeltaSeconds)
 {
-	DECLARE_SCOPE_HIERARCHICAL_COUNTER_FUNC()
-
-	SCOPE_CYCLE_COUNTER(STAT_TickAssetPlayerInstances);
-
-	// Handle all players inside sync groups
-	FSyncGroupMap& SyncGroupMap = SyncGroupMaps[GetSyncGroupWriteIndex()];
-	const FSyncGroupMap& PreviousSyncGroupMap = SyncGroupMaps[GetSyncGroupReadIndex()];
-	TArray<FAnimTickRecord>& UngroupedActivePlayers = UngroupedActivePlayerArrays[GetSyncGroupWriteIndex()];
-
-	for (auto& SyncGroupPair : SyncGroupMap)
-	{
-		FAnimGroupInstance& SyncGroup = SyncGroupPair.Value;
-	
-		if (SyncGroup.ActivePlayers.Num() > 0)
-		{
-			const FAnimGroupInstance* PreviousGroup = PreviousSyncGroupMap.Find(SyncGroupPair.Key);
-			SyncGroup.Prepare(PreviousGroup);
-
-			UE_LOG(LogAnimMarkerSync, Log, TEXT("Ticking Group [%s] GroupLeader [%d]"), *SyncGroupPair.Key.ToString(), SyncGroup.GroupLeaderIndex);
-
-			const bool bOnlyOneAnimationInGroup = SyncGroup.ActivePlayers.Num() == 1;
-
-			// Tick the group leader
-			FAnimAssetTickContext TickContext(DeltaSeconds, RootMotionMode, bOnlyOneAnimationInGroup, SyncGroup.ValidMarkers);
-			if (PreviousGroup)
-			{
-				const FMarkerSyncAnimPosition& EndPosition = PreviousGroup->MarkerTickContext.GetMarkerSyncEndPosition();
-				if ( EndPosition.IsValid() &&
-				     (EndPosition.PreviousMarkerName == NAME_None || SyncGroup.ValidMarkers.Contains(EndPosition.PreviousMarkerName)) &&
-					 (EndPosition.NextMarkerName == NAME_None || SyncGroup.ValidMarkers.Contains(EndPosition.NextMarkerName)))
-				{
-					TickContext.MarkerTickContext.SetMarkerSyncStartPosition(EndPosition);
-				}
-			}
-
-#if DO_CHECK
-			//For debugging UE-54705
-			FName InitialMarkerPrevious = TickContext.MarkerTickContext.GetMarkerSyncStartPosition().PreviousMarkerName;
-			FName InitialMarkerEnd = TickContext.MarkerTickContext.GetMarkerSyncStartPosition().NextMarkerName;
-			const bool bIsLeaderRecordValidPre = SyncGroup.ActivePlayers[0].MarkerTickRecord->IsValid(SyncGroup.ActivePlayers[0].bLooping);
-			FMarkerTickRecord LeaderPreMarkerTickRecord = *SyncGroup.ActivePlayers[0].MarkerTickRecord;
-#endif
-
-			// initialize to invalidate first
-			ensureMsgf(SyncGroup.GroupLeaderIndex == INDEX_NONE, TEXT("SyncGroup %s had a non -1 group leader index of %d in asset %s"), *SyncGroupPair.Key.ToString(), SyncGroup.GroupLeaderIndex, *GetNameSafe(SkeletalMeshComponent));
-			int32 GroupLeaderIndex = 0;
-			for (; GroupLeaderIndex < SyncGroup.ActivePlayers.Num(); ++GroupLeaderIndex)
-			{
-				FAnimTickRecord& GroupLeader = SyncGroup.ActivePlayers[GroupLeaderIndex];
-				// if it has leader score
-				SCOPE_CYCLE_COUNTER(STAT_TickAssetPlayerInstance);
-				FScopeCycleCounterUObject Scope(GroupLeader.SourceAsset);
-				TickContext.MarkerTickContext.MarkersPassedThisTick.Reset();
-				TickContext.RootMotionMovementParams.Clear();
-				GroupLeader.SourceAsset->TickAssetPlayer(GroupLeader, NotifyQueue, TickContext);
-
-				if (RootMotionMode == ERootMotionMode::RootMotionFromEverything && TickContext.RootMotionMovementParams.bHasRootMotion)
-				{
-					ExtractedRootMotion.AccumulateWithBlend(TickContext.RootMotionMovementParams, GroupLeader.GetRootMotionWeight());
-				}
-
-				// if we're not using marker based sync, we don't care, get out
-				if (TickContext.CanUseMarkerPosition() == false)
-				{
-					SyncGroup.GroupLeaderIndex = GroupLeaderIndex;
-					break;
-				}
-				// otherwise, the new position should contain the valid position for end, otherwise, we don't know where to sync to
-				else if (TickContext.MarkerTickContext.IsMarkerSyncEndValid())
-				{
-					// if this leader contains correct position, break
-					SyncGroup.MarkerTickContext = TickContext.MarkerTickContext;
-					SyncGroup.GroupLeaderIndex = GroupLeaderIndex;
-					UE_LOG(LogAnimMarkerSync, Log, TEXT("Previous Sync Group Marker Tick Context :\n%s"), *SyncGroup.MarkerTickContext.ToString());
-					UE_LOG(LogAnimMarkerSync, Log, TEXT("New Sync Group Marker Tick Context :\n%s"), *TickContext.MarkerTickContext.ToString());
-					break;
-				}
-				else
-				{
-					SyncGroup.GroupLeaderIndex = GroupLeaderIndex;
-					UE_LOG(LogAnimMarkerSync, Log, TEXT("Invalid position from Leader %d. Trying next leader"), GroupLeaderIndex);
-				}
-			} 
-
-			check(SyncGroup.GroupLeaderIndex != INDEX_NONE);
-			// we found leader
-			SyncGroup.Finalize(PreviousGroup);
-
-			if (TickContext.CanUseMarkerPosition())
-			{
-				const FMarkerSyncAnimPosition& MarkerStart = TickContext.MarkerTickContext.GetMarkerSyncStartPosition();
-				FName SyncGroupName = SyncGroupPair.Key;
-				FAnimTickRecord& GroupLeader = SyncGroup.ActivePlayers[SyncGroup.GroupLeaderIndex];
-				FString LeaderAnimName = GroupLeader.SourceAsset->GetName();
-
-				//  Updated logic in search for cause of UE-54705
-				const bool bStartMarkerValid = (MarkerStart.PreviousMarkerName == NAME_None) || SyncGroup.ValidMarkers.Contains(MarkerStart.PreviousMarkerName);
-				const bool bEndMarkerValid = (MarkerStart.NextMarkerName == NAME_None) || SyncGroup.ValidMarkers.Contains(MarkerStart.NextMarkerName);
-
-				if (!bStartMarkerValid)
-				{
-#if DO_CHECK
-					FString ErrorMsg = FString(TEXT("Prev Marker name not valid for sync group.\n"));
-					ErrorMsg += FString::Format(TEXT("\tMarker {0} : SyncGroupName {1} : Leader {2}\n"), { MarkerStart.PreviousMarkerName.ToString(), SyncGroupName.ToString(), LeaderAnimName });
-					ErrorMsg += FString::Format(TEXT("\tInitalPrev {0} : InitialNext {1} : GroupLeaderIndex {2}\n"), { InitialMarkerPrevious.ToString(), InitialMarkerEnd.ToString(), GroupLeaderIndex });
-					ErrorMsg += FString::Format(TEXT("\tLeader (0 index) was originally valid: {0} | Record: {1}\n"), { bIsLeaderRecordValidPre, LeaderPreMarkerTickRecord.ToString() });
-					ErrorMsg += FString::Format(TEXT("\t Valid Markers : {0}\n"), { SyncGroup.ValidMarkers.Num() });
-					for (int32 MarkerIndex = 0; MarkerIndex < SyncGroup.ValidMarkers.Num(); ++MarkerIndex)
-					{
-						ErrorMsg += FString::Format(TEXT("\t\t{0}) '{1}'\n"), {MarkerIndex, SyncGroup.ValidMarkers[MarkerIndex].ToString()});
-					}
-					ensureMsgf(false, TEXT("%s"), *ErrorMsg);
-#endif
-					TickContext.InvalidateMarkerSync();
-				}
-				else if (!bEndMarkerValid)
-				{
-#if DO_CHECK
-					FString ErrorMsg = FString(TEXT("Next Marker name not valid for sync group.\n"));
-					ErrorMsg += FString::Format(TEXT("\tMarker {0} : SyncGroupName {1} : Leader {2}\n"), { MarkerStart.NextMarkerName.ToString(), SyncGroupName.ToString(), LeaderAnimName });
-					ErrorMsg += FString::Format(TEXT("\tInitalPrev {0} : InitialNext {1} : GroupLeaderIndex {2}\n"), { InitialMarkerPrevious.ToString(), InitialMarkerEnd.ToString(), GroupLeaderIndex });
-					ErrorMsg += FString::Format(TEXT("\tLeader (0 index) was originally valid: {0} | Record: {1}\n"), { bIsLeaderRecordValidPre, LeaderPreMarkerTickRecord.ToString() });
-					ErrorMsg += FString::Format(TEXT("\t Valid Markers : {0}\n"), { SyncGroup.ValidMarkers.Num() });
-					for (int32 MarkerIndex = 0; MarkerIndex < SyncGroup.ValidMarkers.Num(); ++MarkerIndex)
-					{
-						ErrorMsg += FString::Format(TEXT("\t\t{0}) '{1}'\n"), { MarkerIndex, SyncGroup.ValidMarkers[MarkerIndex].ToString() });
-					}
-					ensureMsgf(false, TEXT("%s"), *ErrorMsg);
-#endif
-					TickContext.InvalidateMarkerSync();
-				}
-			}
-
-			// Update everything else to follow the leader, if there is more followers
-			if (SyncGroup.ActivePlayers.Num() > GroupLeaderIndex + 1)
-			{
-				// if we don't have a good leader, no reason to convert to follower
-				// tick as leader
-				TickContext.ConvertToFollower();
-	
-				for (int32 TickIndex = GroupLeaderIndex + 1; TickIndex < SyncGroup.ActivePlayers.Num(); ++TickIndex)
-				{
-					FAnimTickRecord& AssetPlayer = SyncGroup.ActivePlayers[TickIndex];
-					{
-						SCOPE_CYCLE_COUNTER(STAT_TickAssetPlayerInstance);
-						FScopeCycleCounterUObject Scope(AssetPlayer.SourceAsset);
-						TickContext.RootMotionMovementParams.Clear();
-						AssetPlayer.SourceAsset->TickAssetPlayer(AssetPlayer, NotifyQueue, TickContext);
-					}
-					if (RootMotionMode == ERootMotionMode::RootMotionFromEverything && TickContext.RootMotionMovementParams.bHasRootMotion)
-					{
-						ExtractedRootMotion.AccumulateWithBlend(TickContext.RootMotionMovementParams, AssetPlayer.GetRootMotionWeight());
-					}
-				}
-			}
-
-#if ANIM_TRACE_ENABLED
-			for(const FPassedMarker& PassedMarker : TickContext.MarkerTickContext.MarkersPassedThisTick)
-			{
-				TRACE_ANIM_SYNC_MARKER(CastChecked<UAnimInstance>(GetAnimInstanceObject()), PassedMarker);
-			}
-#endif
-		}
-	}
-
-	// Handle the remaining ungrouped animation players
-	for (int32 TickIndex = 0; TickIndex < UngroupedActivePlayers.Num(); ++TickIndex)
-	{
-		FAnimTickRecord& AssetPlayerToTick = UngroupedActivePlayers[TickIndex];
-		const TArray<FName>* UniqueNames = AssetPlayerToTick.SourceAsset->GetUniqueMarkerNames();
-		const TArray<FName>& ValidMarkers = UniqueNames ? *UniqueNames : FMarkerTickContext::DefaultMarkerNames;
-
-		const bool bOnlyOneAnimationInGroup = true;
-		FAnimAssetTickContext TickContext(DeltaSeconds, RootMotionMode, bOnlyOneAnimationInGroup, ValidMarkers);
-		{
-			SCOPE_CYCLE_COUNTER(STAT_TickAssetPlayerInstance);
-			FScopeCycleCounterUObject Scope(AssetPlayerToTick.SourceAsset);
-			AssetPlayerToTick.SourceAsset->TickAssetPlayer(AssetPlayerToTick, NotifyQueue, TickContext);
-		}
-		if (RootMotionMode == ERootMotionMode::RootMotionFromEverything && TickContext.RootMotionMovementParams.bHasRootMotion)
-		{
-			ExtractedRootMotion.AccumulateWithBlend(TickContext.RootMotionMovementParams, AssetPlayerToTick.GetRootMotionWeight());
-		}
-
-#if ANIM_TRACE_ENABLED
-		for(const FPassedMarker& PassedMarker : TickContext.MarkerTickContext.MarkersPassedThisTick)
-		{
-			TRACE_ANIM_SYNC_MARKER(CastChecked<UAnimInstance>(GetAnimInstanceObject()), PassedMarker);
-		}
-#endif
-	}
+	Sync.TickAssetPlayerInstances(*this, InDeltaSeconds);
+	Sync.TickSyncGroupWriteIndex();
 }
 
 void FAnimInstanceProxy::AddAnimNotifies(const TArray<FAnimNotifyEventReference>& NewNotifies, const float InstanceWeight)
@@ -936,27 +717,7 @@ int32 FAnimInstanceProxy::GetSyncGroupIndexFromName(FName SyncGroupName) const
 
 bool FAnimInstanceProxy::GetTimeToClosestMarker(FName SyncGroup, FName MarkerName, float& OutMarkerTime) const
 {
-	const FSyncGroupMap& SyncGroupMap = SyncGroupMaps[GetSyncGroupReadIndex()];
-
-	if (const FAnimGroupInstance* SyncGroupInstancePtr = SyncGroupMap.Find(SyncGroup))
-	{
-		if (SyncGroupInstancePtr->bCanUseMarkerSync && SyncGroupInstancePtr->ActivePlayers.IsValidIndex(SyncGroupInstancePtr->GroupLeaderIndex))
-		{
-			const FMarkerSyncAnimPosition& EndPosition = SyncGroupInstancePtr->MarkerTickContext.GetMarkerSyncEndPosition();
-			const FAnimTickRecord& Leader = SyncGroupInstancePtr->ActivePlayers[SyncGroupInstancePtr->GroupLeaderIndex];
-			if (EndPosition.PreviousMarkerName == MarkerName)
-			{
-				OutMarkerTime = Leader.MarkerTickRecord->PreviousMarker.TimeToMarker;
-				return true;
-			}
-			else if (EndPosition.NextMarkerName == MarkerName)
-			{
-				OutMarkerTime = Leader.MarkerTickRecord->NextMarker.TimeToMarker;
-				return true;
-			}
-		}
-	}
-	return false;
+	return Sync.GetTimeToClosestMarker(SyncGroup, MarkerName, OutMarkerTime);
 }
 
 void FAnimInstanceProxy::AddAnimNotifyFromGeneratedClass(int32 NotifyIndex)
@@ -976,64 +737,22 @@ void FAnimInstanceProxy::AddAnimNotifyFromGeneratedClass(int32 NotifyIndex)
 
 bool FAnimInstanceProxy::HasMarkerBeenHitThisFrame(FName SyncGroup, FName MarkerName) const
 {
-	const FSyncGroupMap& SyncGroupMap = SyncGroupMaps[GetSyncGroupReadIndex()];
-
-	if (const FAnimGroupInstance* SyncGroupInstancePtr = SyncGroupMap.Find(SyncGroup))
-	{
-		if (SyncGroupInstancePtr->bCanUseMarkerSync)
-		{
-			return SyncGroupInstancePtr->MarkerTickContext.MarkersPassedThisTick.ContainsByPredicate([&MarkerName](const FPassedMarker& PassedMarker) -> bool
-			{
-				return PassedMarker.PassedMarkerName == MarkerName;
-			});
-		}
-	}
-	return false;
+	return Sync.HasMarkerBeenHitThisFrame(SyncGroup, MarkerName);
 }
 
 bool FAnimInstanceProxy::IsSyncGroupBetweenMarkers(FName InSyncGroupName, FName PreviousMarker, FName NextMarker, bool bRespectMarkerOrder) const
 {
-	const FMarkerSyncAnimPosition& SyncGroupPosition = GetSyncGroupPosition(InSyncGroupName);
-	if ((SyncGroupPosition.PreviousMarkerName == PreviousMarker) && (SyncGroupPosition.NextMarkerName == NextMarker))
-	{
-		return true;
-	}
-
-	if (!bRespectMarkerOrder)
-	{
-		return ((SyncGroupPosition.PreviousMarkerName == NextMarker) && (SyncGroupPosition.NextMarkerName == PreviousMarker));
-	}
-
-	return false;
+	return Sync.IsSyncGroupBetweenMarkers(InSyncGroupName, PreviousMarker, NextMarker, bRespectMarkerOrder);
 }
 
 FMarkerSyncAnimPosition FAnimInstanceProxy::GetSyncGroupPosition(FName InSyncGroupName) const
 {
-	const FSyncGroupMap& SyncGroupMap = SyncGroupMaps[GetSyncGroupReadIndex()];
-
-	if (const FAnimGroupInstance* SyncGroupInstancePtr = SyncGroupMap.Find(InSyncGroupName))
-	{
-		if (SyncGroupInstancePtr->bCanUseMarkerSync && SyncGroupInstancePtr->MarkerTickContext.IsMarkerSyncEndValid())
-		{
-			return SyncGroupInstancePtr->MarkerTickContext.GetMarkerSyncEndPosition();
-		}
-	}
-
-	return FMarkerSyncAnimPosition();
+	return Sync.GetSyncGroupPosition(InSyncGroupName);
 }
 
 bool FAnimInstanceProxy::IsSyncGroupValid(FName InSyncGroupName) const
 {
-	const FSyncGroupMap& SyncGroupMap = SyncGroupMaps[GetSyncGroupReadIndex()];
-
-	if (const FAnimGroupInstance* SyncGroupInstancePtr = SyncGroupMap.Find(InSyncGroupName))
-	{
-		// If we don't use Markers, we're always valid.
-		return (!SyncGroupInstancePtr->bCanUseMarkerSync || SyncGroupInstancePtr->MarkerTickContext.IsMarkerSyncEndValid());
-	}
-
-	// If we're querying a sync group that doesn't exist, treat this as invalid
-	return false;
+	return Sync.IsSyncGroupValid(InSyncGroupName);
 }
 
 
@@ -1079,7 +798,7 @@ void FAnimInstanceProxy::UpdateSlotNodeWeight(const FName& SlotNodeName, float I
 	const int32* TrackerIndexPtr = SlotNameToTrackerIndex.Find(SlotNodeName);
 	if (TrackerIndexPtr)
 	{
-		FMontageActiveSlotTracker& Tracker = SlotWeightTracker[GetSyncGroupWriteIndex()][*TrackerIndexPtr];
+		FMontageActiveSlotTracker& Tracker = SlotWeightTracker[GetBufferWriteIndex()][*TrackerIndexPtr];
 		Tracker.MontageLocalWeight = InMontageLocalWeight;
 		Tracker.NodeGlobalWeight = InNodeGlobalWeight;
 
@@ -1102,8 +821,8 @@ bool FAnimInstanceProxy::GetSlotInertializationRequest(const FName& SlotName, fl
 
 void FAnimInstanceProxy::ClearSlotNodeWeights()
 {
-	TArray<FMontageActiveSlotTracker>& SlotWeightTracker_Read = SlotWeightTracker[GetSyncGroupReadIndex()];
-	TArray<FMontageActiveSlotTracker>& SlotWeightTracker_Write = SlotWeightTracker[GetSyncGroupWriteIndex()];
+	TArray<FMontageActiveSlotTracker>& SlotWeightTracker_Read = SlotWeightTracker[GetBufferReadIndex()];
+	TArray<FMontageActiveSlotTracker>& SlotWeightTracker_Write = SlotWeightTracker[GetBufferWriteIndex()];
 
 	for (int32 TrackerIndex = 0; TrackerIndex < SlotWeightTracker_Write.Num(); TrackerIndex++)
 	{
@@ -1117,7 +836,7 @@ bool FAnimInstanceProxy::IsSlotNodeRelevantForNotifies(const FName& SlotNodeName
 	const int32* TrackerIndexPtr = SlotNameToTrackerIndex.Find(SlotNodeName);
 	if (TrackerIndexPtr)
 	{
-		const FMontageActiveSlotTracker& Tracker = SlotWeightTracker[GetSyncGroupReadIndex()][*TrackerIndexPtr];
+		const FMontageActiveSlotTracker& Tracker = SlotWeightTracker[GetBufferReadIndex()][*TrackerIndexPtr];
 		return (Tracker.bIsRelevantThisTick || Tracker.bWasRelevantOnPreviousTick);
 	}
 
@@ -1129,7 +848,7 @@ float FAnimInstanceProxy::GetSlotNodeGlobalWeight(const FName& SlotNodeName) con
 	const int32* TrackerIndexPtr = SlotNameToTrackerIndex.Find(SlotNodeName);
 	if (TrackerIndexPtr)
 	{
-		const FMontageActiveSlotTracker& Tracker = SlotWeightTracker[GetSyncGroupReadIndex()][*TrackerIndexPtr];
+		const FMontageActiveSlotTracker& Tracker = SlotWeightTracker[GetBufferReadIndex()][*TrackerIndexPtr];
 		return Tracker.NodeGlobalWeight;
 	}
 
@@ -1141,7 +860,7 @@ float FAnimInstanceProxy::GetSlotMontageGlobalWeight(const FName& SlotNodeName)
 	const int32* TrackerIndexPtr = SlotNameToTrackerIndex.Find(SlotNodeName);
 	if (TrackerIndexPtr)
 	{
-		const FMontageActiveSlotTracker& Tracker = SlotWeightTracker[GetSyncGroupReadIndex()][*TrackerIndexPtr];
+		const FMontageActiveSlotTracker& Tracker = SlotWeightTracker[GetBufferReadIndex()][*TrackerIndexPtr];
 		return Tracker.MontageLocalWeight * Tracker.NodeGlobalWeight;
 	}
 
@@ -1153,7 +872,7 @@ float FAnimInstanceProxy::GetSlotMontageLocalWeight(const FName& SlotNodeName) c
 	const int32* TrackerIndexPtr = SlotNameToTrackerIndex.Find(SlotNodeName);
 	if (TrackerIndexPtr)
 	{
-		const FMontageActiveSlotTracker& Tracker = SlotWeightTracker[GetSyncGroupReadIndex()][*TrackerIndexPtr];
+		const FMontageActiveSlotTracker& Tracker = SlotWeightTracker[GetBufferReadIndex()][*TrackerIndexPtr];
 		return Tracker.MontageLocalWeight;
 	}
 
@@ -1292,6 +1011,9 @@ void FAnimInstanceProxy::UpdateAnimation()
 
 #if WITH_EDITORONLY_DATA
 	UpdatedNodesThisFrame.Reset();
+	NodeInputAttributesThisFrame.Reset();
+	NodeOutputAttributesThisFrame.Reset();
+	NodeSyncsThisFrame.Reset();
 #endif
 
 	FAnimationUpdateSharedContext SharedContext;
@@ -1303,6 +1025,9 @@ void FAnimInstanceProxy::UpdateAnimation()
 	}
 
 	UpdateAnimation_WithRoot(Context, RootNode, NAME_AnimGraph);
+
+	// Tick syncing
+	Sync.TickAssetPlayerInstances(*this, CurrentDeltaSeconds);
 }
 
 void FAnimInstanceProxy::UpdateAnimation_WithRoot(const FAnimationUpdateContext& InContext, FAnimNode_Base* InRootNode, FName InLayerName)
@@ -1366,7 +1091,11 @@ void FAnimInstanceProxy::UpdateAnimation_WithRoot(const FAnimationUpdateContext&
 		// We re-enter this function when we call layer graphs linked to the main graph. In these cases we
 		// dont want to perform duplicate work
 		TGuardValue<bool> ScopeGuard(bUpdatingRoot, true);
-	
+
+		// Anything syncing within this scope is subject to sync groups.
+		// We only enable syncing here for the main proxy
+		UE::Anim::TOptionalScopedGraphMessage<UE::Anim::FAnimSyncGroupScope> Message(GetMainInstanceProxy() == this, InContext, InContext);
+
 		// update all nodes
 		if(InRootNode == RootNode)
 		{
@@ -1485,11 +1214,11 @@ void FAnimInstanceProxy::EvaluateAnimationNode_WithRoot(FPoseContext& Output, FA
 		if(InRootNode == RootNode)
 		{
 			EvaluationCounter.Increment();
-		}
 
-		if(AnimClassInterface && AnimClassInterface->GetAnimBlueprintFunctions().Num() > 0)
-		{
-			Output.SetNodeId(AnimClassInterface->GetAnimBlueprintFunctions()[0].OutputPoseNodeIndex);
+			if(AnimClassInterface && AnimClassInterface->GetAnimBlueprintFunctions().Num() > 0)
+			{
+				Output.SetNodeId(AnimClassInterface->GetAnimBlueprintFunctions()[0].OutputPoseNodeIndex);
+			}
 		}
 
 		TRACE_SCOPED_ANIM_NODE(Output);
@@ -2858,12 +2587,12 @@ int32 FAnimInstanceProxy::GetInstanceAssetPlayerIndex(FName MachineName, FName S
 
 float FAnimInstanceProxy::GetRecordedMachineWeight(const int32 InMachineClassIndex) const
 {
-	return MachineWeightArrays[GetSyncGroupReadIndex()][InMachineClassIndex];
+	return MachineWeightArrays[GetBufferReadIndex()][InMachineClassIndex];
 }
 
 void FAnimInstanceProxy::RecordMachineWeight(const int32 InMachineClassIndex, const float InMachineWeight)
 {
-	MachineWeightArrays[GetSyncGroupWriteIndex()][InMachineClassIndex] = InMachineWeight;
+	MachineWeightArrays[GetBufferWriteIndex()][InMachineClassIndex] = InMachineWeight;
 }
 
 float FAnimInstanceProxy::GetRecordedStateWeight(const int32 InMachineClassIndex, const int32 InStateIndex) const
@@ -2873,7 +2602,7 @@ float FAnimInstanceProxy::GetRecordedStateWeight(const int32 InMachineClassIndex
 	if(BaseIndexPtr)
 	{
 		const int32 StateIndex = *BaseIndexPtr + InStateIndex;
-		return StateWeightArrays[GetSyncGroupReadIndex()][StateIndex];
+		return StateWeightArrays[GetBufferReadIndex()][StateIndex];
 	}
 
 	return 0.0f;
@@ -2886,7 +2615,7 @@ void FAnimInstanceProxy::RecordStateWeight(const int32 InMachineClassIndex, cons
 	if(BaseIndexPtr)
 	{
 		const int32 StateIndex = *BaseIndexPtr + InStateIndex;
-		StateWeightArrays[GetSyncGroupWriteIndex()][StateIndex] = InStateWeight;
+		StateWeightArrays[GetBufferWriteIndex()][StateIndex] = InStateWeight;
 	}
 
 #if WITH_EDITORONLY_DATA
@@ -2961,6 +2690,18 @@ TArray<FAnimNode_AssetPlayerBase*> FAnimInstanceProxy::GetInstanceAssetPlayers(c
 }
 
 #if WITH_EDITOR
+void FAnimInstanceProxy::RecordNodeAttribute(const FAnimInstanceProxy& InSourceProxy, int32 InTargetNodeIndex, int32 InSourceNodeIndex, FName InAttribute)
+{
+	TArray<FAnimBlueprintDebugData::FAttributeRecord>& InputAttributeRecords = NodeInputAttributesThisFrame.FindOrAdd(InTargetNodeIndex);
+	InputAttributeRecords.Emplace(InSourceNodeIndex, InAttribute);
+
+	if(&InSourceProxy == this)
+	{
+		TArray<FAnimBlueprintDebugData::FAttributeRecord>& OutputAttributeRecords = NodeOutputAttributesThisFrame.FindOrAdd(InSourceNodeIndex);
+		OutputAttributeRecords.Emplace(InTargetNodeIndex, InAttribute);
+	}
+}
+
 void FAnimInstanceProxy::RegisterWatchedPose(const FCompactPose& Pose, int32 LinkID)
 {
 	if(bIsBeingDebugged)
diff --git a/Engine/Source/Runtime/Engine/Private/Animation/AnimNodeBase.cpp b/Engine/Source/Runtime/Engine/Private/Animation/AnimNodeBase.cpp
--- a/Engine/Source/Runtime/Engine/Private/Animation/AnimNodeBase.cpp
+++ b/Engine/Source/Runtime/Engine/Private/Animation/AnimNodeBase.cpp
@@ -6,6 +6,7 @@
 #include "Animation/AnimTrace.h"
 #include "UObject/CoreObjectVersion.h"
 #include "PropertyAccess.h"
+#include "Animation/AnimAttributes.h"
 
 /////////////////////////////////////////////////////
 // FAnimationBaseContext
@@ -342,11 +343,28 @@ void FPoseLink::Evaluate(FPoseContext& Output)
 		CurrentPose.ResetToAdditiveIdentity();
 #endif
 
+		int32 SourceID = Output.GetCurrentNodeId();
+
 		{
 			Output.SetNodeId(LinkID);
 			TRACE_SCOPED_ANIM_NODE(Output);
 			LinkedNode->Evaluate_AnyThread(Output);
+			TRACE_ANIM_NODE_BLENDABLE_ATTRIBUTES(Output, SourceID, LinkID);
+		}
+
+#if WITH_EDITORONLY_DATA
+		if (Output.AnimInstanceProxy->IsBeingDebugged())
+		{
+			if(Output.CustomAttributes.ContainsData())
+			{
+				Output.AnimInstanceProxy->RecordNodeAttribute(*Output.AnimInstanceProxy, SourceID, LinkID, UE::Anim::FAttributes::Attributes);
+			}
+			if(Output.Curve.NumValid() > 0)
+			{
+				Output.AnimInstanceProxy->RecordNodeAttribute(*Output.AnimInstanceProxy, SourceID, LinkID, UE::Anim::FAttributes::Curves);
+			}
 		}
+#endif
 
 #if ENABLE_ANIMNODE_POSE_DEBUG
 		CurrentPose.CopyBonesFrom(Output.Pose);
@@ -426,11 +444,28 @@ void FComponentSpacePoseLink::EvaluateComponentSpace(FComponentSpacePoseContext&
 
 	if (LinkedNode != NULL)
 	{
+		int32 SourceID = Output.GetCurrentNodeId();
+
 		{
 			Output.SetNodeId(LinkID);
 			TRACE_SCOPED_ANIM_NODE(Output);
 			LinkedNode->EvaluateComponentSpace_AnyThread(Output);
+			TRACE_ANIM_NODE_BLENDABLE_ATTRIBUTES(Output, SourceID, LinkID);
+		}
+
+#if WITH_EDITORONLY_DATA
+		if (Output.AnimInstanceProxy->IsBeingDebugged())
+		{
+			if(Output.CustomAttributes.ContainsData())
+			{
+				Output.AnimInstanceProxy->RecordNodeAttribute(*Output.AnimInstanceProxy, SourceID, LinkID, UE::Anim::FAttributes::Attributes);
+			}
+			if(Output.Curve.NumValid() > 0)
+			{
+				Output.AnimInstanceProxy->RecordNodeAttribute(*Output.AnimInstanceProxy, SourceID, LinkID, UE::Anim::FAttributes::Curves);
+			}
 		}
+#endif
 
 #if WITH_EDITOR
 		Output.AnimInstanceProxy->RegisterWatchedPose(Output.Pose, LinkID);
diff --git a/Engine/Source/Runtime/Engine/Private/Animation/AnimNode_AssetPlayerBase.cpp b/Engine/Source/Runtime/Engine/Private/Animation/AnimNode_AssetPlayerBase.cpp
--- a/Engine/Source/Runtime/Engine/Private/Animation/AnimNode_AssetPlayerBase.cpp
+++ b/Engine/Source/Runtime/Engine/Private/Animation/AnimNode_AssetPlayerBase.cpp
@@ -2,14 +2,15 @@
 
 #include "Animation/AnimNode_AssetPlayerBase.h"
 #include "Animation/AnimInstanceProxy.h"
+#include "Animation/AnimSyncScope.h"
 
 FAnimNode_AssetPlayerBase::FAnimNode_AssetPlayerBase()
 	: GroupName(NAME_None)
 #if WITH_EDITORONLY_DATA
 	, GroupIndex_DEPRECATED(INDEX_NONE)
+	, GroupScope_DEPRECATED(EAnimSyncGroupScope::Local)
 #endif
 	, GroupRole(EAnimGroupRole::CanBeLeader)
-	, GroupScope(EAnimSyncGroupScope::Local)
 	, bIgnoreForRelevancyTest(false)
 	, bHasBeenFullWeight(false)
 	, BlendWeight(0.0f)
@@ -37,22 +38,23 @@ void FAnimNode_AssetPlayerBase::Update_AnyThread(const FAnimationUpdateContext&
 
 void FAnimNode_AssetPlayerBase::CreateTickRecordForNode(const FAnimationUpdateContext& Context, UAnimSequenceBase* Sequence, bool bLooping, float PlayRate)
 {
-	// Create a tick record and fill it out
+	// Create a tick record and push into the closest scope
 	const float FinalBlendWeight = Context.GetFinalBlendWeight();
 
-	FAnimGroupInstance* SyncGroup;
-	const FName GroupNameToUse = ((GroupRole < EAnimGroupRole::TransitionLeader) || bHasBeenFullWeight) ? GroupName : NAME_None;
+	UE::Anim::FAnimSyncGroupScope& SyncScope = Context.GetMessageChecked<UE::Anim::FAnimSyncGroupScope>();
 
-	FAnimTickRecord& TickRecord = Context.AnimInstanceProxy->CreateUninitializedTickRecordInScope(/*out*/ SyncGroup, GroupNameToUse, GroupScope);
+	FName GroupNameToUse = ((GroupRole < EAnimGroupRole::TransitionLeader) || bHasBeenFullWeight) ? GroupName : NAME_None;
+	EAnimSyncMethod MethodToUse = Method;
+	if(GroupNameToUse == NAME_None)
+	{
+		MethodToUse = EAnimSyncMethod::DoNotSync;
+	}
 
-	Context.AnimInstanceProxy->MakeSequenceTickRecord(TickRecord, Sequence, bLooping, PlayRate, FinalBlendWeight, /*inout*/ InternalTimeAccumulator, MarkerTickRecord);
+	UE::Anim::FAnimSyncParams SyncParams(GroupNameToUse, GroupRole, MethodToUse);
+	FAnimTickRecord TickRecord(Sequence, bLooping, PlayRate, FinalBlendWeight, /*inout*/ InternalTimeAccumulator, MarkerTickRecord);
 	TickRecord.RootMotionWeightModifier = Context.GetRootMotionWeightModifier();
 
-	// Update the sync group if it exists
-	if (SyncGroup != NULL)
-	{
-		SyncGroup->TestTickRecordForLeadership(GroupRole);
-	}
+	SyncScope.AddTickRecord(TickRecord, SyncParams, UE::Anim::FAnimSyncDebugInfo(Context));
 
 	TRACE_ANIM_TICK_RECORD(Context, TickRecord);
 }
diff --git a/Engine/Source/Runtime/Engine/Private/Animation/AnimNode_Inertialization.cpp b/Engine/Source/Runtime/Engine/Private/Animation/AnimNode_Inertialization.cpp
--- a/Engine/Source/Runtime/Engine/Private/Animation/AnimNode_Inertialization.cpp
+++ b/Engine/Source/Runtime/Engine/Private/Animation/AnimNode_Inertialization.cpp
@@ -10,6 +10,8 @@
 
 IMPLEMENT_ANIMGRAPH_MESSAGE(UE::Anim::IInertializationRequester);
 
+const FName UE::Anim::IInertializationRequester::Attribute("InertialBlending");
+
 TAutoConsoleVariable<int32> CVarAnimInertializationEnable(TEXT("a.AnimNode.Inertialization.Enable"), 1, TEXT("Enable / Disable Inertialization"));
 TAutoConsoleVariable<int32> CVarAnimInertializationIgnoreVelocity(TEXT("a.AnimNode.Inertialization.IgnoreVelocity"), 0, TEXT("Ignore velocity information during Inertialization (effectively reverting to a quintic diff blend)"));
 TAutoConsoleVariable<int32> CVarAnimInertializationIgnoreDeficit(TEXT("a.AnimNode.Inertialization.IgnoreDeficit"), 0, TEXT("Ignore inertialization time deficit caused by interruptions"));
@@ -24,16 +26,35 @@ namespace UE { namespace Anim {
 class FInertializationRequester : public IInertializationRequester
 {
 public:
-	FInertializationRequester(FAnimNode_Inertialization* InNode)
+	FInertializationRequester(const FAnimationBaseContext& InContext, FAnimNode_Inertialization* InNode)
 		: Node(*InNode)
+		, NodeId(InContext.GetCurrentNodeId())
+		, Proxy(*InContext.AnimInstanceProxy)
 	{}
 
 private:
 	// IInertializationRequester interface
-	virtual void RequestInertialization(float InRequestedDuration) { Node.RequestInertialization(InRequestedDuration); }
+	virtual void RequestInertialization(float InRequestedDuration) override
+	{ 
+		Node.RequestInertialization(InRequestedDuration); 
+	}
+
+	virtual void AddDebugRecord(const FAnimInstanceProxy& InSourceProxy, int32 InSourceNodeId)
+	{
+#if WITH_EDITORONLY_DATA
+		Proxy.RecordNodeAttribute(InSourceProxy, NodeId, InSourceNodeId, IInertializationRequester::Attribute);
+#endif
+		TRACE_ANIM_NODE_ATTRIBUTE(Proxy, InSourceProxy, NodeId, InSourceNodeId, IInertializationRequester::Attribute);
+	}
 
 	// Node to target
 	FAnimNode_Inertialization& Node;
+
+	// Node index
+	int32 NodeId;
+
+	// Proxy currently executing
+	FAnimInstanceProxy& Proxy;
 };
 
 }}	// namespace UE::Anim
@@ -103,11 +124,14 @@ void FAnimNode_Inertialization::Update_AnyThread(const FAnimationUpdateContext&
 {
 	DECLARE_SCOPE_HIERARCHICAL_COUNTER_ANIMNODE(Update_AnyThread);
 
+	const int32 NodeId = Context.GetCurrentNodeId();
+	const FAnimInstanceProxy& Proxy = *Context.AnimInstanceProxy;
+
 	// Allow nodes further towards the leaves to inertialize using this node
-	UE::Anim::TScopedGraphMessage<UE::Anim::FInertializationRequester> Inertialization(Context, this);
+	UE::Anim::TScopedGraphMessage<UE::Anim::FInertializationRequester> Inertialization(Context, Context, this);
 
 	// Handle skipped updates for cached poses by forwarding to inertialization nodes in those residual stacks
-	UE::Anim::TScopedGraphMessage<UE::Anim::FCachedPoseSkippedUpdateHandler> CachedPoseSkippedUpdate(Context, [this](TArrayView<const UE::Anim::FMessageStack> InSkippedUpdates)
+	UE::Anim::TScopedGraphMessage<UE::Anim::FCachedPoseSkippedUpdateHandler> CachedPoseSkippedUpdate(Context, [this, NodeId, &Proxy](TArrayView<const UE::Anim::FMessageStack> InSkippedUpdates)
 	{
 		// If we have a pending request forward the request to other Inertialization nodes
 		// that were skipped due to pose caching.
@@ -118,9 +142,10 @@ void FAnimNode_Inertialization::Update_AnyThread(const FAnimationUpdateContext&
 			// So here we forward 'this' node's requests to the ancestors of those skipped UseCachedPose nodes.
 			for (const UE::Anim::FMessageStack& Stack : InSkippedUpdates)
 			{
-				Stack.ForEachMessage<UE::Anim::IInertializationRequester>([this](UE::Anim::IInertializationRequester& InMessage)
+				Stack.ForEachMessage<UE::Anim::IInertializationRequester>([this, NodeId, &Proxy](UE::Anim::IInertializationRequester& InMessage)
 				{
 					InMessage.RequestInertialization(RequestedDuration);
+ 					InMessage.AddDebugRecord(Proxy, NodeId);
 
 					return UE::Anim::FMessageStack::EEnumerate::Stop;
 				});
diff --git a/Engine/Source/Runtime/Engine/Private/Animation/AnimNode_LinkedAnimGraph.cpp b/Engine/Source/Runtime/Engine/Private/Animation/AnimNode_LinkedAnimGraph.cpp
--- a/Engine/Source/Runtime/Engine/Private/Animation/AnimNode_LinkedAnimGraph.cpp
+++ b/Engine/Source/Runtime/Engine/Private/Animation/AnimNode_LinkedAnimGraph.cpp
@@ -80,6 +80,7 @@ void FAnimNode_LinkedAnimGraph::CacheBonesSubGraph_AnyThread(const FAnimationCac
 		// Note not calling Proxy.CacheBones_WithRoot here as it is guarded by
 		// bBoneCachesInvalidated, which is handled at a higher level
 		FAnimationCacheBonesContext LinkedContext(&Proxy);
+		LinkedContext = LinkedContext.WithNodeId(CachedLinkedNodeIndex);
 		LinkedRoot->CacheBones_AnyThread(LinkedContext);
 	}
 }
@@ -111,9 +112,7 @@ void FAnimNode_LinkedAnimGraph::Update_AnyThread(const FAnimationUpdateContext&
 		// in USkeletalMeshComponent::TickAnimation. It used to be the case that we could do non-parallel work in 
 		// USkeletalMeshComponent::TickAnimation, which would mean we would have to skip doing that work here.
 		FAnimationUpdateContext NewContext = InContext.WithOtherProxy(&Proxy);
-
 		NewContext = NewContext.WithNodeId(CachedLinkedNodeIndex);
-
 		Proxy.UpdateAnimation_WithRoot(NewContext, LinkedRoot, GetDynamicLinkFunctionName());
 	}
 	else if(InputPoses.Num() > 0)
@@ -130,6 +129,7 @@ void FAnimNode_LinkedAnimGraph::Update_AnyThread(const FAnimationUpdateContext&
 		if(InertializationRequester)
 		{
 			InertializationRequester->RequestInertialization(PendingBlendDuration);
+			InertializationRequester->AddDebugRecord(*InContext.AnimInstanceProxy, InContext.GetCurrentNodeId());
 		}
 		else if ((PendingBlendDuration != 0.0f) && (InputPoses.Num() > 0))
 		{
@@ -155,7 +155,8 @@ void FAnimNode_LinkedAnimGraph::Evaluate_AnyThread(FPoseContext& Output)
 		// Create an evaluation context
 		FPoseContext EvaluationContext(&Proxy, Output.ExpectsAdditivePose());
 		EvaluationContext.ResetToRefPose();
-			
+		EvaluationContext.SetNodeId(CachedLinkedNodeIndex);
+
 		// Run the anim blueprint
 		Proxy.EvaluateAnimation_WithRoot(EvaluationContext, LinkedRoot);
 
diff --git a/Engine/Source/Runtime/Engine/Private/Animation/AnimNode_LinkedInputPose.cpp b/Engine/Source/Runtime/Engine/Private/Animation/AnimNode_LinkedInputPose.cpp
--- a/Engine/Source/Runtime/Engine/Private/Animation/AnimNode_LinkedInputPose.cpp
+++ b/Engine/Source/Runtime/Engine/Private/Animation/AnimNode_LinkedInputPose.cpp
@@ -29,9 +29,7 @@ void FAnimNode_LinkedInputPose::Update_AnyThread(const FAnimationUpdateContext&
 	if(InputProxy)
 	{
 		FAnimationUpdateContext InputContext = Context.WithOtherProxy(InputProxy);
-
 		InputContext = InputContext.WithNodeId(OuterGraphNodeIndex);
-
 		InputPose.Update(InputContext);
 	}
 }
@@ -43,6 +41,7 @@ void FAnimNode_LinkedInputPose::Evaluate_AnyThread(FPoseContext& Output)
 		Output.Pose.SetBoneContainer(&InputProxy->GetRequiredBones());
 
 		FPoseContext InputContext(InputProxy, Output.ExpectsAdditivePose());
+		InputContext.SetNodeId(OuterGraphNodeIndex);
 		InputPose.Evaluate(InputContext);
 
 		Output.Pose.MoveBonesFrom(InputContext.Pose);
diff --git a/Engine/Source/Runtime/Engine/Private/Animation/AnimNode_StateMachine.cpp b/Engine/Source/Runtime/Engine/Private/Animation/AnimNode_StateMachine.cpp
--- a/Engine/Source/Runtime/Engine/Private/Animation/AnimNode_StateMachine.cpp
+++ b/Engine/Source/Runtime/Engine/Private/Animation/AnimNode_StateMachine.cpp
@@ -458,6 +458,7 @@ void FAnimNode_StateMachine::Update_AnyThread(const FAnimationUpdateContext& Con
 					if (InertializationRequester)
 					{
 						InertializationRequester->RequestInertialization(ReferenceTransition.CrossfadeDuration);
+						InertializationRequester->AddDebugRecord(*Context.AnimInstanceProxy, Context.GetCurrentNodeId());
 					}
 					else
 					{
diff --git a/Engine/Source/Runtime/Engine/Private/Animation/AnimSingleNodeInstanceProxy.cpp b/Engine/Source/Runtime/Engine/Private/Animation/AnimSingleNodeInstanceProxy.cpp
--- a/Engine/Source/Runtime/Engine/Private/Animation/AnimSingleNodeInstanceProxy.cpp
+++ b/Engine/Source/Runtime/Engine/Private/Animation/AnimSingleNodeInstanceProxy.cpp
@@ -523,11 +523,10 @@ void FAnimNode_SingleNode::Update_AnyThread(const FAnimationUpdateContext& Conte
 
 	if(Proxy->CurrentAsset != NULL)
 	{
-		FAnimGroupInstance* SyncGroup;
 		if (UBlendSpaceBase* BlendSpace = Cast<UBlendSpaceBase>(Proxy->CurrentAsset))
 		{
-			FAnimTickRecord& TickRecord = Proxy->CreateUninitializedTickRecord(/*out*/ SyncGroup, NAME_None);
-			Proxy->MakeBlendSpaceTickRecord(TickRecord, BlendSpace, Proxy->BlendSpaceInput, Proxy->BlendSampleData, Proxy->BlendFilter, Proxy->bLooping, NewPlayRate, 1.f, /*inout*/ Proxy->CurrentTime, Proxy->MarkerTickRecord);
+			FAnimTickRecord TickRecord(BlendSpace, Proxy->BlendSpaceInput, Proxy->BlendSampleData, Proxy->BlendFilter, Proxy->bLooping, NewPlayRate, 1.f, /*inout*/ Proxy->CurrentTime, Proxy->MarkerTickRecord);
+			Proxy->AddTickRecord(TickRecord);
 
 			TRACE_ANIM_TICK_RECORD(Context, TickRecord);
 #if WITH_EDITORONLY_DATA
@@ -536,8 +535,8 @@ void FAnimNode_SingleNode::Update_AnyThread(const FAnimationUpdateContext& Conte
 		}
 		else if (UAnimSequence* Sequence = Cast<UAnimSequence>(Proxy->CurrentAsset))
 		{
-			FAnimTickRecord& TickRecord = Proxy->CreateUninitializedTickRecord(/*out*/ SyncGroup, NAME_None);
-			Proxy->MakeSequenceTickRecord(TickRecord, Sequence, Proxy->bLooping, NewPlayRate, 1.f, /*inout*/ Proxy->CurrentTime, Proxy->MarkerTickRecord);
+			FAnimTickRecord TickRecord(Sequence, Proxy->bLooping, NewPlayRate, 1.f, /*inout*/ Proxy->CurrentTime, Proxy->MarkerTickRecord);
+			Proxy->AddTickRecord(TickRecord);
 
 			TRACE_ANIM_TICK_RECORD(Context, TickRecord);
 
@@ -553,8 +552,8 @@ void FAnimNode_SingleNode::Update_AnyThread(const FAnimationUpdateContext& Conte
 		}
 		else if (UAnimStreamable* Streamable = Cast<UAnimStreamable>(Proxy->CurrentAsset))
 		{
-			FAnimTickRecord& TickRecord = Proxy->CreateUninitializedTickRecord(/*out*/ SyncGroup, NAME_None);
-			Proxy->MakeSequenceTickRecord(TickRecord, Streamable, Proxy->bLooping, NewPlayRate, 1.f, /*inout*/ Proxy->CurrentTime, Proxy->MarkerTickRecord);
+			FAnimTickRecord TickRecord(Streamable, Proxy->bLooping, NewPlayRate, 1.f, /*inout*/ Proxy->CurrentTime, Proxy->MarkerTickRecord);
+			Proxy->AddTickRecord(TickRecord);
 
 			TRACE_ANIM_TICK_RECORD(Context, TickRecord);
 
@@ -570,8 +569,8 @@ void FAnimNode_SingleNode::Update_AnyThread(const FAnimationUpdateContext& Conte
 		}
 		else if(UAnimComposite* Composite = Cast<UAnimComposite>(Proxy->CurrentAsset))
 		{
-			FAnimTickRecord& TickRecord = Proxy->CreateUninitializedTickRecord(/*out*/ SyncGroup, NAME_None);
-			Proxy->MakeSequenceTickRecord(TickRecord, Composite, Proxy->bLooping, NewPlayRate, 1.f, /*inout*/ Proxy->CurrentTime, Proxy->MarkerTickRecord);
+			FAnimTickRecord TickRecord(Composite, Proxy->bLooping, NewPlayRate, 1.f, /*inout*/ Proxy->CurrentTime, Proxy->MarkerTickRecord);
+			Proxy->AddTickRecord(TickRecord);
 
 			TRACE_ANIM_TICK_RECORD(Context, TickRecord);
 
@@ -609,8 +608,8 @@ void FAnimNode_SingleNode::Update_AnyThread(const FAnimationUpdateContext& Conte
 		}
 		else if (UPoseAsset* PoseAsset = Cast<UPoseAsset>(Proxy->CurrentAsset))
 		{
-			FAnimTickRecord& TickRecord = Proxy->CreateUninitializedTickRecord(/*out*/ SyncGroup, NAME_None);
-			Proxy->MakePoseAssetTickRecord(TickRecord, PoseAsset, 1.f);
+			FAnimTickRecord TickRecord(PoseAsset, 1.f);
+			Proxy->AddTickRecord(TickRecord);
 
 			TRACE_ANIM_TICK_RECORD(Context, TickRecord);
 		}
diff --git a/Engine/Source/Runtime/Engine/Private/Animation/AnimSync.cpp b/Engine/Source/Runtime/Engine/Private/Animation/AnimSync.cpp
new file mode 100644
--- /dev/null
+++ b/Engine/Source/Runtime/Engine/Private/Animation/AnimSync.cpp
@@ -0,0 +1,391 @@
+// Copyright Epic Games, Inc. All Rights Reserved.
+
+#include "Animation/AnimSync.h"
+#include "Animation/AnimNodeBase.h"
+#include "Animation/AnimInstanceProxy.h"
+#include "AnimationUtils.h"
+
+namespace UE { namespace Anim {
+
+const FName FAnimSync::Attribute("Sync");
+
+void FAnimSync::Reset()
+{
+	TArray<FAnimTickRecord>& UngroupedActivePlayers = UngroupedActivePlayerArrays[GetSyncGroupWriteIndex()];
+	UngroupedActivePlayers.Reset();
+
+	FSyncGroupMap& SyncGroups = SyncGroupMaps[GetSyncGroupWriteIndex()];
+	for (auto& SyncGroupPair : SyncGroups)
+	{
+		SyncGroupPair.Value.Reset();
+	}
+}
+
+void FAnimSync::AddTickRecord(const FAnimTickRecord& InTickRecord, const FAnimSyncParams& InSyncParams)
+{
+	if (InSyncParams.GroupName != NAME_None)
+	{
+		FSyncGroupMap& SyncGroupMap = SyncGroupMaps[GetSyncGroupWriteIndex()];
+		FAnimGroupInstance& SyncGroupInstance = SyncGroupMap.FindOrAdd(InSyncParams.GroupName);
+		SyncGroupInstance.ActivePlayers.Add(InTickRecord);
+		SyncGroupInstance.TestTickRecordForLeadership(InSyncParams.Role);
+	}
+	else
+	{
+		UngroupedActivePlayerArrays[GetSyncGroupWriteIndex()].Add(InTickRecord);
+	}
+}
+
+void FAnimSync::TickAssetPlayerInstances(FAnimInstanceProxy& InProxy, float InDeltaSeconds)
+{
+	DECLARE_SCOPE_HIERARCHICAL_COUNTER_FUNC()
+
+	SCOPE_CYCLE_COUNTER(STAT_TickAssetPlayerInstances);
+
+	const ERootMotionMode::Type RootMotionMode = InProxy.GetRootMotionMode();
+
+	// Handle all players inside sync groups
+	FSyncGroupMap& SyncGroupMap = SyncGroupMaps[GetSyncGroupWriteIndex()];
+	const FSyncGroupMap& PreviousSyncGroupMap = SyncGroupMaps[GetSyncGroupReadIndex()];
+	TArray<FAnimTickRecord>& UngroupedActivePlayers = UngroupedActivePlayerArrays[GetSyncGroupWriteIndex()];
+
+	for (auto& SyncGroupPair : SyncGroupMap)
+	{
+		FAnimGroupInstance& SyncGroup = SyncGroupPair.Value;
+	
+		if (SyncGroup.ActivePlayers.Num() > 0)
+		{
+			const FAnimGroupInstance* PreviousGroup = PreviousSyncGroupMap.Find(SyncGroupPair.Key);
+			SyncGroup.Prepare(PreviousGroup);
+
+			UE_LOG(LogAnimMarkerSync, Log, TEXT("Ticking Group [%s] GroupLeader [%d]"), *SyncGroupPair.Key.ToString(), SyncGroup.GroupLeaderIndex);
+
+			const bool bOnlyOneAnimationInGroup = SyncGroup.ActivePlayers.Num() == 1;
+
+			// Tick the group leader
+			FAnimAssetTickContext TickContext(InDeltaSeconds, RootMotionMode, bOnlyOneAnimationInGroup, SyncGroup.ValidMarkers);
+			if (PreviousGroup)
+			{
+				const FMarkerSyncAnimPosition& EndPosition = PreviousGroup->MarkerTickContext.GetMarkerSyncEndPosition();
+				if ( EndPosition.IsValid() &&
+						(EndPosition.PreviousMarkerName == NAME_None || SyncGroup.ValidMarkers.Contains(EndPosition.PreviousMarkerName)) &&
+						(EndPosition.NextMarkerName == NAME_None || SyncGroup.ValidMarkers.Contains(EndPosition.NextMarkerName)))
+				{
+					TickContext.MarkerTickContext.SetMarkerSyncStartPosition(EndPosition);
+				}
+			}
+
+#if DO_CHECK
+			//For debugging UE-54705
+			FName InitialMarkerPrevious = TickContext.MarkerTickContext.GetMarkerSyncStartPosition().PreviousMarkerName;
+			FName InitialMarkerEnd = TickContext.MarkerTickContext.GetMarkerSyncStartPosition().NextMarkerName;
+			const bool bIsLeaderRecordValidPre = SyncGroup.ActivePlayers[0].MarkerTickRecord->IsValid(SyncGroup.ActivePlayers[0].bLooping);
+			FMarkerTickRecord LeaderPreMarkerTickRecord = *SyncGroup.ActivePlayers[0].MarkerTickRecord;
+#endif
+
+			// initialize to invalidate first
+			ensureMsgf(SyncGroup.GroupLeaderIndex == INDEX_NONE, TEXT("SyncGroup %s had a non -1 group leader index of %d in asset %s"), *SyncGroupPair.Key.ToString(), SyncGroup.GroupLeaderIndex, *GetNameSafe(InProxy.GetAnimInstanceObject()));
+			int32 GroupLeaderIndex = 0;
+			for (; GroupLeaderIndex < SyncGroup.ActivePlayers.Num(); ++GroupLeaderIndex)
+			{
+				FAnimTickRecord& GroupLeader = SyncGroup.ActivePlayers[GroupLeaderIndex];
+				// if it has leader score
+				SCOPE_CYCLE_COUNTER(STAT_TickAssetPlayerInstance);
+				FScopeCycleCounterUObject Scope(GroupLeader.SourceAsset);
+				TickContext.MarkerTickContext.MarkersPassedThisTick.Reset();
+				TickContext.RootMotionMovementParams.Clear();
+				GroupLeader.SourceAsset->TickAssetPlayer(GroupLeader, InProxy.NotifyQueue, TickContext);
+
+				if (RootMotionMode == ERootMotionMode::RootMotionFromEverything && TickContext.RootMotionMovementParams.bHasRootMotion)
+				{
+					InProxy.ExtractedRootMotion.AccumulateWithBlend(TickContext.RootMotionMovementParams, GroupLeader.GetRootMotionWeight());
+				}
+
+				// if we're not using marker based sync, we don't care, get out
+				if (TickContext.CanUseMarkerPosition() == false)
+				{
+					SyncGroup.GroupLeaderIndex = GroupLeaderIndex;
+					break;
+				}
+				// otherwise, the new position should contain the valid position for end, otherwise, we don't know where to sync to
+				else if (TickContext.MarkerTickContext.IsMarkerSyncEndValid())
+				{
+					// if this leader contains correct position, break
+					SyncGroup.MarkerTickContext = TickContext.MarkerTickContext;
+					SyncGroup.GroupLeaderIndex = GroupLeaderIndex;
+					UE_LOG(LogAnimMarkerSync, Log, TEXT("Previous Sync Group Marker Tick Context :\n%s"), *SyncGroup.MarkerTickContext.ToString());
+					UE_LOG(LogAnimMarkerSync, Log, TEXT("New Sync Group Marker Tick Context :\n%s"), *TickContext.MarkerTickContext.ToString());
+					break;
+				}
+				else
+				{
+					SyncGroup.GroupLeaderIndex = GroupLeaderIndex;
+					UE_LOG(LogAnimMarkerSync, Log, TEXT("Invalid position from Leader %d. Trying next leader"), GroupLeaderIndex);
+				}
+			} 
+
+			check(SyncGroup.GroupLeaderIndex != INDEX_NONE);
+			// we found leader
+			SyncGroup.Finalize(PreviousGroup);
+
+			if (TickContext.CanUseMarkerPosition())
+			{
+				const FMarkerSyncAnimPosition& MarkerStart = TickContext.MarkerTickContext.GetMarkerSyncStartPosition();
+				FName SyncGroupName = SyncGroupPair.Key;
+				FAnimTickRecord& GroupLeader = SyncGroup.ActivePlayers[SyncGroup.GroupLeaderIndex];
+				FString LeaderAnimName = GroupLeader.SourceAsset->GetName();
+
+				//  Updated logic in search for cause of UE-54705
+				const bool bStartMarkerValid = (MarkerStart.PreviousMarkerName == NAME_None) || SyncGroup.ValidMarkers.Contains(MarkerStart.PreviousMarkerName);
+				const bool bEndMarkerValid = (MarkerStart.NextMarkerName == NAME_None) || SyncGroup.ValidMarkers.Contains(MarkerStart.NextMarkerName);
+
+				if (!bStartMarkerValid)
+				{
+#if DO_CHECK
+					FString ErrorMsg = FString(TEXT("Prev Marker name not valid for sync group.\n"));
+					ErrorMsg += FString::Format(TEXT("\tMarker {0} : SyncGroupName {1} : Leader {2}\n"), { MarkerStart.PreviousMarkerName.ToString(), SyncGroupName.ToString(), LeaderAnimName });
+					ErrorMsg += FString::Format(TEXT("\tInitalPrev {0} : InitialNext {1} : GroupLeaderIndex {2}\n"), { InitialMarkerPrevious.ToString(), InitialMarkerEnd.ToString(), GroupLeaderIndex });
+					ErrorMsg += FString::Format(TEXT("\tLeader (0 index) was originally valid: {0} | Record: {1}\n"), { bIsLeaderRecordValidPre, LeaderPreMarkerTickRecord.ToString() });
+					ErrorMsg += FString::Format(TEXT("\t Valid Markers : {0}\n"), { SyncGroup.ValidMarkers.Num() });
+					for (int32 MarkerIndex = 0; MarkerIndex < SyncGroup.ValidMarkers.Num(); ++MarkerIndex)
+					{
+						ErrorMsg += FString::Format(TEXT("\t\t{0}) '{1}'\n"), {MarkerIndex, SyncGroup.ValidMarkers[MarkerIndex].ToString()});
+					}
+					ensureMsgf(false, TEXT("%s"), *ErrorMsg);
+#endif
+					TickContext.InvalidateMarkerSync();
+				}
+				else if (!bEndMarkerValid)
+				{
+#if DO_CHECK
+					FString ErrorMsg = FString(TEXT("Next Marker name not valid for sync group.\n"));
+					ErrorMsg += FString::Format(TEXT("\tMarker {0} : SyncGroupName {1} : Leader {2}\n"), { MarkerStart.NextMarkerName.ToString(), SyncGroupName.ToString(), LeaderAnimName });
+					ErrorMsg += FString::Format(TEXT("\tInitalPrev {0} : InitialNext {1} : GroupLeaderIndex {2}\n"), { InitialMarkerPrevious.ToString(), InitialMarkerEnd.ToString(), GroupLeaderIndex });
+					ErrorMsg += FString::Format(TEXT("\tLeader (0 index) was originally valid: {0} | Record: {1}\n"), { bIsLeaderRecordValidPre, LeaderPreMarkerTickRecord.ToString() });
+					ErrorMsg += FString::Format(TEXT("\t Valid Markers : {0}\n"), { SyncGroup.ValidMarkers.Num() });
+					for (int32 MarkerIndex = 0; MarkerIndex < SyncGroup.ValidMarkers.Num(); ++MarkerIndex)
+					{
+						ErrorMsg += FString::Format(TEXT("\t\t{0}) '{1}'\n"), { MarkerIndex, SyncGroup.ValidMarkers[MarkerIndex].ToString() });
+					}
+					ensureMsgf(false, TEXT("%s"), *ErrorMsg);
+#endif
+					TickContext.InvalidateMarkerSync();
+				}
+			}
+
+			// Update everything else to follow the leader, if there is more followers
+			if (SyncGroup.ActivePlayers.Num() > GroupLeaderIndex + 1)
+			{
+				// if we don't have a good leader, no reason to convert to follower
+				// tick as leader
+				TickContext.ConvertToFollower();
+	
+				for (int32 TickIndex = GroupLeaderIndex + 1; TickIndex < SyncGroup.ActivePlayers.Num(); ++TickIndex)
+				{
+					FAnimTickRecord& AssetPlayer = SyncGroup.ActivePlayers[TickIndex];
+					{
+						SCOPE_CYCLE_COUNTER(STAT_TickAssetPlayerInstance);
+						FScopeCycleCounterUObject Scope(AssetPlayer.SourceAsset);
+						TickContext.RootMotionMovementParams.Clear();
+						AssetPlayer.SourceAsset->TickAssetPlayer(AssetPlayer, InProxy.NotifyQueue, TickContext);
+					}
+					if (RootMotionMode == ERootMotionMode::RootMotionFromEverything && TickContext.RootMotionMovementParams.bHasRootMotion)
+					{
+						InProxy.ExtractedRootMotion.AccumulateWithBlend(TickContext.RootMotionMovementParams, AssetPlayer.GetRootMotionWeight());
+					}
+				}
+			}
+
+#if ANIM_TRACE_ENABLED
+			for(const FPassedMarker& PassedMarker : TickContext.MarkerTickContext.MarkersPassedThisTick)
+			{
+				TRACE_ANIM_SYNC_MARKER(CastChecked<UAnimInstance>(InProxy.GetAnimInstanceObject()), PassedMarker);
+			}
+#endif
+		}
+	}
+	
+	// Handle the remaining ungrouped animation players
+	for (int32 TickIndex = 0; TickIndex < UngroupedActivePlayers.Num(); ++TickIndex)
+	{
+		FAnimTickRecord& AssetPlayerToTick = UngroupedActivePlayers[TickIndex];
+		const TArray<FName>* UniqueNames = AssetPlayerToTick.SourceAsset->GetUniqueMarkerNames();
+		const TArray<FName>& ValidMarkers = UniqueNames ? *UniqueNames : FMarkerTickContext::DefaultMarkerNames;
+
+		const bool bOnlyOneAnimationInGroup = true;
+		FAnimAssetTickContext TickContext(InDeltaSeconds, RootMotionMode, bOnlyOneAnimationInGroup, ValidMarkers);
+		{
+			SCOPE_CYCLE_COUNTER(STAT_TickAssetPlayerInstance);
+			FScopeCycleCounterUObject Scope(AssetPlayerToTick.SourceAsset);
+			AssetPlayerToTick.SourceAsset->TickAssetPlayer(AssetPlayerToTick, InProxy.NotifyQueue, TickContext);
+		}
+		if (RootMotionMode == ERootMotionMode::RootMotionFromEverything && TickContext.RootMotionMovementParams.bHasRootMotion)
+		{
+			InProxy.ExtractedRootMotion.AccumulateWithBlend(TickContext.RootMotionMovementParams, AssetPlayerToTick.GetRootMotionWeight());
+		}
+
+#if ANIM_TRACE_ENABLED
+		for(const FPassedMarker& PassedMarker : TickContext.MarkerTickContext.MarkersPassedThisTick)
+		{
+			TRACE_ANIM_SYNC_MARKER(CastChecked<UAnimInstance>(InProxy.GetAnimInstanceObject()), PassedMarker);
+		}
+#endif
+	}
+
+	// Flip buffers now we have ticked
+	TickSyncGroupWriteIndex();
+}
+
+bool FAnimSync::GetTimeToClosestMarker(FName SyncGroup, FName MarkerName, float& OutMarkerTime) const
+{
+	const FSyncGroupMap& SyncGroupMap = SyncGroupMaps[GetSyncGroupReadIndex()];
+
+	if (const FAnimGroupInstance* SyncGroupInstancePtr = SyncGroupMap.Find(SyncGroup))
+	{
+		if (SyncGroupInstancePtr->bCanUseMarkerSync && SyncGroupInstancePtr->ActivePlayers.IsValidIndex(SyncGroupInstancePtr->GroupLeaderIndex))
+		{
+			const FMarkerSyncAnimPosition& EndPosition = SyncGroupInstancePtr->MarkerTickContext.GetMarkerSyncEndPosition();
+			const FAnimTickRecord& Leader = SyncGroupInstancePtr->ActivePlayers[SyncGroupInstancePtr->GroupLeaderIndex];
+			if (EndPosition.PreviousMarkerName == MarkerName)
+			{
+				OutMarkerTime = Leader.MarkerTickRecord->PreviousMarker.TimeToMarker;
+				return true;
+			}
+			else if (EndPosition.NextMarkerName == MarkerName)
+			{
+				OutMarkerTime = Leader.MarkerTickRecord->NextMarker.TimeToMarker;
+				return true;
+			}
+		}
+	}
+	return false;
+}
+
+bool FAnimSync::HasMarkerBeenHitThisFrame(FName SyncGroup, FName MarkerName) const
+{
+	const FSyncGroupMap& SyncGroupMap = SyncGroupMaps[GetSyncGroupReadIndex()];
+
+	if (const FAnimGroupInstance* SyncGroupInstancePtr = SyncGroupMap.Find(SyncGroup))
+	{
+		if (SyncGroupInstancePtr->bCanUseMarkerSync)
+		{
+			return SyncGroupInstancePtr->MarkerTickContext.MarkersPassedThisTick.ContainsByPredicate([&MarkerName](const FPassedMarker& PassedMarker) -> bool
+			{
+				return PassedMarker.PassedMarkerName == MarkerName;
+			});
+		}
+	}
+	return false;
+}
+
+bool FAnimSync::IsSyncGroupBetweenMarkers(FName InSyncGroupName, FName PreviousMarker, FName NextMarker, bool bRespectMarkerOrder) const
+{
+	const FMarkerSyncAnimPosition& SyncGroupPosition = GetSyncGroupPosition(InSyncGroupName);
+	if ((SyncGroupPosition.PreviousMarkerName == PreviousMarker) && (SyncGroupPosition.NextMarkerName == NextMarker))
+	{
+		return true;
+	}
+
+	if (!bRespectMarkerOrder)
+	{
+		return ((SyncGroupPosition.PreviousMarkerName == NextMarker) && (SyncGroupPosition.NextMarkerName == PreviousMarker));
+	}
+
+	return false;
+}
+
+FMarkerSyncAnimPosition FAnimSync::GetSyncGroupPosition(FName InSyncGroupName) const
+{
+	const FSyncGroupMap& SyncGroupMap = SyncGroupMaps[GetSyncGroupReadIndex()];
+
+	if (const FAnimGroupInstance* SyncGroupInstancePtr = SyncGroupMap.Find(InSyncGroupName))
+	{
+		if (SyncGroupInstancePtr->bCanUseMarkerSync && SyncGroupInstancePtr->MarkerTickContext.IsMarkerSyncEndValid())
+		{
+			return SyncGroupInstancePtr->MarkerTickContext.GetMarkerSyncEndPosition();
+		}
+	}
+
+	return FMarkerSyncAnimPosition();
+}
+
+bool FAnimSync::IsSyncGroupValid(FName InSyncGroupName) const
+{
+	const FSyncGroupMap& SyncGroupMap = SyncGroupMaps[GetSyncGroupReadIndex()];
+
+	if (const FAnimGroupInstance* SyncGroupInstancePtr = SyncGroupMap.Find(InSyncGroupName))
+	{
+		// If we don't use Markers, we're always valid.
+		return (!SyncGroupInstancePtr->bCanUseMarkerSync || SyncGroupInstancePtr->MarkerTickContext.IsMarkerSyncEndValid());
+	}
+
+	// If we're querying a sync group that doesn't exist, treat this as invalid
+	return false;
+}
+
+FAnimTickRecord& FAnimSync::CreateUninitializedTickRecord(FAnimGroupInstance*& OutSyncGroupPtr, FName GroupName)
+{
+	// Find or create the sync group if there is one
+	OutSyncGroupPtr = nullptr;
+	if (GroupName != NAME_None)
+	{
+		FSyncGroupMap& SyncGroupMap = SyncGroupMaps[GetSyncGroupWriteIndex()];
+		OutSyncGroupPtr = &SyncGroupMap.FindOrAdd(GroupName);
+	}
+
+	// Create the record
+	FAnimTickRecord* TickRecord = new ((OutSyncGroupPtr != nullptr) ? OutSyncGroupPtr->ActivePlayers : UngroupedActivePlayerArrays[GetSyncGroupWriteIndex()]) FAnimTickRecord();
+	return *TickRecord;
+}
+
+FAnimTickRecord& FAnimSync::CreateUninitializedTickRecordInScope(FAnimInstanceProxy& InProxy, FAnimGroupInstance*& OutSyncGroupPtr, FName GroupName, EAnimSyncGroupScope Scope)
+{
+PRAGMA_DISABLE_DEPRECATION_WARNINGS
+	if (GroupName != NAME_None)
+	{
+		// If we have no main proxy or it is "this", force us to local
+		if(InProxy.GetMainInstanceProxy() == nullptr || InProxy.GetMainInstanceProxy() == &InProxy)
+		{
+			Scope = EAnimSyncGroupScope::Local;
+		}
+
+		switch(Scope)
+		{
+		default:
+			ensureMsgf(false, TEXT("FAnimSync::CreateUninitializedTickRecordInScope: Scope has invalid value %d"), Scope);
+			// Fall through
+
+		case EAnimSyncGroupScope::Local:
+			return CreateUninitializedTickRecord(OutSyncGroupPtr, GroupName);
+		case EAnimSyncGroupScope::Component:
+			// Forward to the main instance to sync with animations there in TickAssetPlayerInstances()
+			return InProxy.GetMainInstanceProxy()->CreateUninitializedTickRecord(OutSyncGroupPtr, GroupName);
+		}
+	}
+	
+	return CreateUninitializedTickRecord(OutSyncGroupPtr, GroupName);
+PRAGMA_ENABLE_DEPRECATION_WARNINGS
+}
+
+void FAnimSync::AddReferencedObjects(UAnimInstance* InAnimInstance, FReferenceCollector& Collector)
+{
+	for (int32 Index = 0; Index < UE_ARRAY_COUNT(UngroupedActivePlayerArrays); ++Index)
+	{
+		TArray<FAnimTickRecord>& UngroupedPlayers = UngroupedActivePlayerArrays[Index];
+		for (FAnimTickRecord& TickRecord : UngroupedPlayers)
+		{
+			Collector.AddReferencedObject(TickRecord.SourceAsset, InAnimInstance);
+		}
+
+		FSyncGroupMap& SyncGroupMap = SyncGroupMaps[Index];
+		for(TPair<FName, FAnimGroupInstance>& GroupPair : SyncGroupMap)
+		{
+			for (FAnimTickRecord& TickRecord : GroupPair.Value.ActivePlayers)
+			{
+				Collector.AddReferencedObject(TickRecord.SourceAsset, InAnimInstance);
+			}
+		}
+	}
+}
+
+}}	// namespace UE::Anim
diff --git a/Engine/Source/Runtime/Engine/Private/Animation/AnimSyncScope.cpp b/Engine/Source/Runtime/Engine/Private/Animation/AnimSyncScope.cpp
new file mode 100644
--- /dev/null
+++ b/Engine/Source/Runtime/Engine/Private/Animation/AnimSyncScope.cpp
@@ -0,0 +1,62 @@
+// Copyright Epic Games, Inc. All Rights Reserved.
+
+#include "Animation/AnimSyncScope.h"
+#include "Animation/AnimInstanceProxy.h"
+#include "Animation/AnimSync.h"
+
+IMPLEMENT_ANIMGRAPH_MESSAGE(UE::Anim::FAnimSyncGroupScope);
+
+namespace UE { namespace Anim {
+
+FAnimSyncDebugInfo::FAnimSyncDebugInfo(const FAnimationBaseContext& InContext)
+	: SourceProxy(InContext.AnimInstanceProxy)
+	, SourceNodeId(InContext.GetCurrentNodeId())
+{}
+
+FAnimSyncGroupScope::FAnimSyncGroupScope(const FAnimationBaseContext& InContext, FName InSyncGroup, EAnimGroupRole::Type InGroupRole)
+	: NodeId(InContext.GetCurrentNodeId())
+	, Proxy(*InContext.AnimInstanceProxy)
+	, SyncGroup(InSyncGroup)
+	, GroupRole(InGroupRole)
+{
+}
+
+void FAnimSyncGroupScope::AddTickRecord(const FAnimTickRecord& InTickRecord, const FAnimSyncParams& InSyncParams, const FAnimSyncDebugInfo& InDebugInfo)
+{
+	FAnimSyncParams NewSyncParams = InSyncParams;
+
+	// Apply method to transform params if necessary
+	switch(InSyncParams.Method)
+	{
+	default:
+	case EAnimSyncMethod::DoNotSync:
+		check(InSyncParams.GroupName == NAME_None);
+		break;
+	case EAnimSyncMethod::SyncGroup:
+		check(InSyncParams.GroupName != NAME_None);
+		break;
+	case EAnimSyncMethod::Graph:
+		// Override sync group/role supplied with our group
+		check(InSyncParams.GroupName == NAME_None);
+		NewSyncParams.GroupName = SyncGroup;
+		NewSyncParams.Role = GroupRole;
+#if !UE_BUILD_SHIPPING
+		if(InDebugInfo.SourceProxy != nullptr && InDebugInfo.SourceNodeId != INDEX_NONE)
+		{
+#if WITH_EDITORONLY_DATA
+			Proxy.RecordNodeAttribute(*InDebugInfo.SourceProxy, NodeId, InDebugInfo.SourceNodeId, FAnimSync::Attribute);
+			InDebugInfo.SourceProxy->RecordNodeSync(InDebugInfo.SourceNodeId, SyncGroup);
+#endif
+			TRACE_ANIM_NODE_ATTRIBUTE(Proxy, *InDebugInfo.SourceProxy, NodeId, InDebugInfo.SourceNodeId, FAnimSync::Attribute);
+			TRACE_ANIM_NODE_SYNC(*InDebugInfo.SourceProxy, InDebugInfo.SourceNodeId, SyncGroup);
+		}
+#endif
+		break;
+	}
+
+	// Forward to main instance
+	check(Proxy.GetMainInstanceProxy());
+	Proxy.GetMainInstanceProxy()->AddTickRecord(InTickRecord, NewSyncParams);
+}
+
+}}	// namespace UE::Anim
\ No newline at end of file
diff --git a/Engine/Source/Runtime/Engine/Private/Animation/AnimTrace.cpp b/Engine/Source/Runtime/Engine/Private/Animation/AnimTrace.cpp
--- a/Engine/Source/Runtime/Engine/Private/Animation/AnimTrace.cpp
+++ b/Engine/Source/Runtime/Engine/Private/Animation/AnimTrace.cpp
@@ -20,6 +20,7 @@
 #include "Animation/AnimNotifies/AnimNotifyState.h"
 #include "Animation/AnimTypes.h"
 #include "TraceFilter.h"
+#include "Animation/AnimAttributes.h"
 
 UE_TRACE_CHANNEL_DEFINE(AnimationChannel);
 
@@ -86,6 +87,15 @@ UE_TRACE_EVENT_BEGIN(Animation, AnimNodeEnd)
 	UE_TRACE_EVENT_FIELD(uint64, AnimInstanceId)
 UE_TRACE_EVENT_END()
 
+UE_TRACE_EVENT_BEGIN(Animation, AnimNodeAttribute)
+	UE_TRACE_EVENT_FIELD(uint64, Cycle)
+	UE_TRACE_EVENT_FIELD(uint64, SourceAnimInstanceId)
+	UE_TRACE_EVENT_FIELD(uint64, TargetAnimInstanceId)
+	UE_TRACE_EVENT_FIELD(int32, SourceNodeId)
+	UE_TRACE_EVENT_FIELD(int32, TargetNodeId)
+	UE_TRACE_EVENT_FIELD(uint32, NameId)
+UE_TRACE_EVENT_END()
+
 UE_TRACE_EVENT_BEGIN(Animation, AnimNodeValueBool)
 	UE_TRACE_EVENT_FIELD(uint64, Cycle)
 	UE_TRACE_EVENT_FIELD(uint64, AnimInstanceId)
@@ -222,6 +232,13 @@ UE_TRACE_EVENT_BEGIN(Animation, Montage)
 	UE_TRACE_EVENT_FIELD(uint16, FrameCounter)
 UE_TRACE_EVENT_END()
 
+UE_TRACE_EVENT_BEGIN(Animation, Sync)
+	UE_TRACE_EVENT_FIELD(uint64, Cycle)
+	UE_TRACE_EVENT_FIELD(uint64, AnimInstanceId)
+	UE_TRACE_EVENT_FIELD(int32, SourceNodeId)
+	UE_TRACE_EVENT_FIELD(uint32, GroupNameId)
+UE_TRACE_EVENT_END()
+
 // Object annotations used for tracing
 FUObjectAnnotationSparseBool GSkeletalMeshTraceAnnotations;
 
@@ -660,6 +677,35 @@ void FAnimTrace::OutputAnimNodeEnd(const FAnimationBaseContext& InContext, uint6
 		<< AnimNodeEnd.AnimInstanceId(FObjectTrace::GetObjectId(AnimInstance));
 }
 
+void FAnimTrace::OutputAnimNodeAttribute(const FAnimInstanceProxy& InTargetProxy, const FAnimInstanceProxy& InSourceProxy, int32 InTargetNodeId, int32 InSourceNodeId, FName InAttribute)
+{
+	bool bChannelEnabled = UE_TRACE_CHANNELEXPR_IS_ENABLED(AnimationChannel);
+	if (!bChannelEnabled)
+	{
+		return;
+	}
+
+	if (CANNOT_TRACE_OBJECT(InSourceProxy.GetSkelMeshComponent()) || CANNOT_TRACE_OBJECT(InTargetProxy.GetSkelMeshComponent()))
+	{
+		return;
+	}
+
+	const UObject* SourceAnimInstance = InSourceProxy.GetAnimInstanceObject();
+	TRACE_OBJECT(SourceAnimInstance)
+	const UObject* TargetAnimInstance = InTargetProxy.GetAnimInstanceObject();
+	TRACE_OBJECT(TargetAnimInstance);
+
+	uint32 NameId = OutputName(InAttribute);
+
+	UE_TRACE_LOG(Animation, AnimNodeAttribute, AnimationChannel)
+		<< AnimNodeAttribute.Cycle(FPlatformTime::Cycles64())
+		<< AnimNodeAttribute.SourceAnimInstanceId(FObjectTrace::GetObjectId(SourceAnimInstance))
+		<< AnimNodeAttribute.TargetAnimInstanceId(FObjectTrace::GetObjectId(TargetAnimInstance))
+		<< AnimNodeAttribute.SourceNodeId(InSourceNodeId)
+		<< AnimNodeAttribute.TargetNodeId(InTargetNodeId)
+		<< AnimNodeAttribute.NameId(NameId);
+}
+
 void FAnimTrace::OutputAnimNodeValue(const FAnimationBaseContext& InContext, const TCHAR* InKey, bool InValue)
 {
 	bool bChannelEnabled = UE_TRACE_CHANNELEXPR_IS_ENABLED(AnimationChannel);
@@ -1104,4 +1150,29 @@ void FAnimTrace::OutputMontage(UAnimInstance* InAnimInstance, const FAnimMontage
 	}
 }
 
+void FAnimTrace::OutputSync(const FAnimInstanceProxy& InSourceProxy, int32 InSourceNodeId, FName InGroupName)
+{
+	bool bChannelEnabled = UE_TRACE_CHANNELEXPR_IS_ENABLED(AnimationChannel);
+	if (!bChannelEnabled)
+	{
+		return;
+	}
+
+
+	if (CANNOT_TRACE_OBJECT(InSourceProxy.GetSkelMeshComponent()))
+	{
+		return;
+	}
+
+	TRACE_OBJECT(InSourceProxy.GetAnimInstanceObject());
+
+	uint32 GroupNameId = FAnimTrace::OutputName(InGroupName);
+
+	UE_TRACE_LOG(Animation, Sync, AnimationChannel)
+		<< Sync.Cycle(FPlatformTime::Cycles64())
+		<< Sync.AnimInstanceId(FObjectTrace::GetObjectId(InSourceProxy.GetAnimInstanceObject()))
+		<< Sync.SourceNodeId(InSourceNodeId)
+		<< Sync.GroupNameId(GroupNameId);
+}
+
 #endif
diff --git a/Engine/Source/Runtime/Engine/Private/AnimationAsset.cpp b/Engine/Source/Runtime/Engine/Private/AnimationAsset.cpp
--- a/Engine/Source/Runtime/Engine/Private/AnimationAsset.cpp
+++ b/Engine/Source/Runtime/Engine/Private/AnimationAsset.cpp
@@ -8,6 +8,8 @@
 #include "AnimationUtils.h"
 #include "Animation/AnimInstance.h"
 #include "UObject/LinkerLoad.h"
+#include "Animation/BlendSpaceBase.h"
+#include "Animation/PoseAsset.h"
 
 #define LEADERSCORE_ALWAYSLEADER  	2.f
 #define LEADERSCORE_MONTAGE			3.f
@@ -20,65 +22,65 @@ const TArray<FName> FMarkerTickContext::DefaultMarkerNames;
 
 void FAnimGroupInstance::TestTickRecordForLeadership(EAnimGroupRole::Type MembershipType)
 {
+	check(ActivePlayers.Num() > 0);
+
 	// always set leader score if you have potential to be leader
 	// that way if the top leader fails, we'll continue to search next available leader
 	int32 TestIndex = ActivePlayers.Num() - 1;
 	FAnimTickRecord& Candidate = ActivePlayers[TestIndex];
 
-	switch (MembershipType)
-	{
-	case EAnimGroupRole::CanBeLeader:
-	case EAnimGroupRole::TransitionLeader:
-		Candidate.LeaderScore = Candidate.EffectiveBlendWeight;
-		break;
-	case EAnimGroupRole::AlwaysLeader:
-		// Always set the leader index
-		Candidate.LeaderScore = LEADERSCORE_ALWAYSLEADER;
-		break;
-	default:
-	case EAnimGroupRole::AlwaysFollower:
-	case EAnimGroupRole::TransitionFollower:
-		// Never set the leader index; the actual tick code will handle the case of no leader by using the first element in the array
-		break;
-	}
-}
+	if(Candidate.SourceAsset->IsA<UAnimMontage>())
+	{
+		// if the candidate has higher weight
+		if (Candidate.EffectiveBlendWeight > MontageLeaderWeight)
+		{
+			// if this is going to be leader, I'll clean ActivePlayers because we don't sync multi montages
+			const int32 LastIndex = TestIndex - 1;
+			if (LastIndex >= 0)
+			{
+				ActivePlayers.RemoveAt(TestIndex - 1, 1);
+			}
 
-void FAnimGroupInstance::TestMontageTickRecordForLeadership()
-{
-	int32 TestIndex = ActivePlayers.Num() - 1;
-	ensure(TestIndex <= 1);
-	FAnimTickRecord& Candidate = ActivePlayers[TestIndex];
+			// at this time, it should only have one
+			ensure(ActivePlayers.Num() == 1);
 
-	// if the candidate has higher weight
-	if (Candidate.EffectiveBlendWeight > MontageLeaderWeight)
-	{
-		// if this is going to be leader, I'll clean ActivePlayers because we don't sync multi montages
-		const int32 LastIndex = TestIndex - 1;
-		if (LastIndex >= 0)
+			// then override
+			// @note : leader weight doesn't applied WITHIN montages
+			// we still only contain one montage at a time, if this montage fails, next candidate will get the chance, not next weight montage
+			MontageLeaderWeight = Candidate.EffectiveBlendWeight;
+			Candidate.LeaderScore = LEADERSCORE_MONTAGE;
+		}
+		else
 		{
-			ActivePlayers.RemoveAt(TestIndex - 1, 1);
+			if (TestIndex != 0)
+			{
+				// we delete the later ones because we only have one montage for leader. 
+				// this can happen if there was already active one with higher weight. 
+				ActivePlayers.RemoveAt(TestIndex, 1);
+			}
 		}
 
-		// at this time, it should only have one
-		ensure(ActivePlayers.Num() == 1);
-
-		// then override
-		// @note : leader weight doesn't applied WITHIN montages
-		// we still only contain one montage at a time, if this montage fails, next candidate will get the chance, not next weight montage
-		MontageLeaderWeight = Candidate.EffectiveBlendWeight;
-		Candidate.LeaderScore = LEADERSCORE_MONTAGE;
+		ensureAlways(ActivePlayers.Num() == 1);
 	}
 	else
 	{
-		if (TestIndex != 0)
+		switch (MembershipType)
 		{
-			// we delete the later ones because we only have one montage for leader. 
-			// this can happen if there was already active one with higher weight. 
-			ActivePlayers.RemoveAt(TestIndex, 1);
+		case EAnimGroupRole::CanBeLeader:
+		case EAnimGroupRole::TransitionLeader:
+			Candidate.LeaderScore = Candidate.EffectiveBlendWeight;
+			break;
+		case EAnimGroupRole::AlwaysLeader:
+			// Always set the leader index
+			Candidate.LeaderScore = LEADERSCORE_ALWAYSLEADER;
+			break;
+		default:
+		case EAnimGroupRole::AlwaysFollower:
+		case EAnimGroupRole::TransitionFollower:
+			// Never set the leader index; the actual tick code will handle the case of no leader by using the first element in the array
+			break;
 		}
 	}
-
-	ensureAlways(ActivePlayers.Num() == 1);
 }
 
 void FAnimGroupInstance::Finalize(const FAnimGroupInstance* PreviousGroup)
@@ -176,6 +178,49 @@ void FAnimGroupInstance::Prepare(const FAnimGroupInstance* PreviousGroup)
 	}
 }
 
+FAnimTickRecord::FAnimTickRecord(UAnimSequenceBase* InSequence, bool bInLooping, float InPlayRate, float InFinalBlendWeight, float& InCurrentTime, FMarkerTickRecord& InMarkerTickRecord)
+{
+	SourceAsset = InSequence;
+	TimeAccumulator = &InCurrentTime;
+	MarkerTickRecord = &InMarkerTickRecord;
+	PlayRateMultiplier = InPlayRate;
+	EffectiveBlendWeight = InFinalBlendWeight;
+	bLooping = bInLooping;
+}
+
+FAnimTickRecord::FAnimTickRecord(UBlendSpaceBase* InBlendSpace, const FVector& InBlendInput, TArray<FBlendSampleData>& InBlendSampleDataCache, FBlendFilter& InBlendFilter, bool bInLooping, float InPlayRate, float InFinalBlendWeight, float& InCurrentTime, FMarkerTickRecord& InMarkerTickRecord)
+{
+	SourceAsset = InBlendSpace;
+	BlendSpace.BlendSpacePositionX = InBlendInput.X;
+	BlendSpace.BlendSpacePositionY = InBlendInput.Y;
+	BlendSpace.BlendSampleDataCache = &InBlendSampleDataCache;
+	BlendSpace.BlendFilter = &InBlendFilter;
+	TimeAccumulator = &InCurrentTime;
+	MarkerTickRecord = &InMarkerTickRecord;
+	PlayRateMultiplier = InPlayRate;
+	EffectiveBlendWeight = InFinalBlendWeight;
+	bLooping = bInLooping;
+}
+
+FAnimTickRecord::FAnimTickRecord(UAnimMontage* InMontage, float InCurrentPosition, float InPreviousPosition, float InMoveDelta, float InWeight, TArray<FPassedMarker>& InMarkersPassedThisTick, FMarkerTickRecord& InMarkerTickRecord)
+{
+	SourceAsset = InMontage;
+	Montage.CurrentPosition = InCurrentPosition;
+	Montage.PreviousPosition = InPreviousPosition;
+	Montage.MoveDelta = InMoveDelta;
+	Montage.MarkersPassedThisTick = &InMarkersPassedThisTick;
+	MarkerTickRecord = &InMarkerTickRecord;
+	PlayRateMultiplier = 1.f; // we don't care here, this is alreayd applied in the montageinstance::Advance
+	EffectiveBlendWeight = InWeight;
+	bLooping = false;
+}
+
+FAnimTickRecord::FAnimTickRecord(UPoseAsset* InPoseAsset, float InFinalBlendWeight)
+{
+	SourceAsset = InPoseAsset;
+	EffectiveBlendWeight = InFinalBlendWeight;
+}
+
 //////////////////////////////////////////////////////////////////////////
 // UAnimationAsset
 
diff --git a/Engine/Source/Runtime/Engine/Public/Animation/AnimAttributes.h b/Engine/Source/Runtime/Engine/Public/Animation/AnimAttributes.h
new file mode 100644
--- /dev/null
+++ b/Engine/Source/Runtime/Engine/Public/Animation/AnimAttributes.h
@@ -0,0 +1,18 @@
+// Copyright Epic Games, Inc. All Rights Reserved.
+
+#pragma once
+
+#include "CoreMinimal.h"
+
+namespace UE { namespace Anim 
+{
+
+// Built-in attributes that most nodes will share
+struct ENGINE_API FAttributes
+{
+	static const FName Pose;
+	static const FName Curves;
+	static const FName Attributes;
+};
+
+}}	// namespace UE::Anim
diff --git a/Engine/Source/Runtime/Engine/Public/Animation/AnimCurveTypes.h b/Engine/Source/Runtime/Engine/Public/Animation/AnimCurveTypes.h
--- a/Engine/Source/Runtime/Engine/Public/Animation/AnimCurveTypes.h
+++ b/Engine/Source/Runtime/Engine/Public/Animation/AnimCurveTypes.h
@@ -700,6 +700,9 @@ struct FBaseBlendedCurve
 	/** Return number of elements */
 	int32 Num() const { return CurveWeights.Num(); }
 
+	/** Return number of valid elements */
+	int32 NumValid() const { return ValidCurveWeights.CountSetBits(); }
+
 	/** CopyFrom as expected. */
 	template <typename OtherAllocator>
 	void CopyFrom(const FBaseBlendedCurve<OtherAllocator>& CurveToCopyFrom)
diff --git a/Engine/Source/Runtime/Engine/Public/Animation/AnimInstanceProxy.h b/Engine/Source/Runtime/Engine/Public/Animation/AnimInstanceProxy.h
--- a/Engine/Source/Runtime/Engine/Public/Animation/AnimInstanceProxy.h
+++ b/Engine/Source/Runtime/Engine/Public/Animation/AnimInstanceProxy.h
@@ -18,6 +18,7 @@
 #include "Animation/AnimBlueprintGeneratedClass.h"
 #include "Logging/TokenizedMessage.h"
 #include "Animation/AnimTrace.h"
+#include "Animation/AnimSync.h"
 
 #include "AnimInstanceProxy.generated.h"
 
@@ -125,7 +126,7 @@ public:
 		, CurrentTimeDilation(1.0f)
 		, RootNode(nullptr)
 		, DefaultLinkedInstanceInputNode(nullptr)
-		, SyncGroupWriteIndex(0)
+		, BufferWriteIndex(0)
 		, RootMotionMode(ERootMotionMode::NoRootMotionExtraction)
 		, FrameCounterForUpdate(0)
 		, FrameCounterForNodeUpdate(0)
@@ -150,7 +151,7 @@ public:
 		, CurrentTimeDilation(1.0f)
 		, RootNode(nullptr)
 		, DefaultLinkedInstanceInputNode(nullptr)
-		, SyncGroupWriteIndex(0)
+		, BufferWriteIndex(0)
 		, RootMotionMode(ERootMotionMode::NoRootMotionExtraction)
 		, FrameCounterForUpdate(0)
 		, FrameCounterForNodeUpdate(0)
@@ -204,7 +205,16 @@ public:
 	/** Record a visited node in the debugger */
 	void RecordNodeVisit(int32 TargetNodeIndex, int32 SourceNodeIndex, float BlendWeight)
 	{
-		new (UpdatedNodesThisFrame) FAnimBlueprintDebugData::FNodeVisit(SourceNodeIndex, TargetNodeIndex, BlendWeight);
+		UpdatedNodesThisFrame.Emplace(SourceNodeIndex, TargetNodeIndex, BlendWeight);
+	}
+
+	/** Record a node attribute in the debugger */
+	void RecordNodeAttribute(const FAnimInstanceProxy& InSourceProxy, int32 InTargetNodeIndex, int32 InSourceNodeIndex, FName InAttribute);
+
+	/** Record a node sync in the debugger */
+	void RecordNodeSync(int32 InSourceNodeIndex, FName InSyncGroup)
+	{
+		NodeSyncsThisFrame.FindOrAdd(InSourceNodeIndex, InSyncGroup);
 	}
 
 	UAnimBlueprint* GetAnimBlueprint() const
@@ -218,13 +228,19 @@ public:
 	void RegisterWatchedPose(const FCSPose<FCompactPose>& Pose, int32 LinkID);
 #endif
 
-	// flip sync group read/write indices
+	UE_DEPRECATED(5.0, "Function renamed to FlipBufferWriteIndex as this no longer deals with sync groups.")
 	void TickSyncGroupWriteIndex()
 	{ 
-		SyncGroupWriteIndex = GetSyncGroupReadIndex();
+		FlipBufferWriteIndex();
 	}
 
-	UE_DEPRECATED(5.0, "Sync groups are no longer stored in arrays, please use GetSyncGroupMapRead")
+	// flip buffer read/write indices
+	void FlipBufferWriteIndex()
+	{ 
+		BufferWriteIndex = GetBufferReadIndex();
+	}
+
+	UE_DEPRECATED(5.0, "Sync groups are no longer stored in arrays, please use GetSyncGroupMapRead.")
 	const TArray<FAnimGroupInstance>& GetSyncGroupRead() const
 	{
 		static const TArray<FAnimGroupInstance> Dummy;
@@ -234,19 +250,19 @@ public:
 	/** Get the sync group we are currently reading from */
 	const FSyncGroupMap& GetSyncGroupMapRead() const
 	{ 
-		return SyncGroupMaps[GetSyncGroupReadIndex()]; 
+		return Sync.GetSyncGroupMapRead(); 
 	}
 
 	/** Get the ungrouped active player we are currently reading from */
 	const TArray<FAnimTickRecord>& GetUngroupedActivePlayersRead() 
 	{ 
-		return UngroupedActivePlayerArrays[GetSyncGroupReadIndex()]; 
+		return Sync.GetUngroupedActivePlayersRead(); 
 	}
 
-	/** Tick active asset players. */
+	UE_DEPRECATED(5.0, "Do not call directly. Instead sync and asset player ticking is controlled via UE::Anim::FSync and UE::Anim::FSyncScope.")
 	void TickAssetPlayerInstances(float DeltaSeconds);
 
-	/** Tick active asset players. This overload with no parameters uses CurrentDeltaSeconds */
+	UE_DEPRECATED(5.0, "Do not call directly. Instead sync and asset player ticking is controlled via UE::Anim::FSync and UE::Anim::FSyncScope.")
 	void TickAssetPlayerInstances();
 
 	/** Queues an Anim Notify from the shared list on our generated class */
@@ -351,34 +367,48 @@ public:
 	/** Get the current skeletal mesh component we are running on. Note that this will return nullptr outside of pre/post update */
 	USkeletalMeshComponent* GetSkelMeshComponent() const
 	{ 
-		// Skeleton is only available during update/eval. If you're calling this function outside of it, it will return null. 
+		// Skeletal mesh component is only available during update/eval. If you're calling this function outside of it, it will return null. 
 		// adding ensure here so that we can catch them earlier
 		ensureAlways(SkeletalMeshComponent);
 		return SkeletalMeshComponent; 
 	}
 
+	/** Get the current main instance proxy. Note that this will return nullptr outside of pre/post update */
+	FAnimInstanceProxy* GetMainInstanceProxy() const
+	{ 
+		// Main instance proxy is only available during update/eval. If you're calling this function outside of it, it will return null. 
+		// adding ensure here so that we can catch them earlier
+		ensureAlways(MainInstanceProxy);
+		return MainInstanceProxy; 
+	}
+
 	UE_DEPRECATED(4.26, "Please use the overload that takes a group FName")
 	FAnimTickRecord& CreateUninitializedTickRecord(int32 GroupIndex, FAnimGroupInstance*& OutSyncGroupPtr);
 
 	UE_DEPRECATED(4.26, "Please use the overload that takes a group FName")
 	FAnimTickRecord& CreateUninitializedTickRecordInScope(int32 GroupIndex, EAnimSyncGroupScope Scope, FAnimGroupInstance*& OutSyncGroupPtr);
 
-	// Creates an uninitialized tick record in the list for the correct group or the ungrouped array.  If the group is valid, OutSyncGroupPtr will point to the group.
+	UE_DEPRECATED(5.0, "Please use FAnimSyncGroupScope")
 	FAnimTickRecord& CreateUninitializedTickRecord(FAnimGroupInstance*& OutSyncGroupPtr, FName GroupName);
 
-	// Creates an uninitialized tick record in the list for the correct group or the ungrouped array.  
-	// If the group is valid, OutSyncGroupPtr will point to the group.
-	// Supply the scope to sync with tick records outside this instance
+	UE_DEPRECATED(5.0, "Please use FAnimSyncGroupScope")
 	FAnimTickRecord& CreateUninitializedTickRecordInScope(FAnimGroupInstance*& OutSyncGroupPtr, FName GroupName, EAnimSyncGroupScope Scope);
 
-	/** Helper function: make a tick record for a sequence */
+	UE_DEPRECATED(5.0, "Please use the FAnimTickRecord constructor that takes a UAnimSequenceBase")
 	void MakeSequenceTickRecord(FAnimTickRecord& TickRecord, UAnimSequenceBase* Sequence, bool bLooping, float PlayRate, float FinalBlendWeight, float& CurrentTime, FMarkerTickRecord& MarkerTickRecord) const;
 
-	/** Helper function: make a tick record for a blend space */
+	UE_DEPRECATED(5.0, "Please use the FAnimTickRecord constructor that takes a UBlendSpaceBase")
 	void MakeBlendSpaceTickRecord(FAnimTickRecord& TickRecord, UBlendSpaceBase* BlendSpace, const FVector& BlendInput, TArray<FBlendSampleData>& BlendSampleDataCache, FBlendFilter& BlendFilter, bool bLooping, float PlayRate, float FinalBlendWeight, float& CurrentTime, FMarkerTickRecord& MarkerTickRecord) const;
 
-	/** Helper function: make a tick record for a pose asset*/
+	UE_DEPRECATED(5.0, "Please use the FAnimTickRecord constructor that takes a UPoseAsset")
 	void MakePoseAssetTickRecord(FAnimTickRecord& TickRecord, class UPoseAsset* PoseAsset, float FinalBlendWeight) const;
+
+	// Adds a tick record in the list for the correct group or the ungrouped array.
+	void AddTickRecord(const FAnimTickRecord& InTickRecord, const UE::Anim::FAnimSyncParams& InSyncParams = UE::Anim::FAnimSyncParams())
+	{
+		Sync.AddTickRecord(InTickRecord, InSyncParams);
+	}
+
 	/**
 	 * Get Slot Node Weight : this returns new Slot Node Weight, Source Weight, Original TotalNodeWeight
 	 *							this 3 values can't be derived from each other
@@ -454,6 +484,12 @@ public:
 		return AnimInstanceName;
 	}
 
+	// Get the root motion mode assigned to this anim instance proxy
+	ERootMotionMode::Type GetRootMotionMode() const
+	{
+		return RootMotionMode;
+	}
+
 	/** Gets the runtime instance of the specified state machine by Name */
 	FAnimNode_StateMachine* GetStateMachineInstanceFromName(FName MachineName);
 
@@ -509,6 +545,7 @@ public:
 	friend struct FAnimNode_LinkedAnimGraph;
 	friend struct FAnimationBaseContext;
 	friend struct FAnimTrace;
+	friend struct UE::Anim::FAnimSync;
 
 protected:
 	/** Called when our anim instance is being initialized */
@@ -635,16 +672,28 @@ protected:
 	// @todo document
 	void BlendSpaceAdvanceImmediate(UBlendSpaceBase* BlendSpace, const FVector& BlendInput, TArray<FBlendSampleData> & BlendSampleDataCache, FBlendFilter & BlendFilter, bool bLooping, float PlayRate, float DeltaSeconds, /*inout*/ float& CurrentTime, FMarkerTickRecord& MarkerTickRecord);
 
-	// Gets the sync group we should be reading from
+	UE_DEPRECATED(5.0, "Function renamed to GetBufferReadIndex as this no longer deals with sync groups.")
 	int32 GetSyncGroupReadIndex() const 
 	{ 
-		return 1 - SyncGroupWriteIndex; 
+		return GetBufferReadIndex(); 
 	}
 
-	// Gets the sync group we should be writing to
+	UE_DEPRECATED(5.0, "Function renamed to GetBufferWriteIndex as this no longer deals with sync groups.")
 	int32 GetSyncGroupWriteIndex() const 
 	{ 
-		return SyncGroupWriteIndex; 
+		return GetBufferWriteIndex(); 
+	}
+
+	// Gets the buffer we should be reading from
+	int32 GetBufferReadIndex() const 
+	{ 
+		return 1 - BufferWriteIndex; 
+	}
+
+	// Gets the sync group we should be writing to
+	int32 GetBufferWriteIndex() const 
+	{ 
+		return BufferWriteIndex; 
 	}
 
 	/** Add anim notifier **/
@@ -880,6 +929,13 @@ private:
 	/** Array of visited nodes this frame */
 	TArray<FAnimBlueprintDebugData::FNodeVisit> UpdatedNodesThisFrame;
 
+	/** Map of node attributes this frame */
+	TMap<int32, TArray<FAnimBlueprintDebugData::FAttributeRecord>> NodeInputAttributesThisFrame;
+	TMap<int32, TArray<FAnimBlueprintDebugData::FAttributeRecord>> NodeOutputAttributesThisFrame;
+
+	/** Map of node syncs this frame - maps from player node index to graph-determined group name */
+	TMap<int32, FName> NodeSyncsThisFrame;
+
 	/** Array of nodes to watch this frame */
 	TArray<FAnimNodePoseWatch> PoseWatchEntriesForThisFrame;
 #endif
@@ -901,11 +957,8 @@ private:
 	/** Map of layer name to saved pose nodes to process after the graph has been updated */
 	TMap<FName, TArray<FAnimNode_SaveCachedPose*>> SavedPoseQueueMap;
 
-	/** The list of animation assets which are going to be evaluated this frame and need to be ticked (ungrouped) */
-	TArray<FAnimTickRecord> UngroupedActivePlayerArrays[2];
-
-	/** The set of tick groups for this anim instance */
-	FSyncGroupMap SyncGroupMaps[2];
+	/** Synchronizes animations according to sync groups/markers */
+	UE::Anim::FAnimSync Sync;
 
 	/** Buffers containing read/write buffers for all current machine weights */
 	TArray<float> MachineWeightArrays[2];
@@ -916,8 +969,8 @@ private:
 	/** Map that transforms state class indices to base offsets into the weight array */
 	TMap<int32, int32> StateMachineClassIndexToWeightOffset;
 
-	// Current sync group buffer index
-	int32 SyncGroupWriteIndex;
+	// Current buffer index
+	int32 BufferWriteIndex;
 
 	/** Animation Notifies that has been triggered in the latest tick **/
 	FAnimNotifyQueue NotifyQueue;
diff --git a/Engine/Source/Runtime/Engine/Public/Animation/AnimSync.h b/Engine/Source/Runtime/Engine/Public/Animation/AnimSync.h
new file mode 100644
--- /dev/null
+++ b/Engine/Source/Runtime/Engine/Public/Animation/AnimSync.h
@@ -0,0 +1,104 @@
+// Copyright Epic Games, Inc. All Rights Reserved.
+
+#pragma once
+
+#include "CoreMinimal.h"
+#include "AnimNodeMessages.h"
+#include "Animation/AnimationAsset.h"
+
+struct FAnimationBaseContext;
+struct FAnimInstanceProxy;
+
+namespace UE { namespace Anim {
+
+// Parameters for group/marker-based sync 
+struct FAnimSyncParams
+{
+	FAnimSyncParams(FName InGroupName = NAME_None, EAnimGroupRole::Type InRole = EAnimGroupRole::CanBeLeader, EAnimSyncMethod InMethod = EAnimSyncMethod::DoNotSync)
+		: GroupName(InGroupName)
+		, Role(InRole)
+		, Method(InMethod)
+	{}
+
+	FName GroupName = NAME_None;
+	EAnimGroupRole::Type Role = EAnimGroupRole::CanBeLeader;
+	EAnimSyncMethod Method = EAnimSyncMethod::DoNotSync;
+};
+
+// Wraps up functionality for ticking and syncing animations according to group (via normalized time) or marker
+struct ENGINE_API FAnimSync
+{
+	static const FName Attribute;
+
+	using FSyncGroupMap = TMap<FName, FAnimGroupInstance>;
+
+	// Reset internal buffers ready for writing
+	void Reset();
+
+	// Adds a tick record in the list for the correct group or the ungrouped array.
+	void AddTickRecord(const FAnimTickRecord& InTickRecord, const FAnimSyncParams& InSyncParams = FAnimSyncParams());
+
+	// Tick all of the asset player tick records that are registered with us
+	void TickAssetPlayerInstances(FAnimInstanceProxy& InProxy, float InDeltaSeconds);
+
+	bool GetTimeToClosestMarker(FName SyncGroup, FName MarkerName, float& OutMarkerTime) const;
+
+	bool HasMarkerBeenHitThisFrame(FName SyncGroup, FName MarkerName) const;
+
+	bool IsSyncGroupBetweenMarkers(FName InSyncGroupName, FName PreviousMarker, FName NextMarker, bool bRespectMarkerOrder = true) const;
+
+	FMarkerSyncAnimPosition GetSyncGroupPosition(FName InSyncGroupName) const;
+
+	bool IsSyncGroupValid(FName InSyncGroupName) const;
+
+	// flip sync group read/write indices
+	void TickSyncGroupWriteIndex()
+	{ 
+		SyncGroupWriteIndex = GetSyncGroupReadIndex();
+	}
+
+	// Gets the sync group we should be reading from
+	int32 GetSyncGroupReadIndex() const 
+	{ 
+		return 1 - SyncGroupWriteIndex; 
+	}
+
+	// Gets the sync group we should be writing to
+	int32 GetSyncGroupWriteIndex() const 
+	{ 
+		return SyncGroupWriteIndex; 
+	}
+
+	/** Get the sync group we are currently reading from */
+	const FSyncGroupMap& GetSyncGroupMapRead() const
+	{ 
+		return SyncGroupMaps[GetSyncGroupReadIndex()]; 
+	}
+
+	/** Get the ungrouped active player we are currently reading from */
+	const TArray<FAnimTickRecord>& GetUngroupedActivePlayersRead() 
+	{ 
+		return UngroupedActivePlayerArrays[GetSyncGroupReadIndex()]; 
+	}
+
+	UE_DEPRECATED(5.0, "Legacy API support - do not use")
+	FAnimTickRecord& CreateUninitializedTickRecord(FAnimGroupInstance*& OutSyncGroupPtr, FName GroupName);
+
+	UE_DEPRECATED(5.0, "Legacy API support - do not use")
+	FAnimTickRecord& CreateUninitializedTickRecordInScope(FAnimInstanceProxy& InProxy, FAnimGroupInstance*& OutSyncGroupPtr, FName GroupName, EAnimSyncGroupScope Scope);
+
+	// GC support
+	void AddReferencedObjects(UAnimInstance* InAnimInstance, FReferenceCollector& Collector);
+
+private:
+	/** The list of animation assets which are going to be evaluated this frame and need to be ticked (ungrouped) */
+	TArray<FAnimTickRecord> UngroupedActivePlayerArrays[2];
+
+	/** The set of tick groups for this anim instance */
+	FSyncGroupMap SyncGroupMaps[2];
+
+	/** Current sync group buffer index */
+	int32 SyncGroupWriteIndex = 0;
+};
+
+}}	// namespace UE::Anim
diff --git a/Engine/Source/Runtime/Engine/Public/Animation/AnimSyncScope.h b/Engine/Source/Runtime/Engine/Public/Animation/AnimSyncScope.h
new file mode 100644
--- /dev/null
+++ b/Engine/Source/Runtime/Engine/Public/Animation/AnimSyncScope.h
@@ -0,0 +1,56 @@
+// Copyright Epic Games, Inc. All Rights Reserved.
+
+#pragma once
+
+#include "CoreMinimal.h"
+#include "AnimNodeMessages.h"
+#include "Animation/AnimationAsset.h"
+#include "Animation/AnimSync.h"
+
+struct FAnimationBaseContext;
+struct FAnimInstanceProxy;
+
+namespace UE { namespace Anim {
+
+// Debug info for a sync group scope request
+struct ENGINE_API FAnimSyncDebugInfo
+{
+	FAnimSyncDebugInfo() = default;
+
+	FAnimSyncDebugInfo(const FAnimationBaseContext& InContext);
+
+	FAnimSyncDebugInfo(FAnimInstanceProxy& InSourceProxy, int32 InSourceNodeId)
+		: SourceProxy(&InSourceProxy)
+		, SourceNodeId(InSourceNodeId)
+	{}
+
+	FAnimInstanceProxy* SourceProxy = nullptr;
+	int32 SourceNodeId = INDEX_NONE;
+};
+
+// Scoped graph message used to synchronize animations at various points in an anim graph
+class ENGINE_API FAnimSyncGroupScope : public IGraphMessage
+{
+	DECLARE_ANIMGRAPH_MESSAGE(FAnimSyncGroupScope);
+
+public:
+	FAnimSyncGroupScope(const FAnimationBaseContext& InContext, FName InSyncGroup = NAME_None, EAnimGroupRole::Type InGroupRole = EAnimGroupRole::CanBeLeader);
+
+	// Adds a tick record in the list for the correct group or the ungrouped array.
+	void AddTickRecord(const FAnimTickRecord& InTickRecord, const FAnimSyncParams& InSyncParams = FAnimSyncParams(), const FAnimSyncDebugInfo& InDebugInfo = FAnimSyncDebugInfo());
+
+private:
+	// The node ID that was used when this scope was entered
+	int32 NodeId;
+
+	// The proxy that we are running
+	FAnimInstanceProxy& Proxy;
+
+	// The sync group that this scope uses.
+	FName SyncGroup;
+
+	// The role assumed by the specified sync group
+	EAnimGroupRole::Type GroupRole;
+};
+
+}}	// namespace UE::Anim
diff --git a/Engine/Source/Runtime/Engine/Public/Animation/AnimTrace.h b/Engine/Source/Runtime/Engine/Public/Animation/AnimTrace.h
--- a/Engine/Source/Runtime/Engine/Public/Animation/AnimTrace.h
+++ b/Engine/Source/Runtime/Engine/Public/Animation/AnimTrace.h
@@ -11,6 +11,8 @@
 
 #if ANIM_TRACE_ENABLED
 
+#include "AnimAttributes.h"
+
 UE_TRACE_CHANNEL_EXTERN(AnimationChannel, ENGINE_API);
 
 struct FAnimInstanceProxy;
@@ -133,6 +135,24 @@ struct FAnimTrace
 	ENGINE_API static void OutputAnimNodeStart(const FAnimationBaseContext& InContext, uint64 InStartCycle, int32 InPreviousNodeId, int32 InNodeId, float InBlendWeight, float InRootMotionWeight, uint8 InPhase);
 	ENGINE_API static void OutputAnimNodeEnd(const FAnimationBaseContext& InContext, uint64 InEndCycle);
 
+	/** Output the current set of attributes in the supplied context */
+	ENGINE_API static void OutputAnimNodeAttribute(const FAnimInstanceProxy& InTargetProxy, const FAnimInstanceProxy& InSourceProxy, int32 InTargetNodeId, int32 InSourceNodeId, FName InAttribute);
+
+	/** Output the current set of attributes in the supplied context */
+	template<typename ContextType>
+	static void OutputAnimNodeBlendableAttributes(const ContextType& InContext, int32 InTargetNodeId, int32 InSourceNodeId)
+	{
+		if(InContext.CustomAttributes.ContainsData())
+		{
+			OutputAnimNodeAttribute(*InContext.AnimInstanceProxy, *InContext.AnimInstanceProxy, InTargetNodeId, InSourceNodeId, UE::Anim::FAttributes::Attributes);
+		}
+
+		if(InContext.Curve.NumValid() > 0)
+		{
+			OutputAnimNodeAttribute(*InContext.AnimInstanceProxy, *InContext.AnimInstanceProxy, InTargetNodeId, InSourceNodeId, UE::Anim::FAttributes::Curves);
+		}
+	}
+
 	/** Helper function to output a tracked value for an anim node */
 	ENGINE_API static void OutputAnimNodeValue(const FAnimationBaseContext& InContext, const TCHAR* InKey, bool InValue);
 	ENGINE_API static void OutputAnimNodeValue(const FAnimationBaseContext& InContext, const TCHAR* InKey, int32 InValue);
@@ -165,6 +185,9 @@ struct FAnimTrace
 
 	/** Helper function to output a montage instance's info */
 	ENGINE_API static void OutputMontage(UAnimInstance* InAnimInstance, const FAnimMontageInstance& InMontageInstance);
+
+	/** Helper function to output a sync record */
+	ENGINE_API static void OutputSync(const FAnimInstanceProxy& InSourceProxy, int32 InSourceNodeId, FName InGroupName);
 };
 
 #define TRACE_ANIM_TICK_RECORD(Context, TickRecord) \
@@ -185,6 +208,12 @@ struct FAnimTrace
 #define TRACE_SCOPED_ANIM_NODE(Context) \
 	FAnimTrace::FScopedAnimNodeTrace _ScopedAnimNodeTrace(Context);
 
+#define TRACE_ANIM_NODE_ATTRIBUTE(TargetProxy, SourceProxy, TargetNodeId, SourceNodeId, Name) \
+	FAnimTrace::OutputAnimNodeAttribute(TargetProxy, SourceProxy, TargetNodeId, SourceNodeId, Name);
+
+#define TRACE_ANIM_NODE_BLENDABLE_ATTRIBUTES(Context, TargetNodeId, SourceNodeId) \
+	FAnimTrace::OutputAnimNodeBlendableAttributes(Context, TargetNodeId, SourceNodeId);
+
 #define TRACE_SCOPED_ANIM_NODE_SUSPEND \
 	FAnimTrace::FScopedAnimNodeTraceSuspend _ScopedAnimNodeTraceSuspend;
 
@@ -206,6 +235,9 @@ struct FAnimTrace
 #define TRACE_ANIM_MONTAGE(AnimInstance, MontageInstance) \
 	FAnimTrace::OutputMontage(AnimInstance, MontageInstance);
 
+#define TRACE_ANIM_NODE_SYNC(SourceProxy, SourceNodeId, GroupName) \
+	FAnimTrace::OutputSync(SourceProxy, SourceNodeId, GroupName);
+
 #else
 
 #define TRACE_ANIM_TICK_RECORD(Context, TickRecord)
@@ -214,6 +246,8 @@ struct FAnimTrace
 #define TRACE_SKELETALMESH_FRAME(Component)
 #define TRACE_SCOPED_ANIM_GRAPH(Context)
 #define TRACE_SCOPED_ANIM_NODE(Context)
+#define TRACE_ANIM_NODE_ATTRIBUTE(TargetProxy, SourceProxy, TargetNodeId, SourceNodeId, Name)
+#define TRACE_ANIM_NODE_BLENDABLE_ATTRIBUTES(Context, TargetNodeId, SourceNodeId)
 #define TRACE_SCOPED_ANIM_NODE_SUSPEND
 #define TRACE_ANIM_NODE_VALUE(Context, Key, Value)
 #define TRACE_ANIM_SEQUENCE_PLAYER(Context, Node)
@@ -221,5 +255,6 @@ struct FAnimTrace
 #define TRACE_ANIM_NOTIFY(AnimInstance, NotifyEvent, EventType)
 #define TRACE_ANIM_SYNC_MARKER(AnimInstance, SyncMarker)
 #define TRACE_ANIM_MONTAGE(AnimInstance, MontageInstance)
+#define TRACE_ANIM_NODE_SYNC(SourceProxy, SourceNodeId, GroupName)
 
 #endif
