Updated callers of CachedDataProbablyExists to make batch requests when possible

diff --git a/Engine/Plugins/Runtime/HairStrands/Source/HairStrandsCore/Private/GroomAsset.cpp b/Engine/Plugins/Runtime/HairStrands/Source/HairStrandsCore/Private/GroomAsset.cpp
--- a/Engine/Plugins/Runtime/HairStrands/Source/HairStrandsCore/Private/GroomAsset.cpp
+++ b/Engine/Plugins/Runtime/HairStrands/Source/HairStrandsCore/Private/GroomAsset.cpp
@@ -1291,12 +1291,12 @@ namespace GroomDerivedDataCacheUtils
 bool UGroomAsset::IsFullyCached()
 {
 	// Check if all the groom derived data for strands, cards and meshes are already stored in the DDC
-	bool bIsFullyCached = true;
+	TArray<FString> CacheKeys;
 	const uint32 GroupCount = HairGroupsInterpolation.Num();
-	for (uint32 GroupIndex = 0; GroupIndex < GroupCount && bIsFullyCached; ++GroupIndex)
+	for (uint32 GroupIndex = 0; GroupIndex < GroupCount; ++GroupIndex)
 	{
 		const FString StrandsDDCKey = UGroomAsset::GetDerivedDataKeyForStrands(GroupIndex);
-		bIsFullyCached &= GetDerivedDataCacheRef().CachedDataProbablyExists(*StrandsDDCKey);
+		CacheKeys.Add(StrandsDDCKey);
 
 		// Some cards and meshes LOD settings may not produce any derived data so those must be excluded
 		bool bHasCardsLODs = false;
@@ -1322,18 +1322,18 @@ bool UGroomAsset::IsFullyCached()
 		{
 			const FString CardsKeySuffix = GroomDerivedDataCacheUtils::BuildCardsDerivedDataKeySuffix(GroupIndex, HairGroupsLOD[GroupIndex].LODs, HairGroupsCards);
 			const FString CardsDDCKey = StrandsDDCKey + CardsKeySuffix;
-			bIsFullyCached &= GetDerivedDataCacheRef().CachedDataProbablyExists(*CardsDDCKey);
+			CacheKeys.Add(CardsDDCKey);
 		}
 
 		if (bHasMeshesLODs)
 		{
 			const FString MeshesKeySuffix = GroomDerivedDataCacheUtils::BuildMeshesDerivedDataKeySuffix(GroupIndex, HairGroupsLOD[GroupIndex].LODs, HairGroupsMeshes);
 			const FString MeshesDDCKey = GroomDerivedDataCacheUtils::BuildGroomDerivedDataKey(MeshesKeySuffix);
-			bIsFullyCached &= GetDerivedDataCacheRef().CachedDataProbablyExists(*MeshesDDCKey);
+			CacheKeys.Add(MeshesDDCKey);
 		}
 	}
 
-	return bIsFullyCached;
+	return GetDerivedDataCacheRef().AllCachedDataProbablyExists(CacheKeys);
 }
 
 FString UGroomAsset::BuildDerivedDataKeySuffix(uint32 GroupIndex, const FHairGroupsInterpolation& InterpolationSettings, const FHairGroupsLOD& LODSettings) const
diff --git a/Engine/Source/Runtime/Engine/Private/AudioDerivedData.cpp b/Engine/Source/Runtime/Engine/Private/AudioDerivedData.cpp
--- a/Engine/Source/Runtime/Engine/Private/AudioDerivedData.cpp
+++ b/Engine/Source/Runtime/Engine/Private/AudioDerivedData.cpp
@@ -878,17 +878,15 @@ int32 FStreamedAudioPlatformData::GetChunkFromDDC(int32 ChunkIndex, uint8** OutC
 #if WITH_EDITORONLY_DATA
 bool FStreamedAudioPlatformData::AreDerivedChunksAvailable() const
 {
-	bool bChunksAvailable = true;
-	FDerivedDataCacheInterface& DDC = GetDerivedDataCacheRef();
-	for (int32 ChunkIndex = 0; bChunksAvailable && ChunkIndex < Chunks.Num(); ++ChunkIndex)
+	TArray<FString> ChunkKeys;
+	for (const FStreamedAudioChunk& Chunk : Chunks)
 	{
-		const FStreamedAudioChunk& Chunk = Chunks[ChunkIndex];
-		if (Chunk.DerivedDataKey.IsEmpty() == false)
+		if (!Chunk.DerivedDataKey.IsEmpty())
 		{
-			bChunksAvailable = DDC.CachedDataProbablyExists(*Chunk.DerivedDataKey);
+			ChunkKeys.Add(Chunk.DerivedDataKey);
 		}
 	}
-	return bChunksAvailable;
+	return GetDerivedDataCacheRef().AllCachedDataProbablyExists(ChunkKeys);
 }
 #endif // #if WITH_EDITORONLY_DATA
 
diff --git a/Engine/Source/Runtime/Engine/Private/TextureDerivedData.cpp b/Engine/Source/Runtime/Engine/Private/TextureDerivedData.cpp
--- a/Engine/Source/Runtime/Engine/Private/TextureDerivedData.cpp
+++ b/Engine/Source/Runtime/Engine/Private/TextureDerivedData.cpp
@@ -1369,32 +1369,28 @@ EPixelFormat FTexturePlatformData::GetLayerPixelFormat(uint32 LayerIndex) const
 #if WITH_EDITOR
 bool FTexturePlatformData::AreDerivedMipsAvailable() const
 {
-	bool bMipsAvailable = true;
-	FDerivedDataCacheInterface& DDC = GetDerivedDataCacheRef();
-	for (int32 MipIndex = 0; bMipsAvailable && MipIndex < Mips.Num(); ++MipIndex)
+	TArray<FString, TInlineAllocator<16>> MipKeys;
+	for (const FTexture2DMipMap& Mip : Mips)
 	{
-		const FTexture2DMipMap& Mip = Mips[MipIndex];
-		if (Mip.DerivedDataKey.IsEmpty() == false)
+		if (!Mip.DerivedDataKey.IsEmpty())
 		{
-			bMipsAvailable = DDC.CachedDataProbablyExists(*Mip.DerivedDataKey);
+			MipKeys.Add(Mip.DerivedDataKey);
 		}
 	}
-	return bMipsAvailable;
+	return GetDerivedDataCacheRef().AllCachedDataProbablyExists(MipKeys);
 }
 bool FTexturePlatformData::AreDerivedVTChunksAvailable() const
 {
-	bool bChunksAvailable = true;
-	FDerivedDataCacheInterface& DDC = GetDerivedDataCacheRef();
 	check(VTData);
-	for (int32 ChunkIndex = 0; bChunksAvailable && ChunkIndex < VTData->Chunks.Num(); ++ChunkIndex)
+	TArray<FString, TInlineAllocator<16>> ChunkKeys;
+	for (const FVirtualTextureDataChunk& Chunk : VTData->Chunks)
 	{
-		const FVirtualTextureDataChunk& Chunk = VTData->Chunks[ChunkIndex];
-		if (Chunk.DerivedDataKey.IsEmpty() == false)
+		if (!Chunk.DerivedDataKey.IsEmpty())
 		{
-			bChunksAvailable = DDC.CachedDataProbablyExists(*Chunk.DerivedDataKey);
+			ChunkKeys.Add(Chunk.DerivedDataKey);
 		}
 	}
-	return bChunksAvailable;
+	return GetDerivedDataCacheRef().AllCachedDataProbablyExists(ChunkKeys);
 }
 #endif // #if WITH_EDITOR
 
@@ -1787,42 +1783,52 @@ void UTexture::BeginCacheForCookedPlatformData( const ITargetPlatform *TargetPla
 			CacheFlags |= ETextureCacheFlags::AllowAsyncBuild;
 		}
 
-
 		// Cull redundant settings by comparing derived data keys.
-		for (int32 SettingsIndex = 0; SettingsIndex < BuildSettingsToCache.Num(); ++SettingsIndex)
+		TArray<FString> BuildSettingsCacheKeys;
+		for (TArray<TArray<FTextureBuildSettings>>::TIterator It = BuildSettingsToCache.CreateIterator(); It; ++It)
 		{
-			check(BuildSettingsToCache[SettingsIndex].Num() == Source.GetNumLayers());
+			check(It->Num() == Source.GetNumLayers());
 
 			FString DerivedDataKey;
-			GetTextureDerivedDataKey(*this, BuildSettingsToCache[SettingsIndex].GetData(), DerivedDataKey);
-
-			FTexturePlatformData *PlatformData = CookedPlatformData.FindRef( DerivedDataKey );
+			GetTextureDerivedDataKey(*this, It->GetData(), DerivedDataKey);
 
-			if ( PlatformData == NULL )
+			if (CookedPlatformData.FindRef(DerivedDataKey) || BuildSettingsCacheKeys.Contains(DerivedDataKey))
 			{
-				// UE_LOG(LogTemp, Warning, TEXT("Caching data for texture %s with id %s"), *GetName(), *DerivedDataKey);
-				uint32 CurrentCacheFlags = CacheFlags;
-				// if the cached data key exists already then we don't need to allowasync build
-				// if it doesn't then allow async builds
-				
-				{
-					if ( GetDerivedDataCacheRef().CachedDataProbablyExists( *DerivedDataKey ) == false )
-					{
-						CurrentCacheFlags |= ETextureCacheFlags::AllowAsyncBuild;
-						CurrentCacheFlags |= ETextureCacheFlags::AllowAsyncLoading;
-					}
-				}
+				It.RemoveCurrent();
+			}
+			else
+			{
+				BuildSettingsCacheKeys.Add(MoveTemp(DerivedDataKey));
+			}
+		}
+
+		TBitArray<> BuildSettingsProbablyCached;
+		if (BuildSettingsCacheKeys.Num())
+		{
+			BuildSettingsProbablyCached = GetDerivedDataCacheRef().CachedDataProbablyExistsBatch(BuildSettingsCacheKeys);
+		}
 
-				FTexturePlatformData* PlatformDataToCache;
-				PlatformDataToCache = new FTexturePlatformData();
-				PlatformDataToCache->Cache(
-					*this,
-					BuildSettingsToCache[SettingsIndex].GetData(),
-					CurrentCacheFlags,
-					nullptr
-					);
-				CookedPlatformData.Add( DerivedDataKey, PlatformDataToCache );
+		for (int32 SettingsIndex = 0; SettingsIndex < BuildSettingsToCache.Num(); ++SettingsIndex)
+		{
+			// UE_LOG(LogTemp, Warning, TEXT("Caching data for texture %s with id %s"), *GetName(), *DerivedDataKey);
+			uint32 CurrentCacheFlags = CacheFlags;
+			// if the cached data key exists already then we don't need to allowasync build
+			// if it doesn't then allow async builds
+			if (!BuildSettingsProbablyCached[SettingsIndex])
+			{
+				CurrentCacheFlags |= ETextureCacheFlags::AllowAsyncBuild;
+				CurrentCacheFlags |= ETextureCacheFlags::AllowAsyncLoading;
 			}
+
+			FTexturePlatformData* PlatformDataToCache;
+			PlatformDataToCache = new FTexturePlatformData();
+			PlatformDataToCache->Cache(
+				*this,
+				BuildSettingsToCache[SettingsIndex].GetData(),
+				CurrentCacheFlags,
+				nullptr
+				);
+			CookedPlatformData.Add(BuildSettingsCacheKeys[SettingsIndex], PlatformDataToCache);
 		}
 	}
 }
