Make sure HLOD assets are outered to the HLOD actors to avoid cross-level references when generating runtime streaming cells.

diff --git a/Engine/Source/Runtime/Engine/Private/WorldPartition/HLOD/HLODBuilder.cpp b/Engine/Source/Runtime/Engine/Private/WorldPartition/HLOD/HLODBuilder.cpp
--- a/Engine/Source/Runtime/Engine/Private/WorldPartition/HLOD/HLODBuilder.cpp
+++ b/Engine/Source/Runtime/Engine/Private/WorldPartition/HLOD/HLODBuilder.cpp
@@ -204,17 +204,19 @@ class FHLODBuilder_MeshMerge : public FHLODBuilder
 			const IMeshMergeUtilities& MeshMergeUtilities = FModuleManager::Get().LoadModuleChecked<IMeshMergeModule>("MeshMergeUtilities").GetUtilities();
 			MeshMergeUtilities.MergeComponentsToStaticMesh(InSubComponents, HLODActor->GetWorld(), HLODLayer->GetMeshMergeSettings(), HLODLayer->GetHLODMaterial().LoadSynchronous(), HLODActor->GetPackage(), CellName.ToString(), Assets, MergedActorLocation, 0.25f, false);
 
-			// All merged mesh assets are stored in the HLOD Actor package
-			Algo::ForEach(Assets, [](UObject* Asset) { Asset->ClearFlags(RF_Public | RF_Standalone); });
-
-			UStaticMesh* StaticMesh = nullptr;
 			UStaticMeshComponent* Component = nullptr;
-			if (Assets.FindItemByClass<UStaticMesh>(&StaticMesh))
+			Algo::ForEach(Assets, [HLODActor, &Component, &MergedActorLocation](UObject* Asset)
 			{
-				Component = NewObject<UStaticMeshComponent>(HLODActor);
-				Component->SetStaticMesh(StaticMesh);
-				Component->SetWorldLocation(MergedActorLocation);
-			}
+				Asset->ClearFlags(RF_Public | RF_Standalone);
+				Asset->Rename(nullptr, HLODActor);
+
+				if (Cast<UStaticMesh>(Asset))
+				{
+					Component = NewObject<UStaticMeshComponent>(HLODActor);
+					Component->SetStaticMesh(static_cast<UStaticMesh*>(Asset));
+					Component->SetWorldLocation(MergedActorLocation);
+				}
+			});
 
 			return TArray<UPrimitiveComponent*>({ Component });
 		};
@@ -244,16 +246,18 @@ class FHLODBuilder_MeshSimplify : public FHLODBuilder
 			const IMeshMergeUtilities& MeshMergeUtilities = FModuleManager::Get().LoadModuleChecked<IMeshMergeModule>("MeshMergeUtilities").GetUtilities();
 			MeshMergeUtilities.CreateProxyMesh(StaticMeshComponents, HLODLayer->GetMeshSimplifySettings(), HLODLayer->GetHLODMaterial().LoadSynchronous(), HLODActor->GetPackage(), CellName.ToString(), FGuid::NewGuid(), ProxyDelegate, true);
 
-			// All merged mesh assets are stored in the HLOD Actor package
-			Algo::ForEach(Assets, [](UObject* Asset) { Asset->ClearFlags(RF_Public | RF_Standalone); });
-
-			UStaticMesh* StaticMesh = nullptr;
 			UStaticMeshComponent* Component = nullptr;
-			if (Assets.FindItemByClass<UStaticMesh>(&StaticMesh))
+			Algo::ForEach(Assets, [HLODActor, &Component](UObject* Asset)
 			{
-				Component = NewObject<UStaticMeshComponent>(HLODActor);
-				Component->SetStaticMesh(StaticMesh);
-			}
+				Asset->ClearFlags(RF_Public | RF_Standalone);
+				Asset->Rename(nullptr, HLODActor);
+
+				if (Cast<UStaticMesh>(Asset))
+				{
+					Component = NewObject<UStaticMeshComponent>(HLODActor);
+					Component->SetStaticMesh(static_cast<UStaticMesh*>(Asset));
+				}
+			});
 
 			return TArray<UPrimitiveComponent*>({ Component });
 		};
