WorldPartition: Avoid getting registry tags when Actors outer level isn't partitionned (fix crash with OFPA LevelInstances)

diff --git a/Engine/Source/Runtime/Engine/Private/Actor.cpp b/Engine/Source/Runtime/Engine/Private/Actor.cpp
--- a/Engine/Source/Runtime/Engine/Private/Actor.cpp
+++ b/Engine/Source/Runtime/Engine/Private/Actor.cpp
@@ -418,7 +418,7 @@ void AActor::GetAssetRegistryTags(TArray<FAssetRegistryTag>& OutTags) const
 
 	if (IsPackageExternal() && !IsChildActor())
 	{
-		if (UWorldPartition* WorldPartition = GetWorldSettings()->GetWorldPartition())
+		if (UWorldPartition* WorldPartition = GetLevel()->GetWorldPartition())
 		{
 			TUniquePtr<FWorldPartitionActorDesc> ActorDesc(CreateActorDesc());
 			
