Changing more shader compilers to export shader debug data in a more unified manner.

diff --git a/Engine/Source/Editor/UnrealEd/Private/Commandlets/CookGlobalShaderCommandlet.cpp b/Engine/Source/Editor/UnrealEd/Private/Commandlets/CookGlobalShaderCommandlet.cpp
--- a/Engine/Source/Editor/UnrealEd/Private/Commandlets/CookGlobalShaderCommandlet.cpp
+++ b/Engine/Source/Editor/UnrealEd/Private/Commandlets/CookGlobalShaderCommandlet.cpp
@@ -76,13 +76,6 @@ int32 UCookGlobalShadersCommandlet::Main(const FString& Params)
 		TargetPlatform->RefreshSettings();
 	}
 
-	// Set PDB path
-	FString ShaderPDBPath;
-	if ( FParse::Value(*Params, TEXT("shaderpdb="), ShaderPDBPath, true) )
-	{
-		GConfig->SetString(TEXT("DevOptions.Shaders"), TEXT("ShaderPDBRoot"), *ShaderPDBPath, GEngineIni);
-	}
-
 	// Get target device
 	ITargetDevicePtr TargetDevice;
 	FString TargetDeviceName;
diff --git a/Engine/Source/Runtime/Engine/Private/ShaderCompiler/ShaderCompiler.cpp b/Engine/Source/Runtime/Engine/Private/ShaderCompiler/ShaderCompiler.cpp
--- a/Engine/Source/Runtime/Engine/Private/ShaderCompiler/ShaderCompiler.cpp
+++ b/Engine/Source/Runtime/Engine/Private/ShaderCompiler/ShaderCompiler.cpp
@@ -4332,15 +4332,6 @@ void GlobalBeginCompileShader(
 
 	Input.Environment.SetDefine(TEXT("HAS_INVERTED_Z_BUFFER"), (bool)ERHIZBuffer::IsInverted);
 
-	{
-		FString ShaderPDBRoot;
-		GConfig->GetString(TEXT("DevOptions.Shaders"), TEXT("ShaderPDBRoot"), ShaderPDBRoot, GEngineIni);
-		if (!ShaderPDBRoot.IsEmpty())
-		{
-			Input.Environment.SetDefine(TEXT("SHADER_PDB_ROOT"), ShaderPDBRoot);
-		}
-	}
-
 	{
 		static const auto CVar = IConsoleManager::Get().FindTConsoleVariableDataInt(TEXT("r.ClearCoatNormal"));
 		Input.Environment.SetDefine(TEXT("CLEAR_COAT_BOTTOM_NORMAL"), CVar ? (CVar->GetValueOnAnyThread() != 0) : 0);
