Removed ensure from APlayerState::OnRep_bIsInactive

diff --git a/Engine/Source/Runtime/Engine/Private/PlayerState.cpp b/Engine/Source/Runtime/Engine/Private/PlayerState.cpp
--- a/Engine/Source/Runtime/Engine/Private/PlayerState.cpp
+++ b/Engine/Source/Runtime/Engine/Private/PlayerState.cpp
@@ -224,7 +224,7 @@ void APlayerState::OnRep_bIsInactive()
 {
 	// remove and re-add from the GameState so it's in the right list  
 	UWorld* World = GetWorld();
-	if (ensure(World && World->GetGameState()))
+	if (World && World->GetGameState())
 	{
 		World->GetGameState()->RemovePlayerState(this);
 		World->GetGameState()->AddPlayerState(this);
