Fixed ensure/check when running a skeletal mesh with only a post process instance

Syncing code now handles this case. Previously the code assumed that we would always have a main instance proxy

diff --git a/Engine/Source/Runtime/Engine/Private/Animation/AnimInstanceProxy.cpp b/Engine/Source/Runtime/Engine/Private/Animation/AnimInstanceProxy.cpp
--- a/Engine/Source/Runtime/Engine/Private/Animation/AnimInstanceProxy.cpp
+++ b/Engine/Source/Runtime/Engine/Private/Animation/AnimInstanceProxy.cpp
@@ -1093,8 +1093,9 @@ void FAnimInstanceProxy::UpdateAnimation_WithRoot(const FAnimationUpdateContext&
 		TGuardValue<bool> ScopeGuard(bUpdatingRoot, true);
 
 		// Anything syncing within this scope is subject to sync groups.
-		// We only enable syncing here for the main proxy
-		UE::Anim::TOptionalScopedGraphMessage<UE::Anim::FAnimSyncGroupScope> Message(GetMainInstanceProxy() == this, InContext, InContext);
+		// We only enable syncing here for the main proxy or if we have no main proxy
+		const bool bEnableSyncScope = (GetMainInstanceProxy() == nullptr) || (GetMainInstanceProxy() == this);
+		UE::Anim::TOptionalScopedGraphMessage<UE::Anim::FAnimSyncGroupScope> Message(bEnableSyncScope, InContext, InContext);
 
 		// update all nodes
 		if(InRootNode == RootNode)
diff --git a/Engine/Source/Runtime/Engine/Private/Animation/AnimSyncScope.cpp b/Engine/Source/Runtime/Engine/Private/Animation/AnimSyncScope.cpp
--- a/Engine/Source/Runtime/Engine/Private/Animation/AnimSyncScope.cpp
+++ b/Engine/Source/Runtime/Engine/Private/Animation/AnimSyncScope.cpp
@@ -54,9 +54,15 @@ void FAnimSyncGroupScope::AddTickRecord(const FAnimTickRecord& InTickRecord, con
 		break;
 	}
 
-	// Forward to main instance
-	check(Proxy.GetMainInstanceProxy());
-	Proxy.GetMainInstanceProxy()->AddTickRecord(InTickRecord, NewSyncParams);
+	// Forward to main instance if we have one
+	if(Proxy.GetMainInstanceProxy())
+	{
+		Proxy.GetMainInstanceProxy()->AddTickRecord(InTickRecord, NewSyncParams);
+	}
+	else
+	{
+		Proxy.AddTickRecord(InTickRecord, NewSyncParams);
+	}
 }
 
 }}	// namespace UE::Anim
\ No newline at end of file
diff --git a/Engine/Source/Runtime/Engine/Public/Animation/AnimInstanceProxy.h b/Engine/Source/Runtime/Engine/Public/Animation/AnimInstanceProxy.h
--- a/Engine/Source/Runtime/Engine/Public/Animation/AnimInstanceProxy.h
+++ b/Engine/Source/Runtime/Engine/Public/Animation/AnimInstanceProxy.h
@@ -373,13 +373,11 @@ public:
 		return SkeletalMeshComponent; 
 	}
 
-	/** Get the current main instance proxy. Note that this will return nullptr outside of pre/post update */
+	/** Get the current main instance proxy. Note that this will return nullptr outside of pre/post update, and may return nullptr anyway if we dont have a main instance */
 	FAnimInstanceProxy* GetMainInstanceProxy() const
 	{ 
 		// Main instance proxy is only available during update/eval. If you're calling this function outside of it, it will return null. 
-		// adding ensure here so that we can catch them earlier
-		ensureAlways(MainInstanceProxy);
-		return MainInstanceProxy; 
+		return MainInstanceProxy;
 	}
 
 	UE_DEPRECATED(4.26, "Please use the overload that takes a group FName")
