Cleaned up minimap rendering code + disabled texture streaming to avoid blurry effect on update.

diff --git a/Engine/Source/Runtime/Engine/Private/WorldPartition/WorldPartitionMiniMapHelper.cpp b/Engine/Source/Runtime/Engine/Private/WorldPartition/WorldPartitionMiniMapHelper.cpp
--- a/Engine/Source/Runtime/Engine/Private/WorldPartition/WorldPartitionMiniMapHelper.cpp
+++ b/Engine/Source/Runtime/Engine/Private/WorldPartition/WorldPartitionMiniMapHelper.cpp
@@ -99,44 +99,43 @@ void FWorldPartitionMiniMapHelper::CaptureWorldMiniMapToTexture(UWorld* InWorld,
 
 	TArray<FColor> CapturedImage;
 	//Using SceneCapture Actor capture the scene to buffer
-	{
-		UTextureRenderTarget2D* RenderTargetTexture = NewObject<UTextureRenderTarget2D>();
-		RenderTargetTexture->AddToRoot();
-		RenderTargetTexture->ClearColor = FLinearColor::Transparent;
-		RenderTargetTexture->TargetGamma = 2.2f;
-		RenderTargetTexture->InitCustomFormat(InMiniMapSize, InMiniMapSize, PF_B8G8R8A8, false);
+	UTextureRenderTarget2D* RenderTargetTexture = NewObject<UTextureRenderTarget2D>();
+	RenderTargetTexture->ClearColor = FLinearColor::Transparent;
+	RenderTargetTexture->TargetGamma = 2.2f;
+	RenderTargetTexture->InitCustomFormat(InMiniMapSize, InMiniMapSize, PF_B8G8R8A8, false);
+	RenderTargetTexture->UpdateResourceImmediate(true);
+
+	FActorSpawnParameters SpawnInfo;
+	SpawnInfo.ObjectFlags |= RF_Transient;
+	FVector CaptureActorLocation = FVector(OutWorldBounds.GetCenter().X, OutWorldBounds.GetCenter().Y, OutWorldBounds.GetCenter().Z + OutWorldBounds.GetExtent().Z);
+	FRotator CaptureActorRotation = FRotator(-90.f, 0.f, -90.f);
+	ASceneCapture2D* CaptureActor = InWorld->SpawnActor<ASceneCapture2D>(CaptureActorLocation, CaptureActorRotation, SpawnInfo);
+	auto CaptureComponent = CaptureActor->GetCaptureComponent2D();
+	CaptureComponent->TextureTarget = RenderTargetTexture;
+	CaptureComponent->ProjectionType = ECameraProjectionMode::Orthographic;
+	CaptureComponent->OrthoWidth = ViewportWidth;
+	CaptureComponent->bUseCustomProjectionMatrix = true;
+	CaptureComponent->CustomProjectionMatrix = ProjectionMatrix;
+	CaptureComponent->bCaptureEveryFrame = false;
+	CaptureComponent->bCaptureOnMovement = false;
+	CaptureComponent->CaptureScene();
+
+	InWorld->DestroyActor(CaptureActor);
+	CaptureActor = nullptr;	
 
-		FActorSpawnParameters SpawnInfo;
-		SpawnInfo.ObjectFlags |= RF_Transient;
-		FVector CaptureActorLocation = FVector(OutWorldBounds.GetCenter().X, OutWorldBounds.GetCenter().Y, OutWorldBounds.GetCenter().Z + OutWorldBounds.GetExtent().Z);
-		FRotator CaptureActorRotation = FRotator(-90.f, 0.f, -90.f);
-		ASceneCapture2D* CaptureActor = InWorld->SpawnActor<ASceneCapture2D>(CaptureActorLocation, CaptureActorRotation, SpawnInfo);
-		auto CaptureComponent = CaptureActor->GetCaptureComponent2D();
-		CaptureComponent->TextureTarget = RenderTargetTexture;
-		CaptureComponent->ProjectionType = ECameraProjectionMode::Orthographic;
-		CaptureComponent->OrthoWidth = ViewportWidth;
-		CaptureComponent->bUseCustomProjectionMatrix = true;
-		CaptureComponent->CustomProjectionMatrix = ProjectionMatrix;
-		CaptureComponent->bCaptureEveryFrame = false;
-		CaptureComponent->bCaptureOnMovement = false;
-		CaptureComponent->CaptureScene();
-
-		auto RenderTargetResource = RenderTargetTexture->GameThread_GetRenderTargetResource();
-		CapturedImage.SetNumUninitialized(InMiniMapSize * InMiniMapSize);
-		RenderTargetResource->ReadPixelsPtr(CapturedImage.GetData(), FReadSurfaceDataFlags(), FIntRect(0, 0, InMiniMapSize, InMiniMapSize));
-		InWorld->DestroyActor(CaptureActor);
-		CaptureActor = nullptr;
-	}
-	
 	//Update the output texture
 	if (!InOutMiniMapTexture)
 	{
-		InOutMiniMapTexture = NewObject<UTexture2D>(InOuterForTexture, "MiniMapTexture");
+		InOutMiniMapTexture = RenderTargetTexture->ConstructTexture2D(InOuterForTexture, "MiniMapTexture", RF_NoFlags, CTF_Default, NULL);
 	}
+	else
+	{
+		RenderTargetTexture->UpdateTexture2D(InOutMiniMapTexture, TSF_BGRA8, CTF_Default);
+	}
+
 	InOutMiniMapTexture->AdjustMinAlpha = 1.f;
-	InOutMiniMapTexture->Source.Init(InMiniMapSize, InMiniMapSize, 1, 1, TSF_BGRA8, (uint8*)CapturedImage.GetData());
+	InOutMiniMapTexture->LODGroup = TEXTUREGROUP_UI;
 	InOutMiniMapTexture->UpdateResource();
-
 }
 
 bool FWorldPartitionMiniMapHelper::DoesActorContributeToBounds(AActor* Actor)
