Make sure the duplicate archive set the cooking target platform. This fix the engine editor python tests

diff --git a/Engine/Source/Runtime/Engine/Private/SkeletalMesh.cpp b/Engine/Source/Runtime/Engine/Private/SkeletalMesh.cpp
--- a/Engine/Source/Runtime/Engine/Private/SkeletalMesh.cpp
+++ b/Engine/Source/Runtime/Engine/Private/SkeletalMesh.cpp
@@ -1539,7 +1539,19 @@ void USkeletalMesh::Serialize( FArchive& Ar )
 			{
 				FSkeletalMeshRenderData* LocalSkeletalMeshRenderData = SkeletalMeshRenderData.Get();
 #if WITH_EDITORONLY_DATA
-				LocalSkeletalMeshRenderData = &GetPlatformSkeletalMeshRenderData(this, Ar.CookingTarget());
+				const ITargetPlatform* ArchiveCookingTarget = Ar.CookingTarget();
+				if (ArchiveCookingTarget)
+				{
+					LocalSkeletalMeshRenderData = &GetPlatformSkeletalMeshRenderData(this, ArchiveCookingTarget);
+				}
+				else
+				{
+					//Fall back in case we use an archive that the cooking target has not been set (i.e. Duplicate archive)
+					ITargetPlatformManagerModule& TargetPlatformManager = GetTargetPlatformManagerRef();
+					ITargetPlatform* RunningPlatform = TargetPlatformManager.GetRunningTargetPlatform();
+					check(RunningPlatform != NULL);
+					LocalSkeletalMeshRenderData = &GetPlatformSkeletalMeshRenderData(this, RunningPlatform);
+				}
 #endif
 				if (bCooked)
 				{
