Reduce redundant code between the different async compilers
  - Unified Notification Manager for all async assets
  - Unified CVars registration for all async assets
  - Unified Stall stacks dumps so any async compilation wait can use it
  - Unified the FinishCompilation logic for all asset compilers
  - Texture wait on async compilation will now be present when calling Editor.AsyncAssetDumpStallStacks

diff --git a/Engine/Source/Runtime/Core/Public/Misc/QueuedThreadPool.h b/Engine/Source/Runtime/Core/Public/Misc/QueuedThreadPool.h
--- a/Engine/Source/Runtime/Core/Public/Misc/QueuedThreadPool.h
+++ b/Engine/Source/Runtime/Core/Public/Misc/QueuedThreadPool.h
@@ -18,6 +18,26 @@ enum class EQueuedWorkPriority : uint8
 	Count
 };
 
+constexpr inline const TCHAR* LexToString(EQueuedWorkPriority Priority)
+{
+	switch (Priority)
+	{
+	case EQueuedWorkPriority::Highest:
+		return TEXT("Highest");
+	case EQueuedWorkPriority::High:
+		return TEXT("High");
+	case EQueuedWorkPriority::Normal:
+		return TEXT("Normal");
+	case EQueuedWorkPriority::Low:
+		return TEXT("Low");
+	case EQueuedWorkPriority::Lowest:
+		return TEXT("Lowest");
+	default:
+		check(false);
+		return TEXT("Unknown");
+	}
+}
+
 /** 
  *  Priority Queue tailored for FQueuedThreadPool implementation
  *
diff --git a/Engine/Source/Runtime/Engine/Private/AssetCompilingManager.cpp b/Engine/Source/Runtime/Engine/Private/AssetCompilingManager.cpp
--- a/Engine/Source/Runtime/Engine/Private/AssetCompilingManager.cpp
+++ b/Engine/Source/Runtime/Engine/Private/AssetCompilingManager.cpp
@@ -12,23 +12,19 @@
 #include "TextureCompiler.h"
 #include "Components/StaticMeshComponent.h"
 #include "Materials/MaterialInstance.h"
+#include "AsyncCompilationHelpers.h"
 
 #define LOCTEXT_NAMESPACE "AssetCompilingManager"
 
-static TAutoConsoleVariable<int32> CVarAsyncAssetCompilation(
-	TEXT("Editor.AsyncAssetCompilation"),
-	1,
-	TEXT("1 - Async asset compilation is enabled.\n")
-	TEXT("2 - Async asset compilation is enabled but on pause (for debugging).\n")
-	TEXT("When enabled, assetes will be replaced by placeholders until they are ready\n")
-	TEXT("to reduce stalls on the game thread and improve overall editor performance."),
-	ECVF_Default);
-
-static TAutoConsoleVariable<int32> CVarAsyncAssetCompilationMaxConcurrency(
-	TEXT("Editor.AsyncAssetCompilationMaxConcurrency"),
-	-1,
-	TEXT("Set the maximum number of concurrent asset compilation, -1 for unlimited."),
-	ECVF_Default);
+static AsyncCompilationHelpers::FAsyncCompilationStandardCVars CVarAsyncCompilationStandard(
+	TEXT("Asset"), 
+	TEXT("assets"),
+	FConsoleCommandDelegate::CreateLambda(
+		[]()
+		{
+			FAssetCompilingManager::Get().FinishAllCompilation();
+		}
+	));
 
 static TAutoConsoleVariable<int32> CVarAsyncAssetCompilationMemoryPerCore(
 	TEXT("Editor.AsyncAssetCompilationMemoryPerCore"),
@@ -38,21 +34,6 @@ static TAutoConsoleVariable<int32> CVarAsyncAssetCompilationMemoryPerCore(
 	TEXT("Limit concurrency for async processing based on RAM available.\n"),
 	ECVF_Default);
 
-static FAutoConsoleCommand CVarAsyncAssetCompilationFinishAll(
-	TEXT("Editor.AsyncAssetCompilationFinishAll"),
-	TEXT("Finish all asset compilations"),
-	FConsoleCommandWithArgsDelegate::CreateLambda([](const TArray<FString>& Args)
-	{
-		FAssetCompilingManager::Get().FinishAllCompilation();
-	})
-);
-
-static TAutoConsoleVariable<int32> CVarAsyncAssetCompilationResume(
-	TEXT("Editor.AsyncAssetCompilationResume"),
-	0,
-	TEXT("Number of queued work to resume while paused."),
-	ECVF_Default);
-
 namespace AssetCompilingManagerImpl
 {
 	static void EnsureInitializedCVars()
@@ -63,28 +44,11 @@ namespace AssetCompilingManagerImpl
 		{
 			bIsInitialized = true;
 
-			FString Value;
-			if (FParse::Value(FCommandLine::Get(), TEXT("-asyncassetcompilation="), Value))
-			{
-				int32 AsyncAssetCompilationValue = 0;
-				if (Value == TEXT("1") || Value == TEXT("on"))
-				{
-					AsyncAssetCompilationValue = 1;
-				}
-
-				if (Value == TEXT("2") || Value == TEXT("paused"))
-				{
-					AsyncAssetCompilationValue = 2;
-				}
-
-				CVarAsyncAssetCompilation->Set(AsyncAssetCompilationValue, ECVF_SetByCommandline);
-			}
-
-			int32 MaxConcurrency = -1;
-			if (FParse::Value(FCommandLine::Get(), TEXT("-asyncassetcompilationmaxconcurrency="), MaxConcurrency))
-			{
-				CVarAsyncAssetCompilationMaxConcurrency->Set(MaxConcurrency, ECVF_SetByCommandline);
-			}
+			AsyncCompilationHelpers::EnsureInitializedCVars(
+				TEXT("asset"),
+				CVarAsyncCompilationStandard.AsyncCompilation,
+				CVarAsyncCompilationStandard.AsyncCompilationMaxConcurrency
+			);
 		}
 	}
 }
@@ -139,50 +103,14 @@ FQueuedThreadPool* FAssetCompilingManager::GetThreadPool() const
 		// We're using GThreadPool instead of GLargeThreadPool because asset compilation is hard on total memory and memory bandwidth and can run slower when going wider than actual cores.
 		// All asset priorities will resolve to a Low priority once being scheduled in the LargeThreadPool.
 		// Any asset supporting being built async should be scheduled lower than Normal to let non-async stuff go first
-		const int32 MaxConcurrency = CVarAsyncAssetCompilationMaxConcurrency.GetValueOnAnyThread();
-		GAssetThreadPool = new FMemoryBoundQueuedThreadPoolWrapper(GThreadPool, MaxConcurrency, [](EQueuedWorkPriority) { return EQueuedWorkPriority::Low; });
-
-		CVarAsyncAssetCompilation->SetOnChangedCallback(
-			FConsoleVariableDelegate::CreateLambda(
-				[](IConsoleVariable* Variable)
-				{
-					if (Variable->GetInt() == 2)
-					{
-						GAssetThreadPool->Pause();
-					}
-					else
-					{
-						GAssetThreadPool->Resume();
-					}
-				}
-			)
-		);
-
-		CVarAsyncAssetCompilationResume->SetOnChangedCallback(
-			FConsoleVariableDelegate::CreateLambda(
-				[](IConsoleVariable* Variable)
-				{
-					if (Variable->GetInt() > 0)
-					{
-						GAssetThreadPool->Resume(Variable->GetInt());
-					}
-				}
-			)
-		);
+		GAssetThreadPool = new FMemoryBoundQueuedThreadPoolWrapper(GThreadPool, -1, [](EQueuedWorkPriority) { return EQueuedWorkPriority::Low; });
 
-		CVarAsyncAssetCompilationMaxConcurrency->SetOnChangedCallback(
-			FConsoleVariableDelegate::CreateLambda(
-				[](IConsoleVariable* Variable)
-				{
-					GAssetThreadPool->SetMaxConcurrency(Variable->GetInt());
-				}
-			)
+		AsyncCompilationHelpers::BindThreadPoolToCVar(
+			GAssetThreadPool,
+			CVarAsyncCompilationStandard.AsyncCompilation,
+			CVarAsyncCompilationStandard.AsyncCompilationResume,
+			CVarAsyncCompilationStandard.AsyncCompilationMaxConcurrency
 		);
-
-		if (CVarAsyncAssetCompilation->GetInt() == 2)
-		{
-			GAssetThreadPool->Pause();
-		}
 	}
 
 	return GAssetThreadPool;
@@ -193,8 +121,8 @@ FQueuedThreadPool* FAssetCompilingManager::GetThreadPool() const
  */
 void FAssetCompilingManager::FinishAllCompilation()
 {
-	FStaticMeshCompilingManager::Get().FinishAllCompilation();
 	FTextureCompilingManager::Get().FinishAllCompilation();
+	FStaticMeshCompilingManager::Get().FinishAllCompilation();
 }
 
 /**
diff --git a/Engine/Source/Runtime/Engine/Private/AsyncCompilationHelpers.cpp b/Engine/Source/Runtime/Engine/Private/AsyncCompilationHelpers.cpp
new file mode 100644
--- /dev/null
+++ b/Engine/Source/Runtime/Engine/Private/AsyncCompilationHelpers.cpp
@@ -0,0 +1,384 @@
+// Copyright Epic Games, Inc. All Rights Reserved.
+
+#include "AsyncCompilationHelpers.h"
+
+#if WITH_EDITOR
+
+#include "Framework/Notifications/NotificationManager.h"
+#include "Widgets/Notifications/SNotificationList.h"
+#include "Misc/ScopedSlowTask.h"
+#include "ObjectCacheContext.h"
+#include "Misc/CommandLine.h"
+#include "Misc/Parse.h"
+#include "Settings/EditorExperimentalSettings.h"
+#include "Misc/QueuedThreadPoolWrapper.h"
+#include "HAL/IConsoleManager.h"
+#include "HAL/PlatformStackWalk.h"
+#include "Hash/CityHash.h"
+#include "Misc/ScopeRWLock.h"
+
+#define LOCTEXT_NAMESPACE "AsyncCompilation"
+
+DEFINE_LOG_CATEGORY(LogAsyncCompilation);
+
+namespace AsyncCompilationHelpers
+{
+	void DumpStallStacks();
+}
+
+static FAutoConsoleCommand CVarAsyncAssetDumpStallStacks(
+	TEXT("Editor.AsyncAssetDumpStallStacks"),
+	TEXT("Dump all the callstacks that have caused waits on async compilation."),
+	FConsoleCommandWithArgsDelegate::CreateLambda([](const TArray<FString>& Args)
+	{
+		AsyncCompilationHelpers::DumpStallStacks();
+	})
+);
+
+void FAsyncCompilationNotification::Update(int32 NumJobs)
+{
+	check(IsInGameThread());
+	TSharedPtr<SNotificationItem> NotificationItem = NotificationItemPtr.Pin();
+
+	FFormatNamedArguments Args;
+	Args.Add(TEXT("AssetType"), AssetType);
+
+	if (NumJobs == 0)
+	{
+		if (NotificationItem.IsValid())
+		{
+			NotificationItem->SetText(FText::Format(LOCTEXT("AsyncCompilationFinished", "Finished preparing {AssetType}!"), Args));
+			NotificationItem->SetCompletionState(SNotificationItem::CS_Success);
+			NotificationItem->ExpireAndFadeout();
+
+			NotificationItemPtr.Reset();
+		}
+	}
+	else
+	{
+		Args.Add(TEXT("NumJobs"), FText::AsNumber(NumJobs));
+		FText ProgressMessage = FText::Format(LOCTEXT("AsyncCompilationProgress", "Preparing {AssetType} ({NumJobs})"), Args);
+
+		if (!NotificationItem.IsValid())
+		{
+			FNotificationInfo Info(ProgressMessage);
+			Info.bFireAndForget = false;
+
+			// Setting fade out and expire time to 0 as the expire message is currently very obnoxious
+			Info.FadeOutDuration = 0.0f;
+			Info.ExpireDuration = 0.0f;
+
+			NotificationItem = FSlateNotificationManager::Get().AddNotification(Info);
+			NotificationItemPtr = NotificationItem;
+		}
+
+		NotificationItem->SetCompletionState(SNotificationItem::CS_Pending);
+		NotificationItem->SetVisibility(EVisibility::HitTestInvisible);
+		NotificationItem->SetText(ProgressMessage);
+	}
+};
+
+namespace AsyncCompilationHelpers
+{
+	void FinishCompilation(
+		TFunctionRef<ICompilable& (int32 Index)> Getter,
+		int32 Num,
+		const FText& AssetType,
+		const FLogCategoryBase& LogCategory,
+		TFunctionRef<void(ICompilable*)> PostCompileSingle
+	)
+	{
+		check(IsInGameThread());
+
+		FObjectCacheContextScope ObjectCacheScope;
+		if (Num)
+		{
+			TOptional<FScopedSlowTask> SlowTask;
+
+			// Do not create a progress during PostLoad as TickSlate could endup calling too many things
+			if (!FUObjectThreadContext::Get().IsRoutingPostLoad)
+			{
+				FFormatNamedArguments Args;
+				Args.Add(TEXT("AssetType"), AssetType.ToLower());
+
+				SlowTask.Emplace((float)Num, FText::Format(LOCTEXT("WaitingOnFinishCompilation", "Waiting on {AssetType} preparation"), Args), true);
+				SlowTask->MakeDialogDelayed(1.0f, false /*bShowCancelButton*/, true /*bAllowInPIE*/);
+			}
+
+			struct FFinishTask : public IQueuedWork
+			{
+				ICompilable* Job = nullptr;
+				FEvent* Event;
+				FFinishTask() { Event = FPlatformProcess::GetSynchEventFromPool(true); }
+				~FFinishTask() { FPlatformProcess::ReturnSynchEventToPool(Event); }
+				void DoThreadedWork() override
+				{
+					FOptionalTaskTagScope Scope(ETaskTag::EParallelGameThread);
+					Job->EnsureCompletion();
+					Event->Trigger();
+				}
+
+				void Abandon() override { }
+			};
+
+			// Perform forced compilation on as many thread as possible in high priority since the game-thread is waiting
+			TArray<FFinishTask> PendingTasks;
+			PendingTasks.SetNum(Num);
+
+			for (int32 Index = 0; Index < Num; ++Index)
+			{
+				PendingTasks[Index].Job = &Getter(Index);
+				GThreadPool->AddQueuedWork(&PendingTasks[Index], EQueuedWorkPriority::High);
+			}
+
+			auto FormatProgress =
+				[&AssetType](int32 Done, int32 Total, FName ObjectName)
+			{
+				FFormatNamedArguments Args;
+				Args.Add(TEXT("AssetType"), AssetType.ToLower());
+				Args.Add(TEXT("Done"), Done);
+				Args.Add(TEXT("Total"), Total);
+				Args.Add(TEXT("ObjectName"), FText::FromName(ObjectName));
+				return FText::Format(LOCTEXT("WaitingOnFinishCompilationWithCount", "Waiting for {AssetType} to be ready {Done}/{Total} ({ObjectName}) ..."), Args);
+			};
+
+			int32 NumDone = 0;
+			for (FFinishTask& PendingTask : PendingTasks)
+			{
+				ICompilable* Job = PendingTask.Job;
+
+				FText Progress = FormatProgress(NumDone++, Num, Job->GetName());
+
+				// Be nice with the game thread and tick the progress at 60 fps even when no progress is being made...
+				while (!PendingTask.Event->Wait(16))
+				{
+					if (SlowTask.IsSet())
+					{
+						SlowTask->EnterProgressFrame(0.0f, Progress);
+					}
+				}
+
+				if (SlowTask.IsSet())
+				{
+					SlowTask->EnterProgressFrame(1.0f, Progress);
+				}
+
+				UE_LOG_REF(LogCategory, Display, TEXT("%s"), *Progress.ToString());
+
+				PostCompileSingle(Job);
+			}
+		}
+	}
+
+	void EnsureInitializedCVars(
+		const TCHAR* InName,
+		TAutoConsoleVariable<int32>& InCVarAsyncCompilation,
+		TAutoConsoleVariable<int32>& InCVarAsyncCompilationMaxConcurrency,
+		FName InExperimentalSettingsName)
+	{
+		if (InExperimentalSettingsName != NAME_None)
+		{
+			auto UpdateCVarFromSettings =
+				[&InCVarAsyncCompilation, InExperimentalSettingsName]
+				{
+					const UEditorExperimentalSettings* Settings = GetDefault<UEditorExperimentalSettings>();
+					FProperty* Property = Settings->GetClass()->FindPropertyByName(InExperimentalSettingsName);
+					if (Property)
+					{
+						if (FBoolProperty* BoolProperty = CastField<FBoolProperty>(Property))
+						{
+							bool bIsEnabled = BoolProperty->GetPropertyValue(BoolProperty->ContainerPtrToValuePtr<void>(Settings));
+							InCVarAsyncCompilation->Set(bIsEnabled ? 1 : 0, ECVF_SetByProjectSetting);
+						}
+					}
+				};
+
+			GetMutableDefault<UEditorExperimentalSettings>()->OnSettingChanged().AddLambda(
+				[UpdateCVarFromSettings, InExperimentalSettingsName](FName Name)
+				{
+					if (Name == InExperimentalSettingsName)
+					{
+						UpdateCVarFromSettings();
+					}
+				}
+			);
+
+			UpdateCVarFromSettings();
+		}
+
+		FString Value;
+		if (FParse::Value(FCommandLine::Get(), *FString::Printf(TEXT("-async%scompilation="), InName), Value))
+		{
+			int32 AsyncCompilationValue = 0;
+			if (Value == TEXT("1") || Value == TEXT("on"))
+			{
+				AsyncCompilationValue = 1;
+			}
+
+			if (Value == TEXT("2") || Value == TEXT("paused"))
+			{
+				AsyncCompilationValue = 2;
+			}
+
+			InCVarAsyncCompilation->Set(AsyncCompilationValue, ECVF_SetByCommandline);
+		}
+
+		int32 MaxConcurrency = -1;
+		if (FParse::Value(FCommandLine::Get(), *FString::Printf(TEXT("-async%scompilationmaxconcurrency="), InName), MaxConcurrency))
+		{
+			InCVarAsyncCompilationMaxConcurrency->Set(MaxConcurrency, ECVF_SetByCommandline);
+		}
+	}
+
+	void BindThreadPoolToCVar(
+		FQueuedThreadPoolWrapper* InThreadPoolWrapper,
+		TAutoConsoleVariable<int32>& InCVarAsyncCompilation,
+		TAutoConsoleVariable<int32>& InCVarAsyncCompilationResume,
+		TAutoConsoleVariable<int32>& InCVarAsyncCompilationMaxConcurrency
+		)
+	{
+		InCVarAsyncCompilation->SetOnChangedCallback(
+			FConsoleVariableDelegate::CreateLambda(
+				[InThreadPoolWrapper](IConsoleVariable* Variable)
+				{
+					if (Variable->GetInt() == 2)
+					{
+						InThreadPoolWrapper->Pause();
+					}
+					else
+					{
+						InThreadPoolWrapper->Resume();
+					}
+				}
+			)
+		);
+
+		InCVarAsyncCompilationResume->SetOnChangedCallback(
+			FConsoleVariableDelegate::CreateLambda(
+				[InThreadPoolWrapper](IConsoleVariable* Variable)
+				{
+					if (Variable->GetInt() > 0)
+					{
+						InThreadPoolWrapper->Resume(Variable->GetInt());
+					}
+				}
+			)
+		);
+
+		InCVarAsyncCompilationMaxConcurrency->SetOnChangedCallback(
+			FConsoleVariableDelegate::CreateLambda(
+				[InThreadPoolWrapper](IConsoleVariable* Variable)
+				{
+					InThreadPoolWrapper->SetMaxConcurrency(Variable->GetInt());
+				}
+			)
+		);
+
+		if (InCVarAsyncCompilation->GetInt() == 2)
+		{
+			InThreadPoolWrapper->Pause();
+		}
+
+		InThreadPoolWrapper->SetMaxConcurrency(InCVarAsyncCompilationMaxConcurrency->GetInt());
+	}
+
+	FAsyncCompilationStandardCVars::FAsyncCompilationStandardCVars(const TCHAR* AssetType, const TCHAR* AssetTypePluralLowerCase, const FConsoleCommandDelegate& FinishAllCommand)
+		: AsyncCompilation(
+			*FString::Printf(TEXT("Editor.Async%sCompilation"), AssetType),
+			0,	// Constructor default is disabled, need to be activated by one of the activation method
+			*FString::Printf(
+				TEXT("1 - Async %s compilation is enabled.\n")
+				TEXT("2 - Async %s compilation is enabled but on pause (for debugging).\n")
+				TEXT("When enabled, %s will be replaced by placeholders until they are ready\n")
+				TEXT("to reduce stalls on the game thread and improve overall editor performance."),
+				AssetTypePluralLowerCase,
+				AssetTypePluralLowerCase,
+				AssetTypePluralLowerCase
+			),
+			ECVF_Default)
+		, AsyncCompilationMaxConcurrency(
+			*FString::Printf(TEXT("Editor.Async%sCompilationMaxConcurrency"), AssetType),
+			-1,
+			*FString::Printf(
+				TEXT("Set the maximum number of concurrent %s compilation, -1 for unlimited."),
+				AssetTypePluralLowerCase
+			),
+			ECVF_Default)
+		, AsyncCompilationFinishAll(
+			*FString::Printf(TEXT("Editor.Async%sCompilationFinishAll"), AssetType),
+			*FString::Printf(TEXT("Finish all %s compilations"), AssetTypePluralLowerCase),
+			FinishAllCommand)
+		, AsyncCompilationResume(
+			*FString::Printf(TEXT("Editor.Async%sCompilationResume"), AssetType),
+			0,
+			TEXT("Number of queued work to resume while paused."),
+			ECVF_Default)
+	{
+	}
+
+	struct FStackData
+	{
+		static const uint32 MaxDepth = 24;
+		uint64 Backtrace[MaxDepth]{ 0 };
+		uint64 Cycles = 0;
+		int64  Count = 0;
+	};
+
+	static FRWLock                  BackTracesLock;
+	static TMap<uint64, FStackData> BackTraces;
+
+	void SaveStallStack(uint64 Cycles)
+	{
+		TRACE_CPUPROFILER_EVENT_SCOPE(AsyncCompilationHelpers::SaveStallStack);
+
+		uint64 Backtrace[FStackData::MaxDepth];
+		const uint32 Depth = FPlatformStackWalk::CaptureStackBackTrace(Backtrace, FStackData::MaxDepth);
+		const uint64 StackHash = CityHash64(reinterpret_cast<const char*>(Backtrace), Depth * sizeof(Backtrace[0]));
+
+		// We're already in an editor-only slow path, we don't mind the locking performance.
+		FRWScopeLock Scope(BackTracesLock, SLT_Write);
+		FStackData& StackData = BackTraces.FindOrAdd(StackHash);
+		if (StackData.Count++ == 0)
+		{
+			FPlatformMemory::Memcpy(StackData.Backtrace, Backtrace, sizeof(StackData.Backtrace));
+		}
+		StackData.Cycles += Cycles;
+	}
+
+	void DumpStallStacks()
+	{
+		FRWScopeLock Scope(BackTracesLock, SLT_Write);
+		TArray<FStackData*> Stacks;
+		for (TPair<uint64, FStackData>& Pair : BackTraces)
+		{
+			Stacks.Add(&Pair.Value);
+		}
+
+		Algo::SortBy(Stacks, [](const FStackData* StackData) { return StackData->Cycles; });
+
+		const int32 HumanReadableStringSize = 4096;
+		ANSICHAR HumanReadableString[HumanReadableStringSize];
+		int64 TotalCount = 0;
+
+		uint64 TotalCycles = 0;
+		for (FStackData* StackData : Stacks)
+		{
+			HumanReadableString[0] = '\0';
+
+			// Start at index 2 to Skip both the CaptureBackTrace and this function
+			for (int32 Index = 2; Index < FStackData::MaxDepth && StackData->Backtrace[Index]; ++Index)
+			{
+				FPlatformStackWalk::ProgramCounterToHumanReadableString(Index, StackData->Backtrace[Index], HumanReadableString, HumanReadableStringSize);
+				FCStringAnsi::Strncat(HumanReadableString, LINE_TERMINATOR_ANSI, HumanReadableStringSize);
+			}
+
+			TotalCount += StackData->Count;
+			TotalCycles += StackData->Cycles;
+			UE_LOG(LogAsyncCompilation, Display, TEXT("Async Compilation Stall Stack: (count: %llu, time: %s)\n%s"), StackData->Count, *FPlatformTime::PrettyTime(FPlatformTime::ToSeconds64(StackData->Cycles)), ANSI_TO_TCHAR(HumanReadableString));
+		}
+	}
+}
+
+#undef LOCTEXT_NAMESPACE
+
+#endif // #if WITH_EDITOR
\ No newline at end of file
diff --git a/Engine/Source/Runtime/Engine/Private/Components/ActorComponent.cpp b/Engine/Source/Runtime/Engine/Private/Components/ActorComponent.cpp
--- a/Engine/Source/Runtime/Engine/Private/Components/ActorComponent.cpp
+++ b/Engine/Source/Runtime/Engine/Private/Components/ActorComponent.cpp
@@ -101,7 +101,7 @@ void FRegisterComponentContext::Process()
 #if WITH_EDITOR
 	// This is required for async static mesh compilation in case a scene proxy is not async aware.
 	// A stall until the compilation is finished might occur, and this is only supported on the game thread for now.
-	bSingleThreaded |= FStaticMeshCompilingManager().Get().IsAsyncStaticMeshCompilationEnabled();
+	bSingleThreaded |= FStaticMeshCompilingManager::Get().IsAsyncStaticMeshCompilationEnabled();
 #endif
 
 	FSceneInterface* Scene = World->Scene;
diff --git a/Engine/Source/Runtime/Engine/Private/StaticMesh.cpp b/Engine/Source/Runtime/Engine/Private/StaticMesh.cpp
--- a/Engine/Source/Runtime/Engine/Private/StaticMesh.cpp
+++ b/Engine/Source/Runtime/Engine/Private/StaticMesh.cpp
@@ -66,9 +66,6 @@
 #include "TessellationRendering.h"
 #include "Misc/MessageDialog.h"
 #include "StaticMeshCompiler.h"
-#include "HAL/PlatformStackWalk.h"
-#include "Hash/CityHash.h"
-#include "Misc/ScopeRWLock.h"
 #include "AssetCompilingManager.h"
 #endif // #if WITH_EDITOR
 
@@ -2857,82 +2854,6 @@ PRAGMA_ENABLE_DEPRECATION_WARNINGS
 
 #if WITH_EDITOR
 
-namespace StaticMeshImpl
-{
-	struct FStackData
-	{
-		static const uint32 MaxDepth = 24;
-		uint64 Backtrace[MaxDepth]{0};
-		uint64 Cycles = 0;
-		int64  Count  = 0;
-	};
-
-	static FRWLock                  BackTracesLock;
-	static TMap<uint64, FStackData> BackTraces;
-
-	void SaveStack(uint64 Cycles)
-	{
-		TRACE_CPUPROFILER_EVENT_SCOPE(StaticMeshImpl::SaveStackHash);
-		
-		uint64 Backtrace[FStackData::MaxDepth];
-		const uint32 Depth = FPlatformStackWalk::CaptureStackBackTrace(Backtrace, FStackData::MaxDepth);
-		const uint64 StackHash = CityHash64(reinterpret_cast<const char*>(Backtrace), Depth * sizeof(Backtrace[0]));
-
-		// We're already in an editor-only slow path, we don't mind the locking performance.
-		FRWScopeLock Scope(BackTracesLock, SLT_Write);
-		FStackData& StackData = BackTraces.FindOrAdd(StackHash);
-		if (StackData.Count++ == 0)
-		{
-			FPlatformMemory::Memcpy(StackData.Backtrace, Backtrace, sizeof(StackData.Backtrace));
-		}
-		StackData.Cycles += Cycles;
-	}
-
-	void DumpStallStacks()
-	{
-		FRWScopeLock Scope(BackTracesLock, SLT_Write);
-		TArray<FStackData*> Stacks;
-		for (TPair<uint64, FStackData>& Pair : BackTraces)
-		{
-			Stacks.Add(&Pair.Value);
-		}
-		
-		Algo::SortBy(Stacks, [](const FStackData* StackData) { return StackData->Cycles; });
-
-		const int32 HumanReadableStringSize = 4096;
-		ANSICHAR HumanReadableString[HumanReadableStringSize];
-
-		int64 TotalCount = 0;
-		uint64 TotalCycles = 0;
-		for (FStackData* StackData : Stacks)
-		{
-			HumanReadableString[0] = '\0';
-
-			// Start at index 2 to Skip both the CaptureBackTrace and this function
-			for (int32 Index = 2; Index < FStackData::MaxDepth && StackData->Backtrace[Index]; ++Index)
-			{
-				FPlatformStackWalk::ProgramCounterToHumanReadableString(Index, StackData->Backtrace[Index], HumanReadableString, HumanReadableStringSize);
-				FCStringAnsi::Strncat(HumanReadableString, LINE_TERMINATOR_ANSI, HumanReadableStringSize);
-			}
-
-			TotalCount += StackData->Count;
-			TotalCycles += StackData->Cycles;
-			UE_LOG(LogStaticMesh, Display, TEXT("StaticMesh Stall Stack: (count: %llu, time: %s)\n%s"), StackData->Count, *FPlatformTime::PrettyTime(FPlatformTime::ToSeconds64(StackData->Cycles)), ANSI_TO_TCHAR(HumanReadableString));
-		}
-
-		UE_LOG(LogStaticMesh, Display, TEXT("StaticMesh Total Stalls: (count: %llu, time: %s)"), TotalCount, *FPlatformTime::PrettyTime(FPlatformTime::ToSeconds64(TotalCycles)));
-	}
-}
-
-static FAutoConsoleCommand CVarAsyncStaticMeshDumpStallStacks(
-	TEXT("Editor.AsyncStaticMeshDumpStallStacks"),
-	TEXT("Dump all the callstacks that have caused waits on async static mesh compilation."),
-	FConsoleCommandWithArgsDelegate::CreateLambda([](const TArray<FString>& Args)
-	{
-		StaticMeshImpl::DumpStallStacks();
-	})
-);
-
 thread_local const UStaticMesh* FStaticMeshAsyncBuildScope::StaticMeshBeingAsyncCompiled = nullptr;
 
 void UStaticMesh::WaitUntilAsyncPropertyReleased(EStaticMeshAsyncProperties AsyncProperties) const
@@ -2949,13 +2870,13 @@ void UStaticMesh::WaitUntilAsyncPropertyReleased(EStaticMeshAsyncProperties Asyn
 				Verbose,
 				TEXT("Accessing property %s of the StaticMesh while it is still being built asynchronously will force it to be compiled before continuing. "
 					 "For better performance, consider making the caller async aware so it can wait until the static mesh is ready to access this property."
-					 "To better understand where those calls are coming from, you can use Editor.AsyncStaticMeshDumpStallStacks on the console." ),
+					 "To better understand where those calls are coming from, you can use Editor.AsyncAssetDumpStallStacks on the console." ),
 				ToString(AsyncProperties)
 			);
 
 			uint64 StartTime = FPlatformTime::Cycles64();
 			FStaticMeshCompilingManager::Get().FinishCompilation({ const_cast<UStaticMesh*>(this) });
-			StaticMeshImpl::SaveStack(FPlatformTime::Cycles64() - StartTime);
+			AsyncCompilationHelpers::SaveStallStack(FPlatformTime::Cycles64() - StartTime);
 		}
 		else
 		{
diff --git a/Engine/Source/Runtime/Engine/Private/StaticMeshCompiler.cpp b/Engine/Source/Runtime/Engine/Private/StaticMeshCompiler.cpp
--- a/Engine/Source/Runtime/Engine/Private/StaticMeshCompiler.cpp
+++ b/Engine/Source/Runtime/Engine/Private/StaticMeshCompiler.cpp
@@ -1,15 +1,13 @@
 // Copyright Epic Games, Inc. All Rights Reserved.
 
 #include "StaticMeshCompiler.h"
-#include "AssetCompilingManager.h"
-#include "Engine/StaticMesh.h"
 
 #if WITH_EDITOR
 
+#include "AssetCompilingManager.h"
+#include "Engine/StaticMesh.h"
 #include "ObjectCacheContext.h"
-#include "Framework/Notifications/NotificationManager.h"
 #include "Settings/EditorExperimentalSettings.h"
-#include "Widgets/Notifications/SNotificationList.h"
 #include "Misc/QueuedThreadPoolWrapper.h"
 #include "Misc/Optional.h"
 #include "EngineModule.h"
@@ -31,15 +29,15 @@
 
 #define LOCTEXT_NAMESPACE "StaticMeshCompiler"
 
-static TAutoConsoleVariable<int32> CVarAsyncStaticMeshCompilation(
-	TEXT("Editor.AsyncStaticMeshCompilation"),
-	0,
-	TEXT("0 - Async static mesh compilation is disabled.\n")
-	TEXT("1 - Async static mesh compilation is enabled.\n")
-	TEXT("2 - Async static mesh compilation is enabled but on pause (for debugging).\n")
-	TEXT("When enabled, static meshes will be replaced by placeholders until they are ready\n")
-	TEXT("to reduce stalls on the game thread and improve overall editor performance."),
-	ECVF_Default);
+static AsyncCompilationHelpers::FAsyncCompilationStandardCVars CVarAsyncStaticMeshStandard(
+	TEXT("StaticMesh"),
+	TEXT("static meshes"),
+	FConsoleCommandDelegate::CreateLambda(
+		[]()
+		{
+			FStaticMeshCompilingManager::Get().FinishAllCompilation();
+		}
+	));
 
 static TAutoConsoleVariable<int32> CVarAsyncStaticMeshPlayInEditorMode(
 	TEXT("Editor.AsyncStaticMeshPlayInEditorMode"),
@@ -66,27 +64,6 @@ static TAutoConsoleVariable<bool> CVarAsyncStaticMeshDebugDraw(
 	TEXT("Any static meshes that were waited on due to being too close to the player will have their bounding box drawn in red for a couple of seconds."),
 	ECVF_Default);
 
-static TAutoConsoleVariable<int32> CVarAsyncStaticMeshMaxConcurrency(
-	TEXT("Editor.AsyncStaticMeshMaxConcurrency"),
-	-1,
-	TEXT("Set the maximum number of concurrent static mesh compilation, -1 for unlimited."),
-	ECVF_Default);
-
-static FAutoConsoleCommand CVarAsyncStaticMeshCompilationFinishAll(
-	TEXT("Editor.AsyncStaticMeshCompilationFinishAll"),
-	TEXT("Finish all static mesh compilations"),
-	FConsoleCommandWithArgsDelegate::CreateLambda([](const TArray<FString>& Args)
-	{
-		FStaticMeshCompilingManager::Get().FinishAllCompilation();
-	})
-);
-
-static TAutoConsoleVariable<int32> CVarAsyncStaticMeshCompilationResume(
-	TEXT("Editor.AsyncStaticMeshCompilationResume"),
-	0,
-	TEXT("Number of queued work to resume while paused."),
-	ECVF_Default);
-
 namespace StaticMeshCompilingManagerImpl
 {
 	static void EnsureInitializedCVars()
@@ -96,44 +73,21 @@ namespace StaticMeshCompilingManagerImpl
 		if (!bIsInitialized)
 		{
 			bIsInitialized = true;
-			GetMutableDefault<UEditorExperimentalSettings>()->OnSettingChanged().AddLambda(
-				[](FName Name)
-				{
-					if (Name == TEXT("bEnableAsyncStaticMeshCompilation"))
-					{
-						CVarAsyncStaticMeshCompilation->Set(GetDefault<UEditorExperimentalSettings>()->bEnableAsyncStaticMeshCompilation ? 1 : 0, ECVF_SetByProjectSetting);
-					}
-				}
-			);
-
-			CVarAsyncStaticMeshCompilation->Set(GetDefault<UEditorExperimentalSettings>()->bEnableAsyncStaticMeshCompilation ? 1 : 0, ECVF_SetByProjectSetting);
-
-			FString Value;
-			if (FParse::Value(FCommandLine::Get(), TEXT("-asyncstaticmeshcompilation="), Value))
-			{
-				int32 AsyncStaticMeshCompilationValue = 0;
-				if (Value == TEXT("1") || Value == TEXT("on"))
-				{
-					AsyncStaticMeshCompilationValue = 1;
-				}
-
-				if (Value == TEXT("2") || Value == TEXT("paused"))
-				{
-					AsyncStaticMeshCompilationValue = 2;
-				}
-
-				CVarAsyncStaticMeshCompilation->Set(AsyncStaticMeshCompilationValue, ECVF_SetByCommandline);
-			}
 
-			int32 MaxConcurrency = -1;
-			if (FParse::Value(FCommandLine::Get(), TEXT("-asyncstaticmeshmaxconcurrency="), MaxConcurrency))
-			{
-				CVarAsyncStaticMeshMaxConcurrency->Set(MaxConcurrency, ECVF_SetByCommandline);
-			}
+			AsyncCompilationHelpers::EnsureInitializedCVars(
+				TEXT("staticmesh"),
+				CVarAsyncStaticMeshStandard.AsyncCompilation,
+				CVarAsyncStaticMeshStandard.AsyncCompilationMaxConcurrency,
+				GET_MEMBER_NAME_CHECKED(UEditorExperimentalSettings, bEnableAsyncStaticMeshCompilation));
 		}
 	}
 }
 
+FStaticMeshCompilingManager::FStaticMeshCompilingManager()
+	: Notification(LOCTEXT("StaticMeshes", "Static Meshes"))
+{
+}
+
 EQueuedWorkPriority FStaticMeshCompilingManager::GetBasePriority(UStaticMesh* InStaticMesh) const
 {
 	return EQueuedWorkPriority::Low;
@@ -146,52 +100,15 @@ FQueuedThreadPool* FStaticMeshCompilingManager::GetThreadPool() const
 	{
 		StaticMeshCompilingManagerImpl::EnsureInitializedCVars();
 
-		const int32 MaxConcurrency = CVarAsyncStaticMeshMaxConcurrency.GetValueOnAnyThread();
-
 		// Static meshes will be scheduled on the asset thread pool, where concurrency limits might by dynamically adjusted depending on memory constraints.
-		GStaticMeshThreadPool = new FQueuedThreadPoolDynamicWrapper(FAssetCompilingManager::Get().GetThreadPool(), MaxConcurrency, [](EQueuedWorkPriority) { return EQueuedWorkPriority::Low; });
-
-		CVarAsyncStaticMeshCompilation->SetOnChangedCallback(
-			FConsoleVariableDelegate::CreateLambda(
-				[](IConsoleVariable* Variable)
-				{
-					if (Variable->GetInt() == 2)
-					{
-						GStaticMeshThreadPool->Pause();
-					}
-					else
-					{
-						GStaticMeshThreadPool->Resume();
-					}
-				}
-				)
-			);
-
-		CVarAsyncStaticMeshCompilationResume->SetOnChangedCallback(
-			FConsoleVariableDelegate::CreateLambda(
-				[](IConsoleVariable* Variable)
-				{
-					if (Variable->GetInt() > 0)
-					{
-						GStaticMeshThreadPool->Resume(Variable->GetInt());
-					}
-				}
-				)
-			);
-
-		CVarAsyncStaticMeshMaxConcurrency->SetOnChangedCallback(
-			FConsoleVariableDelegate::CreateLambda(
-				[](IConsoleVariable* Variable)
-				{
-					GStaticMeshThreadPool->SetMaxConcurrency(Variable->GetInt());
-				}
-				)
-			);
-
-		if (CVarAsyncStaticMeshCompilation->GetInt() == 2)
-		{
-			GStaticMeshThreadPool->Pause();
-		}
+		GStaticMeshThreadPool = new FQueuedThreadPoolDynamicWrapper(FAssetCompilingManager::Get().GetThreadPool(), -1, [](EQueuedWorkPriority) { return EQueuedWorkPriority::Low; });
+
+		AsyncCompilationHelpers::BindThreadPoolToCVar(
+			GStaticMeshThreadPool,
+			CVarAsyncStaticMeshStandard.AsyncCompilation,
+			CVarAsyncStaticMeshStandard.AsyncCompilationResume,
+			CVarAsyncStaticMeshStandard.AsyncCompilationMaxConcurrency
+		);
 	}
 
 	return GStaticMeshThreadPool;
@@ -205,33 +122,30 @@ void FStaticMeshCompilingManager::Shutdown()
 		check(IsInGameThread());
 		TRACE_CPUPROFILER_EVENT_SCOPE(FStaticMeshCompilingManager::Shutdown)
 
-		if (GetNumRemainingMeshes())
-		{
-			TArray<UStaticMesh*> PendingStaticMeshes;
-			PendingStaticMeshes.Reserve(GetNumRemainingMeshes());
+		TArray<UStaticMesh*> PendingStaticMeshes;
+		PendingStaticMeshes.Reserve(GetNumRemainingMeshes());
 
-			for (TWeakObjectPtr<UStaticMesh>& WeakStaticMesh : RegisteredStaticMesh)
+		for (TWeakObjectPtr<UStaticMesh>& WeakStaticMesh : RegisteredStaticMesh)
+		{
+			if (WeakStaticMesh.IsValid())
 			{
-				if (WeakStaticMesh.IsValid())
+				UStaticMesh* StaticMesh = WeakStaticMesh.Get();
+				if (!StaticMesh->IsAsyncTaskComplete())
 				{
-					UStaticMesh* StaticMesh = WeakStaticMesh.Get();
-					if (!StaticMesh->IsAsyncTaskComplete())
+					if (StaticMesh->AsyncTask->Cancel())
 					{
-						if (StaticMesh->AsyncTask->Cancel())
-						{
-							StaticMesh->AsyncTask.Reset();
-						}
+						StaticMesh->AsyncTask.Reset();
 					}
+				}
 
-					if (StaticMesh->AsyncTask)
-					{
-						PendingStaticMeshes.Add(StaticMesh);
-					}
+				if (StaticMesh->AsyncTask)
+				{
+					PendingStaticMeshes.Add(StaticMesh);
 				}
 			}
-
-			FinishCompilation(PendingStaticMeshes);
 		}
+
+		FinishCompilation(PendingStaticMeshes);
 	}
 }
 
@@ -244,54 +158,15 @@ bool FStaticMeshCompilingManager::IsAsyncStaticMeshCompilationEnabled() const
 
 	StaticMeshCompilingManagerImpl::EnsureInitializedCVars();
 
-	return CVarAsyncStaticMeshCompilation.GetValueOnAnyThread() != 0;
+	return CVarAsyncStaticMeshStandard.AsyncCompilation.GetValueOnAnyThread() != 0;
 }
 
 void FStaticMeshCompilingManager::UpdateCompilationNotification()
 {
-	check(IsInGameThread());
-	static TWeakPtr<SNotificationItem> StaticMeshCompilationPtr;
-
-	TSharedPtr<SNotificationItem> NotificationItem = StaticMeshCompilationPtr.Pin();
-
-	const int32 NumRemainingCompilations = GetNumRemainingMeshes();
-	if (NumRemainingCompilations == 0)
-	{
-		if (NotificationItem.IsValid())
-		{
-			NotificationItem->SetText(NSLOCTEXT("StaticMeshBuild", "StaticMeshBuildFinished", "Finished preparing Static Meshes!"));
-			NotificationItem->SetCompletionState(SNotificationItem::CS_Success);
-			NotificationItem->ExpireAndFadeout();
-
-			StaticMeshCompilationPtr.Reset();
-		}
-	}
-	else
-	{
-		if (!NotificationItem.IsValid())
-		{
-			FNotificationInfo Info(NSLOCTEXT("StaticMeshBuild", "StaticMeshBuildInProgress", "Preparing Static Meshes"));
-			Info.bFireAndForget = false;
-
-			// Setting fade out and expire time to 0 as the expire message is currently very obnoxious
-			Info.FadeOutDuration = 0.0f;
-			Info.ExpireDuration = 0.0f;
-
-			NotificationItem = FSlateNotificationManager::Get().AddNotification(Info);
-			StaticMeshCompilationPtr = NotificationItem;
-		}
-
-		FFormatNamedArguments Args;
-		Args.Add(TEXT("BuildTasks"), FText::AsNumber(NumRemainingCompilations));
-		FText ProgressMessage = FText::Format(NSLOCTEXT("StaticMeshBuild", "StaticMeshBuildInProgressFormat", "Preparing Static Meshes ({BuildTasks})"), Args);
-
-		NotificationItem->SetCompletionState(SNotificationItem::CS_Pending);
-		NotificationItem->SetVisibility(EVisibility::HitTestInvisible);
-		NotificationItem->SetText(ProgressMessage);
-	}
+	Notification.Update(GetNumRemainingMeshes());
 }
 
-void FStaticMeshCompilingManager::PostStaticMeshesCompilation(const TSet<UStaticMesh*>& InStaticMeshes)
+void FStaticMeshCompilingManager::PostCompilation(TArrayView<UStaticMesh* const> InStaticMeshes)
 {
 	if (InStaticMeshes.Num())
 	{
@@ -309,7 +184,7 @@ void FStaticMeshCompilingManager::PostStaticMeshesCompilation(const TSet<UStatic
 	}
 }
 
-void FStaticMeshCompilingManager::FinishStaticMeshCompilation(UStaticMesh* StaticMesh)
+void FStaticMeshCompilingManager::PostCompilation(UStaticMesh* StaticMesh)
 {
 	using namespace StaticMeshCompilingManagerImpl;
 	
@@ -317,7 +192,7 @@ void FStaticMeshCompilingManager::FinishStaticMeshCompilation(UStaticMesh* Stati
 	if (StaticMesh->AsyncTask)
 	{
 		check(IsInGameThread());
-		TRACE_CPUPROFILER_EVENT_SCOPE(FinishStaticMeshCompilation);
+		TRACE_CPUPROFILER_EVENT_SCOPE(PostCompilation);
 
 		UE_LOG(LogStaticMesh, Verbose, TEXT("Refreshing static mesh %s because it is ready"), *StaticMesh->GetName());
 
@@ -378,7 +253,7 @@ int32 FStaticMeshCompilingManager::GetNumRemainingMeshes() const
 	return RegisteredStaticMesh.Num();
 }
 
-void FStaticMeshCompilingManager::AddStaticMeshes(const TArray<UStaticMesh*>& InStaticMeshes)
+void FStaticMeshCompilingManager::AddStaticMeshes(TArrayView<UStaticMesh* const> InStaticMeshes)
 {
 	TRACE_CPUPROFILER_EVENT_SCOPE(FStaticMeshCompilingManager::AddStaticMeshes)
 	check(IsInGameThread());
@@ -394,15 +269,13 @@ void FStaticMeshCompilingManager::AddStaticMeshes(const TArray<UStaticMesh*>& In
 	}
 }
 
-void FStaticMeshCompilingManager::FinishCompilation(const TArray<UStaticMesh*>& InStaticMeshes)
+void FStaticMeshCompilingManager::FinishCompilation(TArrayView<UStaticMesh* const> InStaticMeshes)
 {
 	TRACE_CPUPROFILER_EVENT_SCOPE(FStaticMeshCompilingManager::FinishCompilation);
 
-	using namespace StaticMeshCompilingManagerImpl;
 	check(IsInGameThread());
 
-	FObjectCacheContextScope ObjectCacheScope;
-	TSet<UStaticMesh*> PendingStaticMeshes;
+	TArray<UStaticMesh*> PendingStaticMeshes;
 	PendingStaticMeshes.Reserve(InStaticMeshes.Num());
 
 	int32 StaticMeshIndex = 0;
@@ -410,78 +283,41 @@ void FStaticMeshCompilingManager::FinishCompilation(const TArray<UStaticMesh*>&
 	{
 		if (RegisteredStaticMesh.Contains(StaticMesh))
 		{
-			PendingStaticMeshes.Add(StaticMesh);
+			PendingStaticMeshes.Emplace(StaticMesh);
 		}
 	}
 
 	if (PendingStaticMeshes.Num())
 	{
-		TOptional<FScopedSlowTask> SlowTask;
-
-		// Do not create a progress during PostLoad as TickSlate could endup calling stuff that is not supposed to be called during PostLoad
-		if (!FUObjectThreadContext::Get().IsRoutingPostLoad)
+		class FCompilableStaticMesh : public AsyncCompilationHelpers::ICompilable
 		{
-			SlowTask.Emplace((float)PendingStaticMeshes.Num(), LOCTEXT("FinishStaticMeshCompilation", "Waiting on static meshes preparation"), true);
-			SlowTask->MakeDialogDelayed(1.0f, false /*bShowCancelButton*/, true /*bAllowInPIE*/);
-		}
-
-		struct FStaticMeshTask : public IQueuedWork
-		{
-			TStrongObjectPtr<UStaticMesh> StaticMesh;
-			FEvent* Event;
-			FStaticMeshTask() { Event = FPlatformProcess::GetSynchEventFromPool(true); }
-			~FStaticMeshTask() { FPlatformProcess::ReturnSynchEventToPool(Event); }
-			void DoThreadedWork() override 
+		public:
+			FCompilableStaticMesh(UStaticMesh* InStaticMesh)
+				: StaticMesh(InStaticMesh)
 			{
-				if (StaticMesh->AsyncTask) 
-				{ 
-					StaticMesh->AsyncTask->EnsureCompletion();
-				}
-				Event->Trigger(); 
 			}
 
-			void Abandon() override { }
+			TStrongObjectPtr<UStaticMesh> StaticMesh;
+			void EnsureCompletion() override { if (StaticMesh->AsyncTask) { StaticMesh->AsyncTask->EnsureCompletion(); } }
+			FName GetName() override { return StaticMesh->GetFName(); }
 		};
 
-		// Perform forced compilation on as many thread as possible in high priority since the game-thread is waiting
-		TArray<FStaticMeshTask> PendingTasks;
-		PendingTasks.SetNum(PendingStaticMeshes.Num());
-		
-		int32 PendingTaskIndex = 0;
-		for (UStaticMesh* StaticMesh : PendingStaticMeshes)
-		{
-			PendingTasks[PendingTaskIndex].StaticMesh.Reset(StaticMesh);
-			GThreadPool->AddQueuedWork(&PendingTasks[PendingTaskIndex], EQueuedWorkPriority::High);
-			PendingTaskIndex++;
-		}
-
-		auto UpdateProgress =
-			[&SlowTask](float Progress, int32 Done, int32 Total, const FString& CurrentObjectsName)
-			{
-				if (SlowTask.IsSet())
-				{
-					return SlowTask->EnterProgressFrame(Progress, FText::FromString(FString::Printf(TEXT("Waiting for static meshes to be ready %d/%d (%s) ..."), Done, Total, *CurrentObjectsName)));
-				}
-			};
-
-		for (FStaticMeshTask& PendingTask : PendingTasks)
-		{
-			UStaticMesh* StaticMesh = PendingTask.StaticMesh.Get();
-			const FString StaticMeshName = StaticMesh->GetName();
-			// Be nice with the game thread and tick the progress at 60 fps even when no progress is being made...
-			while (!PendingTask.Event->Wait(16))
+		TArray<FCompilableStaticMesh> CompilableStaticMeshes(PendingStaticMeshes);
+		FObjectCacheContextScope ObjectCacheScope;
+		AsyncCompilationHelpers::FinishCompilation(
+			[&CompilableStaticMeshes](int32 Index)	-> AsyncCompilationHelpers::ICompilable& { return CompilableStaticMeshes[Index]; },
+			CompilableStaticMeshes.Num(),
+			LOCTEXT("StaticMeshes", "Static Meshes"),
+			LogStaticMesh,
+			[this](AsyncCompilationHelpers::ICompilable* Object)
 			{
-				UpdateProgress(0.0f, StaticMeshIndex, InStaticMeshes.Num(), StaticMeshName);
+				UStaticMesh* StaticMesh = static_cast<FCompilableStaticMesh*>(Object)->StaticMesh.Get();
+				PostCompilation(StaticMesh);
+				RegisteredStaticMesh.Remove(StaticMesh);
 			}
-			UE_LOG(LogStaticMesh, Display, TEXT("Waiting for static meshes to be ready %d/%d (%s) ..."), StaticMeshIndex, InStaticMeshes.Num(), *StaticMeshName);
-			UpdateProgress(1.f, StaticMeshIndex++, InStaticMeshes.Num(), StaticMeshName);
-
-			FinishStaticMeshCompilation(StaticMesh);
-
-			RegisteredStaticMesh.Remove(StaticMesh);
-		}
+		);
 
-		PostStaticMeshesCompilation(PendingStaticMeshes);
+		PostCompilation(InStaticMeshes);
 	}
 }
 
@@ -751,7 +587,7 @@ void FStaticMeshCompilingManager::ProcessStaticMeshes(bool bLimitExecutionTime,
 			const double TickStartTime = FPlatformTime::Seconds();
 
 			TSet<TWeakObjectPtr<UStaticMesh>> StaticMeshesToPostpone;
-			TSet<UStaticMesh*> ProcessedStaticMeshes;
+			TArray<UStaticMesh*> ProcessedStaticMeshes;
 			if (StaticMeshesToProcess.Num())
 			{
 				for (UStaticMesh* StaticMesh : StaticMeshesToProcess)
@@ -759,7 +595,7 @@ void FStaticMeshCompilingManager::ProcessStaticMeshes(bool bLimitExecutionTime,
 					const bool bHasMeshUpdateLeft = ProcessedStaticMeshes.Num() <= MaxMeshUpdatesPerFrame;
 					if (bHasMeshUpdateLeft && StaticMesh->IsAsyncTaskComplete())
 					{
-						FinishStaticMeshCompilation(StaticMesh);
+						PostCompilation(StaticMesh);
 						ProcessedStaticMeshes.Add(StaticMesh);
 					}
 					else
@@ -771,7 +607,7 @@ void FStaticMeshCompilingManager::ProcessStaticMeshes(bool bLimitExecutionTime,
 
 			RegisteredStaticMesh = MoveTemp(StaticMeshesToPostpone);
 
-			PostStaticMeshesCompilation(ProcessedStaticMeshes);
+			PostCompilation(ProcessedStaticMeshes);
 		}
 	}
 }
@@ -788,6 +624,6 @@ void FStaticMeshCompilingManager::ProcessAsyncTasks(bool bLimitExecutionTime)
 	UpdateCompilationNotification();
 }
 
-#endif // #if WITH_EDITOR
+#undef LOCTEXT_NAMESPACE
 
-#undef LOCTEXT_NAMESPACE
\ No newline at end of file
+#endif // #if WITH_EDITOR
diff --git a/Engine/Source/Runtime/Engine/Private/Texture2D.cpp b/Engine/Source/Runtime/Engine/Private/Texture2D.cpp
--- a/Engine/Source/Runtime/Engine/Private/Texture2D.cpp
+++ b/Engine/Source/Runtime/Engine/Private/Texture2D.cpp
@@ -39,7 +39,7 @@
 #include "Streaming/Texture2DMipDataProvider_DDC.h"
 #include "Streaming/Texture2DMipDataProvider_IO.h"
 #include "Engine/TextureMipDataProviderFactory.h"
-
+#include "AsyncCompilationHelpers.h"
 #include "Async/AsyncFileHandle.h"
 #include "EngineModule.h"
 #include "Engine/Texture2DArray.h"
@@ -272,7 +272,9 @@ const FTexturePlatformData* UTexture2D::GetPlatformData() const
 		FText Msg = FText::Format(LOCTEXT("WaitOnTextureCompilation", "Waiting on texture compilation (%s) ..."), FText::FromString(GetName()));
 		FScopedSlowTask Progress(1.f, Msg, true);
 		Progress.MakeDialog(true);
+		uint64 StartTime = FPlatformTime::Cycles64();
 		PrivatePlatformData->FinishCache();
+		AsyncCompilationHelpers::SaveStallStack(FPlatformTime::Cycles64() - StartTime);
 	}
 #endif
 
diff --git a/Engine/Source/Runtime/Engine/Private/TextureCompiler.cpp b/Engine/Source/Runtime/Engine/Private/TextureCompiler.cpp
--- a/Engine/Source/Runtime/Engine/Private/TextureCompiler.cpp
+++ b/Engine/Source/Runtime/Engine/Private/TextureCompiler.cpp
@@ -6,9 +6,7 @@
 
 #include "Engine/Texture.h"
 #include "ObjectCacheContext.h"
-#include "Framework/Notifications/NotificationManager.h"
 #include "Settings/EditorExperimentalSettings.h"
-#include "Widgets/Notifications/SNotificationList.h"
 #include "Misc/QueuedThreadPoolWrapper.h"
 #include "RendererInterface.h"
 #include "EngineModule.h"
@@ -20,41 +18,21 @@
 #include "TextureDerivedDataTask.h"
 #include "Misc/IQueuedWork.h"
 #include "Components/PrimitiveComponent.h"
+#include "AsyncCompilationHelpers.h"
 #include "AssetCompilingManager.h"
 #include "LevelEditor.h"
 
 #define LOCTEXT_NAMESPACE "TextureCompiler"
 
-static TAutoConsoleVariable<int32> CVarAsyncTextureCompilation(
-	TEXT("Editor.AsyncTextureCompilation"),
-	0,
-	TEXT("0 - Async texture compilation is disabled.\n")
-	TEXT("1 - Async texture compilation is enabled.\n")
-	TEXT("2 - Async texture compilation is enabled but on pause (for debugging).\n")
-	TEXT("When enabled, textures will be replaced by placeholders until they are ready\n")
-	TEXT("to reduce stalls on the game thread and improve overall editor performance."),
-	ECVF_Default);
-
-static TAutoConsoleVariable<int32> CVarAsyncTextureCompilationMaxConcurrency(
-	TEXT("Editor.AsyncTextureCompilationMaxConcurrency"),
-	-1,
-	TEXT("Set the maximum number of concurrent texture compilation, -1 for unlimited."),
-	ECVF_Default);
-
-static FAutoConsoleCommand CVarAsyncTextureCompilationFinishAll(
-	TEXT("Editor.AsyncTextureCompilationFinishAll"),
-	TEXT("Finish all texture compilations"),
-	FConsoleCommandWithArgsDelegate::CreateLambda([](const TArray<FString>& Args)
-	{
-		FTextureCompilingManager::Get().FinishAllCompilation();
-	})
-);
-
-static TAutoConsoleVariable<int32> CVarAsyncTextureCompilationResume(
-	TEXT("Editor.AsyncTextureCompilationResume"),
-	0,
-	TEXT("Number of queued work to resume while paused."),
-	ECVF_Default);
+static AsyncCompilationHelpers::FAsyncCompilationStandardCVars CVarAsyncTextureStandard(
+	TEXT("Texture"),
+	TEXT("textures"),
+	FConsoleCommandDelegate::CreateLambda(
+		[]()
+		{
+			FTextureCompilingManager::Get().FinishAllCompilation();
+		}
+	));
 
 namespace TextureCompilingManagerImpl
 {
@@ -81,25 +59,6 @@ namespace TextureCompilingManagerImpl
 		return (EQueuedWorkPriority)(FMath::Max((uint8)1, (uint8)GetBasePriority(InTexture)) - 1);
 	}
 
-	static const TCHAR* GetPriorityName(EQueuedWorkPriority Priority)
-	{
-		switch (Priority)
-		{
-			case EQueuedWorkPriority::Highest:
-				return TEXT("Highest");
-			case EQueuedWorkPriority::High:
-				return TEXT("High");
-			case EQueuedWorkPriority::Normal:
-				return TEXT("Normal");
-			case EQueuedWorkPriority::Low:
-				return TEXT("Low");
-			case EQueuedWorkPriority::Lowest:
-				return TEXT("Lowest");
-			default:
-				return TEXT("Unknown");
-		}
-	}
-
 	static void EnsureInitializedCVars()
 	{
 		static bool bIsInitialized = false;
@@ -107,49 +66,49 @@ namespace TextureCompilingManagerImpl
 		if (!bIsInitialized)
 		{
 			bIsInitialized = true;
-			GetMutableDefault<UEditorExperimentalSettings>()->OnSettingChanged().AddLambda(
-				[](FName Name)
-				{
-					if (Name == TEXT("bEnableAsyncTextureCompilation"))
-					{
-						CVarAsyncTextureCompilation->Set(GetDefault<UEditorExperimentalSettings>()->bEnableAsyncTextureCompilation ? 1 : 0, ECVF_SetByProjectSetting);
-					}
-				}
-			);
-
-			CVarAsyncTextureCompilation->Set(GetDefault<UEditorExperimentalSettings>()->bEnableAsyncTextureCompilation ? 1 : 0, ECVF_SetByProjectSetting);
 
-			FString Value;
-			if (FParse::Value(FCommandLine::Get(), TEXT("-asynctexturecompilation="), Value))
-			{
-				int32 AsyncTextureCompilationValue = 0;
-				if (Value == TEXT("1") || Value == TEXT("on"))
-				{
-					AsyncTextureCompilationValue = 1;
-				}
-
-				if (Value == TEXT("2") || Value == TEXT("paused"))
-				{
-					AsyncTextureCompilationValue = 2;
-				}
-
-				CVarAsyncTextureCompilation->Set(AsyncTextureCompilationValue, ECVF_SetByCommandline);
-			}
-
-			int32 MaxConcurrency = -1;
-			if (FParse::Value(FCommandLine::Get(), TEXT("-asynctexturecompilationmaxconcurrency="), MaxConcurrency))
-			{
-				CVarAsyncTextureCompilationMaxConcurrency->Set(MaxConcurrency, ECVF_SetByCommandline);
-			}
+			AsyncCompilationHelpers::EnsureInitializedCVars(
+				TEXT("texture"),
+				CVarAsyncTextureStandard.AsyncCompilation,
+				CVarAsyncTextureStandard.AsyncCompilationMaxConcurrency,
+				GET_MEMBER_NAME_CHECKED(UEditorExperimentalSettings, bEnableAsyncTextureCompilation));
 		}
 	}
 }
 
+FTextureCompilingManager::FTextureCompilingManager()
+	: Notification(LOCTEXT("Textures", "Textures"))
+{
+}
+
 EQueuedWorkPriority FTextureCompilingManager::GetBasePriority(UTexture* InTexture) const
 {
 	return TextureCompilingManagerImpl::GetBasePriority(InTexture);
 }
 
+FQueuedThreadPool* FTextureCompilingManager::GetThreadPool() const
+{
+	static FQueuedThreadPoolWrapper* GTextureThreadPool = nullptr;
+	if (GTextureThreadPool == nullptr)
+	{
+		TextureCompilingManagerImpl::EnsureInitializedCVars();
+
+		const auto TexturePriorityMapper = [](EQueuedWorkPriority TexturePriority) { return FMath::Max(TexturePriority, EQueuedWorkPriority::Low); };
+
+		// Textures will be scheduled on the asset thread pool, where concurrency limits might by dynamically adjusted depending on memory constraints.
+		GTextureThreadPool = new FQueuedThreadPoolWrapper(FAssetCompilingManager::Get().GetThreadPool(), -1, TexturePriorityMapper);
+
+		AsyncCompilationHelpers::BindThreadPoolToCVar(
+			GTextureThreadPool,
+			CVarAsyncTextureStandard.AsyncCompilation,
+			CVarAsyncTextureStandard.AsyncCompilationResume,
+			CVarAsyncTextureStandard.AsyncCompilationMaxConcurrency
+		);
+	}
+
+	return GTextureThreadPool;
+}
+
 void FTextureCompilingManager::Shutdown()
 {
 	bHasShutdown = true;
@@ -179,65 +138,6 @@ void FTextureCompilingManager::Shutdown()
 	}
 }
 
-FQueuedThreadPool* FTextureCompilingManager::GetThreadPool() const
-{
-	static FQueuedThreadPoolWrapper* GTextureThreadPool = nullptr;
-	if (GTextureThreadPool == nullptr)
-	{
-		TextureCompilingManagerImpl::EnsureInitializedCVars();
-
-		const auto TexturePriorityMapper = [](EQueuedWorkPriority TexturePriority) { return FMath::Max(TexturePriority, EQueuedWorkPriority::Low); };
-		const int32 MaxConcurrency = CVarAsyncTextureCompilationMaxConcurrency.GetValueOnAnyThread();
-
-		// Textures will be scheduled on the asset thread pool, where concurrency limits might by dynamically adjusted depending on memory constraints.
-		GTextureThreadPool = new FQueuedThreadPoolWrapper(FAssetCompilingManager::Get().GetThreadPool(), MaxConcurrency, TexturePriorityMapper);
-
-		CVarAsyncTextureCompilation->SetOnChangedCallback(
-			FConsoleVariableDelegate::CreateLambda(
-				[](IConsoleVariable* Variable)
-				{
-					if (Variable->GetInt() == 2)
-					{
-						GTextureThreadPool->Pause();
-					}
-					else
-					{
-						GTextureThreadPool->Resume();
-					}
-				}
-				)
-			);
-
-		CVarAsyncTextureCompilationResume->SetOnChangedCallback(
-			FConsoleVariableDelegate::CreateLambda(
-				[](IConsoleVariable* Variable)
-				{
-					if (Variable->GetInt() > 0)
-					{
-						GTextureThreadPool->Resume(Variable->GetInt());
-					}
-				}
-				)
-			);
-
-		CVarAsyncTextureCompilationMaxConcurrency->SetOnChangedCallback(
-			FConsoleVariableDelegate::CreateLambda(
-				[](IConsoleVariable* Variable)
-				{
-					GTextureThreadPool->SetMaxConcurrency(Variable->GetInt());
-				}
-				)
-			);
-
-		if (CVarAsyncTextureCompilation->GetInt() == 2)
-		{
-			GTextureThreadPool->Pause();
-		}
-	}
-
-	return GTextureThreadPool;
-}
-
 bool FTextureCompilingManager::IsAsyncTextureCompilationEnabled() const
 {
 	if (bHasShutdown)
@@ -247,59 +147,18 @@ bool FTextureCompilingManager::IsAsyncTextureCompilationEnabled() const
 
 	TextureCompilingManagerImpl::EnsureInitializedCVars();
 
-	return CVarAsyncTextureCompilation.GetValueOnAnyThread() != 0;
+	return CVarAsyncTextureStandard.AsyncCompilation.GetValueOnAnyThread() != 0;
 }
 
 void FTextureCompilingManager::UpdateCompilationNotification()
 {
-	check(IsInGameThread());
-	static TWeakPtr<SNotificationItem> TextureCompilationPtr;
-
-	TSharedPtr<SNotificationItem> NotificationItem = TextureCompilationPtr.Pin();
-
-	const int32 NumRemainingCompilations = GetNumRemainingTextures();
-	if (NumRemainingCompilations == 0)
-	{
-		if (NotificationItem.IsValid())
-		{
-			NotificationItem->SetText(NSLOCTEXT("TextureBuild", "TextureBuildFinished", "Textures are ready!"));
-			NotificationItem->SetCompletionState(SNotificationItem::CS_Success);
-			NotificationItem->ExpireAndFadeout();
-
-			TextureCompilationPtr.Reset();
-		}
-	}
-	else
-	{
-		if (!NotificationItem.IsValid())
-		{
-			FNotificationInfo Info(NSLOCTEXT("TextureBuild", "TextureBuildInProgress", "Preparing Textures"));
-			Info.bFireAndForget = false;
-
-			// Setting fade out and expire time to 0 as the expire message is currently very obnoxious
-			Info.FadeOutDuration = 0.0f;
-			Info.ExpireDuration = 0.0f;
-
-			NotificationItem = FSlateNotificationManager::Get().AddNotification(Info);
-			TextureCompilationPtr = NotificationItem;
-		}
-
-		FFormatNamedArguments Args;
-		Args.Add(TEXT("BuildTasks"), FText::AsNumber(NumRemainingCompilations));
-		FText ProgressMessage = FText::Format(NSLOCTEXT("TextureBuild", "TextureBuildInProgressFormat", "Preparing Textures ({BuildTasks})"), Args);
-
-		NotificationItem->SetCompletionState(SNotificationItem::CS_Pending);
-		NotificationItem->SetVisibility(EVisibility::HitTestInvisible);
-		NotificationItem->SetText(ProgressMessage);
-	}
+	Notification.Update(GetNumRemainingTextures());
 }
 
-void FTextureCompilingManager::FinishTextureCompilation(UTexture* Texture)
+void FTextureCompilingManager::PostCompilation(UTexture* Texture)
 {
-	using namespace TextureCompilingManagerImpl;
-
 	check(IsInGameThread());
-	TRACE_CPUPROFILER_EVENT_SCOPE(FinishTextureCompilation);
+	TRACE_CPUPROFILER_EVENT_SCOPE(FTextureCompilingManager::PostCompilation);
 
 	UE_LOG(LogTexture, Verbose, TEXT("Refreshing texture %s because it is ready"), *Texture->GetName());
 
@@ -334,7 +193,7 @@ int32 FTextureCompilingManager::GetNumRemainingTextures() const
 	return Num;
 }
 
-void FTextureCompilingManager::AddTextures(const TArray<UTexture*>& InTextures)
+void FTextureCompilingManager::AddTextures(TArrayView<UTexture* const> InTextures)
 {
 	TRACE_CPUPROFILER_EVENT_SCOPE(FTextureCompilingManager::AddTextures)
 	check(IsInGameThread());
@@ -366,14 +225,14 @@ void FTextureCompilingManager::AddTextures(const TArray<UTexture*>& InTextures)
 	}
 }
 
-void FTextureCompilingManager::FinishCompilation(const TArray<UTexture*>& InTextures)
+void FTextureCompilingManager::FinishCompilation(TArrayView<UTexture* const> InTextures)
 {
 	TRACE_CPUPROFILER_EVENT_SCOPE(FTextureCompilingManager::FinishCompilation);
 
 	using namespace TextureCompilingManagerImpl;
 	check(IsInGameThread());
 
-	TSet<UTexture*> PendingTextures;
+	TArray<UTexture*> PendingTextures;
 	PendingTextures.Reserve(InTextures.Num());
 
 	int32 TextureIndex = 0;
@@ -390,61 +249,45 @@ void FTextureCompilingManager::FinishCompilation(const TArray<UTexture*>& InText
 
 	if (PendingTextures.Num())
 	{
-		FScopedSlowTask SlowTask((float)PendingTextures.Num(), LOCTEXT("FinishTextureCompilation", "Waiting on texture preparation"), true);
-		SlowTask.MakeDialogDelayed(1.0f);
-
-		struct FTextureTask : public IQueuedWork
+		class FCompilableTexture : public AsyncCompilationHelpers::ICompilable
 		{
+		public:
+			FCompilableTexture(UTexture* InTexture)
+				: Texture(InTexture)
+			{
+			}
+
 			TStrongObjectPtr<UTexture> Texture;
-			FEvent* Event;
-			FTextureTask() { Event = FPlatformProcess::GetSynchEventFromPool(true); }
-			~FTextureTask() { FPlatformProcess::ReturnSynchEventToPool(Event); }
-			void DoThreadedWork() override { FOptionalTaskTagScope Scope(ETaskTag::EParallelGameThread); Texture->FinishCachePlatformData(); Event->Trigger(); };
-			void Abandon() override { }
-		};
+			void EnsureCompletion() override { Texture->FinishCachePlatformData(); }
+			FName GetName() override { return Texture->GetFName(); }
 
-		// Perform forced compilation on as many thread as possible in high priority since the game-thread is waiting
-		TArray<FTextureTask> PendingTasks;
-		PendingTasks.SetNum(PendingTextures.Num());
-		
-		int32 PendingTaskIndex = 0;
-		for (UTexture* Texture : PendingTextures)
-		{
-			PendingTasks[PendingTaskIndex].Texture.Reset(Texture);
-			GThreadPool->AddQueuedWork(&PendingTasks[PendingTaskIndex], EQueuedWorkPriority::High);
-			PendingTaskIndex++;
-		}
+		};
 
-		auto UpdateProgress =
-			[&SlowTask](float Progress, int32 Done, int32 Total, const FString& CurrentObjectsName)
+		TArray<FCompilableTexture> CompilableTextures(PendingTextures);
+		using namespace AsyncCompilationHelpers;
+		FObjectCacheContextScope ObjectCacheScope;
+		AsyncCompilationHelpers::FinishCompilation(
+			[&CompilableTextures](int32 Index)	-> ICompilable& { return CompilableTextures[Index]; },
+			CompilableTextures.Num(),
+			LOCTEXT("Textures", "Textures"),
+			LogTexture,
+			[this](ICompilable* Object)
 			{
-				return SlowTask.EnterProgressFrame(Progress, FText::FromString(FString::Printf(TEXT("Waiting for textures to be ready %d/%d (%s) ..."), Done, Total, *CurrentObjectsName)));
-			};
+				UTexture* Texture = static_cast<FCompilableTexture*>(Object)->Texture.Get();
+				PostCompilation(Texture);
 
-		for (FTextureTask& PendingTask : PendingTasks)
-		{
-			UTexture* Texture = PendingTask.Texture.Get();
-			const FString TextureName = Texture->GetName();
-			// Be nice with the game thread and tick the progress at 60 fps even when no progress is being made...
-			while (!PendingTask.Event->Wait(16))
-			{
-				UpdateProgress(0.0f, TextureIndex, InTextures.Num(), TextureName);
+				for (TSet<TWeakObjectPtr<UTexture>>& Bucket : RegisteredTextureBuckets)
+				{
+					Bucket.Remove(Texture);
+				}
 			}
-			UE_LOG(LogTexture, Display, TEXT("Waiting for textures to be ready %d/%d (%s) ..."), TextureIndex, InTextures.Num(), *TextureName);
-			UpdateProgress(1.f, TextureIndex++, InTextures.Num(), TextureName);
-			FinishTextureCompilation(Texture);
+		);
 
-			for (TSet<TWeakObjectPtr<UTexture>>& Bucket : RegisteredTextureBuckets)
-			{
-				Bucket.Remove(Texture);
-			}
-		}
+		PostCompilation(PendingTextures);
 	}
-
-	PostTextureCompilation(PendingTextures);
 }
 
-void FTextureCompilingManager::PostTextureCompilation(const TSet<UTexture*>& InCompiledTextures)
+void FTextureCompilingManager::PostCompilation(TArrayView<UTexture* const> InCompiledTextures)
 {
 	using namespace TextureCompilingManagerImpl;
 	if (InCompiledTextures.Num())
@@ -557,8 +400,8 @@ bool FTextureCompilingManager::RequestPriorityChange(UTexture* InTexture, EQueue
 							TEXT("Changing priority of %s (%s) from %s to %s"),
 							*InTexture->GetName(),
 							*GetLODGroupName(InTexture),
-							GetPriorityName(OldPriority),
-							GetPriorityName(InPriority)
+							LexToString(OldPriority),
+							LexToString(InPriority)
 						);
 
 						return true;
@@ -580,7 +423,7 @@ void FTextureCompilingManager::ProcessTextures(bool bLimitExecutionTime, int32 M
 	if (GetNumRemainingTextures())
 	{
 		FObjectCacheContextScope ObjectCacheScope;
-		TSet<UTexture*> ProcessedTextures;
+		TArray<UTexture*> ProcessedTextures;
 		{
 			TRACE_CPUPROFILER_EVENT_SCOPE(ProcessFinishedTextures);
 
@@ -606,7 +449,7 @@ void FTextureCompilingManager::ProcessTextures(bool bLimitExecutionTime, int32 M
 							const bool bHasTimeLeft = bLimitExecutionTime ? ((FPlatformTime::Seconds() - TickStartTime) < MaxSecondsPerFrame) : true;
 							if ((bIsHighestPrio || bHasTimeLeft) && Texture->IsAsyncCacheComplete())
 							{
-								FinishTextureCompilation(Texture.Get());
+								PostCompilation(Texture.Get());
 								ProcessedTextures.Add(Texture.Get());
 							}
 							else
@@ -677,10 +520,7 @@ void FTextureCompilingManager::ProcessTextures(bool bLimitExecutionTime, int32 M
 			}
 		}
 
-		if (ProcessedTextures.Num())
-		{
-			PostTextureCompilation(ProcessedTextures);
-		}
+		PostCompilation(ProcessedTextures);
 	}
 }
 
diff --git a/Engine/Source/Runtime/Engine/Public/AsyncCompilationHelpers.h b/Engine/Source/Runtime/Engine/Public/AsyncCompilationHelpers.h
new file mode 100644
--- /dev/null
+++ b/Engine/Source/Runtime/Engine/Public/AsyncCompilationHelpers.h
@@ -0,0 +1,76 @@
+// Copyright Epic Games, Inc. All Rights Reserved.
+#pragma once
+
+#include "CoreMinimal.h"
+
+#if WITH_EDITOR
+
+#include "HAL/IConsoleManager.h"
+#include "Logging/LogMacros.h"
+
+DECLARE_LOG_CATEGORY_EXTERN(LogAsyncCompilation, Log, All);
+
+class SNotificationItem;
+class FQueuedThreadPoolWrapper;
+
+class FAsyncCompilationNotification
+{
+public:
+	FAsyncCompilationNotification(FText InAssetType)
+		: AssetType(InAssetType)
+	{
+	}
+
+	void Update(int32 NumJobs);
+private:
+	TWeakPtr<SNotificationItem> NotificationItemPtr;
+	FText AssetType;
+};
+
+namespace AsyncCompilationHelpers
+{
+	class FAsyncCompilationStandardCVars
+	{
+	public:
+		TAutoConsoleVariable<int32> AsyncCompilation;
+		TAutoConsoleVariable<int32> AsyncCompilationMaxConcurrency;
+		FAutoConsoleCommand         AsyncCompilationFinishAll;
+		TAutoConsoleVariable<int32> AsyncCompilationResume;
+
+		FAsyncCompilationStandardCVars(const TCHAR* AssetType, const TCHAR* AssetTypePluralLowerCase, const FConsoleCommandDelegate& FinishAllCommand);
+	};
+
+	class ICompilable
+	{
+	public:
+		inline virtual ~ICompilable() {};
+		virtual void EnsureCompletion() = 0;
+		virtual FName GetName() = 0;
+	};
+
+	void FinishCompilation(
+		TFunctionRef<ICompilable& (int32 Index)> Getter,
+		int32 Num,
+		const FText& AssetType,
+		const FLogCategoryBase& LogCategory,
+		TFunctionRef<void(ICompilable*)> PostCompileSingle
+	);
+
+	void EnsureInitializedCVars(
+		const TCHAR* InName,
+		TAutoConsoleVariable<int32>& InCVarAsyncCompilation,
+		TAutoConsoleVariable<int32>& InCVarAsyncCompilationMaxConcurrency,
+		FName ExperimentalSettingsName = NAME_None);
+
+	void BindThreadPoolToCVar(
+		FQueuedThreadPoolWrapper* InThreadPoolWrapper,
+		TAutoConsoleVariable<int32>& InCVarAsyncCompilation,
+		TAutoConsoleVariable<int32>& InCVarAsyncCompilationResume,
+		TAutoConsoleVariable<int32>& InCVarAsyncCompilationMaxConcurrency
+		);
+
+	ENGINE_API void SaveStallStack(uint64 Cycles);
+	ENGINE_API void DumpStallStacks();
+}
+
+#endif // #if WITH_EDITOR
\ No newline at end of file
diff --git a/Engine/Source/Runtime/Engine/Public/StaticMeshCompiler.h b/Engine/Source/Runtime/Engine/Public/StaticMeshCompiler.h
--- a/Engine/Source/Runtime/Engine/Public/StaticMeshCompiler.h
+++ b/Engine/Source/Runtime/Engine/Public/StaticMeshCompiler.h
@@ -4,8 +4,7 @@
 #include "CoreMinimal.h"
 #include "UObject/WeakObjectPtr.h"
 #include "Containers/Set.h"
-#include "Tickable.h"
-#include "Stats/Stats.h"
+#include "AsyncCompilationHelpers.h"
 
 #if WITH_EDITOR
 
@@ -34,12 +33,12 @@ public:
 	/** 
 	 * Adds static meshes compiled asynchronously so they are monitored. 
 	 */
-	ENGINE_API void AddStaticMeshes(const TArray<UStaticMesh*>& InStaticMeshes);
+	ENGINE_API void AddStaticMeshes(TArrayView<UStaticMesh* const> InStaticMeshes);
 
 	/** 
 	 * Blocks until completion of the requested static meshes.
 	 */
-	ENGINE_API void FinishCompilation(const TArray<UStaticMesh*>& InStaticMeshes);
+	ENGINE_API void FinishCompilation(TArrayView<UStaticMesh* const> InStaticMeshes);
 
 	/** 
 	 * Blocks until completion of all async static mesh compilation.
@@ -66,20 +65,25 @@ public:
 	 */
 	ENGINE_API void Shutdown();
 
-	/** Called once per frame, fetches completed tasks and applies them to the scene. */
-	ENGINE_API void ProcessAsyncTasks(bool bLimitExecutionTime = false);
-
 private:
 	friend class FAssetCompilingManager;
-	
+
+	FStaticMeshCompilingManager();
+
+	/** Called once per frame, fetches completed tasks and applies them to the scene. */
+	void ProcessAsyncTasks(bool bLimitExecutionTime = false);
+
 	bool bHasShutdown = false;
 	TSet<TWeakObjectPtr<UStaticMesh>> RegisteredStaticMesh;
+	FAsyncCompilationNotification Notification;
+
 	void FinishCompilationsForGame();
 	void Reschedule();
 	void ProcessStaticMeshes(bool bLimitExecutionTime, int32 MinBatchSize = 1);
 	void UpdateCompilationNotification();
-	void PostStaticMeshesCompilation(const TSet<UStaticMesh*>& InStaticMeshes);
-	void FinishStaticMeshCompilation(UStaticMesh* StaticMesh);
+
+	void PostCompilation(TArrayView<UStaticMesh* const> InStaticMeshes);
+	void PostCompilation(UStaticMesh* StaticMesh);
 };
 
 #endif // #if WITH_EDITOR
\ No newline at end of file
diff --git a/Engine/Source/Runtime/Engine/Public/TextureCompiler.h b/Engine/Source/Runtime/Engine/Public/TextureCompiler.h
--- a/Engine/Source/Runtime/Engine/Public/TextureCompiler.h
+++ b/Engine/Source/Runtime/Engine/Public/TextureCompiler.h
@@ -4,8 +4,7 @@
 #include "CoreMinimal.h"
 #include "UObject/WeakObjectPtr.h"
 #include "Containers/Set.h"
-#include "Tickable.h"
-#include "Stats/Stats.h"
+#include "AsyncCompilationHelpers.h"
 
 #if WITH_EDITOR
 
@@ -31,12 +30,12 @@ public:
 	/** 
 	 * Adds textures compiled asynchronously so they are monitored. 
 	 */
-	ENGINE_API void AddTextures(const TArray<UTexture*>& InTextures);
+	ENGINE_API void AddTextures(TArrayView<UTexture* const> InTextures);
 
 	/** 
 	 * Blocks until completion of the requested textures.
 	 */
-	ENGINE_API void FinishCompilation(const TArray<UTexture*>& InTextures);
+	ENGINE_API void FinishCompilation(TArrayView<UTexture* const> InTextures);
 
 	/** 
 	 * Blocks until completion of all async texture compilation.
@@ -68,19 +67,24 @@ public:
 	 */
 	ENGINE_API void Shutdown();
 
-	/** Called once per frame, fetches completed tasks and applies them to the scene. */
-	ENGINE_API void ProcessAsyncTasks(bool bLimitExecutionTime = false);
 
 private:
 	friend class FAssetCompilingManager;
-	bool bHasShutdown = false;
-	TArray<TSet<TWeakObjectPtr<UTexture>>> RegisteredTextureBuckets;
 
-	void PostTextureCompilation(const TSet<UTexture*>& InCompiledTextures);
+	FTextureCompilingManager();
+	
+	void ProcessAsyncTasks(bool bLimitExecutionTime = false);
+	
 	void FinishCompilationsForGame();
 	void ProcessTextures(bool bLimitExecutionTime, int32 MaximumPriority = -1);
 	void UpdateCompilationNotification();
-	void FinishTextureCompilation(UTexture* Texture);
+	
+	void PostCompilation(UTexture* Texture);
+	void PostCompilation(TArrayView<UTexture* const> InCompiledTextures);
+
+	bool bHasShutdown = false;
+	TArray<TSet<TWeakObjectPtr<UTexture>>> RegisteredTextureBuckets;
+	FAsyncCompilationNotification Notification;
 };
 
 #endif
\ No newline at end of file
