Make sure to PostCompile only the pending static meshes

diff --git a/Engine/Source/Runtime/Engine/Private/StaticMeshCompiler.cpp b/Engine/Source/Runtime/Engine/Private/StaticMeshCompiler.cpp
--- a/Engine/Source/Runtime/Engine/Private/StaticMeshCompiler.cpp
+++ b/Engine/Source/Runtime/Engine/Private/StaticMeshCompiler.cpp
@@ -317,7 +317,7 @@ void FStaticMeshCompilingManager::FinishCompilation(TArrayView<UStaticMesh* cons
 			}
 		);
 
-		PostCompilation(InStaticMeshes);
+		PostCompilation(PendingStaticMeshes);
 	}
 }
 
