WorldPartition HLOD - HLODBuilder - Force a synchronous load of the parent HLODLayer asset

diff --git a/Engine/Source/Runtime/Engine/Private/WorldPartition/HLOD/HLODBuilder.cpp b/Engine/Source/Runtime/Engine/Private/WorldPartition/HLOD/HLODBuilder.cpp
--- a/Engine/Source/Runtime/Engine/Private/WorldPartition/HLOD/HLODBuilder.cpp
+++ b/Engine/Source/Runtime/Engine/Private/WorldPartition/HLOD/HLODBuilder.cpp
@@ -67,7 +67,7 @@ public:
 			HLODActor->SetChildrenPrimitives(InSubComponents);
 			HLODActor->RuntimeGrid = HLODLayer->GetRuntimeGrid(HLODLevel);
 			HLODActor->SetLODLevel(HLODLevel);
-			HLODActor->SetHLODLayer(HLODLayer->GetParentLayer().Get());
+			HLODActor->SetHLODLayer(HLODLayer->GetParentLayer().LoadSynchronous());
 			HLODActor->SetSubActorsHLODLayer(HLODLayer);
 			HLODActor->SetGridIndices(Context->GridIndexX, Context->GridIndexY, Context->GridIndexZ);
 		}
