WorldPartition HLOD - Added proper SCC operations to manage grid actors (ASpatialHashRuntimeGridInfo) packages

diff --git a/Engine/Source/Runtime/Engine/Private/WorldPartition/WorldPartitionRuntimeSpatialHash.cpp b/Engine/Source/Runtime/Engine/Private/WorldPartition/WorldPartitionRuntimeSpatialHash.cpp
--- a/Engine/Source/Runtime/Engine/Private/WorldPartition/WorldPartitionRuntimeSpatialHash.cpp
+++ b/Engine/Source/Runtime/Engine/Private/WorldPartition/WorldPartitionRuntimeSpatialHash.cpp
@@ -1504,6 +1504,33 @@ void UWorldPartitionRuntimeSpatialHash::FlushStreaming()
 	StreamingGrids.Empty();
 }
 
+void CheckoutAndSavePackage(UPackage* Package, ISourceControlHelper* SourceControlHelper)
+{
+	// Checkout package
+	Package->MarkAsFullyLoaded();
+
+	if (!SourceControlHelper->Checkout(Package))
+	{
+		UE_LOG(LogWorldPartitionRuntimeSpatialHash, Error, TEXT("Error checking out package %s."), *Package->GetName());
+		check(0);
+	}
+
+	// Save package
+	FString PackageFileName = SourceControlHelper->GetFilename(Package);
+	if (!UPackage::SavePackage(Package, nullptr, RF_Standalone, *PackageFileName))
+	{
+		UE_LOG(LogWorldPartitionRuntimeSpatialHash, Error, TEXT("Error saving package %s."), *Package->GetName());
+		check(0);
+	}
+
+	// Add new package to source control
+	if (!SourceControlHelper->Add(Package))
+	{
+		UE_LOG(LogWorldPartitionRuntimeSpatialHash, Error, TEXT("Error adding package %s to source control."), *Package->GetName());
+		check(0);
+	}
+}
+
 TArray<FGuid> GenerateHLODsForGrid(UWorldPartition* WorldPartition, const FSpatialHashRuntimeGrid& RuntimeGrid, uint32 HLODLevel, FHLODGenerationContext& Context, ISourceControlHelper* SourceControlHelper, const TArray<FActorCluster>& ActorClusters)
 {
 	auto HasExceededMaxMemory = []()
@@ -1618,35 +1645,13 @@ TArray<FGuid> GenerateHLODsForGrid(UWorldPartition* WorldPartition, const FSpati
 						FGuid ActorGuid = CellHLODActor->GetActorGuid();
 						GridHLODActors.Add(ActorGuid);
 
-						// Checkout packages
 						UPackage* CellHLODActorPackage = CellHLODActor->GetPackage();
-						CellHLODActorPackage->MarkAsFullyLoaded();
-
 						if (CellHLODActorPackage->HasAnyPackageFlags(PKG_NewlyCreated))
 						{
 							NewCellHLODActors.Add(CellHLODActor);
 						}
 
-						if (!SourceControlHelper->Checkout(CellHLODActorPackage))
-						{
-							UE_LOG(LogWorldPartitionRuntimeSpatialHash, Error, TEXT("Error checking out package %s."), *CellHLODActorPackage->GetName());
-							check(0);
-						}
-
-						// Save packages
-						FString PackageFileName = SourceControlHelper->GetFilename(CellHLODActorPackage);
-						if (!UPackage::SavePackage(CellHLODActorPackage, nullptr, RF_Standalone, *PackageFileName))
-						{
-							UE_LOG(LogWorldPartitionRuntimeSpatialHash, Error, TEXT("Error saving package %s."), *CellHLODActorPackage->GetName());
-							check(0);
-						}
-
-						// Add new packages to source control
-						if (!SourceControlHelper->Add(CellHLODActorPackage))
-						{
-							UE_LOG(LogWorldPartitionRuntimeSpatialHash, Error, TEXT("Error adding package %s to source control."), *CellHLODActorPackage->GetName());
-							check(0);
-						}
+						CheckoutAndSavePackage(CellHLODActorPackage, SourceControlHelper);
 					}
 
 					// Manually tick the directory watcher and the asset registry to register newly created actors
@@ -1795,7 +1800,7 @@ void RenameHLODGrids(TMap<FName, FSpatialHashRuntimeGrid>& HLODGrids, TArray<FNa
 }
 
 // Create/destroy HLOD grid actors
-void UpdateHLODGridsActors(UWorld* World, const TMap<FName, FSpatialHashRuntimeGrid>& HLODGrids)
+void UpdateHLODGridsActors(UWorld* World, const TMap<FName, FSpatialHashRuntimeGrid>& HLODGrids, ISourceControlHelper* SourceControlHelper)
 {
 	static const FName HLODGridTag = TEXT("HLOD");
 	static const uint32 HLODGridTagLen = HLODGridTag.ToString().Len();
@@ -1814,6 +1819,9 @@ void UpdateHLODGridsActors(UWorld* World, const TMap<FName, FSpatialHashRuntimeG
 			else
 			{
 				World->DestroyActor(GridActor);
+
+				const FString ActorPackageFilename = SourceControlHelper->GetFilename(GridActor->GetPackage());
+				SourceControlHelper->Delete(ActorPackageFilename);
 			}
 		}
 	}
@@ -1842,6 +1850,8 @@ void UpdateHLODGridsActors(UWorld* World, const TMap<FName, FSpatialHashRuntimeG
 			}
 			HLODLevel = FMath::Clamp(HLODLevel + 2, 0U, LastLODColorationColorIdx);
 			GridActor->GridSettings.DebugColor = GEngine->HLODColorationColors[HLODLevel];
+
+			CheckoutAndSavePackage(GridActor->GetPackage(), SourceControlHelper);
 		}
 	}
 }
@@ -1993,7 +2003,7 @@ bool UWorldPartitionRuntimeSpatialHash::GenerateHLOD(ISourceControlHelper* Sourc
 	Algo::ForEach(InvalidHLODActors, DeleteHLODActor);
 
 	// Create/destroy HLOD grid actors
-	UpdateHLODGridsActors(GetWorld(), HLODGrids);
+	UpdateHLODGridsActors(GetWorld(), HLODGrids, SourceControlHelper);
 
 	return true;
 }
