Fixing PVS 7.7 Issues:  warning V1050: The uninitialized class member 'XYZ' is used when initializing the base class 'FXYZ'.

diff --git a/Engine/Source/Runtime/Core/Public/Windows/WindowsPlatformCodeAnalysis.h b/Engine/Source/Runtime/Core/Public/Windows/WindowsPlatformCodeAnalysis.h
--- a/Engine/Source/Runtime/Core/Public/Windows/WindowsPlatformCodeAnalysis.h
+++ b/Engine/Source/Runtime/Core/Public/Windows/WindowsPlatformCodeAnalysis.h
@@ -103,7 +103,7 @@
 	#pragma warning(disable : 6240) // warning C6240 : (<expression> && <non-zero constant>) always evaluates to the result of <expression>. Did you intend to use the bitwise-and operator?
 
 	//PVS-Studio settings:
-	//-V::505,542,581,591,601,668,677,686,688,690,703,704,711,719,720,728,730,735,751,1002,1008
+	//-V::505,542,581,591,601,668,677,686,688,690,703,704,711,719,720,728,730,735,751,1002,1008,1062
 	//-V:TRYCOMPRESSION:519,547
 	//-V:check(:501,547,560,605
 	//-V:checkf(:510
@@ -233,6 +233,13 @@
 	//-V:AccumulateParentID:678
 	//-V:Resource:623
 
+	// The following classes retain a reference to data supplied in the constructor by the derived class which can not yet be initialized.
+	//-V:FMemoryWriter(:1050
+	//-V:FObjectWriter(:1050
+	//-V:FDurationTimer(:1050
+	//-V:FScopedDurationTimer(:1050
+	//-V:FQueryFastData(:1050
+
 #else // defined(__clang_analyzer__)
 
 // A fake function marked with noreturn that acts as a marker for CA_ASSUME to ensure the
diff --git a/Engine/Source/Runtime/Engine/Classes/PhysicsEngine/BodyInstance.h b/Engine/Source/Runtime/Engine/Classes/PhysicsEngine/BodyInstance.h
--- a/Engine/Source/Runtime/Engine/Classes/PhysicsEngine/BodyInstance.h
+++ b/Engine/Source/Runtime/Engine/Classes/PhysicsEngine/BodyInstance.h
@@ -142,7 +142,7 @@ struct FInitBodiesHelperWithData : public FInitBodiesHelperBase
 {
 	FInitBodiesHelperWithData() { check(false); }
 	FInitBodiesHelperWithData(TArray<FBodyInstance*>&& InBodies, TArray<FTransform>&& InTransforms, class UBodySetup* InBodySetup, class UPrimitiveComponent* InPrimitiveComp, FPhysScene* InRBScene, const FInitBodySpawnParams& InSpawnParams, FPhysicsAggregateHandle InAggregate)
-	: FInitBodiesHelperBase(OwnedBodies, OwnedTransforms, InBodySetup, InPrimitiveComp, InRBScene, InSpawnParams, InAggregate), OwnedBodies(MoveTemp(InBodies)), OwnedTransforms(MoveTemp(InTransforms))
+	: FInitBodiesHelperBase(OwnedBodies, OwnedTransforms, InBodySetup, InPrimitiveComp, InRBScene, InSpawnParams, InAggregate), OwnedBodies(MoveTemp(InBodies)), OwnedTransforms(MoveTemp(InTransforms)) //-V1050
 	{
 		//Compute all the needed constants
 		bStatic = bCompileStatic || SpawnParams.bStaticPhysics;
@@ -154,13 +154,13 @@ struct FInitBodiesHelperWithData : public FInitBodiesHelperBase
 	}
 
 	FInitBodiesHelperWithData(const FInitBodiesHelperWithData& InHelper)
-	: FInitBodiesHelperBase(OwnedBodies, OwnedTransforms, InHelper.BodySetup, InHelper.PrimitiveComp, InHelper.PhysScene, InHelper.SpawnParams, InHelper.Aggregate), OwnedBodies(InHelper.OwnedBodies), OwnedTransforms(InHelper.OwnedTransforms)
+	: FInitBodiesHelperBase(OwnedBodies, OwnedTransforms, InHelper.BodySetup, InHelper.PrimitiveComp, InHelper.PhysScene, InHelper.SpawnParams, InHelper.Aggregate), OwnedBodies(InHelper.OwnedBodies), OwnedTransforms(InHelper.OwnedTransforms) //-V1050
 	{
 		ensure(false);
 	}
 
 	FInitBodiesHelperWithData(FInitBodiesHelperWithData&& InHelper)
-	: FInitBodiesHelperBase(OwnedBodies, OwnedTransforms, InHelper.BodySetup, InHelper.PrimitiveComp, InHelper.PhysScene, InHelper.SpawnParams, InHelper.Aggregate), OwnedBodies(MoveTemp(InHelper.OwnedBodies)), OwnedTransforms(MoveTemp(InHelper.OwnedTransforms))
+	: FInitBodiesHelperBase(OwnedBodies, OwnedTransforms, InHelper.BodySetup, InHelper.PrimitiveComp, InHelper.PhysScene, InHelper.SpawnParams, InHelper.Aggregate), OwnedBodies(MoveTemp(InHelper.OwnedBodies)), OwnedTransforms(MoveTemp(InHelper.OwnedTransforms)) //-V1050
 	{
 		//Compute all the needed constants
 		bStatic = bCompileStatic || SpawnParams.bStaticPhysics;
