Root sync scopes now operate for main & post process instances

Sync scopes always forward to the outermost scope (determined at scope creation time), rather than relying on the cached 'main' proxy, as the context might be the post process instance.

diff --git a/Engine/Source/Runtime/Engine/Private/Animation/AnimInstanceProxy.cpp b/Engine/Source/Runtime/Engine/Private/Animation/AnimInstanceProxy.cpp
--- a/Engine/Source/Runtime/Engine/Private/Animation/AnimInstanceProxy.cpp
+++ b/Engine/Source/Runtime/Engine/Private/Animation/AnimInstanceProxy.cpp
@@ -1093,8 +1093,8 @@ void FAnimInstanceProxy::UpdateAnimation_WithRoot(const FAnimationUpdateContext&
 		TGuardValue<bool> ScopeGuard(bUpdatingRoot, true);
 
 		// Anything syncing within this scope is subject to sync groups.
-		// We only enable syncing here for the main proxy or if we have no main proxy
-		const bool bEnableSyncScope = (GetMainInstanceProxy() == nullptr) || (GetMainInstanceProxy() == this);
+		// We only enable syncing here for the main instance or post process instance
+		const bool bEnableSyncScope = GetAnimInstanceObject() == GetSkelMeshComponent()->GetAnimInstance() || GetAnimInstanceObject() == GetSkelMeshComponent()->GetPostProcessInstance();
 		UE::Anim::TOptionalScopedGraphMessage<UE::Anim::FAnimSyncGroupScope> Message(bEnableSyncScope, InContext, InContext);
 
 		// update all nodes
diff --git a/Engine/Source/Runtime/Engine/Private/Animation/AnimSyncScope.cpp b/Engine/Source/Runtime/Engine/Private/Animation/AnimSyncScope.cpp
--- a/Engine/Source/Runtime/Engine/Private/Animation/AnimSyncScope.cpp
+++ b/Engine/Source/Runtime/Engine/Private/Animation/AnimSyncScope.cpp
@@ -16,9 +16,15 @@ FAnimSyncDebugInfo::FAnimSyncDebugInfo(const FAnimationBaseContext& InContext)
 FAnimSyncGroupScope::FAnimSyncGroupScope(const FAnimationBaseContext& InContext, FName InSyncGroup, EAnimGroupRole::Type InGroupRole)
 	: NodeId(InContext.GetCurrentNodeId())
 	, Proxy(*InContext.AnimInstanceProxy)
+	, OuterProxy(nullptr)
 	, SyncGroup(InSyncGroup)
 	, GroupRole(InGroupRole)
 {
+	// If we have an outer message, grab its proxy for forwarding
+	if(const FAnimSyncGroupScope* OuterMessage = InContext.GetMessage<FAnimSyncGroupScope>())
+	{
+		OuterProxy = OuterMessage->OuterProxy != nullptr ? OuterMessage->OuterProxy : &OuterMessage->Proxy;
+	}
 }
 
 void FAnimSyncGroupScope::AddTickRecord(const FAnimTickRecord& InTickRecord, const FAnimSyncParams& InSyncParams, const FAnimSyncDebugInfo& InDebugInfo)
@@ -54,10 +60,10 @@ void FAnimSyncGroupScope::AddTickRecord(const FAnimTickRecord& InTickRecord, con
 		break;
 	}
 
-	// Forward to main instance if we have one
-	if(Proxy.GetMainInstanceProxy())
+	// Forward to outer instance if we have one
+	if(OuterProxy)
 	{
-		Proxy.GetMainInstanceProxy()->AddTickRecord(InTickRecord, NewSyncParams);
+		OuterProxy->AddTickRecord(InTickRecord, NewSyncParams);
 	}
 	else
 	{
diff --git a/Engine/Source/Runtime/Engine/Public/Animation/AnimSyncScope.h b/Engine/Source/Runtime/Engine/Public/Animation/AnimSyncScope.h
--- a/Engine/Source/Runtime/Engine/Public/Animation/AnimSyncScope.h
+++ b/Engine/Source/Runtime/Engine/Public/Animation/AnimSyncScope.h
@@ -46,6 +46,9 @@ private:
 	// The proxy that we are running
 	FAnimInstanceProxy& Proxy;
 
+	// The outer proxy, if any, to forward tick records to
+	FAnimInstanceProxy* OuterProxy;
+
 	// The sync group that this scope uses.
 	FName SyncGroup;
 
