When updating an actor descriptor, add it to all touching cells if it's loaded.
This is not the correct behavior if there's only non-grid references, but this will prevent a crash for now.

diff --git a/Engine/Source/Runtime/Engine/Private/WorldPartition/WorldPartitionEditorCell.cpp b/Engine/Source/Runtime/Engine/Private/WorldPartition/WorldPartitionEditorCell.cpp
--- a/Engine/Source/Runtime/Engine/Private/WorldPartition/WorldPartitionEditorCell.cpp
+++ b/Engine/Source/Runtime/Engine/Private/WorldPartition/WorldPartitionEditorCell.cpp
@@ -40,15 +40,7 @@ void UWorldPartitionEditorCell::RemoveActor(const FWorldPartitionHandle& ActorHa
 {
 	check(ActorHandle.IsValid());
 	verify(Actors.Remove(ActorHandle));
-
-	if (ActorHandle.IsLoaded())
-	{
-		verify(LoadedActors.Remove(ActorHandle));
-	}
-	else
-	{
-		check(!LoadedActors.Contains(ActorHandle));
-	}
+	LoadedActors.Remove(ActorHandle);
 }
 #endif
 
