Prevent manual delete of selected PartitionActors.

diff --git a/Engine/Source/Runtime/Engine/Private/ActorPartition/ActorPartitionSubsystem.cpp b/Engine/Source/Runtime/Engine/Private/ActorPartition/ActorPartitionSubsystem.cpp
--- a/Engine/Source/Runtime/Engine/Private/ActorPartition/ActorPartitionSubsystem.cpp
+++ b/Engine/Source/Runtime/Engine/Private/ActorPartition/ActorPartitionSubsystem.cpp
@@ -210,6 +210,10 @@ public:
 
 				FoundActor = CastChecked<APartitionActor>(DescActor);
 				check(FoundActor->GridSize == InGridSize && FoundActor->GetGridGuid() == InGuid);
+				// Typically, this could happen if a PartitionActor was manually deleted (it's now be prevented by APartitionActor::CanDeleteSelectedActor);
+				// If at some point it's decided to be supported, a solution for a found PendingKill PartitionActor is to
+				// call Modify(), Rename the actor to a unique/trash name and behave as if no actor was found
+				check(!FoundActor->IsPendingKill());
 				return false;
 			}
 			return true;
diff --git a/Engine/Source/Runtime/Engine/Private/ActorPartition/PartitionActor.cpp b/Engine/Source/Runtime/Engine/Private/ActorPartition/PartitionActor.cpp
--- a/Engine/Source/Runtime/Engine/Private/ActorPartition/PartitionActor.cpp
+++ b/Engine/Source/Runtime/Engine/Private/ActorPartition/PartitionActor.cpp
@@ -8,6 +8,8 @@
 #include "WorldPartition/ActorPartition/PartitionActorDesc.h"
 #endif
 
+#define LOCTEXT_NAMESPACE "PartitionActor"
+
 APartitionActor::APartitionActor(const FObjectInitializer& ObjectInitializer)
 	: Super(ObjectInitializer)
 #if WITH_EDITORONLY_DATA
@@ -24,4 +26,19 @@ TUniquePtr<FWorldPartitionActorDesc> APartitionActor::CreateClassActorDesc() con
 {
 	return TUniquePtr<FWorldPartitionActorDesc>(new FPartitionActorDesc());
 }
-#endif
\ No newline at end of file
+
+bool APartitionActor::CanDeleteSelectedActor(FText& OutReason) const
+{
+	if (!Super::CanDeleteSelectedActor(OutReason))
+	{
+		return false;
+	}
+
+	check(GetLevel());
+	check(GetLevel()->bIsPartitioned);
+	OutReason = LOCTEXT("CantDeleteSelectedPartitionActor", "Can't delete Partition Actor (automatically managed).");
+	return false;
+}
+#endif
+
+#undef LOCTEXT_NAMESPACE
\ No newline at end of file
diff --git a/Engine/Source/Runtime/Engine/Public/ActorPartition/PartitionActor.h b/Engine/Source/Runtime/Engine/Public/ActorPartition/PartitionActor.h
--- a/Engine/Source/Runtime/Engine/Public/ActorPartition/PartitionActor.h
+++ b/Engine/Source/Runtime/Engine/Public/ActorPartition/PartitionActor.h
@@ -24,6 +24,7 @@ public:
 	virtual TUniquePtr<class FWorldPartitionActorDesc> CreateClassActorDesc() const override;
 	virtual uint32 GetDefaultGridSize(UWorld* InWorld) const PURE_VIRTUAL(APartitionActor, return 0;)
 	virtual FGuid GetGridGuid() const { return FGuid(); }
+	virtual bool CanDeleteSelectedActor(FText& OutReason) const override;
 #endif
 	//~ End AActor Interface	
 
