Fix for incorrect state machine debug visibility with child anim blueprints attached to an instance

diff --git a/Engine/Source/Runtime/Engine/Private/Animation/AnimInstance.cpp b/Engine/Source/Runtime/Engine/Private/Animation/AnimInstance.cpp
--- a/Engine/Source/Runtime/Engine/Private/Animation/AnimInstance.cpp
+++ b/Engine/Source/Runtime/Engine/Private/Animation/AnimInstance.cpp
@@ -432,19 +432,6 @@ void UAnimInstance::UpdateAnimation(float DeltaSeconds, bool bNeedsValidRootMoti
 #if WITH_EDITOR
 	if (GIsEditor)
 	{
-		// Reset the anim graph visualization
-		if (Proxy.HasRootNode())
-		{
-			if (UAnimBlueprintGeneratedClass* AnimBlueprintClass = Cast<UAnimBlueprintGeneratedClass>(GetClass()))
-			{
-				UAnimBlueprint* AnimBP = Cast<UAnimBlueprint>(AnimBlueprintClass->ClassGeneratedBy);
-				if (AnimBP && AnimBP->GetObjectBeingDebugged() == this)
-				{
-					AnimBlueprintClass->GetAnimBlueprintDebugData().ResetNodeVisitSites();
-				}
-			}
-		}
-
 		// Update the lifetimer and see if we should use the snapshot instead
 		CurrentLifeTimerScrubPosition += DeltaSeconds;
 		LifeTimer = FMath::Max<double>(CurrentLifeTimerScrubPosition, LifeTimer);
diff --git a/Engine/Source/Runtime/Engine/Private/Animation/AnimInstanceProxy.cpp b/Engine/Source/Runtime/Engine/Private/Animation/AnimInstanceProxy.cpp
--- a/Engine/Source/Runtime/Engine/Private/Animation/AnimInstanceProxy.cpp
+++ b/Engine/Source/Runtime/Engine/Private/Animation/AnimInstanceProxy.cpp
@@ -358,6 +358,11 @@ void FAnimInstanceProxy::PreUpdate(UAnimInstance* InAnimInstance, float DeltaSec
 
 	InitializeObjects(InAnimInstance);
 
+	if (FAnimBlueprintDebugData* DebugData = GetAnimBlueprintDebugData())
+	{
+		DebugData->ResetNodeVisitSites();
+	}
+
 	if (SkelMeshComp)
 	{
 		// Save off LOD level that we're currently using.
