Rename "MassInKg" to "Mass (kg)"

diff --git a/Engine/Source/Runtime/Engine/Classes/PhysicsEngine/BodyInstance.h b/Engine/Source/Runtime/Engine/Classes/PhysicsEngine/BodyInstance.h
--- a/Engine/Source/Runtime/Engine/Classes/PhysicsEngine/BodyInstance.h
+++ b/Engine/Source/Runtime/Engine/Classes/PhysicsEngine/BodyInstance.h
@@ -431,7 +431,7 @@ protected:
 
 	/**Mass of the body in KG. By default we compute this based on physical material and mass scale.
 	*@see bOverrideMass to set this directly */
-	UPROPERTY(EditAnywhere, BlueprintReadOnly, Category = Physics, meta = (editcondition = "bOverrideMass", ClampMin = "0.001", UIMin = "0.001", DisplayName = "MassInKg"))
+	UPROPERTY(EditAnywhere, BlueprintReadOnly, Category = Physics, meta = (editcondition = "bOverrideMass", ClampMin = "0.001", UIMin = "0.001", DisplayName = "Mass (kg)"))
 	float MassInKgOverride;
 
 	/** The body setup holding the default body instance and its collision profile. */
