Restore shielding against duplicates to avoid crash in certain cases

diff --git a/Engine/Source/Runtime/Engine/Private/TextureCompiler.cpp b/Engine/Source/Runtime/Engine/Private/TextureCompiler.cpp
--- a/Engine/Source/Runtime/Engine/Private/TextureCompiler.cpp
+++ b/Engine/Source/Runtime/Engine/Private/TextureCompiler.cpp
@@ -232,7 +232,7 @@ void FTextureCompilingManager::FinishCompilation(TArrayView<UTexture* const> InT
 	using namespace TextureCompilingManagerImpl;
 	check(IsInGameThread());
 
-	TArray<UTexture*> PendingTextures;
+	TSet<UTexture*> PendingTextures;
 	PendingTextures.Reserve(InTextures.Num());
 
 	int32 TextureIndex = 0;
@@ -260,10 +260,10 @@ void FTextureCompilingManager::FinishCompilation(TArrayView<UTexture* const> InT
 			TStrongObjectPtr<UTexture> Texture;
 			void EnsureCompletion() override { Texture->FinishCachePlatformData(); }
 			FName GetName() override { return Texture->GetFName(); }
-
 		};
 
-		TArray<FCompilableTexture> CompilableTextures(PendingTextures);
+		TArray<UTexture*> UniqueTextures(PendingTextures.Array());
+		TArray<FCompilableTexture> CompilableTextures(UniqueTextures);
 		using namespace AsyncCompilationHelpers;
 		FObjectCacheContextScope ObjectCacheScope;
 		AsyncCompilationHelpers::FinishCompilation(
@@ -283,7 +283,7 @@ void FTextureCompilingManager::FinishCompilation(TArrayView<UTexture* const> InT
 			}
 		);
 
-		PostCompilation(PendingTextures);
+		PostCompilation(UniqueTextures);
 	}
 }
 
