Properly release WorldPartition Actor Handles & References (fixes crash when closing world)

diff --git a/Engine/Source/Runtime/Engine/Private/WorldPartition/WorldPartitionEditorCell.cpp b/Engine/Source/Runtime/Engine/Private/WorldPartition/WorldPartitionEditorCell.cpp
--- a/Engine/Source/Runtime/Engine/Private/WorldPartition/WorldPartitionEditorCell.cpp
+++ b/Engine/Source/Runtime/Engine/Private/WorldPartition/WorldPartitionEditorCell.cpp
@@ -20,6 +20,15 @@ void UWorldPartitionEditorCell::Serialize(FArchive& Ar)
 }
 
 #if WITH_EDITOR
+
+void UWorldPartitionEditorCell::BeginDestroy()
+{
+	// Release WorldPartition Actor Handles/References
+	Actors.Empty();
+	LoadedActors.Empty();
+	Super::BeginDestroy();
+}
+
 void UWorldPartitionEditorCell::AddActor(const FWorldPartitionHandle& ActorHandle)
 {
 	check(ActorHandle.IsValid());
diff --git a/Engine/Source/Runtime/Engine/Public/WorldPartition/WorldPartitionEditorCell.h b/Engine/Source/Runtime/Engine/Public/WorldPartition/WorldPartitionEditorCell.h
--- a/Engine/Source/Runtime/Engine/Public/WorldPartition/WorldPartitionEditorCell.h
+++ b/Engine/Source/Runtime/Engine/Public/WorldPartition/WorldPartitionEditorCell.h
@@ -20,6 +20,7 @@ public:
 	virtual void Serialize(FArchive& Ar) override;
 
 #if WITH_EDITOR
+	virtual void BeginDestroy() override;
 	void AddActor(const FWorldPartitionHandle& ActorHandle);
 	void RemoveActor(const FWorldPartitionHandle& ActorHandle);
 
diff --git a/Engine/Source/Runtime/Landscape/Classes/LandscapeInfo.h b/Engine/Source/Runtime/Landscape/Classes/LandscapeInfo.h
--- a/Engine/Source/Runtime/Landscape/Classes/LandscapeInfo.h
+++ b/Engine/Source/Runtime/Landscape/Classes/LandscapeInfo.h
@@ -169,6 +169,8 @@ public:
 	//~ End UObject Interface
 
 #if WITH_EDITOR
+	virtual void BeginDestroy() override;
+
 	// @todo document 
 	// all below.
 	LANDSCAPE_API bool AreAllComponentsRegistered() const;
diff --git a/Engine/Source/Runtime/Landscape/Private/LandscapeInfo.cpp b/Engine/Source/Runtime/Landscape/Private/LandscapeInfo.cpp
--- a/Engine/Source/Runtime/Landscape/Private/LandscapeInfo.cpp
+++ b/Engine/Source/Runtime/Landscape/Private/LandscapeInfo.cpp
@@ -14,3 +14,15 @@ LANDSCAPE_API FLandscapeInfoLayerSettings::FLandscapeInfoLayerSettings(ULandscap
 #endif
 {
 }
+
+#if WITH_EDITOR
+
+void ULandscapeInfo::BeginDestroy()
+{
+	// Release WorldPartition Actor Handles
+	ProxyHandles.Empty();
+	SplineHandles.Empty();
+	Super::BeginDestroy();
+}
+
+#endif // WITH_EDITOR
\ No newline at end of file
