When unloading actors in partitioned worlds, don't mark them as pending kill, as other objects might still have references to them.

diff --git a/Engine/Source/Runtime/Engine/Private/Level.cpp b/Engine/Source/Runtime/Engine/Private/Level.cpp
--- a/Engine/Source/Runtime/Engine/Private/Level.cpp
+++ b/Engine/Source/Runtime/Engine/Private/Level.cpp
@@ -650,9 +650,6 @@ void ULevel::RemoveLoadedActor(AActor* Actor)
 	ActorsForGC.Remove(Actor);
 
 	OnLoadedActorRemovedFromLevelEvent.Broadcast(*Actor);
-
-	Actor->MarkPendingKill();
-	Actor->MarkComponentsAsPendingKill();
 }
 #endif
 
