Fix for build error introduced by 15020921

diff --git a/Engine/Source/Runtime/Engine/Private/Animation/AnimInstanceProxy.cpp b/Engine/Source/Runtime/Engine/Private/Animation/AnimInstanceProxy.cpp
--- a/Engine/Source/Runtime/Engine/Private/Animation/AnimInstanceProxy.cpp
+++ b/Engine/Source/Runtime/Engine/Private/Animation/AnimInstanceProxy.cpp
@@ -358,10 +358,12 @@ void FAnimInstanceProxy::PreUpdate(UAnimInstance* InAnimInstance, float DeltaSec
 
 	InitializeObjects(InAnimInstance);
 
+#if WITH_EDITORONLY_DATA
 	if (FAnimBlueprintDebugData* DebugData = GetAnimBlueprintDebugData())
 	{
 		DebugData->ResetNodeVisitSites();
 	}
+#endif
 
 	if (SkelMeshComp)
 	{
