Adding support for changing audio volume reverb enablement while in PIE

diff --git a/Engine/Source/Runtime/Engine/Classes/Sound/AudioVolume.h b/Engine/Source/Runtime/Engine/Classes/Sound/AudioVolume.h
--- a/Engine/Source/Runtime/Engine/Classes/Sound/AudioVolume.h
+++ b/Engine/Source/Runtime/Engine/Classes/Sound/AudioVolume.h
@@ -135,24 +135,17 @@ struct TStructOpsTypeTraits<FInteriorSettings> : public TStructOpsTypeTraitsBase
 
 struct FAudioVolumeProxy
 {
-	FAudioVolumeProxy()
-		: AudioVolumeID(0)
-		, WorldID(0)
-		, Priority(0.f)
-		, BodyInstance(nullptr)
-	{
-	}
-
 	FAudioVolumeProxy(const AAudioVolume* AudioVolume);
 
-	uint32 AudioVolumeID;
-	uint32 WorldID;
-	float Priority;
+	uint32 AudioVolumeID = 0;
+	uint32 WorldID = 0;
+	float Priority = 0.0f;
 	FReverbSettings ReverbSettings;
 	FInteriorSettings InteriorSettings;
 	TArray<FAudioVolumeSubmixSendSettings> SubmixSendSettings;
 	TArray<FAudioVolumeSubmixOverrideSettings> SubmixOverrideSettings;
-	FBodyInstance* BodyInstance; // This is scary
+	FBodyInstance* BodyInstance = nullptr;
+	bool bChanged = false;
 };
 
 UCLASS(hidecategories=(Advanced, Attachment, Collision, Volume))
diff --git a/Engine/Source/Runtime/Engine/Private/AudioDevice.cpp b/Engine/Source/Runtime/Engine/Private/AudioDevice.cpp
--- a/Engine/Source/Runtime/Engine/Private/AudioDevice.cpp
+++ b/Engine/Source/Runtime/Engine/Private/AudioDevice.cpp
@@ -3249,6 +3249,15 @@ void FAudioDevice::SetDefaultAudioSettings(UWorld* World, const FReverbSettings&
 	}, GET_STATID(STAT_AudioSetDefaultAudioSettings));
 }
 
+void FAudioDevice::ResetAudioVolumeProxyChangedState()
+{
+	for (TPair<uint32, FAudioVolumeProxy>& AudioVolumePair : AudioVolumeProxies)
+	{
+		FAudioVolumeProxy& Proxy = AudioVolumePair.Value;
+		Proxy.bChanged = false;
+	}
+}
+
 void FAudioDevice::GetAudioVolumeSettings(const uint32 WorldID, const FVector& Location, FAudioVolumeSettings& OutSettings) const
 {
 	check(IsInAudioThread());
@@ -3268,6 +3277,7 @@ void FAudioDevice::GetAudioVolumeSettings(const uint32 WorldID, const FVector& L
 				OutSettings.InteriorSettings = Proxy.InteriorSettings;
 				OutSettings.SubmixSendSettings = Proxy.SubmixSendSettings;
 				OutSettings.SubmixOverrideSettings = Proxy.SubmixOverrideSettings;
+				OutSettings.bChanged = Proxy.bChanged;
 				return;
 			}
 		}
@@ -4399,12 +4409,20 @@ void FAudioDevice::UpdateAudioVolumeEffects()
 		}
 	}
 
+	// Reset audio volume proxies to reset state back to before a proxy was updated
+	ResetAudioVolumeProxyChangedState();
+
 	if (PlayerAudioVolumeSettings.AudioVolumeID != CurrentAudioVolumeSettings.AudioVolumeID)
 	{
 		PreviousPlayerAudioVolumeSettings = CurrentAudioVolumeSettings;
 		CurrentAudioVolumeSettings = PlayerAudioVolumeSettings;
 		bReverbChanged = true;
 	}
+	else if (PlayerAudioVolumeSettings.AudioVolumeID == CurrentAudioVolumeSettings.AudioVolumeID && PlayerAudioVolumeSettings.bChanged)
+	{
+		CurrentAudioVolumeSettings = PlayerAudioVolumeSettings;
+		bReverbChanged = true;
+	}
 
 	if (Effects)
 	{
diff --git a/Engine/Source/Runtime/Engine/Private/AudioVolume.cpp b/Engine/Source/Runtime/Engine/Private/AudioVolume.cpp
--- a/Engine/Source/Runtime/Engine/Private/AudioVolume.cpp
+++ b/Engine/Source/Runtime/Engine/Private/AudioVolume.cpp
@@ -180,6 +180,9 @@ void FAudioDevice::UpdateAudioVolumeProxy(const FAudioVolumeProxy& NewProxy)
 
 		*CurrentProxy = NewProxy;
 
+		// Flag that the proxy changed so it can propagate any changes to runtime systems
+		CurrentProxy->bChanged = true;
+
 		if (CurrentPriority != NewProxy.Priority)
 		{
 			AudioVolumeProxies.ValueSort([](const FAudioVolumeProxy& A, const FAudioVolumeProxy& B) { return A.Priority > B.Priority; });
@@ -324,29 +327,46 @@ void AAudioVolume::PostEditChangeProperty(FPropertyChangedEvent& PropertyChanged
 	AmbientZoneSettings.ExteriorTime = FMath::Max<float>( 0.01f, AmbientZoneSettings.ExteriorTime );
 	AmbientZoneSettings.ExteriorLPFTime = FMath::Max<float>( 0.01f, AmbientZoneSettings.ExteriorLPFTime );
 
-	if (PropertyChangedEvent.Property && PropertyChangedEvent.Property->GetFName() == GET_MEMBER_NAME_CHECKED(AAudioVolume, Priority))
+	if (PropertyChangedEvent.Property)
 	{
-		if (UWorld* World = GetWorld())
+		static FName NAME_Priority = GET_MEMBER_NAME_CHECKED(AAudioVolume, Priority);
+		static FName NAME_Enabled = GET_MEMBER_NAME_CHECKED(AAudioVolume, bEnabled);
+		static FName NAME_ApplyReverb = GET_MEMBER_NAME_CHECKED(FReverbSettings, bApplyReverb);
+
+		FName PropertyName = PropertyChangedEvent.Property->GetFName();
+
+		if (PropertyName == NAME_Priority)
 		{
-			World->AudioVolumes.Sort([](const AAudioVolume& A, const AAudioVolume& B) { return (A.GetPriority() > B.GetPriority()); });
+			if (UWorld* World = GetWorld())
+			{
+				World->AudioVolumes.Sort([](const AAudioVolume& A, const AAudioVolume& B) { return (A.GetPriority() > B.GetPriority()); });
+			}
 		}
-	}
-
-	if (PropertyChangedEvent.Property && PropertyChangedEvent.Property->GetFName() == GET_MEMBER_NAME_CHECKED(AAudioVolume, bEnabled))
-	{
-		if (bEnabled)
+		else if (PropertyName == NAME_Enabled)
 		{
-			AddProxy();
+			if (bEnabled)
+			{
+				AddProxy();
+			}
+			else
+			{
+				RemoveProxy();
+			}
 		}
-		else
+		else if (PropertyName == NAME_ApplyReverb)
 		{
-			RemoveProxy();
+			if (Settings.ReverbEffect)
+			{
+				Settings.ReverbEffect->bChanged = true;
+			}
+		}
+
+		if (bEnabled)
+		{
+			UpdateProxy();
 		}
 	}
-	else if (bEnabled)
-	{
-		UpdateProxy();
-	}
+
 
 }
 #endif // WITH_EDITOR
diff --git a/Engine/Source/Runtime/Engine/Public/AudioDevice.h b/Engine/Source/Runtime/Engine/Public/AudioDevice.h
--- a/Engine/Source/Runtime/Engine/Public/AudioDevice.h
+++ b/Engine/Source/Runtime/Engine/Public/AudioDevice.h
@@ -796,9 +796,11 @@ public:
 		FInteriorSettings InteriorSettings;
 		TArray<FAudioVolumeSubmixSendSettings> SubmixSendSettings;
 		TArray<FAudioVolumeSubmixOverrideSettings> SubmixOverrideSettings;
+		bool bChanged = false;
 	};
 
 	void GetAudioVolumeSettings(const uint32 WorldID, const FVector& Location, FAudioVolumeSettings& OutSettings) const;
+	void ResetAudioVolumeProxyChangedState();
 
 public:
 
