Fix non-unity errors

diff --git a/Engine/Source/Runtime/Engine/Private/ViewModeNames.cpp b/Engine/Source/Runtime/Engine/Private/ViewModeNames.cpp
--- a/Engine/Source/Runtime/Engine/Private/ViewModeNames.cpp
+++ b/Engine/Source/Runtime/Engine/Private/ViewModeNames.cpp
@@ -1,6 +1,7 @@
 // Copyright Epic Games, Inc. All Rights Reserved.
 
 #include "CoreMinimal.h"
+#include "Styling/AppStyle.h"
 #include "Engine/EngineBaseTypes.h"
 
 #define LOCTEXT_NAMESPACE "UViewModeUtils"
