Fixed LandscapeActorDesc grid indices (they were for negative coordinates which resulted in crashes when deleting/adding landscape components).

diff --git a/Engine/Source/Runtime/Core/Public/UObject/UE5MainStreamObjectVersion.h b/Engine/Source/Runtime/Core/Public/UObject/UE5MainStreamObjectVersion.h
--- a/Engine/Source/Runtime/Core/Public/UObject/UE5MainStreamObjectVersion.h
+++ b/Engine/Source/Runtime/Core/Public/UObject/UE5MainStreamObjectVersion.h
@@ -67,6 +67,9 @@ struct CORE_API FUE5MainStreamObjectVersion
 		// Animation sync groups explicitly specify sync method
 		AnimSyncGroupsExplicitSyncMethod,
 
+		// Fixup FLandscapeActorDesc Grid indices
+		FLandscapeActorDescFixupGridIndices,
+
 		// -----<new versions can be added above this line>-------------------------------------------------
 		VersionPlusOne,
 		LatestVersion = VersionPlusOne - 1
diff --git a/Engine/Source/Runtime/Engine/Private/WorldPartition/Landscape/LandscapeActorDesc.cpp b/Engine/Source/Runtime/Engine/Private/WorldPartition/Landscape/LandscapeActorDesc.cpp
--- a/Engine/Source/Runtime/Engine/Private/WorldPartition/Landscape/LandscapeActorDesc.cpp
+++ b/Engine/Source/Runtime/Engine/Private/WorldPartition/Landscape/LandscapeActorDesc.cpp
@@ -8,6 +8,7 @@
 #include "LandscapeInfo.h"
 #include "WorldPartition/WorldPartition.h"
 #include "WorldPartition/WorldPartitionHandle.h"
+#include "UObject/UE5MainStreamObjectVersion.h"
 
 void FLandscapeActorDesc::Init(const AActor* InActor)
 {
@@ -15,11 +16,22 @@ void FLandscapeActorDesc::Init(const AActor* InActor)
 
 	const ALandscapeProxy* LandscapeProxy = CastChecked<ALandscapeProxy>(InActor);
 	check(LandscapeProxy);
-	GridIndexX = LandscapeProxy->LandscapeSectionOffset.X / LandscapeProxy->GridSize;
-	GridIndexY = LandscapeProxy->LandscapeSectionOffset.Y / LandscapeProxy->GridSize;
+	GridIndexX = LandscapeProxy->LandscapeSectionOffset.X / (int32)LandscapeProxy->GridSize;
+	GridIndexY = LandscapeProxy->LandscapeSectionOffset.Y / (int32)LandscapeProxy->GridSize;
 	GridIndexZ = 0;
 }
 
+void FLandscapeActorDesc::Serialize(FArchive& Ar)
+{
+	FPartitionActorDesc::Serialize(Ar);
+
+	if (Ar.IsLoading() && Ar.CustomVer(FUE5MainStreamObjectVersion::GUID) < FUE5MainStreamObjectVersion::FLandscapeActorDescFixupGridIndices)
+	{
+		GridIndexX = (int32)(GridIndexX * GridSize) / (int32)GridSize;
+		GridIndexY = (int32)(GridIndexY * GridSize) / (int32)GridSize;
+	}
+}
+
 void FLandscapeActorDesc::Unload()
 {
 	if (ALandscapeStreamingProxy* LandscapeStreamingProxy = Cast<ALandscapeStreamingProxy>(GetActor()))
diff --git a/Engine/Source/Runtime/Engine/Public/WorldPartition/Landscape/LandscapeActorDesc.h b/Engine/Source/Runtime/Engine/Public/WorldPartition/Landscape/LandscapeActorDesc.h
--- a/Engine/Source/Runtime/Engine/Public/WorldPartition/Landscape/LandscapeActorDesc.h
+++ b/Engine/Source/Runtime/Engine/Public/WorldPartition/Landscape/LandscapeActorDesc.h
@@ -19,5 +19,6 @@ public:
 protected:
 	virtual void OnRegister() override;
 	virtual void OnUnregister() override;
+	virtual void Serialize(FArchive& Ar) override;
 #endif
 };
