Fixing PVS 7.7 Issues:

warning V767: Suspicious access to element of 'Args' array by a constant index inside a loop.

diff --git a/Engine/Source/Runtime/Engine/Private/LocalPlayer.cpp b/Engine/Source/Runtime/Engine/Private/LocalPlayer.cpp
--- a/Engine/Source/Runtime/Engine/Private/LocalPlayer.cpp
+++ b/Engine/Source/Runtime/Engine/Private/LocalPlayer.cpp
@@ -517,7 +517,7 @@ public:
 			bool bAnyEmpty = false;
 			for (int32 i = 0; i < Args.Num(); ++i)
 			{
-				bAnyEmpty |= Args[0].Len() == 0;
+				bAnyEmpty |= Args[i].Len() == 0;
 			}
 			if (bAnyEmpty)
 			{
@@ -535,7 +535,7 @@ public:
 			bool bAnyEmpty = false;
 			for (int32 i = 0; i < Args.Num(); ++i)
 			{
-				bAnyEmpty |= Args[0].Len() == 0;
+				bAnyEmpty |= Args[i].Len() == 0;
 			}
 			if (bAnyEmpty)
 			{
