Actor Partition Subsystem now also sets actor's label

diff --git a/Engine/Source/Runtime/Engine/Private/ActorPartition/ActorPartitionSubsystem.cpp b/Engine/Source/Runtime/Engine/Private/ActorPartition/ActorPartitionSubsystem.cpp
--- a/Engine/Source/Runtime/Engine/Private/ActorPartition/ActorPartitionSubsystem.cpp
+++ b/Engine/Source/Runtime/Engine/Private/ActorPartition/ActorPartitionSubsystem.cpp
@@ -267,6 +267,9 @@ public:
 			FoundActor->bLockLocation = true;
 			
 			InActorCreated(FoundActor);
+
+			// Once actor is created, update its label
+			FoundActor->SetActorLabel(FString::Printf(TEXT("%s_%d_%d_%d"), *InActorClass->GetName(), InCellCoord.X, InCellCoord.Y, InCellCoord.Z));
 		}
 
 		check(FoundActor || !bInCreate);
