Actor: Moved FixupDataLayers from Serialize to PreSave to avoid a crash when cooking with -ConcurrentSave

#rnxUn

diff --git a/Engine/Source/Runtime/Engine/Classes/GameFramework/Actor.h b/Engine/Source/Runtime/Engine/Classes/GameFramework/Actor.h
--- a/Engine/Source/Runtime/Engine/Classes/GameFramework/Actor.h
+++ b/Engine/Source/Runtime/Engine/Classes/GameFramework/Actor.h
@@ -1784,6 +1784,7 @@ public:
 
 	virtual bool PreSaveRoot(const TCHAR* InFilename) override;
 	virtual void PostSaveRoot(bool bCleanupIsRequired) override;
+	virtual void PreSave(const ITargetPlatform* TargetPlatform) override;
 
 #if WITH_EDITOR
 	virtual bool Modify(bool bAlwaysMarkDirty = true) override;
diff --git a/Engine/Source/Runtime/Engine/Private/Actor.cpp b/Engine/Source/Runtime/Engine/Private/Actor.cpp
--- a/Engine/Source/Runtime/Engine/Private/Actor.cpp
+++ b/Engine/Source/Runtime/Engine/Private/Actor.cpp
@@ -411,6 +411,14 @@ void AActor::PostSaveRoot(bool bCleanupIsRequired)
 #endif
 }
 
+void AActor::PreSave(const ITargetPlatform* TargetPlatform)
+{
+	Super::PreSave(TargetPlatform);
+#if WITH_EDITOR
+	FixupDataLayers();
+#endif
+}
+
 #if WITH_EDITOR
 void AActor::GetAssetRegistryTags(TArray<FAssetRegistryTag>& OutTags) const
 {
@@ -732,12 +740,6 @@ void AActor::Serialize(FArchive& Ar)
 		}
 	}
 
-	// Fixup DataLayers
-	if (Ar.IsSaving() && !Ar.IsObjectReferenceCollector() && Ar.IsPersistent())
-	{
-		FixupDataLayers();
-	}
-
 	// When duplicating for PIE all components need to be gathered up and duplicated even if there are no other property references to them
 	// otherwise we can end up with Attach Parents that do not get redirected to the correct component. However, if there is a transient component
 	// we'll let that drop
