Make the skeletal mesh USkeletalMeshEditorData a child object of the skeletalmesh. This allow to remove PostDuplicate and PostRename override in the USkeletalmesh.

Edigrate from Dev-Tools-Staging branch cl: 14326873, because we need it to make the sk asynchronous

diff --git a/Engine/Source/Runtime/Engine/Classes/Engine/SkeletalMesh.h b/Engine/Source/Runtime/Engine/Classes/Engine/SkeletalMesh.h
--- a/Engine/Source/Runtime/Engine/Classes/Engine/SkeletalMesh.h
+++ b/Engine/Source/Runtime/Engine/Classes/Engine/SkeletalMesh.h
@@ -529,9 +529,7 @@ private:
 #if WITH_EDITORONLY_DATA
 public:
 	/*
-	 * This editor data asset is save in the same package has the skeletalmesh, the editor data asset is always loaded.
-	 * If the skeletal mesh is rename the editor data asset will also be rename: the name is SkeletalMeshName_USkeletalMeshEditorData
-	 * If the skeletal mesh is duplicate the editor data asset will also be duplicate
+	 * This editor data asset is save under the skeletalmesh(skel mesh is the owner), the editor data asset is always loaded.
 	 * There is only one editor data asset possible per skeletalmesh.
 	 * The reason we store the editor data in a separate asset is because the size of it can be very big and affect the editor performance. (undo/redo transactions)
 	 */
@@ -2104,8 +2102,6 @@ public:
 	virtual FString GetDetailedInfoInternal() const override;
 	virtual void GetResourceSizeEx(FResourceSizeEx& CumulativeResourceSize) override;
 	virtual void GetPreloadDependencies(TArray<UObject*>& OutDeps) override;
-	virtual void PostDuplicate(bool bDuplicateForPIE) override;
-	virtual void PostRename(UObject* OldOuter, const FName OldName) override;
 	//~ End UObject Interface.
 
 	//~ Begin UStreamableRenderAsset Interface.
diff --git a/Engine/Source/Runtime/Engine/Private/SkeletalMesh.cpp b/Engine/Source/Runtime/Engine/Private/SkeletalMesh.cpp
--- a/Engine/Source/Runtime/Engine/Private/SkeletalMesh.cpp
+++ b/Engine/Source/Runtime/Engine/Private/SkeletalMesh.cpp
@@ -1734,85 +1734,6 @@ void USkeletalMesh::GetPreloadDependencies(TArray<UObject*>& OutDeps)
 	OutDeps.Add(GetSkeleton());
 }
 
-void USkeletalMesh::PostDuplicate(bool bDuplicateForPIE)
-{
-	Super::PostDuplicate(bDuplicateForPIE);
-#if WITH_EDITORONLY_DATA
-	//We must not use the same asset for the imported data and duplicate it in the destination skeletalmesh package
-	if (IsMeshEditorDataValid())
-	{
-		PRAGMA_DISABLE_DEPRECATION_WARNINGS
-		USkeletalMeshEditorData* SourceMeshEditorData = MeshEditorDataObject;
-		MeshEditorDataObject = nullptr;
-		PRAGMA_ENABLE_DEPRECATION_WARNINGS
-		FSkeletalMeshModel& ImportedModels = *GetImportedModel();
-		//allocate the necessary data before going multi thread
-		ReserveLODImportData(ImportedModels.LODModels.Num() - 1);
-		USkeletalMeshEditorData& DestMeshEditorData = GetMeshEditorData();
-		//We should have a brand new created asset
-		check(IsMeshEditorDataValid());
-		//Lets duplicate the imported data
-		const int32 NumLODModels = ImportedModels.LODModels.Num();
-		for (int32 LODIndex = 0; LODIndex < NumLODModels; ++LODIndex)
-		{
-			if (!SourceMeshEditorData->IsLODImportDataValid(LODIndex) || SourceMeshEditorData->GetLODImportedData(LODIndex).IsEmpty())
-			{
-				return;
-			}
-
-			LLM_SCOPE(ELLMTag::SkeletalMesh);
-
-			FSkeletalMeshLODModel& ThisLODModel = ImportedModels.LODModels[LODIndex];
-			FRawSkeletalMeshBulkData& SourceRawSkeletalMeshBulkData = SourceMeshEditorData->GetLODImportedData(LODIndex);
-			FSkeletalMeshImportData RawMesh;
-			SourceRawSkeletalMeshBulkData.LoadRawMesh(RawMesh);
-			//This will create the asset into the correct package and assign the transient pointer
-			FRawSkeletalMeshBulkData& RawSkeletalMeshBulkData = DestMeshEditorData.GetLODImportedData(LODIndex);
-			RawSkeletalMeshBulkData.SaveRawMesh(RawMesh);
-			RawSkeletalMeshBulkData.GeoImportVersion = SourceRawSkeletalMeshBulkData.GeoImportVersion;
-			RawSkeletalMeshBulkData.SkinningImportVersion = SourceRawSkeletalMeshBulkData.SkinningImportVersion;
-			//Set all the cache data to avoid loading the asset if we do not edit the asset
-			ThisLODModel.bIsBuildDataAvailable = RawSkeletalMeshBulkData.IsBuildDataAvailable();
-			ThisLODModel.bIsRawSkeletalMeshBulkDataEmpty = RawSkeletalMeshBulkData.IsEmpty();
-			ThisLODModel.BuildStringID = RawSkeletalMeshBulkData.GetIdString();
-		}
-	}
-#endif //WITH_EDITORONLY_DATA
-}
-
-void USkeletalMesh::PostRename(UObject* OldOuter, const FName OldName)
-{
-	Super::PostRename(OldOuter, OldName);
-#if WITH_EDITORONLY_DATA
-	//We must not use the same asset for the imported data and duplicate it in the destination skeletalmesh package
-	if (IsMeshEditorDataValid())
-	{
-		PRAGMA_DISABLE_DEPRECATION_WARNINGS
-		USkeletalMeshEditorData* SourceMeshEditorData = MeshEditorDataObject;
-		PRAGMA_ENABLE_DEPRECATION_WARNINGS
-
-		FString MeshEditorDataString = GetName() + TEXT("_USkeletalMeshEditorData");
-		UPackage* SkeletalMeshPackage = GetOutermost();
-
-		//If the data is already existing, no need to rename it. Tentative to fix a live crash that we cannot reproduce, will log the data to help us
-		UObject* ExistingObject = StaticFindObject(/*Class=*/ NULL, SkeletalMeshPackage, *MeshEditorDataString, true);
-		if(!ExistingObject)
-		{
-			//Do a soft rename avoid: dirty, redirector, transaction and reset of the loaders
-			SourceMeshEditorData->Rename(*MeshEditorDataString, GetOutermost(), (REN_DoNotDirty | REN_DontCreateRedirectors | REN_NonTransactional | REN_ForceNoResetLoaders));
-		}
-		else
-		{
-			if (ExistingObject != SourceMeshEditorData)
-			{
-				//Log that the MeshEditorDataObject is not renamed properly.
-				UE_LOG(LogSkeletalMesh, Warning, TEXT("Renaming of skeletal mesh %s failed to rename the sub UObject MeshEditorDataObject (The imported data), because there is already an object with this name!/nMeshEditorDataObject name before renaming %s./nName of the existing object preventing the renaming %s"), *GetFullName(), *SourceMeshEditorData->GetFullName(), *ExistingObject->GetFullName());
-			}
-		}
-	}
-#endif //WITH_EDITORONLY_DATA
-}
-
 void USkeletalMesh::FlushRenderState()
 {
 	//TComponentReregisterContext<USkeletalMeshComponent> ReregisterContext;
@@ -2174,22 +2095,12 @@ void USkeletalMesh::RemoveLegacyClothingSections()
 USkeletalMeshEditorData& USkeletalMesh::GetMeshEditorData() const
 {
 	PRAGMA_DISABLE_DEPRECATION_WARNINGS
-	if (!IsMeshEditorDataValid())
+	if (MeshEditorDataObject == nullptr)
 	{
-		FString MeshEditorDataString = GetName() + TEXT("_USkeletalMeshEditorData");
-		FName MeshEditorDataName = *MeshEditorDataString;
-		//We can have only one USkeletalMeshEditorData asset per skeletalmesh. Find if there is already one existing 
-		UObject* ExistingObject = StaticFindObjectFast(USkeletalMeshEditorData::StaticClass(), GetOutermost(), MeshEditorDataName, true, false, RF_NoFlags, EInternalObjectFlags::PendingKill);
-		if (ExistingObject)
-		{
-			MeshEditorDataObject = Cast<USkeletalMeshEditorData>(ExistingObject);
-		}
-		if (!IsMeshEditorDataValid())
-		{
-			//The asset is created in the skeletalmesh package. We keep it private so the user cannot see it in the content browser
-			//StandAlone make sure the asset is save when we save the package(i.e. the skeletalmesh)
-			MeshEditorDataObject = NewObject<USkeletalMeshEditorData>(GetOutermost(), MeshEditorDataName);
-		}
+		//The asset is created in the skeletalmesh package. We keep it private so the user cannot see it in the content browser
+		//RF_Transactional make sure the asset can be transactional if we want to edit it
+		USkeletalMesh* NonConstSkeletalMesh = const_cast<USkeletalMesh*>(this);
+		MeshEditorDataObject = NewObject<USkeletalMeshEditorData>(NonConstSkeletalMesh, NAME_None, RF_Transactional);
 	}
 	//Make sure we have a valid pointer
 	check(MeshEditorDataObject != nullptr);
@@ -2498,6 +2409,16 @@ void USkeletalMesh::PostLoad()
 	}
 
 #if WITH_EDITOR
+
+	// Make sure the mesh editor data object is a sub object of the skeletalmesh, rename it to change the owner to be the skeletalmesh.
+	if (IsMeshEditorDataValid() && GetMeshEditorData().GetOuter() != this)
+	{
+		//Post load call so no need to: dirty, redirect, transact or reset the loader.
+		const TCHAR* NewName = nullptr;
+		GetMeshEditorData().Rename(NewName, this, REN_ForceNoResetLoaders | REN_DoNotDirty | REN_DontCreateRedirectors | REN_NonTransactional);
+		GetMeshEditorData().SetFlags(RF_Transactional);
+	}
+
 	if (!GetOutermost()->bIsCookedForEditor)
 	{
 		// If LODInfo is missing - create array of correct size.
