Added option to filter out curves from inertialization

diff --git a/Engine/Source/Runtime/Engine/Classes/Animation/AnimNode_Inertialization.h b/Engine/Source/Runtime/Engine/Classes/Animation/AnimNode_Inertialization.h
--- a/Engine/Source/Runtime/Engine/Classes/Animation/AnimNode_Inertialization.h
+++ b/Engine/Source/Runtime/Engine/Classes/Animation/AnimNode_Inertialization.h
@@ -191,8 +191,9 @@ struct FInertializationPoseDiff
 	// AttachParentName		the current frame's attach parent name (for checking if the attachment has changed)
 	// Prev1				the previous frame's pose
 	// Prev2				the pose from two frames before
+	// FilteredCurvesUIDs	list of curves we don't want to inertialize
 	//
-	void InitFrom(const FCompactPose& Pose, const FBlendedCurve& Curves, const FTransform& ComponentTransform, const FName& AttachParentName, const FInertializationPose& Prev1, const FInertializationPose& Prev2);
+	void InitFrom(const FCompactPose& Pose, const FBlendedCurve& Curves, const FTransform& ComponentTransform, const FName& AttachParentName, const FInertializationPose& Prev1, const FInertializationPose& Prev2, const TSet<SmartName::UID_Type>& FilteredCurvesUIDs);
 
 	// Apply this difference to a pose, decaying over time as InertializationElapsedTime approaches InertializationDuration
 	//
@@ -228,6 +229,10 @@ struct ENGINE_API FAnimNode_Inertialization : public FAnimNode_Base
 	UPROPERTY(EditAnywhere, BlueprintReadWrite, Category = Links)
 	FPoseLink Source;
 
+	// List of curves that should not use inertial blending. These curves will instantly change when inertialization begins.
+	UPROPERTY(EditAnywhere, BlueprintReadWrite, Category = Filter)
+	TArray<FName> FilteredCurves;
+
 public: // FAnimNode_Inertialization
 
 	FAnimNode_Inertialization();
@@ -282,6 +287,9 @@ protected:
 
 private:
 
+	// Set of UIDs for valid curves in FilteredCurves
+	TSet<SmartName::UID_Type> CachedFilteredCurvesUIDs;
+
 	// Snapshots of the actor pose from past frames
 	TArray<FInertializationPose> PoseSnapshots;
 
diff --git a/Engine/Source/Runtime/Engine/Private/Animation/AnimNode_Inertialization.cpp b/Engine/Source/Runtime/Engine/Private/Animation/AnimNode_Inertialization.cpp
--- a/Engine/Source/Runtime/Engine/Private/Animation/AnimNode_Inertialization.cpp
+++ b/Engine/Source/Runtime/Engine/Private/Animation/AnimNode_Inertialization.cpp
@@ -108,6 +108,19 @@ void FAnimNode_Inertialization::Initialize_AnyThread(const FAnimationInitializeC
 	Deactivate();
 
 	InertializationPoseDiff.Reset();
+	CachedFilteredCurvesUIDs.Reset();
+
+	const USkeleton* Skeleton = Context.AnimInstanceProxy->GetSkeleton();
+	check(Skeleton);
+	for (const FName& CurveName : FilteredCurves)
+	{
+		SmartName::UID_Type NameUID = Skeleton->GetUIDByName(USkeleton::AnimCurveMappingName, CurveName);
+		if (NameUID != SmartName::MaxUID)
+		{
+			// Grab UIDs of filtered curves to avoid lookup later
+			CachedFilteredCurvesUIDs.Add(NameUID);
+		}
+	}
 }
 
 
@@ -378,7 +391,7 @@ void FAnimNode_Inertialization::StartInertialization(FPoseContext& Context, FIne
 		}
 	}
 
-	OutPoseDiff.InitFrom(Context.Pose, Context.Curve, Context.AnimInstanceProxy->GetComponentTransform(), AttachParentName, PreviousPose1, PreviousPose2);
+	OutPoseDiff.InitFrom(Context.Pose, Context.Curve, Context.AnimInstanceProxy->GetComponentTransform(), AttachParentName, PreviousPose1, PreviousPose2, CachedFilteredCurvesUIDs);
 }
 
 
@@ -438,7 +451,7 @@ void FInertializationPose::InitFrom(const FCompactPose& Pose, const FBlendedCurv
 // Prev2		the pose and curves from two frames before
 // DeltaTime	the time elapsed between Prev1 and Pose
 //
-void FInertializationPoseDiff::InitFrom(const FCompactPose& Pose, const FBlendedCurve& Curves, const FTransform& ComponentTransform, const FName& AttachParentName, const FInertializationPose& Prev1, const FInertializationPose& Prev2)
+void FInertializationPoseDiff::InitFrom(const FCompactPose& Pose, const FBlendedCurve& Curves, const FTransform& ComponentTransform, const FName& AttachParentName, const FInertializationPose& Prev1, const FInertializationPose& Prev2, const TSet<SmartName::UID_Type>& FilteredCurvesUIDs)
 {
 	const FBoneContainer& BoneContainer = Pose.GetBoneContainer();
 
@@ -621,6 +634,12 @@ void FInertializationPoseDiff::InitFrom(const FCompactPose& Pose, const FBlended
 			continue;
 		}
 
+		// Check if the curve is in our filter set. Leave CurveDiff zeroed if it is.
+		if (FilteredCurvesUIDs.Contains((SmartName::UID_Type)CurveUID))
+		{
+			continue;
+		}
+
 		const float CurrWeight = Curves.CurveWeights[CurrIdx];
 		const float Prev1Weight = Prev1.Curves.CurveWeights[Prev1Idx];
 		FInertializationCurveDiff& CurveDiff = CurveDiffs[CurveUID];
