StaticMesh: Poll the async build instead of blocking on it when cooking

diff --git a/Engine/Source/Runtime/Engine/Classes/Components/StaticMeshComponent.h b/Engine/Source/Runtime/Engine/Classes/Components/StaticMeshComponent.h
--- a/Engine/Source/Runtime/Engine/Classes/Components/StaticMeshComponent.h
+++ b/Engine/Source/Runtime/Engine/Classes/Components/StaticMeshComponent.h
@@ -414,6 +414,8 @@ public:
 	virtual void PreEditUndo() override;
 	virtual void PostEditChangeProperty(FPropertyChangedEvent& PropertyChangedEvent) override;
 	virtual bool CanEditChange(const FProperty* InProperty) const override;
+	virtual void BeginCacheForCookedPlatformData(const ITargetPlatform* TargetPlatform) override;
+	virtual bool IsCachedCookedPlatformDataLoaded(const ITargetPlatform* TargetPlatform) override;
 #endif // WITH_EDITOR
 #if WITH_EDITORONLY_DATA
 	virtual void PreSave(const class ITargetPlatform* TargetPlatform) override;
diff --git a/Engine/Source/Runtime/Engine/Classes/Engine/StaticMesh.h b/Engine/Source/Runtime/Engine/Classes/Engine/StaticMesh.h
--- a/Engine/Source/Runtime/Engine/Classes/Engine/StaticMesh.h
+++ b/Engine/Source/Runtime/Engine/Classes/Engine/StaticMesh.h
@@ -1469,6 +1469,8 @@ public:
 	ENGINE_API virtual void PostEditChangeProperty(FPropertyChangedEvent& PropertyChangedEvent) override;
 	ENGINE_API virtual void PostEditUndo() override;
 	ENGINE_API virtual void GetAssetRegistryTagMetadata(TMap<FName, FAssetRegistryTagMetadata>& OutMetadata) const override;
+	ENGINE_API virtual void BeginCacheForCookedPlatformData(const ITargetPlatform* TargetPlatform) override;
+	ENGINE_API virtual bool IsCachedCookedPlatformDataLoaded(const ITargetPlatform* TargetPlatform) override;
 	ENGINE_API void SetLODGroup(FName NewGroup, bool bRebuildImmediately = true);
 	ENGINE_API void BroadcastNavCollisionChange();
 
@@ -1677,6 +1679,8 @@ public:
 	//~ Begin Interface_CollisionDataProvider Interface
 	ENGINE_API virtual bool GetPhysicsTriMeshData(struct FTriMeshCollisionData* CollisionData, bool InUseAllTriData) override;
 	ENGINE_API virtual bool ContainsPhysicsTriMeshData(bool InUseAllTriData) const override;
+	ENGINE_API virtual bool PollAsyncPhysicsTriMeshData(bool InUseAllTriData) const override;
+
 private:
 		bool GetPhysicsTriMeshDataCheckComplex(struct FTriMeshCollisionData* CollisionData, bool bInUseAllTriData, bool bInCheckComplexCollisionMesh);
 		bool ContainsPhysicsTriMeshDataCheckComplex(bool InUseAllTriData, bool bInCheckComplexCollisionMesh) const;
diff --git a/Engine/Source/Runtime/Engine/Classes/Interfaces/Interface_CollisionDataProvider.h b/Engine/Source/Runtime/Engine/Classes/Interfaces/Interface_CollisionDataProvider.h
--- a/Engine/Source/Runtime/Engine/Classes/Interfaces/Interface_CollisionDataProvider.h
+++ b/Engine/Source/Runtime/Engine/Classes/Interfaces/Interface_CollisionDataProvider.h
@@ -37,6 +37,13 @@ class IInterface_CollisionDataProvider
 	 */
 	virtual bool ContainsPhysicsTriMeshData(bool InUseAllTriData) const { return false; }
 
+	/**
+	 * Poll for availability of asynchronously-populated triangle mesh collision data
+	 *
+	 * @return true if the triangle mesh collision can be accessed without blocking, false otherwise
+	 */
+	virtual bool PollAsyncPhysicsTriMeshData(bool InUseAllTriData) const { return true; }
+
 	/** Do we want to create a negative version of this mesh */
 	virtual bool WantsNegXTriMesh() { return false; }
 
diff --git a/Engine/Source/Runtime/Engine/Classes/PhysicsEngine/BodySetup.h b/Engine/Source/Runtime/Engine/Classes/PhysicsEngine/BodySetup.h
--- a/Engine/Source/Runtime/Engine/Classes/PhysicsEngine/BodySetup.h
+++ b/Engine/Source/Runtime/Engine/Classes/PhysicsEngine/BodySetup.h
@@ -435,6 +435,7 @@ public:
 
 #if WITH_EDITOR
 	ENGINE_API virtual void BeginCacheForCookedPlatformData(  const ITargetPlatform* TargetPlatform ) override;
+	ENGINE_API virtual bool IsCachedCookedPlatformDataLoaded(  const ITargetPlatform* TargetPlatform ) override;
 	ENGINE_API virtual void ClearCachedCookedPlatformData(  const ITargetPlatform* TargetPlatform ) override;
 
 	/*
diff --git a/Engine/Source/Runtime/Engine/Private/Components/StaticMeshComponent.cpp b/Engine/Source/Runtime/Engine/Private/Components/StaticMeshComponent.cpp
--- a/Engine/Source/Runtime/Engine/Private/Components/StaticMeshComponent.cpp
+++ b/Engine/Source/Runtime/Engine/Private/Components/StaticMeshComponent.cpp
@@ -1643,6 +1643,15 @@ bool UStaticMeshComponent::CanEditChange(const FProperty* InProperty) const
 	return Super::CanEditChange(InProperty);
 }
 
+void UStaticMeshComponent::BeginCacheForCookedPlatformData(const ITargetPlatform* TargetPlatform)
+{
+}
+
+bool UStaticMeshComponent::IsCachedCookedPlatformDataLoaded(const ITargetPlatform* TargetPlatform)
+{
+	return !GetStaticMesh() || !GetStaticMesh()->IsCompiling();
+}
+
 #endif // WITH_EDITOR
 
 bool UStaticMeshComponent::SupportsDefaultCollision()
diff --git a/Engine/Source/Runtime/Engine/Private/PhysicsEngine/BodySetup.cpp b/Engine/Source/Runtime/Engine/Private/PhysicsEngine/BodySetup.cpp
--- a/Engine/Source/Runtime/Engine/Private/PhysicsEngine/BodySetup.cpp
+++ b/Engine/Source/Runtime/Engine/Private/PhysicsEngine/BodySetup.cpp
@@ -1328,7 +1328,20 @@ float UBodySetup::GetClosestPointAndNormal(const FVector& WorldPosition, const F
 #if WITH_EDITOR
 void UBodySetup::BeginCacheForCookedPlatformData(const ITargetPlatform* TargetPlatform)
 {
+}
+
+bool UBodySetup::IsCachedCookedPlatformDataLoaded(const ITargetPlatform* TargetPlatform)
+{
+	if (IInterface_CollisionDataProvider* CDP = Cast<IInterface_CollisionDataProvider>(GetOuter()))
+	{
+		bool bInUseAllTriData = true;
+		if (!CDP->PollAsyncPhysicsTriMeshData(bInUseAllTriData))
+		{
+			return false;
+		}
+	}
 	GetCookedData(TargetPlatform->GetPhysicsFormat(this), true);
+	return true;
 }
 
 void UBodySetup::ClearCachedCookedPlatformData( const ITargetPlatform* TargetPlatform )
diff --git a/Engine/Source/Runtime/Engine/Private/StaticMesh.cpp b/Engine/Source/Runtime/Engine/Private/StaticMesh.cpp
--- a/Engine/Source/Runtime/Engine/Private/StaticMesh.cpp
+++ b/Engine/Source/Runtime/Engine/Private/StaticMesh.cpp
@@ -3971,6 +3971,16 @@ void UStaticMesh::GetAssetRegistryTagMetadata(TMap<FName, FAssetRegistryTagMetad
 }
 #endif
 
+#if WITH_EDITOR
+void UStaticMesh::BeginCacheForCookedPlatformData(const ITargetPlatform* TargetPlatform)
+{
+}
+
+bool UStaticMesh::IsCachedCookedPlatformDataLoaded(const ITargetPlatform* TargetPlatform)
+{
+	return !IsCompiling();
+}
+#endif
 
 /*------------------------------------------------------------------------------
 	FStaticMeshSourceModel
@@ -6402,6 +6412,18 @@ bool UStaticMesh::ContainsPhysicsTriMeshDataCheckComplex(bool bInUseAllTriData,
 	return false; 
 }
 
+bool UStaticMesh::PollAsyncPhysicsTriMeshData(bool InUseAllTriData) const
+{
+#if WITH_EDITORONLY_DATA
+	bool bInCheckComplexCollisionMesh = true;
+	if (ComplexCollisionMesh && ComplexCollisionMesh != this && bInCheckComplexCollisionMesh)
+	{
+		return true;
+	}
+#endif
+	return !IsCompiling();
+}
+
 void UStaticMesh::GetMeshId(FString& OutMeshId)
 {
 #if WITH_EDITORONLY_DATA
diff --git a/Engine/Source/Runtime/Engine/Private/StaticMeshCompiler.cpp b/Engine/Source/Runtime/Engine/Private/StaticMeshCompiler.cpp
--- a/Engine/Source/Runtime/Engine/Private/StaticMeshCompiler.cpp
+++ b/Engine/Source/Runtime/Engine/Private/StaticMeshCompiler.cpp
@@ -4,6 +4,7 @@
 
 #if WITH_EDITOR
 
+#include "Algo/NoneOf.h"
 #include "AssetCompilingManager.h"
 #include "Engine/StaticMesh.h"
 #include "ObjectCacheContext.h"
@@ -273,6 +274,12 @@ void FStaticMeshCompilingManager::FinishCompilation(TArrayView<UStaticMesh* cons
 {
 	TRACE_CPUPROFILER_EVENT_SCOPE(FStaticMeshCompilingManager::FinishCompilation);
 
+	// Allow calls from any thread if the meshes are already finished compiling.
+	if (Algo::NoneOf(InStaticMeshes, &UStaticMesh::IsCompiling))
+	{
+		return;
+	}
+
 	check(IsInGameThread());
 
 	TArray<UStaticMesh*> PendingStaticMeshes;
