Added option to disable HLOD spawning on map load
* Used in the world partition conversion commandlet

diff --git a/Engine/Source/Editor/UnrealEd/Classes/Commandlets/WorldPartitionConvertCommandlet.h b/Engine/Source/Editor/UnrealEd/Classes/Commandlets/WorldPartitionConvertCommandlet.h
--- a/Engine/Source/Editor/UnrealEd/Classes/Commandlets/WorldPartitionConvertCommandlet.h
+++ b/Engine/Source/Editor/UnrealEd/Classes/Commandlets/WorldPartitionConvertCommandlet.h
@@ -62,8 +62,9 @@ protected:
 	bool DetachDependantLevelPackages(ULevel* Level);
 	bool RenameWorldPackageWithSuffix(UWorld* World);
 
-	UHLODLayer* CreateHLODLayerFromINI(const FString& InHLODLayerName);
+	void SetupHLOD();
 	void SetupHLODLayerAssets();
+	UHLODLayer* CreateHLODLayerFromINI(const FString& InHLODLayerName);
 
 	void SetActorGuid(AActor* Actor, const FGuid& NewGuid);
 	void CreateWorldMiniMapTexture(UWorld* World);
diff --git a/Engine/Source/Editor/UnrealEd/Private/Commandlets/WorldPartitionConvertCommandlet.cpp b/Engine/Source/Editor/UnrealEd/Private/Commandlets/WorldPartitionConvertCommandlet.cpp
--- a/Engine/Source/Editor/UnrealEd/Private/Commandlets/WorldPartitionConvertCommandlet.cpp
+++ b/Engine/Source/Editor/UnrealEd/Private/Commandlets/WorldPartitionConvertCommandlet.cpp
@@ -26,6 +26,7 @@
 #include "PackageHelperFunctions.h"
 #include "UObject/MetaData.h"
 #include "Editor.h"
+#include "HLOD/HLODEngineSubsystem.h"
 #include "HierarchicalLOD.h"
 #include "IHierarchicalLODUtilities.h"
 #include "HierarchicalLODUtilitiesModule.h"
@@ -596,6 +597,14 @@ UHLODLayer* UWorldPartitionConvertCommandlet::CreateHLODLayerFromINI(const FStri
 	return HLODLayer;
 }
 
+void UWorldPartitionConvertCommandlet::SetupHLOD()
+{
+	// No need to spawn HLOD actors during the conversion
+	GEngine->GetEngineSubsystem<UHLODEngineSubsystem>()->DisableHLODSpawningOnLoad(true);
+
+	SetupHLODLayerAssets();
+}
+
 void UWorldPartitionConvertCommandlet::SetupHLODLayerAssets()
 {
 	TArray<FString> HLODLayerSectionsNames;
@@ -743,7 +752,7 @@ int32 UWorldPartitionConvertCommandlet::Main(const FString& Params)
 		return 1;
 	}
 
-	SetupHLODLayerAssets();
+	SetupHLOD();
 
 	// Delete existing result from running the commandlet, even if not using the suffix mode to cleanup previous conversion
 	if (!bReportOnly)
diff --git a/Engine/Source/Runtime/Engine/Private/HLOD/HLODEngineSubsystem.cpp b/Engine/Source/Runtime/Engine/Private/HLOD/HLODEngineSubsystem.cpp
--- a/Engine/Source/Runtime/Engine/Private/HLOD/HLODEngineSubsystem.cpp
+++ b/Engine/Source/Runtime/Engine/Private/HLOD/HLODEngineSubsystem.cpp
@@ -19,6 +19,7 @@
 void UHLODEngineSubsystem::Initialize(FSubsystemCollectionBase& Collection)
 {
 	bDisableHLODCleanupOnLoad = false;
+	bDisableHLODSpawningOnLoad = false;
 
 	Super::Initialize(Collection);
 	RegisterRecreateLODActorsDelegates();
@@ -35,6 +36,14 @@ void UHLODEngineSubsystem::DisableHLODCleanupOnLoad(bool bInDisableHLODCleanup)
 	bDisableHLODCleanupOnLoad = bInDisableHLODCleanup;
 }
 
+void UHLODEngineSubsystem::DisableHLODSpawningOnLoad(bool bInDisableHLODSpawning)
+{
+	bDisableHLODSpawningOnLoad = bInDisableHLODSpawning;
+
+	UnregisterRecreateLODActorsDelegates();
+	RegisterRecreateLODActorsDelegates();
+}
+
 void UHLODEngineSubsystem::OnSaveLODActorsToHLODPackagesChanged()
 {
 	UnregisterRecreateLODActorsDelegates();
@@ -49,7 +58,7 @@ void UHLODEngineSubsystem::UnregisterRecreateLODActorsDelegates()
 
 void UHLODEngineSubsystem::RegisterRecreateLODActorsDelegates()
 {
-	if (GetDefault<UHierarchicalLODSettings>()->bSaveLODActorsToHLODPackages)
+	if (GetDefault<UHierarchicalLODSettings>()->bSaveLODActorsToHLODPackages && !bDisableHLODSpawningOnLoad)
 	{
 		OnPostWorldInitializationDelegateHandle = FWorldDelegates::OnPostWorldInitialization.AddUObject(this, &UHLODEngineSubsystem::RecreateLODActorsForWorld);
 		OnLevelAddedToWorldDelegateHandle = FWorldDelegates::LevelAddedToWorld.AddUObject(this, &UHLODEngineSubsystem::RecreateLODActorsForLevel);
diff --git a/Engine/Source/Runtime/Engine/Public/HLOD/HLODEngineSubsystem.h b/Engine/Source/Runtime/Engine/Public/HLOD/HLODEngineSubsystem.h
--- a/Engine/Source/Runtime/Engine/Public/HLOD/HLODEngineSubsystem.h
+++ b/Engine/Source/Runtime/Engine/Public/HLOD/HLODEngineSubsystem.h
@@ -37,6 +37,12 @@ public:
 	 */
 	void DisableHLODCleanupOnLoad(bool bInDisableHLODCleanup);
 
+	/**
+	 * By default, when HLODs are saved to HLOD packages, they are spawned as transient on load
+	 * Use this method to disable spawning.
+	 */
+	void DisableHLODSpawningOnLoad(bool bInDisableHLODSpawning);
+
 	// Should be called when the "Save LOD Actors to HLOD Packages" option is toggled.
 	void OnSaveLODActorsToHLODPackagesChanged();
 
@@ -60,6 +66,7 @@ private:
 	FDelegateHandle OnPreSaveWorlDelegateHandle;
 
 	bool bDisableHLODCleanupOnLoad;
+	bool bDisableHLODSpawningOnLoad;
 
 #endif // WITH_EDITOR
 };
