Downgrade the timing error to a warning to unblock builds (UE-106056).

(Edigrating 14950445)

diff --git a/Engine/Source/Runtime/Engine/Private/Materials/MaterialShared.cpp b/Engine/Source/Runtime/Engine/Private/Materials/MaterialShared.cpp
--- a/Engine/Source/Runtime/Engine/Private/Materials/MaterialShared.cpp
+++ b/Engine/Source/Runtime/Engine/Private/Materials/MaterialShared.cpp
@@ -1057,7 +1057,7 @@ void FMaterial::SerializeInlineShaderMap(FArchive& Ar)
 			}
 			else
 			{
-				UE_LOG(LogMaterial, Error, TEXT("Cooking a material resource (in %s hierarchy) that doesn't have a valid ShaderMap! %s"),
+				UE_LOG(LogMaterial, Warning, TEXT("Cooking a material resource (in %s hierarchy) that doesn't have a valid ShaderMap! %s"),
 					*GetFriendlyName(),
 					(GameThreadShaderMap != nullptr) ? TEXT("Shadermap exists but wasn't compiled successfully (yet?)") : TEXT("Shadermap pointer is null.")
 					);
