Fix a potential memory leak if async distance field calculations are cancelled early.

- Memory allocated for each async task was not being cleaned up if the task was cancelled before completion .

#rnx

diff --git a/Engine/Source/Runtime/Engine/Private/DistanceFieldAtlas.cpp b/Engine/Source/Runtime/Engine/Private/DistanceFieldAtlas.cpp
--- a/Engine/Source/Runtime/Engine/Private/DistanceFieldAtlas.cpp
+++ b/Engine/Source/Runtime/Engine/Private/DistanceFieldAtlas.cpp
@@ -1148,6 +1148,14 @@ void FDistanceFieldAsyncQueue::CancelBuild(UStaticMesh* StaticMesh)
 	CancelBackgroundTask(TasksToCancel);
 	for (FAsyncDistanceFieldTask* Task : TasksToCancel)
 	{
+		if (Task->GeneratedVolumeData != nullptr)
+		{
+			Task->GeneratedVolumeData->VolumeTexture.Release();
+
+			// Rendering thread may still be referencing the old one, use the deferred cleanup interface to delete it next frame when it is safe
+			BeginCleanup(Task->GeneratedVolumeData);
+		}
+
 		delete Task;
 	}
 }
diff --git a/Engine/Source/Runtime/Engine/Private/MeshCardRepresentation.cpp b/Engine/Source/Runtime/Engine/Private/MeshCardRepresentation.cpp
--- a/Engine/Source/Runtime/Engine/Private/MeshCardRepresentation.cpp
+++ b/Engine/Source/Runtime/Engine/Private/MeshCardRepresentation.cpp
@@ -300,6 +300,12 @@ void FCardRepresentationAsyncQueue::CancelBuild(UStaticMesh* StaticMesh)
 	CancelBackgroundTask(TasksToCancel);
 	for (FAsyncCardRepresentationTask* Task : TasksToCancel)
 	{
+		if (Task->GeneratedCardRepresentation != nullptr)
+		{
+			// Rendering thread may still be referencing the old one, use the deferred cleanup interface to delete it next frame when it is safe
+			BeginCleanup(Task->GeneratedCardRepresentation);
+		}
+
 		delete Task;
 	}
 }
