Fix an occasional crash when cooking meshes

- In some larger cooks a crash was being seen during the generation of distance field data as the async computation was still running for a mesh after it had already been garbage collected.
- Now we make sure to cancel any currently outstanding async computations that the mesh might havewhen ::IsReadyForFinishDestroy is called to prevent this.

#rnx

diff --git a/Engine/Source/Runtime/Engine/Private/StaticMesh.cpp b/Engine/Source/Runtime/Engine/Private/StaticMesh.cpp
--- a/Engine/Source/Runtime/Engine/Private/StaticMesh.cpp
+++ b/Engine/Source/Runtime/Engine/Private/StaticMesh.cpp
@@ -3845,6 +3845,19 @@ bool UStaticMesh::IsReadyForFinishDestroy()
 	{
 		return false;
 	}
+
+	if (GetRenderData())
+	{
+		if (GDistanceFieldAsyncQueue)
+		{
+			GDistanceFieldAsyncQueue->CancelBuild(this);
+		}
+
+		if (GCardRepresentationAsyncQueue)
+		{
+			GCardRepresentationAsyncQueue->CancelBuild(this);
+		}
+	}
 #endif
 
 	if (!Super::IsReadyForFinishDestroy())
