Disabling code fix added in CL 15049697 due to an editor shutdown crash that it seems to cause while the crash is being investigated.

#rnx

diff --git a/Engine/Source/Runtime/Engine/Private/StaticMesh.cpp b/Engine/Source/Runtime/Engine/Private/StaticMesh.cpp
--- a/Engine/Source/Runtime/Engine/Private/StaticMesh.cpp
+++ b/Engine/Source/Runtime/Engine/Private/StaticMesh.cpp
@@ -3846,6 +3846,9 @@ bool UStaticMesh::IsReadyForFinishDestroy()
 		return false;
 	}
 
+	// The following code block was causing a shutdown crash in the editor, so
+	// disabling it for now pending investigation.
+#if 0 
 	if (GetRenderData())
 	{
 		if (GDistanceFieldAsyncQueue)
@@ -3858,6 +3861,8 @@ bool UStaticMesh::IsReadyForFinishDestroy()
 			GCardRepresentationAsyncQueue->CancelBuild(this);
 		}
 	}
+#endif //0
+
 #endif
 
 	if (!Super::IsReadyForFinishDestroy())
