Fixing PVS 7.7 Issues:

warning V1051: Consider checking for misprints. It's possible that the 'XYZ' should be checked here.

diff --git a/Engine/Plugins/Experimental/ControlRig/Source/ControlRig/Private/ControlRig.cpp b/Engine/Plugins/Experimental/ControlRig/Source/ControlRig/Private/ControlRig.cpp
--- a/Engine/Plugins/Experimental/ControlRig/Source/ControlRig/Private/ControlRig.cpp
+++ b/Engine/Plugins/Experimental/ControlRig/Source/ControlRig/Private/ControlRig.cpp
@@ -1138,7 +1138,7 @@ void UControlRig::CreateRigControlsForCurveContainer()
 		for (FRigCurve& Curve : CurveContainer)
 		{
 			FString Name = Curve.Name.ToString();
-			if (Name.Contains(CTRLName) && Hierarchy.ControlHierarchy.GetIndex(Curve.Name) == INDEX_NONE)
+			if (Name.Contains(CTRLName) && Hierarchy.ControlHierarchy.GetIndex(Curve.Name) == INDEX_NONE) //-V1051
 			{
 				FRigControlValue Value;
 				Value.Set<float>(Curve.Value);
diff --git a/Engine/Plugins/Experimental/GeometryCache/Source/GeometryCache/Private/GeometryCacheSceneProxy.cpp b/Engine/Plugins/Experimental/GeometryCache/Source/GeometryCache/Private/GeometryCacheSceneProxy.cpp
--- a/Engine/Plugins/Experimental/GeometryCache/Source/GeometryCache/Private/GeometryCacheSceneProxy.cpp
+++ b/Engine/Plugins/Experimental/GeometryCache/Source/GeometryCache/Private/GeometryCacheSceneProxy.cpp
@@ -677,7 +677,7 @@ void FGeometryCacheSceneProxy::FrameUpdate() const
 					TrackProxy->NextFrameIndex = OldFrameIndex;
 
 					// Decode the new next frame
-					if (TrackProxy->GetMeshData(NextFrameIndex, *TrackProxy->NextFrameMeshData))
+					if (TrackProxy->GetMeshData(NextFrameIndex, *TrackProxy->NextFrameMeshData)) //-V1051
 					{
 						bDecodedAnything = true;
 						// Only register this if we actually successfully decoded
diff --git a/Engine/Plugins/FX/Niagara/Source/NiagaraEditor/Private/NiagaraNode.cpp b/Engine/Plugins/FX/Niagara/Source/NiagaraEditor/Private/NiagaraNode.cpp
--- a/Engine/Plugins/FX/Niagara/Source/NiagaraEditor/Private/NiagaraNode.cpp
+++ b/Engine/Plugins/FX/Niagara/Source/NiagaraEditor/Private/NiagaraNode.cpp
@@ -344,7 +344,7 @@ void UNiagaraNode::AutowireNewNode(UEdGraphPin* FromPin)
 			if (Pin->Direction == DesiredDirection)
 			{
 				const FPinConnectionResponse Response = Schema->CanCreateConnection(FromPin, Pin);
-				if (Response.Response != ECanCreateConnectionResponse::CONNECT_RESPONSE_DISALLOW)
+				if (Response.Response != ECanCreateConnectionResponse::CONNECT_RESPONSE_DISALLOW) //-V1051
 				{
 					FirstPinOfSameType = Pin;
 					break;
diff --git a/Engine/Plugins/Media/ElectraPlayer/Source/ElectraPlayerRuntime/Private/Runtime/Player/HLS/StreamReaderHLSfmp4.cpp b/Engine/Plugins/Media/ElectraPlayer/Source/ElectraPlayerRuntime/Private/Runtime/Player/HLS/StreamReaderHLSfmp4.cpp
--- a/Engine/Plugins/Media/ElectraPlayer/Source/ElectraPlayerRuntime/Private/Runtime/Player/HLS/StreamReaderHLSfmp4.cpp
+++ b/Engine/Plugins/Media/ElectraPlayer/Source/ElectraPlayerRuntime/Private/Runtime/Player/HLS/StreamReaderHLSfmp4.cpp
@@ -169,7 +169,7 @@ UEMediaError FStreamReaderHLSfmp4::Create(IPlayerSessionServices* InPlayerSessio
 
 	if (!InCreateParam.MemoryProvider ||
 		!InCreateParam.EventListener ||
-		!InCreateParam.PlayerSessionService)
+		!InCreateParam.PlayerSessionService) //-V1051
 	{
 		return UEMEDIA_ERROR_BAD_ARGUMENTS;
 	}
diff --git a/Engine/Plugins/Runtime/Metasound/Source/MetasoundEditor/Private/MetasoundEditorGraphNode.cpp b/Engine/Plugins/Runtime/Metasound/Source/MetasoundEditor/Private/MetasoundEditorGraphNode.cpp
--- a/Engine/Plugins/Runtime/Metasound/Source/MetasoundEditor/Private/MetasoundEditorGraphNode.cpp
+++ b/Engine/Plugins/Runtime/Metasound/Source/MetasoundEditor/Private/MetasoundEditorGraphNode.cpp
@@ -150,7 +150,7 @@ void UMetasoundEditorGraphNode::AutowireNewNode(UEdGraphPin* FromPin)
 			UEdGraphPin* Pin = Pins[i];
 			check(Pin);
 			FPinConnectionResponse Response = Schema->CanCreateConnection(FromPin, Pin);
-			if (ECanCreateConnectionResponse::CONNECT_RESPONSE_MAKE == Response.Response)
+			if (ECanCreateConnectionResponse::CONNECT_RESPONSE_MAKE == Response.Response) //-V1051
 			{
 				if (Schema->TryCreateConnection(FromPin, Pin))
 				{
diff --git a/Engine/Source/Developer/TraceServices/Private/Model/Threads.cpp b/Engine/Source/Developer/TraceServices/Private/Model/Threads.cpp
--- a/Engine/Source/Developer/TraceServices/Private/Model/Threads.cpp
+++ b/Engine/Source/Developer/TraceServices/Private/Model/Threads.cpp
@@ -60,7 +60,7 @@ void FThreadProvider::AddThread(uint32 Id, const TCHAR* Name, EThreadPriority Pr
 	if (Name != nullptr)
 	{
 		ThreadInfo->Name = Session.StoreString(Name);
-		if (!FCString::Strcmp(Name, TEXT("RHIThread")))
+		if (!FCString::Strcmp(Name, TEXT("RHIThread"))) //-V1051
 		{
 			const TCHAR* GroupName = Session.StoreString(TEXT("Render"));
 			SetThreadGroup(Id, GroupName);
diff --git a/Engine/Source/Editor/AudioEditor/Private/SoundCueGraphNode_Base.cpp b/Engine/Source/Editor/AudioEditor/Private/SoundCueGraphNode_Base.cpp
--- a/Engine/Source/Editor/AudioEditor/Private/SoundCueGraphNode_Base.cpp
+++ b/Engine/Source/Editor/AudioEditor/Private/SoundCueGraphNode_Base.cpp
@@ -198,7 +198,7 @@ void USoundCueGraphNode_Base::AutowireNewNode(UEdGraphPin* FromPin)
 			UEdGraphPin* Pin = Pins[i];
 			check(Pin);
 			FPinConnectionResponse Response = Schema->CanCreateConnection(FromPin, Pin);
-			if (ECanCreateConnectionResponse::CONNECT_RESPONSE_MAKE == Response.Response)
+			if (ECanCreateConnectionResponse::CONNECT_RESPONSE_MAKE == Response.Response) //-V1051
 			{
 				if (Schema->TryCreateConnection(FromPin, Pin))
 				{
diff --git a/Engine/Source/Editor/GraphEditor/Private/BlueprintConnectionDrawingPolicy.cpp b/Engine/Source/Editor/GraphEditor/Private/BlueprintConnectionDrawingPolicy.cpp
--- a/Engine/Source/Editor/GraphEditor/Private/BlueprintConnectionDrawingPolicy.cpp
+++ b/Engine/Source/Editor/GraphEditor/Private/BlueprintConnectionDrawingPolicy.cpp
@@ -592,7 +592,7 @@ void FKismetConnectionDrawingPolicy::SetIncompatiblePinDrawState(const TSharedPt
 		if (CheckPin != StartPin)
 		{
 			const FPinConnectionResponse Response = StartPin->GetPinObj()->GetSchema()->CanCreateConnection(StartPin->GetPinObj(), CheckPin->GetPinObj());
-			if (Response.Response == CONNECT_RESPONSE_DISALLOW)
+			if (Response.Response == CONNECT_RESPONSE_DISALLOW) //-V1051
 			{
 				CheckPin->SetPinColorModifier(FLinearColor(0.25f, 0.25f, 0.25f, 0.5f));
 			}
diff --git a/Engine/Source/Editor/KismetCompiler/Private/KismetCompiler.cpp b/Engine/Source/Editor/KismetCompiler/Private/KismetCompiler.cpp
--- a/Engine/Source/Editor/KismetCompiler/Private/KismetCompiler.cpp
+++ b/Engine/Source/Editor/KismetCompiler/Private/KismetCompiler.cpp
@@ -3829,7 +3829,7 @@ void FKismetCompilerContext::ProcessOneFunctionGraph(UEdGraph* SourceGraph, bool
 	FunctionList.Add(&Context);
 	Context.SourceGraph = FunctionGraph;
 
-	if (FBlueprintEditorUtils::IsDelegateSignatureGraph(SourceGraph))
+	if (FBlueprintEditorUtils::IsDelegateSignatureGraph(SourceGraph)) //-V1051
 	{
 		Context.SetDelegateSignatureName(SourceGraph->GetFName());
 	}
diff --git a/Engine/Source/Editor/UnrealEd/Private/MaterialGraphNode_Base.cpp b/Engine/Source/Editor/UnrealEd/Private/MaterialGraphNode_Base.cpp
--- a/Engine/Source/Editor/UnrealEd/Private/MaterialGraphNode_Base.cpp
+++ b/Engine/Source/Editor/UnrealEd/Private/MaterialGraphNode_Base.cpp
@@ -284,7 +284,7 @@ void UMaterialGraphNode_Base::AutowireNewNode(UEdGraphPin* FromPin)
 			UEdGraphPin* Pin = Pins[i];
 			check(Pin);
 			FPinConnectionResponse Response = Schema->CanCreateConnection(FromPin, Pin);
-			if (ECanCreateConnectionResponse::CONNECT_RESPONSE_MAKE == Response.Response)
+			if (ECanCreateConnectionResponse::CONNECT_RESPONSE_MAKE == Response.Response) //-V1051
 			{
 				if (Schema->TryCreateConnection(FromPin, Pin))
 				{
diff --git a/Engine/Source/Editor/ViewportInteraction/ViewportWorldInteraction.cpp b/Engine/Source/Editor/ViewportInteraction/ViewportWorldInteraction.cpp
--- a/Engine/Source/Editor/ViewportInteraction/ViewportWorldInteraction.cpp
+++ b/Engine/Source/Editor/ViewportInteraction/ViewportWorldInteraction.cpp
@@ -525,7 +525,7 @@ void UViewportWorldInteraction::SetTransformer( UViewportTransformer* NewTransfo
 		ViewportTransformer = nullptr;
 	}
 
-	this->ViewportTransformer = NewTransformer;
+	ViewportTransformer = NewTransformer;
 
 	if( ViewportTransformer == nullptr )
 	{
diff --git a/Engine/Source/Runtime/Core/Public/Math/Float16.h b/Engine/Source/Runtime/Core/Public/Math/Float16.h
--- a/Engine/Source/Runtime/Core/Public/Math/Float16.h
+++ b/Engine/Source/Runtime/Core/Public/Math/Float16.h
@@ -155,7 +155,7 @@ FORCEINLINE void FFloat16::Set(float FP32Value)
              uint32 Mantissa = FP32.Components.Mantissa | 0x800000; // Hidden 1 bit
              Components.Mantissa = (uint16)(Mantissa >> (14 - NewExp));
 			 // Check for rounding
-             if ( (Mantissa >> (13 - NewExp)) & 1 )
+             if ( (Mantissa >> (13 - NewExp)) & 1 ) //-V1051
 			 {
                  Encoded++; // Round, might overflow into exp bit, but this is OK
 			 }
diff --git a/Engine/Source/Runtime/Engine/Private/AudioDebug.cpp b/Engine/Source/Runtime/Engine/Private/AudioDebug.cpp
--- a/Engine/Source/Runtime/Engine/Private/AudioDebug.cpp
+++ b/Engine/Source/Runtime/Engine/Private/AudioDebug.cpp
@@ -2468,7 +2468,7 @@ namespace Audio
 				if (const int32* SoundInfoIndex = ActiveSoundToInfoIndex.Find(ActiveSound))
 				{
 					const USoundClass* SoundClass = ActiveSound->GetSoundClass();
-					if (const USoundClass* WaveSoundClass = WaveInstance->SoundClass)
+					if (const USoundClass* WaveSoundClass = WaveInstance->SoundClass) //-V1051
 					{
 						SoundClass = WaveSoundClass;
 					}
diff --git a/Engine/Source/Runtime/Engine/Private/Materials/MaterialUniformExpressions.cpp b/Engine/Source/Runtime/Engine/Private/Materials/MaterialUniformExpressions.cpp
--- a/Engine/Source/Runtime/Engine/Private/Materials/MaterialUniformExpressions.cpp
+++ b/Engine/Source/Runtime/Engine/Private/Materials/MaterialUniformExpressions.cpp
@@ -1485,7 +1485,7 @@ void FUniformExpressionSet::FillUniformBuffer(const FMaterialRenderContext& Mate
 				ensure(PageTable0RHI);
 				*ResourceTablePageTexture0Ptr = PageTable0RHI;
 
-				if (ResourceTablePageTexture1Ptr != nullptr)
+				if (ResourceTablePageTexture1Ptr != nullptr) //-V1051
 				{
 					FRHITexture* PageTable1RHI = AllocatedVT->GetPageTableTexture(1u);
 					ensure(PageTable1RHI);
@@ -1500,7 +1500,7 @@ void FUniformExpressionSet::FillUniformBuffer(const FMaterialRenderContext& Mate
 			{
 				// Don't have valid resources to bind for this VT, so make sure something is bound
 				*ResourceTablePageTexture0Ptr = GBlackUintTexture->TextureRHI;
-				if (ResourceTablePageTexture1Ptr != nullptr)
+				if (ResourceTablePageTexture1Ptr != nullptr) //-V1051
 				{
 					*ResourceTablePageTexture1Ptr = GBlackUintTexture->TextureRHI;
 				}
diff --git a/Engine/Source/Runtime/VulkanRHI/Private/VulkanCommands.cpp b/Engine/Source/Runtime/VulkanRHI/Private/VulkanCommands.cpp
--- a/Engine/Source/Runtime/VulkanRHI/Private/VulkanCommands.cpp
+++ b/Engine/Source/Runtime/VulkanRHI/Private/VulkanCommands.cpp
@@ -918,7 +918,7 @@ void FVulkanCommandListContext::RHICopyToStagingBuffer(FRHIVertexBuffer* SourceB
 	ensureMsgf((SourceBufferRHI->GetUsage() & BUF_SourceCopy) != 0, TEXT("Buffers used as copy source need to be created with BUF_SourceCopy"));
 
 	FVulkanStagingBuffer* StagingBuffer = ResourceCast(StagingBufferRHI);
-	if (!StagingBuffer->StagingBuffer || StagingBuffer->StagingBuffer->GetSize() < NumBytes)
+	if (!StagingBuffer->StagingBuffer || StagingBuffer->StagingBuffer->GetSize() < NumBytes) //-V1051
 	{
 		if (StagingBuffer->StagingBuffer)
 		{
