Chaos : Constraint : Angular Plasticity

- Allow the joints to reset their target angle to the current relative angle when the targeted rotations have deviated from a user defined delta.

diff --git a/Engine/Source/Editor/DetailCustomizations/Private/PhysicsConstraintComponentDetails.cpp b/Engine/Source/Editor/DetailCustomizations/Private/PhysicsConstraintComponentDetails.cpp
--- a/Engine/Source/Editor/DetailCustomizations/Private/PhysicsConstraintComponentDetails.cpp
+++ b/Engine/Source/Editor/DetailCustomizations/Private/PhysicsConstraintComponentDetails.cpp
@@ -485,6 +485,9 @@ void FPhysicsConstraintComponentDetails::AddAngularLimits(IDetailLayoutBuilder&
 
 	AngularLimitCat.AddProperty(ProfilePropertiesProperty->GetChildHandle(GET_MEMBER_NAME_CHECKED(FConstraintProfileProperties, bAngularBreakable)).ToSharedRef());
 	AngularLimitCat.AddProperty(ProfilePropertiesProperty->GetChildHandle(GET_MEMBER_NAME_CHECKED(FConstraintProfileProperties, AngularBreakThreshold)).ToSharedRef());
+	AngularLimitCat.AddProperty(ProfilePropertiesProperty->GetChildHandle(GET_MEMBER_NAME_CHECKED(FConstraintProfileProperties, bAngularPlasticity)).ToSharedRef());
+	AngularLimitCat.AddProperty(ProfilePropertiesProperty->GetChildHandle(GET_MEMBER_NAME_CHECKED(FConstraintProfileProperties, AngularPlasticityThreshold)).ToSharedRef());
+
 }
 
 void FPhysicsConstraintComponentDetails::AddLinearDrive(IDetailLayoutBuilder& DetailBuilder, TSharedPtr<IPropertyHandle> ConstraintInstance, TSharedPtr<IPropertyHandle> ProfilePropertiesProperty)
diff --git a/Engine/Source/Runtime/Engine/Classes/PhysicsEngine/ConstraintInstance.h b/Engine/Source/Runtime/Engine/Classes/PhysicsEngine/ConstraintInstance.h
--- a/Engine/Source/Runtime/Engine/Classes/PhysicsEngine/ConstraintInstance.h
+++ b/Engine/Source/Runtime/Engine/Classes/PhysicsEngine/ConstraintInstance.h
@@ -50,6 +50,10 @@ struct ENGINE_API FConstraintProfileProperties
 	UPROPERTY(EditAnywhere, AdvancedDisplay, Category = Angular, meta = (editcondition = "bAngularBreakable", ClampMin = "0.0"))
 	float AngularBreakThreshold;
 
+	/** Degree threshold from target angle needed to reset the target angle.*/
+	UPROPERTY(EditAnywhere, AdvancedDisplay, Category = Angular, meta = (editcondition = "bAngularPlasticity", ClampMin = "0.0"))
+	float AngularPlasticityThreshold;
+
 	UPROPERTY(EditAnywhere, Category = Linear)
 	FLinearConstraint LinearLimit;
 
@@ -101,6 +105,10 @@ struct ENGINE_API FConstraintProfileProperties
 	UPROPERTY(EditAnywhere, AdvancedDisplay, Category = Angular)
 	uint8 bAngularBreakable : 1;
 
+	/** Whether it is possible to reset target rotations from the angular displacement. */
+	UPROPERTY(EditAnywhere, AdvancedDisplay, Category = Angular)
+	uint8 bAngularPlasticity : 1;
+
 	/** Whether it is possible to break the joint with linear force. */
 	UPROPERTY(EditAnywhere, AdvancedDisplay, Category = Linear)
 	uint8 bLinearBreakable : 1;
@@ -110,10 +118,13 @@ struct ENGINE_API FConstraintProfileProperties
 	/** Updates physx joint properties from unreal properties (limits, drives, flags, etc...) */
 	void Update_AssumesLocked(const FPhysicsConstraintHandle& InConstraintRef, float AverageMass, float UseScale) const;
 
-	/** Updates physx joint breakable properties (threshold, etc...)*/
+	/** Updates joint breakable properties (threshold, etc...)*/
 	void UpdateBreakable_AssumesLocked(const FPhysicsConstraintHandle& InConstraintRef) const;
 
-	/** Updates physx joint flag based on profile properties */
+	/** Updates joint breakable properties (threshold, etc...)*/
+	void UpdatePlasticity_AssumesLocked(const FPhysicsConstraintHandle& InConstraintRef) const;
+
+	/** Updates joint flag based on profile properties */
 	void UpdateConstraintFlags_AssumesLocked(const FPhysicsConstraintHandle& InConstraintRef) const;
 
 #if WITH_EDITOR
@@ -472,6 +483,17 @@ public:
 		UpdateBreakable();
 	}
 
+	/** Sets the Angular Plasticity properties
+	*	@param bInAngularPlasticity 	Whether it is possible to reset the target angles
+	*	@param InAngularPlasticityThreshold	Delta from target needed to reset the target joint
+	*/
+	void SetAngularPlasticity(bool bInAngularPlasticity, float InAngularPlasticityThreshold)
+	{
+		ProfileInstance.bAngularPlasticity = bInAngularPlasticity;
+		ProfileInstance.AngularPlasticityThreshold = InAngularPlasticityThreshold;
+		UpdatePlasticity();
+	}
+
 	// @todo document
 	void CopyConstraintGeometryFrom(const FConstraintInstance* FromInstance);
 
@@ -704,6 +726,7 @@ private:
 		const FTransform& Con1Frame, const FTransform& Con2Frame, bool bDrawAsPoint) const;
 
 	void UpdateBreakable();
+	void UpdatePlasticity();
 	void UpdateDriveTarget();
 
 	FOnConstraintBroken OnConstraintBrokenDelegate;
diff --git a/Engine/Source/Runtime/Engine/Classes/PhysicsEngine/PhysicsConstraintComponent.h b/Engine/Source/Runtime/Engine/Classes/PhysicsEngine/PhysicsConstraintComponent.h
--- a/Engine/Source/Runtime/Engine/Classes/PhysicsEngine/PhysicsConstraintComponent.h
+++ b/Engine/Source/Runtime/Engine/Classes/PhysicsEngine/PhysicsConstraintComponent.h
@@ -276,6 +276,13 @@ public:
 	UFUNCTION(BlueprintCallable, Category = "Physics|Components|PhysicsConstraint")
 	void SetAngularBreakable(bool bAngularBreakable, float AngularBreakThreshold);
 
+	/** Sets the Angular Plasticity properties
+	*	@param bAngularPlasticity		Whether it is possible to reset the target angle from the angular displacement
+	*	@param AngularPlasticityThreshold	Degrees needed to reset the rest state of the joint
+	*/
+	UFUNCTION(BlueprintCallable, Category = "Physics|Components|PhysicsConstraint")
+	void SetAngularPlasticity(bool bAngularPlasticity, float AngularPlasticityThreshold);
+
 	/** Gets the current Angular Twist of the constraint */
 	UFUNCTION(BlueprintCallable, Category = "Physics|Components|PhysicsConstraint")
 	float GetCurrentTwist() const;
diff --git a/Engine/Source/Runtime/Engine/Private/PhysicsEngine/ConstraintInstance.cpp b/Engine/Source/Runtime/Engine/Private/PhysicsEngine/ConstraintInstance.cpp
--- a/Engine/Source/Runtime/Engine/Private/PhysicsEngine/ConstraintInstance.cpp
+++ b/Engine/Source/Runtime/Engine/Private/PhysicsEngine/ConstraintInstance.cpp
@@ -130,11 +130,13 @@ FConstraintProfileProperties::FConstraintProfileProperties()
 	, ProjectionAngularAlpha(0.0f)
 	, LinearBreakThreshold(300.f)
 	, AngularBreakThreshold(500.f)
+	, AngularPlasticityThreshold(10.f)
 	, bDisableCollision(false)
 	, bParentDominates(false)
 	, bEnableProjection(true)
 	, bEnableSoftProjection(false)
 	, bAngularBreakable(false)
+	, bAngularPlasticity(false)
 	, bLinearBreakable(false)
 {
 }
@@ -172,6 +174,22 @@ void FConstraintProfileProperties::UpdateBreakable_AssumesLocked(const FPhysicsC
 	FPhysicsInterface::SetBreakForces_AssumesLocked(InConstraintRef, LinearBreakForce, AngularBreakForce);
 }
 
+void FConstraintInstance::UpdatePlasticity()
+{
+	FPhysicsInterface::ExecuteOnUnbrokenConstraintReadWrite(ConstraintHandle, [&](const FPhysicsConstraintHandle& InUnbrokenConstraint)
+		{
+			ProfileInstance.UpdatePlasticity_AssumesLocked(InUnbrokenConstraint);
+		});
+}
+
+void FConstraintProfileProperties::UpdatePlasticity_AssumesLocked(const FPhysicsConstraintHandle& InConstraintRef) const
+{
+	const float LinearPlasticityLimit = MAX_FLT; // todo(Chaos) : implement
+	const float AngularPlasticityLimit = bAngularPlasticity ? FMath::DegreesToRadians(AngularPlasticityThreshold) : MAX_FLT;
+
+	FPhysicsInterface::SetPlasticityLimits_AssumesLocked(InConstraintRef, LinearPlasticityLimit, AngularPlasticityLimit);
+}
+
 void FConstraintInstance::UpdateDriveTarget()
 {
 	FPhysicsInterface::ExecuteOnUnbrokenConstraintReadWrite(ConstraintHandle, [&](const FPhysicsConstraintHandle& InUnbrokenConstraint)
@@ -499,6 +517,7 @@ void FConstraintProfileProperties::Update_AssumesLocked(const FPhysicsConstraint
 	TwistLimit.UpdateTwistLimit_AssumesLocked(InConstraintRef, AverageMass);
 
 	UpdateBreakable_AssumesLocked(InConstraintRef);
+	UpdatePlasticity_AssumesLocked(InConstraintRef);
 
 	// Motors
 	FPhysicsInterface::UpdateLinearDrive_AssumesLocked(InConstraintRef, LinearDrive);
diff --git a/Engine/Source/Runtime/Engine/Private/PhysicsEngine/Experimental/PhysInterface_Chaos.cpp b/Engine/Source/Runtime/Engine/Private/PhysicsEngine/Experimental/PhysInterface_Chaos.cpp
--- a/Engine/Source/Runtime/Engine/Private/PhysicsEngine/Experimental/PhysInterface_Chaos.cpp
+++ b/Engine/Source/Runtime/Engine/Private/PhysicsEngine/Experimental/PhysInterface_Chaos.cpp
@@ -436,7 +436,11 @@ void FPhysInterface_Chaos::UpdateAngularDrive_AssumesLocked(const FPhysicsConstr
 					Constraint->SetAngularSLerpPositionDriveEnabled(InDriveParams.SlerpDrive.bEnablePositionDrive);
 				}
 
-				Constraint->SetAngularDrivePositionTarget(Chaos::FRotation3(InDriveParams.OrientationTarget.Quaternion()));
+				if (FMath::IsNearlyEqual(Constraint->GetAngularPlasticityLimit(), FLT_MAX))
+				{
+					// Plastic joints should not be re-targeted after initialization. 
+					Constraint->SetAngularDrivePositionTarget(Chaos::FRotation3(InDriveParams.OrientationTarget.Quaternion()));
+				}
 			}
 
 			bool bVelocityDriveEnabled = InDriveParams.IsVelocityDriveEnabled();
@@ -452,7 +456,18 @@ void FPhysInterface_Chaos::UpdateAngularDrive_AssumesLocked(const FPhysicsConstr
 					Constraint->SetAngularSLerpVelocityDriveEnabled(InDriveParams.SlerpDrive.bEnableVelocityDrive);
 				}
 
-				Constraint->SetAngularDriveVelocityTarget(InDriveParams.AngularVelocityTarget);
+				if (!FMath::IsNearlyEqual(Constraint->GetAngularPlasticityLimit(), FLT_MAX))
+				{
+					// Plasticity requires a zero relative velocity.
+					if (!Constraint->GetAngularDriveVelocityTarget().IsZero())
+					{
+						Constraint->SetAngularDriveVelocityTarget(FVector(ForceInitToZero));
+					}
+				}
+				else
+				{
+					Constraint->SetAngularDriveVelocityTarget(InDriveParams.AngularVelocityTarget);
+				}
 			}
 
 			Constraint->SetAngularDriveForceMode(Chaos::EJointForceMode::Acceleration);
diff --git a/Engine/Source/Runtime/Engine/Private/PhysicsEngine/Experimental/PhysicsInterfaceImmediatePhysX.cpp b/Engine/Source/Runtime/Engine/Private/PhysicsEngine/Experimental/PhysicsInterfaceImmediatePhysX.cpp
--- a/Engine/Source/Runtime/Engine/Private/PhysicsEngine/Experimental/PhysicsInterfaceImmediatePhysX.cpp
+++ b/Engine/Source/Runtime/Engine/Private/PhysicsEngine/Experimental/PhysicsInterfaceImmediatePhysX.cpp
@@ -1394,6 +1394,12 @@ void FPhysicsInterface_ImmediatePhysX::SetBreakForces_AssumesLocked(const FPhysi
     */
 }
 
+void FPhysicsInterface_ImmediatePhysX::SetPlasticityLimits_AssumesLocked(const FPhysicsConstraintHandle& InConstraintRef, float InLinearBreakForce, float InAngularBreakTorque)
+{
+	// @todo(chaos) : Implement
+}
+
+
 void FPhysicsInterface_ImmediatePhysX::SetLocalPose(const FPhysicsConstraintHandle& InConstraintRef, const FTransform& InPose, EConstraintFrame::Type InFrame)
 {
     if (InFrame == EConstraintFrame::Frame1)
diff --git a/Engine/Source/Runtime/Engine/Private/PhysicsEngine/ImmediatePhysics/ImmediatePhysicsChaos/ImmediatePhysicsJointHandle_Chaos.cpp b/Engine/Source/Runtime/Engine/Private/PhysicsEngine/ImmediatePhysics/ImmediatePhysicsChaos/ImmediatePhysicsJointHandle_Chaos.cpp
--- a/Engine/Source/Runtime/Engine/Private/PhysicsEngine/ImmediatePhysics/ImmediatePhysicsChaos/ImmediatePhysicsJointHandle_Chaos.cpp
+++ b/Engine/Source/Runtime/Engine/Private/PhysicsEngine/ImmediatePhysics/ImmediatePhysicsChaos/ImmediatePhysicsJointHandle_Chaos.cpp
@@ -113,6 +113,7 @@ namespace ImmediatePhysics_Chaos
 
 		ConstraintSettings.LinearBreakForce = (Profile.bLinearBreakable) ? Chaos::ConstraintSettings::LinearBreakScale() * Profile.LinearBreakThreshold : FLT_MAX;
 		ConstraintSettings.AngularBreakTorque = (Profile.bAngularBreakable) ? Chaos::ConstraintSettings::AngularBreakScale() * Profile.AngularBreakThreshold : FLT_MAX;
+		ConstraintSettings.AngularPlasticityLimit = (Profile.bAngularPlasticity) ? Profile.AngularPlasticityThreshold : FLT_MAX;
 
 		// UE Disables Soft Limits when the Limit is less than some threshold. This is not necessary in Chaos but for now we also do it for parity's sake (See FLinearConstraint::UpdateLinearLimit_AssumesLocked).
 		if (ConstraintSettings.LinearLimit < RB_MinSizeToLockDOF)
diff --git a/Engine/Source/Runtime/Engine/Private/PhysicsEngine/PhysicsConstraintComponent.cpp b/Engine/Source/Runtime/Engine/Private/PhysicsEngine/PhysicsConstraintComponent.cpp
--- a/Engine/Source/Runtime/Engine/Private/PhysicsEngine/PhysicsConstraintComponent.cpp
+++ b/Engine/Source/Runtime/Engine/Private/PhysicsEngine/PhysicsConstraintComponent.cpp
@@ -661,6 +661,11 @@ void UPhysicsConstraintComponent::SetAngularBreakable(bool bAngularBreakable, fl
 	ConstraintInstance.SetAngularBreakable(bAngularBreakable, AngularBreakThreshold);
 }
 
+void UPhysicsConstraintComponent::SetAngularPlasticity(bool bAngularPlasticity, float AngularPlasticityThreshold)
+{
+	ConstraintInstance.SetAngularPlasticity(bAngularPlasticity, AngularPlasticityThreshold);
+}
+
 float UPhysicsConstraintComponent::GetCurrentTwist() const
 {
 	const float CurrentTwistRads = ConstraintInstance.GetCurrentTwist();
diff --git a/Engine/Source/Runtime/Engine/Public/Physics/Experimental/PhysicsInterfaceImmediatePhysX.h b/Engine/Source/Runtime/Engine/Public/Physics/Experimental/PhysicsInterfaceImmediatePhysX.h
--- a/Engine/Source/Runtime/Engine/Public/Physics/Experimental/PhysicsInterfaceImmediatePhysX.h
+++ b/Engine/Source/Runtime/Engine/Public/Physics/Experimental/PhysicsInterfaceImmediatePhysX.h
@@ -353,6 +353,7 @@ struct ENGINE_API FPhysicsInterface_ImmediatePhysX
 	static void SetProjectionEnabled_AssumesLocked(const FPhysicsConstraintReference_ImmediatePhysX& InConstraintRef, bool bInProjectionEnabled, float InLinearTolerance = 0.0f, float InAngularToleranceDegrees = 0.0f);
 	static void SetParentDominates_AssumesLocked(const FPhysicsConstraintReference_ImmediatePhysX& InConstraintRef, bool bInParentDominates);
 	static void SetBreakForces_AssumesLocked(const FPhysicsConstraintReference_ImmediatePhysX& InConstraintRef, float InLinearBreakForce, float InAngularBreakForce);
+	static void SetPlasticityLimits_AssumesLocked(const FPhysicsConstraintHandle& InConstraintRef, float InLinearBreakForce, float InAngularBreakTorque);
 	static void SetLocalPose(const FPhysicsConstraintReference_ImmediatePhysX& InConstraintRef, const FTransform& InPose, EConstraintFrame::Type InFrame);
 
 	static void SetLinearMotionLimitType_AssumesLocked(const FPhysicsConstraintReference_ImmediatePhysX& InConstraintRef, PhysicsInterfaceTypes::ELimitAxis InAxis, ELinearConstraintMotion InMotion);
diff --git a/Engine/Source/Runtime/Experimental/Chaos/Private/Chaos/PBDJointConstraints.cpp b/Engine/Source/Runtime/Experimental/Chaos/Private/Chaos/PBDJointConstraints.cpp
--- a/Engine/Source/Runtime/Experimental/Chaos/Private/Chaos/PBDJointConstraints.cpp
+++ b/Engine/Source/Runtime/Experimental/Chaos/Private/Chaos/PBDJointConstraints.cpp
@@ -171,7 +171,9 @@ namespace Chaos
 		, AngularDriveStiffness(0)
 		, AngularDriveDamping(0)
 		, LinearBreakForce(FLT_MAX)
+		, LinearPlasticityLimit(FLT_MAX)
 		, AngularBreakTorque(FLT_MAX)
+		, AngularPlasticityLimit(FLT_MAX)
 		, UserData(nullptr)
 	{
 		if (bChaos_Joint_ISPC_Enabled)
@@ -1311,6 +1313,11 @@ namespace Chaos
 			ApplyBreakThreshold(Dt, ConstraintIndex, Solver.GetNetLinearImpulse(), Solver.GetNetAngularImpulse());
 		}
 
+		if (/*(JointSettings.LinearPlasticityLimit != FLT_MAX) ||*/ (JointSettings.AngularPlasticityLimit != FLT_MAX)) // todo(chaos) : Implement Linear Plasticity
+		{
+			ApplyPlasticityLimits(Dt, ConstraintIndex, Solver.GetP(1) - Solver.GetP(0), Solver.GetQ(0).Inverse() * Solver.GetQ(1));
+		}
+
 		return Solver.GetIsActive() || !bChaos_Joint_EarlyOut_Enabled;
 	}
 
@@ -1417,6 +1424,23 @@ namespace Chaos
 	}
 
 
+	void FPBDJointConstraints::ApplyPlasticityLimits(const FReal Dt, int32 ConstraintIndex, const FVec3& LinearDisplacement, const FRotation3& AngularDisplacement)
+	{
+		FPBDJointSettings& JointSettings = ConstraintSettings[ConstraintIndex];
+
+		// @todo(chaos) : Implement Linear case
+
+		if (JointSettings.AngularPlasticityLimit)
+		{
+			const FReal AngleDeg = JointSettings.AngularDrivePositionTarget.AngularDistance(AngularDisplacement);
+			if (AngleDeg > JointSettings.AngularPlasticityLimit)
+			{
+				JointSettings.AngularDrivePositionTarget = AngularDisplacement;
+			}
+		}
+	}
+
+
 	// Assign an Island, Level and Color to each constraint. Constraints must be processed in Level order, but
 	// constraints of the same color are independent and can be processed in parallel (SIMD or Task)
 	// NOTE: Constraints are the Vertices, and Edges connect constraints sharing a Particle
diff --git a/Engine/Source/Runtime/Experimental/Chaos/Private/PhysicsProxy/JointConstraintProxy.cpp b/Engine/Source/Runtime/Experimental/Chaos/Private/PhysicsProxy/JointConstraintProxy.cpp
--- a/Engine/Source/Runtime/Experimental/Chaos/Private/PhysicsProxy/JointConstraintProxy.cpp
+++ b/Engine/Source/Runtime/Experimental/Chaos/Private/PhysicsProxy/JointConstraintProxy.cpp
@@ -192,12 +192,14 @@ void TJointConstraintProxy<Chaos::FJointConstraint>::PushStateOnGameThread(Chaos
 			if (Constraint->IsDirty(Chaos::EJointConstraintFlags::LinearBreakForce))
 			{
 				JointSettingsBuffer.LinearBreakForce = Constraint->GetLinearBreakForce();
+				JointSettingsBuffer.LinearPlasticityLimit = Constraint->GetLinearPlasticityLimit();
 				DirtyFlagsBuffer.MarkDirty(Chaos::EJointConstraintFlags::LinearBreakForce);
 			}
 
 			if (Constraint->IsDirty(Chaos::EJointConstraintFlags::AngularBreakTorque))
 			{
 				JointSettingsBuffer.AngularBreakTorque = Constraint->GetAngularBreakTorque();
+				JointSettingsBuffer.AngularPlasticityLimit = Constraint->GetAngularPlasticityLimit();
 				DirtyFlagsBuffer.MarkDirty(Chaos::EJointConstraintFlags::AngularBreakTorque);
 			}
 
@@ -343,11 +345,13 @@ void TJointConstraintProxy<Chaos::FJointConstraint>::PushStateOnPhysicsThread(Ch
 			if (DirtyFlagsBuffer.IsDirty(Chaos::EJointConstraintFlags::LinearBreakForce))
 			{
 				ConstraintSettings.LinearBreakForce = JointSettingsBuffer.LinearBreakForce;
+				ConstraintSettings.LinearPlasticityLimit = JointSettingsBuffer.LinearPlasticityLimit;
 			}
 
 			if (DirtyFlagsBuffer.IsDirty(Chaos::EJointConstraintFlags::AngularBreakTorque))
 			{
 				ConstraintSettings.AngularBreakTorque = JointSettingsBuffer.AngularBreakTorque;
+				ConstraintSettings.AngularPlasticityLimit = JointSettingsBuffer.AngularPlasticityLimit;
 			}
 
 			if (DirtyFlagsBuffer.IsDirty(Chaos::EJointConstraintFlags::UserData))
diff --git a/Engine/Source/Runtime/Experimental/Chaos/Public/Chaos/PBDJointConstraintData.h b/Engine/Source/Runtime/Experimental/Chaos/Public/Chaos/PBDJointConstraintData.h
--- a/Engine/Source/Runtime/Experimental/Chaos/Public/Chaos/PBDJointConstraintData.h
+++ b/Engine/Source/Runtime/Experimental/Chaos/Public/Chaos/PBDJointConstraintData.h
@@ -82,10 +82,18 @@ namespace Chaos
 		//void SetLinearBreakForce(FReal InLinearBreakForce);
 		//FReal GetLinearBreakForce() const
 
+		CONSTRAINT_JOINT_PROPERPETY_IMPL(FReal, LinearPlasticityLimit, EJointConstraintFlags::LinearBreakForce, JointSettings.LinearPlasticityLimit);
+		//void SetLinearPlasticityLimit(FReal InLinearPlasticityLimit);
+		//FReal GetLinearPlasticityLimit() const
+
 		CONSTRAINT_JOINT_PROPERPETY_IMPL(FReal, AngularBreakTorque, EJointConstraintFlags::AngularBreakTorque, JointSettings.AngularBreakTorque);
 		//void SetAngularBreakTorque(FReal InAngularBreakTorque);
 		//FReal GetAngularBreakTorque() const
 
+		CONSTRAINT_JOINT_PROPERPETY_IMPL(FReal, AngularPlasticityLimit, EJointConstraintFlags::AngularBreakTorque, JointSettings.AngularPlasticityLimit);
+		//void SetAngularPlasticityLimit(FReal InAngularPlasticityLimit);
+		//FReal GetAngularPlasticityLimit() const
+
 		CONSTRAINT_JOINT_PROPERPETY_IMPL(void*, UserData, EJointConstraintFlags::UserData, UserData);
 		//void SetUserData(void* InUserData);
 		//void* GetUserData() const
diff --git a/Engine/Source/Runtime/Experimental/Chaos/Public/Chaos/PBDJointConstraintTypes.h b/Engine/Source/Runtime/Experimental/Chaos/Public/Chaos/PBDJointConstraintTypes.h
--- a/Engine/Source/Runtime/Experimental/Chaos/Public/Chaos/PBDJointConstraintTypes.h
+++ b/Engine/Source/Runtime/Experimental/Chaos/Public/Chaos/PBDJointConstraintTypes.h
@@ -160,7 +160,9 @@ namespace Chaos
 		FReal AngularDriveDamping;
 
 		FReal LinearBreakForce;
+		FReal LinearPlasticityLimit;
 		FReal AngularBreakTorque;
+		FReal AngularPlasticityLimit;
 
 		void* UserData;
 	};
diff --git a/Engine/Source/Runtime/Experimental/Chaos/Public/Chaos/PBDJointConstraints.h b/Engine/Source/Runtime/Experimental/Chaos/Public/Chaos/PBDJointConstraints.h
--- a/Engine/Source/Runtime/Experimental/Chaos/Public/Chaos/PBDJointConstraints.h
+++ b/Engine/Source/Runtime/Experimental/Chaos/Public/Chaos/PBDJointConstraints.h
@@ -241,6 +241,7 @@ namespace Chaos
 		bool ApplySingle(const FReal Dt, const int32 ConstraintIndex, const int32 NumPairIts, const int32 It, const int32 NumIts);
 		bool ApplyPushOutSingle(const FReal Dt, const int32 ConstraintIndex, const int32 NumPairIts, const int32 It, const int32 NumIts);
 		void ApplyBreakThreshold(const FReal Dt, int32 ConstraintIndex, const FVec3& LinearImpulse, const FVec3& AngularImpulse);
+		void ApplyPlasticityLimits(const FReal Dt, int32 ConstraintIndex, const FVec3& LinearDisplacement, const FRotation3& AngularDisplacement);
 
 		FPBDJointSolverSettings Settings;
 
diff --git a/Engine/Source/Runtime/PhysicsCore/Private/ChaosEngineInterface.cpp b/Engine/Source/Runtime/PhysicsCore/Private/ChaosEngineInterface.cpp
--- a/Engine/Source/Runtime/PhysicsCore/Private/ChaosEngineInterface.cpp
+++ b/Engine/Source/Runtime/PhysicsCore/Private/ChaosEngineInterface.cpp
@@ -1304,6 +1304,18 @@ void FChaosEngineInterface::SetBreakForces_AssumesLocked(const FPhysicsConstrain
 	}
 }
 
+void FChaosEngineInterface::SetPlasticityLimits_AssumesLocked(const FPhysicsConstraintHandle& InConstraintRef, float InLinearPlasticityLimit, float InAngularPlasticityLimit)
+{
+	if (InConstraintRef.IsValid() && InConstraintRef.Constraint->IsType(Chaos::EConstraintType::JointConstraintType))
+	{
+		if (Chaos::FJointConstraint* Constraint = static_cast<Chaos::FJointConstraint*>(InConstraintRef.Constraint))
+		{
+			Constraint->SetLinearPlasticityLimit(InLinearPlasticityLimit);
+			Constraint->SetAngularPlasticityLimit(InAngularPlasticityLimit);
+		}
+	}
+}
+
 void FChaosEngineInterface::SetLocalPose(const FPhysicsConstraintHandle& InConstraintRef,const FTransform& InPose,EConstraintFrame::Type InFrame)
 {
 	// @todo(chaos) :  Joint Constraints : Motors
diff --git a/Engine/Source/Runtime/PhysicsCore/Public/Chaos/ChaosEngineInterface.h b/Engine/Source/Runtime/PhysicsCore/Public/Chaos/ChaosEngineInterface.h
--- a/Engine/Source/Runtime/PhysicsCore/Public/Chaos/ChaosEngineInterface.h
+++ b/Engine/Source/Runtime/PhysicsCore/Public/Chaos/ChaosEngineInterface.h
@@ -481,6 +481,7 @@ public:
 	static void SetProjectionEnabled_AssumesLocked(const FPhysicsConstraintHandle& InConstraintRef,bool bInProjectionEnabled,float InLinearAlpah = 1.0f,float InAngularAlpha = 0.0f);
 	static void SetParentDominates_AssumesLocked(const FPhysicsConstraintHandle& InConstraintRef,bool bInParentDominates);
 	static void SetBreakForces_AssumesLocked(const FPhysicsConstraintHandle& InConstraintRef,float InLinearBreakForce,float InAngularBreakForce);
+	static void SetPlasticityLimits_AssumesLocked(const FPhysicsConstraintHandle& InConstraintRef, float InLinearPlasticityLimit, float InAngularPlasticityLimit);
 	static void SetLocalPose(const FPhysicsConstraintHandle& InConstraintRef,const FTransform& InPose,EConstraintFrame::Type InFrame);
 
 	static void SetDrivePosition(const FPhysicsConstraintHandle& InConstraintRef,const FVector& InPosition);
