Fix crash when manipulating pinbase that has nullptr reroute I/O pins.

diff --git a/Engine/Source/Runtime/Engine/Private/Materials/MaterialExpressions.cpp b/Engine/Source/Runtime/Engine/Private/Materials/MaterialExpressions.cpp
--- a/Engine/Source/Runtime/Engine/Private/Materials/MaterialExpressions.cpp
+++ b/Engine/Source/Runtime/Engine/Private/Materials/MaterialExpressions.cpp
@@ -10482,8 +10482,11 @@ TArray<FExpressionOutput>& UMaterialExpressionComposite::GetOutputs()
 	{
 		for (const FCompositeReroute& ReroutePin : OutputExpressions->ReroutePins)
 		{
-			ReroutePin.Expression->GetOutputs()[0].OutputName = ReroutePin.Name;
-			Outputs.Add(ReroutePin.Expression->GetOutputs()[0]);
+			if (ReroutePin.Expression)
+			{
+				ReroutePin.Expression->GetOutputs()[0].OutputName = ReroutePin.Name;
+				Outputs.Add(ReroutePin.Expression->GetOutputs()[0]);
+			}
 		}
 	}
 	return Outputs;
@@ -10498,7 +10501,10 @@ const TArray<FExpressionInput*> UMaterialExpressionComposite::GetInputs()
 	{
 		for (const FCompositeReroute& ReroutePin : InputExpressions->ReroutePins)
 		{
-			ExpressionInputs.Add(ReroutePin.Expression->GetInput(0));
+			if (ReroutePin.Expression)
+			{
+				ExpressionInputs.Add(ReroutePin.Expression->GetInput(0));
+			}
 		}
 	}
 	return ExpressionInputs;
@@ -10508,7 +10514,10 @@ FExpressionInput* UMaterialExpressionComposite::GetInput(int32 InputIndex)
 {
 	if (InputIndex >= 0 && InputIndex < InputExpressions->ReroutePins.Num())
 	{
-		return InputExpressions->ReroutePins[InputIndex].Expression->GetInput(0);
+		if (InputExpressions->ReroutePins[InputIndex].Expression)
+		{
+			return InputExpressions->ReroutePins[InputIndex].Expression->GetInput(0);
+		}
 	}
 
 	return nullptr;
@@ -10616,9 +10625,16 @@ void UMaterialExpressionPinBase::DeleteReroutePins()
 	Modify();
 	for (FCompositeReroute& Reroute : ReroutePins)
 	{
-		Reroute.Expression->Modify();
-		Material->Expressions.Remove(Reroute.Expression);
-		Reroute.Expression->MarkPendingKill();
+		if (Reroute.Expression)
+		{
+			Reroute.Expression->Modify();
+			Material->Expressions.Remove(Reroute.Expression);
+			Reroute.Expression->MarkPendingKill();
+		}
+		else
+		{
+			Material->Expressions.Remove(Reroute.Expression);
+		}
 	}
 	ReroutePins.Empty();
 }
@@ -10663,17 +10679,31 @@ void UMaterialExpressionPinBase::PostEditChangeProperty(FPropertyChangedEvent& P
 			uint32 RemovedRerouteIndex = PropertyChangedEvent.GetArrayIndex(PropertyChangedEvent.Property->GetFName().ToString());
 			UMaterialExpression* RemovedReroute = PreEditRereouteExpresions[RemovedRerouteIndex];
 
-			RemovedReroute->Modify();
-			Material->Expressions.Remove(RemovedReroute);
-			RemovedReroute->MarkPendingKill();
+			if (RemovedReroute)
+			{
+				RemovedReroute->Modify();
+				Material->Expressions.Remove(RemovedReroute);
+				RemovedReroute->MarkPendingKill();
+			}
+			else
+			{
+				Material->Expressions.Remove(RemovedReroute);
+			}
 		}
 		else if (PropertyChangedEvent.ChangeType == EPropertyChangeType::ArrayClear)
 		{
 			for (UMaterialExpressionReroute* RemovedReroute : PreEditRereouteExpresions)
 			{
-				RemovedReroute->Modify();
-				Material->Expressions.Remove(RemovedReroute);
-				RemovedReroute->MarkPendingKill();
+				if (RemovedReroute)
+				{
+					RemovedReroute->Modify();
+					Material->Expressions.Remove(RemovedReroute);
+					RemovedReroute->MarkPendingKill();
+				}
+				else
+				{
+					Material->Expressions.Remove(RemovedReroute);
+				}
 			}
 		}
 
@@ -10702,8 +10732,11 @@ TArray<FExpressionOutput>& UMaterialExpressionPinBase::GetOutputs()
 		Outputs.Reset();
 		for (const FCompositeReroute& ReroutePin : ReroutePins)
 		{
-			ReroutePin.Expression->GetOutputs()[0].OutputName = ReroutePin.Name;
-			Outputs.Add(ReroutePin.Expression->GetOutputs()[0]);
+			if (ReroutePin.Expression)
+			{
+				ReroutePin.Expression->GetOutputs()[0].OutputName = ReroutePin.Name;
+				Outputs.Add(ReroutePin.Expression->GetOutputs()[0]);
+			}
 		}
 	}
 	return Outputs;
@@ -10716,7 +10749,10 @@ const TArray<FExpressionInput*> UMaterialExpressionPinBase::GetInputs()
 	{
 		for (const FCompositeReroute& ReroutePin : ReroutePins)
 		{
-			ExpressionInputs.Add(ReroutePin.Expression->GetInput(0));
+			if (ReroutePin.Expression)
+			{
+				ExpressionInputs.Add(ReroutePin.Expression->GetInput(0));
+			}
 		}
 	}
 	return ExpressionInputs;
@@ -10728,7 +10764,10 @@ FExpressionInput* UMaterialExpressionPinBase::GetInput(int32 InputIndex)
 	{
 		if (InputIndex >= 0 && InputIndex < ReroutePins.Num())
 		{
-			return ReroutePins[InputIndex].Expression->GetInput(0);
+			if (ReroutePins[InputIndex].Expression)
+			{
+				return ReroutePins[InputIndex].Expression->GetInput(0);
+			}
 		}
 	}
 
