Fix incomplete shadermaps being cooked (UE-106056).

diff --git a/Engine/Source/Runtime/Engine/Private/ShaderCompiler/ShaderCompiler.cpp b/Engine/Source/Runtime/Engine/Private/ShaderCompiler/ShaderCompiler.cpp
--- a/Engine/Source/Runtime/Engine/Private/ShaderCompiler/ShaderCompiler.cpp
+++ b/Engine/Source/Runtime/Engine/Private/ShaderCompiler/ShaderCompiler.cpp
@@ -2755,7 +2755,19 @@ void FShaderCompilingManager::BlockOnShaderMapCompletion(const TArray<int32>& Sh
 						{
 							if (Results->FinishedJobs.Num() > 0)
 							{
-								CompiledShaderMaps.Add(ShaderMapIdsToFinishCompiling[ShaderMapIndex], FShaderMapFinalizeResults(*Results));
+								FShaderMapFinalizeResults NewResults(*Results);
+
+								// merge with the previous unprocessed jobs, if any
+								if (FShaderMapCompileResults const* PrevResults = CompiledShaderMaps.Find(ShaderMapIdsToFinishCompiling[ShaderMapIndex]))
+								{
+									NewResults.bAllJobsSucceeded = NewResults.bAllJobsSucceeded && PrevResults->bAllJobsSucceeded;
+									NewResults.bSkipResultProcessing = NewResults.bSkipResultProcessing || PrevResults->bSkipResultProcessing;
+									NewResults.TimeStarted = FMath::Min(NewResults.TimeStarted, PrevResults->TimeStarted);
+									NewResults.bIsHung = NewResults.bIsHung || PrevResults->bIsHung;
+									NewResults.FinishedJobs.Append(PrevResults->FinishedJobs);
+								}
+
+								CompiledShaderMaps.Add(ShaderMapIdsToFinishCompiling[ShaderMapIndex], NewResults);
 							}
 							ShaderMapJobs.Remove(ShaderMapIdsToFinishCompiling[ShaderMapIndex]);
 						}
