Fix the editor shutdown crash that was caused by CL 15049697 and re-enable the code which was disabled in CL 15052142

- In the original code we were checking if either of the queues were nullptr before using them however when those queues are deleted they were not actually being set to nullptr, so the checks would not catch this and eventually we would read trash data when trying to use the queues.
- Set both the queues to nullptr on shutdown so no attempt to access them will be made after this point.

#rnx

diff --git a/Engine/Source/Runtime/Engine/Private/StaticMesh.cpp b/Engine/Source/Runtime/Engine/Private/StaticMesh.cpp
--- a/Engine/Source/Runtime/Engine/Private/StaticMesh.cpp
+++ b/Engine/Source/Runtime/Engine/Private/StaticMesh.cpp
@@ -3846,9 +3846,6 @@ bool UStaticMesh::IsReadyForFinishDestroy()
 		return false;
 	}
 
-	// The following code block was causing a shutdown crash in the editor, so
-	// disabling it for now pending investigation.
-#if 0 
 	if (GetRenderData())
 	{
 		if (GDistanceFieldAsyncQueue)
@@ -3861,8 +3858,6 @@ bool UStaticMesh::IsReadyForFinishDestroy()
 			GCardRepresentationAsyncQueue->CancelBuild(this);
 		}
 	}
-#endif //0
-
 #endif
 
 	if (!Super::IsReadyForFinishDestroy())
diff --git a/Engine/Source/Runtime/Launch/Private/LaunchEngineLoop.cpp b/Engine/Source/Runtime/Launch/Private/LaunchEngineLoop.cpp
--- a/Engine/Source/Runtime/Launch/Private/LaunchEngineLoop.cpp
+++ b/Engine/Source/Runtime/Launch/Private/LaunchEngineLoop.cpp
@@ -4300,12 +4300,14 @@ void FEngineLoop::Exit()
 	{
 		GDistanceFieldAsyncQueue->Shutdown();
 		delete GDistanceFieldAsyncQueue;
+		GDistanceFieldAsyncQueue = nullptr;
 	}
 
 	if (GCardRepresentationAsyncQueue)
 	{
 		GCardRepresentationAsyncQueue->Shutdown();
 		delete GCardRepresentationAsyncQueue;
+		GCardRepresentationAsyncQueue = nullptr;
 	}
 #endif // WITH_ENGINE
 
