Fixing PVS 7.7 Issues:

warning V781: The value of the 'XYZ' index is checked after it was used. Perhaps there is a mistake in program logic.

diff --git a/Engine/Source/Developer/MaterialBaking/Private/MaterialBakingHelpers.cpp b/Engine/Source/Developer/MaterialBaking/Private/MaterialBakingHelpers.cpp
--- a/Engine/Source/Developer/MaterialBaking/Private/MaterialBakingHelpers.cpp
+++ b/Engine/Source/Developer/MaterialBaking/Private/MaterialBakingHelpers.cpp
@@ -283,7 +283,7 @@ namespace FMaterialBakingHelpersImpl
 						PreviousY = Y - 1;
 						NextRowsLeft->Add(Y - 1);
 					}
-					if (!RowCompleted[Y] && RowRemainingPixels[RowIndex] && PreviousY < Y)
+					if (PreviousY < Y && !RowCompleted[Y] && RowRemainingPixels[RowIndex])
 					{
 						PreviousY = Y;
 						NextRowsLeft->Add(Y);
diff --git a/Engine/Source/Developer/Windows/LiveCoding/Private/External/LC_StringUtil.cpp b/Engine/Source/Developer/Windows/LiveCoding/Private/External/LC_StringUtil.cpp
--- a/Engine/Source/Developer/Windows/LiveCoding/Private/External/LC_StringUtil.cpp
+++ b/Engine/Source/Developer/Windows/LiveCoding/Private/External/LC_StringUtil.cpp
@@ -89,7 +89,9 @@ namespace string
 	std::wstring ToWideString(const char* utf8Str, size_t count)
 	{
 		size_t length = 0u;
-		while ((utf8Str[length] != '\0') && (length < count))
+		// BEGIN EPIC MOD - PVS FIX
+		while ((length < count) && (utf8Str[length] != '\0'))
+		// END EPIC MOD
 		{
 			// find null-terminator
 			++length;
diff --git a/Engine/Source/Runtime/Core/Private/Misc/Parse.cpp b/Engine/Source/Runtime/Core/Private/Misc/Parse.cpp
--- a/Engine/Source/Runtime/Core/Private/Misc/Parse.cpp
+++ b/Engine/Source/Runtime/Core/Private/Misc/Parse.cpp
@@ -1236,7 +1236,7 @@ bool FParse::SchemeNameFromURI(const TCHAR* URI, FString& OutSchemeName)
 	{
 		if(!FChar::IsAlpha(URI[Idx]) && !FChar::IsDigit(URI[Idx]) && URI[Idx] != TEXT('+') && URI[Idx] != TEXT('.') && URI[Idx] != TEXT('-'))
 		{
-			if(URI[Idx] == TEXT(':') && Idx > 0)
+			if(Idx > 0 && URI[Idx] == TEXT(':'))
 			{
 				OutSchemeName = FString(Idx, URI);
 				return true;
diff --git a/Engine/Source/Runtime/Engine/Private/Audio/AudioAddressPattern.cpp b/Engine/Source/Runtime/Engine/Private/Audio/AudioAddressPattern.cpp
--- a/Engine/Source/Runtime/Engine/Private/Audio/AudioAddressPattern.cpp
+++ b/Engine/Source/Runtime/Engine/Private/Audio/AudioAddressPattern.cpp
@@ -46,7 +46,7 @@ bool FAudioAddressPattern::BracePatternMatches(const FString& Pattern, int32 Pat
 		{
 			bMatches = false;
 			PartIndex = PartIter;
-			while (Pattern[PatternIter] != ',' && PatternIter < PatternEndIndex)
+			while (PatternIter < PatternEndIndex && Pattern[PatternIter] != ',')
 			{
 				++PatternIter;
 			}
diff --git a/Engine/Source/Runtime/Engine/Private/Materials/Material.cpp b/Engine/Source/Runtime/Engine/Private/Materials/Material.cpp
--- a/Engine/Source/Runtime/Engine/Private/Materials/Material.cpp
+++ b/Engine/Source/Runtime/Engine/Private/Materials/Material.cpp
@@ -5786,7 +5786,7 @@ int32 UMaterial::CompilePropertyEx( FMaterialCompiler* Compiler, const FGuid& At
 			{
 				const int32 TextureCoordinateIndex = Property - MP_CustomizedUVs0;
 
-				if (CustomizedUVs[TextureCoordinateIndex].Expression && TextureCoordinateIndex < NumCustomizedUVs)
+				if (TextureCoordinateIndex < NumCustomizedUVs && CustomizedUVs[TextureCoordinateIndex].Expression)
 				{
 					return CustomizedUVs[TextureCoordinateIndex].CompileWithDefault(Compiler, Property);
 				}
diff --git a/Engine/Source/Runtime/Engine/Private/TickTaskManager.cpp b/Engine/Source/Runtime/Engine/Private/TickTaskManager.cpp
--- a/Engine/Source/Runtime/Engine/Private/TickTaskManager.cpp
+++ b/Engine/Source/Runtime/Engine/Private/TickTaskManager.cpp
@@ -668,7 +668,7 @@ private:
 		QUICK_SCOPE_CYCLE_COUNTER(STAT_DispatchTickGroup);
 		for (int32 IndexInner = 0; IndexInner < TG_MAX; IndexInner++)
 		{
-			TArray<TGraphTask<FTickFunctionTask>*>& TickArray = HiPriTickTasks[WorldTickGroup][IndexInner];
+			TArray<TGraphTask<FTickFunctionTask>*>& TickArray = HiPriTickTasks[WorldTickGroup][IndexInner]; //-V781
 			if (IndexInner < WorldTickGroup)
 			{
 				check(TickArray.Num() == 0); // makes no sense to have and end TG before the start TG
@@ -684,7 +684,7 @@ private:
 		}
 		for (int32 IndexInner = 0; IndexInner < TG_MAX; IndexInner++)
 		{
-			TArray<TGraphTask<FTickFunctionTask>*>& TickArray = TickTasks[WorldTickGroup][IndexInner];
+			TArray<TGraphTask<FTickFunctionTask>*>& TickArray = TickTasks[WorldTickGroup][IndexInner]; //-V781
 			if (IndexInner < WorldTickGroup)
 			{
 				check(TickArray.Num() == 0); // makes no sense to have and end TG before the start TG
diff --git a/Engine/Source/Runtime/Navmesh/Private/DetourCrowd/DetourPathCorridor.cpp b/Engine/Source/Runtime/Navmesh/Private/DetourCrowd/DetourPathCorridor.cpp
--- a/Engine/Source/Runtime/Navmesh/Private/DetourCrowd/DetourPathCorridor.cpp
+++ b/Engine/Source/Runtime/Navmesh/Private/DetourCrowd/DetourPathCorridor.cpp
@@ -351,7 +351,7 @@ int dtPathCorridor::findCorners(float* cornerVerts, unsigned char* cornerFlags,
 	{
 		int fromIdx = 0;
 		CA_SUPPRESS(6385);
-		while (m_path[fromIdx] != cornerPolys[i] && fromIdx < m_npath)
+		while (fromIdx < m_npath && m_path[fromIdx] != cornerPolys[i])
 		{
 			fromIdx++;
 		}
