Fix world partition streaming when simulating in the editor (SIE).

diff --git a/Engine/Source/Runtime/Engine/Private/WorldPartition/WorldPartitionStreamingPolicy.cpp b/Engine/Source/Runtime/Engine/Private/WorldPartition/WorldPartitionStreamingPolicy.cpp
--- a/Engine/Source/Runtime/Engine/Private/WorldPartition/WorldPartitionStreamingPolicy.cpp
+++ b/Engine/Source/Runtime/Engine/Private/WorldPartition/WorldPartitionStreamingPolicy.cpp
@@ -17,6 +17,11 @@
 #include "DisplayDebugHelpers.h"
 #include "RenderUtils.h"
 
+#if WITH_EDITOR
+#include "Editor.h"
+#include "LevelEditorViewport.h"
+#endif
+
 UWorldPartitionStreamingPolicy::UWorldPartitionStreamingPolicy(const FObjectInitializer& ObjectInitializer)
 : Super(ObjectInitializer) 
 , bIsServerLoadingDone(false)
@@ -37,10 +42,25 @@ void UWorldPartitionStreamingPolicy::UpdateStreamingSources()
 		return;
 	}
 
+	const FTransform WorldToLocal = WorldPartition->GetInstanceTransform().Inverse();
+
+#if WITH_EDITOR
+	// We are in the SIE
+	if (GEditor->bIsSimulatingInEditor && GCurrentLevelEditingViewportClient->IsSimulateInEditorViewport())
+	{
+		// Transform to Local
+		const FVector ViewLocation = GCurrentLevelEditingViewportClient->GetViewLocation();
+		const FRotator ViewRotation = GCurrentLevelEditingViewportClient->GetViewRotation();
+		const FVector ViewLocationLocal = WorldToLocal.TransformPosition(ViewLocation);
+		const FRotator ViewRotationLocal = WorldToLocal.TransformRotation(ViewRotation.Quaternion()).Rotator();
+		StreamingSources.Add(FWorldPartitionStreamingSource(ViewLocationLocal, ViewRotationLocal));
+		return;
+	}
+#endif
+
 	UWorld* World = WorldPartition->GetWorld();
 	if (World->GetNetMode() != NM_DedicatedServer)
 	{
-		FTransform WorldToLocal = WorldPartition->GetInstanceTransform().Inverse();
 		const int32 NumPlayers = GEngine->GetNumGamePlayers(World);
 		for (int32 PlayerIndex = 0; PlayerIndex < NumPlayers; ++PlayerIndex)
 		{
@@ -51,6 +71,8 @@ void UWorldPartitionStreamingPolicy::UpdateStreamingSources()
 				FRotator ViewRotation;
 				Player->PlayerController->GetPlayerViewPoint(ViewLocation, ViewRotation);
 
+				//@todo_ow: this test is to cover cases where GetPlayerViewPoint returns (0,0,0) when invalid. It should probably return a bool
+				//          to indicate that the returned position is invalid.
 				if (!ViewLocation.IsZero())
 				{
 					// Transform to Local
