Fixing PVS 7.7 Issues:

warning V568: It's odd that 'sizeof()' operator evaluates the size of a pointer to a class, but not the size of the 'OverlapsAtEndLocation' class object.

diff --git a/Engine/Source/Runtime/Engine/Private/Components/PrimitiveComponent.cpp b/Engine/Source/Runtime/Engine/Private/Components/PrimitiveComponent.cpp
--- a/Engine/Source/Runtime/Engine/Private/Components/PrimitiveComponent.cpp
+++ b/Engine/Source/Runtime/Engine/Private/Components/PrimitiveComponent.cpp
@@ -3324,7 +3324,7 @@ bool UPrimitiveComponent::UpdateOverlapsImpl(const TOverlapArrayView* NewPending
 
 			// Ensure these arrays are still in scope, because we kept pointers to them in NewOverlappingComponentPtrs.
 			static_assert(sizeof(OverlapMultiResult) != 0, "Variable must be in this scope");
-			static_assert(sizeof(OverlapsAtEndLocation) != 0, "Variable must be in this scope");
+			static_assert(sizeof(*OverlapsAtEndLocation) != 0, "Variable must be in this scope");
 
 			// NewOverlappingComponents now contains only new overlaps that didn't exist previously.
 			for (const FOverlapInfo* NewOverlap : NewOverlappingComponentPtrs)
