Modify default in editor vt scalability settings to address performance cost.

diff --git a/Engine/Source/Runtime/Engine/Private/VT/VirtualTextureScalability.cpp b/Engine/Source/Runtime/Engine/Private/VT/VirtualTextureScalability.cpp
--- a/Engine/Source/Runtime/Engine/Private/VT/VirtualTextureScalability.cpp
+++ b/Engine/Source/Runtime/Engine/Private/VT/VirtualTextureScalability.cpp
@@ -16,7 +16,7 @@ namespace VirtualTextureScalability
 #if WITH_EDITOR
 	static TAutoConsoleVariable<int32> CVarVTMaxUploadsPerFrameInEditor(
 		TEXT("r.VT.MaxUploadsPerFrameInEditor"),
-		64,
+		32,
 		TEXT("Max number of page uploads per frame when in editor"),
 		ECVF_RenderThreadSafe
 	);
@@ -39,7 +39,7 @@ namespace VirtualTextureScalability
 #if WITH_EDITOR
 	static TAutoConsoleVariable<int32> CVarVTMaxContinuousUpdatesPerFrameInEditor(
 		TEXT("r.VT.MaxContinuousUpdatesPerFrameInEditor"),
-		128,
+		8,
 		TEXT("Max number of page uploads for pages that are already mapped when in editor."),
 		ECVF_RenderThreadSafe | ECVF_Scalability
 	);
diff --git a/Engine/Source/Runtime/Engine/Public/VT/VirtualTextureScalability.h b/Engine/Source/Runtime/Engine/Public/VT/VirtualTextureScalability.h
--- a/Engine/Source/Runtime/Engine/Public/VT/VirtualTextureScalability.h
+++ b/Engine/Source/Runtime/Engine/Public/VT/VirtualTextureScalability.h
@@ -8,6 +8,8 @@ namespace VirtualTextureScalability
 {
 	/** Get max upload rate to virtual textures. */
 	ENGINE_API int32 GetMaxUploadsPerFrame();
+	/** Get max produce rate to virtual textures. */
+	ENGINE_API int32 GetMaxPagesProducedPerFrame();
 	/** Get max update rate of already mapped virtual texture pages. */
 	ENGINE_API int32 GetMaxContinuousUpdatesPerFrame();
 	/** Get scale factor for virtual texture physical pool sizes. */
@@ -15,8 +17,6 @@ namespace VirtualTextureScalability
 	/** Get resolution bias for runtime virtual textures. */
 	ENGINE_API int32 GetRuntimeVirtualTextureSizeBias();
 
-	ENGINE_API int32 GetMaxPagesProducedPerFrame();
-
 	/**
 	 * Get maximum anisotropy when virtual texture sampling. 
 	 * This is also clamped per virtual texture according to the tile border size.
