Fix compilation error on non-Editor platforms. Can't check bNaniteEnabled, need to look at streaming page count > 0

diff --git a/Engine/Source/Runtime/Engine/Private/InstancedStaticMesh.cpp b/Engine/Source/Runtime/Engine/Private/InstancedStaticMesh.cpp
--- a/Engine/Source/Runtime/Engine/Private/InstancedStaticMesh.cpp
+++ b/Engine/Source/Runtime/Engine/Private/InstancedStaticMesh.cpp
@@ -2836,27 +2836,35 @@ static bool ComponentRequestsCPUAccess(UInstancedStaticMeshComponent* InComponen
 {
 	bool bNeedsCPUAccess = false;
 
+	// Ray tracing needs instance transforms on CPU
+	bNeedsCPUAccess |= IsRayTracingEnabled();
+
+	const UStaticMesh* StaticMesh = InComponent->GetStaticMesh();
+
 	// Check mesh distance fields
-	if (FeatureLevel > ERHIFeatureLevel::ES3_1)
+	if (StaticMesh)
 	{
-		static const auto CVar = IConsoleManager::Get().FindTConsoleVariableDataInt(TEXT("r.GenerateMeshDistanceFields"));
+		const FStaticMeshRenderData* RenderData = StaticMesh->GetRenderData();
 
-		bNeedsCPUAccess |= (InComponent->CastShadow && InComponent->bAffectDistanceFieldLighting 
-			// Distance field algorithms need access to instance data on the CPU
-			&& (CVar->GetValueOnAnyThread(true) != 0 || (InComponent->GetStaticMesh() && InComponent->GetStaticMesh()->bGenerateMeshDistanceField)));
-	}
+		if (FeatureLevel > ERHIFeatureLevel::ES3_1)
+		{
+			static const auto CVar = IConsoleManager::Get().FindTConsoleVariableDataInt(TEXT("r.GenerateMeshDistanceFields"));
 
-	// Ray tracing needs instance transforms on CPU
-	bNeedsCPUAccess |= IsRayTracingEnabled();
+			bNeedsCPUAccess |= (InComponent->CastShadow && InComponent->bAffectDistanceFieldLighting
+				// Distance field algorithms need access to instance data on the CPU
+				&& (CVar->GetValueOnAnyThread(true) != 0 || StaticMesh->bGenerateMeshDistanceField));
+		}
 
-	// Check Nanite
-	if (FeatureLevel >= ERHIFeatureLevel::SM5)
-	{
-		// TODO: Call UseNanite()
-		//static const auto CVar = IConsoleManager::Get().FindTConsoleVariableDataInt(TEXT("r.Nanite"));
-		
-		// Nanite needs lightmap/shadowmap UV biases and per instance random values accessible from GPUScene
-		bNeedsCPUAccess |= (InComponent->GetStaticMesh() && InComponent->GetStaticMesh()->NaniteSettings.bEnabled);
+		// Check Nanite
+		if (RenderData && FeatureLevel >= ERHIFeatureLevel::SM5)
+		{
+			// TODO: Call UseNanite(GetScene()->GetShaderPlatform())?
+			//static const auto CVar = IConsoleManager::Get().FindTConsoleVariableDataInt(TEXT("r.Nanite"));
+
+			// Nanite needs lightmap/shadowmap UV biases and per instance random values accessible from GPUScene
+			const bool bHasNaniteData = RenderData->NaniteResources.PageStreamingStates.Num() > 0;
+			bNeedsCPUAccess |= bHasNaniteData;
+		}
 	}
 
 	return bNeedsCPUAccess;
