WorldPartition: Only prevent deleting partition actors in a partition world

diff --git a/Engine/Source/Runtime/Engine/Private/ActorPartition/PartitionActor.cpp b/Engine/Source/Runtime/Engine/Private/ActorPartition/PartitionActor.cpp
--- a/Engine/Source/Runtime/Engine/Private/ActorPartition/PartitionActor.cpp
+++ b/Engine/Source/Runtime/Engine/Private/ActorPartition/PartitionActor.cpp
@@ -35,9 +35,13 @@ bool APartitionActor::CanDeleteSelectedActor(FText& OutReason) const
 	}
 
 	check(GetLevel());
-	check(GetLevel()->bIsPartitioned);
-	OutReason = LOCTEXT("CantDeleteSelectedPartitionActor", "Can't delete Partition Actor (automatically managed).");
-	return false;
+	if (GetLevel()->bIsPartitioned)
+	{
+		OutReason = LOCTEXT("CantDeleteSelectedPartitionActor", "Can't delete Partition Actor (automatically managed).");
+		return false;
+	}
+
+	return true;
 }
 #endif
 
