Ensure that FRichCurve::ReadjustTimeRange uses GetKeyHandle rather than trying to find it directly, as the handle map might not be populated yet
#misc Update animation sequence curve key trimming test to new (correct) values

diff --git a/Engine/Source/Runtime/Engine/Private/Curves/RichCurve.cpp b/Engine/Source/Runtime/Engine/Private/Curves/RichCurve.cpp
--- a/Engine/Source/Runtime/Engine/Private/Curves/RichCurve.cpp
+++ b/Engine/Source/Runtime/Engine/Private/Curves/RichCurve.cpp
@@ -753,10 +753,10 @@ void FRichCurve::ReadjustTimeRange(float NewMinTimeRange, float NewMaxTimeRange,
 		{
 			for (int32 KeyIndex = KeysToDelete.Num()-1; KeyIndex >= 0; --KeyIndex)
 			{
-				const FKeyHandle* KeyHandle = KeyHandlesToIndices.FindKey(KeysToDelete[KeyIndex]);
-				if(KeyHandle && IsKeyHandleValid(*KeyHandle))
+				const FKeyHandle KeyHandle = GetKeyHandle(KeysToDelete[KeyIndex]);
+				if(IsKeyHandleValid(KeyHandle))
 				{
-					DeleteKey(*KeyHandle);
+					DeleteKey(KeyHandle);
 				}
 			}
 
