Added GetLevelTransform function which returns the transform of the level streaming the actor belongs to.

diff --git a/Engine/Source/Runtime/Engine/Classes/GameFramework/Actor.h b/Engine/Source/Runtime/Engine/Classes/GameFramework/Actor.h
--- a/Engine/Source/Runtime/Engine/Classes/GameFramework/Actor.h
+++ b/Engine/Source/Runtime/Engine/Classes/GameFramework/Actor.h
@@ -2670,6 +2670,10 @@ public:
 	UFUNCTION(BlueprintCallable, Category=Level)
 	ULevel* GetLevel() const;
 
+	/** Return the FTransform of the level this actor is a part of. */
+	UFUNCTION(BlueprintCallable, Category=Level)
+	FTransform GetLevelTransform() const;
+
 	/**	Do anything needed to clear out cross level references; Called from ULevel::PreSave	 */
 	virtual void ClearCrossLevelReferences();
 	
diff --git a/Engine/Source/Runtime/Engine/Private/Actor.cpp b/Engine/Source/Runtime/Engine/Private/Actor.cpp
--- a/Engine/Source/Runtime/Engine/Private/Actor.cpp
+++ b/Engine/Source/Runtime/Engine/Private/Actor.cpp
@@ -2079,6 +2079,19 @@ ULevel* AActor::GetLevel() const
 	return GetTypedOuter<ULevel>();
 }
 
+FTransform AActor::GetLevelTransform() const
+{
+	if (ULevel* Level = GetLevel())
+	{
+		if (ULevelStreaming* LevelStreaming = FLevelUtils::FindStreamingLevel(Level))
+		{
+			return LevelStreaming->LevelTransform;
+		}
+	}
+
+	return FTransform::Identity;
+}
+
 bool AActor::IsInPersistentLevel(bool bIncludeLevelStreamingPersistent) const
 {
 	ULevel* MyLevel = GetLevel();
