Temporary fix for GC crash for partitioned maps.

diff --git a/Engine/Source/Runtime/Engine/Private/WorldPartition/WorldPartition.cpp b/Engine/Source/Runtime/Engine/Private/WorldPartition/WorldPartition.cpp
--- a/Engine/Source/Runtime/Engine/Private/WorldPartition/WorldPartition.cpp
+++ b/Engine/Source/Runtime/Engine/Private/WorldPartition/WorldPartition.cpp
@@ -143,6 +143,7 @@ FName UWorldPartition::GetWorldPartitionEditorName()
 void UWorldPartition::Initialize(UWorld* InWorld, const FTransform& InTransform)
 {
 	UE_SCOPED_TIMER(TEXT("WorldPartition initialize"), LogWorldPartition);
+	TRACE_CPUPROFILER_EVENT_SCOPE(UWorldPartition::Initialize);
 	
 	check(!World || (World == InWorld));
 	if (!ensure(!IsInitialized()))
@@ -836,6 +837,16 @@ void UWorldPartition::Serialize(FArchive& Ar)
 #endif
 }
 
+void UWorldPartition::BeginDestroy()
+{
+	Super::BeginDestroy();
+
+	for (TUniquePtr<FWorldPartitionActorDesc>& ActorDescPtr : ActorDescList)
+	{
+		ActorDescPtr.Release();
+	}
+}
+
 UWorldPartitionStreamingPolicy* UWorldPartition::GetStreamingPolicy() const
 {
 	check(GetWorld());
diff --git a/Engine/Source/Runtime/Engine/Private/WorldPartition/WorldPartitionActorDesc.cpp b/Engine/Source/Runtime/Engine/Private/WorldPartition/WorldPartitionActorDesc.cpp
--- a/Engine/Source/Runtime/Engine/Private/WorldPartition/WorldPartitionActorDesc.cpp
+++ b/Engine/Source/Runtime/Engine/Private/WorldPartition/WorldPartitionActorDesc.cpp
@@ -203,7 +203,10 @@ AActor* FWorldPartitionActorDesc::GetActor() const
 		return Actor;
 	}
 
-	Actor = FindObject<AActor>(nullptr, *ActorPath.ToString());
+	if (!IsGarbageCollecting())
+	{
+		Actor = FindObject<AActor>(nullptr, *ActorPath.ToString());
+	}
 
 	if (Actor)
 	{
diff --git a/Engine/Source/Runtime/Engine/Public/WorldPartition/WorldPartition.h b/Engine/Source/Runtime/Engine/Public/WorldPartition/WorldPartition.h
--- a/Engine/Source/Runtime/Engine/Public/WorldPartition/WorldPartition.h
+++ b/Engine/Source/Runtime/Engine/Public/WorldPartition/WorldPartition.h
@@ -94,6 +94,7 @@ private:
 
 	//~ Begin UObject Interface
 	virtual void Serialize(FArchive& Ar) override;
+	virtual void BeginDestroy() override;
 	//~ End UObject Interface
 
 	static void WorldPartitionOnLevelRemovedFromWorld(class ULevel* Level, UWorld* InWorld);
