fix non-editor builds

diff --git a/Engine/Source/Runtime/Engine/Private/WorldPartition/WorldPartition.cpp b/Engine/Source/Runtime/Engine/Private/WorldPartition/WorldPartition.cpp
--- a/Engine/Source/Runtime/Engine/Private/WorldPartition/WorldPartition.cpp
+++ b/Engine/Source/Runtime/Engine/Private/WorldPartition/WorldPartition.cpp
@@ -841,10 +841,12 @@ void UWorldPartition::BeginDestroy()
 {
 	Super::BeginDestroy();
 
+#if WITH_EDITOR
 	for (TUniquePtr<FWorldPartitionActorDesc>& ActorDescPtr : ActorDescList)
 	{
 		ActorDescPtr.Release();
 	}
+#endif
 }
 
 UWorldPartitionStreamingPolicy* UWorldPartition::GetStreamingPolicy() const
