Serialize a hash of each VT chunk, allow us to (hopefully) more quickly detect/fix corrupt VT data from the DDC

diff --git a/Engine/Source/Runtime/Engine/Private/TextureDerivedData.cpp b/Engine/Source/Runtime/Engine/Private/TextureDerivedData.cpp
--- a/Engine/Source/Runtime/Engine/Private/TextureDerivedData.cpp
+++ b/Engine/Source/Runtime/Engine/Private/TextureDerivedData.cpp
@@ -53,7 +53,7 @@
 
 // This GUID is mixed into DDC version for virtual textures only, this allows updating DDC version for VT without invalidating DDC for all textures
 // This is useful during development, but once large numbers of VT are present in shipped content, it will have the same problem as TEXTURE_DERIVEDDATA_VER
-#define TEXTURE_VT_DERIVEDDATA_VER	TEXT("FA56564569064608A11C0F5F1CD22E5E")
+#define TEXTURE_VT_DERIVEDDATA_VER	TEXT("2FF3E5D7EFF5424D8631C71A65DB4BF8")
 
 #if ENABLE_COOK_STATS
 namespace TextureCookStats
diff --git a/Engine/Source/Runtime/Engine/Private/TextureDerivedDataTask.cpp b/Engine/Source/Runtime/Engine/Private/TextureDerivedDataTask.cpp
--- a/Engine/Source/Runtime/Engine/Private/TextureDerivedDataTask.cpp
+++ b/Engine/Source/Runtime/Engine/Private/TextureDerivedDataTask.cpp
@@ -37,7 +37,7 @@
 
 static TAutoConsoleVariable<int32> CVarVTValidateCompressionOnLoad(
 	TEXT("r.VT.ValidateCompressionOnLoad"),
-	0,
+	1,
 	TEXT("Validates that VT data contains no compression errors when loading from DDC")
 	TEXT("This is slow, but allows debugging corrupt VT data (and allows recovering from bad DDC)")
 );
@@ -250,7 +250,7 @@ void FTextureCacheDerivedDataWorker::BuildTexture(bool bReplaceExistingDDC)
 		bool bCompressionValid = true;
 		if (CVarVTValidateCompressionOnSave.GetValueOnAnyThread())
 		{
-			bCompressionValid = DerivedData->VTData->ValidateCompression(Texture.GetPathName());
+			bCompressionValid = DerivedData->VTData->ValidateData(Texture.GetPathName(), true);
 		}
 
 		if (ensureMsgf(bCompressionValid, TEXT("Corrupt Virtual Texture compression for %s, can't store to DDC"), *Texture.GetPathName()))
@@ -493,7 +493,7 @@ void FTextureCacheDerivedDataWorker::DoWork()
 		if (bSucceeded && bForVirtualTextureStreamingBuild && CVarVTValidateCompressionOnLoad.GetValueOnAnyThread())
 		{
 			check(DerivedData->VTData);
-			bSucceeded = DerivedData->VTData->ValidateCompression(Texture.GetPathName());
+			bSucceeded = DerivedData->VTData->ValidateData(Texture.GetPathName(), false);
 			if (!bSucceeded)
 			{
 				UE_LOG(LogTexture, Error, TEXT("Texture %s has corrupt Virtual Texture compression. The texture will be rebuild."), *Texture.GetFullName());
diff --git a/Engine/Source/Runtime/Engine/Private/VT/VirtualTextureBuiltData.cpp b/Engine/Source/Runtime/Engine/Private/VT/VirtualTextureBuiltData.cpp
--- a/Engine/Source/Runtime/Engine/Private/VT/VirtualTextureBuiltData.cpp
+++ b/Engine/Source/Runtime/Engine/Private/VT/VirtualTextureBuiltData.cpp
@@ -155,6 +155,7 @@ void FVirtualTextureBuiltData::Serialize(FArchive& Ar, UObject* Owner, int32 Fir
 	{
 		FVirtualTextureDataChunk& Chunk = Chunks[ChunkId];
 
+		Ar << Chunk.BulkDataHash;
 		Ar << Chunk.SizeInBytes;
 		Ar << Chunk.CodecPayloadSize;
 		for (uint32 LayerIndex = 0u; LayerIndex < NumLayers; ++LayerIndex)
@@ -180,7 +181,7 @@ void FVirtualTextureBuiltData::Serialize(FArchive& Ar, UObject* Owner, int32 Fir
 	}
 }
 
-bool FVirtualTextureBuiltData::ValidateCompression(FStringView const& InDDCDebugContext) const
+bool FVirtualTextureBuiltData::ValidateData(FStringView const& InDDCDebugContext, bool bValidateCompression) const
 {
 	const uint32 TilePixelSize = GetPhysicalTileSize();
 	TArray<uint8> UncompressedResult;
@@ -192,9 +193,11 @@ bool FVirtualTextureBuiltData::ValidateCompression(FStringView const& InDDCDebug
 		const FVirtualTextureDataChunk& Chunk = Chunks[ChunkIndex];
 
 		const uint8* ChunkData = nullptr;
+		uint32 ChunkDataSize = 0u;
 		bool bNeedToUnlockBulkData = false;
 		if (Chunk.BulkData.GetBulkDataSize() > 0)
 		{
+			ChunkDataSize = (uint32)Chunk.BulkData.GetBulkDataSize();
 			ChunkData = (uint8*)Chunk.BulkData.LockReadOnly();
 			bNeedToUnlockBulkData = true;
 		}
@@ -204,43 +207,60 @@ bool FVirtualTextureBuiltData::ValidateCompression(FStringView const& InDDCDebug
 			ChunkDataDDC.Reset();
 			const bool bDDCResult = GetDerivedDataCacheRef().GetSynchronous(*Chunk.DerivedDataKey, ChunkDataDDC, InDDCDebugContext);
 			check(bDDCResult);
+			check(ChunkDataDDC.Num() > 4);
 			ChunkData = ChunkDataDDC.GetData() + 4;
+			ChunkDataSize = (uint32)(ChunkDataDDC.Num() - 4);
 		}
 #endif // WITH_EDITORONLY_DATA
 
-		bResult = bResult && (ChunkData != nullptr);
+		if (!ChunkData)
+		{
+			bResult = false;
+			break;
+		}
+
+		FSHAHash Hash;
+		FSHA1::HashBuffer(ChunkData, ChunkDataSize, Hash.Hash);
+		if (Hash != Chunk.BulkDataHash)
+		{
+			bResult = false;
+			break;
+		}
 
-		uint32 TileIndex = TileIndexPerChunk[ChunkIndex];
-		while (bResult && TileIndex < TileIndexPerChunk[ChunkIndex + 1])
+		if (bValidateCompression)
 		{
-			for (uint32 LayerIndex = 0u; LayerIndex < GetNumLayers(); ++LayerIndex)
+			uint32 TileIndex = TileIndexPerChunk[ChunkIndex];
+			while (bResult && TileIndex < TileIndexPerChunk[ChunkIndex + 1])
 			{
-				const EVirtualTextureCodec VTCodec = Chunk.CodecType[LayerIndex];
-				const EPixelFormat LayerFormat = LayerTypes[LayerIndex];
-				const uint32 TileWidthInBlocks = FMath::DivideAndRoundUp(TilePixelSize, (uint32)GPixelFormats[LayerFormat].BlockSizeX);
-				const uint32 TileHeightInBlocks = FMath::DivideAndRoundUp(TilePixelSize, (uint32)GPixelFormats[LayerFormat].BlockSizeY);
-				const uint32 PackedStride = TileWidthInBlocks * GPixelFormats[LayerFormat].BlockBytes;
-				const size_t PackedOutputSize = PackedStride * TileHeightInBlocks;
-
-				if (VTCodec == EVirtualTextureCodec::ZippedGPU)
+				for (uint32 LayerIndex = 0u; LayerIndex < GetNumLayers(); ++LayerIndex)
 				{
-					const uint32 TileOffset = GetTileOffset(ChunkIndex, TileIndex);
-					const uint32 NextTileOffset = GetTileOffset(ChunkIndex, TileIndex + 1);
-					check(NextTileOffset >= TileOffset);
-					if (NextTileOffset > TileOffset)
+					const EVirtualTextureCodec VTCodec = Chunk.CodecType[LayerIndex];
+					const EPixelFormat LayerFormat = LayerTypes[LayerIndex];
+					const uint32 TileWidthInBlocks = FMath::DivideAndRoundUp(TilePixelSize, (uint32)GPixelFormats[LayerFormat].BlockSizeX);
+					const uint32 TileHeightInBlocks = FMath::DivideAndRoundUp(TilePixelSize, (uint32)GPixelFormats[LayerFormat].BlockSizeY);
+					const uint32 PackedStride = TileWidthInBlocks * GPixelFormats[LayerFormat].BlockBytes;
+					const size_t PackedOutputSize = PackedStride * TileHeightInBlocks;
+
+					if (VTCodec == EVirtualTextureCodec::ZippedGPU)
 					{
-						const uint32 CompressedTileSize = NextTileOffset - TileOffset;
-
-						UncompressedResult.SetNumUninitialized(PackedOutputSize, false);
-						const bool bUncompressResult = FCompression::UncompressMemory(NAME_Zlib, UncompressedResult.GetData(), PackedOutputSize, &ChunkData[TileOffset], CompressedTileSize);
-						if (!bUncompressResult)
+						const uint32 TileOffset = GetTileOffset(ChunkIndex, TileIndex);
+						const uint32 NextTileOffset = GetTileOffset(ChunkIndex, TileIndex + 1);
+						check(NextTileOffset >= TileOffset);
+						if (NextTileOffset > TileOffset)
 						{
-							bResult = false;
-							break;
+							const uint32 CompressedTileSize = NextTileOffset - TileOffset;
+
+							UncompressedResult.SetNumUninitialized(PackedOutputSize, false);
+							const bool bUncompressResult = FCompression::UncompressMemory(NAME_Zlib, UncompressedResult.GetData(), PackedOutputSize, &ChunkData[TileOffset], CompressedTileSize);
+							if (!bUncompressResult)
+							{
+								bResult = false;
+								break;
+							}
 						}
 					}
+					++TileIndex;
 				}
-				++TileIndex;
 			}
 		}
 
diff --git a/Engine/Source/Runtime/Engine/Private/VT/VirtualTextureBuiltData.h b/Engine/Source/Runtime/Engine/Private/VT/VirtualTextureBuiltData.h
--- a/Engine/Source/Runtime/Engine/Private/VT/VirtualTextureBuiltData.h
+++ b/Engine/Source/Runtime/Engine/Private/VT/VirtualTextureBuiltData.h
@@ -31,6 +31,7 @@ enum class EVirtualTextureCodec : uint8
 struct FVirtualTextureDataChunk
 {
 	FByteBulkData BulkData;
+	FSHAHash BulkDataHash;
 	uint32 SizeInBytes;
 	uint32 CodecPayloadSize;
 	uint32 CodecPayloadOffset[VIRTUALTEXTURE_DATA_MAXLAYERS];
@@ -191,8 +192,8 @@ struct FVirtualTextureBuiltData
 	}
 
 	/**
-	* Attempts to decompress every VT tile, to verify data is valid
-	* This will be very slow
+	* Validates the VT data
+	* If bValidateCompression is true, attempts to decompress every VT tile, to verify data is valid. This will be very slow
 	*/
-	bool ValidateCompression(FStringView const& InDDCDebugContext) const;
+	bool ValidateData(FStringView const& InDDCDebugContext, bool bValidateCompression) const;
 };
diff --git a/Engine/Source/Runtime/Engine/Private/VT/VirtualTextureChunkDDCCache.cpp b/Engine/Source/Runtime/Engine/Private/VT/VirtualTextureChunkDDCCache.cpp
--- a/Engine/Source/Runtime/Engine/Private/VT/VirtualTextureChunkDDCCache.cpp
+++ b/Engine/Source/Runtime/Engine/Private/VT/VirtualTextureChunkDDCCache.cpp
@@ -26,7 +26,7 @@ DEFINE_LOG_CATEGORY(LogVTDiskCache);
 struct FVirtualTextureFileHeader
 {
 	static const uint32 CurrentMagic = 0x4558ACDF;
-	static const uint32 CurrentVersion = 1u;
+	static const uint32 CurrentVersion = 2u;
 
 	uint32 Magic = 0u;
 	uint32 Version = 0u;
@@ -111,6 +111,24 @@ public:
 			return;
 		}
 
+		if (Results.Num() <= 4)
+		{
+			UE_LOG(LogVTDiskCache, Error, TEXT("VT DDC data is too small %d (key: %s)"), Results.Num(), *Chunk->DerivedDataKey);
+			return;
+		}
+
+		// skip size embedded in DDC entry
+		const uint8* ChunkData = Results.GetData() + 4;
+		const int32 ChunkDataSize = Results.Num() - 4;
+
+		FSHAHash Hash;
+		FSHA1::HashBuffer(ChunkData, ChunkDataSize, Hash.Hash);
+		if (Hash != Chunk->BulkDataHash)
+		{
+			UE_LOG(LogVTDiskCache, Error, TEXT("VT DDC data is corrupt (key: %s)"), *Chunk->DerivedDataKey);
+			return;
+		}
+
 		// Write to Disk
 		{
 			TUniquePtr<FArchive> Ar(IFileManager::Get().CreateFileWriter(*TempFilename, 0));
@@ -123,13 +141,12 @@ public:
 			FVirtualTextureFileHeader Header;
 			Header.Magic = FVirtualTextureFileHeader::CurrentMagic;
 			Header.Version = FVirtualTextureFileHeader::CurrentVersion;
-			Header.HashSize = FMath::Min<uint32>(Results.Num() - 4, kMaxHashSize);
-			FSHA1::HashBuffer(Results.GetData() + 4, Header.HashSize, Header.Hash.Hash);
+			Header.HashSize = FMath::Min<uint32>(ChunkDataSize, kMaxHashSize);
+			FSHA1::HashBuffer(ChunkData, Header.HashSize, Header.Hash.Hash);
 			*Ar << Header;
 			const int64 ArchiveOffset = Ar->Tell();
 			check(ArchiveOffset == sizeof(FVirtualTextureFileHeader));
-
-			Ar->Serialize(const_cast<uint8*>(Results.GetData() + 4), Results.Num() - 4); // skip size embedded in DDC entry
+			Ar->Serialize(const_cast<uint8*>(ChunkData), ChunkDataSize);
 		}
 
 		if (PlatformFile.MoveFile(*FinalFilename, *TempFilename))
diff --git a/Engine/Source/Runtime/Engine/Private/VT/VirtualTextureDataBuilder.cpp b/Engine/Source/Runtime/Engine/Private/VT/VirtualTextureDataBuilder.cpp
--- a/Engine/Source/Runtime/Engine/Private/VT/VirtualTextureDataBuilder.cpp
+++ b/Engine/Source/Runtime/Engine/Private/VT/VirtualTextureDataBuilder.cpp
@@ -767,6 +767,8 @@ void FVirtualTextureDataBuilder::PushDataToChunk(const TArray<FVTSourceTileEntry
 
 	check(ChunkOffset == TotalSize);
 
+	FSHA1::HashBuffer(NewChunkData, TotalSize, Chunk.BulkDataHash.Hash);
+
 	BulkData.Unlock();
 	BulkData.SetBulkDataFlags(BULKDATA_Force_NOT_InlinePayload);
 }
