Fixing PVS 7.7 Issues:

warning V547: Expression 'XYZ' is always true/false.

diff --git a/Engine/Source/Runtime/Engine/Private/TextureRenderTarget2DArray.cpp b/Engine/Source/Runtime/Engine/Private/TextureRenderTarget2DArray.cpp
--- a/Engine/Source/Runtime/Engine/Private/TextureRenderTarget2DArray.cpp
+++ b/Engine/Source/Runtime/Engine/Private/TextureRenderTarget2DArray.cpp
@@ -200,7 +200,7 @@ UTexture2DArray* UTextureRenderTarget2DArray::ConstructTexture2DArray(UObject* O
 	Texture2DArray->Source.UnlockMip(0);
 	Texture2DArray->SRGB = bSRGB;
 	// If HDR source image then choose HDR compression settings..
-	Texture2DArray->CompressionSettings = TextureFormat == TSF_RGBA16F ? TextureCompressionSettings::TC_HDR : TextureCompressionSettings::TC_Default;
+	Texture2DArray->CompressionSettings = TextureFormat == TSF_RGBA16F ? TextureCompressionSettings::TC_HDR : TextureCompressionSettings::TC_Default; //-V547 - future proofing
 	// Default to no mip generation for cube render target captures.
 	Texture2DArray->MipGenSettings = TextureMipGenSettings::TMGS_NoMipmaps;
 	Texture2DArray->PostEditChange();
diff --git a/Engine/Source/Runtime/Engine/Private/TextureRenderTargetVolume.cpp b/Engine/Source/Runtime/Engine/Private/TextureRenderTargetVolume.cpp
--- a/Engine/Source/Runtime/Engine/Private/TextureRenderTargetVolume.cpp
+++ b/Engine/Source/Runtime/Engine/Private/TextureRenderTargetVolume.cpp
@@ -205,7 +205,7 @@ UVolumeTexture* UTextureRenderTargetVolume::ConstructTextureVolume(UObject* ObjO
 	VolumeTexture->Source.UnlockMip(0);
 	VolumeTexture->SRGB = bSRGB;
 	// If HDR source image then choose HDR compression settings..
-	VolumeTexture->CompressionSettings = TextureFormat == TSF_RGBA16F ? TextureCompressionSettings::TC_HDR : TextureCompressionSettings::TC_Default;
+	VolumeTexture->CompressionSettings = TextureFormat == TSF_RGBA16F ? TextureCompressionSettings::TC_HDR : TextureCompressionSettings::TC_Default; //-V547 - future proofing
 	// Default to no mip generation for cube render target captures.
 	VolumeTexture->MipGenSettings = TextureMipGenSettings::TMGS_NoMipmaps;
 	VolumeTexture->PostEditChange();
