Avoid common case of VT DDC getting getting out of sync, which should hopefully resolve most corrupt VT problems

diff --git a/Engine/Source/Runtime/Engine/Private/TextureDerivedData.cpp b/Engine/Source/Runtime/Engine/Private/TextureDerivedData.cpp
--- a/Engine/Source/Runtime/Engine/Private/TextureDerivedData.cpp
+++ b/Engine/Source/Runtime/Engine/Private/TextureDerivedData.cpp
@@ -53,7 +53,7 @@
 
 // This GUID is mixed into DDC version for virtual textures only, this allows updating DDC version for VT without invalidating DDC for all textures
 // This is useful during development, but once large numbers of VT are present in shipped content, it will have the same problem as TEXTURE_DERIVEDDATA_VER
-#define TEXTURE_VT_DERIVEDDATA_VER	TEXT("2FF3E5D7EFF5424D8631C71A65DB4BF8")
+#define TEXTURE_VT_DERIVEDDATA_VER	TEXT("744516E4C74F4867BDAC2E6D49596219")
 
 #if ENABLE_COOK_STATS
 namespace TextureCookStats
@@ -663,7 +663,7 @@ uint32 PutDerivedDataInCache(FTexturePlatformData* DerivedData, const FString& D
 	FMemoryWriter Ar(RawDerivedData, /*bIsPersistent=*/ true);
 	DerivedData->Serialize(Ar, NULL);
 	TotalBytesPut += RawDerivedData.Num();
-	GetDerivedDataCacheRef().Put(*DerivedDataKey, RawDerivedData, TextureName, /*bPutEvenIfExists*/ true);
+	GetDerivedDataCacheRef().Put(*DerivedDataKey, RawDerivedData, TextureName, bReplaceExistingDDC);
 	UE_LOG(LogTexture,Verbose,TEXT("%s  Derived Data: %d bytes"),*LogString,RawDerivedData.Num());
 	return TotalBytesPut;
 }
