Animation data MVC refactor

diff --git a/Engine/Source/Runtime/Engine/Classes/Animation/AnimData/AnimDataController.h b/Engine/Source/Runtime/Engine/Classes/Animation/AnimData/AnimDataController.h
new file mode 100644
--- /dev/null
+++ b/Engine/Source/Runtime/Engine/Classes/Animation/AnimData/AnimDataController.h
@@ -0,0 +1,483 @@
+// Copyright Epic Games, Inc. All Rights Reserved.
+
+#pragma once
+
+
+#include "CoreTypes.h"
+#include "UObject/ObjectMacros.h"
+#include "Animation/AnimData/AnimDataModel.h"
+#include "Animation/AnimCurveTypes.h"
+
+#if WITH_EDITOR
+#include "ChangeTransactor.h"
+#endif // WITH_EDITOR
+
+#include "AnimDataController.generated.h"
+
+#if WITH_EDITOR
+
+#endif // WITH_EDITOR
+
+struct FAnimationCurveIdentifier;
+static const int32 DefaultCurveFlags = EAnimAssetCurveFlags::AACF_Editable;
+
+namespace UE {
+namespace Anim {
+	class FOpenBracketAction;
+	class FCloseBracketAction;
+}}
+
+/**
+ * The Controller is the sole authority to perform changes on the Animation Data Model. Any mutation to the model made will
+ * cause a subsequent notify (EAnimDataModelNotifyType) to be broadcasted from the Model's ModifiedEvent. Alongside of it is a 
+ * payload containing information relevant to the mutation. These notifies should be relied upon to update any dependent views 
+ * or generated (derived) data.
+ */
+UCLASS(BlueprintType)
+class ENGINE_API UAnimDataController : public UObject
+{
+	GENERATED_BODY()
+
+public:
+	UAnimDataController() 
+#if WITH_EDITOR
+	: BracketDepth(0) 
+#endif // WITH_EDITOR
+	{}
+
+#if WITH_EDITOR
+	/** RAII helper to define a scoped-based bracket, opens and closes a controller bracket automatically */
+	struct FScopedBracket
+	{
+		FScopedBracket(UAnimDataController* InController, const FText& InDescription)
+			: Controller(InController)
+		{
+			Controller->OpenBracket(InDescription);
+		}
+
+		~FScopedBracket()
+		{
+			Controller->CloseBracket();
+		}
+
+		UAnimDataController* Controller;
+	};
+
+	/**
+	* Sets the AnimDataModel instance this controller is supposed to be targeting
+	*
+	* @param	InModel		UAnimDataModel instance to target
+	*/
+	UFUNCTION(BlueprintCallable, Category=AnimationData)
+	void SetModel(UAnimDataModel* InModel);
+
+	/**
+	* @return		The AnimDataModel instance this controller is currently targeting
+	*/
+	UFUNCTION(BlueprintPure, Category = AnimationData)
+	UAnimDataModel* GetModel() { return Model; }
+
+	/**
+	* @return		The AnimDataModel instance this controller is currently targeting
+	*/
+	const UAnimDataModel* const GetModel() const;
+
+	/**
+	* Opens an interaction bracket, used for combining a set of controller actions. Broadcasts a EAnimDataModelNotifyType::BracketOpened notify,
+	* this can be used by any Views or dependendent systems to halt any unnecessary or invalid operations until the (last) bracket is closed.
+	*
+	* @param	InTitle				Description of the bracket, e.g. "Generating Curve Data"
+	* @param	bShouldTransact		Whether or not any undo-redo changes should be generated
+	*/
+	UFUNCTION(BlueprintCallable, Category = AnimationData)
+	void OpenBracket(const FText& InTitle, bool bShouldTransact = true);
+
+	/**
+	* Closes a previously opened interaction bracket, used for combining a set of controller actions. Broadcasts a EAnimDataModelNotifyType::BracketClosed notify.
+	*
+	* @param	InTitle				Description of the bracket, e.g. "Generating Curve Data"
+	* @param	bShouldTransact		Whether or not any undo-redo changes should be generated
+	*/
+	UFUNCTION(BlueprintCallable, Category = AnimationData)
+	void CloseBracket(bool bShouldTransact = true);
+
+	/**
+	* Sets the total play-able length in seconds. Broadcasts a EAnimDataModelNotifyType::SequenceLengthChanged notify if successful.
+	* The number of frames and keys for the provided length is recalculated according to the current value of UAnimDataModel::FrameRate.
+	*
+	* @param	Length				New play-able length value, has to be positive and non-zero
+	* @param	bShouldTransact		Whether or not any undo-redo changes should be generated
+	*/
+	UFUNCTION(BlueprintCallable, Category = AnimationData)
+	void SetPlayLength(float Length, bool bShouldTransact = true);
+
+	/**
+	* Sets the total play-able length in seconds. Broadcasts a EAnimDataModelNotifyType::SequenceLengthChanged notify if successful.
+	* T0 and T1 are expected to represent the window of time that was either added or removed. E.g. for insertion T0 indicates the time
+	* at which additional time starts and T1 were it ends. For removal T0 indicates the time at which time should be started to remove, and T1 indicates the end. Giving a total of T1 - T0 added or removed length.
+	* The number of frames and keys for the provided length is recalculated according to the current value of UAnimDataModel::FrameRate.
+	*
+	* @param	Length				Total new play-able length value, has to be positive and non-zero
+	* @param	T0					Point between 0 and Length at which the change in time starts
+	* @param	T1					Point between 0 and Length at which the change in time ends
+	* @param	bShouldTransact		Whether or not any undo-redo changes should be generated
+	*/
+	UFUNCTION(BlueprintCallable, Category = AnimationData)
+	void Resize(float Length, float T0, float T1, bool bShouldTransact = true);
+
+	/**
+	* Sets the frame rate according to which the bone animation is expected to be sampled. Broadcasts a EAnimDataModelNotifyType::FrameRateChanged notify if successful.
+	* The number of frames and keys for the provided frame rate is recalculated according to the current value of UAnimDataModel::PlayLength.
+	*
+	* @param	FrameRate			The new sampling frame rate, has to be positive and non-zero
+	* @param	bShouldTransact		Whether or not any undo-redo changes should be generated
+	*/
+	UFUNCTION(BlueprintCallable, Category = AnimationData)
+	void SetFrameRate(FFrameRate FrameRate, bool bShouldTransact = true);
+
+	/**
+	* Adds a new bone animation track for the provided name. Broadcasts a EAnimDataModelNotifyType::TrackAdded notify if successful.
+	*
+	* @param	BoneName			Bone name for which a track should be added
+	* @param	bShouldTransact		Whether or not any undo-redo changes should be generated
+	*
+	* @return	The index at which the bone track was added, INDEX_NONE if adding it failed
+	*/
+	UFUNCTION(BlueprintCallable, Category = AnimationData)
+	int32 AddBoneTrack(FName BoneName, bool bShouldTransact = true);
+
+	/**
+	* Inserts a new bone animation track for the provided name, at the provided index. Broadcasts a EAnimDataModelNotifyType::TrackAdded notify if successful.
+	* The bone name is verified with the AnimModel's outer target USkeleton to ensure the bone exists.
+	*
+	* @param	BoneName			Bone name for which a track should be inserted
+	* @param	DesiredIndex		Index at which the track should be inserted
+	* @param	bShouldTransact		Whether or not any undo-redo changes should be generated
+	*
+	* @return	The index at which the bone track was inserted, INDEX_NONE if the insertion failed
+	*/
+	UFUNCTION(BlueprintCallable, Category = AnimationData)
+	int32 InsertBoneTrack(FName BoneName, int32 DesiredIndex, bool bShouldTransact = true);
+
+	/**
+	* Removes an existing bone animation track with the provided name. Broadcasts a EAnimDataModelNotifyType::TrackRemoved notify if successful.
+	*
+	* @param	BoneName			Bone name of the track which should be removed
+	* @param	bShouldTransact		Whether or not any undo-redo changes should be generated
+	*
+	* @return	Whether or not the removal was succesful
+	*/
+	UFUNCTION(BlueprintCallable, Category = AnimationData)
+	bool RemoveBoneTrack(FName BoneName, bool bShouldTransact = true);
+
+	/**
+	* Removes all existing Bone Animation tracks. Broadcasts a EAnimDataModelNotifyType::TrackRemoved for each removed track, wrapped within BracketOpened/BracketClosed notifies.
+	*
+	* @param	bShouldTransact		Whether or not any undo-redo changes should be generated
+	*/
+	UFUNCTION(BlueprintCallable, Category = AnimationData)
+	void RemoveAllBoneTracks(bool bShouldTransact = true);
+
+	/**
+	* Removes an existing bone animation track with the provided name. Broadcasts a EAnimDataModelNotifyType::TrackChanged notify if successful.
+	* The provided number of keys provided is expected to match for each component, and be non-zero.
+	*
+	* @param	BoneName			Bone name of the track for which the keys should be set
+	* @param	PositionalKeys		Array of keys for the translation component
+	* @param	RotationalKeys		Array of keys for the rotation component
+	* @param	ScalingKeys			Array of keys for the scale component
+	* @param	bShouldTransact		Whether or not any undo-redo changes should be generated
+	*
+	* @return	Whether or not the keys were succesfully set
+	*/
+	UFUNCTION(BlueprintCallable, Category = AnimationData)
+	bool SetBoneTrackKeys(FName BoneName, const TArray<FVector>& PositionalKeys, const TArray<FQuat>& RotationalKeys, const TArray<FVector>& ScalingKeys, bool bShouldTransact = true);
+	
+	/**
+	* Adds a new curve with the provided information. Broadcasts a EAnimDataModelNotifyType::CurveAdded notify if successful.
+	*
+	* @param	CurveId				Identifier for the to-be-added curve
+	* @param	CurveFlags			Flags to be set for the curve
+	* @param	bShouldTransact		Whether or not any undo-redo changes should be generated
+	*
+	* @return	Whether or not the curve was succesfully added
+	*/
+	UFUNCTION(BlueprintCallable, Category = CurveData)
+	bool AddCurve(const FAnimationCurveIdentifier& CurveId, int32 CurveFlags = 0x00000004, bool bShouldTransact = true);
+
+	/**
+	* Duplicated the curve with the identifier. Broadcasts a EAnimDataModelNotifyType::CurveAdded notify if successful.
+	*
+	* @param	CopyCurveId			Identifier for the to-be-duplicated curve
+	* @param	NewCurveId			Identifier for the to-be-added curve
+	* @param	bShouldTransact		Whether or not any undo-redo changes should be generated
+	*
+	* @return	Whether or not the curve was succesfully duplicated
+	*/
+	UFUNCTION(BlueprintCallable, Category = CurveData)
+	bool DuplicateCurve(const FAnimationCurveIdentifier& CopyCurveId, const FAnimationCurveIdentifier& NewCurveId, bool bShouldTransact = true);
+	
+
+	/**
+	* Remove the curve with provided identifier. Broadcasts a EAnimDataModelNotifyType::CurveRemoved notify if successful.
+	*
+	* @param	CopyCurveId			Identifier for the to-be-removed curve
+	* @param	bShouldTransact		Whether or not any undo-redo changes should be generated
+	*
+	* @return	Whether or not the curve was succesfully removed
+	*/
+	UFUNCTION(BlueprintCallable, Category = CurveData)
+	bool RemoveCurve(const FAnimationCurveIdentifier& CurveId, bool bShouldTransact = true);
+
+	/**
+	* Removes all the curves of the provided type. Broadcasts a EAnimDataModelNotifyType::CurveRemoved for each removed curve, wrapped within BracketOpened/BracketClosed notifies.
+	*
+	* @param	SupportedCurveType	Type for which all curves are to be removed
+	* @param	bShouldTransact		Whether or not any undo-redo changes should be generated
+	*/
+	UFUNCTION(BlueprintCallable, Category = CurveData)
+	void RemoveAllCurvesOfType(ERawCurveTrackTypes SupportedCurveType, bool bShouldTransact = true);
+
+	/**
+	* Set an individual flag for the curve with provided identifier. Broadcasts a EAnimDataModelNotifyType::CurveFlagsChanged notify if successful.
+	*
+	* @param	CurveId			    Identifier for the curve for which the flag state is to be set
+	* @param	Flag				Flag for which the state is supposed to be set
+	* @param	bState				State of the flag to be, true=set/false=not set
+	* @param	bShouldTransact		Whether or not any undo-redo changes should be generated
+	*
+	* @return	Whether or not the flag state was succesfully set
+	*/
+	UFUNCTION(BlueprintCallable, Category = CurveData)
+	bool SetCurveFlag(const FAnimationCurveIdentifier& CurveId, EAnimAssetCurveFlags Flag, bool bState = true, bool bShouldTransact = true);
+
+	/**
+	* Replace the flags for the curve with provided identifier. Broadcasts a EAnimDataModelNotifyType::CurveFlagsChanged notify if successful.
+	*
+	* @param	CurveId			    Identifier for the curve for which the flag state is to be set
+	* @param	Flags				Flag mask with which the existings flags are to be replaced
+	* @param	bShouldTransact		Whether or not any undo-redo changes should be generated
+	*
+	* @return	Whether or not the flag mask was succesfully set
+	*/
+	UFUNCTION(BlueprintCallable, Category = CurveData)
+	bool SetCurveFlags(const FAnimationCurveIdentifier& CurveId, int32 Flags, bool bShouldTransact = true);
+
+	/**
+	* Replace the keys for the transform curve with provided identifier. Broadcasts a EAnimDataModelNotifyType::CurveChanged notify if successful.
+	*
+	* @param	CurveId			    Identifier for the transform curve for which the keys are to be set
+	* @param	TransformValues		Transform Values with which the existing values are to be replaced
+	* @param	TimeKeys			Time Keys with which the existing keys are to be replaced
+	* @param	bShouldTransact		Whether or not any undo-redo changes should be generated
+	*
+	* @return	Whether or not the transform curve keys were succesfully set
+	*/
+	UFUNCTION(BlueprintCallable, Category = CurveData)
+	bool SetTransformCurveKeys(const FAnimationCurveIdentifier& CurveId, const TArray<FTransform>& TransformValues, const TArray<float>& TimeKeys, bool bShouldTransact = true);
+		
+	/**
+	* Sets a single key for the transform curve with provided identifier. Broadcasts a EAnimDataModelNotifyType::CurveChanged notify if successful.
+	* In case a key for any of the individual transform channel curves already exists the value is replaced.
+	*
+	* @param	CurveId			    Identifier for the transform curve for which the key is to be set
+	* @param	Time				Time of the key to be set
+	* @param	Value				Value of the key to be set
+	* @param	bShouldTransact		Whether or not any undo-redo changes should be generated
+	*
+	* @return	Whether or not the transform curve key was succesfully set
+	*/
+	UFUNCTION(BlueprintCallable, Category = CurveData)
+	bool SetTransformCurveKey(const FAnimationCurveIdentifier& CurveId, float Time, const FTransform& Value, bool bShouldTransact = true);
+
+	/**
+	* Removes a single key for the transform curve with provided identifier. Broadcasts a EAnimDataModelNotifyType::CurveChanged notify if successful.
+	*
+	* @param	CurveId			    Identifier for the transform curve for which the key is to be set
+	* @param	Time				Time of the key to be removed
+	* @param	bShouldTransact		Whether or not any undo-redo changes should be generated
+	*
+	* @return	Whether or not the transform curve key was succesfully removed
+	*/
+	UFUNCTION(BlueprintCallable, Category = CurveData)
+	bool RemoveTransformCurveKey(const FAnimationCurveIdentifier& CurveId, float Time, bool bShouldTransact = true);
+
+	/**
+	* Renames the curve with provided identifier. Broadcasts a EAnimDataModelNotifyType::CurveRenamed notify if successful.
+	*
+	* @param	CurveToRenameId		Identifier for the curve to be renamed
+	* @param	NewCurveId			Time of the key to be removed
+	* @param	bShouldTransact		Whether or not any undo-redo changes should be generated
+	*
+	* @return	Whether or not the curve was succesfully renamed
+	*/
+	UFUNCTION(BlueprintCallable, Category = CurveData)
+	bool RenameCurve(const FAnimationCurveIdentifier& CurveToRenameId, const FAnimationCurveIdentifier& NewCurveId, bool bShouldTransact = true);
+
+	/**
+	* Changes the color of the curve with provided identifier. Broadcasts a EAnimDataModelNotifyType::CurveRenamed notify if successful.
+	* Currently changing curve colors is only supported for float curves.
+	*
+	* @param	CurveId				Identifier of the curve to change the color for
+	* @param	Color				Color to which the curve is to be set
+	* @param	bShouldTransact		Whether or not any undo-redo changes should be generated
+	*
+	* @return	Whether or not the curve color was succesfully changed
+	*/
+	UFUNCTION(BlueprintCallable, Category = CurveData)
+	bool SetCurveColor(const FAnimationCurveIdentifier& CurveId, FLinearColor Color, bool bShouldTransact = true);
+	
+	/**
+	* Scales the curve with provided identifier. Broadcasts a EAnimDataModelNotifyType::CurveScaled notify if successful.
+	*
+	* @param	CurveId				Identifier of the curve to scale
+	* @param	Origin				Time to use as the origin when scaling the curve
+	* @param	Factor				Factor with which the curve is supposed to be scaled
+	* @param	bShouldTransact		Whether or not any undo-redo changes should be generated
+	*
+	* @return	Whether or not scaling the curve was succesful
+	*/
+	UFUNCTION(BlueprintCallable, Category = CurveData)
+	bool ScaleCurve(const FAnimationCurveIdentifier& CurveId, float Origin, float Factor, bool bShouldTransact = true);
+
+	/**
+	* Sets a single key for the curve with provided identifier and name. Broadcasts a EAnimDataModelNotifyType::CurveChanged notify if successful.
+	* In case a key for the provided key time already exists the key is replaced.
+	*
+	* @param	CurveId			    Identifier for the curve for which the key is to be set
+	* @param	Key					Key to be set
+	* @param	bShouldTransact		Whether or not any undo-redo changes should be generated
+	*
+	* @return	Whether or not the curve key was succesfully set
+	*/
+	UFUNCTION(BlueprintCallable, Category = CurveData)
+	bool SetCurveKey(const FAnimationCurveIdentifier& CurveId, const FRichCurveKey& Key, bool bShouldTransact = true);
+	
+	/**
+	* Remove a single key from the curve with provided identifier and name. Broadcasts a EAnimDataModelNotifyType::CurveChanged notify if successful.
+	* In case a key for the provided key time already exists the key is replaced.
+	*
+	* @param	CurveId			    Identifier for the curve for which the key is to be set
+	* @param	NewCurveId			Time of the key to be removed
+	* @param	bShouldTransact		Whether or not any undo-redo changes should be generated
+	*
+	* @return	Whether or not the curve key was succesfully removed
+	*/
+	UFUNCTION(BlueprintCallable, Category = CurveData)
+	bool RemoveCurveKey(const FAnimationCurveIdentifier& CurveId, float Time, bool bShouldTransact = true);
+
+	/**
+	* Replace the keys for the curve with provided identifier and name. Broadcasts a EAnimDataModelNotifyType::CurveChanged notify if successful.
+	*
+	* @param	CurveId			    Identifier for the curve for which the key is to be set
+	* @param	CurveKeys			Keys with which the existing keys are to be replaced
+	* @param	bShouldTransact		Whether or not any undo-redo changes should be generated
+	*
+	* @return	Whether or not the curve key was succesfully set
+	*/
+	UFUNCTION(BlueprintCallable, Category = CurveData)
+	bool SetCurveKeys(const FAnimationCurveIdentifier& CurveId, const TArray<FRichCurveKey>& CurveKeys, bool bShouldTransact = true);
+
+
+	/**
+	* Updates the display name values for any stored curve, with the names being retrieved from the provided skeleton. Broadcasts a EAnimDataModelNotifyType::CurveRenamed for each to-be-updated curve name, wrapped within BracketOpened/BracketClosed notifies.
+	*
+	* @param	Skeleton			Skeleton to retrieve the display name values from
+	* @param	SupportedCurveType	Curve type for which the names should be updated
+	* @param	bShouldTransact		Whether or not any undo-redo changes should be generated
+	*/
+	UFUNCTION(BlueprintCallable, Category = CurveData)
+	void UpdateCurveNamesFromSkeleton(const USkeleton* Skeleton, ERawCurveTrackTypes SupportedCurveType, bool bShouldTransact = true);
+
+	/**
+	* Updates the curve names with the provided skeleton, if a display name is not found it will be added thus modifying the skeleton. Broadcasts a EAnimDataModelNotifyType::CurveRenamed for each curve name for which the UID was different or if it was added as a new smartname, wrapped within BracketOpened/BracketClosed notifies.
+	*
+	* @param	Skeleton			Skeleton to retrieve the display name values from
+	* @param	SupportedCurveType	Curve type for which the names should be updated
+	* @param	bShouldTransact		Whether or not any undo-redo changes should be generated
+	*/
+	UFUNCTION(BlueprintCallable, Category = CurveData)
+	void FindOrAddCurveNamesOnSkeleton(USkeleton* Skeleton, ERawCurveTrackTypes SupportedCurveType, bool bShouldTransact = true);
+
+	/**
+	* Removes any bone track for which the name was not found in the provided skeleton. Broadcasts a EAnimDataModelNotifyType::TrackRemoved for each track which was not found in the skeleton, wrapped within BracketOpened/BracketClosed notifies.
+	*
+	* @param	Skeleton			Skeleton to retrieve the display name values from
+	* @param	bShouldTransact		Whether or not any undo-redo changes should be generated
+	*/
+	bool RemoveBoneTracksMissingFromSkeleton(const USkeleton* Skeleton, bool bShouldTransact = true);
+
+	/**
+	* Broadcast a EAnimDataModelNotifyType::Populated notify.
+	*/
+	void NotifyPopulated();	
+
+	/**
+	* Resets all data stored in the model, broadcasts a EAnimDataModelNotifyType::Reset and wraps all actions within BracketOpened/BracketClosed notifies.
+	*	- Bone tracks, broadcasts a EAnimDataModelNotifyType::TrackRemoved for each
+	*	- Curves, broadcasts a EAnimDataModelNotifyType::CurveRemoves for each
+	*	- Play length to one frame at 30fps, broadcasts a EAnimDataModelNotifyType::PlayLengthChanged
+	*	- Frame rate to 30fps, broadcasts a EAnimDataModelNotifyType::FrameRateChanged
+	*
+	* @param	bShouldTransact		Whether or not any undo-redo changes should be generated
+	*/
+	void ResetModel(bool bShouldTransact = true);
+protected:
+	/** Functionality used by FOpenBracketAction and FCloseBracketAction to broadcast their equivalent notifies without actually opening a bracket. */
+	void NotifyBracketOpen();
+	void NotifyBracketClosed();
+
+private:
+	/** Returns whether or not the supplied curve type is supported by the controller functionality */
+	const bool IsSupportedCurveType(ERawCurveTrackTypes CurveType) const;
+	/** Returns the string representation of the provided curve enum type value */
+	FString GetCurveTypeValueName(ERawCurveTrackTypes InType) const;
+	
+	/** Resizes the curve data stored on the model according to the provided new length and time at which to insert or remove time */
+	void ResizeCurves(float NewLength, float OldLength, float T0, float T1);
+
+	/** Helper functionality to apply any changes in the model's playlength*/
+	void SetPlayLength_Internal(float NewLength, float T0, float T1, bool bShouldTransact);
+	   	
+	/** Ensures that a valid model is currently targeted */
+	void ValidateModel() const;
+
+	/** Verifies whether or not the Model's outer object is (or is derived from) the specified UClass */
+	bool CheckOuterClass(UClass* InClass) const;
+
+	/** Helper functionality to output script-based warnings and errors */
+	void ReportWarning(const FText& InMessage) const;
+	void ReportError(const FText& InMessage) const;
+
+	template <typename FmtType, typename... Types>
+	void ReportWarningf(const FmtType& Fmt, Types... Args) const
+	{	
+		ReportWarning(FText::Format(Fmt, Args...));
+	}
+
+	template <typename FmtType, typename... Types>
+	void ReportErrorf(const FmtType& Fmt, Types... Args) const
+	{
+		ReportError(FText::Format(Fmt, Args...));
+	}
+#endif // WITH_EDITOR
+
+private: 
+#if WITH_EDITOR
+	/** Current depth of outstanding brackets, assumed to remain positive */
+	int32 BracketDepth;
+
+	/** Transactor used to store Change based undo/redo action in the engines transaction buffer */
+	UE::FChangeTransactor ChangeTransactor;	
+#endif // WITH_EDITOR
+
+#if WITH_EDITORONLY_DATA
+	/** Current UAnimDataModel instance targeted by this controller */
+	UPROPERTY(transient)
+	UAnimDataModel* Model;
+#endif // WITH_EDITORONLY_DATA
+
+	friend class FAnimDataControllerTestBase;
+	friend UE::Anim::FOpenBracketAction;
+	friend UE::Anim::FCloseBracketAction;
+};
diff --git a/Engine/Source/Runtime/Engine/Classes/Animation/AnimData/AnimDataModel.h b/Engine/Source/Runtime/Engine/Classes/Animation/AnimData/AnimDataModel.h
new file mode 100644
--- /dev/null
+++ b/Engine/Source/Runtime/Engine/Classes/Animation/AnimData/AnimDataModel.h
@@ -0,0 +1,354 @@
+// Copyright Epic Games, Inc. All Rights Reserved.
+
+#pragma once
+
+#include "CoreMinimal.h"
+#include "UObject/ObjectMacros.h"
+#include "AnimDataNotifications.h"
+
+#include "AnimDataModel.generated.h"
+
+/**
+ * Structure encapsulating a single bone animation track.
+ */
+USTRUCT(BlueprintType)
+struct ENGINE_API FBoneAnimationTrack
+{
+	GENERATED_BODY()
+
+	/** Internally stored data representing the animation bone data */
+	UPROPERTY(VisibleAnywhere, BlueprintReadOnly, Category = "Model")
+	FRawAnimSequenceTrack InternalTrackData;
+
+	/** Index corresponding to the bone this track corresponds to within the target USkeleton */
+	UPROPERTY(VisibleAnywhere, BlueprintReadOnly, Category = "Model")
+	int32 BoneTreeIndex;
+
+	/** Name of the bone this track corresponds to */
+	UPROPERTY(VisibleAnywhere, BlueprintReadOnly, Category = "Model")
+	FName Name;
+};
+
+/**
+ * Structure encapsulating animated curve data. Currently only contains Float and Transform curves.
+ */
+USTRUCT(BlueprintType)
+struct ENGINE_API FAnimationCurveData
+{
+	GENERATED_BODY()
+
+	/** Float-based animation curves */
+	UPROPERTY(VisibleAnywhere, BlueprintReadOnly, Category = "Model")
+	TArray<FFloatCurve>	FloatCurves;
+
+	/** FTransform-based animation curves, used for animation layer editing */
+	UPROPERTY(VisibleAnywhere, BlueprintReadOnly, Category = "Model")
+	TArray<FTransformCurve>	TransformCurves;
+};
+
+struct FTrackToSkeletonMap;
+struct FAnimationCurveIdentifier;
+
+/**
+ * The Model represents the source data for animations. It contains both bone animation data as well as animated curves.
+ * They are currently only a sub-object of a AnimSequenceBase instance. The instance derives all runtime data from the source data. 
+ */
+UCLASS(BlueprintType)
+class ENGINE_API UAnimDataModel : public UObject
+{
+	GENERATED_BODY()
+public:
+
+	/** Begin UAnimDataModel overrides */
+	virtual void PostLoad() override;
+	virtual void PostDuplicate(bool bDuplicateForPIE) override;
+	/** End UAnimDataModel overrides */
+
+	/**
+	* @return	Total length of play-able animation data 
+	*/
+	UFUNCTION(BlueprintPure, Category=AnimationDataModel)
+	float GetPlayLength() const;
+	
+	/**
+	* @return	Total number of frames of animation data stored 
+	*/
+	UFUNCTION(BlueprintPure, Category = AnimationDataModel)
+	int32 GetNumberOfFrames() const;
+
+	/**
+	* @return	Total number of animation data keys stored 
+	*/
+	UFUNCTION(BlueprintPure, Category = AnimationDataModel)
+	int32 GetNumberOfKeys() const;
+
+	/**
+	* @return	Frame rate at which the animation data is key-ed 
+	*/
+	UFUNCTION(BlueprintPure, Category = AnimationDataModel)
+	const FFrameRate& GetFrameRate() const;
+	
+	/**
+	* @return	Array containg all bone animation tracks 
+	*/
+	UFUNCTION(BlueprintPure, Category = AnimationDataModel)
+	const TArray<FBoneAnimationTrack>& GetBoneAnimationTracks() const;
+	
+	/**
+	* @return	Bone animation track for the provided index
+	*/
+	UFUNCTION(BlueprintPure, Category = AnimationDataModel)
+	const FBoneAnimationTrack& GetBoneTrackByIndex(int32 TrackIndex) const;
+
+	/**
+	* @return	Bone animation track for the provided (bone) name
+	*/
+	UFUNCTION(BlueprintPure, Category = AnimationDataModel)
+	const FBoneAnimationTrack& GetBoneTrackByName(FName TrackName) const;
+	
+	/**
+	* @return	Bone animation track for the provided (bone) name if found, otherwise returns a nullptr 
+	*/
+	const FBoneAnimationTrack* FindBoneTrackByName(FName Name) const;
+
+	/**
+	* @return	Bone animation track for the provided index if valid, otherwise returns a nullptr 
+	*/
+	const FBoneAnimationTrack* FindBoneTrackByIndex(int32 BoneIndex) const;
+
+	/**
+	* @return	Internal track index for the provided bone animation track if found, otherwise returns INDEX_NONE 
+	*/
+	UFUNCTION(BlueprintPure, Category = AnimationDataModel)
+	int32 GetBoneTrackIndex(const FBoneAnimationTrack& Track) const;
+
+	/**
+	* @return	Internal track index for the provided (bone) name if found, otherwise returns INDEX_NONE 
+	*/
+	UFUNCTION(BlueprintPure, Category = AnimationDataModel)
+	int32 GetBoneTrackIndexByName(FName TrackName) const;
+
+	/**
+	* @return	Whether or not the provided track index is valid 
+	*/
+	UFUNCTION(BlueprintPure, Category = AnimationDataModel)
+	bool IsValidBoneTrackIndex(int32 TrackIndex) const;
+
+	/**
+	* @return	Total number of bone animation tracks
+	*/
+	UFUNCTION(BlueprintPure, Category = AnimationDataModel)
+	const int32 GetNumBoneTracks() const;
+
+	/**
+	* Populates the provided array with all contained (bone) track names
+	*
+	* @param	OutNames	[out] Array containing all bone track names
+	*/
+	UFUNCTION(BlueprintPure, Category = AnimationDataModel)
+	void GetBoneTrackNames(TArray<FName>& OutNames) const;
+
+	/** Returns all contained curve animation data */
+	const FAnimationCurveData& GetCurveData() const;
+
+	/**
+	* @return	Total number of stored FTransform curves
+	*/
+	UFUNCTION(BlueprintPure, Category = AnimationDataModel)
+	int32 GetNumberOfTransformCurves() const;
+
+	/**
+	* @return	Total number of stored float curves
+	*/
+	UFUNCTION(BlueprintPure, Category = AnimationDataModel)
+	int32 GetNumberOfFloatCurves() const;
+
+	/**
+	* @return	Array containing all stored float curves 
+	*/
+	const TArray<struct FFloatCurve>& GetFloatCurves() const;
+
+	/**
+	* @return	Array containing all stored FTransform curves 
+	*/
+	const TArray<struct FTransformCurve>& GetTransformCurves() const;
+
+	/**
+	* @return	Curve ptr for the provided identifier if valid, otherwise returns a nullptr 
+	*/
+	const FAnimCurveBase* FindCurve(const FAnimationCurveIdentifier& CurveIdentifier) const;
+
+	/**
+	* @return	Float Curve ptr for the provided identifier if valid, otherwise returns a nullptr
+	*/
+	const FFloatCurve* FindFloatCurve(const FAnimationCurveIdentifier& CurveIdentifier) const;
+
+	/**
+	* @return	Transform Curve ptr for the provided identifier if valid, otherwise returns a nullptr
+	*/
+	const FTransformCurve* FindTransformCurve(const FAnimationCurveIdentifier& CurveIdentifier) const;
+
+	/**
+	* @return	Rich curve ptr for the provided identifier if valid, otherwise returns a nullptr
+	*/
+	const FRichCurve* FindRichCurve(const FAnimationCurveIdentifier& CurveIdentifier) const;
+
+	/**
+	* @return	Curve object for the provided identifier if valid
+	*/
+	const FAnimCurveBase& GetCurve(const FAnimationCurveIdentifier& CurveIdentifier) const;
+
+	/**
+	* @return	Float Curve object for the provided identifier if valid
+	*/
+	const FFloatCurve& GetFloatCurve(const FAnimationCurveIdentifier& CurveIdentifier) const;
+
+	/**
+	* @return	Transform Curve object for the provided identifier if valid
+	*/
+	const FTransformCurve& GetTransformCurve(const FAnimationCurveIdentifier& CurveIdentifier) const;
+
+	/**
+	* @return	Rich Curve object for the provided identifier if valid
+	*/
+	const FRichCurve& GetRichCurve(const FAnimationCurveIdentifier& CurveIdentifier) const;
+		
+	/**
+	* @return	The outer UAnimSequence object if found, otherwise returns a nullptr 
+	*/
+	UFUNCTION(BlueprintPure, Category = AnimationDataModel)
+	UAnimSequence* GetAnimationSequence() const;
+
+	/**
+	* @return	Multicast delegate which is broadcasted to propagated changes to any internal data, see FAnimDataModelModifiedEvent and EAnimDataModelNotifyType
+	*/
+	FAnimDataModelModifiedEvent& GetModifiedEvent() { return ModifiedEvent; }
+
+	/**
+	* @return	GUID representing the contained data and state 
+	*/
+	FGuid GenerateGuid() const;
+	
+	/** Backwards compatibility functionality */
+	UE_DEPRECATED(5.0, "Use GetBoneAnimationTracks(), and FBoneAnimationTrack::InternalTrackData instead")
+	const TArray<FRawAnimSequenceTrack>& GetTransientRawAnimationTracks() const;
+	
+	UE_DEPRECATED(5.0, "Use GetBoneAnimationTracks(), and FBoneAnimationTrack::Name instead")
+	const TArray<FName>& GetTransientRawAnimationTrackNames() const;
+
+	UE_DEPRECATED(5.0, "Use GetBoneAnimationTracks(), and FBoneAnimationTrack::BoneTreeIndex instead")
+	const TArray<FTrackToSkeletonMap>& GetTransientRawAnimationTrackSkeletonMappings() const;
+
+	UE_DEPRECATED(5.0, "Non-const access to track data is prohibited, use UAnimDataController API to modify AnimDataModel instead")
+	FRawAnimSequenceTrack& GetNonConstRawAnimationTrackByIndex(int32 TrackIndex);
+
+	UE_DEPRECATED(5.0, "Use GetCurveData() instead")
+	const FRawCurveTracks& GetTransientRawCurveTracks() const;
+
+	UE_DEPRECATED(5.0, "Use Controller for mutating curve data")
+	FAnimationCurveData& GetNonConstCurveData();	
+
+private:
+	/** Helper functionality used by UAnimDataController to retrieve mutable data */ 
+	FBoneAnimationTrack* FindMutableBoneTrackByName(FName Name);
+	FBoneAnimationTrack& GetMutableBoneTrackByName(FName Name);
+	FTransformCurve* FindMutableTransformCurveById(const FAnimationCurveIdentifier& CurveIdentifier);
+	FFloatCurve* FindMutableFloatCurveById(const FAnimationCurveIdentifier& CurveIdentifier);
+	FAnimCurveBase* FindMutableCurveById(const FAnimationCurveIdentifier& CurveIdentifier);	   
+	FRichCurve* GetMutableRichCurve(const FAnimationCurveIdentifier& CurveIdentifier);
+
+	/**
+	* Broadcasts a new EAnimDataModelNotifyType with the provided payload data alongside it.
+	*
+	* @param	NotifyType			Type of notify to broadcast
+	* @param	PayloadData			Typed payload data
+	*/
+	template<typename T>
+	void Notify(EAnimDataModelNotifyType NotifyType, const T& PayloadData)
+	{
+		UScriptStruct* TypeScriptStruct = T::StaticStruct();
+
+		const FAnimDataModelNotifPayload Payload((int8*)&PayloadData, TypeScriptStruct);
+		ModifiedEvent.Broadcast(NotifyType, this, Payload);
+
+		if (ModifiedEventDynamic.IsBound())
+		{
+			ModifiedEventDynamic.Broadcast(NotifyType, this, Payload);
+		}
+
+		GenerateTransientData();
+	}
+
+	/**
+	* Broadcasts a new EAnimDataModelNotifyType alongside of an empty payload.
+	*
+	* @param	NotifyType			Type of notify to broadcast
+	*/
+	void Notify(EAnimDataModelNotifyType NotifyType)
+	{
+		FEmptyPayload EmptyPayload;
+		const FAnimDataModelNotifPayload Payload((int8*)&EmptyPayload, FEmptyPayload::StaticStruct());
+
+		ModifiedEvent.Broadcast(NotifyType, this, Payload);
+
+		if (ModifiedEventDynamic.IsBound())
+		{
+			ModifiedEventDynamic.Broadcast(NotifyType, this, Payload);
+		}
+
+		GenerateTransientData();
+	}
+
+private:
+	void GenerateTransientData();
+
+private:
+	/** Dynamic delegate event allows scripting to register to any broadcasted notify. */
+	UPROPERTY(BlueprintAssignable, Transient, Category = AnimationDataModel, meta = (ScriptName = "ModifiedEvent", AllowPrivateAccess = "true"))
+	FAnimDataModelModifiedDynamicEvent ModifiedEventDynamic;
+	
+	/** Native delegate event allows for registerings to any broadcasted notify. */
+	FAnimDataModelModifiedEvent ModifiedEvent;
+
+	/** All individual bone animation tracks */
+	UPROPERTY(VisibleAnywhere, BlueprintReadOnly, Category="Animation Data Model", meta = (AllowPrivateAccess = "true"))
+	TArray<FBoneAnimationTrack> BoneAnimationTracks;
+
+	/** Total playable length of the contained animation data */
+	UPROPERTY(VisibleAnywhere, BlueprintReadOnly, Category = "Animation Data Model", meta = (AllowPrivateAccess = "true"))
+	float PlayLength;
+	
+	/** Rate at which the animated data is sampled */
+	UPROPERTY(VisibleAnywhere, BlueprintReadOnly, Category = "Animation Data Model", meta = (AllowPrivateAccess = "true"))
+	FFrameRate FrameRate;
+
+	/** Total number of sampled animated frames */
+	UPROPERTY(VisibleAnywhere, BlueprintReadOnly, Category = "Animation Data Model", meta = (AllowPrivateAccess = "true"))
+	int32 NumberOfFrames;
+
+	/** Total number of sampled animated keys */
+	UPROPERTY(VisibleAnywhere, BlueprintReadOnly, Category = "Animation Data Model", meta = (AllowPrivateAccess = "true"))
+	int32 NumberOfKeys;
+	
+	/** Container with all animated curve data */
+	UPROPERTY(VisibleAnywhere, BlueprintReadOnly, Category = "Animation Data Model", meta = (AllowPrivateAccess = "true"))
+	FAnimationCurveData CurveData;
+	
+private:
+	/** Transient data, kept around for backward-compatibility */
+	UPROPERTY(Transient, VisibleAnywhere, Category = "Animation Data Model")
+	TArray<FRawAnimSequenceTrack> RawAnimationTracks;
+
+	UPROPERTY(Transient, VisibleAnywhere, Category = "Animation Data Model")
+	TArray<FName> RawAnimationTrackNames;
+
+	UPROPERTY(Transient, VisibleAnywhere, Category = "Animation Data Model")
+	TArray<FTrackToSkeletonMap> RawAnimationTrackSkeletonMappings;
+
+	UPROPERTY(Transient, VisibleAnywhere, Category = "Animation Data Model")
+	FRawCurveTracks RawCurveTracks;
+
+	friend class UAnimDataController;
+	friend class FAnimDataControllerTestBase;
+};
+
diff --git a/Engine/Source/Runtime/Engine/Classes/Animation/AnimData/AnimDataNotifyCollector.h b/Engine/Source/Runtime/Engine/Classes/Animation/AnimData/AnimDataNotifyCollector.h
new file mode 100644
--- /dev/null
+++ b/Engine/Source/Runtime/Engine/Classes/Animation/AnimData/AnimDataNotifyCollector.h
@@ -0,0 +1,81 @@
+// Copyright Epic Games, Inc. All Rights Reserved.
+
+#pragma once
+
+#include "Animation/AnimData/AnimDataNotifications.h"
+#include "Containers/Set.h"
+#include "Containers/Array.h"
+
+enum class EAnimDataModelNotifType : uint8;
+
+namespace UE {
+namespace Anim {
+
+#if WITH_EDITOR
+
+/** Helper structure for keeping track of which notifies of type EAnimDataModelNotifType are broadcasted
+between top-level EAnimDataModelNotifType::BracketOpened and EAnimDataModelNotifType::BracketClosed notifies */
+struct ENGINE_API FAnimationDataNotifyCollector
+{
+	FAnimationDataNotifyCollector() : BracketDepth(0) {}
+
+	/** Handle a broadcasted notify, reset if we are opening a new top-level bracket*/
+	void Handle(EAnimDataModelNotifType NotifyType)
+	{
+		NotifyTypes.Add(NotifyType);
+
+		if (NotifyType == EAnimDataModelNotifType::BracketOpened)
+		{
+			if (BracketDepth == 0)
+			{
+				Reset();
+			}
+
+			++BracketDepth;
+		}
+		else if (NotifyType == EAnimDataModelNotifType::BracketClosed)
+		{
+			--BracketDepth;
+		}
+	}
+	
+	/** Returns whether or not the notify of the provided types was broadcasted */
+	bool Contains(EAnimDataModelNotifType NotifyType) const
+	{
+		return NotifyTypes.Find(NotifyType) != nullptr;
+	}
+
+	/** Returns whether or not any of the provided notify types were broadcasted */
+	bool Contains(const TArray<EAnimDataModelNotifType>& TestNotifyTypes) const
+	{
+		for (EAnimDataModelNotifType Notify : TestNotifyTypes)
+		{
+			if (NotifyTypes.Find(Notify) != nullptr)
+			{
+				return true;
+			}
+		}
+
+		return false;
+	}
+
+	/** Returns whether or not a bracket is still open */
+	bool IsWithinBracket() const { return BracketDepth > 0; }
+
+	/** Returns whether or not all brackets have been closed */
+	bool IsNotWithinBracket() const { return BracketDepth == 0; }
+protected:
+	void Reset()
+	{
+		NotifyTypes.Empty();
+	}
+protected:
+	TSet<EAnimDataModelNotifType> NotifyTypes;
+	int32 BracketDepth;
+};
+
+#endif // WITH_EDITOR
+
+} // namespace Anim	
+
+} // namespace UE
\ No newline at end of file
diff --git a/Engine/Source/Runtime/Engine/Classes/Animation/AnimData/CurveIdentifier.h b/Engine/Source/Runtime/Engine/Classes/Animation/AnimData/CurveIdentifier.h
new file mode 100644
--- /dev/null
+++ b/Engine/Source/Runtime/Engine/Classes/Animation/AnimData/CurveIdentifier.h
@@ -0,0 +1,172 @@
+// Copyright Epic Games, Inc. All Rights Reserved.
+
+#pragma once
+
+#include "CoreTypes.h"
+#include "UObject/ObjectMacros.h"
+#include "Kismet/BlueprintFunctionLibrary.h"
+#include "Animation/SmartName.h"
+#include "Animation/AnimCurveTypes.h"
+
+#include "CurveIdentifier.generated.h"
+
+/** Enum used to determine a component channel of a transform curve */
+UENUM(BlueprintType)
+enum class ETransformCurveChannel : uint8
+{
+	Position,
+	Rotation,
+	Scale,
+	Invalid
+};
+
+/** Enum used to determine an axis channel of a vector curve */
+UENUM(BlueprintType)
+enum class EVectorCurveChannel : uint8
+{
+	X,
+	Y,
+	Z,
+	Invalid
+};
+
+/** Script-friendly structure for identifying an animation curve.
+* Wrapping the unique type and smart-name for use within the AnimDataController API. */
+USTRUCT(BlueprintType)
+struct ENGINE_API FAnimationCurveIdentifier
+{
+	GENERATED_USTRUCT_BODY();
+
+	FAnimationCurveIdentifier() : CurveType(ERawCurveTrackTypes::RCT_MAX), Channel(ETransformCurveChannel::Invalid), Axis(EVectorCurveChannel::Invalid) {}
+	FAnimationCurveIdentifier(const FSmartName& InSmartName, ERawCurveTrackTypes InCurveType) : InternalName(InSmartName), CurveType(InCurveType), Channel(ETransformCurveChannel::Invalid), Axis(EVectorCurveChannel::Invalid) {}
+	
+	explicit FAnimationCurveIdentifier(const SmartName::UID_Type& InUID, ERawCurveTrackTypes InCurveType) : CurveType(InCurveType)
+	{
+		InternalName.UID = InUID;
+	}
+
+	bool operator==(const FAnimationCurveIdentifier& Other) const
+	{
+		return (InternalName == Other.InternalName && CurveType == Other.CurveType 
+			&& Channel == Other.Channel && Axis == Other.Axis);
+	}
+	bool operator!=(const FAnimationCurveIdentifier& Other) const
+	{
+		return !(*this == Other);
+	}
+
+	bool IsValid() const
+	{
+		return InternalName.IsValid() && CurveType != ERawCurveTrackTypes::RCT_MAX;
+	}
+
+	FAnimationCurveIdentifier& operator=(const FAnimationCurveIdentifier& Other)
+	{
+		InternalName = Other.InternalName;
+		CurveType = Other.CurveType;
+		Channel = Other.Channel;
+		Axis = Other.Axis;
+
+		return *this;
+	}
+
+	FSmartName InternalName;
+	ERawCurveTrackTypes CurveType;
+	ETransformCurveChannel Channel;
+	EVectorCurveChannel Axis;
+};
+
+template<>
+struct TStructOpsTypeTraits<FAnimationCurveIdentifier> : public TStructOpsTypeTraitsBase2<FAnimationCurveIdentifier>
+{
+	enum
+	{
+		WithCopy = true,
+		WithIdenticalViaEquality = true
+	};
+};
+
+/** Script-exposed functionality for wrapping native functionality and constructing valid FAnimationCurveIdentifier instances */
+UCLASS()
+class ENGINE_API UAnimationCurveIdentifierExtensions : public UBlueprintFunctionLibrary
+{
+	GENERATED_BODY()
+public:
+	/**
+	* @return	Whether or not the Curve Identifier is valid
+	*/
+	UFUNCTION(BlueprintPure, Category = Curve, meta = (ScriptMethod))
+	static bool IsValid(UPARAM(ref) FAnimationCurveIdentifier& Identifier)
+	{
+		return Identifier.IsValid();
+	}
+
+	/**
+	* @return	The name used for displaying the Curve Identifier
+	*/
+	UFUNCTION(BlueprintPure, Category = Curve, meta = (ScriptMethod))
+	static FName GetName(UPARAM(ref) FAnimationCurveIdentifier& Identifier)
+	{
+		return Identifier.InternalName.DisplayName;
+	}
+
+	/**
+	* @return	The animation curve type for the curve represented by the Curve Identifier
+	*/
+	UFUNCTION(BlueprintPure, Category = Curve, meta = (ScriptMethod))
+	static ERawCurveTrackTypes GetType(UPARAM(ref) FAnimationCurveIdentifier& Identifier)
+	{
+		return Identifier.CurveType;
+	}
+
+#if WITH_EDITOR
+	/**
+	* Constructs a valid FAnimationCurveIdentifier instance. Ensuring that the underlying SmartName exists on the provided Skeleton for the provided curve type.
+	* If it is not found initially it will add it to the Skeleton thus modifying it.
+	*
+	* @param	InSkeleton			Skeleton on which to look for or add the curve name
+	* @param	Name				Name of the curve to find or add
+	* @param	CurveType			Type of the curve to find or add
+	*
+	* @return	Valid FAnimationCurveIdentifier if the operation was successful
+	*/
+	UFUNCTION(BlueprintCallable, Category = Curves, meta = (ScriptMethod))
+	static FAnimationCurveIdentifier GetCurveIdentifier(USkeleton* InSkeleton, FName Name, ERawCurveTrackTypes CurveType);
+
+	/**
+	* Tries to construct a valid FAnimationCurveIdentifier instance. It tries to find the underlying SmartName on the provided Skeleton for the provided curve type.
+	*
+	* @param	InSkeleton			Skeleton on which to look for the curve name
+	* @param	Name				Name of the curve to find
+	* @param	CurveType			Type of the curve to find
+	*
+	* @return	Valid FAnimationCurveIdentifier if the name exists on the skeleton and the operation was successful, invalid otherwise
+	*/
+	UFUNCTION(BlueprintCallable, Category = Curves, meta = (ScriptMethod))
+	static FAnimationCurveIdentifier FindCurveIdentifier(const USkeleton* InSkeleton, FName Name, ERawCurveTrackTypes CurveType);
+
+	/**
+	* Retrieves all curve identifiers for a specific curve types from the provided Skeleton
+	*
+	* @param	InSkeleton			Skeleton from which to retrieve the curve identifiers
+	* @param	CurveType			Type of the curve identifers to retrieve
+	*
+	* @return	Array of FAnimationCurveIdentifier instances each representing a unique curve if the operation was successful, empyty array otherwise
+	*/
+	UFUNCTION(BlueprintCallable, Category = Curves, meta = (ScriptMethod))
+	static TArray<FAnimationCurveIdentifier> GetCurveIdentifiers(USkeleton* InSkeleton, ERawCurveTrackTypes CurveType);
+
+	/**
+	* Converts a valid FAnimationCurveIdentifier instance with RCT_Transform curve type to target a child curve.
+	*
+	* @param	InOutIdentifier		[out] Identifier to be converted
+	* @param	Channel				Channel to target
+	* @param	Axis				Axis within channel to target
+	*
+	* @return	Valid FAnimationCurveIdentifier if the operation was successful
+	*/
+	UFUNCTION(BlueprintCallable, Category = Curves, meta = (ScriptMethod))
+	static bool GetTransformChildCurveIdentifier(UPARAM(ref) FAnimationCurveIdentifier& InOutIdentifier, ETransformCurveChannel Channel, EVectorCurveChannel Axis);
+#endif // WITH_EDITOR
+};
+
diff --git a/Engine/Source/Runtime/Engine/Classes/Animation/AnimMontage.h b/Engine/Source/Runtime/Engine/Classes/Animation/AnimMontage.h
--- a/Engine/Source/Runtime/Engine/Classes/Animation/AnimMontage.h
+++ b/Engine/Source/Runtime/Engine/Classes/Animation/AnimMontage.h
@@ -918,7 +918,9 @@ public:
 
 	//~Begin UAnimCompositeBase Interface
 	virtual void InvalidateRecursiveAsset() override;
-	virtual bool ContainRecursive(TArray<UAnimCompositeBase*>& CurrentAccumulatedList) override;	//~End UAnimCompositeBase Interface
+	virtual bool ContainRecursive(TArray<UAnimCompositeBase*>& CurrentAccumulatedList) override;
+	virtual void SetCompositeLength(float InLength) override;
+	//~End UAnimCompositeBase Interface
 
 	/** Utility function to create dynamic montage from AnimSequence */
 	ENGINE_API static UAnimMontage* CreateSlotAnimationAsDynamicMontage(UAnimSequenceBase* Asset, FName SlotNodeName, float BlendInTime = 0.25f, float BlendOutTime = 0.25f, float InPlayRate = 1.f, int32 LoopCount = 1, float BlendOutTriggerTime = -1.f, float InTimeToStartMontageAt = 0.f);
@@ -935,7 +937,8 @@ public:
 	FName TimeStretchCurveName;
 
 private:
+#if WITH_EDITOR
 	void BakeTimeStretchCurve();
 	//~End Time Stretch Curve
-
+#endif // WITH_EDITOR
 };
diff --git a/Engine/Source/Runtime/Engine/Classes/Animation/AnimSequence.h b/Engine/Source/Runtime/Engine/Classes/Animation/AnimSequence.h
--- a/Engine/Source/Runtime/Engine/Classes/Animation/AnimSequence.h
+++ b/Engine/Source/Runtime/Engine/Classes/Animation/AnimSequence.h
@@ -19,6 +19,12 @@
 #include "CustomAttributes.h"
 #include "Containers/ArrayView.h"
 #include "Animation/CustomAttributes.h"
+#include "Animation/AnimData/AnimDataNotifications.h"
+
+#if WITH_EDITOR
+#include "AnimData/AnimDataModel.h"
+#include "AnimData/AnimDataController.h"
+#endif // WITH_EDITOR
 
 #include "AnimSequence.generated.h"
 
@@ -209,37 +215,37 @@ class ENGINE_API UAnimSequence : public UAnimSequenceBase
 	/** The resample framerate that was computed during import. UI information only, unit are Hz */
 	UPROPERTY(AssetRegistrySearchable, meta = (DisplayName = "Import Resample Framerate"))
 	int32 ImportResampleFramerate;
-#endif
 
 protected:
-#if WITH_EDITORONLY_DATA
 	/** Contains the number of keys expected within the individual animation tracks. */
-	UE_DEPRECATED(5.0, "Num Frames is deprecated use NumberOfKeys instead")
+	UE_DEPRECATED(5.0, "NumFrames is deprecated see UAnimDataModel::GetNumberOfFrames for the number of source data frames, or GetNumberOfSampledKeys for the target keys")
 	UPROPERTY()
 	int32 NumFrames;
-#endif
 
 	/** The number of keys expected within the individual (non-uniform) animation tracks. */
-	UPROPERTY(AssetRegistrySearchable)
+	UE_DEPRECATED(5.0, "NumberOfKeys is deprecated see UAnimDataModel::GetNumberOfKeys for the number of source data keys, or GetNumberOfSampledKeys for the target keys")
+	UPROPERTY()
 	int32 NumberOfKeys;
 
 	/** The frame rate at which the source animation is sampled. */
-	UPROPERTY(AssetRegistrySearchable)
+	UE_DEPRECATED(5.0, "SamplingFrameRate is deprecated see UAnimDataModel::GetFrameRate for the source frame rate, or GetSamplingFrameRate for the target frame rate instead")
+	UPROPERTY()
 	FFrameRate SamplingFrameRate;
 
 	/**
 	 * In the future, maybe keeping RawAnimSequenceTrack + TrackMap as one would be good idea to avoid inconsistent array size
 	 * TrackToSkeletonMapTable(i) should contains  track mapping data for RawAnimationData(i). 
 	 */
+	UE_DEPRECATED(5.0, "TrackToSkeletonMapTable has been deprecated see FBoneAnimationTrack::BoneTreeIndex")
 	UPROPERTY()
 	TArray<struct FTrackToSkeletonMap> TrackToSkeletonMapTable;
 
 	/**
 	 * Raw uncompressed keyframe data. 
 	 */
+	UE_DEPRECATED(5.0, "RawAnimationData has been deprecated see FBoneAnimationTrack::InternalTrackData")
 	TArray<struct FRawAnimSequenceTrack> RawAnimationData;
 
-#if WITH_EDITORONLY_DATA
 	// Update this if the contents of RawAnimationData changes;
 	UPROPERTY()
 	FGuid RawDataGuid;
@@ -247,6 +253,7 @@ protected:
 	/**
 	 * This is name of RawAnimationData tracks for editoronly - if we lose skeleton, we'll need relink them
 	 */
+	UE_DEPRECATED(5.0, "Animation track names has been deprecated see FBoneAnimationTrack::Name")
 	UPROPERTY(VisibleAnywhere, Category="Animation")
 	TArray<FName> AnimationTrackNames;
 
@@ -254,11 +261,9 @@ protected:
 	 * Source RawAnimationData. Only can be overridden by when transform curves are added first time OR imported
 	 */
 	TArray<struct FRawAnimSequenceTrack> SourceRawAnimationData_DEPRECATED;
-#endif
 
 public:
 
-#if WITH_EDITORONLY_DATA
 	/**
 	 * Allow frame stripping to be performed on this animation if the platform requests it
 	 * Can be disabled if animation has high frequency movements that are being lost.
@@ -286,14 +291,15 @@ public:
 	FCompressedAnimSequence CompressedData;
 
 	// Accessors for animation frame count
-	UE_DEPRECATED(5.0, "GetRawNumberOfFrames has been deprecated, use GetNumberOfSampledKeys instead")
-	int32 GetRawNumberOfFrames() const { return NumberOfKeys; }
+	UE_DEPRECATED(5.0, "GetRawNumberOfFrames is deprecated see UAnimDataModel::GetNumberOfKeys for the source keys, or GetNumberOfSampledKeys for resampled keys")
+	int32 GetRawNumberOfFrames() const { return GetNumberOfSampledKeys(); }
 
-	UE_DEPRECATED(5.0, "SetRawNumberOfFrame has been deprecated, use SetNumberOfSampledKeys instead")
-	void SetRawNumberOfFrame(int32 InNumFrames) { SetNumberOfSampledKeys(InNumFrames); }
+	UE_DEPRECATED(5.0, "SetRawNumberOfFrame has been deprecated see UAnimDataController::SetFrameRate")
+	void SetRawNumberOfFrame(int32 InNumFrames) {}
 
 	/** Update the the number of expected keys in the (non-uniform) animation tracks, including T0 */
-	void SetNumberOfSampledKeys(int32 InNumberOfKeys);
+	UE_DEPRECATED(5.0, "SetNumberOfSampledKeys has been deprecated see UAnimDataController::SetFrameRate")
+	void SetNumberOfSampledKeys(int32 InNumberOfKeys) {}
 
 	/** Additive animation type. **/
 	UPROPERTY(EditAnywhere, Category=AdditiveSettings, AssetRegistrySearchable)
@@ -418,14 +424,15 @@ public:
 	virtual bool HasRootMotion() const override { return bEnableRootMotion; }
 	virtual void RefreshCacheData() override;
 	virtual EAdditiveAnimationType GetAdditiveAnimType() const override { return AdditiveAnimType; }
-	virtual int32 GetNumberOfSampledKeys() const override { return NumberOfKeys; }
-	virtual const FFrameRate& GetSamplingFrameRate() const override { return SamplingFrameRate; }
-	virtual void SetSequenceLength(float NewLength) override;
+	virtual int32 GetNumberOfSampledKeys() const override;
+	virtual const FFrameRate& GetSamplingFrameRate() const override { return TargetFrameRate; }
 	virtual void EvaluateCurveData(FBlendedCurve& OutCurve, float CurrentTime, bool bForceUseRawData = false) const override;
 	virtual float EvaluateCurveData(SmartName::UID_Type CurveUID, float CurrentTime, bool bForceUseRawData = false) const override;
 	virtual bool HasCurveData(SmartName::UID_Type CurveUID, bool bForceUseRawData) const override;
 
 #if WITH_EDITOR
+	UE_DEPRECATED(5.0, "MarkRawDataAsModified has been deprecated, any (Raw Data) modification should be applied using the UAnimDataController API instead. This will handle updating the GUID instead.")
+	PRAGMA_DISABLE_DEPRECATION_WARNINGS
 	virtual void MarkRawDataAsModified(bool bForceNewRawDatGuid = true) override
 	{
 		Super::MarkRawDataAsModified();
@@ -434,15 +441,13 @@ public:
 		FlagDependentAnimationsAsRawDataOnly();
 		UpdateDependentStreamingAnimations();
 	}
+	PRAGMA_ENABLE_DEPRECATION_WARNINGS
 #endif
 	//~ End UAnimSequenceBase Interface
 
 	// Returns the framerate of the animation
-	UE_DEPRECATED(5.0, "Use GetSamplingFrameRate instead")
-	float GetFrameRate() const 
-	{ 
-		return (float)(FMath::Max(NumberOfKeys - 1, 1)) / (GetPlayLength() > 0.f ? GetPlayLength() : 1.f); 
-	}
+	UE_DEPRECATED(5.0, "GetFrameRate is deprecated see UAnimDataModel::GetFrameRate for the source frame rate, or GetSamplingFrameRate for the target frame rate instead")
+	float GetFrameRate() const { return GetSamplingFrameRate().AsDecimal(); }
 
 	// Extract Root Motion transform from the animation
 	FTransform ExtractRootMotion(float StartTime, float DeltaTime, bool bAllowLooping) const;
@@ -476,33 +481,42 @@ public:
 	*/
 	void GetBonePose(struct FAnimationPoseData& OutAnimationPoseData, const FAnimExtractContext& ExtractionContext, bool bForceUseRawData = false) const;
 
-	const TArray<FRawAnimSequenceTrack>& GetRawAnimationData() const { return RawAnimationData; }
+	UE_DEPRECATED(5.0, "GetRawAnimationTrack has been deprecated see UAnimDataModel::GetBoneAnimationTracks")
+	const TArray<FRawAnimSequenceTrack>& GetRawAnimationData() const;
 
 #if WITH_EDITORONLY_DATA
 	UE_DEPRECATED(5.0, "SourceRawAnimationData has been deprecated")
 	bool HasSourceRawData() const { return false; }
-	const TArray<FName>& GetAnimationTrackNames() const { return AnimationTrackNames; }
+	UE_DEPRECATED(5.0, "GetAnimationTrackNames has been deprecated see FBoneAnimationTrack::Name")
+	const TArray<FName>& GetAnimationTrackNames() const;
 	
-	void  UpdateCompressedCurveName(SmartName::UID_Type CurveUID, const struct FSmartName& NewCurveName);
+	UE_DEPRECATED(5.0, "UpdateCompressedCurveName will be marked protected, updating compressed curve names is now handled by EAnimDataModelNotifyType::CurveRenamed")
+	void UpdateCompressedCurveName(SmartName::UID_Type CurveUID, const struct FSmartName& NewCurveName);
 	
 	// Adds a new track (if no track of the supplied name is found) to the raw animation data, optionally setting it to TrackData.
+	UE_DEPRECATED(5.0, "AddNewRawTrack has been deprecated see UAnimDataController::AddBoneTrack and UAnimDataController::SetBoneTrackKeys")
 	int32 AddNewRawTrack(FName TrackName, FRawAnimSequenceTrack* TrackData = nullptr);
 #endif
+	UE_DEPRECATED(5.0, "GetRawTrackToSkeletonMapTable has been deprecated see FBoneAnimationTrack::BoneTreeIndex")
+	const TArray<FTrackToSkeletonMap>& GetRawTrackToSkeletonMapTable() const;
 
-	const TArray<FTrackToSkeletonMap>& GetRawTrackToSkeletonMapTable() const { return TrackToSkeletonMapTable; }
 	const TArray<FTrackToSkeletonMap>& GetCompressedTrackToSkeletonMapTable() const { return CompressedData.CompressedTrackToSkeletonMapTable; }
 	const TArray<struct FSmartName>& GetCompressedCurveNames() const { return CompressedData.CompressedCurveNames; }
 
-	FRawAnimSequenceTrack& GetRawAnimationTrack(int32 TrackIndex) { return RawAnimationData[TrackIndex]; }
-	const FRawAnimSequenceTrack& GetRawAnimationTrack(int32 TrackIndex) const { return RawAnimationData[TrackIndex]; }
+	UE_DEPRECATED(5.0, "GetRawAnimationTrack, and non-const access to Raw Animation Data has been deprecated see UAnimDataModel::GetBoneTrackByIndex")
+	FRawAnimSequenceTrack& GetRawAnimationTrack(int32 TrackIndex);
+
+	UE_DEPRECATED(5.0, "GetRawAnimationTrack has been deprecated see UAnimDataModel::GetBoneTrackByIndex")
+	const FRawAnimSequenceTrack& GetRawAnimationTrack(int32 TrackIndex) const;
 
 private:
 #if WITH_EDITORONLY_DATA
 	void UpdateRetargetSourceAsset();
-#endif
 
 	/** Updates the stored sampling frame-rate using the sequence length and number of sampling keys */
+	UE_DEPRECATED(5.0, "UpdateFrameRate has been deprecated see UAnimDataController::SetFrameRate")
 	void UpdateFrameRate();
+#endif
 	
 	const TArray<FTransform>& GetRetargetTransforms() const;
 	FName GetRetargetTransformsSourceName() const;
@@ -568,6 +582,7 @@ public:
 	 * @param	TrackIndex			Index of track to interpolate.
 	 * @param	Time				Time on track to interpolate to.
 	 */
+	UE_DEPRECATED(5.0, "ExtractBoneTransform has been deprecated see FAnimationUtils::ExtractTransformFromTrack")
 	void ExtractBoneTransform(const TArray<struct FRawAnimSequenceTrack> & InRawAnimationData, FTransform& OutAtom, int32 TrackIndex, float Time) const;
 
 	/**
@@ -577,8 +592,10 @@ public:
 	* @param	OutAtom				[out] Output bone transform.
 	* @param	Time				Time on track to interpolate to.
 	*/
+	UE_DEPRECATED(5.0, "ExtractBoneTransform has been deprecated see FAnimationUtils::ExtractTransformFromTrack")
 	void ExtractBoneTransform(const struct FRawAnimSequenceTrack& InRawAnimationTrack, FTransform& OutAtom, float Time) const;
 
+	UE_DEPRECATED(5.0, "ExtractBoneTransform has been deprecated see FAnimSequenceHelpers::ExtractBoneTransform")
 	void ExtractBoneTransform(const struct FRawAnimSequenceTrack& RawTrack, FTransform& OutAtom, int32 KeyIndex) const;
 
 	// End Transform related functions 
@@ -587,11 +604,17 @@ public:
 
 	/** @return	estimate uncompressed raw size. This is *not* the real raw size. 
 				Here we estimate what it would be with no trivial compression. */
+#if !WITH_EDITOR
+	UE_DEPRECATED(5.0, "GetUncompressedRawSize will be marked EDITOR_ONLY")
+#endif // !WITH_EDITOR
 	int32 GetUncompressedRawSize() const;
 
 	/**
 	 * @return		The approximate size of raw animation data.
 	 */
+#if !WITH_EDITOR
+	UE_DEPRECATED(5.0, "GetApproxRawSize will be marked EDITOR_ONLY")
+#endif // !WITH_EDITOR
 	int32 GetApproxRawSize() const;
 
 	/**
@@ -610,7 +633,8 @@ public:
 	 * go down to a single frame, the time is stripped out as well.
 	 * @return true if keys were removed.
 	 */
-	bool CompressRawAnimData(float MaxPosDiff, float MaxAngleDiff);
+	UE_DEPRECATED(5.0, "CompressRawAnimData has been deprecated, reduction of Raw Animation data now happens during compression instead")
+	bool CompressRawAnimData(float MaxPosDiff, float MaxAngleDiff) { return false; }
 
 	/**
 	 * Removes trivial frames -- frames of tracks when position or orientation is constant
@@ -618,7 +642,8 @@ public:
 	 * go down to a single frame, the time is stripped out as well.
 	 * @return true if keys were removed.
 	 */
-	bool CompressRawAnimData();
+	UE_DEPRECATED(5.0, "CompressRawAnimData has been deprecated, reduction of Raw Animation data now happens during compression instead")
+	bool CompressRawAnimData() { return false; }
 
 	// Get compressed data for this UAnimSequence. May be built directly or pulled from DDC
 #if WITH_EDITOR
@@ -659,7 +684,6 @@ public:
 
 protected:
 	void ApplyCompressedData(const TArray<uint8>& Data);
-	void SanitizeAnimationTrackData(FRawAnimSequenceTrack& RawAnim);
 
 public:
 	bool IsCompressedDataValid() const;
@@ -681,10 +705,8 @@ public:
 	 *
 	 * @param	TrackIndex		Track Index
 	 */
-	int32 GetSkeletonIndexFromRawDataTrackIndex(const int32 TrackIndex) const 
-	{ 
-		return TrackToSkeletonMapTable[TrackIndex].BoneTreeIndex; 
-	}
+	UE_DEPRECATED(5.0, "DoesContainTransformCurves has been deprecated see FBoneAnimationTrack::BoneTreeIndex")
+	int32 GetSkeletonIndexFromRawDataTrackIndex(const int32 TrackIndex) const;
 
 	/**
 	* Get Skeleton Bone Index from Track Index for compressed data
@@ -697,21 +719,25 @@ public:
 	}
 
 	/** Clears any data in the AnimSequence */
+	UE_DEPRECATED(5.0, "RecycleAnimSequence has been deprecated use ResetAnimation instead")
 	void RecycleAnimSequence();
 
 #if WITH_EDITOR
 	/** Clears some data in the AnimSequence, so it can be reused when importing a new animation with same name over it. */
+	UE_DEPRECATED(5.0, "CleanAnimSequenceForImport has been deprecated see UAnimDataController:ResetModel")
 	void CleanAnimSequenceForImport();
 #endif
 
 	/** 
 	 * Copy AnimNotifies from one UAnimSequence to another.
 	 */
+	UE_DEPRECATED(5.0, "CopyNotifies has been moved, see UAnimSequenceHelpers::CopyNotifies")
 	static bool CopyNotifies(UAnimSequence* SourceAnimSeq, UAnimSequence* DestAnimSeq, bool bShowDialogs = true);
 
 	/**
 	 * Flip Rotation's W For NonRoot items, and compress it again if SkelMesh exists
 	 */
+	UE_DEPRECATED(5.0, "FlipRotationWForNonRoot has been deprecated, use UAnimDataModel::GetBoneAnimationTracks and UAnimDataController::SetBoneTrackKeys instead")
 	void FlipRotationWForNonRoot(USkeletalMesh * SkelMesh);
 
 	// End Utility functions
@@ -719,20 +745,24 @@ public:
 	/**
 	 * After imported or any other change is made, call this to apply post process
 	 */
+	UE_DEPRECATED(5.0, "PostProcessSequence has been deprecated, any Raw Animation Data modifications should go through UAnimDataController")
 	void PostProcessSequence(bool bForceNewRawDatGuid = true);
 
 	// Kick off compression request when our raw data has changed
+	UE_DEPRECATED(5.0, "OnRawDataChanged has been deprecated, any Raw Animation Data modifications should go through UAnimDataController")
 	void OnRawDataChanged();
 
 	/** 
 	 * Insert extra frame of the first frame at the end of the frame so that it improves the interpolation when it loops
 	 * This increases framecount + time, so that it requires recompression
 	 */
+	UE_DEPRECATED(5.0, "AddLoopingInterpolation has been deprecated see FAnimSequenceHelpers::AnimationData::AddLoopingInterpolation")
 	bool AddLoopingInterpolation();
 
 	/*
 	* Clear all raw animation data that contains bone tracks
 	*/
+	UE_DEPRECATED(5.0, "RemoveAllTracks has been deprecated see UAnimDataController::RemoveAllBoneTracks")
 	void RemoveAllTracks();
 
 	/** 
@@ -762,6 +792,7 @@ public:
 	/**
 	 * Return true if it contains transform curves
 	 */
+	UE_DEPRECATED(5.0, "DoesContainTransformCurves has been deprecated see UAnimDataModel::GetNumberOfTransformCurves")
 	bool DoesContainTransformCurves() const;
 
 	/**
@@ -802,8 +833,8 @@ public:
 	 * @param	bFromStart		whether marker is begin or end marker
 	 * @return					true if the operation was successful.
 	 */
+	UE_DEPRECATED(5.0, "InsertFramesToRawAnimData has been deprecated see FAnimSequenceHelpers::AnimationData::Trim")
 	bool CropRawAnimData( float CurrentTime, bool bFromStart );
-
 		
 	/**
 	 * Crops the raw anim data either from Start to CurrentTime or CurrentTime to End depending on
@@ -814,6 +845,7 @@ public:
 	 * @param	CopyFrame		A frame that we copy from (0-based)
 	 * @return					true if the operation was successful.
 	 */
+	UE_DEPRECATED(5.0, "InsertFramesToRawAnimData has been deprecated see FAnimSequenceHelpers::AnimationData::DuplicateKeys")
 	bool InsertFramesToRawAnimData( int32 StartFrame, int32 EndFrame, int32 CopyFrame);
 
 	/** 
@@ -866,10 +898,10 @@ public:
 	void TestEvalauteAnimation() const;
 
 	// Bakes out the additive version of this animation into the raw data.
-	void BakeOutAdditiveIntoRawData(TArray<FRawAnimSequenceTrack>& NewRawTracks, TArray<FName>& NewAnimationTrackNames, TArray<FTrackToSkeletonMap>& NewTrackToSkeletonMapTable, FRawCurveTracks& NewCurveTracks, TArray<FRawAnimSequenceTrack>& AdditiveBaseAnimationData);
+	void BakeOutAdditiveIntoRawData(TArray<FRawAnimSequenceTrack>& NewRawTracks, TArray<FName>& NewAnimationTrackNames, TArray<FTrackToSkeletonMap>& NewTrackToSkeletonMapTable, TArray<FFloatCurve>& NewCurveTracks, TArray<FRawAnimSequenceTrack>& AdditiveBaseAnimationData);
 
 	// Test whether at any point we will scale a bone to 0 (needed for validating additive anims)
-	bool DoesSequenceContainZeroScale();
+	bool DoesSequenceContainZeroScale() const;
 
 	// Helper function to allow us to notify animations that depend on us that they need to update
 	void FlagDependentAnimationsAsRawDataOnly() const;
@@ -885,7 +917,13 @@ public:
 	void SetUseRawDataOnly(bool bInUseRawDataOnly) { bUseRawDataOnly = bInUseRawDataOnly; }
 
 	// Return this animations guid for the raw data
-	FGuid GetRawDataGuid() const { return RawDataGuid; }
+	FGuid GetRawDataGuid() const
+	{ 
+		return RawDataGuid;
+	}
+
+	/** Resets Bone Animation, Curve data and Notify tracks **/
+	void ResetAnimation();
 #endif
 
 private:
@@ -905,8 +943,9 @@ private:
 #if WITH_EDITOR
 	/**
 	 * Remap Tracks to New Skeleton
-	 */
+	 */ 
 	virtual void RemapTracksToNewSkeleton( USkeleton* NewSkeleton, bool bConvertSpaces ) override;
+
 	/**
 	 * Remap NaN tracks from the RawAnimation data and recompress
 	 */	
@@ -914,35 +953,35 @@ private:
 
 	/** Retargeting functions */
 	bool ConvertAnimationDataToRiggingData(FAnimSequenceTrackContainer & RiggingAnimationData);
-	bool ConvertRiggingDataToAnimationData(FAnimSequenceTrackContainer & RiggingAnimationData, bool bPerformPostProcess=true);
+	bool ConvertRiggingDataToAnimationData(FAnimSequenceTrackContainer & RiggingAnimationData);
 	int32 GetSpaceBasedAnimationData(TArray< TArray<FTransform> > & AnimationDataInComponentSpace, FAnimSequenceTrackContainer * RiggingAnimationData) const;
 
 	/** Verify Track Map is valid, if not, fix up */
+	UE_DEPRECATED(5.0, "RemoveTrack has been deprecated see UAnimDataController::RemoveBoneTracksMissingFromSkeleton")
 	void VerifyTrackMap(USkeleton* MySkeleton=NULL);
-	/** Reset Animation Data. Called before Creating new Animation data **/
-	void ResetAnimation();
+
 	/** Refresh Track Map from Animation Track Names **/
+	UE_DEPRECATED(5.0, "RemoveTrack has been deprecated see UAnimDataController::RemoveBoneTracksMissingFromSkeleton")
 	void RefreshTrackMapFromAnimTrackNames();
 
 	/**
 	 * Utility function that helps to remove track, you can't just remove RawAnimationData
 	 */
+	UE_DEPRECATED(5.0, "RemoveTrack has been deprecated see UAnimDataController::RemoveBoneTrack")
 	void RemoveTrack(int32 TrackIndex);
 
 	/**
 	 * Utility function that finds the correct spot to insert track to 
 	 */
+	UE_DEPRECATED(5.0, "InsertTrack has been deprecated see UAnimDataController::InsertBoneTrack")
 	int32 InsertTrack(const FName& BoneName);
 
-private:
-	/** Internal insert function to be used by BakeTrackCurvesToRawAnimation*/ 
-	int32 InsertTrackInternal(const FName& BoneName);
-
 public:
 	/**
 	 * Utility function to resize the sequence
 	 * It rearranges curve data + notifies
 	 */
+	UE_DEPRECATED(5.0, "ResizeSequence has been deprecated see UAnimDataController::Resize")
 	void ResizeSequence(float NewLength, int32 NewNumFrames, bool bInsert, int32 StartFrame/*inclusive */, int32 EndFrame/*inclusive*/);
 
 #endif
@@ -1038,9 +1077,41 @@ protected:
 	FCustomAttributePerBoneData& FindOrAddCustomAttributeForBone(const FName& BoneName);
 #endif // WITH_EDITOR
 
+protected:
+#if WITH_EDITOR
+	// Begin UAnimSequenceBase virtual overrides
+	virtual void OnModelModified(const EAnimDataModelNotifyType& NotifyType, UAnimDataModel* Model, const FAnimDataModelNotifPayload& Payload) override;
+	virtual void PopulateModel() override;
+	// End UAnimSequenceBase virtual overrides
+
+	void EnsureValidRawDataGuid();
+	void RecompressAnimationData();
+	void ResampleAnimationTrackData();
+
+	void DeleteBoneAnimationData();
+	void DeleteDeprecatedRawAnimationData();
+public:
+	const TArray<FBoneAnimationTrack>& GetResampledTrackData() const { return ResampledAnimationTrackData; }
+	void DeleteNotifyTrackData();
+#endif // WITH_EDITOR
+
+protected:
+	UPROPERTY(VisibleAnywhere, AssetRegistrySearchable, Category = "Animation")
+	FFrameRate TargetFrameRate;
 
-private:
 #if WITH_EDITORONLY_DATA
+	UPROPERTY(VisibleAnywhere, AssetRegistrySearchable, Category = "Animation", Transient, DuplicateTransient)
+	int32 NumberOfSampledKeys;
+
+	UPROPERTY(VisibleAnywhere, Category = "Animation", Transient, DuplicateTransient)
+	int32 NumberOfSampledFrames;
+
+	UPROPERTY(VisibleAnywhere, Category = "Animation", Transient, DuplicateTransient)
+	TArray<FBoneAnimationTrack> ResampledAnimationTrackData;
+
+	bool bBlockCompressionRequests;
+
+private:
 	UPROPERTY(VisibleAnywhere, EditFixedSize, Category=CustomAttributes)
 	TArray<FCustomAttributePerBoneData> PerBoneCustomAttributeData;
 	
diff --git a/Engine/Source/Runtime/Engine/Classes/Animation/AnimSequenceBase.h b/Engine/Source/Runtime/Engine/Classes/Animation/AnimSequenceBase.h
--- a/Engine/Source/Runtime/Engine/Classes/Animation/AnimSequenceBase.h
+++ b/Engine/Source/Runtime/Engine/Classes/Animation/AnimSequenceBase.h
@@ -13,6 +13,8 @@
 #include "Animation/AnimationAsset.h"
 #include "Animation/AnimCurveTypes.h"
 #include "Animation/AnimNotifyQueue.h"
+#include "Animation/AnimData/AnimDataModelNotifyCollector.h"
+
 #include "AnimSequenceBase.generated.h"
 
 UENUM()
@@ -24,6 +26,16 @@ enum ETypeAdvanceAnim
 };
 
 struct FAnimationPoseData;
+struct FAnimDataModelNotifPayload;
+class UAnimDataModel;
+class UAnimDataController;
+enum class EAnimDataModelNotifyType : uint8;
+
+#if 0
+#define REFACTOR_V2_WARNING(text) EMIT_CUSTOM_WARNING(text)
+#else
+#define REFACTOR_V2_WARNING(text)
+#endif
 
 UCLASS(abstract, BlueprintType)
 class ENGINE_API UAnimSequenceBase : public UAnimationAsset
@@ -36,19 +48,12 @@ public:
 	TArray<struct FAnimNotifyEvent> Notifies;
 
 	/** Length (in seconds) of this AnimSequence if played back with a speed of 1.0. */
-	UE_DEPRECATED(5.0, "Direct access to SequenceLength should be removed, use GetPlayLength or SetSequenceLength instead")
+	UE_DEPRECATED(5.0, "Public access to SequenceLength is deprecated, use GetPlayLength or UAnimDataController::SetPlayLength instead")
 	UPROPERTY(Category=Length, AssetRegistrySearchable, VisibleAnywhere, BlueprintReadOnly)
 	float SequenceLength;
 
-	virtual void SetSequenceLength(float NewLength) 	
-	{
-		PRAGMA_DISABLE_DEPRECATION_WARNINGS
-		if (NewLength > 0.f)
-		{
-			SequenceLength = NewLength;
-		}
-		PRAGMA_ENABLE_DEPRECATION_WARNINGS
-	}
+	UE_DEPRECATED(5.0, "SetSequenceLength is deprecated use UAnimDataController::SetPlayLength instead")
+	virtual void SetSequenceLength(float NewLength);
 
 	/** Number for tweaking playback rate of this animation globally. */
 	UPROPERTY(EditAnywhere, Category=Animation)
@@ -57,6 +62,7 @@ public:
 	/**
 	 * Raw uncompressed float curve data 
 	 */
+	UE_DEPRECATED(5.0, "Public access to RawCurveData is deprecated, see UAnimDataModel for source data or use GetCurveData for runtime instead")
 	UPROPERTY()
 	struct FRawCurveTracks RawCurveData;
 
@@ -69,6 +75,7 @@ public:
 	//~ Begin UObject Interface
 	virtual void PostLoad() override;
 	virtual bool IsPostLoadThreadSafe() const override;
+	virtual void PostDuplicate(EDuplicateMode::Type DuplicateMode) override;
 	//~ End UObject Interface
 
 	/** Returns the total play length of the montage, if played back with a speed of 1.0. */
@@ -118,7 +125,7 @@ public:
 	virtual void EvaluateCurveData(FBlendedCurve& OutCurve, float CurrentTime, bool bForceUseRawData = false) const;
 	virtual float EvaluateCurveData(SmartName::UID_Type CurveUID, float CurrentTime, bool bForceUseRawData = false) const;
 
-	virtual const FRawCurveTracks& GetCurveData() const { return RawCurveData; }
+	virtual const FRawCurveTracks& GetCurveData() const;
 	virtual bool HasCurveData(SmartName::UID_Type CurveUID, bool bForceUseRawData = false) const;
 
 	/** Return Number of Keys **/
@@ -132,7 +139,6 @@ public:
 	virtual const FFrameRate& GetSamplingFrameRate() const;
 
 #if WITH_EDITOR
-
 	/** Get the frame number for the provided time */
 	virtual int32 GetFrameAtTime(const float Time) const;
 
@@ -164,7 +170,8 @@ public:
 	virtual void RefreshCacheData();
 
 #if WITH_EDITOR
-	void RefreshCurveData();
+	UE_DEPRECATED(5.0, "Updating of any relevent Curve data is now handled through the UAnimDataModel notifies")
+	void RefreshCurveData() {}
 #endif // WITH_EDITOR
 
 	//~ Begin UAnimationAsset Interface
@@ -172,11 +179,8 @@ public:
 	virtual void RemapTracksToNewSkeleton(USkeleton* NewSkeleton, bool bConvertSpaces) override;
 #endif
 	virtual void TickAssetPlayer(FAnimTickRecord& Instance, struct FAnimNotifyQueue& NotifyQueue, FAnimAssetTickContext& Context) const override;
-
 	void TickByMarkerAsFollower(FMarkerTickRecord &Instance, FMarkerTickContext &MarkerContext, float& CurrentTime, float& OutPreviousTime, const float MoveDelta, const bool bLooping) const;
-
 	void TickByMarkerAsLeader(FMarkerTickRecord& Instance, FMarkerTickContext& MarkerContext, float& CurrentTime, float& OutPreviousTime, const float MoveDelta, const bool bLooping) const;
-
 	//~ End UAnimationAsset Interface
 
 	/**
@@ -218,6 +222,7 @@ public:
 
 #if WITH_EDITOR
 	// Store that our raw data has changed so that we can get correct compressed data later on
+	UE_DEPRECATED(5.0, "MarkRawDataAsModified has been deprecated, any (Raw Data) modification should be applied using the UAnimDataController API instead. This will handle updating the GUID instead.")
 	virtual void MarkRawDataAsModified(bool bForceNewRawDatGuid = true) {}
 
 private:
@@ -234,30 +239,25 @@ public:
 	// ideally this would be animsequcnebase, but we might have some issue with that. For now, just allow AnimSequence
 	virtual class UAnimSequence* GetAdditiveBasePose() const { return nullptr; }
 
-#endif
-	// return true if anim notify is available 
-	virtual bool IsNotifyAvailable() const;
-
-#if WITH_EDITOR
-private:
-	DECLARE_MULTICAST_DELEGATE(FOnAnimCurvesChangedMulticaster);
-	FOnAnimCurvesChangedMulticaster OnAnimCurvesChanged;
-
-	DECLARE_MULTICAST_DELEGATE(FOnAnimTrackCurvesChangedMulticaster);
-	FOnAnimTrackCurvesChangedMulticaster OnAnimTrackCurvesChanged;
-
-public:
-	typedef FOnAnimCurvesChangedMulticaster::FDelegate FOnAnimCurvesChanged;	
+	typedef FSimpleMulticastDelegate::FDelegate FOnAnimCurvesChanged;
 	/** Registers a delegate to be called after anim curves have changed*/
-	void RegisterOnAnimCurvesChanged(const FOnAnimCurvesChanged& Delegate);
-	void UnregisterOnAnimCurvesChanged(void* Unregister);
+	UE_DEPRECATED(5.0, "Functionality has been deprecated, register to UAnimDataModel::GetModifiedEvent instead")
+	void RegisterOnAnimCurvesChanged(const FOnAnimCurvesChanged& Delegate) {}
+
+	UE_DEPRECATED(5.0, "Functionality has been deprecated, register to UAnimDataModel::GetModifiedEvent instead")
+	void UnregisterOnAnimCurvesChanged(void* Unregister) {}
 
-	typedef FOnAnimTrackCurvesChangedMulticaster::FDelegate FOnAnimTrackCurvesChanged;
+	typedef FSimpleMulticastDelegate::FDelegate FOnAnimTrackCurvesChanged;
 	/** Registers a delegate to be called after anim track curves have changed*/
-	void RegisterOnAnimTrackCurvesChanged(const FOnAnimTrackCurvesChanged& Delegate);
-	void UnregisterOnAnimTrackCurvesChanged(void* Unregister);
+	UE_DEPRECATED(5.0, "Functionality has been deprecated, register to UAnimDataModel::GetModifiedEvent instead")
+	void RegisterOnAnimTrackCurvesChanged(const FOnAnimTrackCurvesChanged& Delegate) {}
+
+	UE_DEPRECATED(5.0, "Functionality has been deprecated, register to UAnimDataModel::GetModifiedEvent instead")
+	void UnregisterOnAnimTrackCurvesChanged(void* Unregister) {}
 #endif
 
+	// return true if anim notify is available 
+	virtual bool IsNotifyAvailable() const;
 protected:
 	template <typename DataType>
 	void VerifyCurveNames(USkeleton& Skeleton, const FName& NameContainer, TArray<DataType>& CurveList)
@@ -267,4 +267,43 @@ protected:
 			Skeleton.VerifySmartName(NameContainer, Curve.Name);
 		}
 	}
+#if WITH_EDITOR
+public:
+	/** Returns the UAnimDataModel object embedded in this UAnimSequenceBase */
+	UAnimDataModel* GetDataModel() const;
+
+	/** Returns the transient UAnimDataController set to operate on DataModel */
+	UAnimDataController* GetController();
+
+protected:
+	/** Populates the UAnimDataModel object according to any pre-existing data. (overrides expect to populate the model according to their data) */
+	virtual void PopulateModel();
+
+	/** Callback registered to UAnimDatModel::GetModifiedEvent for the embedded object */
+	virtual void OnModelModified(const EAnimDataModelNotifyType& NotifyType, UAnimDataModel* Model, const FAnimDataModelNotifPayload& Payload);
+
+	/** Validates that DataModel contains a valid UAnimDataModel object */
+	void ValidateModel() const;
+
+	/** Replaces the current DataModel, if any, with the provided one */
+	void CopyDataModel(const UAnimDataModel* ModelToDuplicate);
+private:
+	/** Creates a new UAnimDataModel instance and sets DataModel accordingly */
+	void CreateModel();
+#endif // WITH_EDITOR
+
+#if WITH_EDITORONLY_DATA
+protected:
+	/** UAnimDataModel instance containing source animation data */
+	UPROPERTY(VisibleAnywhere, BlueprintReadOnly, Category = "Animation Model")
+	UAnimDataModel* DataModel;
+
+	/** UAnimDataController instance set to operate on DataModel */
+	UPROPERTY(VisibleAnywhere, BlueprintReadOnly, Transient, DuplicateTransient, Category = "Animation Model")
+	UAnimDataController* Controller;
+	
+	/** Helper object that keeps track of any controller brackets, and all unique notify types that are broadcasted during it */
+	UE::Anim::FAnimDataModelNotifyCollector NotifyCollector;
+#endif // WITH_EDITORONLY_DATA
+
 };
diff --git a/Engine/Source/Runtime/Engine/Classes/Animation/AnimSequenceHelpers.h b/Engine/Source/Runtime/Engine/Classes/Animation/AnimSequenceHelpers.h
new file mode 100644
--- /dev/null
+++ b/Engine/Source/Runtime/Engine/Classes/Animation/AnimSequenceHelpers.h
@@ -0,0 +1,214 @@
+// Copyright Epic Games, Inc. All Rights Reserved.
+
+#pragma once
+
+#include "UObject/WeakObjectPtrTemplates.h"
+#include "AnimData/AnimDataNotifications.h"
+#include "Containers/Array.h"
+#include "Delegates/IDelegateInstance.h"
+
+class UAnimSequence;
+class UAnimDataModel;
+class UAnimDataController;
+class FName;
+struct FCompactPose;
+struct FRawCurveTracks;
+struct FRawAnimSequenceTrack;
+struct FBoneAnimationTrack;
+struct FRawAnimSequenceTrack;
+
+enum class EAnimInterpolationType : uint8;
+
+namespace UE {
+
+namespace Anim {
+
+#if WITH_EDITOR
+	/**
+	* Populates a FCompactPose according to the data stored within a UAnimDataModel provided a time to sample.
+	*
+	* @param	Model				Model to sample from
+	* @param	OutPose				[out] Pose to be populated using the model
+	* @param	Time				Time value to sample at
+	* @param	InterpolationType	Method to be used to interpolate between different sample points
+	* @param	RetargetSource		Name of the retarget source to retrieve from the skeleton
+	* @param	RetargetTransforms	Base transforms to use when retargetting the pose
+	*/
+	ENGINE_API void BuildPoseFromModel(const UAnimDataModel* Model, FCompactPose& OutPose, const float Time, const EAnimInterpolationType& InterpolationType, const FName& RetargetSource, const TArray<FTransform>& RetargetTransforms);
+
+	/**
+	* Populates float curves according to the data stored within a UAnimDataModel provided a time to sample.
+	*
+	* @param	Model				Model to sample from
+	* @param	OutCurves			[out] Curves to be populated using the model
+	* @param	Time				Time value to sample at
+	*/
+	ENGINE_API void EvaluateFloatCurvesFromModel(const UAnimDataModel* Model, FBlendedCurve& OutCurves, float Time);
+
+	/**
+	* Populates transform curves according to the data stored within a UAnimDataModel provided a time to sample.
+	*
+	* @param	Model				Model to sample from
+	* @param	OutCurves			[out] Curves which were sampled from the model, key-ed with their equivalent name
+	* @param	Time				Time value to sample at
+	* @param	BlendWeight			Blending weight to apply when sampling
+	*/
+	ENGINE_API void EvaluateTransformCurvesFromModel(const UAnimDataModel* Model, TMap<FName, FTransform>& OutCurves, float Time, float BlendWeight);
+
+	/**
+	* Retrieves a single bone transform according to the data stored within a UAnimDataModel provided a time and track to sample.
+	*
+	* @param	Model				Model to sample from
+	* @param	OutTransform		[out] Transform which was sampled from the model
+	* @param	TrackIndex			Track / bone index which should be sampld
+	* @param	Time				Time value to sample at
+	* @param	InterpolationType	Method to be used to interpolate between different sample points
+	*/
+	ENGINE_API void GetBoneTransformFromModel(const UAnimDataModel* Model, FTransform& OutTransform, int32 TrackIndex, float Time, const EAnimInterpolationType& Interpolation);
+	
+	/**
+	* Retrieves a single bone transform according to the data stored within a UAnimDataModel provided a key and track to sample.
+	*
+	* @param	Model				Model to sample from
+	* @param	OutTransform		[out] Transform which was sampled from the model
+	* @param	TrackIndex			Track / bone index which should be sampld
+	* @param	KeyIndex			Specific key index to sample
+	*/
+	ENGINE_API void GetBoneTransformFromModel(const UAnimDataModel* Model, FTransform& OutTransform, int32 TrackIndex, int32 KeyIndex);
+
+	/**
+	* Copies over any individual curve from FRawCurveTracks to a UAnimDataModel instance targeted by the provided controller.
+	*
+	* @param	CurveData			Container of the curves to be copied over
+	* @param	Skeleton			Skeleton to use for verifying and or add the curve names
+	* @param	Controller			Controller to use for adding curve data to a model
+	*/
+	ENGINE_API void CopyCurveDataToModel(const FRawCurveTracks& CurveData, const USkeleton* Skeleton, UAnimDataController* Controller);
+
+	/**
+	* Copies over any individual curve from FRawCurveTracks to a UAnimDataModel instance targeted by the provided controller.
+	*
+	* @param	SourceAnimSeq		Animation Sequence to copy the notifies from
+	* @param	Skeleton			Animation Sequence to copy the notifies to
+	* @param	bShowDialogs		Whether or not to show any user-facing dialogs for confirmation of the copy
+	*
+	* @return Whether or not the copy was succesful
+	*/
+	ENGINE_API bool CopyNotifies(const UAnimSequence* SourceAnimSeq, UAnimSequence* DestAnimSeq, bool bShowDialogs = true);
+
+	namespace AnimationData
+	{
+		/**
+		* Adds an additional animated frame to match with the first animated frame, attempting to make it loop
+		*
+		* @param	InSequence		Animation Sequence to add the animated frame to
+		*
+		* @return Whether or not the operation was succesful
+		*/
+		ENGINE_API bool AddLoopingInterpolation(UAnimSequence* InSequence);
+
+		/**
+		* Trim a specific window from the animation data
+		*
+		* @param	InSequence		Animation Sequence to add the animated frame to
+		* @param	TrimStart		Time value at which to start removing frames
+		* @param	TrimEnd			Time value at which to stop removing frames
+		*
+		* @return Whether or not the operation was succesful
+		*/
+		ENGINE_API bool Trim(UAnimSequence* InSequence, float TrimStart, float TrimEnd);
+				
+		/**
+		* Insert duplicate key(s) for all tracks for the provided Animation Sequence
+		*
+		* @param	InSequence			Animation Sequence to add the duplicated animation frames
+		* @param	StartKeyIndex		Key index after which the duplicate keys should be inserted
+		* @param	NumDuplicates		Number of duplicated keys to insert
+		* @param	SourceKeyIndex		Specific key index to use as a source for the duplicated keys, otherwise uses key at StartKeyIndex
+		*/
+		ENGINE_API void DuplicateKeys(UAnimSequence* InSequence, int32 StartKeyIndex, int32 NumDuplicates, int32 SourceKeyIndex = INDEX_NONE);
+
+		/**
+		* Remove a number of keys for all tracks for the provided Animation Sequence
+		*
+		* @param	InSequence			Animation Sequence to remove the animated frames from
+		* @param	StartKeyIndex		Key index at which the key removal should start
+		* @param	NumKeysToRemove		Number of keys to remove
+		*/
+		ENGINE_API void RemoveKeys(UAnimSequence* InSequence, int32 StartKeyIndex, int32 NumKeysToRemove);
+
+		/**
+		* Finds the index of the first child track (bone) for the provided bone name
+		*
+		* @param	InSequence			Animation Sequence to remove the animated frames from
+		* @param	Skeleton			Skeleton to use for retrieving child bone names
+		* @param	BoneName			Name of the bone to find the child track index for
+		*
+		* @return The first child track index for the provided bone name, if none found will return the end of the tracks array
+		*/
+		ENGINE_API int32 FindFirstChildTrackIndex(const UAnimSequence* InSequence, const USkeleton* Skeleton, const FName& BoneName);
+	}
+#endif // WITH_EDITOR
+
+	/**
+	* Extract Bone Transform of the Time given, from the provided FRawAnimSequenceTrack
+	*
+	* @param	RawTrack		RawAnimationTrack it extracts bone transform from
+	* @param	OutAtom			[out] Output bone transform.
+	* @param	KeyIndex		Key index to retrieve from the track.
+	*/
+	ENGINE_API void ExtractBoneTransform(const struct FRawAnimSequenceTrack& RawTrack, FTransform& OutTransform, int32 KeyIndex);
+		
+	namespace Compression
+	{
+		/**
+		* Compress the provided tracks using two methods.
+		* 1. Checking for uniform key data according to the provided error metrics
+		* 2. Setting number of scale keys to 0 if it is uniform and equal to FVector::OneVector
+		*
+		* @param	RawAnimationData	Set of tracks to apply the reduction to
+		* @param	NumberOfKeys		Expected number of keys for each component within the tracks
+		* @param	ErrorName			Identifier for outer callsite when failing to reduce the data
+		* @param	MaxPosDiff			Maximum positional delta when determining whether or not two positional keys are identical
+		* @param	MaxAngleDiff		Maximum rotational delta when determining whether or not two rotational keys are identical
+		*
+		* @return Whether or not the operation was succesful
+		*/
+		ENGINE_API bool CompressAnimationDataTracks(TArray<FRawAnimSequenceTrack>& RawAnimationData, int32 NumberOfKeys, FName ErrorName, float MaxPosDiff = 0.0001f, float MaxAngleDiff = 0.0003f);
+
+		/**
+		* Compress the provided track by checking for uniform key data according to the provided error metrics
+		*
+		* @param	RawTrack			Individual track to apply the reduction to
+		* @param	NumberOfKeys		Expected number of keys for each component within the tracks
+		* @param	ErrorName			Identifier for outer callsite when failing to reduce the data
+		* @param	MaxPosDiff			Maximum positional delta when determining whether or not two positional keys are identical
+		* @param	MaxAngleDiff		Maximum rotational delta when determining whether or not two rotational keys are identical
+		*
+		* @return Whether or not the operation was succesful
+		*/
+		ENGINE_API bool CompressRawAnimSequenceTrack(FRawAnimSequenceTrack& RawTrack, int32 NumberOfKeys, FName ErrorName, float MaxPosDiff, float MaxAngleDiff);
+
+		/**
+		* Sanitize the provided track by snapping small scale values to 0, and normalizing any rotational keys
+		*
+		* @param	RawTrack			Individual track to sanitize
+		*/
+		ENGINE_API void SanitizeRawAnimSequenceTrack(FRawAnimSequenceTrack& RawTrack);
+	}
+
+	namespace Retargeting
+	{
+		/**
+		* Retargeting the provided pose using the retarget source or transforms
+		*
+		* @param	InOutPose			[out] Individual track to apply the reduction to
+		* @param	RetargetSource		Name of the retarget source to retrieve data for
+		* @param	RetargetTransforms	Set of transforms to use as the ref skeleton bone transforms
+		*/
+		ENGINE_API void RetargetPose(FCompactPose& InOutPose, const FName& RetargetSource, const TArray<FTransform>& RetargetTransforms);
+	}
+
+} // namespace Anim	
+
+} // namespace UE
\ No newline at end of file
diff --git a/Engine/Source/Runtime/Engine/Classes/Curves/RichCurve.h b/Engine/Source/Runtime/Engine/Classes/Curves/RichCurve.h
--- a/Engine/Source/Runtime/Engine/Classes/Curves/RichCurve.h
+++ b/Engine/Source/Runtime/Engine/Classes/Curves/RichCurve.h
@@ -73,7 +73,7 @@ enum ERichCurveKeyTimeCompressionFormat
 };
 
 /** One key in a rich, editable float curve */
-USTRUCT()
+USTRUCT(BlueprintType)
 struct ENGINE_API FRichCurveKey
 {
 	GENERATED_USTRUCT_BODY()
@@ -91,11 +91,11 @@ struct ENGINE_API FRichCurveKey
 	TEnumAsByte<ERichCurveTangentWeightMode> TangentWeightMode;
 
 	/** Time at this key */
-	UPROPERTY(EditAnywhere, Category="Key")
+	UPROPERTY(EditAnywhere, BlueprintReadWrite, Category="Key")
 	float Time;
 
 	/** Value at this key */
-	UPROPERTY(EditAnywhere, Category="Key")
+	UPROPERTY(EditAnywhere, BlueprintReadWrite, Category="Key")
 	float Value;
 
 	/** If RCIM_Cubic, the arriving tangent at this key */
diff --git a/Engine/Source/Runtime/Engine/Classes/Engine/Engine.h b/Engine/Source/Runtime/Engine/Classes/Engine/Engine.h
--- a/Engine/Source/Runtime/Engine/Classes/Engine/Engine.h
+++ b/Engine/Source/Runtime/Engine/Classes/Engine/Engine.h
@@ -3193,6 +3193,7 @@ public:
 	//~ Begin Transaction Interfaces.
 	virtual int32 BeginTransaction(const TCHAR* TransactionContext, const FText& Description, UObject* PrimaryObject) { return INDEX_NONE; }
 	virtual int32 EndTransaction() { return INDEX_NONE; }
+	virtual bool CanTransact() { return false; }
 	virtual void CancelTransaction(int32 Index) { }
 #endif
 
diff --git a/Engine/Source/Runtime/Engine/Private/Animation/AnimComposite.cpp b/Engine/Source/Runtime/Engine/Private/Animation/AnimComposite.cpp
--- a/Engine/Source/Runtime/Engine/Private/Animation/AnimComposite.cpp
+++ b/Engine/Source/Runtime/Engine/Private/Animation/AnimComposite.cpp
@@ -7,6 +7,7 @@
 #include "Animation/AnimComposite.h"
 #include "Animation/AnimationPoseData.h"
 #include "Animation/CustomAttributesRuntime.h"
+#include "Animation/AnimData/AnimDataController.h"
 
 UAnimComposite::UAnimComposite(const FObjectInitializer& ObjectInitializer)
 	: Super(ObjectInitializer)
@@ -140,3 +141,14 @@ bool UAnimComposite::ContainRecursive(TArray<UAnimCompositeBase*>& CurrentAccumu
 
 	return false;
 }
+
+void UAnimComposite::SetCompositeLength(float InLength)
+{
+#if WITH_EDITOR		
+	Controller->SetPlayLength(InLength);
+#else
+	PRAGMA_DISABLE_DEPRECATION_WARNINGS
+	SetSequenceLength(InLength);
+	PRAGMA_ENABLE_DEPRECATION_WARNINGS
+#endif	
+}
diff --git a/Engine/Source/Runtime/Engine/Private/Animation/AnimCompressionDerivedData.cpp b/Engine/Source/Runtime/Engine/Private/Animation/AnimCompressionDerivedData.cpp
--- a/Engine/Source/Runtime/Engine/Private/Animation/AnimCompressionDerivedData.cpp
+++ b/Engine/Source/Runtime/Engine/Private/Animation/AnimCompressionDerivedData.cpp
@@ -35,7 +35,7 @@ const TCHAR* FDerivedDataAnimationCompression::GetVersionString() const
 	// This is a version string that mimics the old versioning scheme. If you
 	// want to bump this version, generate a new guid using VS->Tools->Create GUID and
 	// return it here. Ex.
-	return TEXT("D0E639028BB74237AD0F9775C097BB85");
+	return TEXT("3B32894C1C3A45A883A84C1CB2F349F0");
 }
 
 bool FDerivedDataAnimationCompression::Build( TArray<uint8>& OutDataArray )
diff --git a/Engine/Source/Runtime/Engine/Private/Animation/AnimCompressionTypes.cpp b/Engine/Source/Runtime/Engine/Private/Animation/AnimCompressionTypes.cpp
--- a/Engine/Source/Runtime/Engine/Private/Animation/AnimCompressionTypes.cpp
+++ b/Engine/Source/Runtime/Engine/Private/Animation/AnimCompressionTypes.cpp
@@ -11,86 +11,14 @@
 #include "Animation/AnimCurveCompressionSettings.h"
 #include "AnimationRuntime.h"
 #include "UObject/FortniteMainBranchObjectVersion.h"
+#include "Animation/AnimSequenceHelpers.h"
+#include "Animation/AnimData/AnimDataModel.h"
 
 CSV_DECLARE_CATEGORY_MODULE_EXTERN(ENGINE_API, Animation);
 
 DECLARE_CYCLE_STAT(TEXT("Build Anim Track Pairs"), STAT_BuildAnimTrackPairs, STATGROUP_Anim);
 DECLARE_CYCLE_STAT(TEXT("Extract Pose From Anim Data"), STAT_ExtractPoseFromAnimData, STATGROUP_Anim);
 
-template <typename ArrayType>
-void UpdateSHAWithArray(FSHA1& Sha, const TArray<ArrayType>& Array)
-{
-	Sha.Update((uint8*)Array.GetData(), Array.Num() * Array.GetTypeSize());
-}
-
-void UpdateSHAWithRawTrack(FSHA1& Sha, const FRawAnimSequenceTrack& RawTrack)
-{
-	UpdateSHAWithArray(Sha, RawTrack.PosKeys);
-	UpdateSHAWithArray(Sha, RawTrack.RotKeys);
-	UpdateSHAWithArray(Sha, RawTrack.ScaleKeys);
-}
-
-template<class DataType>
-void UpdateWithData(FSHA1& Sha, const DataType& Data)
-{
-	Sha.Update((uint8*)(&Data), sizeof(DataType));
-}
-
-void UpdateSHAWithCurves(FSHA1& Sha, const FRawCurveTracks& InRawCurveData) 
-{
-	auto UpdateWithFloatCurve = [&Sha](const FRichCurve& Curve)
-	{
-		UpdateWithData(Sha, Curve.DefaultValue);
-		UpdateSHAWithArray(Sha, Curve.GetConstRefOfKeys());
-		UpdateWithData(Sha, Curve.PreInfinityExtrap);
-		UpdateWithData(Sha, Curve.PostInfinityExtrap);
-	};
-
-	for (const FFloatCurve& Curve : InRawCurveData.FloatCurves)
-	{
-		UpdateWithData(Sha, Curve.Name.UID);
-		UpdateWithFloatCurve(Curve.FloatCurve);
-	}
-
-#if WITH_EDITOR
-	for (const FTransformCurve& Curve : InRawCurveData.TransformCurves)
-	{
-		UpdateWithData(Sha, Curve.Name.UID);
-
-		auto UpdateWithComponent = [&Sha, &UpdateWithFloatCurve](const FVectorCurve& VectorCurve)
-		{
-			for (int32 ChannelIndex = 0; ChannelIndex < 3; ++ChannelIndex)
-			{
-				UpdateWithFloatCurve(VectorCurve.FloatCurves[ChannelIndex]);
-			}
-		};
-		
-		UpdateWithComponent(Curve.TranslationCurve);
-		UpdateWithComponent(Curve.RotationCurve);
-		UpdateWithComponent(Curve.ScaleCurve);
-	}
-#endif // WITH_EDITOR
-}
-
-FGuid GenerateGuidFromRawAnimData(const TArray<FRawAnimSequenceTrack>& RawAnimationData, const FRawCurveTracks& RawCurveData)
-{
-	FSHA1 Sha;
-
-	for (const FRawAnimSequenceTrack& Track : RawAnimationData)
-	{
-		UpdateSHAWithRawTrack(Sha, Track);
-	}
-
-	UpdateSHAWithCurves(Sha, RawCurveData);
-
-	Sha.Final();
-
-	uint32 Hash[5];
-	Sha.GetHash((uint8*)Hash);
-	FGuid Guid(Hash[0] ^ Hash[4], Hash[1], Hash[2], Hash[3]);
-	return Guid;
-}
-
 template<typename ArrayValue>
 void StripFramesEven(TArray<ArrayValue>& Keys, const int32 NumFrames)
 {
@@ -152,7 +80,6 @@ void StripFramesOdd(TArray<ArrayValue>& Keys, const int32 NumFrames)
 FCompressibleAnimData::FCompressibleAnimData(class UAnimSequence* InSeq, const bool bPerformStripping)
 	: CurveCompressionSettings(InSeq->CurveCompressionSettings)
 	, BoneCompressionSettings(InSeq->BoneCompressionSettings)
-	, TrackToSkeletonMapTable(InSeq->GetRawTrackToSkeletonMapTable())
 	, Interpolation(InSeq->Interpolation)
 	, SequenceLength(InSeq->GetPlayLength())
 	, NumberOfKeys(InSeq->GetNumberOfSampledKeys())
@@ -175,35 +102,85 @@ FCompressibleAnimData::FCompressibleAnimData(class UAnimSequence* InSeq, const b
 
 	const bool bHasVirtualBones = InSeq->GetSkeleton()->GetVirtualBones().Num() > 0;
 
+	/* Always get the resampled data to start off with */
+	const TArray<FBoneAnimationTrack>& ResampledTrackData = InSeq->GetResampledTrackData();
+
 	TArray<FName> TempTrackNames;
+	RawAnimationData.Empty(ResampledTrackData.Num());
+	TrackToSkeletonMapTable.Empty(ResampledTrackData.Num());
+	TempTrackNames.Empty(ResampledTrackData.Num());
 
-	if (InSeq->CanBakeAdditive())
+	for (const FBoneAnimationTrack& AnimTrack : ResampledTrackData)
 	{
-		InSeq->BakeOutAdditiveIntoRawData(RawAnimationData, TempTrackNames, TrackToSkeletonMapTable, RawCurveData, AdditiveBaseAnimationData);
+		FRawAnimSequenceTrack& Track = RawAnimationData.Add_GetRef(AnimTrack.InternalTrackData);
+		UE::Anim::Compression::SanitizeRawAnimSequenceTrack(Track);
+		TrackToSkeletonMapTable.Add(AnimTrack.BoneTreeIndex);
+		TempTrackNames.Add(AnimTrack.Name);
 	}
-	else
+
+	const bool bBakeAdditive = InSeq->CanBakeAdditive();
+	if (bBakeAdditive)
 	{
+		InSeq->BakeOutAdditiveIntoRawData(RawAnimationData, TempTrackNames, TrackToSkeletonMapTable, RawFloatCurves, AdditiveBaseAnimationData);
+	}
+	else
+    {
 		// In case we require baking down transform curves, do so now meaning Virtual Bone baking will incorporate the correct bone transforms
-		if (InSeq->DoesContainTransformCurves())
-		{
-			InSeq->BakeTrackCurvesToRawAnimationTracks(RawAnimationData, TempTrackNames, TrackToSkeletonMapTable);			
-		}
-		else
-		{
-			RawAnimationData = InSeq->GetRawAnimationData();
-			TrackToSkeletonMapTable = InSeq->GetRawTrackToSkeletonMapTable();
-			TempTrackNames = InSeq->GetAnimationTrackNames();
-		}
+	    if (InSeq->GetDataModel()->GetNumberOfTransformCurves() > 0)
+	    {
+		    InSeq->BakeTrackCurvesToRawAnimationTracks(RawAnimationData, TempTrackNames, TrackToSkeletonMapTable);
+	    }
 
-		RawCurveData = InSeq->RawCurveData;
+		RawFloatCurves = InSeq->GetCurveData().FloatCurves;
 
 		// If we aren't additive we must bake virtual bones
-		if (bHasVirtualBones)
+	    if (bHasVirtualBones)
+	    {
+		    InSeq->BakeOutVirtualBoneTracks(RawAnimationData, TempTrackNames, TrackToSkeletonMapTable);
+	    }
+    } 
+
+	// Apply any key reduction if possible
+	if (RawAnimationData.Num())
+	{ 
+		UE::Anim::Compression::CompressAnimationDataTracks(RawAnimationData, NumberOfKeys, InSeq->GetFName(), -1.f, -1.f);
+		UE::Anim::Compression::CompressAnimationDataTracks(RawAnimationData, NumberOfKeys, InSeq->GetFName());
+	}
+
+	auto IsKeyArrayValidForRemoval = [](const auto& Keys, const auto& IdentityValue) -> bool
+	{
+		return Keys.Num() == 0 || (Keys.Num() == 1 && Keys[0].Equals(IdentityValue));
+	};
+
+	auto IsRawTrackValidForRemoval = [IsKeyArrayValidForRemoval](const FRawAnimSequenceTrack& Track) -> bool
+	{
+		return IsKeyArrayValidForRemoval(Track.PosKeys, FVector::ZeroVector) &&
+			IsKeyArrayValidForRemoval(Track.RotKeys, FQuat::Identity) &&
+			IsKeyArrayValidForRemoval(Track.ScaleKeys, FVector::ZeroVector);
+	};
+
+	// Verify bone track names and data, removing any bone that does not exist on the skeleton
+	// Note on (TrackIndex > 0) below : deliberately stop before track 0, compression code doesn't like getting a completely empty animation
+	for (int32 TrackIndex = RawAnimationData.Num() - 1; TrackIndex > 0; --TrackIndex)
+	{
+		const FRawAnimSequenceTrack& Track = RawAnimationData[TrackIndex];
+		// Try find correct bone index
+		const int32 BoneIndex = RefSkeleton.FindBoneIndex(TempTrackNames[TrackIndex]);
+		if (IsRawTrackValidForRemoval(Track) || (BoneIndex == INDEX_NONE && !bBakeAdditive))
 		{
-			InSeq->BakeOutVirtualBoneTracks(RawAnimationData, TempTrackNames, TrackToSkeletonMapTable);
+			RawAnimationData.RemoveAtSwap(TrackIndex, 1, false);
+			TempTrackNames.RemoveAtSwap(TrackIndex, 1, false);
+			TrackToSkeletonMapTable.RemoveAtSwap(TrackIndex, 1, false);
 		}
 	}
 
+	// Find or add curve names on skeleton
+	const FSmartNameMapping* Mapping = Skeleton->GetSmartNameContainer(USkeleton::AnimCurveMappingName);
+	for (FFloatCurve& Curve : RawFloatCurves)
+	{
+		Skeleton->VerifySmartName(USkeleton::AnimCurveMappingName, Curve.Name);
+	}
+	   
 	if (bPerformStripping)
 	{
 		const int32 NumTracks = RawAnimationData.Num();
@@ -272,12 +249,12 @@ void FCompressibleAnimData::Update(FCompressedAnimSequence& InOutCompressedData)
 	InOutCompressedData.CompressedTrackToSkeletonMapTable = TrackToSkeletonMapTable;
 	InOutCompressedData.CompressedRawDataSize = GetApproxRawSize();
 
-	const int32 NumCurves = RawCurveData.FloatCurves.Num();
+	const int32 NumCurves = RawFloatCurves.Num();
 	InOutCompressedData.CompressedCurveNames.Reset(NumCurves);
 	InOutCompressedData.CompressedCurveNames.AddUninitialized(NumCurves);
 	for (int32 CurveIndex = 0; CurveIndex < NumCurves; ++CurveIndex)
 	{
-		const FFloatCurve& Curve = RawCurveData.FloatCurves[CurveIndex];
+		const FFloatCurve& Curve = RawFloatCurves[CurveIndex];
 		InOutCompressedData.CompressedCurveNames[CurveIndex] = Curve.Name;
 	}
 }
