Animation data MVC refactor

diff --git a/Engine/Source/Runtime/Engine/Private/Animation/AnimCurveCompressionCodec_CompressedRichCurve.cpp b/Engine/Source/Runtime/Engine/Private/Animation/AnimCurveCompressionCodec_CompressedRichCurve.cpp
--- a/Engine/Source/Runtime/Engine/Private/Animation/AnimCurveCompressionCodec_CompressedRichCurve.cpp
+++ b/Engine/Source/Runtime/Engine/Private/Animation/AnimCurveCompressionCodec_CompressedRichCurve.cpp
@@ -28,7 +28,7 @@ struct FCurveDesc
 #if WITH_EDITORONLY_DATA
 bool UAnimCurveCompressionCodec_CompressedRichCurve::Compress(const FCompressibleAnimData& AnimSeq, FAnimCurveCompressionResult& OutResult)
 {
-	int32 NumCurves = AnimSeq.RawCurveData.FloatCurves.Num();
+	int32 NumCurves = AnimSeq.RawFloatCurves.Num();
 
 	TArray<FCurveDesc> Curves;
 	Curves.Reserve(NumCurves);
@@ -44,7 +44,7 @@ bool UAnimCurveCompressionCodec_CompressedRichCurve::Compress(const FCompressibl
 
 	for (int32 CurveIndex = 0; CurveIndex < NumCurves; ++CurveIndex)
 	{
-		const FFloatCurve& Curve = AnimSeq.RawCurveData.FloatCurves[CurveIndex];
+		const FFloatCurve& Curve = AnimSeq.RawFloatCurves[CurveIndex];
 
 		FRichCurve RawCurve = Curve.FloatCurve;	// Copy
 		RawCurve.RemoveRedundantKeys(MaxCurveError);
diff --git a/Engine/Source/Runtime/Engine/Private/Animation/AnimCurveCompressionCodec_UniformIndexable.cpp b/Engine/Source/Runtime/Engine/Private/Animation/AnimCurveCompressionCodec_UniformIndexable.cpp
--- a/Engine/Source/Runtime/Engine/Private/Animation/AnimCurveCompressionCodec_UniformIndexable.cpp
+++ b/Engine/Source/Runtime/Engine/Private/Animation/AnimCurveCompressionCodec_UniformIndexable.cpp
@@ -83,7 +83,7 @@ UAnimCurveCompressionCodec_UniformIndexable::UAnimCurveCompressionCodec_UniformI
 #if WITH_EDITORONLY_DATA
 bool UAnimCurveCompressionCodec_UniformIndexable::Compress(const FCompressibleAnimData& AnimSeq, FAnimCurveCompressionResult& OutResult)
 {
-	const int32 NumCurves = AnimSeq.RawCurveData.FloatCurves.Num();
+	const int32 NumCurves = AnimSeq.RawFloatCurves.Num();
 	const float Duration = AnimSeq.SequenceLength;
 
 	const FAnimKeyHelper Helper(AnimSeq.SequenceLength, AnimSeq.NumberOfKeys);
@@ -119,10 +119,8 @@ bool UAnimCurveCompressionCodec_UniformIndexable::Compress(const FCompressibleAn
 
 		// Curve 0 Key 0, Curve 0 Key 1, Curve 0 Key N, Curve 1 Key 0, Curve 1 Key 1, Curve 1 Key N, Curve M Key 0, ...
 		const float InvSampleRate = 1.0f / SampleRate;
-		for (int32 CurveIndex = 0; CurveIndex < NumCurves; ++CurveIndex)
+		for (const FFloatCurve& Curve : AnimSeq.RawFloatCurves)
 		{
-			const FFloatCurve& Curve = AnimSeq.RawCurveData.FloatCurves[CurveIndex];
-
 			for (int32 SampleIndex = 0; SampleIndex < NumSamples; ++SampleIndex)
 			{
 				const float SampleTime = FMath::Clamp(SampleIndex * InvSampleRate, 0.0f, Duration);
diff --git a/Engine/Source/Runtime/Engine/Private/Animation/AnimCurveCompressionCodec_UniformlySampled.cpp b/Engine/Source/Runtime/Engine/Private/Animation/AnimCurveCompressionCodec_UniformlySampled.cpp
--- a/Engine/Source/Runtime/Engine/Private/Animation/AnimCurveCompressionCodec_UniformlySampled.cpp
+++ b/Engine/Source/Runtime/Engine/Private/Animation/AnimCurveCompressionCodec_UniformlySampled.cpp
@@ -16,7 +16,7 @@ UAnimCurveCompressionCodec_UniformlySampled::UAnimCurveCompressionCodec_Uniforml
 #if WITH_EDITORONLY_DATA
 bool UAnimCurveCompressionCodec_UniformlySampled::Compress(const FCompressibleAnimData& AnimSeq, FAnimCurveCompressionResult& OutResult)
 {
-	const int32 NumCurves = AnimSeq.RawCurveData.FloatCurves.Num();
+	const int32 NumCurves = AnimSeq.RawFloatCurves.Num();
 	const float Duration = AnimSeq.SequenceLength;
 
 	int32 NumSamples;
@@ -37,7 +37,7 @@ bool UAnimCurveCompressionCodec_UniformlySampled::Compress(const FCompressibleAn
 	}
 
 	int32 NumConstantCurves = 0;
-	for (const FFloatCurve& Curve : AnimSeq.RawCurveData.FloatCurves)
+	for (const FFloatCurve& Curve : AnimSeq.RawFloatCurves)
 	{
 		if (Curve.FloatCurve.IsConstant())
 		{
@@ -92,7 +92,7 @@ bool UAnimCurveCompressionCodec_UniformlySampled::Compress(const FCompressibleAn
 		{
 			for (int32 CurveIndex = 0, ConstantCurveIndex = 0; CurveIndex < NumCurves; ++CurveIndex)
 			{
-				const FFloatCurve& Curve = AnimSeq.RawCurveData.FloatCurves[CurveIndex];
+				const FFloatCurve& Curve = AnimSeq.RawFloatCurves[CurveIndex];
 				if (Curve.FloatCurve.IsConstant())
 				{
 					if (Curve.FloatCurve.IsEmpty())
@@ -124,7 +124,7 @@ bool UAnimCurveCompressionCodec_UniformlySampled::Compress(const FCompressibleAn
 
 				for (int32 CurveIndex = 0, AnimatedCurveIndex = 0; CurveIndex < NumCurves; ++CurveIndex)
 				{
-					const FFloatCurve& Curve = AnimSeq.RawCurveData.FloatCurves[CurveIndex];
+					const FFloatCurve& Curve = AnimSeq.RawFloatCurves[CurveIndex];
 					if (Curve.FloatCurve.IsConstant())
 					{
 						// Skip constant curves, their data has already been written
diff --git a/Engine/Source/Runtime/Engine/Private/Animation/AnimCurveTypes.cpp b/Engine/Source/Runtime/Engine/Private/Animation/AnimCurveTypes.cpp
--- a/Engine/Source/Runtime/Engine/Private/Animation/AnimCurveTypes.cpp
+++ b/Engine/Source/Runtime/Engine/Private/Animation/AnimCurveTypes.cpp
@@ -393,6 +393,7 @@ const FFloatCurve* FCachedFloatCurve::GetFloatCurve(UAnimSequenceBase* InAnimSeq
 		USkeleton::AnimCurveUID DistanceCurveUID = GetAnimCurveUID(InAnimSequence);
 		if (DistanceCurveUID != SmartName::MaxUID)
 		{
+			const FAnimationCurveIdentifier CurveId(DistanceCurveUID, ERawCurveTrackTypes::RCT_Float);
 			return (const FFloatCurve*)(InAnimSequence->GetCurveData().GetCurveData(DistanceCurveUID));
 		}
 	}
diff --git a/Engine/Source/Runtime/Engine/Private/Animation/AnimDataModel.cpp b/Engine/Source/Runtime/Engine/Private/Animation/AnimDataModel.cpp
new file mode 100644
--- /dev/null
+++ b/Engine/Source/Runtime/Engine/Private/Animation/AnimDataModel.cpp
@@ -0,0 +1,442 @@
+// Copyright Epic Games, Inc. All Rights Reserved.
+
+#include "Animation/AnimData/AnimDataModel.h"
+#include "UObject/NameTypes.h"
+
+#include "Animation/AnimSequence.h"
+
+#include "Algo/Transform.h"
+#include "Animation/SmartName.h"
+
+
+void UAnimDataModel::PostLoad()
+{
+	UObject::PostLoad();
+
+	GenerateTransientData();
+}
+
+void UAnimDataModel::PostDuplicate(bool bDuplicateForPIE)
+{
+	UObject::PostDuplicate(bDuplicateForPIE);
+
+	GenerateTransientData();
+	Notify(EAnimDataModelNotifyType::Populated);
+}
+
+const TArray<FBoneAnimationTrack>& UAnimDataModel::GetBoneAnimationTracks() const
+{
+	return BoneAnimationTracks;
+}
+
+const FBoneAnimationTrack& UAnimDataModel::GetBoneTrackByIndex(int32 TrackIndex) const
+{
+	checkf(BoneAnimationTracks.IsValidIndex(TrackIndex), TEXT("Unable to find animation track by index"));
+	return BoneAnimationTracks[TrackIndex];
+}
+
+const FBoneAnimationTrack& UAnimDataModel::GetBoneTrackByName(FName TrackName) const
+{
+	const FBoneAnimationTrack* TrackPtr = BoneAnimationTracks.FindByPredicate([TrackName](FBoneAnimationTrack& Track)
+	{
+		return Track.Name == TrackName;
+	});
+
+	checkf(TrackPtr != nullptr, TEXT("Unable to find animation track by name"));
+
+	return *TrackPtr;
+}
+
+const FBoneAnimationTrack* UAnimDataModel::FindBoneTrackByIndex(int32 BoneIndex) const
+{
+	const FBoneAnimationTrack* TrackPtr = BoneAnimationTracks.FindByPredicate([BoneIndex](FBoneAnimationTrack& Track)
+	{
+		return Track.BoneTreeIndex == BoneIndex;
+	});
+
+	return TrackPtr;
+}
+
+int32 UAnimDataModel::GetBoneTrackIndex(const FBoneAnimationTrack& Track) const
+{
+	return BoneAnimationTracks.IndexOfByPredicate([&Track](const FBoneAnimationTrack SearchTrack)
+	{
+		return SearchTrack.Name == Track.Name;
+	});
+}
+
+int32 UAnimDataModel::GetBoneTrackIndexByName(FName TrackName) const
+{
+	if (const FBoneAnimationTrack* TrackPtr = FindBoneTrackByName(TrackName))
+	{
+		return GetBoneTrackIndex(*TrackPtr);
+	}
+
+	return INDEX_NONE;
+}
+
+float UAnimDataModel::GetPlayLength() const
+{
+	return PlayLength;
+}
+
+int32 UAnimDataModel::GetNumberOfFrames() const
+{
+	return NumberOfFrames;
+}
+
+int32 UAnimDataModel::GetNumberOfKeys() const
+{
+	return NumberOfKeys;
+}
+
+const FFrameRate& UAnimDataModel::GetFrameRate() const
+{
+	return FrameRate;
+}
+
+bool UAnimDataModel::IsValidBoneTrackIndex(int32 TrackIndex) const
+{
+	return BoneAnimationTracks.IsValidIndex(TrackIndex);
+}
+
+const int32 UAnimDataModel::GetNumBoneTracks() const
+{
+	return BoneAnimationTracks.Num();
+}
+
+UAnimSequence* UAnimDataModel::GetAnimationSequence() const
+{
+	return Cast<UAnimSequence>(GetOuter());
+}
+
+void UAnimDataModel::GetBoneTrackNames(TArray<FName>& OutNames) const
+{
+	Algo::Transform(BoneAnimationTracks, OutNames, [](FBoneAnimationTrack Track)
+	{
+		return Track.Name;
+	});
+}
+
+const FAnimationCurveData& UAnimDataModel::GetCurveData() const
+{
+	return CurveData;
+}
+
+int32 UAnimDataModel::GetNumberOfTransformCurves() const
+{
+	return CurveData.TransformCurves.Num();
+}
+
+int32 UAnimDataModel::GetNumberOfFloatCurves() const
+{
+	return CurveData.FloatCurves.Num();
+}
+
+const TArray<FFloatCurve>& UAnimDataModel::GetFloatCurves() const
+{
+	return CurveData.FloatCurves;
+}
+
+const TArray<struct FTransformCurve>& UAnimDataModel::GetTransformCurves() const
+{
+	return CurveData.TransformCurves;
+}
+
+const FFloatCurve* UAnimDataModel::FindFloatCurve(const FAnimationCurveIdentifier& CurveIdentifier) const
+{
+	ensure(CurveIdentifier.CurveType == ERawCurveTrackTypes::RCT_Float);
+	for (const FFloatCurve& FloatCurve : CurveData.FloatCurves)
+	{
+		if (FloatCurve.Name == CurveIdentifier.InternalName)
+		{
+			return &FloatCurve;
+		}
+	}
+
+	return nullptr;
+}
+
+const FTransformCurve* UAnimDataModel::FindTransformCurve(const FAnimationCurveIdentifier& CurveIdentifier) const
+{
+	ensure(CurveIdentifier.CurveType == ERawCurveTrackTypes::RCT_Transform);
+	for (const FTransformCurve& TransformCurve : CurveData.TransformCurves)
+	{
+		if (TransformCurve.Name == CurveIdentifier.InternalName)
+		{
+			return &TransformCurve;
+		}
+	}
+
+	return nullptr;
+}
+
+const FRichCurve* UAnimDataModel::FindRichCurve(const FAnimationCurveIdentifier& CurveIdentifier) const
+{
+	const FRichCurve* RichCurve = nullptr;
+
+	if (CurveIdentifier.IsValid())
+	{
+		if (CurveIdentifier.CurveType == ERawCurveTrackTypes::RCT_Float)
+		{
+			const FFloatCurve* Curve = FindFloatCurve(CurveIdentifier);
+			if (Curve)
+			{
+				RichCurve = &Curve->FloatCurve;
+			}
+		}
+		else if (CurveIdentifier.CurveType == ERawCurveTrackTypes::RCT_Transform)
+		{
+			if (CurveIdentifier.Channel != ETransformCurveChannel::Invalid && CurveIdentifier.Axis != EVectorCurveChannel::Invalid)
+			{
+				// Dealing with transform curve
+				const FTransformCurve* TransformCurve = FindTransformCurve(CurveIdentifier);
+				if (TransformCurve)
+				{
+					const FVectorCurve* VectorCurve = TransformCurve->GetVectorCurveByIndex((int32)CurveIdentifier.Channel);
+					if (VectorCurve)
+					{
+						RichCurve = &VectorCurve->FloatCurves[(int32)CurveIdentifier.Axis];
+					}
+				}
+
+			}
+		}
+	}
+
+	return RichCurve;
+}
+
+const FAnimCurveBase& UAnimDataModel::GetCurve(const FAnimationCurveIdentifier& CurveIdentifier) const
+{
+	const FAnimCurveBase* CurvePtr = FindCurve(CurveIdentifier);
+	checkf(CurvePtr, TEXT("Tried to retrieve non-existing curve"));
+	return *CurvePtr;
+}
+
+const FFloatCurve& UAnimDataModel::GetFloatCurve(const FAnimationCurveIdentifier& CurveIdentifier) const
+{
+	const FFloatCurve* CurvePtr = FindFloatCurve(CurveIdentifier);
+	checkf(CurvePtr, TEXT("Tried to retrieve non-existing curve"));
+	return *CurvePtr;
+}
+
+const FTransformCurve& UAnimDataModel::GetTransformCurve(const FAnimationCurveIdentifier& CurveIdentifier) const
+{
+	const FTransformCurve* CurvePtr = FindTransformCurve(CurveIdentifier);
+	checkf(CurvePtr, TEXT("Tried to retrieve non-existing curve"));
+	return *CurvePtr;
+}
+
+const FRichCurve& UAnimDataModel::GetRichCurve(const FAnimationCurveIdentifier& CurveIdentifier) const
+{
+	const FRichCurve* CurvePtr = FindRichCurve(CurveIdentifier);
+	checkf(CurvePtr, TEXT("Tried to retrieve non-existing curve"));
+	return *CurvePtr;
+}
+
+const FAnimCurveBase* UAnimDataModel::FindCurve(const FAnimationCurveIdentifier& CurveIdentifier) const
+{
+	switch (CurveIdentifier.CurveType)
+	{
+		case ERawCurveTrackTypes::RCT_Float:
+			return FindFloatCurve(CurveIdentifier);
+		case ERawCurveTrackTypes::RCT_Transform:
+			return FindTransformCurve(CurveIdentifier);
+	}
+
+	return nullptr;
+}
+
+FGuid UAnimDataModel::GenerateGuid() const
+{
+	FSHA1 Sha;
+	
+	auto UpdateSHAWithArray = [&](const auto& Array)
+	{
+		Sha.Update((uint8*)Array.GetData(), Array.Num() * Array.GetTypeSize());
+	};
+	   	
+	auto UpdateWithData = [&](const auto& Data)
+	{
+		Sha.Update((uint8*)(&Data), sizeof(Data));
+	};
+
+	for (const FBoneAnimationTrack& Track : BoneAnimationTracks)
+	{
+		UpdateSHAWithArray(Track.InternalTrackData.PosKeys);
+		UpdateSHAWithArray(Track.InternalTrackData.RotKeys);
+		UpdateSHAWithArray(Track.InternalTrackData.ScaleKeys);
+	}
+
+	for (const FFloatCurve& Curve : CurveData.FloatCurves)
+	{
+		UpdateWithData(Curve.Name.UID);
+		UpdateWithData(Curve.FloatCurve.DefaultValue);
+		UpdateSHAWithArray(Curve.FloatCurve.GetConstRefOfKeys());
+		UpdateWithData(Curve.FloatCurve.PreInfinityExtrap);
+		UpdateWithData(Curve.FloatCurve.PostInfinityExtrap);
+	}
+
+	Sha.Final();
+
+	uint32 Hash[5];
+	Sha.GetHash((uint8*)Hash);
+	FGuid Guid(Hash[0] ^ Hash[4], Hash[1], Hash[2], Hash[3]);
+	
+	return Guid;
+}
+
+const TArray<FRawAnimSequenceTrack>& UAnimDataModel::GetTransientRawAnimationTracks() const
+{
+	return RawAnimationTracks;
+}
+
+const TArray<FName>& UAnimDataModel::GetTransientRawAnimationTrackNames() const
+{
+	return RawAnimationTrackNames;
+}
+
+const TArray<FTrackToSkeletonMap>& UAnimDataModel::GetTransientRawAnimationTrackSkeletonMappings() const
+{
+	return RawAnimationTrackSkeletonMappings;
+}
+
+FRawAnimSequenceTrack& UAnimDataModel::GetNonConstRawAnimationTrackByIndex(int32 TrackIndex)
+{
+	checkf(BoneAnimationTracks.IsValidIndex(TrackIndex), TEXT("Invalid track index"));
+	return BoneAnimationTracks[TrackIndex].InternalTrackData;
+}
+
+const FRawCurveTracks& UAnimDataModel::GetTransientRawCurveTracks() const
+{
+	return RawCurveTracks;
+}
+
+FAnimationCurveData& UAnimDataModel::GetNonConstCurveData()
+{
+	return CurveData;	
+}
+
+void UAnimDataModel::GenerateTransientData()
+{
+	RawAnimationTracks.Empty(BoneAnimationTracks.Num());
+	RawAnimationTrackNames.Empty(BoneAnimationTracks.Num());
+	RawAnimationTrackSkeletonMappings.Empty(BoneAnimationTracks.Num());
+
+	for (const FBoneAnimationTrack& AnimTrack : BoneAnimationTracks)
+	{
+		RawAnimationTracks.Add(AnimTrack.InternalTrackData);
+		RawAnimationTrackNames.Add(AnimTrack.Name);
+		RawAnimationTrackSkeletonMappings.Add(AnimTrack.BoneTreeIndex);
+	}
+
+	RawCurveTracks.FloatCurves = CurveData.FloatCurves;
+#if WITH_EDITOR
+	RawCurveTracks.TransformCurves = CurveData.TransformCurves;
+#endif
+}
+
+FRichCurve* UAnimDataModel::GetMutableRichCurve(const FAnimationCurveIdentifier& CurveIdentifier)
+{
+	FRichCurve* RichCurve = nullptr;
+
+	if (CurveIdentifier.IsValid())
+	{
+		if (CurveIdentifier.CurveType == ERawCurveTrackTypes::RCT_Float)
+		{
+			FFloatCurve* Curve = FindMutableFloatCurveById(CurveIdentifier);
+			if (Curve)
+			{
+				RichCurve = &Curve->FloatCurve;
+			}
+		}
+		else if (CurveIdentifier.CurveType == ERawCurveTrackTypes::RCT_Transform)
+		{
+			if (CurveIdentifier.Channel != ETransformCurveChannel::Invalid && CurveIdentifier.Axis != EVectorCurveChannel::Invalid)
+			{
+				// Dealing with transform curve
+				FTransformCurve* TransformCurve = FindMutableTransformCurveById(CurveIdentifier);
+				if (TransformCurve)
+				{
+					FVectorCurve* VectorCurve = TransformCurve->GetVectorCurveByIndex((int32)CurveIdentifier.Channel);
+					if (VectorCurve)
+					{
+						RichCurve = &VectorCurve->FloatCurves[(int32)CurveIdentifier.Axis];
+					}
+				}
+
+			}
+		}
+	}
+
+	return RichCurve;
+}
+
+FBoneAnimationTrack* UAnimDataModel::FindMutableBoneTrackByName(FName Name)
+{
+	return BoneAnimationTracks.FindByPredicate([Name](FBoneAnimationTrack& Track)
+	{
+		return Track.Name == Name;
+	});
+}
+
+const FBoneAnimationTrack* UAnimDataModel::FindBoneTrackByName(FName Name) const
+{
+	return BoneAnimationTracks.FindByPredicate([Name](FBoneAnimationTrack& Track)
+	{
+		return Track.Name == Name;
+	});
+}
+
+FBoneAnimationTrack& UAnimDataModel::GetMutableBoneTrackByName(FName Name)
+{
+	FBoneAnimationTrack* TrackPtr = BoneAnimationTracks.FindByPredicate([Name](FBoneAnimationTrack& Track)
+	{
+		return Track.Name == Name;
+	});
+
+	checkf(TrackPtr, TEXT("Failed to find track by name"));
+
+	return *TrackPtr;
+}
+
+FTransformCurve* UAnimDataModel::FindMutableTransformCurveById(const FAnimationCurveIdentifier& CurveIdentifier)
+{
+	for (FTransformCurve& TransformCurve : CurveData.TransformCurves)
+	{
+		if (TransformCurve.Name.UID == CurveIdentifier.InternalName.UID)
+		{
+			return &TransformCurve;
+		}
+	}
+
+	return nullptr;
+}
+
+FFloatCurve* UAnimDataModel::FindMutableFloatCurveById(const FAnimationCurveIdentifier& CurveIdentifier)
+{
+	for (FFloatCurve& FloatCurve : CurveData.FloatCurves)
+	{
+		if (FloatCurve.Name.UID == CurveIdentifier.InternalName.UID)
+		{
+			return &FloatCurve;
+		}
+	}
+
+	return nullptr;
+}
+
+FAnimCurveBase* UAnimDataModel::FindMutableCurveById(const FAnimationCurveIdentifier& CurveIdentifier)
+{
+	switch (CurveIdentifier.CurveType)
+	{
+	case ERawCurveTrackTypes::RCT_Float:
+		return FindMutableFloatCurveById(CurveIdentifier);
+	case ERawCurveTrackTypes::RCT_Transform:
+		return FindMutableTransformCurveById(CurveIdentifier);
+	}
+
+	return nullptr;
+}
+
diff --git a/Engine/Source/Runtime/Engine/Private/Animation/AnimMontage.cpp b/Engine/Source/Runtime/Engine/Private/Animation/AnimMontage.cpp
--- a/Engine/Source/Runtime/Engine/Private/Animation/AnimMontage.cpp
+++ b/Engine/Source/Runtime/Engine/Private/Animation/AnimMontage.cpp
@@ -18,6 +18,7 @@
 #include "Animation/AnimSingleNodeInstance.h"
 #include "Engine/Engine.h"
 #include "Animation/AnimTrace.h"
+#include "Animation/AnimData/AnimDataController.h"
 
 DEFINE_LOG_CATEGORY(LogAnimMontage);
 
@@ -343,8 +344,9 @@ void UAnimMontage::UnregisterOnMontageChanged(void* Unregister)
 
 void UAnimMontage::PreSave(const class ITargetPlatform* TargetPlatform)
 {
+#if WITH_EDITOR
 	BakeTimeStretchCurve();
-
+#endif // WITH_EDITOR
 	Super::PreSave(TargetPlatform);
 }
 
@@ -362,7 +364,7 @@ void UAnimMontage::PostLoad()
 		if(CurrentCalculatedLength != GetPlayLength())		
 		{
 			UE_LOG(LogAnimMontage, Display, TEXT("UAnimMontage::PostLoad: The actual sequence length for %s does not match the length stored in the asset, please resave the asset."), *GetFullName());
-			 SetSequenceLength(CurrentCalculatedLength);
+			SetCompositeLength(CurrentCalculatedLength);
 		}
 	}
 
@@ -1288,6 +1290,17 @@ bool UAnimMontage::ContainRecursive(TArray<UAnimCompositeBase*>& CurrentAccumula
 	return false;
 }
 
+void UAnimMontage::SetCompositeLength(float InLength)
+{
+#if WITH_EDITOR	
+	Controller->SetPlayLength(InLength);
+#else
+	PRAGMA_DISABLE_DEPRECATION_WARNINGS
+	SetSequenceLength(InLength);
+	PRAGMA_ENABLE_DEPRECATION_WARNINGS
+#endif	
+}
+
 //////////////////////////////////////////////////////////////////////////////////////////////
 // MontageInstance
 /////////////////////////////////////////////////////////////////////////////////////////////
@@ -1351,7 +1364,7 @@ void FAnimMontageInstance::Play(float InPlayRate)
 		BlendInSettings.Blend = Montage->BlendIn;
 		BlendInSettings.BlendMode = Montage->BlendModeIn;
 		BlendInSettings.BlendProfile = Montage->BlendProfileIn;
-	}
+}
 
 	Play(InPlayRate, BlendInSettings);
 }
@@ -1419,7 +1432,7 @@ void FAnimMontageInstance::Stop(const FMontageBlendSettings& InBlendOutSettings,
 		BlendOutArgs.BlendTime = bShouldInertialize ? 0.0f : BlendOutArgs.BlendTime;
 
 		// do not use default Montage->BlendOut 
-		// depending on situation, the BlendOut time can change
+		// depending on situation, the BlendOut time can change 
 		InitializeBlend(FAlphaBlend(BlendOutArgs));
 		BlendStartAlpha = Blend.GetAlpha();
 		Blend.SetDesiredValue(0.f);
@@ -1428,7 +1441,7 @@ void FAnimMontageInstance::Stop(const FMontageBlendSettings& InBlendOutSettings,
 		// Only change the active blend profile if the montage isn't stopped. This is to prevent pops on a sudden blend profile switch
 		ActiveBlendProfile = InBlendOutSettings.BlendProfile;
 
-		if (Montage)
+		if(Montage)
 		{
 			if (UAnimInstance* Inst = AnimInstance.Get())
 			{
@@ -3029,7 +3042,7 @@ UAnimMontage* UAnimMontage::CreateSlotAnimationAsDynamicMontage(UAnimSequenceBas
 	NewSegment.AnimPlayRate = 1.f;
 	NewSegment.StartPos = 0.f;
 	NewSegment.LoopingCount = LoopCount;
-	NewMontage->SetSequenceLength(NewSegment.GetLength());
+	NewMontage->SetCompositeLength(NewSegment.GetLength());
 	NewTrack.AnimTrack.AnimSegments.Add(NewSegment);
 
 	FCompositeSection NewSection;
@@ -3051,12 +3064,13 @@ bool FAnimMontageInstance::CanUseMarkerSync() const
 	return SyncGroupName != NAME_None && IsStopped() && Blend.IsComplete() == false;
 }
 
+#if WITH_EDITOR
 void UAnimMontage::BakeTimeStretchCurve()
 {
 	TimeStretchCurve.Reset();
 
 	// See if Montage is hosting a curve named 'TimeStretchCurveName'
-	FFloatCurve* TimeStretchFloatCurve = nullptr;
+	const FFloatCurve* TimeStretchFloatCurve = nullptr;
 	if (const USkeleton* MySkeleton = GetSkeleton())
 	{
 		if (const FSmartNameMapping* CurveNameMapping = MySkeleton->GetSmartNameContainer(USkeleton::AnimCurveMappingName))
@@ -3064,7 +3078,7 @@ void UAnimMontage::BakeTimeStretchCurve()
 			const USkeleton::AnimCurveUID CurveUID = CurveNameMapping->FindUID(TimeStretchCurveName);
 			if (CurveUID != SmartName::MaxUID)
 			{
-				TimeStretchFloatCurve = (FFloatCurve*)(GetCurveData().GetCurveData(CurveUID));
+				TimeStretchFloatCurve = GetDataModel()->FindFloatCurve(FAnimationCurveIdentifier(CurveUID, ERawCurveTrackTypes::RCT_Float));
 			}
 		}
 	}
@@ -3078,6 +3092,7 @@ void UAnimMontage::BakeTimeStretchCurve()
 	TimeStretchCurve.BakeFromFloatCurve(*TimeStretchFloatCurve, SequenceLength);
 	PRAGMA_ENABLE_DEPRECATION_WARNINGS
 }
+#endif // WITH_EDITOR
 
 FMontageBlendSettings::FMontageBlendSettings()
 	: BlendProfile(nullptr)
@@ -3085,4 +3100,4 @@ FMontageBlendSettings::FMontageBlendSettings()
 	, BlendMode(EMontageBlendMode::Standard)
 {
 
-}
+}
\ No newline at end of file
