Animation data MVC refactor

diff --git a/Engine/Source/Runtime/Engine/Private/Animation/AnimSequence.cpp b/Engine/Source/Runtime/Engine/Private/Animation/AnimSequence.cpp
--- a/Engine/Source/Runtime/Engine/Private/Animation/AnimSequence.cpp
+++ b/Engine/Source/Runtime/Engine/Private/Animation/AnimSequence.cpp
@@ -52,6 +52,19 @@
 #include "ITimeManagementModule.h"
 #include "CommonFrameRates.h"
 
+#include "HAL/FileManager.h"
+
+#if WITH_EDITOR
+#include "Animation/AnimData/AnimDataController.h"
+#include "Animation/AnimData/AnimDataModel.h"
+#endif // WITH_EDITOR
+
+#include "Animation/AnimSequenceHelpers.h"
+
+#include "Algo/Transform.h"
+#include "Math/UnrealMathUtility.h"
+#include "Animation/AnimCurveTypes.h"
+
 #define USE_SLERP 0
 #define LOCTEXT_NAMESPACE "AnimSequence"
 
@@ -134,7 +147,7 @@ void OnCVarsChanged()
 		const FCompressedAnimSequence& AnimData = Seq->CompressedData;
 		const FUECompressedAnimData& UEAnimData = AnimData.CompressedDataStructure;
 		const int32 Additive = Seq->IsValidAdditive() ? 1 : 0;
-		OutputMessage += FString::Printf(TEXT("%s - %.2f Fr:%i Add:%i TO:%i SO:%i CBS:%i\n"), *Seq->GetName(), Seq->GetPlayLength(), Seq->GetRawNumberOfFrames(), Additive, UEAnimData.CompressedTrackOffsets.Num(), UEAnimData.CompressedScaleOffsets.OffsetData.Num(), UEAnimData.CompressedByteStream.Num());
+		OutputMessage += FString::Printf(TEXT("%s - %.2f Fr:%i Add:%i TO:%i SO:%i CBS:%i\n"), *Seq->GetName(), Seq->SequenceLength, Seq->GetRawNumberOfFrames(), Additive, UEAnimData.CompressedTrackOffsets.Num(), UEAnimData.CompressedScaleOffsets.OffsetData.Num(), UEAnimData.CompressedByteStream.Num());
 		OutputMessage += FString::Printf(TEXT("\t K:%i (%i : %i : %i)\n"), (int32)UEAnimData.KeyEncodingFormat, (int32)UEAnimData.TranslationCompressionFormat, (int32)UEAnimData.RotationCompressionFormat, (int32)UEAnimData.ScaleCompressionFormat);
 		OutputMessage += FString::Printf(TEXT("\t Curve Codec:%s\n"), AnimData.CurveCompressionCodec ? *AnimData.CurveCompressionCodec->GetPathName() : TEXT("nullptr"));
 		OutputMessage += FString::Printf(TEXT("\t TrackOff:%s\n"), *GetArrayGuid<int32>(UEAnimData.CompressedTrackOffsets).ToString());
@@ -250,13 +263,17 @@ FString GetAnimSequenceSpecificCacheKeySuffix(const UAnimSequence& Seq, bool bPe
 
 	FString Ret = FString::Printf(TEXT("%i_%s%s%s_%c%c%i_%s_%s"),
 		Seq.CompressCommandletVersion,
+		PRAGMA_DISABLE_DEPRECATION_WARNINGS
 		*Seq.GetRawDataGuid().ToString(),
+		PRAGMA_ENABLE_DEPRECATION_WARNINGS
 		*Seq.GetSkeleton()->GetGuid().ToString(),
 		*Seq.GetSkeleton()->GetVirtualBoneGuid().ToString(),
 		AdditiveType,
 		RefType,
 		Seq.RefFrameIndex,
+		PRAGMA_DISABLE_DEPRECATION_WARNINGS
 		(bIsValidAdditive && Seq.RefPoseSeq) ? *Seq.RefPoseSeq->GetRawDataGuid().ToString() : TEXT("NoAdditiveGuid"),
+		PRAGMA_ENABLE_DEPRECATION_WARNINGS
 		*ArcToHexString.MakeString()
 	);
 
@@ -264,170 +281,6 @@ FString GetAnimSequenceSpecificCacheKeySuffix(const UAnimSequence& Seq, bool bPe
 }
 #endif
 
-bool CompressRawAnimSequenceTrack(FRawAnimSequenceTrack& RawTrack, int32 NumFrames, FName ErrorName, float MaxPosDiff, float MaxAngleDiff)
-{
-	bool bRemovedKeys = false;
-
-	// First part is to make sure we have valid input
-	bool const bPosTrackIsValid = (RawTrack.PosKeys.Num() == 1 || RawTrack.PosKeys.Num() == NumFrames);
-	if (!bPosTrackIsValid)
-	{
-		UE_LOG(LogAnimation, Warning, TEXT("Found non valid position track for %s, %d frames, instead of %d. Chopping!"), *ErrorName.ToString(), RawTrack.PosKeys.Num(), NumFrames);
-		bRemovedKeys = true;
-		RawTrack.PosKeys.RemoveAt(1, RawTrack.PosKeys.Num() - 1);
-		RawTrack.PosKeys.Shrink();
-		check(RawTrack.PosKeys.Num() == 1);
-	}
-
-	bool const bRotTrackIsValid = (RawTrack.RotKeys.Num() == 1 || RawTrack.RotKeys.Num() == NumFrames);
-	if (!bRotTrackIsValid)
-	{
-		UE_LOG(LogAnimation, Warning, TEXT("Found non valid rotation track for %s, %d frames, instead of %d. Chopping!"), *ErrorName.ToString(), RawTrack.RotKeys.Num(), NumFrames);
-		bRemovedKeys = true;
-		RawTrack.RotKeys.RemoveAt(1, RawTrack.RotKeys.Num() - 1);
-		RawTrack.RotKeys.Shrink();
-		check(RawTrack.RotKeys.Num() == 1);
-	}
-
-	// scale keys can be empty, and that is valid 
-	bool const bScaleTrackIsValid = (RawTrack.ScaleKeys.Num() == 0 || RawTrack.ScaleKeys.Num() == 1 || RawTrack.ScaleKeys.Num() == NumFrames);
-	if (!bScaleTrackIsValid)
-	{
-		UE_LOG(LogAnimation, Warning, TEXT("Found non valid Scaleation track for %s, %d frames, instead of %d. Chopping!"), *ErrorName.ToString(), RawTrack.ScaleKeys.Num(), NumFrames);
-		bRemovedKeys = true;
-		RawTrack.ScaleKeys.RemoveAt(1, RawTrack.ScaleKeys.Num() - 1);
-		RawTrack.ScaleKeys.Shrink();
-		check(RawTrack.ScaleKeys.Num() == 1);
-	}
-
-	// Second part is actual compression.
-
-	// Check variation of position keys
-	if ((RawTrack.PosKeys.Num() > 1) && (MaxPosDiff >= 0.0f))
-	{
-		FVector FirstPos = RawTrack.PosKeys[0];
-		bool bFramesIdentical = true;
-		for (int32 j = 1; j < RawTrack.PosKeys.Num() && bFramesIdentical; j++)
-		{
-			if ((FirstPos - RawTrack.PosKeys[j]).SizeSquared() > FMath::Square(MaxPosDiff))
-			{
-				bFramesIdentical = false;
-			}
-		}
-
-		// If all keys are the same, remove all but first frame
-		if (bFramesIdentical)
-		{
-			bRemovedKeys = true;
-			RawTrack.PosKeys.RemoveAt(1, RawTrack.PosKeys.Num() - 1);
-			RawTrack.PosKeys.Shrink();
-			check(RawTrack.PosKeys.Num() == 1);
-		}
-	}
-
-	// Check variation of rotational keys
-	if ((RawTrack.RotKeys.Num() > 1) && (MaxAngleDiff >= 0.0f))
-	{
-		FQuat FirstRot = RawTrack.RotKeys[0];
-		bool bFramesIdentical = true;
-		for (int32 j = 1; j < RawTrack.RotKeys.Num() && bFramesIdentical; j++)
-		{
-			if (FQuat::Error(FirstRot, RawTrack.RotKeys[j]) > MaxAngleDiff)
-			{
-				bFramesIdentical = false;
-			}
-		}
-
-		// If all keys are the same, remove all but first frame
-		if (bFramesIdentical)
-		{
-			bRemovedKeys = true;
-			RawTrack.RotKeys.RemoveAt(1, RawTrack.RotKeys.Num() - 1);
-			RawTrack.RotKeys.Shrink();
-			check(RawTrack.RotKeys.Num() == 1);
-		}
-	}
-
-	float MaxScaleDiff = 0.0001f;
-
-	// Check variation of Scaleition keys
-	if ((RawTrack.ScaleKeys.Num() > 1) && (MaxScaleDiff >= 0.0f))
-	{
-		FVector FirstScale = RawTrack.ScaleKeys[0];
-		bool bFramesIdentical = true;
-		for (int32 j = 1; j < RawTrack.ScaleKeys.Num() && bFramesIdentical; j++)
-		{
-			if ((FirstScale - RawTrack.ScaleKeys[j]).SizeSquared() > FMath::Square(MaxScaleDiff))
-			{
-				bFramesIdentical = false;
-			}
-		}
-
-		// If all keys are the same, remove all but first frame
-		if (bFramesIdentical)
-		{
-			bRemovedKeys = true;
-			RawTrack.ScaleKeys.RemoveAt(1, RawTrack.ScaleKeys.Num() - 1);
-			RawTrack.ScaleKeys.Shrink();
-			check(RawTrack.ScaleKeys.Num() == 1);
-		}
-	}
-
-	return bRemovedKeys;
-}
-
-bool StaticCompressRawAnimData(TArray<FRawAnimSequenceTrack>& RawAnimationData, int32 NumFrames, FName ErrorName, float MaxPosDiff, float MaxAngleDiff)
-{
-	bool bRemovedKeys = false;
-
-#if WITH_EDITORONLY_DATA
-	if (ensureMsgf(RawAnimationData.Num() > 0, TEXT("%s is trying to compress while raw animation is missing"), *ErrorName.ToString()))
-	{
-		// This removes trivial keys, and this has to happen before the removing tracks
-		for (int32 TrackIndex = 0; TrackIndex < RawAnimationData.Num(); TrackIndex++)
-		{
-			bRemovedKeys |= CompressRawAnimSequenceTrack(RawAnimationData[TrackIndex], NumFrames, ErrorName, MaxPosDiff, MaxAngleDiff);
-		}
-
-		bool bCompressScaleKeys = false;
-		// go through remove keys if not needed
-		for (int32 TrackIndex = 0; TrackIndex < RawAnimationData.Num(); TrackIndex++)
-		{
-			FRawAnimSequenceTrack const& RawData = RawAnimationData[TrackIndex];
-			if (RawData.ScaleKeys.Num() > 0)
-			{
-				// if scale key exists, see if we can just empty it
-				if ((RawData.ScaleKeys.Num() > 1) || (RawData.ScaleKeys[0].Equals(FVector(1.f)) == false))
-				{
-					bCompressScaleKeys = true;
-					break;
-				}
-			}
-		}
-
-		// if we don't have scale, we should delete all scale keys
-		// if you have one track that has scale, we still should support scale, so compress scale
-		if (!bCompressScaleKeys)
-		{
-			// then remove all scale keys
-			for (int32 TrackIndex = 0; TrackIndex < RawAnimationData.Num(); TrackIndex++)
-			{
-				FRawAnimSequenceTrack& RawData = RawAnimationData[TrackIndex];
-				RawData.ScaleKeys.Empty();
-			}
-		}
-	}
-#endif
-	return bRemovedKeys;
-}
-
-bool StaticCompressRawAnimData(TArray<FRawAnimSequenceTrack>& RawAnimationData, int32 NumFrames, FName ErrorName)
-{
-	const float MaxPosDiff = 0.0001f;
-	const float MaxAngleDiff = 0.0003f;
-	return StaticCompressRawAnimData(RawAnimationData, NumFrames, ErrorName, MaxPosDiff, MaxAngleDiff);
-}
-
 /////////////////////////////////////////////////////
 // FRequestAnimCompressionParams
 FRequestAnimCompressionParams::FRequestAnimCompressionParams(bool bInAsyncCompression, bool bInAllowAlternateCompressor, bool bInOutput)
@@ -582,6 +435,11 @@ UAnimSequence::UAnimSequence(const FObjectInitializer& ObjectInitializer)
 
 	CustomAttributesGuid.Invalidate();
 	BakedCustomAttributesGuid.Invalidate();
+
+	if (!HasAllFlags(EObjectFlags::RF_ClassDefaultObject))
+	{
+		check(DataModel);
+	}
 #endif
 }
 
@@ -604,11 +462,24 @@ void UAnimSequence::GetAssetRegistryTags(TArray<FAssetRegistryTag>& OutTags) con
 	{
 		OutTags.Add( FAssetRegistryTag(SourceFileTagName(), AssetImportData->GetSourceData().ToJson(), FAssetRegistryTag::TT_Hidden) );
 	}
-#endif
 
 	OutTags.Add(FAssetRegistryTag(TEXT("Compression Ratio"), FString::Printf(TEXT("%.03f"), (float)GetApproxCompressedSize() / (float)GetUncompressedRawSize()), FAssetRegistryTag::TT_Numerical));
+	if (DataModel)
+	{
+		OutTags.Add(FAssetRegistryTag(TEXT("Source Frame Rate"), FString::Printf(TEXT("%.2f"), GetDataModel()->GetFrameRate().AsDecimal()), FAssetRegistryTag::TT_Numerical));
+		OutTags.Add(FAssetRegistryTag(TEXT("Number of Frames"), FString::Printf(TEXT("%.i"), GetDataModel()->GetNumberOfFrames()), FAssetRegistryTag::TT_Numerical));
+		OutTags.Add(FAssetRegistryTag(TEXT("Number of Keys"), FString::Printf(TEXT("%i"), GetDataModel()->GetNumberOfKeys()), FAssetRegistryTag::TT_Numerical));
+	}
+	else
+	{
+		OutTags.Add(FAssetRegistryTag(TEXT("Source Frame Rate"), TEXT("0.0"), FAssetRegistryTag::TT_Numerical));
+		OutTags.Add(FAssetRegistryTag(TEXT("Number of Frames"), TEXT("0"), FAssetRegistryTag::TT_Numerical));
+		OutTags.Add(FAssetRegistryTag(TEXT("Number of Keys"), TEXT("0"), FAssetRegistryTag::TT_Numerical));
+	}
+#endif
+
 	OutTags.Add(FAssetRegistryTag(TEXT("Compressed Size (KB)"), FString::Printf(TEXT("%.02f"), (float)GetApproxCompressedSize() / 1024.0f), FAssetRegistryTag::TT_Numerical));
-	OutTags.Add(FAssetRegistryTag(TEXT("FrameRate"), FString::Printf(TEXT("%.2f"), SamplingFrameRate.AsDecimal()), FAssetRegistryTag::TT_Numerical));
+	OutTags.Add(FAssetRegistryTag(TEXT("Target Frame Rate"), FString::Printf(TEXT("%.2f"), TargetFrameRate.AsDecimal()), FAssetRegistryTag::TT_Numerical));
 	Super::GetAssetRegistryTags(OutTags);
 }
 
@@ -623,9 +494,13 @@ void UAnimSequence::AddReferencedObjects(UObject* This, FReferenceCollector& Col
 
 int32 UAnimSequence::GetUncompressedRawSize() const
 {
-	int32 BoneRawSize = ((sizeof(FVector) + sizeof(FQuat) + sizeof(FVector)) * RawAnimationData.Num() * NumberOfKeys);
+#if WITH_EDITOR
+	if (DataModel)
+	{
+		const int32 BoneRawSize = FRawAnimSequenceTrack::SingleKeySize * DataModel->GetNumBoneTracks() * DataModel->GetNumberOfKeys();
 	int32 CurveRawSize = 0;
-	for (const FFloatCurve& Curve : RawCurveData.FloatCurves)
+
+		for (const FFloatCurve& Curve : DataModel->GetCurveData().FloatCurves)
 	{
 		CurveRawSize += sizeof(FFloatCurve);
 		CurveRawSize += sizeof(FRichCurveKey) * Curve.FloatCurve.Keys.Num();
@@ -633,24 +508,35 @@ int32 UAnimSequence::GetUncompressedRawSize() const
 	return BoneRawSize + CurveRawSize;
 }
 
+#endif // WITH_EDITOR
+
+	return 0;
+}
+
 int32 UAnimSequence::GetApproxRawSize() const
 {
-	int32 Total = sizeof(FRawAnimSequenceTrack) * RawAnimationData.Num();
-	for (int32 i=0;i<RawAnimationData.Num();++i)
+#if WITH_EDITOR
+	if (DataModel)
 	{
-		const FRawAnimSequenceTrack& RawTrack = RawAnimationData[i];
-		Total +=
-			sizeof( FVector ) * RawTrack.PosKeys.Num() +
-			sizeof( FQuat ) * RawTrack.RotKeys.Num() + 
-			sizeof( FVector ) * RawTrack.ScaleKeys.Num(); 
+		int32 Total = sizeof(FRawAnimSequenceTrack) * DataModel->GetNumBoneTracks();
+		for (const FBoneAnimationTrack& AnimationTrack : DataModel->GetBoneAnimationTracks())
+	{
+			Total += AnimationTrack.InternalTrackData.PosKeys.Num() * AnimationTrack.InternalTrackData.PosKeys.GetTypeSize();
+			Total += AnimationTrack.InternalTrackData.RotKeys.Num() * AnimationTrack.InternalTrackData.RotKeys.GetTypeSize();
+			Total += AnimationTrack.InternalTrackData.ScaleKeys.Num() * AnimationTrack.InternalTrackData.ScaleKeys.GetTypeSize();
 	}
-	for (const FFloatCurve& Curve : RawCurveData.FloatCurves)
+
+		for (const FFloatCurve& Curve : DataModel->GetCurveData().FloatCurves)
 	{
 		Total += sizeof(FFloatCurve);
 		Total += sizeof(FRichCurveKey) * Curve.FloatCurve.Keys.Num();
 	}
 	return Total;
 }
+#endif // WITH_EDITOR
+
+	return 0;
+}
 
 int32 UAnimSequence::GetApproxBoneCompressedSize() const
 {
@@ -693,38 +579,15 @@ void UAnimSequence::Serialize(FArchive& Ar)
 	Ar.UsingCustomVersion(FFrameworkObjectVersion::GUID);
 	Ar.UsingCustomVersion(FUE5MainStreamObjectVersion::GUID);
 
-	FRawCurveTracks RawCurveCache;
-
-	if (Ar.IsCooking())
-	{
-		RawCurveCache.FloatCurves = MoveTemp(RawCurveData.FloatCurves);
-		RawCurveData.FloatCurves.Reset();
-
-#if WITH_EDITORONLY_DATA
-		RawCurveCache.VectorCurves = MoveTemp(RawCurveData.VectorCurves);
-		RawCurveData.VectorCurves.Reset();
-
-		RawCurveCache.TransformCurves = MoveTemp(RawCurveData.TransformCurves);
-		RawCurveData.TransformCurves.Reset();
-#endif
-	}
-
 	Super::Serialize(Ar);
 
-	if (Ar.IsCooking())
-	{
-		RawCurveData.FloatCurves = MoveTemp(RawCurveCache.FloatCurves);
-#if WITH_EDITORONLY_DATA
-		RawCurveData.VectorCurves = MoveTemp(RawCurveCache.VectorCurves);
-		RawCurveData.TransformCurves = MoveTemp(RawCurveCache.TransformCurves);
-#endif
-	}
-
 	FStripDataFlags StripFlags( Ar );
 	if( !StripFlags.IsEditorDataStripped() )
 	{
-		Ar << RawAnimationData;
 #if WITH_EDITORONLY_DATA
+		PRAGMA_DISABLE_DEPRECATION_WARNINGS
+		Ar << RawAnimationData;
+		PRAGMA_ENABLE_DEPRECATION_WARNINGS
 		if (!Ar.IsCooking())
 		{
 			if (Ar.UE4Ver() >= VER_UE4_ANIMATION_ADD_TRACKCURVES)
@@ -742,9 +605,9 @@ void UAnimSequence::Serialize(FArchive& Ar)
 						MarkRawDataAsModified();
 					}
 					PRAGMA_ENABLE_DEPRECATION_WARNINGS
+				}
 			}
 		}
-		}
 #endif // WITH_EDITORONLY_DATA
 	}
 
@@ -820,9 +683,9 @@ void UAnimSequence::Serialize(FArchive& Ar)
 			const float SampledTime = SamplingFrameRate.AsSeconds(NumberOfFrames);
 
 			// Validate that provided the number of keys and frame-rate the sequence length is accurate
-			if (!FMath::IsNearlyEqual(SampledTime, GetPlayLength(), KINDA_SMALL_NUMBER))
+			if (!FMath::IsNearlyEqual(SampledTime, SequenceLength, KINDA_SMALL_NUMBER))
 			{
-				UE_LOG(LogAnimation, Warning, TEXT("Animation %s needs resaving - Inaccurate sequence length %5.5f according to number of frames and frame rate (delta of %f). Setting new length %5.5f."), *GetName(), GetPlayLength(), SampledTime - GetPlayLength(), SampledTime);
+				UE_LOG(LogAnimation, Warning, TEXT("Animation %s needs resaving - Inaccurate sequence length %5.5f according to number of frames and frame rate (delta of %f). Setting new length %5.5f."), *GetName(), SequenceLength, SampledTime - SequenceLength, SampledTime);
 				SequenceLength = SampledTime;
 			}
 		}
@@ -846,7 +709,6 @@ void UAnimSequence::Serialize(FArchive& Ar)
 		SourceFilePath_DEPRECATED = TEXT("");
 		SourceFileTimestamp_DEPRECATED = TEXT("");
 	}
-
 #endif // WITH_EDITORONLY_DATA
 }
 
@@ -856,26 +718,14 @@ bool UAnimSequence::IsValidToPlay() const
 	// make sure sequence length is valid and raw animation data exists, and compressed
 	return (GetPlayLength() > 0.f);
 }
-#endif
-
-void UAnimSequence::SetNumberOfSampledKeys(int32 InNumberOfKeys)
-{
-	NumberOfKeys = InNumberOfKeys;	
-	UpdateFrameRate();
-}
-
-void UAnimSequence::SetSequenceLength(float NewLength)
-{
-	UAnimSequenceBase::SetSequenceLength(NewLength);
-	UpdateFrameRate();
-}
 
 void UAnimSequence::UpdateFrameRate()
 {
+	PRAGMA_DISABLE_DEPRECATION_WARNINGS
 	const int32 NumberOfFrames = FMath::Max(NumberOfKeys - 1, 1);
 
 	// Generate the frame-rate according to the number of frames and sequence length
-	const double DecimalFrameRate = (double)NumberOfFrames / ((double)GetPlayLength() > 0.0 ? (double)GetPlayLength() : 1.0);
+	const double DecimalFrameRate = (double)NumberOfFrames / ((double)SequenceLength > 0.0 ? (double)SequenceLength : 1.0);
 
 	// Account for non-whole number frame rates using large denominator
 	const double Denominator = 1000000.0;
@@ -893,7 +743,9 @@ void UAnimSequence::UpdateFrameRate()
 			break;
 		}
 	}
+	PRAGMA_ENABLE_DEPRECATION_WARNINGS
 }
+#endif
 
 void UAnimSequence::SortSyncMarkers()
 {
@@ -971,11 +823,13 @@ void UAnimSequence::PostLoad()
 	Super::PostLoad();
 
 #if WITH_EDITOR
-	if (!RawDataGuid.IsValid())
+	if (RefPoseSeq && RefPoseSeq->HasAnyFlags(EObjectFlags::RF_NeedPostLoad))
 	{
-		RawDataGuid = GenerateGuidFromRawData();
+		RefPoseSeq->ConditionalPostLoad();
 	}
 
+	EnsureValidRawDataGuid();
+
 	// I have to do this first thing in here
 	// so that remove all NaNs before even being read
 	if(GetLinkerUE4Version() < VER_UE4_ANIMATION_REMOVE_NANS)
@@ -983,8 +837,7 @@ void UAnimSequence::PostLoad()
 		RemoveNaNTracks();
 	}
 
-	VerifyTrackMap(nullptr);
-
+	Controller->RemoveBoneTracksMissingFromSkeleton(GetSkeleton());
 #endif // WITH_EDITOR
 
 	// if valid additive, but if base additive isn't 
@@ -1006,7 +859,9 @@ void UAnimSequence::PostLoad()
 	if (ForcedRecompressionSetting)
 	{
 		//Force recompression
+		PRAGMA_DISABLE_DEPRECATION_WARNINGS
 		RawDataGuid = FGuid::NewGuid();
+		PRAGMA_ENABLE_DEPRECATION_WARNINGS
 		bUseRawDataOnly = true;
 	}
 
@@ -1022,10 +877,11 @@ void UAnimSequence::PostLoad()
 	// No animation data is found. Warn - this should check before we check CompressedTrackOffsets size
 	// Otherwise, we'll see empty data set crashing game due to no CompressedTrackOffsets
 	// You can't check RawAnimationData size since it gets removed during cooking
-	if ( NumberOfKeys == 0 && RawCurveData.FloatCurves.Num() == 0 )
+#if WITH_EDITOR
+	if (DataModel->GetNumberOfKeys() == 0 && DataModel->GetNumberOfFloatCurves() == 0 )
 	{
 		UE_LOG(LogAnimation, Warning, TEXT("No animation data exists for sequence %s (%s)"), *GetName(), (GetOuter() ? *GetOuter()->GetFullName() : *GetFullName()) );
-#if WITH_EDITOR
+
 		if (!IsRunningGame())
 		{
 			static FName NAME_LoadErrors("LoadErrors");
@@ -1037,17 +893,18 @@ void UAnimSequence::PostLoad()
 			Message->AddToken(FTextToken::Create(LOCTEXT("EmptyAnimationData2", " has no animation data. Recommend to remove.")));
 			LoadErrors.Notify();
 		}
-#endif
 	}
 	// @remove temp hack for fixing length
 	// @todo need to fix importer/editing feature
 	else if (GetPlayLength() == 0.f )
 	{
-		ensure(NumberOfKeys == 1);
-		SetSequenceLength(MINIMUM_ANIMATION_LENGTH);
+		ensure(DataModel->GetNumberOfKeys() == 1);
+		Controller->SetPlayLength(MINIMUM_ANIMATION_LENGTH);
 	}
 	// Raw data exists, but missing compress animation data
-	else if( GetSkeleton() && !IsCompressedDataValid()
+	else
+#endif
+	if( GetSkeleton() && !IsCompressedDataValid()
 #if WITH_EDITOR
 		&& !bCompressionInProgress
 #endif
@@ -1056,31 +913,13 @@ void UAnimSequence::PostLoad()
 		UE_LOG(LogAnimation, Fatal, TEXT("No animation compression exists for sequence %s (%s)"), *GetName(), (GetOuter() ? *GetOuter()->GetFullName() : *GetFullName()) );
 	}
 
-	// If we're in the game and compressed animation data exists, whack the raw data.
+	// Update the virtual bone GUID according to the skeleton
 	if (FPlatformProperties::RequiresCookedData())
 	{
 		if (GetSkeleton())
 		{
 			SetSkeletonVirtualBoneGuid(GetSkeleton()->GetVirtualBoneGuid());
 		}
-		if( RawAnimationData.Num() > 0  && CompressedData.CompressedByteStream.Num() > 0 )
-		{
-#if [REDACTED_EMAIL]/Package...
-			// Don't do this on consoles; raw animation data should have been stripped during cook!
-			UE_LOG(LogAnimation, Fatal, TEXT("Cooker did not strip raw animation from sequence %s"), *GetName() );
-#else
-			// Remove raw animation data.
-			for ( int32 TrackIndex = 0 ; TrackIndex < RawAnimationData.Num() ; ++TrackIndex )
-			{
-				FRawAnimSequenceTrack& RawTrack = RawAnimationData[TrackIndex];
-				RawTrack.PosKeys.Empty();
-				RawTrack.RotKeys.Empty();
-				RawTrack.ScaleKeys.Empty();
-			}
-			
-			RawAnimationData.Empty();
-#endif
-		}
 	}
 
 	for(FAnimNotifyEvent& Notify : Notifies)
@@ -1107,9 +946,9 @@ void UAnimSequence::PostLoad()
 		{
 			CurrentSkeleton->VerifySmartName(USkeleton::AnimCurveMappingName, CurveName);
 		}
-
 #if WITH_EDITOR
-		VerifyCurveNames<FTransformCurve>(*CurrentSkeleton, USkeleton::AnimTrackCurveMappingName, RawCurveData.TransformCurves);
+		static const bool DoNotTransactAction = false;
+		Controller->FindOrAddCurveNamesOnSkeleton(CurrentSkeleton, ERawCurveTrackTypes::RCT_Transform, DoNotTransactAction);
 
 		for (const FAnimSyncMarker& SyncMarker : AuthoredSyncMarkers)
 		{
@@ -1126,80 +965,26 @@ void ShowResaveMessage(const UAnimSequence* Sequence)
 	{
 		UE_LOG(LogAnimation, Log, TEXT("Resave Animation Required(%s, %s): Fixing track data and recompressing."), *GetNameSafe(Sequence), *Sequence->GetPathName());
 
-		/*static FName NAME_LoadErrors("LoadErrors");
+		static FName NAME_LoadErrors("LoadErrors");
 		FMessageLog LoadErrors(NAME_LoadErrors);
 
 		TSharedRef<FTokenizedMessage> Message = LoadErrors.Warning();
 		Message->AddToken(FTextToken::Create(LOCTEXT("AnimationNeedsResave1", "The Animation ")));
 		Message->AddToken(FAssetNameToken::Create(Sequence->GetPathName(), FText::FromString(GetNameSafe(Sequence))));
 		Message->AddToken(FTextToken::Create(LOCTEXT("AnimationNeedsResave2", " needs resave.")));
-		LoadErrors.Notify();*/
+		LoadErrors.Notify();
 	}
 }
 
 void UAnimSequence::VerifyTrackMap(USkeleton* MySkeleton)
 {
+	ValidateModel();
+
 	USkeleton* UseSkeleton = (MySkeleton) ? MySkeleton : GetSkeleton();
 
-	if( AnimationTrackNames.Num() != TrackToSkeletonMapTable.Num() && UseSkeleton!=nullptr)
+	if(Controller->RemoveBoneTracksMissingFromSkeleton(UseSkeleton))
 	{
 		ShowResaveMessage(this);
-
-		AnimationTrackNames.Empty();
-		AnimationTrackNames.AddUninitialized(TrackToSkeletonMapTable.Num());
-		for(int32 I=0; I<TrackToSkeletonMapTable.Num(); ++I)
-		{
-			const FTrackToSkeletonMap& TrackMap = TrackToSkeletonMapTable[I];
-			AnimationTrackNames[I] = UseSkeleton->GetReferenceSkeleton().GetBoneName(TrackMap.BoneTreeIndex);
-		}
-	}
-	else if (UseSkeleton != nullptr)
-	{
-		// first check if any of them needs to be removed
-		{
-			int32 NumTracks = AnimationTrackNames.Num();
-			int32 NumSkeletonBone = UseSkeleton->GetReferenceSkeleton().GetRawBoneNum();
-
-			// the first fix is to make sure 
-			bool bNeedsFixing = false;
-			// verify all tracks are still valid
-			for(int32 TrackIndex=0; TrackIndex<NumTracks; TrackIndex++)
-			{
-				int32 SkeletonBoneIndex = TrackToSkeletonMapTable[TrackIndex].BoneTreeIndex;
-				// invalid index found
-				if(SkeletonBoneIndex == INDEX_NONE || NumSkeletonBone <= SkeletonBoneIndex)
-				{
-					// if one is invalid, fix up for all.
-					// you don't know what index got messed up
-					bNeedsFixing = true;
-					break;
-				}
-			}
-
-			if(bNeedsFixing)
-			{
-				ShowResaveMessage(this);
-
-				for(int32 I=NumTracks-1; I>=0; --I)
-				{
-					int32 BoneTreeIndex = UseSkeleton->GetReferenceSkeleton().FindBoneIndex(AnimationTrackNames[I]);
-					if(BoneTreeIndex == INDEX_NONE)
-					{
-						RemoveTrack(I);
-					}
-					else
-					{
-						TrackToSkeletonMapTable[I].BoneTreeIndex = BoneTreeIndex;
-					}
-				}
-			}
-		}
-		
-		for(int32 I=0; I<AnimationTrackNames.Num(); ++I)
-		{
-			FTrackToSkeletonMap& TrackMap = TrackToSkeletonMapTable[I];
-			TrackMap.BoneTreeIndex = UseSkeleton->GetReferenceSkeleton().FindBoneIndex(AnimationTrackNames[I]);
-		}		
 	}
 }
 
@@ -1251,6 +1036,7 @@ void UAnimSequence::PostEditChangeProperty(FPropertyChangedEvent& PropertyChange
 
 	bool bAdditiveSettingsChanged = false;
 	bool bCompressionAffectingSettingsChanged = false;
+	bool bShouldResample = false;
 
 	if(PropertyChangedEvent.Property)
 	{
@@ -1275,28 +1061,30 @@ void UAnimSequence::PostEditChangeProperty(FPropertyChangedEvent& PropertyChange
 		}
 		
 		bCompressionAffectingSettingsChanged =   PropertyChangedEvent.Property->GetFName() == GET_MEMBER_NAME_CHECKED(UAnimSequence, bAllowFrameStripping)
-											  || PropertyChangedEvent.Property->GetFName() == GET_MEMBER_NAME_CHECKED(UAnimSequence, CompressionErrorThresholdScale);
-	}
+											  || PropertyChangedEvent.Property->GetFName() == GET_MEMBER_NAME_CHECKED(UAnimSequence, CompressionErrorThresholdScale)
+											  || PropertyChangedEvent.Property->GetFName() == GET_MEMBER_NAME_CHECKED(UAnimSequence, BoneCompressionSettings)
+											  || PropertyChangedEvent.Property->GetFName() == GET_MEMBER_NAME_CHECKED(UAnimSequence, BoneCompressionSettings);
 
-	const bool bNeedPostProcess = !IsCompressedDataValid() || bAdditiveSettingsChanged || bCompressionAffectingSettingsChanged;
+		bShouldResample = PropertyChangedEvent.Property->GetFName() == GET_MEMBER_NAME_CHECKED(UAnimSequence, TargetFrameRate);
+	}
 
-	// @Todo fix me: This is temporary fix to make sure they always have compressed data
-	if (RawAnimationData.Num() > 0 && bNeedPostProcess)
+	if (bShouldResample)
 	{
-		PostProcessSequence(false);
+		ResampleAnimationTrackData();
 	}
 
-	if (PropertyChangedEvent.Property != nullptr)
-	{
-		if (PropertyChangedEvent.Property->GetFName() == GET_MEMBER_NAME_CHECKED(UAnimSequence, CurveCompressionSettings))
-		{
-			RequestSyncAnimRecompression(false);
-		}
+	REFACTOR_V2_WARNING("V2 we should not rely on PostEditChange to check for compressed data");
+	const bool bNeedPostProcess = !IsCompressedDataValid() || bAdditiveSettingsChanged || bCompressionAffectingSettingsChanged || bShouldResample;
 
-		if (PropertyChangedEvent.Property->GetFName() == GET_MEMBER_NAME_CHECKED(UAnimSequence, BoneCompressionSettings))
+	// @Todo fix me: This is temporary fix to make sure they always have compressed data
+	ValidateModel();
+	if (DataModel->GetNumBoneTracks() && bNeedPostProcess)
 		{
-			RequestSyncAnimRecompression(false);
-		}
+		// QQ need new function here to queue async compression
+		PRAGMA_DISABLE_DEPRECATION_WARNINGS
+		MarkRawDataAsModified(false);
+		PRAGMA_ENABLE_DEPRECATION_WARNINGS
+		RecompressAnimationData();
 	}
 }
 #endif // WITH_EDITOR
@@ -1326,7 +1114,9 @@ void UAnimSequence::GetBoneTransform(FTransform& OutAtom, int32 TrackIndex, floa
 	// If the caller didn't request that raw animation data be used . . .
 	if ( !bUseRawData && IsCompressedDataValid() )
 	{
-		FAnimSequenceDecompressionContext DecompContext(GetPlayLength(), Interpolation, GetFName(), *CompressedData.CompressedDataStructure);
+		PRAGMA_DISABLE_DEPRECATION_WARNINGS
+		FAnimSequenceDecompressionContext DecompContext(SequenceLength, Interpolation, GetFName(), *CompressedData.CompressedDataStructure);
+		PRAGMA_ENABLE_DEPRECATION_WARNINGS
 		DecompContext.Seek(Time);
 		if (CompressedData.BoneCompressionCodec)
 		{
@@ -1335,7 +1125,10 @@ void UAnimSequence::GetBoneTransform(FTransform& OutAtom, int32 TrackIndex, floa
 		}
 	}
 
-	ExtractBoneTransform(RawAnimationData, OutAtom, TrackIndex, Time);
+#if WITH_EDITOR
+	ValidateModel();
+	UE::Anim::GetBoneTransformFromModel(DataModel, OutAtom, TrackIndex, Time, Interpolation);
+#endif
 }
 
 void UAnimSequence::GetBoneTransform(FTransform& OutAtom, int32 TrackIndex, FAnimSequenceDecompressionContext& DecompContext, bool bUseRawData) const
@@ -1350,52 +1143,39 @@ void UAnimSequence::GetBoneTransform(FTransform& OutAtom, int32 TrackIndex, FAni
 		}
 	}
 
-	ExtractBoneTransform(RawAnimationData, OutAtom, TrackIndex, DecompContext.Time);
+#if WITH_EDITOR
+	ValidateModel();
+	UE::Anim::GetBoneTransformFromModel(DataModel, OutAtom, TrackIndex, DecompContext.Time, Interpolation);
+#endif
 }
 
 void UAnimSequence::ExtractBoneTransform(const TArray<struct FRawAnimSequenceTrack>& InRawAnimationData, FTransform& OutAtom, int32 TrackIndex, float Time) const
 {
 	// Bail out if the animation data doesn't exists (e.g. was stripped by the cooker).
-	if(InRawAnimationData.Num() == 0)
+	if (InRawAnimationData.IsValidIndex(TrackIndex))
 	{
-		UE_LOG(LogAnimation, Log, TEXT("UAnimSequence::GetBoneTransform : No anim data in AnimSequence[%s]!"),*GetFullName());
+		UE_LOG(LogAnimation, Log, TEXT("UAnimSequence::GetBoneTransform : No valid anim data track!"));
 		OutAtom.SetIdentity();
 		return;
 	}
 
-	ExtractBoneTransform(InRawAnimationData[TrackIndex], OutAtom, Time);
+	PRAGMA_DISABLE_DEPRECATION_WARNINGS
+	return ExtractBoneTransform(InRawAnimationData[TrackIndex], OutAtom, Time);
+	PRAGMA_ENABLE_DEPRECATION_WARNINGS
 }
 
 void UAnimSequence::ExtractBoneTransform(const struct FRawAnimSequenceTrack& RawTrack, FTransform& OutAtom, int32 KeyIndex) const
 {
-	// Bail out (with rather wacky data) if data is empty for some reason.
-	if (RawTrack.PosKeys.Num() == 0 || RawTrack.RotKeys.Num() == 0)
-	{
-		UE_LOG(LogAnimation, Log, TEXT("UAnimSequence::GetBoneTransform : No anim data in AnimSequence!"));
-		OutAtom.SetIdentity();
-		return;
-	}
-
-	const int32 PosKeyIndex = FMath::Min(KeyIndex, RawTrack.PosKeys.Num() - 1);
-	const int32 RotKeyIndex = FMath::Min(KeyIndex, RawTrack.RotKeys.Num() - 1);
-	static const FVector DefaultScale3D = FVector(1.f);
-
-	OutAtom.SetTranslation(RawTrack.PosKeys[PosKeyIndex]);
-	OutAtom.SetRotation(RawTrack.RotKeys[RotKeyIndex]);
-	if (RawTrack.ScaleKeys.Num() > 0)
-	{
-		const int32 ScaleKeyIndex = FMath::Min(KeyIndex, RawTrack.ScaleKeys.Num() - 1);
-		OutAtom.SetScale3D(RawTrack.ScaleKeys[ScaleKeyIndex]);
-	}
-	else
-	{
-		OutAtom.SetScale3D(DefaultScale3D);
-	}
+	UE::Anim::ExtractBoneTransform(RawTrack, OutAtom, KeyIndex);
 }
 
 void UAnimSequence::ExtractBoneTransform(const struct FRawAnimSequenceTrack& RawTrack, FTransform& OutAtom, float Time) const
 {
-	FAnimationUtils::ExtractTransformFromTrack(Time, NumberOfKeys, GetPlayLength(), RawTrack, Interpolation, OutAtom);
+#if WITH_EDITOR
+	FAnimationUtils::ExtractTransformFromTrack(Time, NumberOfSampledKeys, GetPlayLength(), RawTrack, Interpolation, OutAtom);
+#else
+	FAnimationUtils::ExtractTransformFromTrack(Time, CompressedData.CompressedDataStructure->CompressedNumberOfKeys, GetPlayLength(), RawTrack, Interpolation, OutAtom);
+#endif
 }
 
 void UAnimSequence::HandleAssetPlayerTickedInternal(FAnimAssetTickContext &Context, const float PreviousTime, const float MoveDelta, const FAnimTickRecord &Instance, struct FAnimNotifyQueue& NotifyQueue) const
@@ -1410,14 +1190,37 @@ void UAnimSequence::HandleAssetPlayerTickedInternal(FAnimAssetTickContext &Conte
 
 FTransform UAnimSequence::ExtractRootTrackTransform(float Pos, const FBoneContainer * RequiredBones) const
 {
-	const TArray<FTrackToSkeletonMap> & TrackToSkeletonMap = bUseRawDataOnly ? TrackToSkeletonMapTable : CompressedData.CompressedTrackToSkeletonMapTable;
+	const int32 RootBoneIndex = 0;
+	const int32 TrackIndex = [this, RootBoneIndex]() -> int32
+	{
+#if WITH_EDITOR
+		if (bUseRawDataOnly)
+		{
+			ValidateModel();
+			if (const FBoneAnimationTrack* RootTrack = DataModel->FindBoneTrackByIndex(RootBoneIndex))
+			{
+				return DataModel->GetBoneTrackIndex(*RootTrack);
+			}
+		}
+		else
+#endif // WITH_EDITOR
+		{
+			const TArray<FTrackToSkeletonMap>& TrackToSkeletonMap = CompressedData.CompressedTrackToSkeletonMapTable;
+			// we assume root is in first data if available = SkeletonIndex == 0 && BoneTreeIndex == 0)
+			if ((TrackToSkeletonMap.Num() > 0) && (TrackToSkeletonMap[0].BoneTreeIndex == RootBoneIndex))
+			{
+				return 0;
+			}
+		}
+
+		return INDEX_NONE;
+	}();
 
-	// we assume root is in first data if available = SkeletonIndex == 0 && BoneTreeIndex == 0)
-	if ((TrackToSkeletonMap.Num() > 0) && (TrackToSkeletonMap[0].BoneTreeIndex == 0))
+	if(TrackIndex != INDEX_NONE)
 	{
 		// if we do have root data, then return root data
 		FTransform RootTransform;
-		GetBoneTransform(RootTransform, 0, Pos, bUseRawDataOnly);
+		GetBoneTransform(RootTransform, TrackIndex, Pos, bUseRawDataOnly);
 		return RootTransform;
 	}
 
@@ -1427,7 +1230,7 @@ FTransform UAnimSequence::ExtractRootTrackTransform(float Pos, const FBoneContai
 		const FReferenceSkeleton& RefSkeleton = RequiredBones->GetReferenceSkeleton();
 		if( RefSkeleton.GetNum() > 0 )
 		{
-			return RefSkeleton.GetRefBonePose()[0];
+			return RefSkeleton.GetRefBonePose()[RootBoneIndex];
 		}
 	}
 
@@ -1438,7 +1241,7 @@ FTransform UAnimSequence::ExtractRootTrackTransform(float Pos, const FBoneContai
 		const FReferenceSkeleton& RefSkeleton = MySkeleton->GetReferenceSkeleton();
 		if( RefSkeleton.GetNum() > 0 )
 		{
-			return RefSkeleton.GetRefBonePose()[0];
+			return RefSkeleton.GetRefBonePose()[RootBoneIndex];
 		}
 	}
 
@@ -1670,7 +1473,12 @@ void UAnimSequence::GetBonePose(FAnimationPoseData& OutAnimationPoseData, const
 	// extract curve data . Even if no track, it can contain curve data
 	EvaluateCurveData(OutAnimationPoseData.GetCurve(), ExtractionContext.CurrentTime, bUseRawDataForPoseExtraction);
 
-	const int32 NumTracks = bUseRawDataForPoseExtraction ? TrackToSkeletonMapTable.Num() : CompressedData.CompressedTrackToSkeletonMapTable.Num();
+#if WITH_EDITOR
+	ValidateModel();
+	const int32 NumTracks = bUseRawDataForPoseExtraction ? DataModel->GetNumBoneTracks() : CompressedData.CompressedTrackToSkeletonMapTable.Num();
+#else
+	const int32 NumTracks = CompressedData.CompressedTrackToSkeletonMapTable.Num();
+#endif 
 	if (NumTracks == 0)
 	{
 		return;
@@ -1684,41 +1492,21 @@ void UAnimSequence::GetBonePose(FAnimationPoseData& OutAnimationPoseData, const
 	// Slower path for disable retargeting, that's only used in editor and for debugging.
 	if (bUseRawDataForPoseExtraction)
 	{
-		const bool bShouldUseSourceData = RequiredBones.ShouldUseSourceData();
-		const TArray<FRawAnimSequenceTrack>& AnimationData = RawAnimationData;
-
 		// Warning if we have invalid data
 		for (int32 TrackIndex = 0; TrackIndex < NumTracks; TrackIndex++)
 		{
-			const FRawAnimSequenceTrack& TrackToExtract = AnimationData[TrackIndex];
+			const FBoneAnimationTrack& BoneAnimationTrack = DataModel->GetBoneTrackByIndex(TrackIndex);
+			const FRawAnimSequenceTrack& TrackToExtract = BoneAnimationTrack.InternalTrackData;
 
 			// Bail out (with rather wacky data) if data is empty for some reason.
 			if (TrackToExtract.PosKeys.Num() == 0 || TrackToExtract.RotKeys.Num() == 0)
 			{
-				UE_LOG(LogAnimation, Warning, TEXT("UAnimSequence::GetBoneTransform : No anim data in AnimSequence '%s' Track '%s'"), *GetPathName(), *AnimationTrackNames[TrackIndex].ToString());
-			}
-		}
-
-		TMap<int32, const FTransformCurve*> ActiveBoneCurves;
-		if (!bShouldUseSourceData)
-		{
-			const FSmartNameMapping* NameMapping = MySkeleton->GetSmartNameContainer(USkeleton::AnimTrackCurveMappingName);
-			for (const FTransformCurve& Curve : RawCurveData.TransformCurves)
-			{
-				FName BoneName;
-				if (Curve.GetCurveTypeFlag(AACF_Disabled) == false &&
-					ensureAlways(NameMapping->GetName(Curve.Name.UID, BoneName)))
-				{
-					const int32 BoneIndex = MySkeleton->GetReferenceSkeleton().FindBoneIndex(BoneName);
-					if (BoneIndex != INDEX_NONE)
-					{
-						ActiveBoneCurves.Add(BoneIndex, &Curve);
-					}
-				}
+				UE_LOG(LogAnimation, Warning, TEXT("UAnimSequence::GetBoneTransform : No anim data in AnimSequence '%s' Track '%s'"), *GetPathName(), 
+					*BoneAnimationTrack.Name.ToString());
 			}
 		}
 
-		BuildPoseFromRawData(AnimationData, TrackToSkeletonMapTable, OutPose, ExtractionContext.CurrentTime, Interpolation, NumberOfKeys, GetPlayLength(), GetRetargetTransformsSourceName(), GetRetargetTransforms());
+		UE::Anim::BuildPoseFromModel(DataModel, OutPose, ExtractionContext.CurrentTime, Interpolation, GetRetargetTransformsSourceName(), GetRetargetTransforms());
 
 		if ((ExtractionContext.bExtractRootMotion && RootMotionReset.bEnableRootMotion) || RootMotionReset.bForceRootLock)
 		{
@@ -1736,13 +1524,43 @@ void UAnimSequence::GetBonePose(FAnimationPoseData& OutAnimationPoseData, const
 	GetCustomAttributes(OutAnimationPoseData, ExtractionContext, false);
 }
 
-#if WITH_EDITORONLY_DATA
+const TArray<FRawAnimSequenceTrack>& UAnimSequence::GetRawAnimationData() const
+{
+#if WITH_EDITOR
+	ValidateModel();
+	PRAGMA_DISABLE_DEPRECATION_WARNINGS
+	return DataModel->GetTransientRawAnimationTracks();
+	PRAGMA_ENABLE_DEPRECATION_WARNINGS
+#else
+	static const TArray<FRawAnimSequenceTrack> Tracks;
+	check(false);
+	return Tracks;
+#endif
+}
 
-void UAnimSequence::UpdateCompressedCurveName(SmartName::UID_Type CurveUID, const struct FSmartName& NewCurveName)
+
+#if WITH_EDITORONLY_DATA
+const TArray<FName>& UAnimSequence::GetAnimationTrackNames() const
 {
-	for (FSmartName& CurveName : CompressedData.CompressedCurveNames)
-	{
-		if (CurveName.UID == CurveUID)
+#if WITH_EDITOR
+	ValidateModel();
+	PRAGMA_DISABLE_DEPRECATION_WARNINGS
+	return DataModel->GetTransientRawAnimationTrackNames();
+	PRAGMA_ENABLE_DEPRECATION_WARNINGS
+#else
+	static const TArray<FName> Names;
+	check(false);
+	return Names;
+#endif 
+}
+#endif 
+
+#if WITH_EDITOR
+void UAnimSequence::UpdateCompressedCurveName(SmartName::UID_Type CurveUID, const struct FSmartName& NewCurveName)
+{
+	for (FSmartName& CurveName : CompressedData.CompressedCurveNames)
+	{
+		if (CurveName.UID == CurveUID)
 		{
 			CurveName = NewCurveName;
 			break;
@@ -1752,54 +1570,36 @@ void UAnimSequence::UpdateCompressedCurveName(SmartName::UID_Type CurveUID, cons
 
 int32 UAnimSequence::AddNewRawTrack(FName TrackName, FRawAnimSequenceTrack* TrackData)
 {
-	const int32 SkeletonIndex = GetSkeleton() ? GetSkeleton()->GetReferenceSkeleton().FindBoneIndex(TrackName) : INDEX_NONE;
+	ValidateModel();
 
-	if(SkeletonIndex != INDEX_NONE)
+	int32 NewTrackIndex = INDEX_NONE;
+	Controller->OpenBracket(LOCTEXT("AddNewRawTrack_Bracket", "Adding new Bone Animation Track"));
 	{
-		int32 TrackIndex = AnimationTrackNames.IndexOfByKey(TrackName);
-		if (TrackIndex != INDEX_NONE)
+		NewTrackIndex = Controller->AddBoneTrack(TrackName);
+		if (NewTrackIndex != INDEX_NONE && TrackData)
 		{
-			if (TrackData)
-			{
-				RawAnimationData[TrackIndex] = *TrackData;
-			}
-			return TrackIndex;
+			Controller->SetBoneTrackKeys(TrackName, TrackData->PosKeys, TrackData->RotKeys, TrackData->ScaleKeys);
 		}
+	}
+	Controller->CloseBracket();
 
-		// During compression, we store the track indices on 16 bits
-		const int32 MAX_NUM_TRACKS = 65535;
-		if (RawAnimationData.Num() >= MAX_NUM_TRACKS)
-		{
-			FFormatNamedArguments Args;
-			Args.Add(TEXT("TrackName"), FText::FromName(TrackName));
-			FNotificationInfo Info(FText::Format(LOCTEXT("TooManyRawTracks", "Cannot add track with name \"{TrackName}\". An animation sequence cannot contain more than 65535 tracks."), Args));
-
-			Info.bUseLargeFont = false;
-			Info.ExpireDuration = 5.0f;
+	if (NewTrackIndex == INDEX_NONE)
+	{
+		FFormatNamedArguments Args;
+		Args.Add(TEXT("TrackName"), FText::FromName(TrackName));
+		FNotificationInfo Info(FText::Format(LOCTEXT("FailedToAddTracks", "Failed to add track with name \"{TrackName}\"."), Args));
 
-			TSharedPtr<SNotificationItem> Notification = FSlateNotificationManager::Get().AddNotification(Info);
-			if (Notification.IsValid())
-			{
-				Notification->SetCompletionState(SNotificationItem::CS_Fail);
-			}
+		Info.bUseLargeFont = false;
+		Info.ExpireDuration = 5.0f;
 
-			UE_LOG(LogAnimation, Error, TEXT("Cannot add track with name \"%s\". An animation sequence cannot contain more than 65535 tracks."), *TrackName.ToString());
-		}
-
-		check(AnimationTrackNames.Num() == RawAnimationData.Num());
-		TrackIndex = AnimationTrackNames.Add(TrackName);
-		TrackToSkeletonMapTable.Add(FTrackToSkeletonMap(SkeletonIndex));
-		if (TrackData)
-		{
-			RawAnimationData.Add(*TrackData);
-		}
-		else
+		TSharedPtr<SNotificationItem> Notification = FSlateNotificationManager::Get().AddNotification(Info);
+		if (Notification.IsValid())
 		{
-			RawAnimationData.Add(FRawAnimSequenceTrack());
+			Notification->SetCompletionState(SNotificationItem::CS_Fail);
 		}
-		return TrackIndex;
 	}
-	return INDEX_NONE;
+
+	return NewTrackIndex;
 }
 #endif
 
@@ -1861,8 +1661,10 @@ void UAnimSequence::GetAdditiveBasePose(FAnimationPoseData& OutAnimationPoseData
 		case ABPT_AnimScaled:
 		{
 			// normalize time to fit base seq
-			const float Fraction = (GetPlayLength() > 0.f)? FMath::Clamp<float>(ExtractionContext.CurrentTime / GetPlayLength(), 0.f, 1.f) : 0.f;
-			const float BasePoseTime = RefPoseSeq->GetPlayLength() * Fraction;
+			PRAGMA_DISABLE_DEPRECATION_WARNINGS
+			const float Fraction = (SequenceLength > 0.f)? FMath::Clamp<float>(ExtractionContext.CurrentTime / SequenceLength, 0.f, 1.f) : 0.f;
+			const float BasePoseTime = RefPoseSeq->SequenceLength * Fraction;
+			PRAGMA_ENABLE_DEPRECATION_WARNINGS
 
 			FAnimExtractContext BasePoseExtractionContext(ExtractionContext);
 			BasePoseExtractionContext.CurrentTime = BasePoseTime;
@@ -1872,9 +1674,8 @@ void UAnimSequence::GetAdditiveBasePose(FAnimationPoseData& OutAnimationPoseData
 		// use animation as a base pose. Need BasePoseSeq and RefFrameIndex (will clamp if outside).
 		case ABPT_AnimFrame:
 		{
-			const float Fraction = (RefPoseSeq->NumberOfKeys > 0) ? FMath::Clamp<float>((float)RefFrameIndex / (float)RefPoseSeq->NumberOfKeys, 0.f, 1.f) : 0.f;
+			const float Fraction = (RefPoseSeq->GetNumberOfSampledKeys() > 0) ? FMath::Clamp<float>((float)RefFrameIndex / (float)RefPoseSeq->GetNumberOfSampledKeys(), 0.f, 1.f) : 0.f;
 			const float BasePoseTime = RefPoseSeq->GetPlayLength() * Fraction;
-
 			FAnimExtractContext BasePoseExtractionContext(ExtractionContext);
 			BasePoseExtractionContext.CurrentTime = BasePoseTime;
 			RefPoseSeq->GetBonePose(OutAnimationPoseData, BasePoseExtractionContext, true);
@@ -1984,227 +1785,64 @@ FName UAnimSequence::GetRetargetTransformsSourceName() const
 	}
 }
 
-void UAnimSequence::RetargetBoneTransform(FTransform& BoneTransform, const int32 SkeletonBoneIndex, const FCompactPoseBoneIndex& BoneIndex, const FBoneContainer& RequiredBones, const bool bIsBakedAdditive) const
+const TArray<FTrackToSkeletonMap>& UAnimSequence::GetRawTrackToSkeletonMapTable() const
 {
-	FAnimationRuntime::RetargetBoneTransform(GetSkeleton(), GetRetargetTransformsSourceName(), GetRetargetTransforms(), BoneTransform, SkeletonBoneIndex, BoneIndex, RequiredBones, bIsBakedAdditive);
-}
-
 #if WITH_EDITOR
-/** Utility function to crop data from a RawAnimSequenceTrack */
-static int32 CropRawTrack(FRawAnimSequenceTrack& RawTrack, int32 StartKey, int32 NumKeys, int32 TotalNumOfFrames)
-{
-	check(RawTrack.PosKeys.Num() == 1 || RawTrack.PosKeys.Num() == TotalNumOfFrames);
-	check(RawTrack.RotKeys.Num() == 1 || RawTrack.RotKeys.Num() == TotalNumOfFrames);
-	// scale key can be empty
-	check(RawTrack.ScaleKeys.Num() == 0 || RawTrack.ScaleKeys.Num() == 1 || RawTrack.ScaleKeys.Num() == TotalNumOfFrames);
-
-	if( RawTrack.PosKeys.Num() > 1 )
-	{
-		RawTrack.PosKeys.RemoveAt(StartKey, NumKeys);
-		check(RawTrack.PosKeys.Num() > 0);
-		RawTrack.PosKeys.Shrink();
-	}
-
-	if( RawTrack.RotKeys.Num() > 1 )
-	{
-		RawTrack.RotKeys.RemoveAt(StartKey, NumKeys);
-		check(RawTrack.RotKeys.Num() > 0);
-		RawTrack.RotKeys.Shrink();
-	}
-
-	if( RawTrack.ScaleKeys.Num() > 1 )
-	{
-		RawTrack.ScaleKeys.RemoveAt(StartKey, NumKeys);
-		check(RawTrack.ScaleKeys.Num() > 0);
-		RawTrack.ScaleKeys.Shrink();
-	}
+	ValidateModel();
+	PRAGMA_DISABLE_DEPRECATION_WARNINGS
+	return DataModel->GetTransientRawAnimationTrackSkeletonMappings();
+	PRAGMA_ENABLE_DEPRECATION_WARNINGS
+#endif
 
-	// Update NumFrames below to reflect actual number of keys.
-	return FMath::Max<int32>( RawTrack.PosKeys.Num(), FMath::Max<int32>(RawTrack.RotKeys.Num(), RawTrack.ScaleKeys.Num()) );
+	static const TArray<FTrackToSkeletonMap> TempMap;
+	check(false);
+	return TempMap;
 }
 
-void UAnimSequence::ResizeSequence(float NewLength, int32 NewNumKeys, bool bInsert, int32 StartFrame/*inclusive */, int32 EndFrame/*inclusive*/)
+const FRawAnimSequenceTrack& UAnimSequence::GetRawAnimationTrack(int32 TrackIndex) const
 {
-	check (NewNumKeys > 0);
-	check (StartFrame < EndFrame);
-
-	int32 OldNumberOfKeys = NumberOfKeys;
-	float OldSequenceLength = GetPlayLength();
-
-	// verify condition
-	NumberOfKeys = NewNumKeys;
-	// Update sequence length to match new number of frames.
-	SetSequenceLength(NewLength);
-
-	const float OldInterval = OldSequenceLength / (OldNumberOfKeys - 1);
-	const float NewInterval = NewLength / (NewNumKeys - 1);
-	ensure(FMath::IsNearlyEqual(OldInterval, NewInterval));
-
-	const float OldStartTime = StartFrame * OldInterval;
-	const float OldEndTime = EndFrame * OldInterval;
-	const float Duration = OldEndTime - OldStartTime;
-
-	// re-locate notifies
-	for (auto& Notify: Notifies)
-	{
-		float CurrentTime = Notify.GetTime();
-		float NewDuration = 0.f;
-		if (bInsert)
-		{
-			// if state, make sure to adjust end time
-			if(Notify.NotifyStateClass)
-			{
-				float NotifyDuration = Notify.GetDuration();
-				float NotifyEnd = CurrentTime + NotifyDuration;
-				if(NotifyEnd >= OldStartTime)
-				{
-					NewDuration = NotifyDuration + Duration;
-				}
-				else
-				{
-					NewDuration = NotifyDuration;
-				}
-			}
-
-			// when insert, we only care about start time
-			// if it's later than start time
-			if (CurrentTime >= OldStartTime)
-			{
-				CurrentTime += Duration;
-			}
-		}
-		else
-		{
-			// if state, make sure to adjust end time
-			if(Notify.NotifyStateClass)
-			{
-				const float NotifyDuration = Notify.GetDuration();
-				float NotifyEnd = CurrentTime + NotifyDuration;
-				NewDuration = NotifyDuration;
-
-				// If Notify is inside of the trimmed time frame, zero out the duration
-				if(NotifyEnd >= OldStartTime && NotifyEnd <= OldEndTime)
-				{
-					// Set duration to a single frame according to the current sampling rate
-					NewDuration = (float)GetSamplingFrameRate().AsInterval();
-				}
-				// If Notify overlaps trimmed time frame, remove trimmed duration
-				else if (CurrentTime < OldEndTime && NotifyEnd > OldEndTime)
-				{
-					NewDuration = NotifyEnd-Duration-CurrentTime;
-				}
-
-				// Clamp duration to a single frame according to the current sampling rate
-				NewDuration = FMath::Max(NewDuration, (float)GetSamplingFrameRate().AsInterval());
-			}
-
-			if (CurrentTime >= OldStartTime && CurrentTime <= OldEndTime)
-			{
-				CurrentTime = OldStartTime;
-			}
-			else if (CurrentTime > OldEndTime)
-			{
-				CurrentTime -= Duration;
-			}
-		}
-
-		float ClampedCurrentTime = FMath::Clamp(CurrentTime, 0.f, GetPlayLength());
-		Notify.LinkSequence(this, ClampedCurrentTime);
-		Notify.SetDuration(NewDuration);
+#if WITH_EDITOR
+	ValidateModel();
+	PRAGMA_DISABLE_DEPRECATION_WARNINGS
+	return DataModel->GetBoneTrackByIndex(TrackIndex).InternalTrackData;
+	PRAGMA_ENABLE_DEPRECATION_WARNINGS 
+#endif
+	static const FRawAnimSequenceTrack TempTrack;
+	check(false);
+	return TempTrack;
+}
 
-		if (ClampedCurrentTime == 0.f)
-		{
-			Notify.TriggerTimeOffset = GetTriggerTimeOffsetForType(EAnimEventTriggerOffsets::OffsetAfter);
-		}
-		else if (ClampedCurrentTime == GetPlayLength())
-		{
-			Notify.TriggerTimeOffset = GetTriggerTimeOffsetForType(EAnimEventTriggerOffsets::OffsetBefore);
-		}
-	}
+FRawAnimSequenceTrack& UAnimSequence::GetRawAnimationTrack(int32 TrackIndex)
+{
+#if WITH_EDITOR
+	ValidateModel();
+	PRAGMA_DISABLE_DEPRECATION_WARNINGS
+	return DataModel->GetNonConstRawAnimationTrackByIndex(TrackIndex);
+	PRAGMA_ENABLE_DEPRECATION_WARNINGS
+#endif
+		static FRawAnimSequenceTrack TempTrack;
+	check(false);
+	return TempTrack;
+}
 
-	for (FAnimSyncMarker& Marker : AuthoredSyncMarkers)
-	{
-		float CurrentTime = Marker.Time;
-		if (bInsert)
-		{
-			// when insert, we only care about start time
-			// if it's later than start time
-			if (CurrentTime >= OldStartTime)
-			{
-				// Move forwards
-				CurrentTime += Duration;
-			}
-		}
-		else
-		{
-			if (CurrentTime >= OldStartTime && CurrentTime <= OldEndTime)
-			{
-				// Snap to start of the trim/crop
-				CurrentTime = OldStartTime;
-			}
-			else if (CurrentTime > OldEndTime)
-			{
-				// Move backwards relative to crop/trim range end
-				CurrentTime -= Duration;
-			}
-		}
-		Marker.Time = FMath::Clamp(CurrentTime, 0.f, GetPlayLength());
-	}
-	// resize curves
-	RawCurveData.Resize(NewLength, bInsert, OldStartTime, OldEndTime);
+void UAnimSequence::RetargetBoneTransform(FTransform& BoneTransform, const int32 SkeletonBoneIndex, const FCompactPoseBoneIndex& BoneIndex, const FBoneContainer& RequiredBones, const bool bIsBakedAdditive) const
+{
+	FAnimationRuntime::RetargetBoneTransform(GetSkeleton(), GetRetargetTransformsSourceName(), GetRetargetTransforms(), BoneTransform, SkeletonBoneIndex, BoneIndex, RequiredBones, bIsBakedAdditive);
 }
 
+#if WITH_EDITOR
 bool UAnimSequence::InsertFramesToRawAnimData( int32 StartFrame, int32 EndFrame, int32 CopyFrame)
 {
+	ValidateModel();
+
 	// make sure the copyframe is valid and start frame is valid
-	const int32 NumberOfKeysToInsert = EndFrame-StartFrame;
-	if ((CopyFrame>=0 && CopyFrame< NumberOfKeys) && (StartFrame >= 0 && StartFrame <=NumberOfKeys) && NumberOfKeysToInsert > 0)
+	const int32 NumFramesToInsert = EndFrame-StartFrame;
+	if ((CopyFrame>=0 && CopyFrame< DataModel->GetNumberOfKeys()) && (StartFrame >= 0 && StartFrame <= DataModel->GetNumberOfKeys()) && NumFramesToInsert > 0)
 	{
-		for (auto& RawData : RawAnimationData)
-		{
-			if (RawData.PosKeys.Num() > 1 && RawData.PosKeys.IsValidIndex(CopyFrame))
-			{
-				auto Source = RawData.PosKeys[CopyFrame];
-				RawData.PosKeys.InsertZeroed(StartFrame, NumberOfKeysToInsert);
-				for (int32 Index=StartFrame; Index<EndFrame; ++Index)
-				{
-					RawData.PosKeys[Index] = Source;
-				}
-			}
+		UAnimDataController::FScopedBracket ScopedBracket(Controller, LOCTEXT("InsertFramesToRawAnimData", "Inserting Frames into Animation Track Data"));
+		UE::Anim::AnimationData::DuplicateKeys(this, StartFrame, NumFramesToInsert, CopyFrame);
 
-			if(RawData.RotKeys.Num() > 1 && RawData.RotKeys.IsValidIndex(CopyFrame))
-			{
-				auto Source = RawData.RotKeys[CopyFrame];
-				RawData.RotKeys.InsertZeroed(StartFrame, NumberOfKeysToInsert);
-				for(int32 Index=StartFrame; Index<EndFrame; ++Index)
-				{
-					RawData.RotKeys[Index] = Source;
-				}
-			}
-
-			if(RawData.ScaleKeys.Num() > 1 && RawData.ScaleKeys.IsValidIndex(CopyFrame))
-			{
-				auto Source = RawData.ScaleKeys[CopyFrame];
-				RawData.ScaleKeys.InsertZeroed(StartFrame, NumberOfKeysToInsert);
-
-				for(int32 Index=StartFrame; Index<EndFrame; ++Index)
-				{
-					RawData.ScaleKeys[Index] = Source;
-				}
-			}
-		}
-
-		float const FrameTime = SamplingFrameRate.AsInterval();
-
-		const int32 NewNumKeys = NumberOfKeys + NumberOfKeysToInsert;
-		const int32 NewNumFrames = NewNumKeys - 1;
-		ResizeSequence((float)NewNumFrames * FrameTime, NewNumKeys, true, StartFrame, EndFrame);
-
-		UE_LOG(LogAnimation, Log, TEXT("\tGetPlayLength(): %f, Number of Keys: %d"), GetPlayLength(), NumberOfKeys);
-		
-		MarkRawDataAsModified();
-		OnRawDataChanged();
-		MarkPackageDirty();
+		UE_LOG(LogAnimation, Log, TEXT("\tPlay Length: %f, Number of Keys: %d"), GetPlayLength(), DataModel->GetNumberOfKeys());
 
 		return true;
 	}
@@ -2214,11 +1852,8 @@ bool UAnimSequence::InsertFramesToRawAnimData( int32 StartFrame, int32 EndFrame,
 
 bool UAnimSequence::CropRawAnimData( float CurrentTime, bool bFromStart )
 {
-	// Length of one frame.
-	float const FrameTime = SamplingFrameRate.AsInterval();
 	// Save Total Number of Frames before crop
-	const int32 TotalNumOfKeys = NumberOfKeys;
-
+	const int32 TotalNumOfKeys = GetDataModel()->GetNumberOfKeys();
 	// if current frame is 1, do not try crop. There is nothing to crop
 	if (TotalNumOfKeys <= 1 )
 	{
@@ -2236,105 +1871,57 @@ bool UAnimSequence::CropRawAnimData( float CurrentTime, bool bFromStart )
 		return false;
 	}
 
-	// Find the right key to cut at.
-	// This assumes that all keys are equally spaced (ie. won't work if we have dropped unimportant frames etc).
-	// The reason I'm changing to TotalNumOfFrames is CT/SL = KeyIndexWithFraction/TotalNumOfFrames
-	// To play TotalNumOfFrames, it takes GetPlayLength(). Each key will take GetPlayLength()/TotalNumOfFrames
-	float const KeyIndexWithFraction = (CurrentTime * (float)(TotalNumOfKeys)) / GetPlayLength();
-	int32 KeyIndex = bFromStart ? FMath::FloorToInt(KeyIndexWithFraction) : FMath::CeilToInt(KeyIndexWithFraction);
-	// Ensure KeyIndex is in range.
-	KeyIndex = FMath::Clamp<int32>(KeyIndex, 1, TotalNumOfKeys-1); 
+	const FFrameTime FrameTime = DataModel->GetFrameRate().AsFrameTime(CurrentTime);
+	const FFrameNumber FrameNumber = FrameTime.CeilToFrame();
 	// determine which keys need to be removed.
-	const int32 StartKey = bFromStart ? 0 : KeyIndex;
-	const int32 NumKeysToRemove = bFromStart ? KeyIndex : TotalNumOfKeys - KeyIndex;
+	const int32 StartKey = bFromStart ? 0 : FrameNumber.Value;
+	const int32 NumKeysToRemove = bFromStart ? FrameNumber.Value : TotalNumOfKeys - (StartKey + 1);
 
 	// Recalculate number of keys and frames
 	const int32 NewNumKeys = TotalNumOfKeys - NumKeysToRemove;
 	const int32 NewNumFrames = NewNumKeys - 1;
 
-	UE_LOG(LogAnimation, Log, TEXT("UAnimSequence::CropRawAnimData %s - CurrentTime: %f, bFromStart: %d, TotalNumOfFrames: %d, KeyIndex: %d, StartKey: %d, NumKeys: %d"), *GetName(), CurrentTime, bFromStart, TotalNumOfKeys, KeyIndex, StartKey, NumKeysToRemove);
+	UE_LOG(LogAnimation, Log, TEXT("UAnimSequence::CropRawAnimData %s - CurrentTime: %f, bFromStart: %d, TotalNumOfKeys: %d, KeyIndex: %d, StartKey: %d, NumKeys: %d"), *GetName(), CurrentTime, bFromStart, TotalNumOfKeys, FrameNumber.Value, StartKey, NumKeysToRemove);
 
-	// Iterate over tracks removing keys from each one.
-	for(int32 i=0; i<RawAnimationData.Num(); i++)
-	{
-		// Update NewNumKeys below to reflect actual number of keys while we crop the anim data
-		CropRawTrack(RawAnimationData[i], StartKey, NumKeysToRemove, TotalNumOfKeys);
-	}
-
-	// Double check that everything is fine
-	for(int32 i=0; i<RawAnimationData.Num(); i++)
-	{
-		FRawAnimSequenceTrack& RawTrack = RawAnimationData[i];
-		check(RawTrack.PosKeys.Num() == 1 || RawTrack.PosKeys.Num() == NewNumKeys);
-		check(RawTrack.RotKeys.Num() == 1 || RawTrack.RotKeys.Num() == NewNumKeys);
-	}
+	ValidateModel();
 
-	// Update sequence length to match new number of frames.
-	ResizeSequence((float)NewNumFrames * FrameTime, NewNumKeys, false, StartKey, StartKey + NumKeysToRemove);
+	Controller->OpenBracket(LOCTEXT("CropRawAnimation_Bracket", "Cropping Animation Track Data"));
+	UE::Anim::AnimationData::RemoveKeys(this, StartKey, NumKeysToRemove);
+	Controller->CloseBracket();
 
-	UE_LOG(LogAnimation, Log, TEXT("\tGetPlayLength(): %f, Number of Keys: %d"), GetPlayLength(), NumberOfKeys);
+	UE_LOG(LogAnimation, Log, TEXT("\tPlay Length: %f, Number of Keys: %d"), GetPlayLength(), GetDataModel()->GetNumberOfKeys());
 
-	MarkRawDataAsModified();
-	OnRawDataChanged();
 	MarkPackageDirty();
-	return true;
-}
-
-bool UAnimSequence::CompressRawAnimData(float MaxPosDiff, float MaxAngleDiff)
-{
-	if (RawAnimationData.Num() > 0)
-	{
-		return StaticCompressRawAnimData(RawAnimationData, NumberOfKeys, GetFName(), MaxPosDiff, MaxAngleDiff);
-	}
-	return false;
-}
-
-bool UAnimSequence::CompressRawAnimData()
-{
-	if (RawAnimationData.Num() > 0)
-	{
-		return StaticCompressRawAnimData(RawAnimationData, NumberOfKeys, GetFName());
-	}
-	return false;
-}
-
-/** 
- * Flip Rotation W for the RawTrack
- */
-void FlipRotationW(FRawAnimSequenceTrack& RawTrack)
-{
-	int32 TotalNumOfRotKey = RawTrack.RotKeys.Num();
 
-	for (int32 I=0; I<TotalNumOfRotKey; ++I)
-	{
-		FQuat& RotKey = RawTrack.RotKeys[I];
-		RotKey.W *= -1.f;
-	}
+	return true;
 }
 
-
-void UAnimSequence::FlipRotationWForNonRoot(USkeletalMesh * SkelMesh)
+void UAnimSequence::FlipRotationWForNonRoot(USkeletalMesh* SkelMesh)
 {
 	if (!GetSkeleton())
 	{
 		return;
 	}
 
-	// Now add additive animation to destination.
-	for(int32 TrackIdx=0; TrackIdx<TrackToSkeletonMapTable.Num(); TrackIdx++)
+	ValidateModel();
+	Controller->OpenBracket(LOCTEXT("FlipRotationWForNonRoot_Bracket", "Flipping rotation W component for non-root tracks"));
 	{
-		// Figure out which bone this track is mapped to
-		const int32 BoneIndex = TrackToSkeletonMapTable[TrackIdx].BoneTreeIndex;
-		if ( BoneIndex > 0 )
+		for (const FBoneAnimationTrack& Track : DataModel->GetBoneAnimationTracks())
 		{
-			FlipRotationW( RawAnimationData[TrackIdx] );
+			// Only apply to non-root bones
+			if (Track.BoneTreeIndex > 0)
+			{
+				TArray<FQuat> NewRotationalKeys = Track.InternalTrackData.RotKeys;
+				for (FQuat& RotationalKey : NewRotationalKeys)
+				{
+					RotationalKey.W *= -1.f;
+				}
 
+				Controller->SetBoneTrackKeys(Track.Name, Track.InternalTrackData.PosKeys, NewRotationalKeys, Track.InternalTrackData.ScaleKeys);
+			}
 		}
 	}
-
-	// Apply compression
-	MarkRawDataAsModified();
-	OnRawDataChanged();
+	Controller->CloseBracket();
 }
 #endif 
 
@@ -2344,10 +1931,10 @@ bool UAnimSequence::ShouldPerformStripping(const bool bPerformFrameStripping, co
 	const bool bShouldPerformFrameStripping = bPerformFrameStripping && bAllowFrameStripping;
 
 	// Can only do stripping on animations that have an even number of frames once the end frame is removed)
-	const bool bIsEvenFramed = ((NumberOfKeys - 1) % 2) == 0;
+	const bool bIsEvenFramed = ((NumberOfSampledKeys - 1) % 2) == 0;
 	const bool bIsValidForStripping = bIsEvenFramed || bPerformStrippingOnOddFramedAnims;
 
-	const bool bStripCandidate = (NumberOfKeys > 10) && bIsValidForStripping;
+	const bool bStripCandidate = (NumberOfSampledKeys > 10) && bIsValidForStripping;
 
 	return bStripCandidate && bShouldPerformFrameStripping;
 }
@@ -2383,7 +1970,7 @@ void UAnimSequence::RequestAnimCompression(FRequestAnimCompressionParams Params)
 		return;
 	}
 
-	if (FPlatformProperties::RequiresCookedData())
+	if (FPlatformProperties::RequiresCookedData() || bBlockCompressionRequests)
 	{
 		return;
 	}
@@ -2403,20 +1990,13 @@ void UAnimSequence::RequestAnimCompression(FRequestAnimCompressionParams Params)
 	// Make sure all our required dependencies are loaded
 	FAnimationUtils::EnsureAnimSequenceLoaded(*this);
 
-	if (!RawDataGuid.IsValid())
-	{
-		RawDataGuid = GenerateGuidFromRawData();
-	}
+	EnsureValidRawDataGuid();
 
 	bUseRawDataOnly = true;
 
 	check(!bCompressionInProgress);
 	bCompressionInProgress = true;
 
-	// Need to make sure this is up to date.
-	VerifyCurveNames<FFloatCurve>(*CurrentSkeleton, USkeleton::AnimCurveMappingName, RawCurveData.FloatCurves);
-	VerifyTrackMap(CurrentSkeleton);
-
 	Params.CompressContext->GatherPreCompressionStats(GetName(), GetApproxRawSize(), GetApproxCompressedSize());
 
 	const double CompressionStartTime = FPlatformTime::Seconds();
@@ -2445,10 +2025,6 @@ void UAnimSequence::RequestAnimCompression(FRequestAnimCompressionParams Params)
 		else
 		{
 			// Data does not exist, need to build it.
-			// Filter RAW data to get rid of mismatched tracks (translation/rotation data with a different number of keys than there are frames)
-			// No trivial key removal is done at this point (impossible error metrics of -1), since all of the techniques will perform it themselves
-			CompressRawAnimData(-1.0f, -1.0f);
-
 			FCompressibleAnimRef CompressibleData = MakeShared<FCompressibleAnimData, ESPMode::ThreadSafe>(this, bPerformStripping);
 			AnimCompressor->SetCompressibleData(CompressibleData);
 
@@ -2537,58 +2113,42 @@ void UAnimSequence::ApplyCompressedData(const TArray<uint8>& Data)
 	}
 }
 
-void UAnimSequence::SanitizeAnimationTrackData(FRawAnimSequenceTrack& RawAnim)
-{
-	for (auto ScaleIter = RawAnim.ScaleKeys.CreateIterator(); ScaleIter; ++ScaleIter)
-	{
-		FVector& Scale3D = *ScaleIter;
-		if (FMath::IsNearlyZero(Scale3D.X))
-		{
-			Scale3D.X = 0.f;
-		}
-		if (FMath::IsNearlyZero(Scale3D.Y))
-		{
-			Scale3D.Y = 0.f;
-		}
-		if (FMath::IsNearlyZero(Scale3D.Z))
-		{
-			Scale3D.Z = 0.f;
-		}
-	}
-
-	// make sure Rotation part is normalized before compress
-	for (auto RotIter = RawAnim.RotKeys.CreateIterator(); RotIter; ++RotIter)
-	{
-		FQuat& Rotation = *RotIter;
-		if (!Rotation.IsNormalized())
-		{
-			Rotation.Normalize();
-		}
-	}
-}
-
 void UAnimSequence::SerializeCompressedData(FArchive& Ar, bool bDDCData)
 {
 	if (!HasAnyFlags(RF_ClassDefaultObject))
 	{
 		CompressedData.SerializeCompressedData(Ar, bDDCData, this, this->GetSkeleton(), BoneCompressionSettings, CurveCompressionSettings);
+		}
 	}
+
+int32 UAnimSequence::GetSkeletonIndexFromRawDataTrackIndex(const int32 TrackIndex) const
+{
+#if WITH_EDITOR
+	ValidateModel();
+	return DataModel->GetBoneTrackByIndex(TrackIndex).BoneTreeIndex;
+#else
+	check(false);
+	return INDEX_NONE;
+#endif
 }
 
 #if WITH_EDITOR
 
 bool UAnimSequence::CanBakeAdditive() const
 {
-	return	(NumberOfKeys > 0) &&
+	return	(NumberOfSampledKeys > 0) &&
 			IsValidAdditive() &&
 			GetSkeleton();
 }
 
-bool UAnimSequence::DoesSequenceContainZeroScale()
+bool UAnimSequence::DoesSequenceContainZeroScale() const
 {
-	for (const FRawAnimSequenceTrack& RawTrack : RawAnimationData)
+	ValidateModel();
+
+	const TArray<FBoneAnimationTrack>& BoneAnimationTracks = DataModel->GetBoneAnimationTracks();
+	for (const FBoneAnimationTrack& Track : BoneAnimationTracks)
 	{
-		for (const FVector& ScaleKey : RawTrack.ScaleKeys)
+		for (const FVector& ScaleKey : Track.InternalTrackData.ScaleKeys)
 		{
 			if (ScaleKey.IsZero())
 			{
@@ -2602,7 +2162,9 @@ bool UAnimSequence::DoesSequenceContainZeroScale()
 
 FGuid UAnimSequence::GenerateGuidFromRawData() const
 {
-	return GenerateGuidFromRawAnimData(RawAnimationData, RawCurveData);
+	ValidateModel();
+	return DataModel->GenerateGuid();
+
 }
 
 void CopyTransformToRawAnimationData(const FTransform& BoneTransform, FRawAnimSequenceTrack& Track, int32 Frame)
@@ -2648,9 +2210,9 @@ public:
 
 void UAnimSequence::BakeOutVirtualBoneTracks(TArray<FRawAnimSequenceTrack>& NewRawTracks, TArray<FName>& NewAnimationTrackNames, TArray<FTrackToSkeletonMap>& NewTrackToSkeletonMapTable)
 {
-	const int32 NumVirtualBonesOnSkeleton = GetSkeleton()->GetVirtualBones().Num();
+	ValidateModel();
 
-	check( (RawAnimationData.Num() == TrackToSkeletonMapTable.Num()) && (RawAnimationData.Num() == AnimationTrackNames.Num()) ); //Make sure starting data is valid
+	const int32 NumVirtualBonesOnSkeleton = GetSkeleton()->GetVirtualBones().Num();
 
 	TArray<int> VirtualBonesToAdd;
 	VirtualBonesToAdd.Reserve(NumVirtualBonesOnSkeleton);
@@ -2658,6 +2220,23 @@ void UAnimSequence::BakeOutVirtualBoneTracks(TArray<FRawAnimSequenceTrack>& NewR
 	TArray<int32> SourceParents;
 
 	const FReferenceSkeleton& RefSkeleton = GetSkeleton()->GetReferenceSkeleton();
+	
+	const TArray<FBoneAnimationTrack>& BoneAnimationTracks = ResampledAnimationTrackData;
+
+	// Presize outgoing arrays
+	const int32 NumAnimationTracks = DataModel->GetNumBoneTracks();
+	NewRawTracks.Reset(NumAnimationTracks);
+	NewAnimationTrackNames.Reset(NumAnimationTracks);
+	NewTrackToSkeletonMapTable.Reset(NumAnimationTracks);
+
+	for (const FBoneAnimationTrack& Track : BoneAnimationTracks)
+	{
+		FRawAnimSequenceTrack& RawTrack = NewRawTracks.Add_GetRef(Track.InternalTrackData);
+		NewAnimationTrackNames.Add(Track.Name);
+		NewTrackToSkeletonMapTable.Add(Track.BoneTreeIndex);
+
+		UE::Anim::Compression::SanitizeRawAnimSequenceTrack(RawTrack);
+	}
 
 	const TArray<FVirtualBone>& VirtualBones = GetSkeleton()->GetVirtualBones();
 	for (int32 VBIndex = 0; VBIndex < NumVirtualBonesOnSkeleton; ++VBIndex)
@@ -2685,7 +2264,7 @@ void UAnimSequence::BakeOutVirtualBoneTracks(TArray<FRawAnimSequenceTrack>& NewR
 
 			while (!SourceParents.Contains(CurrentBone))
 			{
-				if (Algo::FindBy(TrackToSkeletonMapTable, CurrentBone, &FTrackToSkeletonMap::BoneTreeIndex) != nullptr)
+				if (Algo::FindBy(NewTrackToSkeletonMapTable, CurrentBone, &FTrackToSkeletonMap::BoneTreeIndex) != nullptr)
 				{
 					//We animate this bone so the virtual bone is needed
 					bBuildVirtualBone = true;
@@ -2700,7 +2279,7 @@ void UAnimSequence::BakeOutVirtualBoneTracks(TArray<FRawAnimSequenceTrack>& NewR
 			const int32 FirstCommon = SourceParents.IndexOfByKey(CurrentBone);
 			for (int32 i = FirstCommon - 1; i >= 0; --i)
 			{
-				if (Algo::FindBy(TrackToSkeletonMapTable, i, &FTrackToSkeletonMap::BoneTreeIndex) != nullptr)
+				if (Algo::FindBy(NewTrackToSkeletonMapTable, i, &FTrackToSkeletonMap::BoneTreeIndex) != nullptr)
 				{
 					//We animate this bone so the virtual bone is needed
 					bBuildVirtualBone = true;
@@ -2717,35 +2296,30 @@ void UAnimSequence::BakeOutVirtualBoneTracks(TArray<FRawAnimSequenceTrack>& NewR
 
 	const int32 NumVirtualBonesToAdd = VirtualBonesToAdd.Num();
 
-	auto SetupData = [NumVirtualBonesToAdd](auto& InOutData, const auto& BaseData)
-	{
-		if (InOutData.Num() == 0)
-		{
-			InOutData = BaseData;
-		}
-
-		InOutData.Reserve(InOutData.Num() + NumVirtualBonesToAdd);
-	};
-
-	SetupData(NewRawTracks, RawAnimationData);
-	SetupData(NewTrackToSkeletonMapTable, TrackToSkeletonMapTable);
-	SetupData(NewAnimationTrackNames, AnimationTrackNames);
+	const USkeleton* TargetSkeleton = GetSkeleton();
 
 	const int32 VirtualBoneStartIndex = NewRawTracks.Num();
 
+	const int32 NumTrackKeys = NumberOfSampledKeys;
+
+	// Generate track data for virtual bones
 	for (int32 VBIndex : VirtualBonesToAdd)
 	{
-		const int32 TrackIndex = NewRawTracks.Add(FRawAnimSequenceTrack());
+		const FVirtualBone& VirtualBone = VirtualBones[VBIndex];
+
+		FRawAnimSequenceTrack& RawTrack = NewRawTracks.AddDefaulted_GetRef();
 
 		//Init new tracks
-		NewRawTracks[TrackIndex].PosKeys.SetNumUninitialized(NumberOfKeys);
-		NewRawTracks[TrackIndex].RotKeys.SetNumUninitialized(NumberOfKeys);
-		NewRawTracks[TrackIndex].ScaleKeys.SetNumUninitialized(NumberOfKeys);		
+		RawTrack.PosKeys.SetNumUninitialized(NumTrackKeys);
+		RawTrack.RotKeys.SetNumUninitialized(NumTrackKeys);
+		RawTrack.ScaleKeys.SetNumUninitialized(NumTrackKeys);
 
-		NewTrackToSkeletonMapTable.Add(FTrackToSkeletonMap(GetSkeleton()->GetReferenceSkeleton().GetRequiredVirtualBones()[VBIndex]));
-		NewAnimationTrackNames.Add(VirtualBones[VBIndex].VirtualBoneName);
+		NewAnimationTrackNames.Add(VirtualBone.VirtualBoneName);
+		const int32 VirtualBoneSkeletonIndex = TargetSkeleton->GetReferenceSkeleton().GetRequiredVirtualBones()[VBIndex];
+		NewTrackToSkeletonMapTable.Add(VirtualBoneSkeletonIndex);
 	}
 
+	// Setup data for evaluating bone transforms
 	FMemMark Mark(FMemStack::Get());
 	FByFramePoseEvalContext EvalContext(this);
 
@@ -2754,10 +2328,8 @@ void UAnimSequence::BakeOutVirtualBoneTracks(TArray<FRawAnimSequenceTrack>& NewR
 	Pose.SetBoneContainer(&EvalContext.RequiredBones);
 
 	FAnimExtractContext ExtractContext;
-
 	const TArray<FVirtualBoneRefData>& VBRefData = GetSkeleton()->GetReferenceSkeleton().GetVirtualBoneRefData();
-
-	for (int32 Frame = 0; Frame < NumberOfKeys; ++Frame)
+	for (int32 Frame = 0; Frame < NumTrackKeys; ++Frame)
 	{
 		// Initialise curve data from Skeleton
 		FBlendedCurve Curve;
@@ -2771,18 +2343,15 @@ void UAnimSequence::BakeOutVirtualBoneTracks(TArray<FRawAnimSequenceTrack>& NewR
 		FAnimationPoseData AnimPoseData(Pose, Curve, TempAttributes);
 		GetAnimationPose(AnimPoseData, ExtractContext);
 
+		// Copy the retrieved frame data into the virtual bone track keys
 		for (int32 VBToAddIndex = 0; VBToAddIndex < VirtualBonesToAdd.Num(); ++VBToAddIndex)
 		{
 			const int32 VBIndex = VirtualBonesToAdd[VBToAddIndex];
 			const FVirtualBoneRefData& VB = VBRefData[VBIndex];
 			CopyTransformToRawAnimationData(Pose[FCompactPoseBoneIndex(VB.VBRefSkelIndex)], NewRawTracks[VBToAddIndex + VirtualBoneStartIndex], Frame);
+			UE::Anim::Compression::SanitizeRawAnimSequenceTrack(NewRawTracks[VBToAddIndex + VirtualBoneStartIndex]);
 		}
 	}
-
-	if (NewRawTracks.Num() > 0)
-	{
-		StaticCompressRawAnimData(NewRawTracks, NumberOfKeys, GetFName());
-	}
 }
 
 bool IsIdentity(const FVector& Pos)
@@ -2819,7 +2388,7 @@ void UAnimSequence::TestEvalauteAnimation() const
 
 	FAnimExtractContext ExtractContext;
 
-	for (int32 Frame = 0; Frame < NumberOfKeys; ++Frame)
+	for (int32 Frame = 0; Frame < NumberOfSampledKeys; ++Frame)
 	{
 		// Initialise curve data from Skeleton
 		FBlendedCurve Curve;
@@ -2835,9 +2404,31 @@ void UAnimSequence::TestEvalauteAnimation() const
 	}
 }
 
-FFloatCurve* GetFloatCurve(FRawCurveTracks& RawCurveTracks, USkeleton::AnimCurveUID& CurveUID)
+FFloatCurve* GetFloatCurve(TArray<FFloatCurve>& FloatCurves, const USkeleton::AnimCurveUID& CurveUID)
+{
+	for (FFloatCurve& Curve : FloatCurves)
+	{
+		if (Curve.Name.UID == CurveUID)
+		{
+			return &Curve;
+		}
+	}
+
+	return nullptr;
+}
+
+FFloatCurve* FindOrAddCurve(TArray<FFloatCurve>& FloatCurves, const FSmartName& CurveName)
 {
-	return static_cast<FFloatCurve *>(RawCurveTracks.GetCurveData(CurveUID, ERawCurveTrackTypes::RCT_Float));
+	FFloatCurve* ReturnCurve = GetFloatCurve(FloatCurves, CurveName.UID);
+
+	if (ReturnCurve == nullptr)
+	{
+		FFloatCurve& NewCurve = FloatCurves.Add_GetRef(FFloatCurve(CurveName, 0));
+		ReturnCurve = &NewCurve;
+	}
+
+	return ReturnCurve;
+	
 }
 
 bool IsNewKeyDifferent(const FRichCurveKey& LastKey, float NewValue)
@@ -2845,7 +2436,7 @@ bool IsNewKeyDifferent(const FRichCurveKey& LastKey, float NewValue)
 	return LastKey.Value != NewValue;
 }
 
-void UAnimSequence::BakeOutAdditiveIntoRawData(TArray<FRawAnimSequenceTrack>& NewRawTracks, TArray<FName>& NewAnimationTrackNames, TArray<FTrackToSkeletonMap>& NewTrackToSkeletonMapTable, FRawCurveTracks& NewCurveTracks, TArray<FRawAnimSequenceTrack>& AdditiveBaseAnimationData)
+void UAnimSequence::BakeOutAdditiveIntoRawData(TArray<FRawAnimSequenceTrack>& NewRawTracks, TArray<FName>& NewAnimationTrackNames, TArray<FTrackToSkeletonMap>& NewTrackToSkeletonMapTable, TArray<FFloatCurve>& NewFloatCurves, TArray<FRawAnimSequenceTrack>& AdditiveBaseAnimationData)
 {
 	if (!CanBakeAdditive())
 	{
@@ -2855,9 +2446,14 @@ void UAnimSequence::BakeOutAdditiveIntoRawData(TArray<FRawAnimSequenceTrack>& Ne
 	USkeleton* MySkeleton = GetSkeleton();
 	check(MySkeleton);
 
-	if (RefPoseSeq && RefPoseSeq->HasAnyFlags(EObjectFlags::RF_NeedPostLoad))
+	ensure(RefPoseSeq && !RefPoseSeq->HasAnyFlags(EObjectFlags::RF_NeedPostLoad));
+	
+	const FSmartNameMapping* CurveMapping = MySkeleton->GetSmartNameContainer(USkeleton::AnimCurveMappingName);
+	TArray<FName> SkeletonCurveNames;
+	CurveMapping->FillNameArray(SkeletonCurveNames);
+	for (const FFloatCurve& FloatCurve : DataModel->GetFloatCurves())
 	{
-		RefPoseSeq->VerifyCurveNames<FFloatCurve>(*MySkeleton, USkeleton::AnimCurveMappingName, RefPoseSeq->RawCurveData.FloatCurves);
+		ensure(SkeletonCurveNames.Contains(FloatCurve.Name.DisplayName));
 	}
 
 	FMemMark Mark(FMemStack::Get());
@@ -2869,9 +2465,9 @@ void UAnimSequence::BakeOutAdditiveIntoRawData(TArray<FRawAnimSequenceTrack>& Ne
 
 	for (FRawAnimSequenceTrack& RawTrack : NewRawTracks)
 	{
-		RawTrack.PosKeys.SetNumUninitialized(NumberOfKeys);
-		RawTrack.RotKeys.SetNumUninitialized(NumberOfKeys);
-		RawTrack.ScaleKeys.SetNumUninitialized(NumberOfKeys);
+		RawTrack.PosKeys.SetNumUninitialized(NumberOfSampledKeys);
+		RawTrack.RotKeys.SetNumUninitialized(NumberOfSampledKeys);
+		RawTrack.ScaleKeys.SetNumUninitialized(NumberOfSampledKeys);
 	}
 
 	// keep the same buffer size
@@ -2897,7 +2493,7 @@ void UAnimSequence::BakeOutAdditiveIntoRawData(TArray<FRawAnimSequenceTrack>& Ne
 
 	FAnimExtractContext ExtractContext;
 
-	for (int32 Frame = 0; Frame < NumberOfKeys; ++Frame)
+	for (int32 Frame = 0; Frame < NumberOfSampledKeys; ++Frame)
 	{
 		// Initialise curve data from Skeleton
 		FBlendedCurve Curve;
@@ -2934,15 +2530,14 @@ void UAnimSequence::BakeOutAdditiveIntoRawData(TArray<FRawAnimSequenceTrack>& Ne
 			if (ArrayIndex != INDEX_NONE)
 			{
 				float CurveWeight = Curve.CurveWeights[ArrayIndex];
-				FFloatCurve* RawCurve = GetFloatCurve(NewCurveTracks, CurveUID);
-				if (!RawCurve && CurveWeight != 0.f) //Only make a new curve if we are going to give it data
+				FFloatCurve* RawCurve = GetFloatCurve(NewFloatCurves, CurveUID);
+				if (!RawCurve && !FMath::IsNearlyZero(CurveWeight)) //Only make a new curve if we are going to give it data
 				{
 					FSmartName NewCurveName;
 					// if we don't have name, there is something wrong here. 
 					ensureAlways(MySkeleton->GetSmartNameByUID(USkeleton::AnimCurveMappingName, CurveUID, NewCurveName));
 					// curve flags don't matter much for compressed curves
-					NewCurveTracks.AddCurveData(NewCurveName, 0, ERawCurveTrackTypes::RCT_Float);
-					RawCurve = GetFloatCurve(NewCurveTracks, CurveUID);
+					RawCurve = FindOrAddCurve(NewFloatCurves, NewCurveName);
 				}
 
 				if (RawCurve)
@@ -2953,8 +2548,7 @@ void UAnimSequence::BakeOutAdditiveIntoRawData(TArray<FRawAnimSequenceTrack>& Ne
 						//Add pre key of 0
 						if (Frame > 0)
 						{
-							const float PreKeyTime = (Frame - 1)*EvalContext.IntervalTime;
-							RawCurve->UpdateOrAddKey(0.f, PreKeyTime);
+							RawCurve->UpdateOrAddKey(0.f, PreviousFrameTime);
 						}
 
 					}
@@ -2976,216 +2570,102 @@ void UAnimSequence::BakeOutAdditiveIntoRawData(TArray<FRawAnimSequenceTrack>& Ne
 			}
 		}
 	}
+}
 
-	const FSmartNameMapping* Mapping = GetSkeleton()->GetSmartNameContainer(USkeleton::AnimCurveMappingName);
-	check(Mapping); // Should always exist
-	NewCurveTracks.RefreshName(Mapping);
-
-#if 0 //Validate baked data
-	for (FRawAnimSequenceTrack& RawTrack : NewRawTracks)
+void UAnimSequence::FlagDependentAnimationsAsRawDataOnly() const
+{
+	for (TObjectIterator<UAnimSequence> Iter; Iter; ++Iter)
 	{
-		for (FQuat& Rot : RawTrack.RotKeys)
+		UAnimSequence* Seq = *Iter;
+		if (Seq->RefPoseSeq == this)
 		{
-			check(Rot.IsNormalized());
+			Seq->bUseRawDataOnly = true;
 		}
 	}
-#endif
-
-	StaticCompressRawAnimData(NewRawTracks, NumberOfKeys, GetFName());
+}
 
-	// Note on (TrackIndex > 0) below : deliberately stop before track 0, compression code doesn't like getting a completely empty animation
-	for (int32 TrackIndex = NewRawTracks.Num() - 1; TrackIndex > 0; --TrackIndex)
+void UAnimSequence::UpdateDependentStreamingAnimations() const
+{
+	for (TObjectIterator<UAnimStreamable> Iter; Iter; ++Iter)
 	{
-		const FRawAnimSequenceTrack& Track = NewRawTracks[TrackIndex];
-		if (IsRawTrackValidForRemoval(Track))
+		UAnimStreamable* Seq = *Iter;
+		if (Seq->SourceSequence == this)
 		{
-			NewRawTracks.RemoveAtSwap(TrackIndex, 1, false);
-			NewAnimationTrackNames.RemoveAtSwap(TrackIndex, 1, false);
-			NewTrackToSkeletonMapTable.RemoveAtSwap(TrackIndex, 1, false);
+			Seq->InitFrom(this);
 		}
 	}
+}
+
+#endif
 
+void UAnimSequence::RecycleAnimSequence()
+{
+#if WITH_EDITOR
+	ResetAnimation();
+#endif // WITH_EDITORONLY_DATA
 }
 
-void UAnimSequence::FlagDependentAnimationsAsRawDataOnly() const
+#if WITH_EDITOR
+void UAnimSequence::CleanAnimSequenceForImport()
 {
-	for (TObjectIterator<UAnimSequence> Iter; Iter; ++Iter)
-	{
-		UAnimSequence* Seq = *Iter;
-		if (Seq->RefPoseSeq == this)
-		{
-			Seq->bUseRawDataOnly = true;
-		}
-	}
+	ValidateModel();
+	DeleteBoneAnimationData();
+	RawDataGuid.Invalidate();
 }
 
-void UAnimSequence::UpdateDependentStreamingAnimations() const
+void UAnimSequence::ResetAnimation()
 {
-	for (TObjectIterator<UAnimStreamable> Iter; Iter; ++Iter)
+	// clear everything. Making new animation, so need to reset all the things that belong here
+	ValidateModel();
+
+	Controller->OpenBracket(LOCTEXT("ResetAnimation_Bracket", "Resetting Animation Sequence"));
 	{
-		UAnimStreamable* Seq = *Iter;
-		if (Seq->SourceSequence == this)
-		{
-			Seq->InitFrom(this);
-		}
+		DeleteBoneAnimationData();
+		DeleteNotifyTrackData();
 	}
-}
+	Controller->CloseBracket();
 
-#endif
+	Notifies.Empty();
+	AuthoredSyncMarkers.Empty();
+	UniqueMarkerNames.Empty();
+	AnimNotifyTracks.Empty();
+	RawDataGuid.Invalidate();
+}
 
-void UAnimSequence::RecycleAnimSequence()
+void UAnimSequence::DeleteDeprecatedRawAnimationData()
 {
-#if WITH_EDITORONLY_DATA
-	// Clear RawAnimData
+	PRAGMA_DISABLE_DEPRECATION_WARNINGS
 	RawAnimationData.Empty();
-	RawDataGuid.Invalidate();
 	AnimationTrackNames.Empty();
 	TrackToSkeletonMapTable.Empty();
-	RawCurveData.Empty();
-
-	PRAGMA_DISABLE_DEPRECATION_WARNINGS
-	ClearCompressedBoneData();
-	ClearCompressedCurveData();
+	SourceRawAnimationData_DEPRECATED.Empty();
 	PRAGMA_ENABLE_DEPRECATION_WARNINGS
+}
 
+void UAnimSequence::DeleteNotifyTrackData()
+{
 	AuthoredSyncMarkers.Empty();
 	UniqueMarkerNames.Empty();
 	Notifies.Empty();
 	AnimNotifyTracks.Empty();
-#endif // WITH_EDITORONLY_DATA
+
+	InitializeNotifyTrack();
 }
 
-#if WITH_EDITOR
-void UAnimSequence::CleanAnimSequenceForImport()
+void UAnimSequence::DeleteBoneAnimationData()
 {
-	// Clear RawAnimData
-	RawAnimationData.Empty();
-	RawDataGuid.Invalidate();
-	AnimationTrackNames.Empty();
-	TrackToSkeletonMapTable.Empty();
-	PRAGMA_DISABLE_DEPRECATION_WARNINGS
-	ClearCompressedBoneData();
-	ClearCompressedCurveData();
-	PRAGMA_ENABLE_DEPRECATION_WARNINGS
+	ValidateModel();
+	Controller->ResetModel();
 }
 #endif // WITH_EDITOR
-
+		
 
 bool UAnimSequence::CopyNotifies(UAnimSequence* SourceAnimSeq, UAnimSequence* DestAnimSeq, bool bShowDialogs /*= true */)
 {
 #if WITH_EDITOR
-	// Abort if source == destination.
-	if( SourceAnimSeq == DestAnimSeq )
-	{
-		return true;
-	}
-
-	// If the destination sequence is shorter than the source sequence, we'll be dropping notifies that
-	// occur at later times than the dest sequence is long.  Give the user a chance to abort if we
-	// find any notifies that won't be copied over.
-	if(bShowDialogs && DestAnimSeq->GetPlayLength() < SourceAnimSeq->GetPlayLength())
-	{
-		for(int32 NotifyIndex=0; NotifyIndex<SourceAnimSeq->Notifies.Num(); ++NotifyIndex)
-		{
-			// If a notify is found which occurs off the end of the destination sequence, prompt the user to continue.
-			const FAnimNotifyEvent& SrcNotifyEvent = SourceAnimSeq->Notifies[NotifyIndex];
-			if( SrcNotifyEvent.GetTriggerTime() > DestAnimSeq->GetPlayLength() )
-			{
-				const bool bProceed = EAppReturnType::Yes == FMessageDialog::Open( EAppMsgType::YesNo, NSLOCTEXT("UnrealEd", "SomeNotifiesWillNotBeCopiedQ", "Some notifies will not be copied because the destination sequence is not long enough.  Proceed?") );
-				if( !bProceed )
-				{
-					return false;
-				}
-				else
-				{
-					break;
-				}
-			}
-		}
-	}
-
-	// If the destination sequence contains any notifies, ask the user if they'd like
-	// to delete the existing notifies before copying over from the source sequence.
-	if(bShowDialogs && DestAnimSeq->Notifies.Num() > 0)
-	{
-		const bool bDeleteExistingNotifies = EAppReturnType::Yes == FMessageDialog::Open(EAppMsgType::YesNo, FText::Format(
-			NSLOCTEXT("UnrealEd", "DestSeqAlreadyContainsNotifiesMergeQ", "The destination sequence already contains {0} notifies.  Delete these before copying?"), FText::AsNumber(DestAnimSeq->Notifies.Num())) );
-		if( bDeleteExistingNotifies )
-		{
-			DestAnimSeq->Notifies.Empty();
-			DestAnimSeq->MarkPackageDirty();
-		}
-	}
-
-	// Do the copy.
-	int32 NumNotifiesThatWereNotCopied = 0;
-
-	for(int32 NotifyIndex=0; NotifyIndex<SourceAnimSeq->Notifies.Num(); ++NotifyIndex)
-	{
-		const FAnimNotifyEvent& SrcNotifyEvent = SourceAnimSeq->Notifies[NotifyIndex];
-
-		// Skip notifies which occur at times later than the destination sequence is long.
-		if( SrcNotifyEvent.GetTriggerTime() > DestAnimSeq->GetPlayLength() )
-		{
-			++NumNotifiesThatWereNotCopied;
-			continue;
-		}
-
-		// Copy notify tracks from src to dest if they are missing
-		if (SrcNotifyEvent.TrackIndex >= DestAnimSeq->AnimNotifyTracks.Num())
-		{
-			for (int32 TrackIndex = DestAnimSeq->AnimNotifyTracks.Num(); TrackIndex <= SrcNotifyEvent.TrackIndex; ++TrackIndex)
-			{
-				DestAnimSeq->AnimNotifyTracks.Add(FAnimNotifyTrack(SourceAnimSeq->AnimNotifyTracks[TrackIndex].TrackName, SourceAnimSeq->AnimNotifyTracks[TrackIndex].TrackColor));
-			}
-		}
-
-		// Track the location of the new notify.
-		int32 NewNotifyIndex = DestAnimSeq->Notifies.AddDefaulted();
-		FAnimNotifyEvent& NotifyEvent = DestAnimSeq->Notifies[NewNotifyIndex];
-		
-		// Copy properties of the NotifyEvent
-		NotifyEvent.TrackIndex = SrcNotifyEvent.TrackIndex;
-		NotifyEvent.NotifyName = SrcNotifyEvent.NotifyName;
-		NotifyEvent.Duration = SrcNotifyEvent.Duration;
-
-		// Copy the notify itself, and point the new one at it.
-		if( SrcNotifyEvent.Notify )
-		{
-			DestAnimSeq->Notifies[NewNotifyIndex].Notify = static_cast<UAnimNotify*>(StaticDuplicateObject(SrcNotifyEvent.Notify, DestAnimSeq, NAME_None, RF_AllFlags, nullptr, EDuplicateMode::Normal, ~EInternalObjectFlags::RootSet));
-		}
-		else
-		{
-			DestAnimSeq->Notifies[NewNotifyIndex].Notify = nullptr;
-		}
-
-		if( SrcNotifyEvent.NotifyStateClass )
-		{
-			DestAnimSeq->Notifies[NewNotifyIndex].NotifyStateClass = static_cast<UAnimNotifyState*>(StaticDuplicateObject(SrcNotifyEvent.NotifyStateClass, DestAnimSeq, NAME_None, RF_AllFlags, nullptr, EDuplicateMode::Normal, ~EInternalObjectFlags::RootSet));
-		}
-		else
-		{
-			DestAnimSeq->Notifies[NewNotifyIndex].NotifyStateClass = nullptr;
-		}
-		
-		// Copy notify timing
-		NotifyEvent.LinkSequence(DestAnimSeq, SrcNotifyEvent.GetTriggerTime());
-		NotifyEvent.TriggerTimeOffset = GetTriggerTimeOffsetForType(DestAnimSeq->CalculateOffsetForNotify(NotifyEvent.GetTriggerTime()));
-
-		// Make sure editor knows we've changed something.
-		DestAnimSeq->MarkPackageDirty();
-		DestAnimSeq->RefreshCacheData();
-	}
-
-	// Inform the user if some notifies weren't copied.
-	if(bShowDialogs && NumNotifiesThatWereNotCopied > 0)
-	{
-		FMessageDialog::Open( EAppMsgType::Ok, FText::Format(
-			NSLOCTEXT("UnrealEd", "SomeNotifiesWereNotCopiedF", "Because the destination sequence was shorter, {0} notifies were not copied."), FText::AsNumber(NumNotifiesThatWereNotCopied)) );
-	}
+	return UE::Anim::CopyNotifies(SourceAnimSeq, DestAnimSeq, bShowDialogs);
 #endif // WITH_EDITOR
-
-	return true;
+	return false;
 }
 
 bool UAnimSequence::IsValidAdditive() const		
@@ -3333,6 +2813,10 @@ void UAnimSequence::RemapTracksToNewSkeleton( USkeleton* NewSkeleton, bool bConv
 {
 	// this is not cheap, so make sure it only happens in editor
 
+	ValidateModel();
+
+	UAnimDataController::FScopedBracket ScopedBracket(Controller, LOCTEXT("RemapTrackToNewSkeleton_Bracket", "Remapping Bone Animation tracks to new Skeleton"));
+
 	// @Todo : currently additive will work fine since we don't bake anything except when we extract
 	// but in the future if we bake this can be problem
 	if (bConvertSpaces)
@@ -3425,8 +2909,8 @@ void UAnimSequence::RemapTracksToNewSkeleton( USkeleton* NewSkeleton, bool bConv
 			ConvertedSpaceAnimations.AddZeroed(SrcNumTracks);
 			ConvertedLocalSpaceAnimations.AddZeroed(SrcNumTracks);
 
-			int32 NumKeys = NumberOfKeys;
-			float Interval = SamplingFrameRate.AsInterval();
+			const int32 NumKeys = DataModel->GetNumberOfKeys();
+			float Interval = DataModel->GetFrameRate().AsInterval();
 
 			// allocate arrays
 			for(int32 SrcTrackIndex=0; SrcTrackIndex<SrcNumTracks; ++SrcTrackIndex)
@@ -3449,7 +2933,8 @@ void UAnimSequence::RemapTracksToNewSkeleton( USkeleton* NewSkeleton, bool bConv
 				for(int32 Key=0; Key<NumKeys; ++Key)
 				{
 					FTransform AnimatedLocalKey;
-					ExtractBoneTransform(RiggingAnimationData.AnimationTracks, AnimatedLocalKey, SrcTrackIndex, Interval*Key);
+					//ExtractBoneTransform(RiggingAnimationData.AnimationTracks[SrcTrackIndex], AnimatedLocalKey, Interval * Key);
+					FAnimationUtils::ExtractTransformFromTrack(Interval* Key, NumberOfSampledKeys, GetPlayLength(), RiggingAnimationData.AnimationTracks[SrcTrackIndex], Interpolation, AnimatedLocalKey);
 
 					AnimatedLocalKey.ScaleTranslation(OldToNewTranslationRatio[NodeIndex]);
 
@@ -3544,50 +3029,16 @@ void UAnimSequence::RemapTracksToNewSkeleton( USkeleton* NewSkeleton, bool bConv
 			SetSkeleton(NewSkeleton);
 
 			// convert back to animated data with new skeleton
-			const bool bPerformPostProcess = false; //Don't do PostProcess during Remap as any animations we reference may not have been updated yet
-			ConvertRiggingDataToAnimationData(RiggingAnimationData, bPerformPostProcess);
+			ConvertRiggingDataToAnimationData(RiggingAnimationData);
 		}
 		// @todo end rig testing
 		// @IMPORTANT: now otherwise this will try to do bone to bone mapping
 		else if(OldSkeleton)
 		{
-			// this only replaces the primary one, it doesn't replace old ones
-			TArray<struct FTrackToSkeletonMap> NewTrackToSkeletonMapTable;
-			NewTrackToSkeletonMapTable.Empty(AnimationTrackNames.Num());
-			NewTrackToSkeletonMapTable.AddUninitialized(AnimationTrackNames.Num());
-			for (int32 Track = 0; Track < AnimationTrackNames.Num(); ++Track)
-			{
-				int32 BoneIndex = NewSkeleton->GetReferenceSkeleton().FindBoneIndex(AnimationTrackNames[Track]);
-				NewTrackToSkeletonMapTable[Track].BoneTreeIndex = BoneIndex;
-			}
-
-			// now I have all NewTrack To Skeleton Map Table
-			// I'll need to compare with old tracks and copy over if SkeletonIndex == 0
-			// if SkeletonIndex != 0, we need to see if we can 
-			for (int32 TableId = 0; TableId < NewTrackToSkeletonMapTable.Num(); ++TableId)
-			{
-				if (ensure(TrackToSkeletonMapTable.IsValidIndex(TableId)))
-				{
-					if (NewTrackToSkeletonMapTable[TableId].BoneTreeIndex != INDEX_NONE)
-					{
-						TrackToSkeletonMapTable[TableId].BoneTreeIndex = NewTrackToSkeletonMapTable[TableId].BoneTreeIndex;
-					}
-					else
-					{
-						// if not found, delete the track data
-						RemoveTrack(TableId);
-						NewTrackToSkeletonMapTable.RemoveAt(TableId);
-						--TableId;
-					}
-				}
-			}
+			// Validate animation tracks against the new skeleton, any tracks linked to bones that do not exist in the new hierarchy are removed
+			Controller->RemoveBoneTracksMissingFromSkeleton(NewSkeleton);
+			ensure(DataModel->GetNumBoneTracks() != 0);
 
-			if (TrackToSkeletonMapTable.Num() == 0)
-			{
-				// no bones to retarget
-				// return with error
-				//@todo fail message
-			}
 			// make sure you do update reference pose before coming here
 			
 			// first calculate component space ref pose to get the relative transform between
@@ -3649,8 +3100,8 @@ void UAnimSequence::RemapTracksToNewSkeleton( USkeleton* NewSkeleton, bool bConv
 			ConvertedLocalSpaces.AddZeroed(NumBones);
 			ConvertedSpaceBases.AddZeroed(NumBones);
 
-			int32 NumKeys = NumberOfKeys;
-			float Interval = SamplingFrameRate.AsInterval();
+			const int32 NumKeys = DataModel->GetNumberOfKeys();
+			const float Interval = DataModel->GetFrameRate().AsInterval();
 
 			// allocate arrays
 			for(int32 BoneIndex=0; BoneIndex<NumBones; ++BoneIndex)
@@ -3665,18 +3116,17 @@ void UAnimSequence::RemapTracksToNewSkeleton( USkeleton* NewSkeleton, bool bConv
 			for(int32 BoneIndex=0; BoneIndex<NumBones; ++BoneIndex)
 			{
 				FName BoneName = NewSkeleton->GetReferenceSkeleton().GetBoneName(BoneIndex);
-				int32 OldBoneIndex = OldSkeleton->GetReferenceSkeleton().FindBoneIndex(BoneName);
-				int32 TrackIndex = AnimationTrackNames.Find(BoneName);
-				int32 ParentBoneIndex = NewSkeleton->GetReferenceSkeleton().GetParentIndex(BoneIndex);
+				const int32 OldBoneIndex = OldSkeleton->GetReferenceSkeleton().FindBoneIndex(BoneName);
+				const int32 TrackIndex = DataModel->GetBoneTrackIndexByName(BoneName);
+				const int32 ParentBoneIndex = NewSkeleton->GetReferenceSkeleton().GetParentIndex(BoneIndex);
 
 				if(TrackIndex != INDEX_NONE)
 				{
-					auto& RawAnimation = RawAnimationData[TrackIndex];
 					// fill up keys - calculate PK1 * K1
 					for(int32 Key=0; Key<NumKeys; ++Key)
 					{
 						FTransform AnimatedLocalKey;
-						ExtractBoneTransform(RawAnimationData, AnimatedLocalKey, TrackIndex, Interval*Key);
+						UE::Anim::GetBoneTransformFromModel(DataModel, AnimatedLocalKey, TrackIndex, Interval * Key, Interpolation);
 
 						// note that we apply scale in the animated space
 						// at this point, you should have scaled version of animated skeleton
@@ -3711,12 +3161,15 @@ void UAnimSequence::RemapTracksToNewSkeleton( USkeleton* NewSkeleton, bool bConv
 				}
 			}
 
+			TArray<FVector> PosKeys;
+			TArray<FQuat> RotKeys;
+			TArray<FVector> ScaleKeys;
+
 			// now apply the theta back to the animated space bases
-			TArray<struct FRawAnimSequenceTrack> NewRawAnimationData = RawAnimationData;
 			for(int32 BoneIndex=0; BoneIndex<NumBones; ++BoneIndex)
 			{
-				FName BoneName = NewSkeleton->GetReferenceSkeleton().GetBoneName(BoneIndex);
-				int32 TrackIndex = AnimationTrackNames.Find(BoneName);
+				const FName BoneName = NewSkeleton->GetReferenceSkeleton().GetBoneName(BoneIndex);
+				const int32 TrackIndex = DataModel->GetBoneTrackIndexByName(BoneName);
 				int32 ParentBoneIndex = NewSkeleton->GetReferenceSkeleton().GetParentIndex(BoneIndex);
 
 				for(int32 Key=0; Key<NumKeys; ++Key)
@@ -3751,56 +3204,24 @@ void UAnimSequence::RemapTracksToNewSkeleton( USkeleton* NewSkeleton, bool bConv
 				// now save back to animation data
 				if(TrackIndex != INDEX_NONE)
 				{
-					auto& RawAnimation = NewRawAnimationData[TrackIndex];
-					RawAnimation.PosKeys.Empty(NumKeys);
-					RawAnimation.PosKeys.AddUninitialized(NumKeys);
-					RawAnimation.RotKeys.Empty(NumKeys);
-					RawAnimation.RotKeys.AddUninitialized(NumKeys);
-					RawAnimation.ScaleKeys.Empty(NumKeys);
-					RawAnimation.ScaleKeys.AddUninitialized(NumKeys);
+					PosKeys.SetNumUninitialized(NumKeys);
+					RotKeys.SetNumUninitialized(NumKeys);
+					ScaleKeys.SetNumUninitialized(NumKeys);
 
 					for(int32 Key=0; Key<NumKeys; ++Key)
 					{
-						RawAnimation.PosKeys[Key] = ConvertedLocalSpaces[BoneIndex][Key].GetLocation();
-						RawAnimation.RotKeys[Key] = ConvertedLocalSpaces[BoneIndex][Key].GetRotation();
-						RawAnimation.ScaleKeys[Key] = ConvertedLocalSpaces[BoneIndex][Key].GetScale3D();
+						PosKeys[Key] = ConvertedLocalSpaces[BoneIndex][Key].GetLocation();
+						RotKeys[Key] = ConvertedLocalSpaces[BoneIndex][Key].GetRotation();
+						ScaleKeys[Key] = ConvertedLocalSpaces[BoneIndex][Key].GetScale3D();
 					}
+
+					Controller->SetBoneTrackKeys(BoneName, PosKeys, RotKeys, ScaleKeys);
 				}
 			}
-			RawAnimationData = NewRawAnimationData;
 		}
 		else
 		{
-			// this only replaces the primary one, it doesn't replace old ones
-			TArray<struct FTrackToSkeletonMap> NewTrackToSkeletonMapTable;
-			NewTrackToSkeletonMapTable.Empty(AnimationTrackNames.Num());
-			NewTrackToSkeletonMapTable.AddUninitialized(AnimationTrackNames.Num());
-			for (int32 Track = 0; Track < AnimationTrackNames.Num(); ++Track)
-			{
-				int32 BoneIndex = NewSkeleton->GetReferenceSkeleton().FindBoneIndex(AnimationTrackNames[Track]);
-				NewTrackToSkeletonMapTable[Track].BoneTreeIndex = BoneIndex;
-			}
-
-			// now I have all NewTrack To Skeleton Map Table
-			// I'll need to compare with old tracks and copy over if SkeletonIndex == 0
-			// if SkeletonIndex != 0, we need to see if we can 
-			for (int32 TableId = 0; TableId < NewTrackToSkeletonMapTable.Num(); ++TableId)
-			{
-				if (ensure(TrackToSkeletonMapTable.IsValidIndex(TableId)))
-				{
-					if (NewTrackToSkeletonMapTable[TableId].BoneTreeIndex != INDEX_NONE)
-					{
-						TrackToSkeletonMapTable[TableId].BoneTreeIndex = NewTrackToSkeletonMapTable[TableId].BoneTreeIndex;
-					}
-					else
-					{
-						// if not found, delete the track data
-						RemoveTrack(TableId);
-						NewTrackToSkeletonMapTable.RemoveAt(TableId);
-						--TableId;
-					}
-				}
-			}
+			Controller->RemoveBoneTracksMissingFromSkeleton(NewSkeleton);
 		}
 
 		// I have to set this here in order for compression
@@ -3809,7 +3230,7 @@ void UAnimSequence::RemapTracksToNewSkeleton( USkeleton* NewSkeleton, bool bConv
 	}
 	else
 	{
-		VerifyTrackMap(NewSkeleton);
+		Controller->RemoveBoneTracksMissingFromSkeleton(NewSkeleton);
 	}
 
 	Super::RemapTracksToNewSkeleton(NewSkeleton, bConvertSpaces);
@@ -3817,37 +3238,26 @@ void UAnimSequence::RemapTracksToNewSkeleton( USkeleton* NewSkeleton, bool bConv
 
 void UAnimSequence::PostProcessSequence(bool bForceNewRawDatGuid)
 {
-	// pre process before compress raw animation data
-
-	// if scale is too small, zero it out. Cause it hard to retarget when compress
-	// inverse scale is applied to translation, and causing translation to be huge to retarget, but
-	// compression can't handle that much precision. 
-	for (auto Iter = RawAnimationData.CreateIterator(); Iter; ++Iter)
-	{
-		FRawAnimSequenceTrack& RawAnim = (*Iter);
-		SanitizeAnimationTrackData(RawAnim);
-			}
-
-	UpdateFrameRate();
-	CompressRawAnimData();
+	PRAGMA_DISABLE_DEPRECATION_WARNINGS
 	// Apply compression
 	MarkRawDataAsModified(bForceNewRawDatGuid);
 	OnRawDataChanged();
-	// initialize notify track
-	InitializeNotifyTrack();
-	//Make sure we dont have any notifies off the end of the sequence
-	ClampNotifiesAtEndOfSequence();
+	PRAGMA_ENABLE_DEPRECATION_WARNINGS
+
 	// mark package as dirty
 	MarkPackageDirty();
 }
 
 void UAnimSequence::RemoveNaNTracks()
 {
-	bool bRecompress = false;
+	ValidateModel();
 
-	for( int32 TrackIndex=0; TrackIndex<RawAnimationData.Num(); ++TrackIndex )
+	TArray<FName> TracksToRemove;
+
+	const TArray<FBoneAnimationTrack>& BoneAnimationTracks = DataModel->GetBoneAnimationTracks();		
+	for (const FBoneAnimationTrack& AnimationTrack : BoneAnimationTracks)
 	{
-		const FRawAnimSequenceTrack& RawTrack = RawAnimationData[TrackIndex];
+		const FRawAnimSequenceTrack& RawTrack = AnimationTrack.InternalTrackData;
 
 		bool bContainsNaN = false;
 		for ( auto Key : RawTrack.PosKeys )
@@ -3873,46 +3283,45 @@ void UAnimSequence::RemoveNaNTracks()
 
 		if (bContainsNaN)
 		{
-			UE_LOG(LogAnimation, Warning, TEXT("Animation raw data contains NaNs - Removing the following track [%s Track (%s)]"), (GetOuter() ? *GetOuter()->GetFullName() : *GetFullName()), *AnimationTrackNames[TrackIndex].ToString());
-			// remove this track
-			RemoveTrack(TrackIndex);
-			--TrackIndex;
+			UE_LOG(LogAnimation, Warning, TEXT("Animation raw data contains NaNs - Removing the following track [%s Track (%s)]"), (GetOuter() ? *GetOuter()->GetFullName() : *GetFullName()), *AnimationTrack.Name.ToString());
 
-			bRecompress = true;
+			TracksToRemove.Add(AnimationTrack.Name);
 		}
 	}
 
-	if(bRecompress)
+	UAnimDataController::FScopedBracket ScopedBracket(Controller, LOCTEXT("RemoveNaNTracks_Bracket", "Removing track(s) containing NaN key data"));	
+	for (const FName& TrackName : TracksToRemove)
 	{
-		MarkRawDataAsModified();
-		OnRawDataChanged();
-	}
+		Controller->RemoveBoneTrack(TrackName);
+	}	
 }
 
-
 void UAnimSequence::RemoveAllTracks()
 {
+	PRAGMA_DISABLE_DEPRECATION_WARNINGS
 	RawAnimationData.Empty();
 	AnimationTrackNames.Empty();
 	TrackToSkeletonMapTable.Empty();
+	SourceRawAnimationData_DEPRECATED.Empty();
+	PRAGMA_ENABLE_DEPRECATION_WARNINGS
+	
+	ValidateModel();
+	UAnimDataController::FScopedBracket ScopedBracket(Controller, LOCTEXT("RemoveAllTracks_Bracket", "Removing all Bone Animation and Transform Curve Tracks"));
+	Controller->RemoveAllBoneTracks();
+
 	// clear all transform tracks
 	// not deleting curve names from skeleton 
 	// since we don't know if that name is used by other assets
-	RawCurveData.TransformCurves.Empty();
-
-	// recompress and clear
-	PostProcessSequence();
+	Controller->RemoveAllCurvesOfType(ERawCurveTrackTypes::RCT_Transform);	
 }
 
 void UAnimSequence::RemoveTrack(int32 TrackIndex)
 {
-	if (RawAnimationData.IsValidIndex(TrackIndex))
+	ValidateModel();
+	
+	if (DataModel->IsValidBoneTrackIndex(TrackIndex))
 	{
-		RawAnimationData.RemoveAt(TrackIndex);
-		AnimationTrackNames.RemoveAt(TrackIndex);
-		TrackToSkeletonMapTable.RemoveAt(TrackIndex);
-
-		check (RawAnimationData.Num() == AnimationTrackNames.Num() && AnimationTrackNames.Num() == TrackToSkeletonMapTable.Num() );
+		Controller->RemoveBoneTrack(DataModel->GetBoneTrackByIndex(TrackIndex).Name);
 	}
 }
 
@@ -3973,16 +3382,13 @@ int32 FindFirstChildTrack(const USkeleton* MySkeleton, const FReferenceSkeleton&
 
 int32 UAnimSequence::InsertTrack(const FName& BoneName)
 {
-	return InsertTrackInternal(BoneName);
-}
+	ValidateModel();
 
-int32 UAnimSequence::InsertTrackInternal(const FName& BoneName)
-{
 	// first verify if it doesn't exists, if it does, return
-	int32 CurrentTrackIndex = AnimationTrackNames.Find(BoneName);
-	if (CurrentTrackIndex != INDEX_NONE)
+	const int32 ExistingTrackIndex = DataModel->GetBoneTrackIndexByName(BoneName);
+	if (ExistingTrackIndex != INDEX_NONE)
 	{
-		return CurrentTrackIndex;
+		return ExistingTrackIndex;
 	}
 
 	USkeleton * MySkeleton = GetSkeleton();
@@ -3993,8 +3399,13 @@ int32 UAnimSequence::InsertTrackInternal(const FName& BoneName)
 	}
 
 	const FReferenceSkeleton& RefSkeleton = MySkeleton->GetReferenceSkeleton();
-	int32 NewTrackIndex = FindFirstChildTrack(MySkeleton, RefSkeleton, AnimationTrackNames, BoneName);
-	int32 BoneIndex = RefSkeleton.FindBoneIndex(BoneName);
+
+	TArray<FName> TrackNames;
+	DataModel->GetBoneTrackNames(TrackNames);
+
+	const int32 NewTrackIndex = FindFirstChildTrack(MySkeleton, RefSkeleton, TrackNames, BoneName);
+	ensure(NewTrackIndex <= DataModel->GetNumBoneTracks());
+	const int32 BoneIndex = RefSkeleton.FindBoneIndex(BoneName);
 	if (NewTrackIndex != INDEX_NONE)
 	{
 		const TArray<FTransform>& RefPose = RefSkeleton.GetRefBonePose();
@@ -4005,17 +3416,27 @@ int32 UAnimSequence::InsertTrackInternal(const FName& BoneName)
 		RawTrack.ScaleKeys.Add(RefPose[BoneIndex].GetScale3D());
 
 		// now insert to the track
-		RawAnimationData.Insert(RawTrack, NewTrackIndex);
-		AnimationTrackNames.Insert(BoneName, NewTrackIndex);
-
-		RefreshTrackMapFromAnimTrackNames();
-
-		check(RawAnimationData.Num() == AnimationTrackNames.Num() && AnimationTrackNames.Num() == TrackToSkeletonMapTable.Num());
+		Controller->InsertBoneTrack(BoneName, NewTrackIndex);
 	}
 
 	return NewTrackIndex;
 }
 
+void UAnimSequence::ResizeSequence(float NewLength, int32 NewNumFrames, bool bInsert, int32 StartFrame/*inclusive */, int32 EndFrame/*inclusive*/)
+{
+ 	ValidateModel();
+
+	const FFrameRate& FrameRate = DataModel->GetFrameRate();
+	const float T0 = FrameRate.AsSeconds(StartFrame);
+	const float T1 = FrameRate.AsSeconds(EndFrame);
+
+	FFrameRate NewFrameRate = FFrameRate((NewLength / NewNumFrames), 1);
+
+	GetController()->Resize(NewLength, T0, T1);
+
+	ensure(DataModel->GetNumberOfKeys() == NewNumFrames);
+}
+
 bool UAnimSequence::GetAllAnimationSequencesReferred(TArray<UAnimationAsset*>& AnimationAssets, bool bRecursive /*= true*/)
 {
 	Super::GetAllAnimationSequencesReferred(AnimationAssets, bRecursive);
@@ -4042,46 +3463,9 @@ void UAnimSequence::ReplaceReferredAnimations(const TMap<UAnimationAsset*, UAnim
 
 bool UAnimSequence::AddLoopingInterpolation()
 {
-	int32 NumTracks = AnimationTrackNames.Num();
-	float Interval = SamplingFrameRate.AsInterval();
-
-	if(NumberOfKeys > 0)
-	{
-		// added one more key
-		int32 NewNumKeys = NumberOfKeys +1 ;
-
-		// now I need to calculate back to new animation data
-		for(int32 TrackIndex=0; TrackIndex<NumTracks; ++TrackIndex)
-		{
-			auto& RawAnimation = RawAnimationData[TrackIndex];
-			if (RawAnimation.PosKeys.Num() > 1)
-			{
-				FVector FirstKey = RawAnimation.PosKeys[0];
-				RawAnimation.PosKeys.Add(FirstKey);
-			}
-
-			if(RawAnimation.RotKeys.Num() > 1)
-			{
-				FQuat FirstKey = RawAnimation.RotKeys[0];
-				RawAnimation.RotKeys.Add(FirstKey);
-			}
-
-			if(RawAnimation.ScaleKeys.Num() > 1)
-			{
-				FVector FirstKey = RawAnimation.ScaleKeys[0];
-				RawAnimation.ScaleKeys.Add(FirstKey);
-			}
-		}
-
-		NumberOfKeys = NewNumKeys;
-		SetSequenceLength(GetPlayLength() + Interval);
-
-		PostProcessSequence();
-		return true;
-	}
-
-	return false;
+	return UE::Anim::AnimationData::AddLoopingInterpolation(this);
 }
+
 int32 FindParentNodeIndex(URig* Rig, USkeleton* Skeleton, FName ParentNodeName)
 {
 	const int32& ParentNodeIndex = Rig->FindNode(ParentNodeName);
@@ -4102,8 +3486,8 @@ int32 UAnimSequence::GetSpaceBasedAnimationData(TArray< TArray<FTransform> >& An
 	AnimationDataInComponentSpace.AddZeroed(NumBones);
 
 	// 2d array of animated time [boneindex][time key]
-	int32 NumKeys = NumberOfKeys;
-	float Interval = SamplingFrameRate.AsInterval();
+	const int32 NumKeys = DataModel->GetNumberOfKeys();
+	const float Interval = DataModel->GetFrameRate().AsInterval();
 
 	// allocate arrays
 	for (int32 BoneIndex = 0; BoneIndex < NumBones; ++BoneIndex)
@@ -4264,22 +3648,23 @@ int32 UAnimSequence::GetSpaceBasedAnimationData(TArray< TArray<FTransform> >& An
 	}
 	else
 	{
+		ValidateModel();
+
 		// now calculating old animated space bases
 		// this one calculates aniamted space per bones and per key
 		for (int32 BoneIndex = 0; BoneIndex < NumBones; ++BoneIndex)
 		{
-			FName BoneName = MySkeleton->GetReferenceSkeleton().GetBoneName(BoneIndex);
-			int32 TrackIndex = AnimationTrackNames.Find(BoneName);
-			int32 ParentBoneIndex = MySkeleton->GetReferenceSkeleton().GetParentIndex(BoneIndex);
+			const FName BoneName = MySkeleton->GetReferenceSkeleton().GetBoneName(BoneIndex);
+			const int32 TrackIndex = DataModel->GetBoneTrackIndexByName(BoneName);
+			const int32 ParentBoneIndex = MySkeleton->GetReferenceSkeleton().GetParentIndex(BoneIndex);
 
 			if (TrackIndex != INDEX_NONE)
 			{
-				auto& RawAnimation = RawAnimationData[TrackIndex];
 				// fill up keys - calculate PK1 * K1
 				for (int32 Key = 0; Key < NumKeys; ++Key)
 				{
 					FTransform AnimatedLocalKey;
-					ExtractBoneTransform(RawAnimationData, AnimatedLocalKey, TrackIndex, Interval*Key);
+					UE::Anim::GetBoneTransformFromModel(DataModel, AnimatedLocalKey, TrackIndex, Interval * Key, Interpolation);
 
 					if (ParentBoneIndex != INDEX_NONE)
 					{
@@ -4326,6 +3711,8 @@ bool UAnimSequence::ConvertAnimationDataToRiggingData(FAnimSequenceTrackContaine
 		TArray< TArray<FTransform> > AnimationDataInComponentSpace;
 		int32 NumBones = GetSpaceBasedAnimationData(AnimationDataInComponentSpace, NULL);
 
+		const int32 NumModelKeys = DataModel->GetNumberOfKeys();
+
 		if (NumBones > 0)
 		{
 			RiggingAnimationData.Initialize(ValidNodeNames);
@@ -4340,12 +3727,12 @@ bool UAnimSequence::ConvertAnimationDataToRiggingData(FAnimSequenceTrackContaine
 
 				if (ensure(BoneIndex != INDEX_NONE))
 				{
-					Track.PosKeys.Empty(NumberOfKeys);
-					Track.RotKeys.Empty(NumberOfKeys);
-					Track.ScaleKeys.Empty(NumberOfKeys);
-					Track.PosKeys.AddUninitialized(NumberOfKeys);
-					Track.RotKeys.AddUninitialized(NumberOfKeys);
-					Track.ScaleKeys.AddUninitialized(NumberOfKeys);
+					Track.PosKeys.Empty(NumModelKeys);
+					Track.RotKeys.Empty(NumModelKeys);
+					Track.ScaleKeys.Empty(NumModelKeys);
+					Track.PosKeys.AddUninitialized(NumModelKeys);
+					Track.RotKeys.AddUninitialized(NumModelKeys);
+					Track.ScaleKeys.AddUninitialized(NumModelKeys);
 
 					int32 RigConstraintIndex = Rig->FindTransformBaseByNodeName(NodeName);
 
@@ -4364,7 +3751,7 @@ bool UAnimSequence::ConvertAnimationDataToRiggingData(FAnimSequenceTrackContaine
 							if (ParentBoneIndex != INDEX_NONE)
 							{
 								// if no rig control, component space is used
-								for (int32 KeyIndex = 0; KeyIndex < NumberOfKeys; ++KeyIndex)
+								for (int32 KeyIndex = 0; KeyIndex < NumModelKeys; ++KeyIndex)
 								{
 									FTransform ParentTransform = AnimationDataInComponentSpace[ParentBoneIndex][KeyIndex];
 									FTransform RelativeTransform = AnimationDataInComponentSpace[BoneIndex][KeyIndex].GetRelativeTransform(ParentTransform);
@@ -4374,7 +3761,7 @@ bool UAnimSequence::ConvertAnimationDataToRiggingData(FAnimSequenceTrackContaine
 							else
 							{
 								// if no rig control, component space is used
-								for (int32 KeyIndex = 0; KeyIndex < NumberOfKeys; ++KeyIndex)
+								for (int32 KeyIndex = 0; KeyIndex < NumModelKeys; ++KeyIndex)
 								{
 									Track.RotKeys[KeyIndex] = AnimationDataInComponentSpace[BoneIndex][KeyIndex].GetRotation();
 								}
@@ -4383,7 +3770,7 @@ bool UAnimSequence::ConvertAnimationDataToRiggingData(FAnimSequenceTrackContaine
 						else
 						{
 							// if no rig control, component space is used
-							for (int32 KeyIndex = 0; KeyIndex < NumberOfKeys; ++KeyIndex)
+							for (int32 KeyIndex = 0; KeyIndex < NumModelKeys; ++KeyIndex)
 							{
 								Track.RotKeys[KeyIndex] = AnimationDataInComponentSpace[BoneIndex][KeyIndex].GetRotation();
 							}
@@ -4400,7 +3787,7 @@ bool UAnimSequence::ConvertAnimationDataToRiggingData(FAnimSequenceTrackContaine
 							if (ParentBoneIndex != INDEX_NONE)
 							{
 								// if no rig control, component space is used
-								for (int32 KeyIndex = 0; KeyIndex < NumberOfKeys; ++KeyIndex)
+								for (int32 KeyIndex = 0; KeyIndex < NumModelKeys; ++KeyIndex)
 								{
 									FTransform ParentTransform = AnimationDataInComponentSpace[ParentBoneIndex][KeyIndex];
 									FTransform RelativeTransform = AnimationDataInComponentSpace[BoneIndex][KeyIndex].GetRelativeTransform(ParentTransform);
@@ -4410,7 +3797,7 @@ bool UAnimSequence::ConvertAnimationDataToRiggingData(FAnimSequenceTrackContaine
 							}
 							else
 							{
-								for (int32 KeyIndex = 0; KeyIndex < NumberOfKeys; ++KeyIndex)
+								for (int32 KeyIndex = 0; KeyIndex < NumModelKeys; ++KeyIndex)
 								{
 									Track.PosKeys[KeyIndex] = AnimationDataInComponentSpace[BoneIndex][KeyIndex].GetTranslation();
 									Track.ScaleKeys[KeyIndex] = AnimationDataInComponentSpace[BoneIndex][KeyIndex].GetScale3D();
@@ -4419,7 +3806,7 @@ bool UAnimSequence::ConvertAnimationDataToRiggingData(FAnimSequenceTrackContaine
 						}
 						else
 						{
-							for (int32 KeyIndex = 0; KeyIndex < NumberOfKeys; ++KeyIndex)
+							for (int32 KeyIndex = 0; KeyIndex < NumModelKeys; ++KeyIndex)
 							{
 								Track.PosKeys[KeyIndex] = AnimationDataInComponentSpace[BoneIndex][KeyIndex].GetTranslation();
 								Track.ScaleKeys[KeyIndex] = AnimationDataInComponentSpace[BoneIndex][KeyIndex].GetScale3D();
@@ -4429,7 +3816,7 @@ bool UAnimSequence::ConvertAnimationDataToRiggingData(FAnimSequenceTrackContaine
 					else
 					{
 						// if no rig control, component space is used
-						for (int32 KeyIndex = 0; KeyIndex < NumberOfKeys; ++KeyIndex)
+						for (int32 KeyIndex = 0; KeyIndex < NumModelKeys; ++KeyIndex)
 						{
 							Track.PosKeys[KeyIndex] = AnimationDataInComponentSpace[BoneIndex][KeyIndex].GetTranslation();
 							Track.RotKeys[KeyIndex] = AnimationDataInComponentSpace[BoneIndex][KeyIndex].GetRotation();
@@ -4446,16 +3833,18 @@ bool UAnimSequence::ConvertAnimationDataToRiggingData(FAnimSequenceTrackContaine
 	return false;
 }
 
-bool UAnimSequence::ConvertRiggingDataToAnimationData(FAnimSequenceTrackContainer& RiggingAnimationData, bool bPerformPostProcess)
+bool UAnimSequence::ConvertRiggingDataToAnimationData(FAnimSequenceTrackContainer& RiggingAnimationData)
 {
 	if (RiggingAnimationData.GetNum() > 0)
 	{
+		ValidateModel();
+
+		UAnimDataController::FScopedBracket ScopedBracket(Controller, LOCTEXT("ConvertRiggingDataToAnimationData_Bracket", "Converting Rigging Data to Animation Data"));
+
 		TArray< TArray<FTransform> > AnimationDataInComponentSpace;
 		int32 NumBones = GetSpaceBasedAnimationData(AnimationDataInComponentSpace, &RiggingAnimationData);
 
 		USkeleton* MySkeleton = GetSkeleton();
-		TArray<FRawAnimSequenceTrack> OldAnimationData = RawAnimationData;
-		TArray<FName> OldAnimationTrackNames = AnimationTrackNames;
 		TArray<FName> ValidNodeNames;
 		MySkeleton->GetMappedValidNodes(ValidNodeNames);
 		// remove from ValidNodeNames if it doesn't belong to AnimationTrackNames
@@ -4472,14 +3861,14 @@ bool UAnimSequence::ConvertRiggingDataToAnimationData(FAnimSequenceTrackContaine
 
 		// get local spaces
 		// add all tracks?
-		AnimationTrackNames.Empty(ValidNumNodes);
-		AnimationTrackNames.AddUninitialized(ValidNumNodes);
-		RawAnimationData.Empty(ValidNumNodes);
-		RawAnimationData.AddZeroed(ValidNumNodes);
+		Controller->RemoveAllBoneTracks();
 
-		// Clear Transform curve data
-		RawCurveData.DeleteAllCurveData(ERawCurveTrackTypes::RCT_Transform);
+		// Reused track data arrays
+		TArray<FVector> PosKeys;
+		TArray<FQuat> RotKeys;
+		TArray<FVector> ScaleKeys;
 
+		const int32 NumModelKeys = DataModel->GetNumberOfKeys();
 		const FReferenceSkeleton& RefSkeleton = MySkeleton->GetReferenceSkeleton();
 		const URig* Rig = MySkeleton->GetRig();
 		for (int32 NodeIndex = 0; NodeIndex < ValidNumNodes; ++NodeIndex)
@@ -4489,58 +3878,41 @@ bool UAnimSequence::ConvertRiggingDataToAnimationData(FAnimSequenceTrackContaine
 
 			if (BoneIndex != INDEX_NONE)
 			{
-				// add track names
-				AnimationTrackNames[NodeIndex] = BoneName;
+				// Add a new track by name
+				Controller->AddBoneTrack(BoneName);
 
-				// update bone trasfnrom
-				FRawAnimSequenceTrack& Track = RawAnimationData[NodeIndex];
-
-				Track.PosKeys.Empty();
-				Track.RotKeys.Empty();
-				Track.ScaleKeys.Empty();
-				Track.PosKeys.AddUninitialized(NumberOfKeys);
-				Track.RotKeys.AddUninitialized(NumberOfKeys);
-				Track.ScaleKeys.AddUninitialized(NumberOfKeys);
+				PosKeys.SetNumUninitialized(NumModelKeys);
+				RotKeys.SetNumUninitialized(NumModelKeys);
+				ScaleKeys.SetNumUninitialized(NumModelKeys);
 
 				const int32& ParentBoneIndex = RefSkeleton.GetParentIndex(BoneIndex);
 
 				if(ParentBoneIndex != INDEX_NONE)
 				{
-					for(int32 KeyIndex = 0; KeyIndex < NumberOfKeys; ++KeyIndex)
+					for(int32 KeyIndex = 0; KeyIndex < NumModelKeys; ++KeyIndex)
 					{
-						FTransform LocalTransform = AnimationDataInComponentSpace[BoneIndex][KeyIndex].GetRelativeTransform(AnimationDataInComponentSpace[ParentBoneIndex][KeyIndex]);
+						const FTransform LocalTransform = AnimationDataInComponentSpace[BoneIndex][KeyIndex].GetRelativeTransform(AnimationDataInComponentSpace[ParentBoneIndex][KeyIndex]);
 
-						Track.PosKeys[KeyIndex] = LocalTransform.GetTranslation();
-						Track.RotKeys[KeyIndex] = LocalTransform.GetRotation();
-						Track.ScaleKeys[KeyIndex] = LocalTransform.GetScale3D();
+						PosKeys[KeyIndex] = LocalTransform.GetTranslation();
+						RotKeys[KeyIndex] = LocalTransform.GetRotation();
+						ScaleKeys[KeyIndex] = LocalTransform.GetScale3D();
 					}
 				}
 				else
 				{
-					for(int32 KeyIndex = 0; KeyIndex < NumberOfKeys; ++KeyIndex)
+					for(int32 KeyIndex = 0; KeyIndex < NumModelKeys; ++KeyIndex)
 					{
-						FTransform LocalTransform = AnimationDataInComponentSpace[BoneIndex][KeyIndex];
+						const FTransform LocalTransform = AnimationDataInComponentSpace[BoneIndex][KeyIndex];
 
-						Track.PosKeys[KeyIndex] = LocalTransform.GetTranslation();
-						Track.RotKeys[KeyIndex] = LocalTransform.GetRotation();
-						Track.ScaleKeys[KeyIndex] = LocalTransform.GetScale3D();
+						PosKeys[KeyIndex] = LocalTransform.GetTranslation();
+						RotKeys[KeyIndex] = LocalTransform.GetRotation();
+						ScaleKeys[KeyIndex] = LocalTransform.GetScale3D();
 					}
 				}
-			}
-		}
-
-		// recreate track map
-		TrackToSkeletonMapTable.Empty(AnimationTrackNames.Num());
-		TrackToSkeletonMapTable.AddUninitialized(AnimationTrackNames.Num());
-		int32 TrackIdx = 0;
-		for (auto TrackName : AnimationTrackNames)
-		{
-			TrackToSkeletonMapTable[TrackIdx++].BoneTreeIndex = MySkeleton->GetReferenceSkeleton().FindBoneIndex(TrackName);
-		}
 
-		if (bPerformPostProcess)
-		{
-			PostProcessSequence();
+				// Set the track's key data
+				Controller->SetBoneTrackKeys(BoneName, PosKeys, RotKeys, ScaleKeys);
+			}
 		}
 
 		return true;
@@ -4553,59 +3925,60 @@ void UAnimSequence::ClearBakedTransformData()
 {
 	UE_LOG(LogAnimation, Warning, TEXT("[%s] Detected previous edited data is invalidated. Clearing transform curve data and Source Data. This can happen if you do retarget another animation to this. If not, please report back to Epic. "), *GetName());
 	//Clear Transform curve data
-	RawCurveData.DeleteAllCurveData(ERawCurveTrackTypes::RCT_Transform);
+	Controller->RemoveAllCurvesOfType(ERawCurveTrackTypes::RCT_Transform);
 }
 
 bool UAnimSequence::DoesContainTransformCurves() const
 {
-	return (RawCurveData.TransformCurves.Num() > 0);
-	}
+	ValidateModel();
+	return DataModel->GetNumberOfTransformCurves() > 0;
+}
 
+void UAnimSequence::BakeTrackCurvesToRawAnimationTracks(TArray<FRawAnimSequenceTrack>& InOutNewRawTracks, TArray<FName>& InOutNewTrackNames, TArray<FTrackToSkeletonMap>& InOutNewTrackToSkeletonMapTable)
+{
+	ValidateModel();
 
-void UAnimSequence::BakeTrackCurvesToRawAnimationTracks(TArray<FRawAnimSequenceTrack>& NewRawTracks, TArray<FName>& NewTrackNames, TArray<FTrackToSkeletonMap>& NewTrackToSkeletonMapTable)
+	if (DataModel->GetNumberOfTransformCurves() > 0)
 	{
-	if (DoesContainTransformCurves())
-		{
-		USkeleton * CurSkeleton = GetSkeleton();
+		USkeleton* CurSkeleton = GetSkeleton();
 		check(CurSkeleton);
-		
-		NewRawTracks = RawAnimationData;
-		NewTrackToSkeletonMapTable = TrackToSkeletonMapTable;
-		NewTrackNames = AnimationTrackNames;
+
+		const TArray<FBoneAnimationTrack>& BoneAnimationTracks = ResampledAnimationTrackData;
 			
-		VerifyCurveNames<FTransformCurve>(*CurSkeleton, USkeleton::AnimTrackCurveMappingName, RawCurveData.TransformCurves);
-		const FSmartNameMapping*  NameMapping = CurSkeleton->GetSmartNameContainer(USkeleton::AnimTrackCurveMappingName);
-		
+		InOutNewRawTracks.Reserve(BoneAnimationTracks.Num());
+		InOutNewTrackToSkeletonMapTable.Reserve(BoneAnimationTracks.Num());
+		InOutNewTrackNames.Reserve(BoneAnimationTracks.Num());
+
+		for (FRawAnimSequenceTrack& Track : InOutNewRawTracks)
+		{
 		// since now I'm about to modify Scale Keys. I should add all of them here at least one key. 
 		// if all turns out to be same, it will clear it up. 
-		for (FRawAnimSequenceTrack& RawTrack : NewRawTracks)
-		{
-			if (RawTrack.ScaleKeys.Num() == 0)
+			if (Track.ScaleKeys.Num() == 0)
 			{
 				// at least add one
 				static FVector ScaleConstantKey(1.f);
-				RawTrack.ScaleKeys.Add(ScaleConstantKey);
+				Track.ScaleKeys.Add(ScaleConstantKey);
 			}
 		}
 
-		auto InsertNewBoneTrack = [&NewRawTracks, &NewTrackToSkeletonMapTable, &NewTrackNames, CurSkeleton](const FName& BoneName) -> int32
+		auto InsertNewBoneTrack = [&InOutNewRawTracks, &InOutNewTrackToSkeletonMapTable, &InOutNewTrackNames, CurSkeleton](const FName& BoneName) -> int32
 		{
 			const FReferenceSkeleton& RefSkeleton = CurSkeleton->GetReferenceSkeleton();
-			const int32 NewTrackIndex = FindFirstChildTrack(CurSkeleton, RefSkeleton, NewTrackNames, BoneName);
+			const int32 NewTrackIndex = FindFirstChildTrack(CurSkeleton, RefSkeleton, InOutNewTrackNames, BoneName);
 			const int32 BoneIndex = RefSkeleton.FindBoneIndex(BoneName);
 			if (NewTrackIndex != INDEX_NONE)
 			{
 				const TArray<FTransform>& RefPose = RefSkeleton.GetRefBonePose();
 
-				FRawAnimSequenceTrack& RawTrack = NewRawTracks.InsertDefaulted_GetRef(NewTrackIndex);
+				FRawAnimSequenceTrack& RawTrack = InOutNewRawTracks.InsertDefaulted_GetRef(NewTrackIndex);
 				RawTrack.PosKeys.Add(RefPose[BoneIndex].GetTranslation());
 				RawTrack.RotKeys.Add(RefPose[BoneIndex].GetRotation());
 				RawTrack.ScaleKeys.Add(RefPose[BoneIndex].GetScale3D());
 
-				NewTrackToSkeletonMapTable.Insert(BoneIndex, NewTrackIndex);
-				NewTrackNames.Insert(BoneName, NewTrackIndex);
+				InOutNewTrackToSkeletonMapTable.Insert(BoneIndex, NewTrackIndex);
+				InOutNewTrackNames.Insert(BoneName, NewTrackIndex);
 
-				check(NewRawTracks.Num() == NewTrackToSkeletonMapTable.Num());
+				check(InOutNewRawTracks.Num() == InOutNewTrackToSkeletonMapTable.Num());
 			}
 
 			return NewTrackIndex;
@@ -4613,24 +3986,26 @@ void UAnimSequence::BakeTrackCurvesToRawAnimationTracks(TArray<FRawAnimSequenceT
 
 
 		// NumFrames can't be zero (filtered earlier)
-		const float Interval = SamplingFrameRate.AsInterval();
+		const float Interval = GetIntervalPerKey(NumberOfSampledKeys, GetPlayLength());
 
-		for (const FTransformCurve& Curve : RawCurveData.TransformCurves)
+		const FSmartNameMapping* NameMapping = CurSkeleton->GetSmartNameContainer(USkeleton::AnimTrackCurveMappingName);
+		const FAnimationCurveData& AnimationCurveData = DataModel->GetCurveData();
+		for (const FTransformCurve& Curve : AnimationCurveData.TransformCurves)
 		{
 			// find curves first, and then see what is index of this curve
 			FName BoneName;
 
-			if(Curve.GetCurveTypeFlag(AACF_Disabled)== false &&
+			if (Curve.GetCurveTypeFlag(AACF_Disabled) == false &&
 				ensureAlways(NameMapping->GetName(Curve.Name.UID, BoneName)))
 			{
-				int32 TrackIndex = NewTrackNames.Find(BoneName);
+				int32 TrackIndex = InOutNewTrackNames.Find(BoneName);
 
 				// the animation data doesn't have this track, so insert it
-				if(TrackIndex == INDEX_NONE)
+				if (TrackIndex == INDEX_NONE)
 				{
 					TrackIndex = InsertNewBoneTrack(BoneName);
 					// if it still didn't find, something went horribly wrong
-					if(ensure(TrackIndex != INDEX_NONE) == false)
+					if (ensure(TrackIndex != INDEX_NONE) == false)
 					{
 						UE_LOG(LogAnimation, Warning, TEXT("Animation Baking : Error adding %s track."), *BoneName.ToString());
 						// I can't do anything about it
@@ -4638,69 +4013,59 @@ void UAnimSequence::BakeTrackCurvesToRawAnimationTracks(TArray<FRawAnimSequenceT
 					}
 				}
 
-			// now modify data
-				FRawAnimSequenceTrack& RawTrack = NewRawTracks[TrackIndex];
-
-			// since now we're editing keys, 
-			// if 1 (which meant constant), just expands to # of frames
-			if(RawTrack.PosKeys.Num() == 1)
-			{
-				FVector OneKey = RawTrack.PosKeys[0];
-					RawTrack.PosKeys.Init(OneKey, NumberOfKeys);
-			}
-			else
-			{
-					ensure(RawTrack.PosKeys.Num() == NumberOfKeys);
-			}
-
-			if(RawTrack.RotKeys.Num() == 1)
-			{
-				FQuat OneKey = RawTrack.RotKeys[0];
-					RawTrack.RotKeys.Init(OneKey, NumberOfKeys);
-			}
-			else
-			{
-					ensure(RawTrack.RotKeys.Num() == NumberOfKeys);
-			}
+				// now modify data
+				FRawAnimSequenceTrack& RawTrack = InOutNewRawTracks[TrackIndex];
 
-			// although we don't allow edit of scale
-			// it is important to consider scale when apply transform
-			// so make sure this also is included
-			if(RawTrack.ScaleKeys.Num() == 1)
-			{
-				FVector OneKey = RawTrack.ScaleKeys[0];
-					RawTrack.ScaleKeys.Init(OneKey, NumberOfKeys);
-			}
-			else
-			{
-					ensure(RawTrack.ScaleKeys.Num() == NumberOfKeys);
-			}
+				// since now we're editing keys, 
+				// if 1 (which meant constant), just expands to # of frames
+				if (RawTrack.PosKeys.Num() == 1)
+				{
+					FVector OneKey = RawTrack.PosKeys[0];
+					RawTrack.PosKeys.Init(OneKey, NumberOfSampledKeys);
+				}
+				else
+				{
+					ensure(RawTrack.PosKeys.Num() == NumberOfSampledKeys);
+				}
 
-			// now we have all data ready to apply
-				for (int32 KeyIndex = 0; KeyIndex < NumberOfKeys; ++KeyIndex)
-			{
-				// now evaluate
-				FTransformCurve* TransformCurve = static_cast<FTransformCurve*>(RawCurveData.GetCurveData(Curve.Name.UID, ERawCurveTrackTypes::RCT_Transform));
+				if (RawTrack.RotKeys.Num() == 1)
+				{
+					FQuat OneKey = RawTrack.RotKeys[0];
+					RawTrack.RotKeys.Init(OneKey, NumberOfSampledKeys);
+				}
+				else
+				{
+					ensure(RawTrack.RotKeys.Num() == NumberOfSampledKeys);
+				}
 
-				if(ensure(TransformCurve))
+				// although we don't allow edit of scale
+				// it is important to consider scale when apply transform
+				// so make sure this also is included
+				if (RawTrack.ScaleKeys.Num() == 1)
 				{
-					FTransform AdditiveTransform = TransformCurve->Evaluate(KeyIndex * Interval, 1.0);
-					FTransform LocalTransform(RawTrack.RotKeys[KeyIndex], RawTrack.PosKeys[KeyIndex], RawTrack.ScaleKeys[KeyIndex]);
-					RawTrack.RotKeys[KeyIndex] = LocalTransform.GetRotation() * AdditiveTransform.GetRotation();
-					RawTrack.PosKeys[KeyIndex] = LocalTransform.TransformPosition(AdditiveTransform.GetTranslation());
-					RawTrack.ScaleKeys[KeyIndex] = LocalTransform.GetScale3D() * AdditiveTransform.GetScale3D();
+					FVector OneKey = RawTrack.ScaleKeys[0];
+					RawTrack.ScaleKeys.Init(OneKey, NumberOfSampledKeys);
 				}
 				else
 				{
-					UE_LOG(LogAnimation, Warning, TEXT("Animation Baking : Missing Curve for %s."), *BoneName.ToString());
+					ensure(RawTrack.ScaleKeys.Num() == NumberOfSampledKeys);
+				}
+
+				// now we have all data ready to apply
+				for (int32 KeyIndex = 0; KeyIndex < NumberOfSampledKeys; ++KeyIndex)
+				{
+					// now evaluate
+					FTransform AdditiveTransform = Curve.Evaluate(KeyIndex * Interval, 1.0);
+						FTransform LocalTransform(RawTrack.RotKeys[KeyIndex], RawTrack.PosKeys[KeyIndex], RawTrack.ScaleKeys[KeyIndex]);
+						RawTrack.RotKeys[KeyIndex] = LocalTransform.GetRotation() * AdditiveTransform.GetRotation();
+						RawTrack.PosKeys[KeyIndex] = LocalTransform.TransformPosition(AdditiveTransform.GetTranslation());
+						RawTrack.ScaleKeys[KeyIndex] = LocalTransform.GetScale3D() * AdditiveTransform.GetScale3D();
 				}
-			}
 
 				// Apply PostProcess behaviour in-place
-				SanitizeAnimationTrackData(RawTrack);
-}
-}
-		StaticCompressRawAnimData(NewRawTracks, NumberOfKeys, GetFName());
+				UE::Anim::Compression::SanitizeRawAnimSequenceTrack(RawTrack);
+			}
+		}
 	}
 }
 
@@ -4711,67 +4076,29 @@ void UAnimSequence::AddKeyToSequence(float Time, const FName& BoneName, const FT
 	USkeleton * CurrentSkeleton = GetSkeleton();
 	check (CurrentSkeleton);
 
-	FSmartName NewCurveName;
-	CurrentSkeleton->AddSmartNameAndModify(USkeleton::AnimTrackCurveMappingName, CurveName, NewCurveName);
-
-	// add curve - this won't add duplicate curve
-	RawCurveData.AddCurveData(NewCurveName, AACF_DriveTrack | AACF_Editable, ERawCurveTrackTypes::RCT_Transform);
-
-	//Add this curve
-	FTransformCurve* TransformCurve = static_cast<FTransformCurve*>(RawCurveData.GetCurveData(NewCurveName.UID, ERawCurveTrackTypes::RCT_Transform));
-	check(TransformCurve);
-
-	TransformCurve->UpdateOrAddKey(AdditiveTransform, Time);	
-	MarkRawDataAsModified();
-}
-
-void UAnimSequence::ResetAnimation()
-{
-	// clear everything. Making new animation, so need to reset all the things that belong here
-	NumberOfKeys = 0;
-	PRAGMA_DISABLE_DEPRECATION_WARNINGS
-	SequenceLength = 0.f;
-	PRAGMA_ENABLE_DEPRECATION_WARNINGS
-	RawAnimationData.Empty();
-	AnimationTrackNames.Empty();
-	TrackToSkeletonMapTable.Empty();
+	FSmartName NewCurveName;
+	CurrentSkeleton->AddSmartNameAndModify(USkeleton::AnimTrackCurveMappingName, CurveName, NewCurveName);
 
-	PRAGMA_DISABLE_DEPRECATION_WARNINGS
-	ClearCompressedBoneData();
-	ClearCompressedCurveData();
-	PRAGMA_ENABLE_DEPRECATION_WARNINGS
 
-	Notifies.Empty();
-	AuthoredSyncMarkers.Empty();
-	UniqueMarkerNames.Empty();
-	AnimNotifyTracks.Empty();
-	RawCurveData.Empty();
-	RateScale = 1.f;
+	ValidateModel();
+
+	UAnimDataController::FScopedBracket ScopedBracket(Controller, LOCTEXT("AddKeyToSequence_Bracket", "Adding key to sequence"));
+	FAnimationCurveIdentifier TransformCurveId(NewCurveName, ERawCurveTrackTypes::RCT_Transform);
+	Controller->AddCurve(TransformCurveId, AACF_DriveTrack | AACF_Editable);
+
+	Controller->RemoveBoneTracksMissingFromSkeleton(CurrentSkeleton);
+
+	const FTransformCurve* TransformCurve = DataModel->FindTransformCurve(TransformCurveId);
+	check(TransformCurve);
+
+	Controller->SetTransformCurveKey(TransformCurveId, Time, AdditiveTransform);
 }
 
 void UAnimSequence::RefreshTrackMapFromAnimTrackNames()
 {
-	TrackToSkeletonMapTable.Empty();
-
+	ValidateModel();
 	const USkeleton * MySkeleton = GetSkeleton();
-	const FReferenceSkeleton& RefSkeleton = MySkeleton->GetReferenceSkeleton();
-	const int32 NumBones = AnimationTrackNames.Num();
-	TrackToSkeletonMapTable.AddUninitialized(NumBones);
-
-	bool bNeedsFixing = false;
-	const int32 NumTracks = AnimationTrackNames.Num();
-	for(int32 I=NumTracks-1; I>=0; --I)
-	{
-		int32 BoneTreeIndex = RefSkeleton.FindBoneIndex(AnimationTrackNames[I]);
-		if(BoneTreeIndex == INDEX_NONE)
-		{
-			RemoveTrack(I);
-		}
-		else
-		{
-			TrackToSkeletonMapTable[I].BoneTreeIndex = BoneTreeIndex;
-		}
-	}
+	Controller->RemoveBoneTracksMissingFromSkeleton(MySkeleton);
 }
 
 uint8* UAnimSequence::FindSyncMarkerPropertyData(int32 SyncMarkerIndex, FArrayProperty*& ArrayProperty)
@@ -4790,36 +4117,31 @@ bool UAnimSequence::CreateAnimation(USkeletalMesh* Mesh)
 	// create animation from Mesh's ref pose
 	if (Mesh)
 	{
+		ValidateModel();
+		const FReferenceSkeleton& RefSkeleton = Mesh->GetRefSkeleton();
+
+		UAnimDataController::FScopedBracket ScopedBracket(Controller, LOCTEXT("CreateAnimationSkeletalMesh_Bracket", "Creating Animation Sequence based up Skeletal Mesh"));
 		ResetAnimation();
 
-		const FReferenceSkeleton& RefSkeleton = Mesh->GetRefSkeleton();
-		SetSequenceLength(MINIMUM_ANIMATION_LENGTH);
-		NumberOfKeys = 1;
+		Controller->SetPlayLength(MINIMUM_ANIMATION_LENGTH);
+		Controller->SetFrameRate(FFrameRate(1.f / MINIMUM_ANIMATION_LENGTH, 1));
 
 		const int32 NumBones = RefSkeleton.GetRawBoneNum();
-		RawAnimationData.AddZeroed(NumBones);
-		AnimationTrackNames.AddUninitialized(NumBones);
 
 		const TArray<FTransform>& RefBonePose = RefSkeleton.GetRawRefBonePose();
 
-		check (RefBonePose.Num() == NumBones);
+		check(RefBonePose.Num() == NumBones);
 
-		for (int32 BoneIndex=0; BoneIndex<NumBones; ++BoneIndex)
+		for (int32 BoneIndex = 0; BoneIndex < NumBones; ++BoneIndex)
 		{
-			AnimationTrackNames[BoneIndex] = RefSkeleton.GetBoneName(BoneIndex);
-
-			FRawAnimSequenceTrack& RawTrack = RawAnimationData[BoneIndex];
-
-			RawTrack.PosKeys.Add(RefBonePose[BoneIndex].GetTranslation());
-			RawTrack.RotKeys.Add(RefBonePose[BoneIndex].GetRotation());
-			RawTrack.ScaleKeys.Add(RefBonePose[BoneIndex].GetScale3D());
+			const FName& BoneName = RefSkeleton.GetBoneName(BoneIndex);
+			Controller->AddBoneTrack(BoneName);
+			Controller->SetBoneTrackKeys(BoneName, { RefBonePose[BoneIndex].GetTranslation() }, { RefBonePose[BoneIndex].GetRotation() }, { RefBonePose[BoneIndex].GetScale3D() });
 		}
 
-		// refresh TrackToskeletonMapIndex
-		RefreshTrackMapFromAnimTrackNames();
+		Controller->NotifyPopulated();
+		
 
-		// should recreate track map
-		PostProcessSequence();
 		return true;
 	}
 
@@ -4832,36 +4154,29 @@ bool UAnimSequence::CreateAnimation(USkeletalMeshComponent* MeshComponent)
 	{
 		USkeletalMesh * Mesh = MeshComponent->SkeletalMesh;
 
+		ValidateModel();
+		const FReferenceSkeleton& RefSkeleton = Mesh->GetRefSkeleton();
+
+		UAnimDataController::FScopedBracket ScopedBracket(Controller, LOCTEXT("CreateAnimationSkeletalMeshComponent_Bracket", "Creating Animation Sequence based up Skeletal Mesh Component"));
 		ResetAnimation();
 
-		const FReferenceSkeleton& RefSkeleton = Mesh->GetRefSkeleton();
-		SetSequenceLength(MINIMUM_ANIMATION_LENGTH);
-		NumberOfKeys = 1;
+		Controller->SetPlayLength(MINIMUM_ANIMATION_LENGTH);
+		Controller->SetFrameRate(FFrameRate(1.f / MINIMUM_ANIMATION_LENGTH, 1));
 
 		const int32 NumBones = RefSkeleton.GetRawBoneNum();
-		RawAnimationData.AddZeroed(NumBones);
-		AnimationTrackNames.AddUninitialized(NumBones);
 
 		const TArray<FTransform> BoneSpaceTransforms = MeshComponent->GetBoneSpaceTransforms();
 
 		check(BoneSpaceTransforms.Num() >= NumBones);
 
-		for(int32 BoneIndex=0; BoneIndex<NumBones; ++BoneIndex)
+		for (int32 BoneIndex = 0; BoneIndex < NumBones; ++BoneIndex)
 		{
-			AnimationTrackNames[BoneIndex] = RefSkeleton.GetBoneName(BoneIndex);
-
-			FRawAnimSequenceTrack& RawTrack = RawAnimationData[BoneIndex];
-
-			RawTrack.PosKeys.Add(BoneSpaceTransforms[BoneIndex].GetTranslation());
-			RawTrack.RotKeys.Add(BoneSpaceTransforms[BoneIndex].GetRotation());
-			RawTrack.ScaleKeys.Add(BoneSpaceTransforms[BoneIndex].GetScale3D());
+			const FName& BoneName = RefSkeleton.GetBoneName(BoneIndex);
+			Controller->AddBoneTrack(BoneName);
+			Controller->SetBoneTrackKeys(BoneName, { BoneSpaceTransforms[BoneIndex].GetTranslation() }, { BoneSpaceTransforms[BoneIndex].GetRotation() }, { BoneSpaceTransforms[BoneIndex].GetScale3D() });
 		}
 
-		// refresh TrackToskeletonMapIndex
-		RefreshTrackMapFromAnimTrackNames();
-
-		// should recreate track map
-		PostProcessSequence();
+		Controller->NotifyPopulated();
 		return true;
 	}
 
@@ -4874,22 +4189,13 @@ bool UAnimSequence::CreateAnimation(UAnimSequence* Sequence)
 	{
 		ResetAnimation();
 
-		SetSequenceLength(Sequence->GetPlayLength());
-		NumberOfKeys = Sequence->NumberOfKeys;
-
-		RawAnimationData = Sequence->RawAnimationData;
-		AnimationTrackNames = Sequence->AnimationTrackNames;
-
+		CopyDataModel(Sequence->GetDataModel());
+		Controller->SetModel(DataModel);
 		Notifies = Sequence->Notifies;
 		AnimNotifyTracks = Sequence->AnimNotifyTracks;
-		RawCurveData = Sequence->RawCurveData;
 
-		UpdateFrameRate();
-		// refresh TrackToskeletonMapIndex
-		RefreshTrackMapFromAnimTrackNames();
+		Controller->NotifyPopulated();
 
-		// should recreate track map
-		PostProcessSequence();
 		return true;
 	}
 	
@@ -4925,6 +4231,15 @@ void UAnimSequence::RefreshCacheData()
 	Super::RefreshCacheData();
 }
 
+int32 UAnimSequence::GetNumberOfSampledKeys() const 
+{ 
+#if WITH_EDITOR
+	return NumberOfSampledKeys;
+#else
+	return CompressedData.CompressedDataStructure->CompressedNumberOfKeys;
+#endif
+}
+
 void UAnimSequence::EvaluateCurveData(FBlendedCurve& OutCurve, float CurrentTime, bool bForceUseRawData) const
 {
 	SCOPE_CYCLE_COUNTER(STAT_AnimSeq_EvalCurveData);
@@ -5504,7 +4819,7 @@ void UAnimSequence::SynchronousCustomAttributesCompression()
 			}
 			else if (RefPoseType == ABPT_AnimFrame)
 			{
-				const float Fraction = (RefPoseSeq->NumberOfKeys > 0) ? FMath::Clamp<float>((float)RefFrameIndex / (float)RefPoseSeq->NumberOfKeys, 0.f, 1.f) : 0.f;
+				const float Fraction = (RefPoseSeq->GetNumberOfSampledKeys() > 0) ? FMath::Clamp<float>((float)RefFrameIndex / (float)RefPoseSeq->GetNumberOfSampledKeys(), 0.f, 1.f) : 0.f;
 				BasePoseTime = RefPoseSeq->GetPlayLength() * Fraction;
 
 			}
@@ -5552,7 +4867,7 @@ void UAnimSequence::SynchronousCustomAttributesCompression()
 								BakedFloatAttribute.AttributeName = Attribute.Name;
 
 								FSimpleCurve& FloatCurve = BakedFloatAttribute.FloatCurve;
-								ConvertAttributeToAdditive<float, FSimpleCurve>(Attribute, *RefAttribute, FloatCurve, EvalContext.IntervalTime, NumberOfKeys, GetBasePoseTimeToSample);
+								ConvertAttributeToAdditive<float, FSimpleCurve>(Attribute, *RefAttribute, FloatCurve, EvalContext.IntervalTime, NumberOfSampledKeys, GetBasePoseTimeToSample);
 								FloatCurve.RemoveRedundantKeys(0.f);
 
 								break;
@@ -5564,7 +4879,7 @@ void UAnimSequence::SynchronousCustomAttributesCompression()
 								BakedIntAttribute.AttributeName = Attribute.Name;
 
 								FIntegralCurve& IntCurve = BakedIntAttribute.IntCurve;
-								ConvertAttributeToAdditive<int32, FIntegralCurve>(Attribute, *RefAttribute, IntCurve, EvalContext.IntervalTime, NumberOfKeys, GetBasePoseTimeToSample);
+								ConvertAttributeToAdditive<int32, FIntegralCurve>(Attribute, *RefAttribute, IntCurve, EvalContext.IntervalTime, NumberOfSampledKeys, GetBasePoseTimeToSample);
 								IntCurve.RemoveRedundantKeys();
 							
 								break;
@@ -6062,13 +5377,7 @@ void UAnimSequence::EnableRootMotionSettingFromMontage(bool bInEnableRootMotion,
 #if WITH_EDITOR
 void UAnimSequence::OnRawDataChanged()
 {
-	PRAGMA_DISABLE_DEPRECATION_WARNINGS
-	ClearCompressedBoneData();
-	PRAGMA_ENABLE_DEPRECATION_WARNINGS
-
-	bUseRawDataOnly = true;
-
-	RequestAsyncAnimRecompression(false);
+	RecompressAnimationData();
 	//MDW - Once we have async anim ddc requests we should do this too
 	//RequestDependentAnimRecompression();
 }
@@ -6078,21 +5387,27 @@ bool UAnimSequence::IsCompressedDataValid() const
 {
 	// For bone compressed data, we don't check if we have a codec. It is valid to have no compressed data
 	// if we have no raw data. This can happen with sequences that only has curves.
-
-	if (RawAnimationData.Num() == 0)
+#if WITH_EDITOR
+	ValidateModel();
+	PRAGMA_DISABLE_DEPRECATION_WARNINGS
+	if (DataModel->GetNumBoneTracks() == 0 && RawAnimationData.Num() == 0)
+	PRAGMA_ENABLE_DEPRECATION_WARNINGS
 	{
 		return true;
 	}
+#endif
 
-	return CompressedData.CompressedDataStructure != nullptr;
+	return CompressedData.CompressedDataStructure != nullptr || CompressedData.CompressedTrackToSkeletonMapTable.Num() == 0;
 }
 
 bool UAnimSequence::IsCurveCompressedDataValid() const
 {
+#if WITH_EDITOR
+	ValidateModel();
 	// For curve compressed data, we don't check if we have a codec. It is valid to have no compressed data
 	// if we have no raw data. This can happen with sequences that only has bones.
 
-	if (CompressedData.CompressedCurveByteStream.Num() == 0 && RawCurveData.FloatCurves.Num() != 0)
+	if (CompressedData.CompressedCurveByteStream.Num() == 0 && DataModel->GetNumberOfFloatCurves() != 0)
 	{
 		// No compressed data but we have raw data
 		if (!IsValidAdditive())
@@ -6103,6 +5418,7 @@ bool UAnimSequence::IsCurveCompressedDataValid() const
 		// Additive sequences can have raw curves that all end up being 0.0 (e.g. they 100% match the base sequence curves)
 		// in which case there will be no compressed curve data.
 	}
+#endif
 
 	return true;
 }
@@ -6117,6 +5433,369 @@ void UAnimSequence::ClearCompressedCurveData()
 	CompressedData.ClearCompressedCurveData();
 }
 
+#if WITH_EDITOR
+void UAnimSequence::EnsureValidRawDataGuid()
+{
+PRAGMA_DISABLE_DEPRECATION_WARNINGS
+	if (!RawDataGuid.IsValid())
+	{
+		RawDataGuid = GenerateGuidFromRawData();
+	}
+PRAGMA_ENABLE_DEPRECATION_WARNINGS
+}
+
+void UAnimSequence::PopulateModel()
+{
+	check(!HasAnyFlags(EObjectFlags::RF_ClassDefaultObject));
+	
+	PRAGMA_DISABLE_DEPRECATION_WARNINGS
+	const int32 NumKeys = NumberOfKeys;
+	const float PlayLength = SequenceLength;
+	// Reset target framerate to current sampling frame rate
+	TargetFrameRate = SamplingFrameRate;
+	const FFrameRate FrameRate = SamplingFrameRate;
+	const FRawCurveTracks CurveData = RawCurveData;
+	const TArray<FRawAnimSequenceTrack> SequenceTracks = SourceRawAnimationData_DEPRECATED.Num() && SourceRawAnimationData_DEPRECATED.Num() == RawAnimationData.Num() ?
+		SourceRawAnimationData_DEPRECATED :
+		RawAnimationData;
+	const TArray<FName> TempAnimationTrackNames = AnimationTrackNames;
+	
+	const bool bCurveOnlySequence = RawAnimationData.Num() == 0;
+	const bool bValidNumberOfKeys = NumberOfKeys >= 1 || bCurveOnlySequence;
+	const float ExpectedPlayLength = (float)SamplingFrameRate.AsSeconds(FMath::Max(NumberOfKeys - 1, 1));
+	const bool bValidSequenceLength = FMath::IsNearlyEqual(ExpectedPlayLength, SequenceLength, KINDA_SMALL_NUMBER) || bCurveOnlySequence;
+	const bool bValidNumberOfTrackArrayEntries = AnimationTrackNames.Num() == RawAnimationData.Num() && RawAnimationData.Num() == TrackToSkeletonMapTable.Num();
+	const bool bValidAnimationData = SourceRawAnimationData_DEPRECATED.Num() == 0 || SourceRawAnimationData_DEPRECATED.Num() == RawAnimationData.Num();
+	const bool bValidFrameRate = SamplingFrameRate.IsValid() && (SamplingFrameRate.AsInterval() >= 0.f);
+	const bool bValidCurveData = RawCurveData.FloatCurves.Num() >= 1;
+
+	if (!((bValidNumberOfTrackArrayEntries || bValidCurveData) && bValidAnimationData && bValidFrameRate && bValidNumberOfKeys && bValidSequenceLength))
+	{
+		UE_LOG(LogAnimation, Warning, TEXT("Possible invalid animation data during model generation %s - \n\tNumber of Keys: %i\n\tPlay length: %f\n\tNumber of Tracks: %i\n\tFrameRate: %s\n\tNumber of float Curves: %i"), *GetName(), NumKeys, PlayLength, RawAnimationData.Num(), *SamplingFrameRate.ToPrettyText().ToString(), RawCurveData.FloatCurves.Num());
+	}
+	PRAGMA_ENABLE_DEPRECATION_WARNINGS
+
+	if (RefPoseSeq && RefPoseSeq->HasAnyFlags(EObjectFlags::RF_NeedPostLoad))
+	{
+		RefPoseSeq->ConditionalPostLoad();
+	}
+
+	Controller->OpenBracket(LOCTEXT("UAnimSequence::PopulateModel_Bracket", "Generating Animation Model Data from Animation Sequence"));
+
+	Controller->SetPlayLength(bCurveOnlySequence ? PlayLength : ExpectedPlayLength);
+	Controller->SetFrameRate(FrameRate);
+
+	USkeleton* TargetSkeleton = GetSkeleton();
+	UE::Anim::CopyCurveDataToModel(CurveData, TargetSkeleton, Controller);
+
+	const int32 NumTracks = SequenceTracks.Num();
+	for (int32 TrackIndex = 0; TrackIndex < NumTracks; ++TrackIndex)
+	{
+		Controller->AddBoneTrack(TempAnimationTrackNames[TrackIndex]);
+
+		auto GenerateUniformKeys = [NumKeys](auto& Keys, auto IdentityKey)
+		{
+			// Convert track keys to be uniform			
+			if (Keys.Num() == 0)
+			{
+				// set all to identity
+				for (int32 Index = 0; Index < NumKeys; ++Index)
+				{
+					Keys.Add(IdentityKey);
+				}
+			}
+			else if (Keys.Num() == 1)
+			{
+				// set all to single key
+				const auto KeyZero = Keys[0];
+				for (int32 Index = 0; Index < (NumKeys - 1); ++Index)
+				{
+					Keys.Add(KeyZero);
+				}
+			}
+
+			ensure(Keys.Num() == NumKeys);
+		};
+
+		TArray<FVector> PosKeys = SequenceTracks[TrackIndex].PosKeys;
+		GenerateUniformKeys(PosKeys, FVector::ZeroVector);
+
+		TArray<FQuat> RotKeys = SequenceTracks[TrackIndex].RotKeys;
+		GenerateUniformKeys(RotKeys, FQuat::Identity);
+
+		TArray<FVector> ScaleKeys = SequenceTracks[TrackIndex].ScaleKeys;
+		GenerateUniformKeys(ScaleKeys, FVector::OneVector);
+
+		Controller->SetBoneTrackKeys(TempAnimationTrackNames[TrackIndex], PosKeys, RotKeys, ScaleKeys);
+	}
+
+	DeleteDeprecatedRawAnimationData();
+	
+	Controller->CloseBracket();
+
+	ResampleAnimationTrackData();
+	RequestSyncAnimRecompression(false);
+}
+
+void UAnimSequence::OnModelModified(const EAnimDataModelNotifyType& NotifyType, UAnimDataModel* Model, const FAnimDataModelNotifPayload& Payload)
+{
+	UAnimSequenceBase::OnModelModified(NotifyType, Model, Payload);	
+
+	static const bool GenerateNewGUID = true;
+	static const bool RegenerateGUID = false;
+
+	auto UpdateRawDataGuid = [this](bool bForceNewRawDataGuid)
+	{
+		RawDataGuid = bForceNewRawDataGuid ? FGuid::NewGuid() : GenerateGuidFromRawData();
+		bUseRawDataOnly = true;
+
+		FlagDependentAnimationsAsRawDataOnly();
+		UpdateDependentStreamingAnimations();
+	};
+
+	auto HandleTrackDataChanged = [this, UpdateRawDataGuid]()
+	{
+		if (NotifyCollector.IsNotWithinBracket())
+		{
+			ResampleAnimationTrackData();
+			UpdateRawDataGuid(RegenerateGUID);
+			RecompressAnimationData();
+		}
+	};
+
+	auto HandleLengthChanged = [this](float NewLength, float OldLength, float T0, float T1)
+	{
+		if (NewLength > OldLength)
+		{
+			const float InsertTime = T0;
+			const float Duration = T1 - T0;
+
+			for (FAnimSyncMarker& Marker : AuthoredSyncMarkers)
+			{
+				float CurrentTime = Marker.Time;
+
+				// when insert, we only care about start time
+				// if it's later than start time
+				if (CurrentTime >= InsertTime)
+				{
+					CurrentTime += Duration;
+				}
+
+				Marker.Time = FMath::Clamp(CurrentTime, 0.f, NewLength);
+			}
+		}
+		else if (NewLength < OldLength)
+		{
+			const float StartRemoveTime = T0;
+			const float EndRemoveTime = T1;
+
+			// Total time value for frames that were removed
+			const float Duration = T1 - T0;
+
+			for (FAnimSyncMarker& Marker : AuthoredSyncMarkers)
+			{
+				float CurrentTime = Marker.Time;
+
+				if (CurrentTime >= StartRemoveTime && CurrentTime <= EndRemoveTime)
+				{
+					CurrentTime = StartRemoveTime;
+				}
+				else if (CurrentTime > EndRemoveTime)
+				{
+					CurrentTime -= Duration;
+				}
+
+				Marker.Time = FMath::Clamp(CurrentTime, 0.f, NewLength);
+			}
+
+		}
+	};
+
+	PRAGMA_DISABLE_DEPRECATION_WARNINGS
+	switch (NotifyType)
+	{
+		case EAnimDataModelNotifyType::SequenceLengthChanged:		
+		{
+			const FSequenceLengthChangedPayload& TypedPayload = Payload.GetPayload<FSequenceLengthChangedPayload>();
+
+			const float OldSequenceLength = TypedPayload.PreviousLength;
+			const float CurrentSequenceLength = Model->GetPlayLength();
+			
+			HandleLengthChanged(CurrentSequenceLength, OldSequenceLength, TypedPayload.T0, TypedPayload.T1);
+						
+			if (NotifyCollector.IsNotWithinBracket())
+			{
+				HandleTrackDataChanged();
+			}
+
+			break;
+		}
+
+		case EAnimDataModelNotifyType::FrameRateChanged:
+		{
+			const FFrameRateChangedPayload& TypedPayload = Payload.GetPayload<FFrameRateChangedPayload>();
+			
+			if (NotifyCollector.IsNotWithinBracket())
+			{
+				HandleTrackDataChanged();
+			}
+
+			break;
+		}
+
+		case EAnimDataModelNotifyType::Populated:
+		{
+			TargetFrameRate = Model->GetFrameRate();
+
+			if (NotifyCollector.IsNotWithinBracket())
+			{
+				HandleTrackDataChanged();
+			}	
+			break;
+		}
+		
+		case EAnimDataModelNotifyType::BracketClosed:
+		{
+			if (NotifyCollector.IsNotWithinBracket())
+			{
+				const auto LengthChangingNotifies = { EAnimDataModelNotifyType::SequenceLengthChanged, EAnimDataModelNotifyType::FrameRateChanged, EAnimDataModelNotifyType::Reset };
+				const auto ResamplingNotifies = { EAnimDataModelNotifyType::TrackAdded, EAnimDataModelNotifyType::TrackChanged, EAnimDataModelNotifyType::TrackRemoved,  EAnimDataModelNotifyType::Populated };
+				const auto RecompressNotifies = { EAnimDataModelNotifyType::CurveAdded, EAnimDataModelNotifyType::CurveChanged, EAnimDataModelNotifyType::CurveRemoved, EAnimDataModelNotifyType::CurveFlagsChanged, EAnimDataModelNotifyType::CurveScaled };
+				
+				if (NotifyCollector.Contains(LengthChangingNotifies) || NotifyCollector.Contains(ResamplingNotifies))
+				{
+					ResampleAnimationTrackData();
+				}
+
+				if (NotifyCollector.Contains(LengthChangingNotifies) || NotifyCollector.Contains(ResamplingNotifies) || NotifyCollector.Contains(RecompressNotifies))
+				{
+					const bool bWasModelReset = NotifyCollector.Contains(EAnimDataModelNotifyType::Reset);
+					UpdateRawDataGuid(bWasModelReset ? GenerateNewGUID : RegenerateGUID);
+					ClearCompressedCurveData();
+					RecompressAnimationData();
+				}
+			}
+
+			break;
+		}
+
+		case EAnimDataModelNotifyType::BracketOpened:
+		{
+			break;
+		}
+
+		case EAnimDataModelNotifyType::Reset:
+		{
+			if (NotifyCollector.IsNotWithinBracket())
+			{
+				ResampleAnimationTrackData();
+				UpdateRawDataGuid(GenerateNewGUID);
+				RecompressAnimationData();
+			}
+
+			break;
+		}
+		
+		case EAnimDataModelNotifyType::TrackAdded:
+		case EAnimDataModelNotifyType::TrackChanged:
+		case EAnimDataModelNotifyType::TrackRemoved:		
+		{
+			HandleTrackDataChanged();				
+			break;
+		}
+
+		case EAnimDataModelNotifyType::CurveAdded:
+		case EAnimDataModelNotifyType::CurveChanged:
+		case EAnimDataModelNotifyType::CurveRemoved:
+		case EAnimDataModelNotifyType::CurveFlagsChanged:
+		case EAnimDataModelNotifyType::CurveScaled:
+		{
+			if (NotifyCollector.IsNotWithinBracket())
+			{
+				UpdateRawDataGuid(RegenerateGUID);
+				ClearCompressedCurveData();
+				RecompressAnimationData();
+			}
+
+			break;
+		}
+
+		case EAnimDataModelNotifyType::CurveColorChanged:
+		{
+			break;
+		}
+
+		case EAnimDataModelNotifyType::CurveRenamed:
+		{
+			const FCurveRenamedPayload& TypedPayload = Payload.GetPayload<FCurveRenamedPayload>();
+			UpdateCompressedCurveName(TypedPayload.Identifier.InternalName.UID, TypedPayload.NewIdentifier.InternalName);
+			break;
+		}
+		default:
+		{
+			checkf(false, TEXT("Missing case statement for animation model notify type"));
+			break;
+		}
+	}
+	PRAGMA_ENABLE_DEPRECATION_WARNINGS
+	
+	MarkPackageDirty();
+}
+
+
+void UAnimSequence::ResampleAnimationTrackData()
+{
+	ValidateModel();
+
+	if (TargetFrameRate == DataModel->GetFrameRate())
+	{
+		// Can directly copy
+		ResampledAnimationTrackData = DataModel->GetBoneAnimationTracks();
+		NumberOfSampledKeys = DataModel->GetNumberOfKeys();
+		NumberOfSampledFrames = DataModel->GetNumberOfFrames();
+	}
+	else
+	{
+		// We actually need to resample
+		ResampledAnimationTrackData.SetNum(DataModel->GetBoneAnimationTracks().Num());
+
+		NumberOfSampledFrames = TargetFrameRate.AsFrameNumber(DataModel->GetPlayLength()).Value;
+		NumberOfSampledKeys = NumberOfSampledFrames + 1;
+				
+		for (int32 TrackIndex = 0; TrackIndex < ResampledAnimationTrackData.Num(); ++TrackIndex)	
+		{
+			FBoneAnimationTrack& TrackData = ResampledAnimationTrackData[TrackIndex];
+
+			TrackData.InternalTrackData.PosKeys.SetNumUninitialized(NumberOfSampledKeys);
+			TrackData.InternalTrackData.RotKeys.SetNumUninitialized(NumberOfSampledKeys);
+			TrackData.InternalTrackData.ScaleKeys.SetNumUninitialized(NumberOfSampledKeys);
+
+			for (int32 FrameIndex = 0; FrameIndex < NumberOfSampledKeys; ++FrameIndex)
+			{
+				FTransform Transform;
+
+				const double SampleTime = TargetFrameRate.AsSeconds(FrameIndex);
+
+				UE::Anim::GetBoneTransformFromModel(DataModel, Transform, TrackIndex, SampleTime, EAnimInterpolationType::Linear);
+
+				TrackData.InternalTrackData.PosKeys[FrameIndex] = Transform.GetLocation();
+				TrackData.InternalTrackData.RotKeys[FrameIndex] = Transform.GetRotation();
+				TrackData.InternalTrackData.ScaleKeys[FrameIndex] = Transform.GetScale3D();
+			}
+		}
+	}
+}
+
+void UAnimSequence::RecompressAnimationData()
+{
+	PRAGMA_DISABLE_DEPRECATION_WARNINGS
+	ClearCompressedBoneData();
+	PRAGMA_ENABLE_DEPRECATION_WARNINGS
+	bUseRawDataOnly = true;
+	RequestAsyncAnimRecompression(false);
+}
+#endif // WITH_EDITOR
+
 /*-----------------------------------------------------------------------------
 	AnimNotify& subclasses
 -----------------------------------------------------------------------------*/
