Animation data MVC refactor

diff --git a/Engine/Source/Runtime/Engine/Private/Animation/AnimSequenceBase.cpp b/Engine/Source/Runtime/Engine/Private/Animation/AnimSequenceBase.cpp
--- a/Engine/Source/Runtime/Engine/Private/Animation/AnimSequenceBase.cpp
+++ b/Engine/Source/Runtime/Engine/Private/Animation/AnimSequenceBase.cpp
@@ -11,6 +11,11 @@
 #include "UObject/FrameworkObjectVersion.h"
 #include "UObject/FortniteMainBranchObjectVersion.h"
 #include "Animation/AnimationPoseData.h"
+#include "Animation/AnimSequenceHelpers.h"
+#include "Animation/AnimData/AnimDataModel.h"
+#include "Animation/AnimData/AnimDataController.h"
+
+#include "UObject/UE5MainStreamObjectVersion.h"
 
 DEFINE_LOG_CATEGORY(LogAnimMarkerSync);
 CSV_DECLARE_CATEGORY_MODULE_EXTERN(ENGINE_API, Animation);
@@ -21,7 +26,31 @@ CSV_DECLARE_CATEGORY_MODULE_EXTERN(ENGINE_API, Animation);
 UAnimSequenceBase::UAnimSequenceBase(const FObjectInitializer& ObjectInitializer)
 	: Super(ObjectInitializer)
 	, RateScale(1.0f)
+#if WITH_EDITORONLY_DATA
+	, DataModel(nullptr)
+	, Controller(nullptr)
+#endif // WITH_EDITORONLY_DATA
 {
+#if WITH_EDITOR
+	if (!HasAllFlags(EObjectFlags::RF_ClassDefaultObject))
+	{
+		CreateModel();
+		GetController();
+	}
+#endif // WITH_EDITOR
+}
+
+void UAnimSequenceBase::SetSequenceLength(float NewLength)
+{
+#if WITH_EDITOR
+	// Editor time we update the source data
+	GetController()->SetPlayLength(NewLength);
+#else 
+	// In case this is called during runtime, update the sequence length directly. Current use-case is runtime created Montages.
+	PRAGMA_DISABLE_DEPRECATION_WARNINGS
+	SequenceLength = NewLength;
+	PRAGMA_ENABLE_DEPRECATION_WARNINGS
+#endif
 }
 
 bool UAnimSequenceBase::IsPostLoadThreadSafe() const
@@ -31,6 +60,44 @@ bool UAnimSequenceBase::IsPostLoadThreadSafe() const
 
 void UAnimSequenceBase::PostLoad()
 {
+#if WITH_EDITORONLY_DATA
+	if (!HasAnyFlags(EObjectFlags::RF_ClassDefaultObject))
+	{
+		const bool bRequiresModelCreation = DataModel != nullptr;
+		const bool bRequiresModelPopulation = GetLinkerCustomVersion(FUE5MainStreamObjectVersion::GUID) < FUE5MainStreamObjectVersion::IntroducingAnimationDataModel;
+		checkf(bRequiresModelPopulation || DataModel != nullptr, TEXT("Invalid Animation Sequence base state, no data model found past upgrade object version"));
+
+		// Construct a new UAnimDataModel instance
+		if (bRequiresModelCreation)
+		{
+			CreateModel();
+		}
+		else
+		{
+			DataModel->GetModifiedEvent().AddUObject(this, &UAnimSequenceBase::OnModelModified);
+		}
+
+		if (USkeleton* MySkeleton = GetSkeleton())
+		{
+			if (FLinkerLoad* SkeletonLinker = MySkeleton->GetLinker())
+			{
+				SkeletonLinker->Preload(MySkeleton);
+			}
+			MySkeleton->ConditionalPostLoad();
+
+			PRAGMA_DISABLE_DEPRECATION_WARNINGS
+			VerifyCurveNames<FFloatCurve>(*MySkeleton, USkeleton::AnimCurveMappingName, RawCurveData.FloatCurves);
+			PRAGMA_ENABLE_DEPRECATION_WARNINGS
+		}
+
+		GetController();
+		if (bRequiresModelPopulation)
+		{
+			PopulateModel();
+		}
+	}
+#endif // WITH_EDITORONLY_DATA
+
 	Super::PostLoad();
 
 	// Convert Notifies to new data
@@ -49,42 +116,51 @@ void UAnimSequenceBase::PostLoad()
 		}
 	}
 
-#if WITH_EDITOR
-	InitializeNotifyTrack();
-#endif
 	RefreshCacheData();
 
 	if(USkeleton* MySkeleton = GetSkeleton())
 	{
+		const bool bDoNotTransactAction = false;
 #if WITH_EDITOR
 		if (GetLinkerCustomVersion(FFortniteMainBranchObjectVersion::GUID) < FFortniteMainBranchObjectVersion::FixUpNoneNameAnimationCurves)
 		{
-			for (int32 Index = 0; Index < RawCurveData.FloatCurves.Num(); ++Index)
+			Controller->OpenBracket(LOCTEXT("FFortniteMainBranchObjectVersion::FixUpNoneNameAnimationCurves_Bracket","FFortniteMainBranchObjectVersion::FixUpNoneNameAnimationCurves"), bDoNotTransactAction);
 			{
-				FFloatCurve& Curve = RawCurveData.FloatCurves[Index];
-				if (Curve.Name.DisplayName == NAME_None)
+				const TArray<FFloatCurve>& FloatCurves = DataModel->GetFloatCurves();
+				for (int32 Index = 0; Index < FloatCurves.Num(); ++Index)
 				{
-					// give unique name
-					Curve.Name.DisplayName = FName(*FString(GetName() + TEXT("_CurveNameFix_") + FString::FromInt(Index)));
-					UE_LOG(LogAnimation, Warning, TEXT("[AnimSequence %s] contains invalid curve name \'None\'. Renaming this to %s. Please fix this curve in the editor. "), *GetFullName(), *Curve.Name.DisplayName.ToString());
+					const FFloatCurve& Curve = FloatCurves[Index];
+					if (Curve.Name.DisplayName == NAME_None)
+					{
+						// give unique name
+						const FName UniqueName = FName(*FString(GetName() + TEXT("_CurveNameFix_") + FString::FromInt(Index)));
+						UE_LOG(LogAnimation, Warning, TEXT("[AnimSequence %s] contains invalid curve name \'None\'. Renaming this to %s. Please fix this curve in the editor. "), *GetFullName(), *Curve.Name.DisplayName.ToString());
+
+						FSmartName NewSmartName = Curve.Name;
+						NewSmartName.DisplayName = UniqueName;
+
+						Controller->RenameCurve(FAnimationCurveIdentifier(Curve.Name, ERawCurveTrackTypes::RCT_Float), FAnimationCurveIdentifier(NewSmartName, ERawCurveTrackTypes::RCT_Float), bDoNotTransactAction);
+					}
 				}
 			}
+			Controller->CloseBracket(bDoNotTransactAction);
 		}
 		else
 		{
-			for (int32 Index = 0; Index < RawCurveData.FloatCurves.Num(); ++Index)
+			const TArray<FFloatCurve>& FloatCurves = DataModel->GetFloatCurves();
+			for (int32 Index = 0; Index < FloatCurves.Num(); ++Index)
 			{
-				const FFloatCurve& Curve = RawCurveData.FloatCurves[Index];
+				const FFloatCurve& Curve = FloatCurves[Index];
 				ensureMsgf(Curve.Name.DisplayName != NAME_None, TEXT("[AnimSequencer %s] has invalid curve name."), *GetFullName());
 			}
 		}
-#endif // WITH_EDITOR
-
 
+		Controller->FindOrAddCurveNamesOnSkeleton(MySkeleton, ERawCurveTrackTypes::RCT_Float, bDoNotTransactAction);
+		Controller->FindOrAddCurveNamesOnSkeleton(MySkeleton, ERawCurveTrackTypes::RCT_Transform, bDoNotTransactAction);
+#else
+		PRAGMA_DISABLE_DEPRECATION_WARNINGS
 		VerifyCurveNames<FFloatCurve>(*MySkeleton, USkeleton::AnimCurveMappingName, RawCurveData.FloatCurves);
-
-#if WITH_EDITOR
-		VerifyCurveNames<FTransformCurve>(*MySkeleton, USkeleton::AnimTrackCurveMappingName, RawCurveData.TransformCurves);
+		PRAGMA_ENABLE_DEPRECATION_WARNINGS
 #endif
 
 		// this should continue to add if skeleton hasn't been saved either 
@@ -92,8 +168,16 @@ void UAnimSequenceBase::PostLoad()
 		if (GetLinkerCustomVersion(FFrameworkObjectVersion::GUID) < FFrameworkObjectVersion::MoveCurveTypesToSkeleton
 			|| MySkeleton->GetLinkerCustomVersion(FFrameworkObjectVersion::GUID) < FFrameworkObjectVersion::MoveCurveTypesToSkeleton)
 		{
+
+#if WITH_EDITOR
+			const TArray<FFloatCurve>& FloatCurves = DataModel->GetFloatCurves();
+#else
+			PRAGMA_DISABLE_DEPRECATION_WARNINGS
+			const TArray<FFloatCurve>& FloatCurves = RawCurveData.FloatCurves;
+			PRAGMA_ENABLE_DEPRECATION_WARNINGS
+#endif
 			// fix up curve flags to skeleton
-			for (const FFloatCurve& Curve : RawCurveData.FloatCurves)
+			for (const FFloatCurve& Curve : FloatCurves)
 			{
 				bool bMorphtargetSet = Curve.GetCurveTypeFlag(AACF_DriveMorphTarget_DEPRECATED);
 				bool bMaterialSet = Curve.GetCurveTypeFlag(AACF_DriveMaterial_DEPRECATED);
@@ -108,6 +192,19 @@ void UAnimSequenceBase::PostLoad()
 	}
 }
 
+void UAnimSequenceBase::PostDuplicate(EDuplicateMode::Type DuplicateMode)
+{
+	Super::PostDuplicate(DuplicateMode);
+
+#if WITH_EDITOR
+	if (ensure(DataModel))
+	{
+		DataModel->GetModifiedEvent().RemoveAll(this);
+		DataModel->GetModifiedEvent().AddUObject(this, &UAnimSequenceBase::OnModelModified);
+	}
+#endif // WITH_EDITOR
+}
+
 float UAnimSequenceBase::GetPlayLength() const
 {
 	PRAGMA_DISABLE_DEPRECATION_WARNINGS
@@ -151,28 +248,6 @@ bool UAnimSequenceBase::IsNotifyAvailable() const
 	return (Notifies.Num() != 0) && (GetPlayLength() > 0.f);
 }
 
-#if WITH_EDITOR
-void UAnimSequenceBase::RegisterOnAnimCurvesChanged(const FOnAnimCurvesChanged& Delegate)
-{
-	OnAnimCurvesChanged.Add(Delegate);	
-}
-
-void UAnimSequenceBase::UnregisterOnAnimCurvesChanged(void* Unregister)
-{
-	OnAnimCurvesChanged.RemoveAll(Unregister);
-}
-
-void UAnimSequenceBase::RegisterOnAnimTrackCurvesChanged(const FOnAnimTrackCurvesChanged& Delegate)
-{
-	OnAnimTrackCurvesChanged.Add(Delegate);
-}
-
-void UAnimSequenceBase::UnregisterOnAnimTrackCurvesChanged(void* Unregister)
-{
-	OnAnimTrackCurvesChanged.RemoveAll(Unregister);
-}
-#endif 
-
 void UAnimSequenceBase::GetAnimNotifies(const float& StartTime, const float& DeltaTime, const bool bAllowLooping, TArray<const FAnimNotifyEvent *> & OutActiveNotifies) const
 {
 	TArray<FAnimNotifyEventReference> NotifyRefs;
@@ -294,8 +369,9 @@ void UAnimSequenceBase::GetAnimNotifiesFromDeltaPositions(const float& PreviousP
 void UAnimSequenceBase::RemapTracksToNewSkeleton(USkeleton* NewSkeleton, bool bConvertSpaces)
 {
 	Super::RemapTracksToNewSkeleton(NewSkeleton, bConvertSpaces);
-	VerifyCurveNames<FFloatCurve>(*NewSkeleton, USkeleton::AnimCurveMappingName, RawCurveData.FloatCurves);
-	VerifyCurveNames<FTransformCurve>(*NewSkeleton, USkeleton::AnimTrackCurveMappingName, RawCurveData.TransformCurves);
+
+	Controller->FindOrAddCurveNamesOnSkeleton(NewSkeleton, ERawCurveTrackTypes::RCT_Float);
+	Controller->FindOrAddCurveNamesOnSkeleton(NewSkeleton, ERawCurveTrackTypes::RCT_Transform);
 }
 #endif
 
@@ -547,12 +623,6 @@ const FFrameRate& UAnimSequenceBase::GetSamplingFrameRate() const
 }
 
 #if WITH_EDITOR
-void UAnimSequenceBase::RefreshCurveData()
-{
-	OnAnimCurvesChanged.Broadcast();
-	OnAnimTrackCurvesChanged.Broadcast();
-}
-
 void UAnimSequenceBase::InitializeNotifyTrack()
 {
 	if ( AnimNotifyTracks.Num() == 0 ) 
@@ -635,9 +705,12 @@ void UAnimSequenceBase::GetAssetRegistryTags(TArray<FAssetRegistryTag>& OutTags)
 	// as possible.
 	FString CurveNameList;
 
-	for(const FFloatCurve& Curve : RawCurveData.FloatCurves)
+	if (DataModel)
 	{
-		CurveNameList += FString::Printf(TEXT("%s%s"), *Curve.Name.DisplayName.ToString(), *USkeleton::CurveTagDelimiter);
+		for(const FFloatCurve& Curve : DataModel->GetFloatCurves())
+		{
+			CurveNameList += FString::Printf(TEXT("%s%s"), *Curve.Name.DisplayName.ToString(), *USkeleton::CurveTagDelimiter);
+		}
 	}
 	OutTags.Add(FAssetRegistryTag(USkeleton::CurveNameTag, CurveNameList, FAssetRegistryTag::TT_Hidden));
 }
@@ -708,9 +781,36 @@ void UAnimSequenceBase::RefreshParentAssetData()
 		NotifyEvent.EndLink.Link(this, NotifyEvent.GetTime() + NotifyEvent.Duration, NotifyEvent.GetSlotIndex());
 	}
 
-	SetSequenceLength(ParentSeqBase->GetPlayLength());
 	RateScale = ParentSeqBase->RateScale;
-	RawCurveData = ParentSeqBase->RawCurveData;
+
+	Controller->OpenBracket(LOCTEXT("RefreshParentAssetData_Bracket", "Refreshing Parent Asset Data"));
+	{
+		const UAnimDataModel* ParentDataModel = ParentSeqBase->GetDataModel();
+
+		Controller->SetPlayLength(ParentDataModel->GetPlayLength());
+		
+		Controller->RemoveAllCurvesOfType(ERawCurveTrackTypes::RCT_Float);
+		for (const FFloatCurve& FloatCurve : ParentDataModel->GetFloatCurves())
+		{
+			const FAnimationCurveIdentifier CurveId(FloatCurve.Name, ERawCurveTrackTypes::RCT_Float);
+			Controller->AddCurve(CurveId, FloatCurve.GetCurveTypeFlags());
+			Controller->SetCurveKeys(CurveId, FloatCurve.FloatCurve.GetConstRefOfKeys());
+		}
+
+		Controller->RemoveAllCurvesOfType(ERawCurveTrackTypes::RCT_Transform);
+		for (const FTransformCurve& TransformCurve : ParentDataModel->GetTransformCurves())
+		{
+			const FAnimationCurveIdentifier CurveId(TransformCurve.Name, ERawCurveTrackTypes::RCT_Transform);
+			Controller->AddCurve(CurveId, TransformCurve.GetCurveTypeFlags());
+
+			TArray<FTransform> Values;
+			TArray<float> Times;
+			TransformCurve.GetKeys(Times, Values);
+
+			Controller->SetTransformCurveKeys(CurveId, Values, Times);
+		}
+	}
+	Controller->CloseBracket();
 
 #if WITH_EDITORONLY_DATA
 	// if you change Notifies array, this will need to be rebuilt
@@ -741,29 +841,55 @@ void UAnimSequenceBase::RefreshParentAssetData()
 void UAnimSequenceBase::EvaluateCurveData(FBlendedCurve& OutCurve, float CurrentTime, bool bForceUseRawData) const
 {
 	CSV_SCOPED_TIMING_STAT(Animation, EvaluateCurveData);
+
+	REFACTOR_V2_WARNING("V2 Curve evaluation goes through original curve data structure. This contains copied data from the model. Goal is to have compressed curve data within AnimSequenceBase so that model would be editor / raw data. And compressed data is used at runtime.");
+	PRAGMA_DISABLE_DEPRECATION_WARNINGS
 	RawCurveData.EvaluateCurveData(OutCurve, CurrentTime);
+	PRAGMA_ENABLE_DEPRECATION_WARNINGS
 }
 
 float UAnimSequenceBase::EvaluateCurveData(SmartName::UID_Type CurveUID, float CurrentTime, bool bForceUseRawData) const
 {
+	PRAGMA_DISABLE_DEPRECATION_WARNINGS
 	const FFloatCurve* Curve = (const FFloatCurve*)RawCurveData.GetCurveData(CurveUID, ERawCurveTrackTypes::RCT_Float);
+	PRAGMA_ENABLE_DEPRECATION_WARNINGS
+
+	REFACTOR_V2_WARNING("V2 Curve evaluation goes through original curve data structure. This contains copied data from the model. Goal is to have compressed curve data within AnimSequenceBase so that model would be editor / raw data. And compressed data is used at runtime.");
 	return Curve != nullptr ? Curve->Evaluate(CurrentTime) : 0.0f;
 }
 
+const FRawCurveTracks& UAnimSequenceBase::GetCurveData() const
+{
+	REFACTOR_V2_WARNING("V2 Curve evaluation goes through original curve data structure. This contains copied data from the model. Goal is to have compressed curve data within AnimSequenceBase so that model would be editor / raw data. And compressed data is used at runtime.");
+	PRAGMA_DISABLE_DEPRECATION_WARNINGS
+	return RawCurveData;
+	PRAGMA_ENABLE_DEPRECATION_WARNINGS
+}
+
 bool UAnimSequenceBase::HasCurveData(SmartName::UID_Type CurveUID, bool bForceUseRawData) const
 {
+#if WITH_EDITOR
+	ValidateModel();
+	return DataModel->FindFloatCurve(FAnimationCurveIdentifier(CurveUID, ERawCurveTrackTypes::RCT_Float)) != nullptr;
+#else
+	PRAGMA_DISABLE_DEPRECATION_WARNINGS
 	return RawCurveData.GetCurveData(CurveUID) != nullptr;
+	PRAGMA_ENABLE_DEPRECATION_WARNINGS
+#endif
 }
 
 void UAnimSequenceBase::Serialize(FArchive& Ar)
 {
 	Ar.UsingCustomVersion(FFrameworkObjectVersion::GUID);
 	Ar.UsingCustomVersion(FFortniteMainBranchObjectVersion::GUID);
+	Ar.UsingCustomVersion(FUE5MainStreamObjectVersion::GUID);
 
 	Super::Serialize(Ar);
 
 	// fix up version issue and so on
+	PRAGMA_DISABLE_DEPRECATION_WARNINGS
 	RawCurveData.PostSerialize(Ar);
+	PRAGMA_ENABLE_DEPRECATION_WARNINGS
 }
 
 void UAnimSequenceBase::GetAnimationPose(struct FCompactPose& OutPose, FBlendedCurve & OutCurve, const FAnimExtractContext & ExtractionContext) const
@@ -780,4 +906,291 @@ void UAnimSequenceBase::HandleAssetPlayerTickedInternal(FAnimAssetTickContext &C
 	GetAnimNotifies(PreviousTime, MoveDelta, Instance.bLooping, AnimNotifies);
 	NotifyQueue.AddAnimNotifies(Context.ShouldGenerateNotifies(),AnimNotifies, Instance.EffectiveBlendWeight);
 }
+
+#if WITH_EDITOR
+void UAnimSequenceBase::OnModelModified(const EAnimDataModelNotifyType& NotifyType, UAnimDataModel* Model, const FAnimDataModelNotifPayload& Payload)
+{
+	NotifyCollector.Handle(NotifyType);
+
+	PRAGMA_DISABLE_DEPRECATION_WARNINGS
+	InitializeNotifyTrack();
+	PRAGMA_ENABLE_DEPRECATION_WARNINGS
+
+	auto HandleLengthChange = [&](float NewLength, float OldLength, float T0, float T1)
+	{
+		const float StartRemoveTime = T0;
+		const float EndRemoveTime = T1;
+
+		// Total time value for frames that were removed
+		const float Duration = T1 - T0;
+
+		if (NewLength > OldLength)
+		{
+			const float InsertTime = T0;
+			for (FAnimNotifyEvent& Notify : Notifies)
+			{
+				float CurrentTime = Notify.GetTime();
+				float NewDuration = 0.f;
+
+				// if state, make sure to adjust end time
+				if (Notify.NotifyStateClass)
+				{
+					float NotifyDuration = Notify.GetDuration();
+					float NotifyEnd = CurrentTime + NotifyDuration;
+					if (NotifyEnd >= InsertTime)
+					{
+						NewDuration = NotifyDuration + Duration;
+					}
+					else
+					{
+						NewDuration = NotifyDuration;
+					}
+				}
+
+				// Shift out notify by the time alloted for the inserted keys
+				if (CurrentTime >= InsertTime)
+				{
+					CurrentTime += Duration;
+				}
+
+				// Clamps against the sequence length, ensuring that the notify does not end up outside of the playback bounds
+				const float ClampedCurrentTime = FMath::Clamp(CurrentTime, 0.f, NewLength);
+
+				Notify.LinkSequence(this, ClampedCurrentTime);
+
+				Notify.SetDuration(NewDuration);
+
+				if (ClampedCurrentTime == 0.f)
+				{
+					Notify.TriggerTimeOffset = GetTriggerTimeOffsetForType(EAnimEventTriggerOffsets::OffsetAfter);
+				}
+				else if (ClampedCurrentTime == NewLength)
+				{
+					Notify.TriggerTimeOffset = GetTriggerTimeOffsetForType(EAnimEventTriggerOffsets::OffsetBefore);
+				}
+			}
+		}
+		else if (NewLength < OldLength)
+		{
+			// re-locate notifies
+			for (FAnimNotifyEvent& Notify : Notifies)
+			{
+				float CurrentTime = Notify.GetTime();
+				float NewDuration = 0.f;
+				
+				// if state, make sure to adjust end time
+				if (Notify.NotifyStateClass)
+				{
+					float NotifyDuration = Notify.GetDuration();
+					float NotifyEnd = CurrentTime + NotifyDuration;
+					NewDuration = NotifyDuration;
+
+					// If Notify is inside of the trimmed time frame, zero out the duration
+					if (NotifyEnd >= StartRemoveTime && NotifyEnd <= EndRemoveTime)
+					{
+						// small number @todo see if there is define for this
+						NewDuration = DataModel->GetFrameRate().AsInterval();
+					}
+					// If Notify overlaps trimmed time frame, remove trimmed duration
+					else if (CurrentTime < EndRemoveTime && NotifyEnd > EndRemoveTime)
+					{
+						NewDuration = NotifyEnd - Duration - CurrentTime;
+					}
+
+					NewDuration = FMath::Max(NewDuration, (float)DataModel->GetFrameRate().AsInterval());
+				}
+
+				if (CurrentTime >= StartRemoveTime && CurrentTime <= EndRemoveTime)
+				{
+					CurrentTime = StartRemoveTime;
+				}
+				else if (CurrentTime > EndRemoveTime)
+				{
+					CurrentTime -= Duration;
+				}
+
+				const float ClampedCurrentTime = FMath::Clamp(CurrentTime, 0.f, NewLength);
+
+				Notify.LinkSequence(this, ClampedCurrentTime);
+				Notify.SetDuration(NewDuration);
+
+				if (ClampedCurrentTime == 0.f)
+				{
+					Notify.TriggerTimeOffset = GetTriggerTimeOffsetForType(EAnimEventTriggerOffsets::OffsetAfter);
+				}
+				else if (ClampedCurrentTime == NewLength)
+				{
+					Notify.TriggerTimeOffset = GetTriggerTimeOffsetForType(EAnimEventTriggerOffsets::OffsetBefore);
+				}
+			}
+		}	
+	};
+
+	// Source animation curve data can be coarse, so copy the data down and remove any redundant keys
+	auto CopyCurvesFromModel = [this]()
+	{
+		PRAGMA_DISABLE_DEPRECATION_WARNINGS
+		RawCurveData.FloatCurves = DataModel->GetFloatCurves();
+		RawCurveData.TransformCurves = DataModel->GetTransformCurves();
+		RawCurveData.RemoveRedundantKeys();
+		PRAGMA_ENABLE_DEPRECATION_WARNINGS
+	};
+
+	switch (NotifyType)
+	{
+		case EAnimDataModelNotifyType::SequenceLengthChanged:
+		{
+			const FSequenceLengthChangedPayload& TypedPayload = Payload.GetPayload<FSequenceLengthChangedPayload>();
+			const float PreviousSequenceLength = TypedPayload.PreviousLength;
+			const float CurrentSequenceLength = Model->GetPlayLength();
+
+			PRAGMA_DISABLE_DEPRECATION_WARNINGS
+			SequenceLength = CurrentSequenceLength;
+			PRAGMA_ENABLE_DEPRECATION_WARNINGS
+
+			HandleLengthChange(CurrentSequenceLength, PreviousSequenceLength, TypedPayload.T0, TypedPayload.T1);
+
+			// Ensure that we only clamp the notifies at the end of a bracket
+			if (NotifyCollector.IsNotWithinBracket())
+			{
+				ClampNotifiesAtEndOfSequence();
+				CopyCurvesFromModel();
+			}
+			
+			break;
+		}
+
+		case EAnimDataModelNotifyType::CurveAdded:
+		case EAnimDataModelNotifyType::CurveChanged:
+		case EAnimDataModelNotifyType::CurveRemoved:
+		case EAnimDataModelNotifyType::CurveFlagsChanged:
+		case EAnimDataModelNotifyType::CurveRenamed:
+		case EAnimDataModelNotifyType::CurveScaled:
+		{
+			if (NotifyCollector.IsNotWithinBracket())
+			{
+				CopyCurvesFromModel();
+			}
+			break;
+		}
+
+		case EAnimDataModelNotifyType::Populated:
+		{
+			const float CurrentSequenceLength = Model->GetPlayLength();
+			PRAGMA_DISABLE_DEPRECATION_WARNINGS
+			SequenceLength = CurrentSequenceLength;
+			PRAGMA_ENABLE_DEPRECATION_WARNINGS
+
+			if (NotifyCollector.IsNotWithinBracket())
+			{
+				ClampNotifiesAtEndOfSequence();
+				CopyCurvesFromModel();
+			}
+
+			break;
+		}
+
+		case EAnimDataModelNotifyType::BracketClosed:
+		{
+			if (NotifyCollector.IsNotWithinBracket())
+			{
+				const auto CurveCopyNotifies = { EAnimDataModelNotifyType::CurveAdded, EAnimDataModelNotifyType::CurveChanged, EAnimDataModelNotifyType::CurveRemoved, EAnimDataModelNotifyType::CurveFlagsChanged, EAnimDataModelNotifyType::CurveRenamed, EAnimDataModelNotifyType::CurveScaled, EAnimDataModelNotifyType::Populated, EAnimDataModelNotifyType::Reset };
+
+				const auto LengthChangingNotifies = { EAnimDataModelNotifyType::SequenceLengthChanged, EAnimDataModelNotifyType::FrameRateChanged, EAnimDataModelNotifyType::Reset, EAnimDataModelNotifyType::Populated };
+
+				if (NotifyCollector.Contains(CurveCopyNotifies) || NotifyCollector.Contains(LengthChangingNotifies))
+				{
+					CopyCurvesFromModel();
+				}
+				
+				if (NotifyCollector.Contains(LengthChangingNotifies))
+				{
+					CopyCurvesFromModel();
+					ClampNotifiesAtEndOfSequence();
+				}
+			}
+			break;
+		}
+	}
+}
+
+UAnimDataModel* UAnimSequenceBase::GetDataModel() const
+{
+	ValidateModel();
+	return DataModel;
+}
+
+UAnimDataController* UAnimSequenceBase::GetController()
+{
+	ValidateModel();
+
+	if (HasAllFlags(EObjectFlags::RF_ClassDefaultObject))
+	{
+		return nullptr;
+	}
+
+	if (Controller == nullptr)
+	{
+		Controller = NewObject<UAnimDataController>(GetTransientPackage());		
+	}
+
+	if (Controller->GetModel() != DataModel)
+	{
+		Controller->SetModel(DataModel);
+	}
+	return Controller;
+}
+
+void UAnimSequenceBase::PopulateModel()
+{
+	check(!HasAnyFlags(EObjectFlags::RF_ClassDefaultObject));	
+	
+	// Make a copy of the current data, to mitigate any changes due to notify callbacks
+	const FFrameRate FrameRate = GetSamplingFrameRate();
+	const float PlayLength = GetPlayLength();
+	PRAGMA_DISABLE_DEPRECATION_WARNINGS
+	const FRawCurveTracks CurveData = RawCurveData;
+	PRAGMA_ENABLE_DEPRECATION_WARNINGS
+		
+	UAnimDataController::FScopedBracket ScopedBracket(Controller, LOCTEXT("UAnimSequenceBase::PopulateModel_Bracket", "Generating Animation Model Data for Animation Sequence Base"));
+	Controller->SetPlayLength(PlayLength);
+	Controller->SetFrameRate(FrameRate);
+
+	USkeleton* TargetSkeleton = GetSkeleton();	
+	UE::Anim::CopyCurveDataToModel(CurveData, TargetSkeleton, Controller);	
+}
+
+void UAnimSequenceBase::ValidateModel() const
+{
+	checkf(DataModel != nullptr, TEXT("Invalidate data model"));
+}
+
+void UAnimSequenceBase::CopyDataModel(const UAnimDataModel* ModelToDuplicate)
+{
+	checkf(ModelToDuplicate != nullptr, TEXT("Invalidate data model"));
+	if (DataModel)
+	{
+		DataModel->GetModifiedEvent().RemoveAll(this);
+	}
+
+	DataModel = DuplicateObject(ModelToDuplicate, this);
+	GetController()->SetModel(DataModel);
+	DataModel->GetModifiedEvent().AddUObject(this, &UAnimSequenceBase::OnModelModified);
+}
+
+void UAnimSequenceBase::CreateModel()
+{	
+	if (DataModel == nullptr)
+	{
+		DataModel = CreateDefaultSubobject<UAnimDataModel>(FName(TEXT("AnimationDataModel")));
+	}
+
+	GetController()->SetModel(DataModel);
+	
+	DataModel->GetModifiedEvent().RemoveAll(this);
+	DataModel->GetModifiedEvent().AddUObject(this, &UAnimSequenceBase::OnModelModified);
+}
+#endif // WITH_EDITOR
+
 #undef LOCTEXT_NAMESPACE 
+
diff --git a/Engine/Source/Runtime/Engine/Private/Animation/AnimSequenceHelpers.cpp b/Engine/Source/Runtime/Engine/Private/Animation/AnimSequenceHelpers.cpp
new file mode 100644
--- /dev/null
+++ b/Engine/Source/Runtime/Engine/Private/Animation/AnimSequenceHelpers.cpp
@@ -0,0 +1,960 @@
+// Copyright Epic Games, Inc. All Rights Reserved.
+
+#include "Animation/AnimSequenceHelpers.h"
+#include "Animation/AnimData/AnimDataModel.h"
+
+#include "UObject/NameTypes.h"
+
+#include "BonePose.h"
+#include "BoneContainer.h"
+#include "Animation/AnimTypes.h"
+
+#include "AnimationRuntime.h"
+#include "AnimEncoding.h"
+#include "AnimationUtils.h"
+#include "Animation/AnimTypes.h"
+#include "Animation/AnimNotifies/AnimNotify.h"
+#include "Animation/AnimNotifies/AnimNotifyState.h"
+#include "Misc/MessageDialog.h"
+
+#define LOCTEXT_NAMESPACE "AnimSequenceHelpers"
+
+namespace UE
+{
+
+namespace Anim
+{
+
+struct FRetargetTracking
+{
+	const FCompactPoseBoneIndex PoseBoneIndex;
+	const int32 SkeletonBoneIndex;
+
+	FRetargetTracking(const FCompactPoseBoneIndex InPoseBoneIndex, const int32 InSkeletonBoneIndex)
+		: PoseBoneIndex(InPoseBoneIndex), SkeletonBoneIndex(InSkeletonBoneIndex)
+	{
+	}
+};
+
+struct FBuildRawPoseScratchArea : public TThreadSingleton<FBuildRawPoseScratchArea>
+{
+	TArray<FRetargetTracking> RetargetTracking;
+	TArray<FVirtualBoneCompactPoseData> VirtualBoneCompactPoseData;
+};
+
+
+#if WITH_EDITOR
+FTransform ExtractTransformForKey(int32 Key, const FRawAnimSequenceTrack& TrackToExtract)
+{
+	static const FVector DefaultScale3D = FVector(1.f);
+	const bool bHasScaleKey = TrackToExtract.ScaleKeys.Num() > 0;
+
+	const int32 PosKeyIndex = FMath::Min(Key, TrackToExtract.PosKeys.Num() - 1);
+	const int32 RotKeyIndex = FMath::Min(Key, TrackToExtract.RotKeys.Num() - 1);
+	if (bHasScaleKey)
+	{
+		const int32 ScaleKeyIndex = FMath::Min(Key, TrackToExtract.ScaleKeys.Num() - 1);
+		return FTransform(TrackToExtract.RotKeys[RotKeyIndex], TrackToExtract.PosKeys[PosKeyIndex], TrackToExtract.ScaleKeys[ScaleKeyIndex]);
+	}
+	else
+	{
+		return FTransform(TrackToExtract.RotKeys[RotKeyIndex], TrackToExtract.PosKeys[PosKeyIndex], DefaultScale3D);
+	}
+}
+
+template<bool bInterpolateT>
+void ExtractPoseFromModel(const TArray<FBoneAnimationTrack>& BoneAnimationTracks, const TMap<FName, FTransform>& OverrideBoneTransforms, FCompactPose& InOutPose, const int32 KeyIndex1, const int32 KeyIndex2, float Alpha)
+{
+	const int32 NumAnimationTracks = BoneAnimationTracks.Num();
+	const FBoneContainer& RequiredBones = InOutPose.GetBoneContainer();
+
+	TArray<UE::Anim::FRetargetTracking>& RetargetTracking = UE::Anim::FBuildRawPoseScratchArea::Get().RetargetTracking;
+	RetargetTracking.Reset(NumAnimationTracks);
+
+	TArray<FVirtualBoneCompactPoseData>& VBCompactPoseData = UE::Anim::FBuildRawPoseScratchArea::Get().VirtualBoneCompactPoseData;
+	VBCompactPoseData = RequiredBones.GetVirtualBoneCompactPoseData();
+
+	FCompactPose Key2Pose;
+	if (bInterpolateT)
+	{
+		Key2Pose.CopyBonesFrom(InOutPose);
+	}
+
+	for (const FBoneAnimationTrack& AnimationTrack : BoneAnimationTracks)
+	{
+		const int32 SkeletonBoneIndex = AnimationTrack.BoneTreeIndex;
+		// not sure it's safe to assume that SkeletonBoneIndex can never be INDEX_NONE
+		if ((SkeletonBoneIndex != INDEX_NONE) && (SkeletonBoneIndex < MAX_BONES))
+		{
+			const FCompactPoseBoneIndex PoseBoneIndex = RequiredBones.GetCompactPoseIndexFromSkeletonIndex(SkeletonBoneIndex);
+			if (PoseBoneIndex != INDEX_NONE)
+			{
+				for (int32 Idx = 0; Idx < VBCompactPoseData.Num(); ++Idx)
+				{
+					FVirtualBoneCompactPoseData& VB = VBCompactPoseData[Idx];
+					if (PoseBoneIndex == VB.VBIndex)
+					{
+						// Remove this bone as we have written data for it (false so we dont resize allocation)
+						VBCompactPoseData.RemoveAtSwap(Idx, 1, false);
+						break; //Modified TArray so must break here
+					}
+				}
+				
+				// extract animation
+				const FRawAnimSequenceTrack& TrackToExtract = AnimationTrack.InternalTrackData;
+				const FTransform* OverrideTransform = OverrideBoneTransforms.Find(AnimationTrack.Name);
+				{
+					// Bail out (with rather wacky data) if data is empty for some reason.
+					if (TrackToExtract.PosKeys.Num() == 0 || TrackToExtract.RotKeys.Num() == 0)
+					{
+						InOutPose[PoseBoneIndex].SetIdentity();
+
+						if (bInterpolateT)
+						{
+							Key2Pose[PoseBoneIndex].SetIdentity();
+						}
+					}
+					else
+					{
+						InOutPose[PoseBoneIndex] = UE::Anim::ExtractTransformForKey(KeyIndex1, TrackToExtract);
+
+						if (bInterpolateT)
+						{
+							Key2Pose[PoseBoneIndex] = UE::Anim::ExtractTransformForKey(KeyIndex2, TrackToExtract);
+						}
+					}
+				}
+
+				if (OverrideTransform)
+				{
+					InOutPose[PoseBoneIndex].SetRotation(InOutPose[PoseBoneIndex].GetRotation() * OverrideTransform->GetRotation());
+					InOutPose[PoseBoneIndex].TransformPosition(OverrideTransform->GetTranslation());
+					InOutPose[PoseBoneIndex].SetScale3D(InOutPose[PoseBoneIndex].GetScale3D() * OverrideTransform->GetScale3D());
+
+					if (bInterpolateT)
+					{
+						Key2Pose[PoseBoneIndex].SetRotation(Key2Pose[PoseBoneIndex].GetRotation() * OverrideTransform->GetRotation());
+						Key2Pose[PoseBoneIndex].TransformPosition(OverrideTransform->GetTranslation());
+						Key2Pose[PoseBoneIndex].SetScale3D(Key2Pose[PoseBoneIndex].GetScale3D() * OverrideTransform->GetScale3D());
+					}
+				}
+
+				RetargetTracking.Add(UE::Anim::FRetargetTracking(PoseBoneIndex, SkeletonBoneIndex));
+			}
+		}
+	}
+
+	//Build Virtual Bones
+	if (VBCompactPoseData.Num() > 0)
+	{
+		FCSPose<FCompactPose> CSPose1;
+		CSPose1.InitPose(InOutPose);
+
+		FCSPose<FCompactPose> CSPose2;
+		if (bInterpolateT)
+		{
+			CSPose2.InitPose(Key2Pose);
+		}
+
+		for (FVirtualBoneCompactPoseData& VB : VBCompactPoseData)
+		{
+			FTransform Source = CSPose1.GetComponentSpaceTransform(VB.SourceIndex);
+			FTransform Target = CSPose1.GetComponentSpaceTransform(VB.TargetIndex);
+			InOutPose[VB.VBIndex] = Target.GetRelativeTransform(Source);
+
+			if (bInterpolateT)
+			{
+				FTransform Source2 = CSPose2.GetComponentSpaceTransform(VB.SourceIndex);
+				FTransform Target2 = CSPose2.GetComponentSpaceTransform(VB.TargetIndex);
+				Key2Pose[VB.VBIndex] = Target2.GetRelativeTransform(Source2);
+			}
+		}
+	}
+
+	if (bInterpolateT)
+	{
+		for (FCompactPoseBoneIndex BoneIndex : InOutPose.ForEachBoneIndex())
+		{
+			InOutPose[BoneIndex].Blend(InOutPose[BoneIndex], Key2Pose[BoneIndex], Alpha);
+		}
+	}
+}
+
+void BuildPoseFromModel(const UAnimDataModel* Model, FCompactPose& OutPose, const float Time, const EAnimInterpolationType& InterpolationType, const FName& RetargetSource, const TArray<FTransform>& RetargetTransforms)
+{
+	check(Model);
+
+	const int32 NumberOfKeys = Model->GetNumberOfKeys();
+	const float SequenceLength = Model->GetPlayLength();
+
+	// Generate keys to interpolate between
+	int32 KeyIndex1, KeyIndex2;
+	float Alpha;
+	FAnimationRuntime::GetKeyIndicesFromTime(KeyIndex1, KeyIndex2, Alpha, Time, NumberOfKeys, SequenceLength);
+
+	if (InterpolationType == EAnimInterpolationType::Step)
+	{
+		// Force stepping between keys
+		Alpha = 0.f;
+	}
+
+	bool bShouldInterpolate = true;
+
+	if (Alpha < KINDA_SMALL_NUMBER)
+	{
+		Alpha = 0.f;
+		bShouldInterpolate = false;
+	}
+	else if (Alpha > 1.f - KINDA_SMALL_NUMBER)
+	{
+		bShouldInterpolate = false;
+		KeyIndex1 = KeyIndex2;
+	}
+
+	TMap<FName, FTransform> ActiveCurves;
+
+	if (!OutPose.GetBoneContainer().ShouldUseSourceData())
+	{
+		EvaluateTransformCurvesFromModel(Model, ActiveCurves, Time, 1.f);
+	}
+
+	if (bShouldInterpolate)
+	{
+		ExtractPoseFromModel<true>(Model->GetBoneAnimationTracks(), ActiveCurves, OutPose, KeyIndex1, KeyIndex2, Alpha);
+	}
+	else
+	{
+		ExtractPoseFromModel<false>(Model->GetBoneAnimationTracks(), ActiveCurves, OutPose, KeyIndex1, KeyIndex2, Alpha);
+	}
+
+	Retargeting::RetargetPose(OutPose, RetargetSource, RetargetTransforms);
+}
+
+void EvaluateFloatCurvesFromModel(const UAnimDataModel* Model, FBlendedCurve& OutCurves, float Time)
+{
+	check(Model);
+
+	//SCOPE_CYCLE_COUNTER(STAT_EvalRawCurveData);
+	if (OutCurves.NumValidCurveCount > 0)
+	{
+		// evaluate the curve data at the Time and add to Instance
+		for (const FFloatCurve& Curve : Model->GetFloatCurves())
+		{
+			if (OutCurves.IsEnabled(Curve.Name.UID))
+			{
+				float Value = Curve.Evaluate(Time);
+				OutCurves.Set(Curve.Name.UID, Value);
+			}
+		}
+	}
+}
+
+void EvaluateTransformCurvesFromModel(const UAnimDataModel* Model, TMap<FName, FTransform>& OutCurves, float Time, float BlendWeight)
+{
+	check(Model);
+
+	for (const FTransformCurve& Curve : Model->GetTransformCurves())
+	{
+		// if disabled, do not handle
+		if (Curve.GetCurveTypeFlag(AACF_Disabled))
+		{
+			continue;
+		}
+
+		// Add or retrieve curve
+		const FName& CurveName = Curve.Name.DisplayName;
+
+		// note we're not checking Curve.GetCurveTypeFlags() yet
+		FTransform& Value = OutCurves.FindOrAdd(CurveName);
+		Value = Curve.Evaluate(Time, BlendWeight);
+	}
+}
+
+void GetBoneTransformFromModel(const UAnimDataModel* Model, FTransform& OutTransform, int32 TrackIndex, float Time, const EAnimInterpolationType& Interpolation)
+{
+	const FBoneAnimationTrack& TrackData = Model->GetBoneTrackByIndex(TrackIndex);
+	const FRawAnimSequenceTrack& RawTrack = TrackData.InternalTrackData;
+	
+	FAnimationUtils::ExtractTransformFromTrack(Time, Model->GetNumberOfKeys(), Model->GetPlayLength(), RawTrack, Interpolation, OutTransform);
+}
+
+void GetBoneTransformFromModel(const UAnimDataModel* Model, FTransform& OutTransform, int32 TrackIndex, int32 KeyIndex)
+{
+	const FBoneAnimationTrack& TrackData = Model->GetBoneTrackByIndex(TrackIndex);
+	const FRawAnimSequenceTrack& RawTrack = TrackData.InternalTrackData;
+
+	// Bail out (with rather wacky data) if data is empty for some reason.
+	if (RawTrack.PosKeys.Num() == 0 || RawTrack.RotKeys.Num() == 0)
+	{
+		UE_LOG(LogAnimation, Log, TEXT("UAnimSequence::GetBoneTransform : No anim data in AnimSequence!"));
+		OutTransform.SetIdentity();
+		return;
+	}
+
+	const int32 PosKeyIndex = FMath::Min(KeyIndex, RawTrack.PosKeys.Num() - 1);
+	const int32 RotKeyIndex = FMath::Min(KeyIndex, RawTrack.RotKeys.Num() - 1);
+	static const FVector DefaultScale3D = FVector(1.f);
+
+	OutTransform.SetTranslation(RawTrack.PosKeys[PosKeyIndex]);
+	OutTransform.SetRotation(RawTrack.RotKeys[RotKeyIndex]);
+	if (RawTrack.ScaleKeys.Num() > 0)
+	{
+		const int32 ScaleKeyIndex = FMath::Min(KeyIndex, RawTrack.ScaleKeys.Num() - 1);
+		OutTransform.SetScale3D(RawTrack.ScaleKeys[ScaleKeyIndex]);
+	}
+	else
+	{
+		OutTransform.SetScale3D(DefaultScale3D);
+	}
+}
+
+void CopyCurveDataToModel(const FRawCurveTracks& CurveData, const USkeleton* Skeleton, UAnimDataController* Controller)
+{
+	// Populate float curve data
+	for (const FFloatCurve& FloatCurve : CurveData.FloatCurves)
+	{
+		const FAnimationCurveIdentifier CurveId = UAnimationCurveIdentifierExtensions::FindCurveIdentifier(Skeleton, FloatCurve.Name.DisplayName, ERawCurveTrackTypes::RCT_Float);
+		if (CurveId.IsValid())
+		{
+			Controller->AddCurve(CurveId, FloatCurve.GetCurveTypeFlags());
+			Controller->SetCurveKeys(CurveId, FloatCurve.FloatCurve.GetConstRefOfKeys());
+		}
+	}
+
+	// Populate transform curve data
+	for (const FTransformCurve& TransformCurve : CurveData.TransformCurves)
+	{
+		const FAnimationCurveIdentifier CurveId = UAnimationCurveIdentifierExtensions::FindCurveIdentifier(Skeleton, TransformCurve.Name.DisplayName, ERawCurveTrackTypes::RCT_Transform);
+		if (CurveId.IsValid())
+		{
+			Controller->AddCurve(CurveId, TransformCurve.GetCurveTypeFlags());
+
+			// Set each individual channel rich curve keys, to account for any custom tangents etc.
+			for (int32 SubCurveIndex = 0; SubCurveIndex < 3; ++SubCurveIndex)
+			{
+				const ETransformCurveChannel Channel = (ETransformCurveChannel)SubCurveIndex;
+				const FVectorCurve* VectorCurve = TransformCurve.GetVectorCurveByIndex(SubCurveIndex);
+				for (int32 ChannelIndex = 0; ChannelIndex < 3; ++ChannelIndex)
+				{
+					const EVectorCurveChannel Axis = (EVectorCurveChannel)ChannelIndex;
+					FAnimationCurveIdentifier TargetCurveIdentifier = CurveId;
+					UAnimationCurveIdentifierExtensions::GetTransformChildCurveIdentifier(TargetCurveIdentifier, Channel, Axis);
+					Controller->SetCurveKeys(TargetCurveIdentifier, VectorCurve->FloatCurves[ChannelIndex].GetConstRefOfKeys());
+				}
+			}
+		}
+	}
+}
+#endif // WITH_EDITOR
+ 
+ 
+ void ExtractBoneTransform(const FRawAnimSequenceTrack& RawTrack, FTransform& OutTransform, int32 KeyIndex)
+{
+	// Bail out (with rather wacky data) if data is empty for some reason.
+	if (RawTrack.PosKeys.Num() == 0 || RawTrack.RotKeys.Num() == 0)
+	{
+		UE_LOG(LogAnimation, Log, TEXT("GetBoneTransform : No anim data in track!"));
+		OutTransform.SetIdentity();
+		return;
+	}
+
+	const int32 PosKeyIndex = FMath::Min(KeyIndex, RawTrack.PosKeys.Num() - 1);
+	const int32 RotKeyIndex = FMath::Min(KeyIndex, RawTrack.RotKeys.Num() - 1);
+	static const FVector DefaultScale3D = FVector(1.f);
+
+	OutTransform.SetTranslation(RawTrack.PosKeys[PosKeyIndex]);
+	OutTransform.SetRotation(RawTrack.RotKeys[RotKeyIndex]);
+	if (RawTrack.ScaleKeys.Num() > 0)
+	{
+		const int32 ScaleKeyIndex = FMath::Min(KeyIndex, RawTrack.ScaleKeys.Num() - 1);
+		OutTransform.SetScale3D(RawTrack.ScaleKeys[ScaleKeyIndex]);
+	}
+	else
+	{
+		OutTransform.SetScale3D(DefaultScale3D);
+	}
+}
+
+void Retargeting::RetargetPose(FCompactPose& InOutPose, const FName& RetargetSource, const TArray<FTransform>& RetargetTransforms)
+{
+	const FBoneContainer& RequiredBones = InOutPose.GetBoneContainer();
+	const bool bDisableRetargeting = RequiredBones.GetDisableRetargeting();
+
+	if (!bDisableRetargeting)
+	{
+		const TArray<UE::Anim::FRetargetTracking>& RetargetTracking = UE::Anim::FBuildRawPoseScratchArea::Get().RetargetTracking;
+
+		USkeleton* Skeleton = RequiredBones.GetSkeletonAsset();
+
+		for (const UE::Anim::FRetargetTracking& RT : RetargetTracking)
+		{
+			FAnimationRuntime::RetargetBoneTransform(Skeleton, RetargetSource, RetargetTransforms, InOutPose[RT.PoseBoneIndex], RT.SkeletonBoneIndex, RT.PoseBoneIndex, RequiredBones, false);
+		}
+	}
+}
+
+
+#if WITH_EDITOR
+bool CopyNotifies(const UAnimSequence* SourceAnimSeq, UAnimSequence* DestAnimSeq, bool bShowDialogs)
+{
+	// Abort if source == destination.
+	if (SourceAnimSeq == DestAnimSeq)
+	{
+		return true;
+	}
+
+	// If the destination sequence is shorter than the source sequence, we'll be dropping notifies that
+	// occur at later times than the dest sequence is long.  Give the user a chance to abort if we
+	// find any notifies that won't be copied over.
+	if (bShowDialogs && DestAnimSeq->GetPlayLength() < SourceAnimSeq->GetPlayLength())
+	{
+		for (int32 NotifyIndex = 0; NotifyIndex < SourceAnimSeq->Notifies.Num(); ++NotifyIndex)
+		{
+			// If a notify is found which occurs off the end of the destination sequence, prompt the user to continue.
+			const FAnimNotifyEvent& SrcNotifyEvent = SourceAnimSeq->Notifies[NotifyIndex];
+			if (SrcNotifyEvent.GetTriggerTime() > DestAnimSeq->GetPlayLength())
+			{
+				const bool bProceed = EAppReturnType::Yes == FMessageDialog::Open(EAppMsgType::YesNo, NSLOCTEXT("UnrealEd", "SomeNotifiesWillNotBeCopiedQ", "Some notifies will not be copied because the destination sequence is not long enough.  Proceed?"));
+				if (!bProceed)
+				{
+					return false;
+				}
+				else
+				{
+					break;
+				}
+			}
+		}
+	}
+
+	// If the destination sequence contains any notifies, ask the user if they'd like
+	// to delete the existing notifies before copying over from the source sequence.
+	if (bShowDialogs && DestAnimSeq->Notifies.Num() > 0)
+	{
+		const bool bDeleteExistingNotifies = EAppReturnType::Yes == FMessageDialog::Open(EAppMsgType::YesNo, FText::Format(
+			NSLOCTEXT("UnrealEd", "DestSeqAlreadyContainsNotifiesMergeQ", "The destination sequence already contains {0} notifies.  Delete these before copying?"), FText::AsNumber(DestAnimSeq->Notifies.Num())));
+		if (bDeleteExistingNotifies)
+		{
+			DestAnimSeq->Notifies.Empty();
+			DestAnimSeq->MarkPackageDirty();
+		}
+	}
+
+	// Do the copy.
+	int32 NumNotifiesThatWereNotCopied = 0;
+
+	for (int32 NotifyIndex = 0; NotifyIndex < SourceAnimSeq->Notifies.Num(); ++NotifyIndex)
+	{
+		const FAnimNotifyEvent& SrcNotifyEvent = SourceAnimSeq->Notifies[NotifyIndex];
+
+		// Skip notifies which occur at times later than the destination sequence is long.
+		if (SrcNotifyEvent.GetTriggerTime() > DestAnimSeq->GetPlayLength())
+		{
+			++NumNotifiesThatWereNotCopied;
+			continue;
+		}
+
+		// Copy notify tracks from src to dest if they are missing
+		if (SrcNotifyEvent.TrackIndex >= DestAnimSeq->AnimNotifyTracks.Num())
+		{
+			for (int32 TrackIndex = DestAnimSeq->AnimNotifyTracks.Num(); TrackIndex <= SrcNotifyEvent.TrackIndex; ++TrackIndex)
+			{
+				DestAnimSeq->AnimNotifyTracks.Add(FAnimNotifyTrack(SourceAnimSeq->AnimNotifyTracks[TrackIndex].TrackName, SourceAnimSeq->AnimNotifyTracks[TrackIndex].TrackColor));
+			}
+		}
+
+		// Track the location of the new notify.
+		int32 NewNotifyIndex = DestAnimSeq->Notifies.AddDefaulted();
+		FAnimNotifyEvent& NotifyEvent = DestAnimSeq->Notifies[NewNotifyIndex];
+
+		// Copy properties of the NotifyEvent
+		NotifyEvent.TrackIndex = SrcNotifyEvent.TrackIndex;
+		NotifyEvent.NotifyName = SrcNotifyEvent.NotifyName;
+		NotifyEvent.Duration = SrcNotifyEvent.Duration;
+
+		// Copy the notify itself, and point the new one at it.
+		if (SrcNotifyEvent.Notify)
+		{
+			DestAnimSeq->Notifies[NewNotifyIndex].Notify = static_cast<UAnimNotify*>(StaticDuplicateObject(SrcNotifyEvent.Notify, DestAnimSeq, NAME_None, RF_AllFlags, nullptr, EDuplicateMode::Normal, ~EInternalObjectFlags::RootSet));
+		}
+		else
+		{
+			DestAnimSeq->Notifies[NewNotifyIndex].Notify = nullptr;
+		}
+
+		if (SrcNotifyEvent.NotifyStateClass)
+		{
+			DestAnimSeq->Notifies[NewNotifyIndex].NotifyStateClass = static_cast<UAnimNotifyState*>(StaticDuplicateObject(SrcNotifyEvent.NotifyStateClass, DestAnimSeq, NAME_None, RF_AllFlags, nullptr, EDuplicateMode::Normal, ~EInternalObjectFlags::RootSet));
+		}
+		else
+		{
+			DestAnimSeq->Notifies[NewNotifyIndex].NotifyStateClass = nullptr;
+		}
+
+		// Copy notify timing
+		NotifyEvent.LinkSequence(DestAnimSeq, SrcNotifyEvent.GetTriggerTime());
+		NotifyEvent.TriggerTimeOffset = GetTriggerTimeOffsetForType(DestAnimSeq->CalculateOffsetForNotify(NotifyEvent.GetTriggerTime()));
+
+		// Make sure editor knows we've changed something.
+		DestAnimSeq->MarkPackageDirty();
+		DestAnimSeq->RefreshCacheData();
+	}
+
+	// Inform the user if some notifies weren't copied.
+	if (bShowDialogs && NumNotifiesThatWereNotCopied > 0)
+	{
+		FMessageDialog::Open(EAppMsgType::Ok, FText::Format(
+			NSLOCTEXT("UnrealEd", "SomeNotifiesWereNotCopiedF", "Because the destination sequence was shorter, {0} notifies were not copied."), FText::AsNumber(NumNotifiesThatWereNotCopied)));
+	}
+
+	return true;
+}
+
+#endif // WITH_EDITOR
+
+bool Compression::CompressAnimationDataTracks(TArray<FRawAnimSequenceTrack>& RawAnimationData, int32 NumberOfKeys, FName ErrorName, float MaxPosDiff /*= 0.0001f*/, float MaxAngleDiff /*= 0.0003f*/)
+{
+	bool bRemovedKeys = false;
+
+#if WITH_EDITORONLY_DATA
+	if (ensureMsgf(RawAnimationData.Num() > 0, TEXT("%s is trying to compress while raw animation is missing"), * ErrorName.ToString()))
+	{
+		// This removes trivial keys, and this has to happen before the removing tracks
+		for (int32 TrackIndex = 0; TrackIndex < RawAnimationData.Num(); TrackIndex++)
+		{
+			bRemovedKeys |= CompressRawAnimSequenceTrack(RawAnimationData[TrackIndex], NumberOfKeys, ErrorName, MaxPosDiff, MaxAngleDiff);
+		}
+
+		bool bCompressScaleKeys = false;
+		// go through remove keys if not needed
+		for (int32 TrackIndex = 0; TrackIndex < RawAnimationData.Num(); TrackIndex++)
+		{
+			FRawAnimSequenceTrack const& RawData = RawAnimationData[TrackIndex];
+			if (RawData.ScaleKeys.Num() > 0)
+			{
+				// if scale key exists, see if we can just empty it
+				if ((RawData.ScaleKeys.Num() > 1) || (RawData.ScaleKeys[0].Equals(FVector::OneVector) == false))
+				{
+					bCompressScaleKeys = true;
+					break;
+				}
+			}
+		}
+
+		// if we don't have scale, we should delete all scale keys
+		// if you have one track that has scale, we still should support scale, so compress scale
+		if (!bCompressScaleKeys)
+		{
+			// then remove all scale keys
+			for (int32 TrackIndex = 0; TrackIndex < RawAnimationData.Num(); TrackIndex++)
+			{
+				FRawAnimSequenceTrack& RawData = RawAnimationData[TrackIndex];
+				RawData.ScaleKeys.Empty();
+			}
+		}
+	}
+#endif
+	return bRemovedKeys;
+}
+
+
+bool Compression::CompressRawAnimSequenceTrack(FRawAnimSequenceTrack& RawTrack, int32 NumberOfKeys, FName ErrorName, float MaxPosDiff, float MaxAngleDiff)
+{
+	bool bRemovedKeys = false;
+
+	// First part is to make sure we have valid input
+	bool const bPosTrackIsValid = (RawTrack.PosKeys.Num() == 1 || RawTrack.PosKeys.Num() == NumberOfKeys);
+	if (!bPosTrackIsValid)
+	{
+		UE_LOG(LogAnimation, Warning, TEXT("Found non valid position track for %s, %d frames, instead of %d. Chopping!"), *ErrorName.ToString(), RawTrack.PosKeys.Num(), NumberOfKeys);
+		bRemovedKeys = true;
+		RawTrack.PosKeys.RemoveAt(1, RawTrack.PosKeys.Num() - 1);
+		RawTrack.PosKeys.Shrink();
+		check(RawTrack.PosKeys.Num() == 1);
+	}
+
+	bool const bRotTrackIsValid = (RawTrack.RotKeys.Num() == 1 || RawTrack.RotKeys.Num() == NumberOfKeys);
+	if (!bRotTrackIsValid)
+	{
+		UE_LOG(LogAnimation, Warning, TEXT("Found non valid rotation track for %s, %d frames, instead of %d. Chopping!"), *ErrorName.ToString(), RawTrack.RotKeys.Num(), NumberOfKeys);
+		bRemovedKeys = true;
+		RawTrack.RotKeys.RemoveAt(1, RawTrack.RotKeys.Num() - 1);
+		RawTrack.RotKeys.Shrink();
+		check(RawTrack.RotKeys.Num() == 1);
+	}
+
+	// scale keys can be empty, and that is valid 
+	bool const bScaleTrackIsValid = (RawTrack.ScaleKeys.Num() == 0 || RawTrack.ScaleKeys.Num() == 1 || RawTrack.ScaleKeys.Num() == NumberOfKeys);
+	if (!bScaleTrackIsValid)
+	{
+		UE_LOG(LogAnimation, Warning, TEXT("Found non valid scaling track for %s, %d frames, instead of %d. Chopping!"), *ErrorName.ToString(), RawTrack.ScaleKeys.Num(), NumberOfKeys);
+		bRemovedKeys = true;
+		RawTrack.ScaleKeys.RemoveAt(1, RawTrack.ScaleKeys.Num() - 1);
+		RawTrack.ScaleKeys.Shrink();
+		check(RawTrack.ScaleKeys.Num() == 1);
+	}
+
+	// Second part is actual compression.
+
+	// Check variation of position keys
+	if ((RawTrack.PosKeys.Num() > 1) && (MaxPosDiff >= 0.0f))
+	{
+		FVector FirstPos = RawTrack.PosKeys[0];
+		bool bFramesIdentical = true;
+		for (int32 j = 1; j < RawTrack.PosKeys.Num() && bFramesIdentical; j++)
+		{
+			if ((FirstPos - RawTrack.PosKeys[j]).SizeSquared() > FMath::Square(MaxPosDiff))
+			{
+				bFramesIdentical = false;
+			}
+		}
+
+		// If all keys are the same, remove all but first frame
+		if (bFramesIdentical)
+		{
+			bRemovedKeys = true;
+			RawTrack.PosKeys.RemoveAt(1, RawTrack.PosKeys.Num() - 1);
+			RawTrack.PosKeys.Shrink();
+			check(RawTrack.PosKeys.Num() == 1);
+		}
+	}
+
+	// Check variation of rotational keys
+	if ((RawTrack.RotKeys.Num() > 1) && (MaxAngleDiff >= 0.0f))
+	{
+		FQuat FirstRot = RawTrack.RotKeys[0];
+		bool bFramesIdentical = true;
+		for (int32 j = 1; j < RawTrack.RotKeys.Num() && bFramesIdentical; j++)
+		{
+			if (FQuat::Error(FirstRot, RawTrack.RotKeys[j]) > MaxAngleDiff)
+			{
+				bFramesIdentical = false;
+			}
+		}
+
+		// If all keys are the same, remove all but first frame
+		if (bFramesIdentical)
+		{
+			bRemovedKeys = true;
+			RawTrack.RotKeys.RemoveAt(1, RawTrack.RotKeys.Num() - 1);
+			RawTrack.RotKeys.Shrink();
+			check(RawTrack.RotKeys.Num() == 1);
+		}
+	}
+
+	float MaxScaleDiff = 0.0001f;
+
+	// Check variation of Scaleition keys
+	if ((RawTrack.ScaleKeys.Num() > 1) && (MaxScaleDiff >= 0.0f))
+	{
+		FVector FirstScale = RawTrack.ScaleKeys[0];
+		bool bFramesIdentical = true;
+		for (int32 j = 1; j < RawTrack.ScaleKeys.Num() && bFramesIdentical; j++)
+		{
+			if ((FirstScale - RawTrack.ScaleKeys[j]).SizeSquared() > FMath::Square(MaxScaleDiff))
+			{
+				bFramesIdentical = false;
+			}
+		}
+
+		// If all keys are the same, remove all but first frame
+		if (bFramesIdentical)
+		{
+			bRemovedKeys = true;
+			RawTrack.ScaleKeys.RemoveAt(1, RawTrack.ScaleKeys.Num() - 1);
+			RawTrack.ScaleKeys.Shrink();
+			check(RawTrack.ScaleKeys.Num() == 1);
+		}
+	}
+
+	return bRemovedKeys;
+}
+
+void Compression::SanitizeRawAnimSequenceTrack(FRawAnimSequenceTrack& RawTrack)
+{
+	// if scale is too small, zero it out. Cause it hard to retarget when compress
+	// inverse scale is applied to translation, and causing translation to be huge to retarget, but
+	// compression can't handle that much precision. 
+	for (auto ScaleIter = RawTrack.ScaleKeys.CreateIterator(); ScaleIter; ++ScaleIter)
+	{
+		FVector& Scale3D = *ScaleIter;
+		if (FMath::IsNearlyZero(Scale3D.X))
+		{
+			Scale3D.X = 0.f;
+		}
+		if (FMath::IsNearlyZero(Scale3D.Y))
+		{
+			Scale3D.Y = 0.f;
+		}
+		if (FMath::IsNearlyZero(Scale3D.Z))
+		{
+			Scale3D.Z = 0.f;
+		}
+	}
+
+	// make sure Rotation part is normalized before compress
+	for (auto RotIter = RawTrack.RotKeys.CreateIterator(); RotIter; ++RotIter)
+	{
+		FQuat& Rotation = *RotIter;
+		if (!Rotation.IsNormalized())
+		{
+			Rotation.Normalize();
+		}
+	}
+}
+
+#if WITH_EDITOR
+bool AnimationData::AddLoopingInterpolation(UAnimSequence* InSequence)
+{
+	const UAnimDataModel* DataModel = InSequence->GetDataModel();
+	UAnimDataController* Controller = InSequence->GetController();
+
+	const int32 NumTracks = DataModel->GetNumBoneTracks();
+	const int32 NumKeys = DataModel->GetNumberOfKeys();
+
+	if (NumTracks > 0 && NumKeys > 0)
+	{
+		// now I need to calculate back to new animation data
+		auto LoopKeyData = [&](auto& KeyData)
+		{
+			// Need at least a single 
+			if (KeyData.Num() > 1)
+			{
+				auto FirstKey = KeyData[0];
+				KeyData.Add(FirstKey);
+			}
+		};
+		
+		UAnimDataController::FScopedBracket ScopedBracket(Controller, LOCTEXT("AddLoopingInterpolation_Bracket", "Adding looping interpolation"));
+
+		const TArray<FBoneAnimationTrack>& BoneAnimationTracks = DataModel->GetBoneAnimationTracks();
+		for (const FBoneAnimationTrack& AnimationTrack : BoneAnimationTracks)
+		{
+			auto PositionalKeys = AnimationTrack.InternalTrackData.PosKeys;
+			LoopKeyData(PositionalKeys);
+
+			auto RotationalKeys = AnimationTrack.InternalTrackData.RotKeys;
+			LoopKeyData(RotationalKeys);
+
+			auto ScaleKeys = AnimationTrack.InternalTrackData.ScaleKeys;
+			LoopKeyData(ScaleKeys);
+
+			Controller->SetBoneTrackKeys(AnimationTrack.Name, PositionalKeys, RotationalKeys, ScaleKeys);
+		}
+
+		const float Interval = DataModel->GetFrameRate().AsInterval();
+		Controller->SetPlayLength(Interval + DataModel->GetPlayLength());
+
+		return true;
+	}
+
+	return false;
+}
+
+bool AnimationData::Trim(UAnimSequence* InSequence, float TrimStart, float TrimEnd)
+{
+	const UAnimDataModel* DataModel = InSequence->GetDataModel();
+	UAnimDataController* Controller = InSequence->GetController();
+
+	const int32 NumTracks = DataModel->GetNumBoneTracks();
+	const int32 NumKeys = DataModel->GetNumberOfKeys();
+
+	if (NumTracks > 0 && NumKeys > 0)
+	{
+		const FFrameRate& FrameRate = DataModel->GetFrameRate();
+				
+		// Save Total Number of Keys before trimming
+		int32 TotalNumOfKeys = NumKeys;
+
+		// if there is only one key, there is nothing to trim away
+		if (NumKeys <= 1)
+		{
+			return false;
+		}
+
+		const FFrameNumber StartFrameTrim = FrameRate.AsFrameNumber(TrimStart);
+		const FFrameNumber EndFrameTrim = FrameRate.AsFrameNumber(TrimEnd);
+
+		const int32 NumTrimmedFrames = EndFrameTrim.Value - StartFrameTrim.Value;
+		
+		UAnimDataController::FScopedBracket ScopedBracket(Controller, LOCTEXT("TrimRawAnimation_Bracket", "Trimming Animation Track Data"));
+		RemoveKeys(InSequence, StartFrameTrim.Value, NumTrimmedFrames);
+		
+		return true;
+	}
+
+	return false;
+}
+
+void AnimationData::DuplicateKeys(UAnimSequence* InSequence, int32 StartKeyIndex, int32 NumDuplicates, int32 SourceKeyIndex /*= INDEX_NONE */)
+{
+	const UAnimDataModel* Model = InSequence->GetDataModel();
+	UAnimDataController* Controller = InSequence->GetController();
+
+	const int32 NumberOfKeys = Model->GetNumberOfKeys();
+	const FFrameRate& FrameRate = Model->GetFrameRate();
+
+	// Ensure that the index at which keys will be inserted, and the source key index for the duplicates is valid as well
+	if (StartKeyIndex >= 0 && StartKeyIndex <= NumberOfKeys && NumDuplicates >= 1)
+	{
+		const int32 CopyKeyIndex = SourceKeyIndex == INDEX_NONE ? StartKeyIndex : SourceKeyIndex;
+		if (CopyKeyIndex >= 0 && CopyKeyIndex < NumberOfKeys)
+		{
+			UAnimDataController::FScopedBracket ScopedBracket(Controller, LOCTEXT("DuplicateKeys_Bracket", "Inserting Duplicate Animation Track Keys"));
+
+			const int32 NumFramesToInsert = NumDuplicates;
+			const int32 EndFrameIndex = StartKeyIndex + NumDuplicates;
+
+			TArray<FBoneAnimationTrack> BoneAnimationTracks = Model->GetBoneAnimationTracks();
+			auto InsertFrames = [&](auto& KeyData)
+			{
+				if (KeyData.Num() >= 1 && KeyData.IsValidIndex(CopyKeyIndex))
+				{
+					auto SourceKeyData = KeyData[CopyKeyIndex];
+					KeyData.InsertZeroed(StartKeyIndex, NumFramesToInsert);
+
+					for (int32 FrameIndex = StartKeyIndex; FrameIndex < EndFrameIndex; ++FrameIndex)
+					{
+						KeyData[FrameIndex] = SourceKeyData;
+					}
+				}
+			};
+
+			for (FBoneAnimationTrack& AnimationTrack : BoneAnimationTracks)
+			{
+				FRawAnimSequenceTrack& TrackData = AnimationTrack.InternalTrackData;
+
+				InsertFrames(TrackData.PosKeys);
+				InsertFrames(TrackData.RotKeys);
+				InsertFrames(TrackData.ScaleKeys);
+
+				Controller->SetBoneTrackKeys(AnimationTrack.Name, TrackData.PosKeys, TrackData.RotKeys, TrackData.ScaleKeys);
+			}
+
+			// The number of keys has changed, which means that the sequence length and number of frames should be updated as well
+			const int32 NewNumKeys = NumberOfKeys + NumDuplicates;
+			const int32 NewNumFrames = NewNumKeys - 1;
+			const float NewSequenceLength = FrameRate.AsSeconds(NewNumFrames);
+
+			const float StartTime = FrameRate.AsSeconds(StartKeyIndex);
+			const float InsertedTime = FrameRate.AsInterval() * NumDuplicates;
+
+			// Notify will happen with time slice that was inserted
+			Controller->Resize(NewSequenceLength, StartTime, StartTime + InsertedTime);
+		}
+	}
+}
+
+void AnimationData::RemoveKeys(UAnimSequence* InSequence, int32 StartKeyIndex, int32 NumKeysToRemove)
+{
+	const UAnimDataModel* Model = InSequence->GetDataModel();
+	UAnimDataController* Controller = InSequence->GetController();
+
+	const int32 NumberOfKeys = Model->GetNumberOfKeys();
+	const FFrameRate& FrameRate = Model->GetFrameRate();
+
+	const int32 EndKeyIndex = StartKeyIndex + NumKeysToRemove;
+	if (StartKeyIndex >= 0 && StartKeyIndex < NumberOfKeys && NumKeysToRemove > 0 && EndKeyIndex <= NumberOfKeys)
+	{
+		UAnimDataController::FScopedBracket ScopedBracket(Controller, LOCTEXT("RemoveKeys_Bracket", "Removing Animation Track Keys"));
+
+		auto ShrinkKeys = [&](auto& KeyData)
+		{
+			if (KeyData.Num() >= (StartKeyIndex + NumKeysToRemove))
+			{
+				KeyData.RemoveAt(StartKeyIndex, NumKeysToRemove);
+				check(KeyData.Num() > 0);
+				KeyData.Shrink();
+			}
+		};
+
+		const int32 NewNumberOfKeys = NumberOfKeys - NumKeysToRemove;
+
+		TArray<FBoneAnimationTrack> BoneAnimationTracks = Model->GetBoneAnimationTracks();
+		for (FBoneAnimationTrack& AnimationTrack : BoneAnimationTracks)
+		{
+			FRawAnimSequenceTrack& TrackData = AnimationTrack.InternalTrackData;
+
+			ShrinkKeys(TrackData.PosKeys);
+			ShrinkKeys(TrackData.RotKeys);
+			ShrinkKeys(TrackData.ScaleKeys);
+			
+			Controller->SetBoneTrackKeys(AnimationTrack.Name, TrackData.PosKeys, TrackData.RotKeys, TrackData.ScaleKeys);
+		}
+
+		const int32 NewNumberOfFrames = FMath::Max(NewNumberOfKeys - 1, 0);
+		const float NewSequenceLength = FrameRate.AsSeconds(NewNumberOfFrames);
+
+		const float StartTime = FrameRate.AsSeconds(StartKeyIndex);
+		const float RemovedTime = NumKeysToRemove * FrameRate.AsInterval();
+
+		// Notify will happen with time slice that was removed
+		Controller->Resize(NewSequenceLength, StartTime, RemovedTime);
+	}
+}
+
+int32 AnimationData::FindFirstChildTrackIndex(const UAnimSequence* InSequence, const USkeleton* Skeleton, const FName& BoneName)
+{
+	const UAnimDataModel* DataModel = InSequence->GetDataModel();
+	const FReferenceSkeleton& RefSkeleton = Skeleton->GetReferenceSkeleton();
+
+	const int32 BoneIndex = RefSkeleton.FindBoneIndex(BoneName);
+	if (BoneIndex == INDEX_NONE)
+	{
+		// get out, nothing to do
+		return INDEX_NONE;
+	}
+
+	// find children
+	TArray<int32> Childs;
+	if (Skeleton->GetChildBones(BoneIndex, Childs) > 0)
+	{
+		// first look for direct children
+		for (auto ChildIndex : Childs)
+		{
+			FName ChildBoneName = RefSkeleton.GetBoneName(ChildIndex);
+
+			int32 ChildTrackIndex = DataModel->GetBoneTrackIndexByName(ChildBoneName);
+			if (ChildTrackIndex != INDEX_NONE)
+			{
+				// found the new track
+				return ChildTrackIndex;
+			}
+		}
+
+		int32 BestGrandChildIndex = INDEX_NONE;
+		// if you didn't find yet, now you have to go through all children
+		for (auto ChildIndex : Childs)
+		{
+			FName ChildBoneName = RefSkeleton.GetBoneName(ChildIndex);
+			// now I have to go through all childrewn and find who is earliest since I don't know which one might be the closest one
+			int32 GrandChildIndex = FindFirstChildTrackIndex(InSequence, Skeleton, ChildBoneName);
+			if (GrandChildIndex != INDEX_NONE)
+			{
+				if (BestGrandChildIndex == INDEX_NONE)
+				{
+					BestGrandChildIndex = GrandChildIndex;
+				}
+				else if (BestGrandChildIndex > GrandChildIndex)
+				{
+					// best should be earlier track index
+					BestGrandChildIndex = GrandChildIndex;
+				}
+			}
+		}
+
+		return BestGrandChildIndex;
+	}
+
+	// there is no child, just add at the end
+	return DataModel->GetNumBoneTracks();
+}
+
+#endif // WITH_EDITOR
+
+
+} // Namespace Anim
+
+} // Namespace UE
+
+#undef LOCTEXT_NAMESPACE 
\ No newline at end of file
diff --git a/Engine/Source/Runtime/Engine/Private/Animation/AnimSet.cpp b/Engine/Source/Runtime/Engine/Private/Animation/AnimSet.cpp
--- a/Engine/Source/Runtime/Engine/Private/Animation/AnimSet.cpp
+++ b/Engine/Source/Runtime/Engine/Private/Animation/AnimSet.cpp
@@ -136,7 +136,7 @@ void UAnimSet::ResetAnimSet()
 		UAnimSequence* AnimSeq = Sequences[i];
 		if( AnimSeq )
 		{
-			AnimSeq->RecycleAnimSequence();
+			AnimSeq->ResetAnimation();
 		}
 	}
 	Sequences.Empty();
@@ -164,7 +164,7 @@ bool UAnimSet::RemoveAnimSequenceFromAnimSet(UAnimSequence* AnimSeq)
 	if( SequenceIndex != INDEX_NONE )
 	{
 		// Handle reference clean up properly
-		AnimSeq->RecycleAnimSequence();
+		AnimSeq->ResetAnimation();
 		// Remove from array
 		Sequences.RemoveAt(SequenceIndex, 1);
 		if( GIsEditor )
diff --git a/Engine/Source/Runtime/Engine/Private/Animation/AnimStreamable.cpp b/Engine/Source/Runtime/Engine/Private/Animation/AnimStreamable.cpp
--- a/Engine/Source/Runtime/Engine/Private/Animation/AnimStreamable.cpp
+++ b/Engine/Source/Runtime/Engine/Private/Animation/AnimStreamable.cpp
@@ -23,6 +23,10 @@
 #include "Animation/CustomAttributesRuntime.h"
 #include "UObject/UE5MainStreamObjectVersion.h"
 
+#include "Animation/AnimData/AnimDataModel.h"
+#include "Animation/AnimSequenceHelpers.h"
+#include "Animation/AnimData/AnimDataController.h"
+
 CSV_DECLARE_CATEGORY_MODULE_EXTERN(ENGINE_API, Animation);
 
 DECLARE_CYCLE_STAT(TEXT("AnimStreamable GetAnimationPose"), STAT_AnimStreamable_GetAnimationPose, STATGROUP_Anim);
@@ -274,24 +278,22 @@ void UAnimStreamable::GetAnimationPose(FAnimationPoseData& OutAnimationPoseData,
 	if (!HasRunningPlatformData() || RequiredBones.ShouldUseRawData())
 	{
 		//Need to evaluate raw data
-		RawCurveData.EvaluateCurveData(OutCurve, ExtractionContext.CurrentTime);
+		ValidateModel();
 
-		const int32 NumTracks = TrackToSkeletonMapTable.Num();
+		UE::Anim::EvaluateFloatCurvesFromModel(DataModel, OutCurve, ExtractionContext.CurrentTime);
 
 		// Warning if we have invalid data
-
-		for (int32 TrackIndex = 0; TrackIndex < NumTracks; TrackIndex++)
+		for (const FBoneAnimationTrack& AnimationTrack : DataModel->GetBoneAnimationTracks())
 		{
-			const FRawAnimSequenceTrack& TrackToExtract = RawAnimationData[TrackIndex];
-
+			const FRawAnimSequenceTrack& TrackToExtract = AnimationTrack.InternalTrackData;
 			// Bail out (with rather wacky data) if data is empty for some reason.
 			if (TrackToExtract.PosKeys.Num() == 0 || TrackToExtract.RotKeys.Num() == 0)
 			{
-				UE_LOG(LogAnimation, Warning, TEXT("UAnimSequence::GetBoneTransform : No anim data in AnimSequence '%s' Track '%s'"), *GetPathName(), *AnimationTrackNames[TrackIndex].ToString());
+				UE_LOG(LogAnimation, Warning, TEXT("UAnimSequence::GetBoneTransform : No anim data in AnimSequence '%s' Track '%s'"), *GetPathName(), *AnimationTrack.Name.ToString());
 			}
 		}
 
-		BuildPoseFromRawData(RawAnimationData, TrackToSkeletonMapTable, OutPose, ExtractionContext.CurrentTime, Interpolation, NumberOfKeys, GetPlayLength(), RetargetSource);
+		UE::Anim::BuildPoseFromModel(DataModel, OutPose, ExtractionContext.CurrentTime, Interpolation, RetargetSource, MySkeleton->GetRefLocalPoses(RetargetSource));
 
 		if ((ExtractionContext.bExtractRootMotion && RootMotionReset.bEnableRootMotion) || RootMotionReset.bForceRootLock)
 		{
@@ -381,8 +383,10 @@ void UAnimStreamable::PostLoad()
 		}
 		NonConstSeq->ConditionalPostLoad();
 	}
-
-	if (SourceSequence && (GenerateGuidFromRawAnimData(SourceSequence->GetRawAnimationData(), SourceSequence->RawCurveData) != RawDataGuid))
+	
+	const bool bRequiresModelPopulation = GetLinkerCustomVersion(FUE5MainStreamObjectVersion::GUID) < FUE5MainStreamObjectVersion::IntroducingAnimationDataModel;
+	const FGuid CurrentGuid = SourceSequence ? SourceSequence->GetDataModel()->GenerateGuid() : FGuid();
+	if (SourceSequence && ( CurrentGuid != RawDataGuid || bRequiresModelPopulation))
 	{
 		InitFrom(SourceSequence);
 	}
@@ -449,22 +453,15 @@ void UAnimStreamable::InitFrom(const UAnimSequence* InSourceSequence)
 	BoneCompressionSettings = InSourceSequence->BoneCompressionSettings;
 	CurveCompressionSettings = InSourceSequence->CurveCompressionSettings;
 
-	RawAnimationData = InSourceSequence->GetRawAnimationData();
-	RawCurveData = InSourceSequence->RawCurveData;
+	DataModel = DuplicateObject(InSourceSequence->GetDataModel(), this);
+	Controller->SetModel(DataModel);
 
 	Notifies = InSourceSequence->Notifies;
-
-	TrackToSkeletonMapTable = InSourceSequence->GetRawTrackToSkeletonMapTable();
-	AnimationTrackNames = InSourceSequence->GetAnimationTrackNames();
-	
-	NumberOfKeys = InSourceSequence->GetNumberOfSampledKeys();
-	SetSequenceLength(InSourceSequence->GetPlayLength());
-	
+	Controller->SetPlayLength(InSourceSequence->GetDataModel()->GetPlayLength());
 	RateScale = InSourceSequence->RateScale;
-
 	Interpolation = InSourceSequence->Interpolation;
-
 	RetargetSource = InSourceSequence->RetargetSource;
+	NumberOfKeys = DataModel->GetNumberOfKeys();
 
 	bEnableRootMotion = InSourceSequence->bEnableRootMotion;
 	RootMotionRootLock = InSourceSequence->RootMotionRootLock;
@@ -649,12 +646,13 @@ void UAnimStreamable::RequestCompressedDataForChunk(const FString& ChunkDDCKey,
 		{
 			FCompressibleAnimRef CompressibleData = MakeShared<FCompressibleAnimData, ESPMode::ThreadSafe>(BoneCompressionSettings, CurveCompressionSettings, GetSkeleton(), Interpolation, Chunk.SequenceLength, ChunkNumFrames+1);
 
-			CompressibleData->RawAnimationData.AddDefaulted(RawAnimationData.Num());
+			const TArray<FBoneAnimationTrack>& BoneAnimationTracks = DataModel->GetBoneAnimationTracks();
+			CompressibleData->TrackToSkeletonMapTable.Empty();
 
-			for (int32 TrackIndex = 0; TrackIndex < RawAnimationData.Num(); ++TrackIndex)
+			for (const FBoneAnimationTrack& AnimTrack : BoneAnimationTracks)
 			{
-				FRawAnimSequenceTrack& SrcTrack = RawAnimationData[TrackIndex];
-				FRawAnimSequenceTrack& DestTrack = CompressibleData->RawAnimationData[TrackIndex];
+				const FRawAnimSequenceTrack& SrcTrack = AnimTrack.InternalTrackData;
+				FRawAnimSequenceTrack& DestTrack = CompressibleData->RawAnimationData.AddDefaulted_GetRef();
 
 				MakeKeyChunk(SrcTrack.PosKeys, DestTrack.PosKeys, NumberOfKeys, FrameStart, FrameEnd);
 				MakeKeyChunk(SrcTrack.RotKeys, DestTrack.RotKeys, NumberOfKeys, FrameStart, FrameEnd);
@@ -662,12 +660,14 @@ void UAnimStreamable::RequestCompressedDataForChunk(const FString& ChunkDDCKey,
 				{
 					MakeKeyChunk(SrcTrack.ScaleKeys, DestTrack.ScaleKeys, NumberOfKeys, FrameStart, FrameEnd);
 				}
+
+				CompressibleData->TrackToSkeletonMapTable.Add(AnimTrack.BoneTreeIndex);
 			}
 
 			if (FrameStart == 0)
 			{
 				//Crop curve logic broken, for the moment store curve data in always loaded chunk 0
-				CompressibleData->RawCurveData = RawCurveData;
+				CompressibleData->RawFloatCurves = DataModel->GetFloatCurves();
 			}
 
 			/*if (SourceSequence && ChunkIndex == 7)
@@ -683,7 +683,6 @@ void UAnimStreamable::RequestCompressedDataForChunk(const FString& ChunkDDCKey,
 				SourceSequence = Seq;
 			}*/
 
-			CompressibleData->TrackToSkeletonMapTable = TrackToSkeletonMapTable;
 			AnimCompressor->SetCompressibleData(CompressibleData);
 
 			if (bSkipDDC)
@@ -725,7 +724,7 @@ void UAnimStreamable::RequestCompressedDataForChunk(const FString& ChunkDDCKey,
 
 void UAnimStreamable::UpdateRawData()
 {
-	RawDataGuid = GenerateGuidFromRawAnimData(RawAnimationData, RawCurveData);
+	RawDataGuid = DataModel->GenerateGuid();
 	RequestCompressedData();
 }
 
diff --git a/Engine/Source/Runtime/Engine/Private/Animation/BonePose.cpp b/Engine/Source/Runtime/Engine/Private/Animation/BonePose.cpp
--- a/Engine/Source/Runtime/Engine/Private/Animation/BonePose.cpp
+++ b/Engine/Source/Runtime/Engine/Private/Animation/BonePose.cpp
@@ -297,7 +297,9 @@ void BuildPoseFromRawData(
 	if (MySkeleton)
 	{
 		const TArray<FTransform>& RetargetTransforms = MySkeleton->GetRefLocalPoses(RetargetSource);
+		PRAGMA_DISABLE_DEPRECATION_WARNINGS
 		BuildPoseFromRawData(InAnimationData, TrackToSkeletonMapTable, InOutPose, InTime, Interpolation, NumFrames, SequenceLength, RetargetSource, RetargetTransforms, AdditiveBoneTransformCurves);
+		PRAGMA_ENABLE_DEPRECATION_WARNINGS
 	}
 }
 
diff --git a/Engine/Source/Runtime/Engine/Private/Animation/CurveIdentifier.cpp b/Engine/Source/Runtime/Engine/Private/Animation/CurveIdentifier.cpp
new file mode 100644
--- /dev/null
+++ b/Engine/Source/Runtime/Engine/Private/Animation/CurveIdentifier.cpp
@@ -0,0 +1,99 @@
+// Copyright Epic Games, Inc. All Rights Reserved.
+
+#include "Animation/AnimData/CurveIdentifier.h"
+#include "EngineLogs.h"
+
+#if WITH_EDITOR
+FAnimationCurveIdentifier UAnimationCurveIdentifierExtensions::GetCurveIdentifier(USkeleton* InSkeleton, FName Name, ERawCurveTrackTypes CurveType)
+{
+	FAnimationCurveIdentifier Identifier;
+
+	if (InSkeleton)
+	{
+		const FName MappingName = CurveType == ERawCurveTrackTypes::RCT_Float ? USkeleton::AnimCurveMappingName : USkeleton::AnimTrackCurveMappingName;
+		const bool bExists = InSkeleton->GetSmartNameByName(MappingName, Name, Identifier.InternalName);
+		if (!bExists)
+		{
+			InSkeleton->AddSmartNameAndModify(MappingName, Name, Identifier.InternalName);
+		}
+
+		Identifier.CurveType = CurveType;
+	}
+	else
+	{
+		UE_LOG(LogAnimation, Warning, TEXT("Invalid Skeleton provided for GetCurveIdentifier"));
+	}
+
+	return Identifier;
+}
+
+FAnimationCurveIdentifier UAnimationCurveIdentifierExtensions::FindCurveIdentifier(const USkeleton* InSkeleton, FName Name, ERawCurveTrackTypes CurveType)
+{
+	FAnimationCurveIdentifier Identifier;
+
+	if (InSkeleton)
+	{
+		const FName MappingName = CurveType == ERawCurveTrackTypes::RCT_Float ? USkeleton::AnimCurveMappingName : USkeleton::AnimTrackCurveMappingName;
+		const bool bExists = InSkeleton->GetSmartNameByName(MappingName, Name, Identifier.InternalName);
+		Identifier.CurveType = CurveType;
+	}
+	else
+	{
+		UE_LOG(LogAnimation, Warning, TEXT("Invalid Skeleton provided for FindCurveIdentifier"));
+	}
+
+	return Identifier;
+}
+
+TArray<FAnimationCurveIdentifier> UAnimationCurveIdentifierExtensions::GetCurveIdentifiers(USkeleton* InSkeleton, ERawCurveTrackTypes CurveType)
+{
+	TArray<FAnimationCurveIdentifier> Identifiers;
+	if (InSkeleton)
+	{
+		const FName MappingName = CurveType == ERawCurveTrackTypes::RCT_Float ? USkeleton::AnimCurveMappingName : USkeleton::AnimTrackCurveMappingName;
+		const FSmartNameMapping* SmartNameContainer = InSkeleton->GetSmartNameContainer(MappingName);
+		if (SmartNameContainer)
+		{
+			TArray<SmartName::UID_Type> UIDs;
+			SmartNameContainer->FillUidArray(UIDs);
+
+			for (const SmartName::UID_Type& UID : UIDs)
+			{
+				FSmartName SmartName;
+				if (SmartNameContainer->FindSmartNameByUID(UID, SmartName))
+				{
+
+					Identifiers.Add(FAnimationCurveIdentifier(SmartName, CurveType));
+				}
+			}
+		}
+		else
+		{
+			UE_LOG(LogAnimation, Warning, TEXT("Unable to find smartname mapping %s in Skeleton (%s)"), *MappingName.ToString(), *InSkeleton->GetName());
+		}
+	}
+	else
+	{
+		UE_LOG(LogAnimation, Warning, TEXT("Invalid Skeleton provided for GetCurveIdentifiers"));
+	}
+
+	return Identifiers;
+}
+
+bool UAnimationCurveIdentifierExtensions::GetTransformChildCurveIdentifier(FAnimationCurveIdentifier& InOutIdentifier, ETransformCurveChannel Channel, EVectorCurveChannel Axis)
+{
+	if (InOutIdentifier.IsValid())
+	{
+		if (InOutIdentifier.CurveType == ERawCurveTrackTypes::RCT_Transform)
+		{
+			InOutIdentifier.Channel = Channel;
+			InOutIdentifier.Axis = Axis;
+
+			return true;
+		}
+	}
+
+	return false;
+}
+#endif // WITH_EDITOR
+
diff --git a/Engine/Source/Runtime/Engine/Private/Animation/PoseAsset.cpp b/Engine/Source/Runtime/Engine/Private/Animation/PoseAsset.cpp
--- a/Engine/Source/Runtime/Engine/Private/Animation/PoseAsset.cpp
+++ b/Engine/Source/Runtime/Engine/Private/Animation/PoseAsset.cpp
@@ -8,6 +8,9 @@
 #include "Animation/AnimSequence.h"
 #include "Animation/AnimSequenceBase.h"
 #include "Animation/AnimationPoseData.h"
+#include "Animation/AnimData/AnimDataModel.h"
+#include "Animation/AnimSequenceHelpers.h"
+
 #define LOCTEXT_NAMESPACE "PoseAsset"
 
 // utility function 
@@ -1192,7 +1195,12 @@ void UPoseAsset::CreatePoseFromAnimation(class UAnimSequence* AnimSequence, cons
 				FMemMark Mark(FMemStack::Get());
 
 				// set up track data - @todo: add revaliation code when checked
-				for (const FName& TrackName : AnimSequence->GetAnimationTrackNames())
+				UAnimDataModel* DataModel = AnimSequence->GetDataModel();
+
+				TArray<FName> TrackNames;
+				DataModel->GetBoneTrackNames(TrackNames);
+
+				for (const FName& TrackName : TrackNames)
 				{
 					PoseContainer.Tracks.Add(TrackName);
 				}
@@ -1200,7 +1208,7 @@ void UPoseAsset::CreatePoseFromAnimation(class UAnimSequence* AnimSequence, cons
 				// now create pose transform
 				TArray<FTransform> NewPose;
 				
-				const int32 NumTracks = AnimSequence->GetAnimationTrackNames().Num();
+				const int32 NumTracks = TrackNames.Num();
 				NewPose.Reset(NumTracks);
 				NewPose.AddUninitialized(NumTracks);
 
@@ -1209,7 +1217,9 @@ void UPoseAsset::CreatePoseFromAnimation(class UAnimSequence* AnimSequence, cons
 				float IntervalBetweenKeys = (NumPoses > 1)? AnimSequence->GetPlayLength() / (NumPoses -1 ) : 0.f;
 
 				// add curves - only float curves
-				const int32 TotalFloatCurveCount = AnimSequence->RawCurveData.FloatCurves.Num();
+				const FAnimationCurveData& AnimationCurveData = DataModel->GetCurveData();
+
+				const int32 TotalFloatCurveCount = AnimationCurveData.FloatCurves.Num();
 				
 				// have to construct own UIDList;
 				// copy default UIDLIst
@@ -1217,7 +1227,7 @@ void UPoseAsset::CreatePoseFromAnimation(class UAnimSequence* AnimSequence, cons
 
 				if (TotalFloatCurveCount > 0)
 				{
-					for (const FFloatCurve& Curve : AnimSequence->RawCurveData.FloatCurves)
+					for (const FFloatCurve& Curve : AnimationCurveData.FloatCurves)
 					{
 						PoseContainer.Curves.Add(FAnimCurveBase(Curve.Name, Curve.GetCurveTypeFlags()));
 						UIDList.Add(Curve.Name.UID);
@@ -1232,8 +1242,7 @@ void UPoseAsset::CreatePoseFromAnimation(class UAnimSequence* AnimSequence, cons
 					// now get rawanimationdata, and each key is converted to new pose
 					for (int32 TrackIndex = 0; TrackIndex < NumTracks; ++TrackIndex)
 					{
-						const auto& RawTrack = AnimSequence->GetRawAnimationTrack(TrackIndex);
-						AnimSequence->ExtractBoneTransform(RawTrack, NewPose[TrackIndex], PoseIndex);
+						UE::Anim::GetBoneTransformFromModel(AnimSequence->GetDataModel(), NewPose[TrackIndex], TrackIndex, PoseIndex);
 					}
 
 					if (TotalFloatCurveCount > 0)
diff --git a/Engine/Source/Runtime/Engine/Private/Animation/Skeleton.cpp b/Engine/Source/Runtime/Engine/Private/Animation/Skeleton.cpp
--- a/Engine/Source/Runtime/Engine/Private/Animation/Skeleton.cpp
+++ b/Engine/Source/Runtime/Engine/Private/Animation/Skeleton.cpp
@@ -31,6 +31,7 @@
 #include "Animation/AnimBlueprint.h"
 #include "UObject/AnimObjectVersion.h"
 #include "EngineUtils.h"
+#include "Animation/AnimData/AnimDataModel.h"
 
 #define LOCTEXT_NAMESPACE "Skeleton"
 #define ROOT_BONE_PARENT	INDEX_NONE
@@ -776,10 +777,17 @@ int32 USkeleton::GetRawAnimationTrackIndex(const int32 InSkeletonBoneIndex, cons
 {
 	if( InSkeletonBoneIndex != INDEX_NONE )
 	{
-		return InAnimSeq->GetRawTrackToSkeletonMapTable().IndexOfByPredicate([&](const FTrackToSkeletonMap& TrackToSkel)
+#if WITH_EDITOR
+		return InAnimSeq->GetDataModel()->GetBoneAnimationTracks().IndexOfByPredicate([InSkeletonBoneIndex](const FBoneAnimationTrack& AnimationTrack)
 		{
-			return TrackToSkel.BoneTreeIndex == InSkeletonBoneIndex;
+			return AnimationTrack.BoneTreeIndex == InSkeletonBoneIndex;
 		});
+#else
+		return InAnimSeq->GetCompressedTrackToSkeletonMapTable().IndexOfByPredicate([InSkeletonBoneIndex](const FTrackToSkeletonMap& Mapping)
+		{
+			return Mapping.BoneTreeIndex == InSkeletonBoneIndex;
+		});
+#endif
 	}
 
 	return INDEX_NONE;
diff --git a/Engine/Source/Runtime/Engine/Private/AnimationAsset.cpp b/Engine/Source/Runtime/Engine/Private/AnimationAsset.cpp
--- a/Engine/Source/Runtime/Engine/Private/AnimationAsset.cpp
+++ b/Engine/Source/Runtime/Engine/Private/AnimationAsset.cpp
@@ -10,6 +10,10 @@
 #include "UObject/LinkerLoad.h"
 #include "Animation/BlendSpaceBase.h"
 #include "Animation/PoseAsset.h"
+#include "Animation/AnimSequenceHelpers.h"
+#include "Animation/AnimData/AnimDataController.h"
+
+#define LOCTEXT_NAMESPACE "AnimationAsset"
 
 #define LEADERSCORE_ALWAYSLEADER  	2.f
 #define LEADERSCORE_MONTAGE			3.f
@@ -401,6 +405,13 @@ bool UAnimationAsset::ReplaceSkeleton(USkeleton* NewSkeleton, bool bConvertSpace
 				}
 				AnimAsset->ConditionalPostLoad();
 
+				// This ensure that in subsequent behaviour the RawData GUID is never 'new-ed' but always calculated from the 
+				// raw animation data itself.
+				if (UAnimSequence* Sequence = Cast<UAnimSequence>(AnimAsset))
+				{
+					Sequence->GetController()->OpenBracket(LOCTEXT("ReplaceSkeleton_Bracket", "Replacing USkeleton"));
+				}
+
 				// these two are different functions for now
 				// technically if you have implementation for Remap, it will also set skeleton 
 				AnimAsset->RemapTracksToNewSkeleton(NewSkeleton, bConvertSpaces);
@@ -411,17 +422,26 @@ bool UAnimationAsset::ReplaceSkeleton(USkeleton* NewSkeleton, bool bConvertSpace
 			{
 				if (UAnimSequence* Seq = Cast<UAnimSequence>(AnimAsset))
 				{
-					// We don't force gen here as that can cause us to constantly generate
-					// new anim ddc keys if users never resave anims that need to remap.
-					Seq->PostProcessSequence(false);
+					Seq->GetController()->CloseBracket();
 				}
 			}
 		}
 
-		RemapTracksToNewSkeleton(NewSkeleton, bConvertSpaces);
-		if (UAnimSequence* Seq = Cast<UAnimSequence>(this))
-		{
-			Seq->PostProcessSequence(false);
+		UAnimSequence* Seq = Cast<UAnimSequence>(this);
+		{			
+			// This ensure that in subsequent behaviour the RawData GUID is never 'new-ed' but always calculated from the 
+			// raw animation data itself.
+			if (Seq)
+			{
+				Seq->GetController()->OpenBracket(LOCTEXT("ReplaceSkeleton_Bracket", "Replacing USkeleton"));
+			}
+  
+			RemapTracksToNewSkeleton(NewSkeleton, bConvertSpaces);
+
+			if (Seq)
+			{
+				Seq->GetController()->CloseBracket();
+			}
 		}
 
 		PostEditChange();
@@ -714,3 +734,4 @@ void FBlendSampleData::NormalizeDataWeight(TArray<FBlendSampleData>& SampleDataL
 	}
 }
 
+#undef LOCTEXT_NAMESPACE // "AnimationAsset"
\ No newline at end of file
diff --git a/Engine/Source/Runtime/Engine/Private/UnrealEngine.cpp b/Engine/Source/Runtime/Engine/Private/UnrealEngine.cpp
--- a/Engine/Source/Runtime/Engine/Private/UnrealEngine.cpp
+++ b/Engine/Source/Runtime/Engine/Private/UnrealEngine.cpp
@@ -5926,7 +5926,7 @@ bool UEngine::HandleListAnimsCommand(const TCHAR* Cmd, FOutputDevice& Ar)
 			NumKeys = AnimSeq->GetNumberOfSampledKeys();
 			SequenceLength = AnimSeq->GetPlayLength();
 			RateScale = AnimSeq->RateScale;
-			NumCurves = AnimSeq->RawCurveData.FloatCurves.Num();
+			NumCurves = AnimSeq->GetCurveData().FloatCurves.Num();
 
 			if (AnimSeq->CompressedData.CompressedDataStructure != nullptr)
 			{
diff --git a/Engine/Source/Runtime/Engine/Public/Animation/AnimCompressionTypes.h b/Engine/Source/Runtime/Engine/Public/Animation/AnimCompressionTypes.h
--- a/Engine/Source/Runtime/Engine/Public/Animation/AnimCompressionTypes.h
+++ b/Engine/Source/Runtime/Engine/Public/Animation/AnimCompressionTypes.h
@@ -45,8 +45,6 @@ class UAnimBoneCompressionSettings;
 class UAnimBoneCompressionCodec;
 class USkeleton;
 
-extern FGuid GenerateGuidFromRawAnimData(const TArray<FRawAnimSequenceTrack>& RawAnimationData, const FRawCurveTracks& RawCurveData);
-
 template<typename ArrayClass>
 struct ENGINE_API FCompressedOffsetDataBase
 {
@@ -215,7 +213,7 @@ public:
 
 	TArray<FBoneData> BoneData;
 
-	FRawCurveTracks RawCurveData;
+	TArray<FFloatCurve> RawFloatCurves;
 
 	float SequenceLength;
 
@@ -256,7 +254,7 @@ public:
 	int32 GetApproxRawCurveSize() const
 	{
 		int32 Total = 0;
-		for (const FFloatCurve& Curve : RawCurveData.FloatCurves)
+		for (const FFloatCurve& Curve : RawFloatCurves)
 		{
 			Total += sizeof(FFloatCurve);
 			Total += sizeof(FRichCurveKey) * Curve.FloatCurve.Keys.Num();
diff --git a/Engine/Source/Runtime/Engine/Public/Animation/AnimCurveTypes.h b/Engine/Source/Runtime/Engine/Public/Animation/AnimCurveTypes.h
--- a/Engine/Source/Runtime/Engine/Public/Animation/AnimCurveTypes.h
+++ b/Engine/Source/Runtime/Engine/Public/Animation/AnimCurveTypes.h
@@ -8,6 +8,7 @@
 #include "Animation/SmartName.h"
 #include "Animation/Skeleton.h"
 #include "Curves/RichCurve.h"
+#include "Misc/EnumRange.h"
 #include "AnimCurveTypes.generated.h"
 
 /** This is curve flags that are saved in asset and **/
@@ -32,6 +33,8 @@ enum EAnimAssetCurveFlags
 };
 static const EAnimAssetCurveFlags AACF_DefaultCurve = AACF_Editable;
 
+ENUM_RANGE_BY_FIRST_AND_LAST(EAnimAssetCurveFlags, AACF_DriveMorphTarget_DEPRECATED, AACF_Disabled);
+
 /** UI Curve Parameter type
  * This gets name, and cached UID and use it when needed
  * Also it contains curve types 
@@ -69,7 +72,7 @@ struct ENGINE_API FAnimCurveParam
 /**
  * Float curve data for one track
  */
-USTRUCT()
+USTRUCT(BlueprintType)
 struct ENGINE_API FAnimCurveBase
 {
 	GENERATED_USTRUCT_BODY()
@@ -153,7 +156,7 @@ private:
 #endif
 };
 
-USTRUCT()
+USTRUCT(BlueprintType)
 struct FFloatCurve : public FAnimCurveBase
 {
 	GENERATED_USTRUCT_BODY()
@@ -211,7 +214,7 @@ struct FVectorCurve : public FAnimCurveBase
 	int32 GetNumKeys() const;
 };
 
-USTRUCT()
+USTRUCT(BlueprintType)
 struct FTransformCurve: public FAnimCurveBase
 {
 	GENERATED_USTRUCT_BODY()
@@ -771,11 +774,11 @@ struct ENGINE_API FBlendedHeapCurve : public FBaseBlendedCurve<FDefaultAllocator
 {
 };
 
-UENUM()
+UENUM(BlueprintType)
 enum class ERawCurveTrackTypes : uint8
 {
 	RCT_Float UMETA(DisplayName = "Float Curve"),
-	RCT_Vector UMETA(DisplayName = "Vector Curve"),
+	RCT_Vector UMETA(DisplayName = "Vector Curve", Hidden),
 	RCT_Transform UMETA(DisplayName = "Transformation Curve"),
 	RCT_MAX
 };
diff --git a/Engine/Source/Runtime/Engine/Public/Animation/AnimTypes.h b/Engine/Source/Runtime/Engine/Public/Animation/AnimTypes.h
--- a/Engine/Source/Runtime/Engine/Public/Animation/AnimTypes.h
+++ b/Engine/Source/Runtime/Engine/Public/Animation/AnimTypes.h
@@ -27,6 +27,8 @@ class FMemoryWriter;
 // Enable this if you want to locally measure detailed anim perf.  Disabled by default as it is introduces a lot of additional profile markers and associated overhead.
 #define ENABLE_VERBOSE_ANIM_PERF_TRACKING 0
 
+#define MAX_ANIMATION_TRACKS 65535
+
 namespace EAnimEventTriggerOffsets
 {
 	enum Type
@@ -839,21 +841,21 @@ struct ENGINE_API FTrackToSkeletonMap
 * One element is used as a simple compression scheme where if all keys are the same, they'll be
 * reduced to 1 key that is constant over the entire sequence.
 */
-USTRUCT()
+USTRUCT(BlueprintType)
 struct ENGINE_API FRawAnimSequenceTrack
 {
 	GENERATED_USTRUCT_BODY()
 
 	/** Position keys. */
-	UPROPERTY()
+	UPROPERTY(BlueprintReadOnly, Category=TrackData)
 	TArray<FVector> PosKeys;
 
 	/** Rotation keys. */
-	UPROPERTY()
+	UPROPERTY(BlueprintReadOnly, Category = TrackData)
 	TArray<FQuat> RotKeys;
 
 	/** Scale keys. */
-	UPROPERTY()
+	UPROPERTY(BlueprintReadOnly, Category = TrackData)
 	TArray<FVector> ScaleKeys;
 
 	// Serializer.
@@ -869,6 +871,8 @@ struct ENGINE_API FRawAnimSequenceTrack
 
 		return Ar;
 	}
+
+	static const uint32 SingleKeySize = sizeof(FVector) + sizeof(FQuat) + sizeof(FVector);
 };
 
 
diff --git a/Engine/Source/Runtime/Engine/Public/BonePose.h b/Engine/Source/Runtime/Engine/Public/BonePose.h
--- a/Engine/Source/Runtime/Engine/Public/BonePose.h
+++ b/Engine/Source/Runtime/Engine/Public/BonePose.h
@@ -952,6 +952,7 @@ void FCSPose<PoseType>::ConvertComponentPosesToLocalPosesSafe(FCSPose<PoseType>&
 }
 
 // Populate InOutPose based on raw animation data. 
+UE_DEPRECATED(5.0, "BuildPoseFromRawData has been deprecated, use BuildPoseFromRawData signature with RetargetTransforms parameter")
 extern void BuildPoseFromRawData(
 	const TArray<FRawAnimSequenceTrack>& InAnimationData,
 	const TArray<struct FTrackToSkeletonMap>& TrackToSkeletonMapTable,
@@ -964,6 +965,7 @@ extern void BuildPoseFromRawData(
 	const TMap<int32, const struct FTransformCurve*>* AdditiveBoneTransformCurves = nullptr
 );
 
+UE_DEPRECATED(5.0, "BuildPoseFromRawData has been deprecated, use UE::Anim::BuildPoseFromModel")
 extern void BuildPoseFromRawData(
 	const TArray<FRawAnimSequenceTrack>& InAnimationData, 
 	const TArray<struct FTrackToSkeletonMap>& TrackToSkeletonMapTable, 
diff --git a/Engine/Source/Runtime/MovieSceneTracks/Private/Sections/MovieSceneSkeletalAnimationSection.cpp b/Engine/Source/Runtime/MovieSceneTracks/Private/Sections/MovieSceneSkeletalAnimationSection.cpp
--- a/Engine/Source/Runtime/MovieSceneTracks/Private/Sections/MovieSceneSkeletalAnimationSection.cpp
+++ b/Engine/Source/Runtime/MovieSceneTracks/Private/Sections/MovieSceneSkeletalAnimationSection.cpp
@@ -450,10 +450,37 @@ int32 UMovieSceneSkeletalAnimationSection::SetBoneIndexForRootMotionCalculations
 	{ 
 		//but if not first find first
 		int32 RootIndex = INDEX_NONE;
-		for (int32 TrackIndex = 0; TrackIndex < AnimSequence->GetRawAnimationData().Num(); ++TrackIndex)
+#if WITH_EDITOR
+		const UAnimDataModel* DataModel = AnimSequence->GetDataModel();
+		const TArray<FBoneAnimationTrack>& BoneAnimationTracks = DataModel->GetBoneAnimationTracks();
+		const int32 NumTracks = BoneAnimationTracks.Num();
+		for (int32 TrackIndex = 0; TrackIndex < NumTracks; ++TrackIndex)
+		{
+			const FBoneAnimationTrack& AnimationTrack = BoneAnimationTracks[TrackIndex];
+			// verify if this bone exists in skeleton
+			const int32 BoneTreeIndex = AnimationTrack.BoneTreeIndex;
+			if (BoneTreeIndex != INDEX_NONE)
+			{
+				int32 ParentIndex = AnimSequence->GetSkeleton()->GetReferenceSkeleton().GetParentIndex(BoneTreeIndex);
+				if (ParentIndex == INDEX_NONE)
+				{
+					RootIndex = TrackIndex;
+				}
+				else if (ParentIndex == RootIndex)
+				{
+					TempRootBoneIndex = TrackIndex;
+					break;
+				}
+			}
+		}
+	#else
+		const TArray<FTrackToSkeletonMap>& BoneMappings = AnimSequence->GetCompressedTrackToSkeletonMapTable();
+		const int32 NumTracks = BoneMappings.Num();
+		for (int32 TrackIndex = 0; TrackIndex < NumTracks; ++TrackIndex)
 		{
+			const FTrackToSkeletonMap& Mapping = BoneMappings[TrackIndex];
 			// verify if this bone exists in skeleton
-			int32 BoneTreeIndex = AnimSequence->GetSkeletonIndexFromRawDataTrackIndex(TrackIndex);
+			const int32 BoneTreeIndex = Mapping.BoneTreeIndex;
 			if (BoneTreeIndex != INDEX_NONE)
 			{
 				int32 ParentIndex = AnimSequence->GetSkeleton()->GetReferenceSkeleton().GetParentIndex(BoneTreeIndex);
@@ -468,6 +495,7 @@ int32 UMovieSceneSkeletalAnimationSection::SetBoneIndexForRootMotionCalculations
 				}
 			}
 		}
+	#endif
 	}
 	return TempRootBoneIndex.GetValue();
 }
diff --git a/Engine/Source/Runtime/MovieSceneTracks/Private/Tracks/MovieSceneSkeletalAnimationTrack.cpp b/Engine/Source/Runtime/MovieSceneTracks/Private/Tracks/MovieSceneSkeletalAnimationTrack.cpp
--- a/Engine/Source/Runtime/MovieSceneTracks/Private/Tracks/MovieSceneSkeletalAnimationTrack.cpp
+++ b/Engine/Source/Runtime/MovieSceneTracks/Private/Tracks/MovieSceneSkeletalAnimationTrack.cpp
@@ -586,17 +586,22 @@ static FTransform GetWorldTransformForBone(UAnimSequence* AnimSequence, USkeleta
 	{
 		int32 BoneIndex = MeshComponent->GetBoneIndex(BoneName);
 		FTransform BoneTransform;
-		int32 TrackIndex;
+		int32 TrackIndex = INDEX_NONE;
 
-		for (TrackIndex = 0; TrackIndex < AnimSequence->GetRawTrackToSkeletonMapTable().Num(); ++TrackIndex)
+#if WITH_EDITOR
+		const UAnimDataModel* Model = AnimSequence->GetDataModel();
+		if (const FBoneAnimationTrack* TrackData = Model->FindBoneTrackByName(InBoneName))
 		{
-			if (AnimSequence->GetRawTrackToSkeletonMapTable()[TrackIndex].BoneTreeIndex == BoneIndex)
-			{
-				break;
-			}
+			TrackIndex = TrackData->BoneTreeIndex;
 		}
+#else
+		TrackIndex = AnimSequence->GetCompressedTrackToSkeletonMapTable().IndexOfByPredicate([BoneIndex](const FTrackToSkeletonMap& Mapping)
+		{
+			return Mapping.BoneTreeIndex == BoneIndex;
+		});
+#endif
 
-		if (TrackIndex == AnimSequence->GetRawTrackToSkeletonMapTable().Num())
+		if (TrackIndex == INDEX_NONE)
 		{
 			break;
 		}
