Minor AnimDataController fixes:
 - Apply curve resizing to a temp curve copy and SetCurveKeys instead of applying directly (to incorporate into undo/redo stack)
 - Limit scope of bone validation bracket to when any changes are actually mace
 - Remove RenameCurve check of current curve name (in case it is invalid)

diff --git a/Engine/Source/Runtime/Engine/Classes/Animation/AnimData/AnimDataController.h b/Engine/Source/Runtime/Engine/Classes/Animation/AnimData/AnimDataController.h
--- a/Engine/Source/Runtime/Engine/Classes/Animation/AnimData/AnimDataController.h
+++ b/Engine/Source/Runtime/Engine/Classes/Animation/AnimData/AnimDataController.h
@@ -434,7 +434,7 @@ private:
 	FString GetCurveTypeValueName(ERawCurveTrackTypes InType) const;
 	
 	/** Resizes the curve data stored on the model according to the provided new length and time at which to insert or remove time */
-	void ResizeCurves(float NewLength, float OldLength, float T0, float T1);
+	void ResizeCurves(float NewLength, bool bInserted, float T0, float T1, bool bShouldTransact = true);
 
 	/** Helper functionality to apply any changes in the model's playlength*/
 	void SetPlayLength_Internal(float NewLength, float T0, float T1, bool bShouldTransact);
diff --git a/Engine/Source/Runtime/Engine/Private/Animation/AnimDataController.cpp b/Engine/Source/Runtime/Engine/Private/Animation/AnimDataController.cpp
--- a/Engine/Source/Runtime/Engine/Private/Animation/AnimDataController.cpp
+++ b/Engine/Source/Runtime/Engine/Private/Animation/AnimDataController.cpp
@@ -170,7 +170,9 @@ void UAnimDataController::Resize(float Length, float T0, float T1, bool bShouldT
 				if (T0 < T1)
 				{
 					CONDITIONAL_TRANSACTION(LOCTEXT("ResizeModel", "Resizing Animation Data"));
+					const bool bInserted = Length > Model->PlayLength;
 					SetPlayLength_Internal(Length, T0, T1, bShouldTransact);
+					ResizeCurves(Length, bInserted, T0, T1, bShouldTransact);
 				}
 				else
 				{
@@ -342,8 +344,6 @@ bool UAnimDataController::RemoveBoneTracksMissingFromSkeleton(const USkeleton* S
 
 	if (Skeleton)
 	{
-		CONDITIONAL_BRACKET(LOCTEXT("RemoveBoneTracksMissingFromSkeleton", "Validating Bone Animation Track Data against Skeleton"));
-
 		TArray<FName> TracksToBeRemoved;
 		const FReferenceSkeleton& ReferenceSkeleton = Skeleton->GetReferenceSkeleton();
 
@@ -377,9 +377,13 @@ bool UAnimDataController::RemoveBoneTracksMissingFromSkeleton(const USkeleton* S
 			}
 		}
 
-		for (const FName& TrackName : TracksToBeRemoved)
+		if (TracksToBeRemoved.Num())
 		{
-			RemoveBoneTrack(TrackName);
+			CONDITIONAL_BRACKET(LOCTEXT("RemoveBoneTracksMissingFromSkeleton", "Validating Bone Animation Track Data against Skeleton"));
+			for (const FName& TrackName : TracksToBeRemoved)
+			{
+				RemoveBoneTrack(TrackName);
+			}
 		}
 
 		return TracksToBeRemoved.Num() > 0;
@@ -888,7 +892,7 @@ bool UAnimDataController::RenameCurve(const FAnimationCurveIdentifier& CurveToRe
 {
 	ValidateModel();
 
-	if (CurveToRenameId.IsValid() && NewCurveId.IsValid())
+	if (NewCurveId.IsValid())
 	{
 		if (CurveToRenameId != NewCurveId)
 		{
@@ -933,8 +937,7 @@ bool UAnimDataController::RenameCurve(const FAnimationCurveIdentifier& CurveToRe
 	}
 	else
 	{
-		ReportWarningf(LOCTEXT("InvalidCurveIdentiferProvidedWarning", "Invalid curve identifier(s) provide: {0} ({1}) and {2} ({3})"), FText::FromName(CurveToRenameId.InternalName.DisplayName), FText::AsNumber(CurveToRenameId.InternalName.UID),
-			FText::FromName(NewCurveId.InternalName.DisplayName), FText::AsNumber(NewCurveId.InternalName.UID));
+		ReportWarningf(LOCTEXT("InvalidCurveIdentiferProvidedWarning", "Invalid new curve identifier provided: {2} ({3})"), FText::FromName(NewCurveId.InternalName.DisplayName), FText::AsNumber(NewCurveId.InternalName.UID));
 	}
 
 	return false;
@@ -1165,8 +1168,6 @@ void UAnimDataController::SetPlayLength_Internal(float NewLength, float T0, floa
 
 	Model->NumberOfFrames = Model->FrameRate.AsFrameTime(Model->PlayLength).RoundToFrame().Value;
 	Model->NumberOfKeys = Model->NumberOfFrames + 1;
-
-	ResizeCurves(NewLength, Payload.PreviousLength, T0, T1);
 	
 	Model->Notify<FSequenceLengthChangedPayload>(EAnimDataModelNotifyType::SequenceLengthChanged, Payload);
 }
@@ -1428,18 +1429,20 @@ bool UAnimDataController::SetBoneTrackKeys(FName BoneName, const TArray<FVector>
 	return false;
 }
 
-void UAnimDataController::ResizeCurves(float NewLength, float OldLength, float T0, float T1)
+void UAnimDataController::ResizeCurves(float NewLength, bool bInserted, float T0, float T1, bool bShouldTransact /*= true*/)
 {
-	const bool bInsert = NewLength > OldLength;
-	
-	for (auto& Curve : Model->CurveData.FloatCurves)
+	CONDITIONAL_BRACKET(LOCTEXT("ResizeCurves", "Resizing all Curves"));
+
+	for (FFloatCurve& Curve : Model->CurveData.FloatCurves)
 	{
-		Curve.Resize(NewLength, bInsert, T0, T1);
+		FFloatCurve ResizedCurve = Curve;
+		ResizedCurve.Resize(NewLength, bInserted, T0, T1);
+		SetCurveKeys(FAnimationCurveIdentifier(Curve.Name, ERawCurveTrackTypes::RCT_Float), ResizedCurve.FloatCurve.GetConstRefOfKeys(), bShouldTransact);
 	}
 
-	for (auto& Curve : Model->CurveData.TransformCurves)
+	for (FTransformCurve& Curve : Model->CurveData.TransformCurves)
 	{
-		Curve.Resize(NewLength, bInsert, T0, T1);
+		Curve.Resize(NewLength, bInserted, T0, T1);
 	}
 }
 
