CIS editor tests fix: change transaction to bracket for UAnimDataController::Resize as it now contains multiple actions (through sub-calls)

diff --git a/Engine/Source/Runtime/Engine/Private/Animation/AnimDataController.cpp b/Engine/Source/Runtime/Engine/Private/Animation/AnimDataController.cpp
--- a/Engine/Source/Runtime/Engine/Private/Animation/AnimDataController.cpp
+++ b/Engine/Source/Runtime/Engine/Private/Animation/AnimDataController.cpp
@@ -169,7 +169,7 @@ void UAnimDataController::Resize(float Length, float T0, float T1, bool bShouldT
 				// Ensure that the start and end length of either removal or insertion are valid
 				if (T0 < T1)
 				{
-					CONDITIONAL_TRANSACTION(LOCTEXT("ResizeModel", "Resizing Animation Data"));
+					CONDITIONAL_BRACKET(LOCTEXT("ResizeModel", "Resizing Animation Data"));
 					const bool bInserted = Length > Model->PlayLength;
 					SetPlayLength_Internal(Length, T0, T1, bShouldTransact);
 					ResizeCurves(Length, bInserted, T0, T1, bShouldTransact);
