World Partition
- Store the actor label in the actor descriptors, to be able to show labels for unloaded actors.
- Show the actor label in the world partition editor debug display when possible.
- Always rely on the actor weak pointer to determine if an actor descriptor is loaded or not and stop calling FindObject unless we are loading.
- Removed on engine exit and GC hacks in actor descriptors references since we don't call FindObject anymore from IsLoaded.

diff --git a/Engine/Source/Editor/WorldPartitionEditor/Private/WorldPartition/SWorldPartitionEditorGrid2D.cpp b/Engine/Source/Editor/WorldPartitionEditor/Private/WorldPartition/SWorldPartitionEditorGrid2D.cpp
--- a/Engine/Source/Editor/WorldPartitionEditor/Private/WorldPartition/SWorldPartitionEditorGrid2D.cpp
+++ b/Engine/Source/Editor/WorldPartitionEditor/Private/WorldPartition/SWorldPartitionEditorGrid2D.cpp
@@ -407,23 +407,19 @@ uint32 SWorldPartitionEditorGrid2D::PaintActors(const FGeometry& AllottedGeometr
 				{
 					ActorColor = FLinearColor::Yellow;
 
-					FString ActorPath = ActorDesc->GetActorPath().ToString();
-
-					int32 LastDotPos = ActorPath.Find(TEXT("."), ESearchCase::CaseSensitive, ESearchDir::FromEnd);
-					if (LastDotPos != INDEX_NONE)
+					FName ActorLabel = ActorDesc->GetActorLabel();
+					if (!ActorLabel.IsNone())
 					{
-						ActorPath = ActorPath.Mid(LastDotPos + 1);
+						FSlateDrawElement::MakeText(
+							OutDrawElements,
+							++LayerId,
+							ActorGeometry,
+							ActorLabel.ToString(),
+							SmallLayoutFont,
+							ESlateDrawEffect::None,
+							ActorColor
+						);
 					}
-
-					FSlateDrawElement::MakeText(
-						OutDrawElements,
-						++LayerId,
-						ActorGeometry,
-						ActorPath,
-						SmallLayoutFont, 
-						ESlateDrawEffect::None,
-						ActorColor
-					);
 				}
 				else if ((SelectBox.GetVolume() > 0) && SelectBox.Intersect(ActorDesc->GetBounds()))
 				{
diff --git a/Engine/Source/Runtime/Core/Public/UObject/UE5MainStreamObjectVersion.h b/Engine/Source/Runtime/Core/Public/UObject/UE5MainStreamObjectVersion.h
--- a/Engine/Source/Runtime/Core/Public/UObject/UE5MainStreamObjectVersion.h
+++ b/Engine/Source/Runtime/Core/Public/UObject/UE5MainStreamObjectVersion.h
@@ -76,6 +76,9 @@ struct CORE_API FUE5MainStreamObjectVersion
 		// Introducing UAnimDataModel sub-object for UAnimSequenceBase containing all animation source data
 		IntroducingAnimationDataModel,
 
+		// Serialize ActorLabel in WorldPartitionActorDesc
+		WorldPartitionActorDescSerializeActorLabel,
+
 		// -----<new versions can be added above this line>-------------------------------------------------
 		VersionPlusOne,
 		LatestVersion = VersionPlusOne - 1
diff --git a/Engine/Source/Runtime/Engine/Classes/GameFramework/Actor.h b/Engine/Source/Runtime/Engine/Classes/GameFramework/Actor.h
--- a/Engine/Source/Runtime/Engine/Classes/GameFramework/Actor.h
+++ b/Engine/Source/Runtime/Engine/Classes/GameFramework/Actor.h
@@ -2087,7 +2087,7 @@ public:
 
 	/** Returns this actor's current label.  Actor labels are only available in development builds. */
 	UFUNCTION(BlueprintCallable, Category = "Editor Scripting | Actor Editing", meta = (KeyWords = "Display Name"))
-	const FString& GetActorLabel() const;
+	const FString& GetActorLabel(bool bCreateIfNone = true) const;
 
 	/**
 	 * Assigns a new label to this actor.  Actor labels are only available in development builds.
diff --git a/Engine/Source/Runtime/Engine/Private/ActorEditor.cpp b/Engine/Source/Runtime/Engine/Private/ActorEditor.cpp
--- a/Engine/Source/Runtime/Engine/Private/ActorEditor.cpp
+++ b/Engine/Source/Runtime/Engine/Private/ActorEditor.cpp
@@ -911,7 +911,7 @@ TUniquePtr<class FWorldPartitionActorDesc> AActor::CreateClassActorDesc(const TS
 	return CastChecked<AActor>(ActorClass->GetDefaultObject())->CreateClassActorDesc();
 }
 
-const FString& AActor::GetActorLabel() const
+const FString& AActor::GetActorLabel(bool bCreateIfNone) const
 {
 	// If the label string is empty then we'll use the default actor label (usually the actor's class name.)
 	// We actually cache the default name into our ActorLabel property.  This will be saved out with the
@@ -923,7 +923,7 @@ const FString& AActor::GetActorLabel() const
 	//
 	// Remember, ActorLabel is currently an editor-only property.
 
-	if( ActorLabel.IsEmpty() )
+	if( ActorLabel.IsEmpty() && bCreateIfNone )
 	{
 		// Get the class
 		UClass* ActorClass = GetClass();
diff --git a/Engine/Source/Runtime/Engine/Private/WorldPartition/WorldPartitionActorDesc.cpp b/Engine/Source/Runtime/Engine/Private/WorldPartition/WorldPartitionActorDesc.cpp
--- a/Engine/Source/Runtime/Engine/Private/WorldPartition/WorldPartitionActorDesc.cpp
+++ b/Engine/Source/Runtime/Engine/Private/WorldPartition/WorldPartitionActorDesc.cpp
@@ -77,6 +77,8 @@ void FWorldPartitionActorDesc::Init(const AActor* InActor)
 		}
 	}
 
+	ActorLabel = *InActor->GetActorLabel(false);
+
 	WorldPartition = InActor->GetLevel()->GetWorldPartition();
 	check(WorldPartition);
 }
@@ -107,6 +109,8 @@ void FWorldPartitionActorDesc::Init(UWorldPartition* InWorldPartition, const FWo
 	}
 
 	WorldPartition = InWorldPartition;
+
+	ActorPtr = FindObject<AActor>(nullptr, *ActorPath.ToString());
 }
 
 void FWorldPartitionActorDesc::SerializeTo(TArray<uint8>& OutData)
@@ -187,6 +191,11 @@ void FWorldPartitionActorDesc::Serialize(FArchive& Ar)
 	{
 		Ar << DataLayers;
 	}
+
+	if (Ar.CustomVer(FUE5MainStreamObjectVersion::GUID) >= FUE5MainStreamObjectVersion::WorldPartitionActorDescSerializeActorLabel)
+	{
+		Ar << ActorLabel;
+	}
 }
 
 FBox FWorldPartitionActorDesc::GetBounds() const
@@ -194,36 +203,34 @@ FBox FWorldPartitionActorDesc::GetBounds() const
 	return FBox(BoundsLocation - BoundsExtent, BoundsLocation + BoundsExtent);
 }
 
-AActor* FWorldPartitionActorDesc::GetActor() const
+bool FWorldPartitionActorDesc::IsLoaded() const
 {
-	AActor* Actor = ActorPtr.Get();
-
-	if (Actor)
-	{
-		return Actor;
-	}
-
-	if (!IsGarbageCollecting())
-	{
-		Actor = FindObject<AActor>(nullptr, *ActorPath.ToString());
-	}
-
-	if (Actor)
-	{
-		ActorPtr = Actor;
-	}
+	return ActorPtr.IsValid();
+}
 
-	return Actor;
+AActor* FWorldPartitionActorDesc::GetActor() const
+{
+	return ActorPtr.Get();
 }
 
 AActor* FWorldPartitionActorDesc::Load() const
 {
-	AActor* Actor = GetActor();
+	if (ActorPtr.IsExplicitlyNull())
+	{
+		// First, try to find the existing actor which could have been loaded by another actor (through standard serialization)
+		ActorPtr = FindObject<AActor>(nullptr, *ActorPath.ToString());
+	}
 
-	if (!Actor)
+	// The, if the actor isn't loaded, load it
+	if (ActorPtr.IsExplicitlyNull())
 	{
 		check(WorldPartition || GIsAutomationTesting);
-		const FLinkerInstancingContext* InstancingContext = WorldPartition ? &WorldPartition->InstancingContext : nullptr;
+
+		const FLinkerInstancingContext* InstancingContext = nullptr;
+		if (WorldPartition && WorldPartition->InstancingContext.IsInstanced())
+		{
+			InstancingContext = &WorldPartition->InstancingContext;
+		}
 
 		UPackage* Package = nullptr;
 
@@ -239,12 +246,12 @@ AActor* FWorldPartitionActorDesc::Load() const
 
 		if (Package)
 		{
-			Actor = GetActor();
-			check(Actor);
+			ActorPtr = FindObject<AActor>(nullptr, *ActorPath.ToString());
+			check(ActorPtr.IsValid());
 		}
 	}
 
-	return Actor;
+	return ActorPtr.Get();
 }
 
 void FWorldPartitionActorDesc::Unload()
diff --git a/Engine/Source/Runtime/Engine/Private/WorldPartition/WorldPartitionHandle.cpp b/Engine/Source/Runtime/Engine/Private/WorldPartition/WorldPartitionHandle.cpp
--- a/Engine/Source/Runtime/Engine/Private/WorldPartition/WorldPartitionHandle.cpp
+++ b/Engine/Source/Runtime/Engine/Private/WorldPartition/WorldPartitionHandle.cpp
@@ -24,7 +24,7 @@ bool FWorldPartitionHandleUtils::IsActorDescLoaded(FWorldPartitionActorDesc* Act
 	}
 #endif
 
-	return !!ActorDesc->GetActor();
+	return ActorDesc->IsLoaded();
 }
 
 void FWorldPartitionHandleImpl::IncRefCount(FWorldPartitionActorDesc* ActorDesc)
@@ -43,7 +43,6 @@ void FWorldPartitionReferenceImpl::IncRefCount(FWorldPartitionActorDesc* ActorDe
 	{
 		AActor* Actor = ActorDesc->Load();
 		check(Actor || GIsAutomationTesting);
-		check(!IsEngineExitRequested());
 
 		if (Actor)
 		{
@@ -56,14 +55,12 @@ void FWorldPartitionReferenceImpl::DecRefCount(FWorldPartitionActorDesc* ActorDe
 {
 	if (!ActorDesc->DecHardRefCount())
 	{
-		if (!IsEngineExitRequested())
+		// The only case where an actor can be unloaded while still holding a reference 
+		// is when it was manually deleted in the editor.
+		if (ActorDesc->IsLoaded())
 		{
-			// We can still hold a reference to an actor that was manually deleted in the editor
-			if (AActor* Actor = ActorDesc->GetActor())
-			{
-				ActorDesc->UnregisterActor();
-				ActorDesc->Unload();
-			}
+			ActorDesc->UnregisterActor();
+			ActorDesc->Unload();
 		}
 	}
 }
diff --git a/Engine/Source/Runtime/Engine/Public/WorldPartition/WorldPartitionActorDesc.h b/Engine/Source/Runtime/Engine/Public/WorldPartition/WorldPartitionActorDesc.h
--- a/Engine/Source/Runtime/Engine/Public/WorldPartition/WorldPartitionActorDesc.h
+++ b/Engine/Source/Runtime/Engine/Public/WorldPartition/WorldPartitionActorDesc.h
@@ -58,6 +58,7 @@ public:
 	inline const TArray<FName>& GetDataLayers() const { return DataLayers; }
 	inline FName GetActorPackage() const { return ActorPackage; }
 	inline FName GetActorPath() const { return ActorPath; }
+	inline FName GetActorLabel() const { return ActorLabel; }
 	FBox GetBounds() const;
 
 protected:
@@ -101,6 +102,7 @@ public:
 
 	FString ToString() const;
 
+	bool IsLoaded() const;
 	AActor* GetActor() const;
 	AActor* Load() const;
 	virtual void Unload();
@@ -135,6 +137,7 @@ protected:
 	FName							Class;
 	FName							ActorPackage;
 	FName							ActorPath;
+	FName							ActorLabel;
 	FVector							BoundsLocation;
 	FVector							BoundsExtent;
 	EActorGridPlacement				GridPlacement;
