Change "Partition" category to "World Partition".

diff --git a/Engine/Source/Runtime/Engine/Classes/GameFramework/Actor.h b/Engine/Source/Runtime/Engine/Classes/GameFramework/Actor.h
--- a/Engine/Source/Runtime/Engine/Classes/GameFramework/Actor.h
+++ b/Engine/Source/Runtime/Engine/Classes/GameFramework/Actor.h
@@ -535,14 +535,14 @@ public:
 	 * Determine how this actor will be placed in the partition (if the world is partitioned). This value is not taken into account
 	 * if the function ActorClass->GetDefaultGridPlacement returns other than EActorGridPlacement::None.
 	 */
-	UPROPERTY(EditAnywhere, Category=Partition)
+	UPROPERTY(EditAnywhere, Category=WorldPartition)
 	EActorGridPlacement GridPlacement;
 
 	/** 
 	 * Determine in which partition grid this actor will be placed in the partition (if the world is partitioned).
 	 * If None, the decision will be left to the partition.
 	 */
-	UPROPERTY(EditAnywhere, Category=Partition)
+	UPROPERTY(EditAnywhere, Category=WorldPartition)
 	FName RuntimeGrid;
 
 private:
