Fixed issue where adding SMC to actor in an editing level instance would cause the new mesh to be greyed out.

diff --git a/Engine/Source/Runtime/Engine/Classes/Components/PrimitiveComponent.h b/Engine/Source/Runtime/Engine/Classes/Components/PrimitiveComponent.h
--- a/Engine/Source/Runtime/Engine/Classes/Components/PrimitiveComponent.h
+++ b/Engine/Source/Runtime/Engine/Classes/Components/PrimitiveComponent.h
@@ -367,10 +367,6 @@ public:
 	UPROPERTY()
 	uint8 bSelectable:1;
 
-	/** If this is True, the component may not be greyed if the level instance visualization effect is enabled */
-	UPROPERTY(Transient)
-	uint8 bEditingLevelInstanceState:1;
-
 	/** If true, forces mips for textures used by this component to be resident when this component's level is loaded. */
 	UPROPERTY(EditAnywhere, AdvancedDisplay, BlueprintReadOnly, Category=TextureStreaming)
 	uint8 bForceMipStreaming:1;
@@ -1601,6 +1597,9 @@ public:
 	/** Return true if the owner is selected and this component is selectable */
 	virtual bool ShouldRenderSelected() const;
 
+	/** Returns true if the owning actor is part of a level instance which is being edited. */
+	bool GetLevelInstanceEditingState() const;
+
 	/** Component is directly selected in the editor separate from its parent actor */
 	bool IsComponentIndividuallySelected() const;
 
@@ -2024,7 +2023,7 @@ public:
 	void PushSelectionToProxy();
 
 	/**
-	 * Pushes new LevelInstance editing state to the render thread primitive proxy and caches it locally
+	 * Pushes new LevelInstance editing state to the render thread primitive proxy.
 	 */
 	void PushLevelInstanceEditingStateToProxy(bool bInEditingState);
 
diff --git a/Engine/Source/Runtime/Engine/Classes/GameFramework/Actor.h b/Engine/Source/Runtime/Engine/Classes/GameFramework/Actor.h
--- a/Engine/Source/Runtime/Engine/Classes/GameFramework/Actor.h
+++ b/Engine/Source/Runtime/Engine/Classes/GameFramework/Actor.h
@@ -217,6 +217,13 @@ private:
 	UPROPERTY()
 	uint8 bForceNetAddressable:1;
 
+#if WITH_EDITORONLY_DATA
+	/**
+	 * Whether this actor belongs to a level instance which is currently being edited.
+	 */
+	UPROPERTY(Transient)
+	uint8 bIsInEditingLevelInstance:1;
+#endif
 public:
 
 	/** If true, this actor is no longer replicated to new clients, and is "torn off" (becomes a ROLE_Authority) on clients to which it was being replicated. */
@@ -225,6 +232,14 @@ public:
 		return bTearOff;
 	}
 
+#if WITH_EDITOR
+	/** If true, the actor belongs to a level instance which is currently being edited. */
+	virtual bool IsInEditingLevelInstance() const
+	{
+		return bIsInEditingLevelInstance;
+	}
+#endif
+
 	/** Networking - Server - TearOff this actor to stop replication to clients. Will set bTearOff to true. */
 	UFUNCTION(BlueprintCallable, Category=Replication)
 	virtual void TearOff();
@@ -2519,8 +2534,10 @@ public:
 	/** Push Selection to actor */
 	virtual void PushSelectionToProxies();
 
+#if WITH_EDITOR
 	/** Push Foundation Editing state to primitive scene proxy */
 	virtual void PushLevelInstanceEditingStateToProxies(bool bInEditingState);
+#endif
 
 	/** 
 	 * Returns a list of all actors spawned by our Child Actor Components, including children of children. 
diff --git a/Engine/Source/Runtime/Engine/Private/Actor.cpp b/Engine/Source/Runtime/Engine/Private/Actor.cpp
--- a/Engine/Source/Runtime/Engine/Private/Actor.cpp
+++ b/Engine/Source/Runtime/Engine/Private/Actor.cpp
@@ -143,7 +143,7 @@ void AActor::InitializeDefaults()
 	bAllowReceiveTickEventOnDedicatedServer = true;
 	bRelevantForNetworkReplays = true;
 	bRelevantForLevelBounds = true;
-	
+
 	// Overlap collision settings
 	bGenerateOverlapEventsDuringLevelStreaming = false;
 	UpdateOverlapsMethodDuringLevelStreaming = EActorUpdateOverlapsMethod::UseConfigDefault;
@@ -151,6 +151,7 @@ void AActor::InitializeDefaults()
 	
 	bHasDeferredComponentRegistration = false;
 #if WITH_EDITORONLY_DATA
+	bIsInEditingLevelInstance = false;
 	PivotOffset = FVector::ZeroVector;
 #endif
 	SpawnCollisionHandlingMethod = ESpawnActorCollisionHandlingMethod::AlwaysSpawn;
@@ -4667,16 +4668,23 @@ void AActor::PushSelectionToProxies()
 	}
 }
 
+#if WITH_EDITOR
 void AActor::PushLevelInstanceEditingStateToProxies(bool bInEditingState)
 {
 	TInlineComponentArray<UPrimitiveComponent*> PrimComponents;
 	GetComponents(PrimComponents);
 
+
+	bIsInEditingLevelInstance = bInEditingState;
+
+	// We call the virtual method here to allow subclasses to extend the logic for determining level instance state.
+	const bool bPushEditingLevelInstanceState = IsInEditingLevelInstance();
+
 	for (const auto& PrimComponent : PrimComponents)
 	{
 		if (PrimComponent->IsRegistered())
 		{
-			PrimComponent->PushLevelInstanceEditingStateToProxy(bInEditingState);
+			PrimComponent->PushLevelInstanceEditingStateToProxy(bPushEditingLevelInstanceState);
 		}
 	}
 
@@ -4688,6 +4696,7 @@ void AActor::PushLevelInstanceEditingStateToProxies(bool bInEditingState)
 
 	ForEachAttachedActors(PushAllChildrenToProxies);
 }
+#endif
 
 bool AActor::IsChildActor() const
 {
diff --git a/Engine/Source/Runtime/Engine/Private/Components/PrimitiveComponent.cpp b/Engine/Source/Runtime/Engine/Private/Components/PrimitiveComponent.cpp
--- a/Engine/Source/Runtime/Engine/Private/Components/PrimitiveComponent.cpp
+++ b/Engine/Source/Runtime/Engine/Private/Components/PrimitiveComponent.cpp
@@ -331,7 +331,6 @@ UPrimitiveComponent::UPrimitiveComponent(const FObjectInitializer& ObjectInitial
 	bCastVolumetricTranslucentShadow = false;
 	bCastContactShadow = true;
 	IndirectLightingCacheQuality = ILCQ_Point;
-	bEditingLevelInstanceState = false;
 	bSelectable = true;
 	bFillCollisionUnderneathForNavmesh = false;
 	AlwaysLoadOnClient = true;
@@ -1453,6 +1452,18 @@ bool UPrimitiveComponent::ShouldRenderSelected() const
 	return false;
 }
 
+bool UPrimitiveComponent::GetLevelInstanceEditingState() const
+{
+#if WITH_EDITOR
+	if (const AActor* Owner = GetOwner())
+	{
+		return Owner->IsInEditingLevelInstance();
+	}
+#endif
+
+	return false;
+}
+
 void UPrimitiveComponent::SetCastShadow(bool NewCastShadow)
 {
 	if(NewCastShadow != CastShadow)
@@ -1548,10 +1559,9 @@ void UPrimitiveComponent::PushSelectionToProxy()
 void UPrimitiveComponent::PushLevelInstanceEditingStateToProxy(bool bInEditingState)
 {
 	//although this should only be called for attached components, some billboard components can get in without valid proxies
-	bEditingLevelInstanceState = bInEditingState;
 	if (SceneProxy)
 	{
-		SceneProxy->SetLevelInstanceEditingState_GameThread(bEditingLevelInstanceState);
+		SceneProxy->SetLevelInstanceEditingState_GameThread(bInEditingState);
 	}
 }
 
diff --git a/Engine/Source/Runtime/Engine/Private/PrimitiveSceneProxy.cpp b/Engine/Source/Runtime/Engine/Private/PrimitiveSceneProxy.cpp
--- a/Engine/Source/Runtime/Engine/Private/PrimitiveSceneProxy.cpp
+++ b/Engine/Source/Runtime/Engine/Private/PrimitiveSceneProxy.cpp
@@ -116,7 +116,7 @@ FPrimitiveSceneProxy::FPrimitiveSceneProxy(const UPrimitiveComponent* InComponen
 ,	bOwnerNoSee(InComponent->bOwnerNoSee)
 ,	bParentSelected(InComponent->ShouldRenderSelected())
 ,	bIndividuallySelected(InComponent->IsComponentIndividuallySelected())
-,	bLevelInstanceEditingState(InComponent->bEditingLevelInstanceState)
+,	bLevelInstanceEditingState(InComponent->GetLevelInstanceEditingState())
 ,	bHovered(false)
 ,	bUseViewOwnerDepthPriorityGroup(InComponent->bUseViewOwnerDepthPriorityGroup)
 ,	bHasMotionBlurVelocityMeshes(InComponent->bHasMotionBlurVelocityMeshes)
