Fix actor descriptors serialization archive not being marked as persistent.

diff --git a/Engine/Source/Runtime/Core/Public/UObject/UE5MainStreamObjectVersion.h b/Engine/Source/Runtime/Core/Public/UObject/UE5MainStreamObjectVersion.h
--- a/Engine/Source/Runtime/Core/Public/UObject/UE5MainStreamObjectVersion.h
+++ b/Engine/Source/Runtime/Core/Public/UObject/UE5MainStreamObjectVersion.h
@@ -79,6 +79,9 @@ struct CORE_API FUE5MainStreamObjectVersion
 		// Serialize ActorLabel in WorldPartitionActorDesc
 		WorldPartitionActorDescSerializeActorLabel,
 
+		// Fix WorldPartitionActorDesc serialization archive not persistent
+		WorldPartitionActorDescSerializeArchivePersistent,
+
 		// -----<new versions can be added above this line>-------------------------------------------------
 		VersionPlusOne,
 		LatestVersion = VersionPlusOne - 1
diff --git a/Engine/Source/Runtime/Engine/Private/WorldPartition/WorldPartitionActorDesc.cpp b/Engine/Source/Runtime/Engine/Private/WorldPartition/WorldPartitionActorDesc.cpp
--- a/Engine/Source/Runtime/Engine/Private/WorldPartition/WorldPartitionActorDesc.cpp
+++ b/Engine/Source/Runtime/Engine/Private/WorldPartition/WorldPartitionActorDesc.cpp
@@ -91,7 +91,7 @@ void FWorldPartitionActorDesc::Init(UWorldPartition* InWorldPartition, const FWo
 	Class = DescData.NativeClass->GetFName();
 
 	// Serialize actor metadata
-	FMemoryReader MetadataAr(DescData.SerializedData);
+	FMemoryReader MetadataAr(DescData.SerializedData, true);
 
 	// Serialize metadata custom versions
 	FCustomVersionContainer CustomVersions;
@@ -117,7 +117,7 @@ void FWorldPartitionActorDesc::SerializeTo(TArray<uint8>& OutData)
 {
 	// Serialize to archive and gather custom versions
 	TArray<uint8> PayloadData;
-	FMemoryWriter PayloadAr(PayloadData);
+	FMemoryWriter PayloadAr(PayloadData, true);
 	Serialize(PayloadAr);
 
 	// Serialize custom versions
@@ -170,6 +170,8 @@ FString FWorldPartitionActorDesc::ToString() const
 
 void FWorldPartitionActorDesc::Serialize(FArchive& Ar)
 {
+	check(Ar.IsPersistent());
+
 	Ar.UsingCustomVersion(FUE5MainStreamObjectVersion::GUID);
 
 	Ar << Class << Guid << BoundsLocation << BoundsExtent << GridPlacement << RuntimeGrid << bActorIsEditorOnly << bLevelBoundsRelevant;
@@ -182,7 +184,7 @@ void FWorldPartitionActorDesc::Serialize(FArchive& Ar)
 
 	Ar << References;
 
-	if (!Ar.IsPersistent())
+	if (Ar.CustomVer(FUE5MainStreamObjectVersion::GUID) < FUE5MainStreamObjectVersion::WorldPartitionActorDescSerializeArchivePersistent)
 	{
 		Ar << ActorPackage << ActorPath;
 	}
