Added world partition settings 'MinimumWorldSize' to control the minimum world size.
When loading a partitioned world smaller that this value, we automatically load all cells.

diff --git a/Engine/Source/Editor/WorldPartitionEditor/Private/WorldPartitionEditorModule.cpp b/Engine/Source/Editor/WorldPartitionEditor/Private/WorldPartitionEditorModule.cpp
--- a/Engine/Source/Editor/WorldPartitionEditor/Private/WorldPartitionEditorModule.cpp
+++ b/Engine/Source/Editor/WorldPartitionEditor/Private/WorldPartitionEditorModule.cpp
@@ -141,6 +141,11 @@ void FWorldPartitionEditorModule::SetConversionPromptEnabled(bool bEnabled)
 	GetMutableDefault<UWorldPartitionEditorSettings>()->bEnableConversionPrompt = bEnabled;
 }
 
+float FWorldPartitionEditorModule::GetMinimumWorldSize() const
+{
+	return GetDefault<UWorldPartitionEditorSettings>()->MinimumWorldSize;
+}
+
 bool FWorldPartitionEditorModule::ConvertMap(const FString& InLongPackageName)
 {
 	if (!GetDefault<UWorldPartitionEditorSettings>()->bEnableConversionPrompt || ULevel::GetIsLevelPartitionedFromPackage(FName(*InLongPackageName)))
diff --git a/Engine/Source/Editor/WorldPartitionEditor/Public/WorldPartition/WorldPartitionEditorSettings.h b/Engine/Source/Editor/WorldPartitionEditor/Public/WorldPartition/WorldPartitionEditorSettings.h
--- a/Engine/Source/Editor/WorldPartitionEditor/Public/WorldPartition/WorldPartitionEditorSettings.h
+++ b/Engine/Source/Editor/WorldPartitionEditor/Public/WorldPartition/WorldPartitionEditorSettings.h
@@ -27,4 +27,7 @@ public:
 
 	UPROPERTY(config, EditAnywhere, Category = MapConversion, meta = (ToolTip = "Commandlet class to use for World Parition conversion"))
 	TSubclassOf<UWorldPartitionConvertCommandlet> CommandletClass;
+
+	UPROPERTY(config, EditAnywhere, Category = WorldPartition, meta = (ToolTip = "This option is used to automatically load all cells when the world is smaller than this value"))
+	float MinimumWorldSize = 100000.0f;
 };
\ No newline at end of file
diff --git a/Engine/Source/Editor/WorldPartitionEditor/Public/WorldPartitionEditorModule.h b/Engine/Source/Editor/WorldPartitionEditor/Public/WorldPartitionEditorModule.h
--- a/Engine/Source/Editor/WorldPartitionEditor/Public/WorldPartitionEditorModule.h
+++ b/Engine/Source/Editor/WorldPartitionEditor/Public/WorldPartitionEditorModule.h
@@ -40,6 +40,11 @@ public:
 	 */
 	virtual void SetConversionPromptEnabled(bool bEnabled) override;
 
+	/**
+	 *
+	 */
+	virtual float GetMinimumWorldSize() const override;
+
 	/**
 	 * 
 	 */
diff --git a/Engine/Source/Runtime/Engine/Private/WorldPartition/WorldPartition.cpp b/Engine/Source/Runtime/Engine/Private/WorldPartition/WorldPartition.cpp
--- a/Engine/Source/Runtime/Engine/Private/WorldPartition/WorldPartition.cpp
+++ b/Engine/Source/Runtime/Engine/Private/WorldPartition/WorldPartition.cpp
@@ -27,6 +27,7 @@
 #include "ScopedTransaction.h"
 #include "UnrealEdMisc.h"
 #include "AssetRegistryModule.h"
+#include "WorldPartition/IWorldPartitionEditorModule.h"
 #include "WorldPartition/WorldPartitionLevelStreamingDynamic.h"
 #include "WorldPartition/WorldPartitionEditorHash.h"
 #include "WorldPartition/WorldPartitionEditorSpatialHash.h"
@@ -261,8 +262,15 @@ void UWorldPartition::Initialize(UWorld* InWorld, const FTransform& InTransform)
 			// Load the always loaded cell, don't call LoadCells to avoid creating a transaction
 			UpdateLoadingEditorCell(EditorHash->GetAlwaysLoadedCell(), true);
 
+			// If the world is smaller than the project setting's minimum world size, load all cells
+			IWorldPartitionEditorModule& WorldPartitionEditorModule = FModuleManager::LoadModuleChecked<IWorldPartitionEditorModule>("WorldPartitionEditor");
+			const float MinimumWorldSize = WorldPartitionEditorModule.GetMinimumWorldSize();
+			const bool bWorldIsSmallerThanMinimumWorldSize = GetWorldBounds().GetVolume() <= (MinimumWorldSize * MinimumWorldSize * MinimumWorldSize);
+			
 			// When loading a subworld, load all actors
-			if (!IsMainWorldPartition())
+			const bool bWorldIsSubPartition = !IsMainWorldPartition();
+
+			if (bWorldIsSubPartition || bWorldIsSmallerThanMinimumWorldSize)
 			{
 				EditorHash->ForEachCell([this](UWorldPartitionEditorCell* Cell)
 				{
diff --git a/Engine/Source/Runtime/Engine/Public/WorldPartition/IWorldPartitionEditorModule.h b/Engine/Source/Runtime/Engine/Public/WorldPartition/IWorldPartitionEditorModule.h
--- a/Engine/Source/Runtime/Engine/Public/WorldPartition/IWorldPartitionEditorModule.h
+++ b/Engine/Source/Runtime/Engine/Public/WorldPartition/IWorldPartitionEditorModule.h
@@ -15,5 +15,6 @@ public:
 	virtual bool IsWorldPartitionEnabled() = 0;
 	virtual bool IsConversionPromptEnabled() const = 0;
 	virtual void SetConversionPromptEnabled(bool bEnabled) = 0;
+	virtual float GetMinimumWorldSize() const = 0;
 	virtual bool ConvertMap(const FString& InLongPackageName) = 0;
 };
\ No newline at end of file
