Hide DataLayers members when not in a partitioned world.

diff --git a/Engine/Source/Runtime/Engine/Private/ActorEditor.cpp b/Engine/Source/Runtime/Engine/Private/ActorEditor.cpp
--- a/Engine/Source/Runtime/Engine/Private/ActorEditor.cpp
+++ b/Engine/Source/Runtime/Engine/Private/ActorEditor.cpp
@@ -58,7 +58,8 @@ bool AActor::CanEditChange(const FProperty* PropertyThatWillChange) const
 	}
 
 	if ((PropertyThatWillChange->GetFName() == GET_MEMBER_NAME_CHECKED(AActor, GridPlacement)) || 
-		(PropertyThatWillChange->GetFName() == GET_MEMBER_NAME_CHECKED(AActor, RuntimeGrid)))
+		(PropertyThatWillChange->GetFName() == GET_MEMBER_NAME_CHECKED(AActor, RuntimeGrid)) ||
+		(PropertyThatWillChange->GetFName() == GET_MEMBER_NAME_CHECKED(AActor, DataLayers)))
 	{
 		if (!IsTemplate())
 		{
