Added basic starta material conversion for default lit, hair, volumetric, cloud and water. But water has issue due to missing connection on input parameters.
Tested on a simple level + material.

Renamed Emissve to EmissiveColor on UMaterialExpressionStrataUnlitBSDFfor concistency
Roughness on slab BSDF now defaults to 0.5 as the default material node.

diff --git a/Engine/Source/Runtime/Engine/Classes/Materials/Material.h b/Engine/Source/Runtime/Engine/Classes/Materials/Material.h
--- a/Engine/Source/Runtime/Engine/Classes/Materials/Material.h
+++ b/Engine/Source/Runtime/Engine/Classes/Materials/Material.h
@@ -1198,6 +1198,7 @@ private:
 
 	void BackwardsCompatibilityInputConversion();
 	void BackwardsCompatibilityVirtualTextureOutputConversion();
+	void ConvertMaterialToStrataMaterial();
 
 	/** Handles setting up an annotation for this object if a flag has changed value */
 	void MarkUsageFlagDirty(EMaterialUsage Usage, bool CurrentValue, bool NewValue);
diff --git a/Engine/Source/Runtime/Engine/Classes/Materials/MaterialExpressionStrata.h b/Engine/Source/Runtime/Engine/Classes/Materials/MaterialExpressionStrata.h
--- a/Engine/Source/Runtime/Engine/Classes/Materials/MaterialExpressionStrata.h
+++ b/Engine/Source/Runtime/Engine/Classes/Materials/MaterialExpressionStrata.h
@@ -61,7 +61,7 @@ class UMaterialExpressionStrataSlabBSDF : public UMaterialExpression
 	// STRATA_TODO: edge or F82 EdgeColor?
 
 	/**
-	 * Controls how rough the Material is. Roughness of 0 (smooth) is a mirror reflection and 1 (rough) is completely matte or diffuse. When using anisotropy, it is the roughness used along the Tangent axis. (type = float, unit = unitless, defaults to 0)
+	 * Controls how rough the Material is. Roughness of 0 (smooth) is a mirror reflection and 1 (rough) is completely matte or diffuse. When using anisotropy, it is the roughness used along the Tangent axis. (type = float, unit = unitless, defaults to 0.5)
 	 */
 	UPROPERTY()
 	FExpressionInput RoughnessX;
@@ -204,7 +204,7 @@ class UMaterialExpressionStrataUnlitBSDF : public UMaterialExpression
 	* Emissive (type = float3, unit = Luminance, default = 0)
 	*/
 	UPROPERTY()
-	FExpressionInput Emissive;
+	FExpressionInput EmissiveColor;
 
 	/**
 	 * The amount of transmitted light from the back side of the surface to the front side of the surface ()
diff --git a/Engine/Source/Runtime/Engine/Private/Materials/HLSLMaterialTranslator.cpp b/Engine/Source/Runtime/Engine/Private/Materials/HLSLMaterialTranslator.cpp
--- a/Engine/Source/Runtime/Engine/Private/Materials/HLSLMaterialTranslator.cpp
+++ b/Engine/Source/Runtime/Engine/Private/Materials/HLSLMaterialTranslator.cpp
@@ -914,7 +914,9 @@ bool FHLSLMaterialTranslator::Translate()
 			{
 				Errorf(TEXT("SingleLayerWater materials cannot be combined with other shading models.")); // Simply untested for now
 			}
-			if (Material->GetMaterialInterface() && !Material->GetMaterialInterface()->GetMaterial()->HasAnyExpressionsInMaterialAndFunctionsOfType<UMaterialExpressionSingleLayerWaterMaterialOutput>())
+			if (Material->GetMaterialInterface() && !Material->GetMaterialInterface()->GetMaterial()->HasAnyExpressionsInMaterialAndFunctionsOfType<UMaterialExpressionSingleLayerWaterMaterialOutput>()
+				&& CONVERT_MATERIAL_TO_STRATA_ON_LOAD == 0
+				)
 			{
 				Errorf(TEXT("SingleLayerWater materials requires the use of SingleLayerWaterMaterial output node."));
 			}
diff --git a/Engine/Source/Runtime/Engine/Private/Materials/Material.cpp b/Engine/Source/Runtime/Engine/Private/Materials/Material.cpp
--- a/Engine/Source/Runtime/Engine/Private/Materials/Material.cpp
+++ b/Engine/Source/Runtime/Engine/Private/Materials/Material.cpp
@@ -72,6 +72,12 @@
 #include "RenderCore/Public/RenderUtils.h"
 #include "Materials/MaterialStaticParameterValueResolver.h"
 #include "Engine/Font.h"
+#include "UObject/UE5MainStreamObjectVersion.h"
+#include "Materials/MaterialExpressionAdd.h"
+#include "Materials/MaterialExpressionDivide.h"
+#include "Materials/MaterialExpressionSingleLayerWaterMaterialOutput.h"
+#include "Materials/MaterialExpressionStrata.h"
+#include "Materials/StrataMaterial.h"
 
 #if WITH_EDITOR
 #include "Logging/TokenizedMessage.h"
@@ -3586,6 +3592,149 @@ void UMaterial::GetQualityLevelUsage(TArray<bool, TInlineAllocator<EMaterialQual
 	}
 }
 
+void UMaterial::ConvertMaterialToStrataMaterial()
+{
+#if WITH_EDITOR && CONVERT_MATERIAL_TO_STRATA_ON_LOAD
+	const URendererSettings* RendererSettings = GetDefault<URendererSettings>();
+	const bool bStrataEnabled = RendererSettings && RendererSettings->bEnableStrata;
+
+	if (bStrataEnabled && !FrontMaterial.IsConnected())
+	{
+		auto MoveConnectionTo = [](auto& OldNodeInput, UMaterialExpression* NewNode, uint32 NewInputIndex)
+		{
+			if (OldNodeInput.IsConnected())
+			{
+				NewNode->GetInput(NewInputIndex)->Connect(OldNodeInput.OutputIndex, OldNodeInput.Expression);
+				OldNodeInput.Expression = nullptr;
+			}
+		};
+		auto CopyConnectionTo = [](auto& OldNodeInput, UMaterialExpression* NewNode, uint32 NewInputIndex)
+		{
+			if (OldNodeInput.IsConnected())
+			{
+				NewNode->GetInput(NewInputIndex)->Connect(OldNodeInput.OutputIndex, OldNodeInput.Expression);
+			}
+		};
+
+		if (MaterialDomain == MD_Surface && ShadingModel == MSM_Unlit)
+		{
+			UMaterialExpressionStrataUnlitBSDF* UnlitBSDF = NewObject<UMaterialExpressionStrataUnlitBSDF>(this);
+			MoveConnectionTo(EmissiveColor, UnlitBSDF, 0);	// EmissiveColor
+			FrontMaterial.Connect(0, UnlitBSDF);
+		}
+		else if (MaterialDomain == MD_Surface && ShadingModel == MSM_DefaultLit)
+		{
+			UMaterialExpressionStrataSlabBSDF* SlabBSDF = NewObject<UMaterialExpressionStrataSlabBSDF>(this);
+			MoveConnectionTo(BaseColor, SlabBSDF, 0);		// BaseColor
+			MoveConnectionTo(Metallic, SlabBSDF, 1);		// Metallic
+			MoveConnectionTo(Specular, SlabBSDF, 2);		// Specular
+			MoveConnectionTo(Roughness, SlabBSDF, 3);		// RoughnessX
+															// STRATA_TODO RoughnessY
+			MoveConnectionTo(Normal, SlabBSDF, 5);			// Normal
+			MoveConnectionTo(Tangent, SlabBSDF, 6);			// Tangent
+			// Opacity can remain on the end point node
+			// Opacity mask can remain on the end point node
+			// WorldPositionOffset can remain on the end point node
+			// STRATA_TODO WorldDisplacement and TessellationMultiplier are going to be removed and should be ignored?
+			// STRATA_TODO ShadingModel
+			// STRATA_TODO SubsurfaceColor
+			// STRATA_TODO ClearCoat
+			// STRATA_TODO ClearCoatRoughness
+			// STRATA_TODO AmbientOcclusion should be removed and put on the BSDF node
+			// STRATA_TODO Refraction
+			// STRATA_TODO PixelDepthOffset
+			// STRATA_TODO Anisotropy
+			// STRATA_TODO EmissiveColor
+
+			FrontMaterial.Connect(0, SlabBSDF);
+		}
+		else if (MaterialDomain == MD_Surface && ShadingModel == MSM_Hair)
+		{
+			UMaterialExpressionStrataHairBSDF* HairBSDF = NewObject<UMaterialExpressionStrataHairBSDF>(this);
+			MoveConnectionTo(BaseColor, HairBSDF, 0);		// BaseColor
+			MoveConnectionTo(Metallic, HairBSDF, 1);		// Scatter
+			MoveConnectionTo(Specular, HairBSDF, 2);		// Specular
+			MoveConnectionTo(Roughness, HairBSDF, 3);		// Roughness
+			MoveConnectionTo(ClearCoat, HairBSDF, 4);		// Backlit 
+			MoveConnectionTo(Normal, HairBSDF, 5);			// Tangent
+			MoveConnectionTo(EmissiveColor, HairBSDF, 6);	// Emissive
+
+			FrontMaterial.Connect(0, HairBSDF);
+		}
+		else if (MaterialDomain == MD_Surface && ShadingModel == MSM_SingleLayerWater)
+		{
+			UMaterialExpressionStrataSingleLayerWaterBSDF* SLWBSDF = NewObject<UMaterialExpressionStrataSingleLayerWaterBSDF>(this);
+			MoveConnectionTo(BaseColor, SLWBSDF, 0);		// BaseColor
+			MoveConnectionTo(Metallic, SLWBSDF, 1);			// Metallic
+			MoveConnectionTo(Specular, SLWBSDF, 2);			// Specular
+			MoveConnectionTo(Roughness, SLWBSDF, 3);		// Roughness
+			MoveConnectionTo(Normal, SLWBSDF, 4);			// Normal
+			MoveConnectionTo(EmissiveColor, SLWBSDF, 5);	// Emissive
+			MoveConnectionTo(Opacity, SLWBSDF, 6);			// TopMaterialOpacity
+
+			TArray<class UMaterialExpressionCustomOutput*> CustomOutputExpressions;
+			GetAllCustomOutputExpressions(CustomOutputExpressions);
+			UMaterialExpressionSingleLayerWaterMaterialOutput* SLWCustomOutput = nullptr;
+			for (UMaterialExpressionCustomOutput* Expression : CustomOutputExpressions)
+			{
+				SLWCustomOutput = Cast<UMaterialExpressionSingleLayerWaterMaterialOutput>(Expression);
+				if (SLWCustomOutput)
+				{
+					break;
+				}
+			}
+
+			if (SLWCustomOutput)
+			{
+				// Building a graph to get:
+				// WaterExtinction = ScatteringCoefficients + AbsorptionCoefficients
+				// WaterAlbedo     = ScatteringCoefficients / WaterExtinction
+				UMaterialExpressionAdd* AddToExtinction = NewObject<UMaterialExpressionAdd>(this);
+				UMaterialExpressionDivide* DivToAlbedo = NewObject<UMaterialExpressionDivide>(this);
+
+				CopyConnectionTo(*SLWCustomOutput->GetInput(0), AddToExtinction, 0);	// ScatteringCoefficients -> A
+				CopyConnectionTo(*SLWCustomOutput->GetInput(1), AddToExtinction, 1);	// AbsorptionCoefficients -> B
+				SLWBSDF->GetInput(8)->Connect(0, AddToExtinction);						// AddToExtinction -> WaterExtinction
+
+				CopyConnectionTo(*AddToExtinction->GetInput(0), DivToAlbedo, 0);		// ScatteringCoefficients -> A
+				DivToAlbedo->GetInput(1)->Connect(0, AddToExtinction);					// WaterExtinction -> B
+				SLWBSDF->GetInput(7)->Connect(0, DivToAlbedo);							// DivToAlbedo -> WaterAlbedo
+
+				MoveConnectionTo(*SLWCustomOutput->GetInput(2), SLWBSDF, 9);			// WaterPhaseG
+				MoveConnectionTo(*SLWCustomOutput->GetInput(3), SLWBSDF, 10);			// ColorScaleBehindWater
+			}
+
+			FrontMaterial.Connect(0, SLWBSDF);
+		}
+		else if (MaterialDomain == MD_Volume)
+		{
+			UMaterialExpressionStrataVolumetricFogCloudBSDF* VolBSDF = NewObject<UMaterialExpressionStrataVolumetricFogCloudBSDF>(this);
+			MoveConnectionTo(BaseColor, VolBSDF, 0);		// Albedo
+			MoveConnectionTo(SubsurfaceColor, VolBSDF, 1);	// Extinction
+			MoveConnectionTo(EmissiveColor, VolBSDF, 2);	// Emissive
+			MoveConnectionTo(AmbientOcclusion, VolBSDF, 3);	// AmbientOcclusion
+
+			// STRATA_TODO remove the VolumetricAdvancedOutput node and add the input onto FogCloudBSDF even if only used by the cloud renderer?
+
+			FrontMaterial.Connect(0, VolBSDF);
+		}
+
+		// STRATA_TODO Other conversion: SSS, cloth, subsurface, Water, thin translucent, Clear Coat, etc. See EMaterialShadingModel
+
+		// STRATA_TODO unified translucency and coverage
+
+		// Now force the material to recompile and we use a hash of the original StateId.
+		// This is to avoid having different StateId each time we load the material and to not forever recompile it, i.e. use a cached version.
+		uint32 HashBuffer[5];
+		FSHA1::HashBuffer(&StateId, sizeof(FGuid), reinterpret_cast<uint8*>(HashBuffer));
+		StateId.A = HashBuffer[0];
+		StateId.B = HashBuffer[1];
+		StateId.C = HashBuffer[2];
+		StateId.D = HashBuffer[3];
+	}
+#endif
+}
+
 TMap<FGuid, UMaterialInterface*> LightingGuidFixupMap;
 
 void UMaterial::PostLoad()
@@ -3730,6 +3879,7 @@ void UMaterial::PostLoad()
 
 	BackwardsCompatibilityInputConversion();
 	BackwardsCompatibilityVirtualTextureOutputConversion();
+	ConvertMaterialToStrataMaterial();
 
 #if WITH_EDITOR
 	if ( GMaterialsThatNeedSamplerFixup.Get( this ) )
diff --git a/Engine/Source/Runtime/Engine/Private/Materials/MaterialExpressions.cpp b/Engine/Source/Runtime/Engine/Private/Materials/MaterialExpressions.cpp
--- a/Engine/Source/Runtime/Engine/Private/Materials/MaterialExpressions.cpp
+++ b/Engine/Source/Runtime/Engine/Private/Materials/MaterialExpressions.cpp
@@ -19104,7 +19104,8 @@ int32 UMaterialExpressionSingleLayerWaterMaterialOutput::Compile(class FMaterial
 {
 	int32 CodeInput = INDEX_NONE;
 
-	if (!ScatteringCoefficients.IsConnected() && !AbsorptionCoefficients.IsConnected() && !PhaseG.IsConnected())
+	if (!ScatteringCoefficients.IsConnected() && !AbsorptionCoefficients.IsConnected() && !PhaseG.IsConnected()
+		&& CONVERT_MATERIAL_TO_STRATA_ON_LOAD == 0)
 	{
 		Compiler->Error(TEXT("No inputs to Single Layer Water Material."));
 	}
@@ -19669,7 +19670,7 @@ int32 UMaterialExpressionStrataSlabBSDF::Compile(class FMaterialCompiler* Compil
 	int32 TangentCodeChunk = CompileWithDefaultTangentWS(Compiler, Tangent, &bDefaultTangentIsUsed);
 	uint8 SharedNormalIndex = bDefaultTangentIsUsed ? StrataCompilationInfoCreateSharedNormal(Compiler, NormalCodeChunk) : StrataCompilationInfoCreateSharedNormal(Compiler, NormalCodeChunk, TangentCodeChunk);
 
-	int32 RoughnessXCodeChunk = CompileWithDefaultFloat1(Compiler, RoughnessX, 0.0f);
+	int32 RoughnessXCodeChunk = CompileWithDefaultFloat1(Compiler, RoughnessX, 0.5f);
 	// If not plugged in, RoughnessYCodeChunk is set to RoughnessXCodeChunk to get an isotropic behavior
 	int32 RoughnessYCodeChunk = CompileWithDefaultCodeChunk(Compiler, RoughnessY, RoughnessXCodeChunk);
 
@@ -19982,7 +19983,7 @@ UMaterialExpressionStrataUnlitBSDF::UMaterialExpressionStrataUnlitBSDF(const FOb
 int32 UMaterialExpressionStrataUnlitBSDF::Compile(class FMaterialCompiler* Compiler, int32 OutputIndex)
 {
 	int32 OutputCodeChunk = Compiler->StrataUnlitBSDF(
-		CompileWithDefaultFloat3(Compiler, Emissive, 0.0f, 0.0f, 0.0f),
+		CompileWithDefaultFloat3(Compiler, EmissiveColor, 0.0f, 0.0f, 0.0f),
 		CompileWithDefaultFloat3(Compiler, TransmittanceColor, 1.0f, 1.0f, 1.0f));
 
 	uint8 FakeSharedNormalIndex = 0;
diff --git a/Engine/Source/Runtime/Engine/Private/Materials/StrataMaterial.h b/Engine/Source/Runtime/Engine/Private/Materials/StrataMaterial.h
--- a/Engine/Source/Runtime/Engine/Private/Materials/StrataMaterial.h
+++ b/Engine/Source/Runtime/Engine/Private/Materials/StrataMaterial.h
@@ -8,6 +8,17 @@
 
 
 
+// Material conversion to strata is disabled until ready.
+// To enable it locally:
+//		- Set CONVERT_MATERIAL_TO_STRATA_ON_LOAD to 1
+//		- Open UE5MainStreamObjectVersion.h
+//		- Add StrataMaterialConversion at the end of FUE5MainStreamObjectVersion::Type
+//		- Locally change ShaderVersion.ush GUID
+// !!!!!! !!!!!! DO NOT CHECK THIS IN SET TO 1 !!!!!! !!!!!!
+#define CONVERT_MATERIAL_TO_STRATA_ON_LOAD 0
+
+
+
 class FMaterialCompiler;
 
 
