Trivial changes following review

diff --git a/Engine/Source/Runtime/Engine/Classes/Animation/BlendSpaceBase.h b/Engine/Source/Runtime/Engine/Classes/Animation/BlendSpaceBase.h
--- a/Engine/Source/Runtime/Engine/Classes/Animation/BlendSpaceBase.h
+++ b/Engine/Source/Runtime/Engine/Classes/Animation/BlendSpaceBase.h
@@ -58,17 +58,16 @@ struct FBlendParameter
 	UPROPERTY(EditAnywhere, DisplayName = "Number of Grid Divisions", Category=BlendParameter, meta=(UIMin="1", ClampMin="1"))
 	int32 GridNum;
 
-	/** Max value for this parameter. */
-	UPROPERTY(EditAnywhere, DisplayName = "Wrap Input", Category = BlendParameter, meta=(ToolTip=
-		"If false then input parameters are clamped to the min/max values on this axis. If true then the input can go outside the min/max range and the blend space is treated as being cyclic on this axis."))
-	bool WrapInput;
+	/** If false then input parameters are clamped to the min/max values on this axis. If true then the input can go outside the min/max range and the blend space is treated as being cyclic on this axis. */
+	UPROPERTY(EditAnywhere, DisplayName = "Wrap Input", Category = BlendParameter)
+	bool bWrapInput;
 
 	FBlendParameter()
 		: DisplayName(TEXT("None"))
 		, Min(0.f)
 		, Max(100.f)
 		, GridNum(4) // TODO when changing GridNum's default value, it breaks all grid samples ATM - provide way to rebuild grid samples during loading
-		, WrapInput(false)
+		, bWrapInput(false)
 	{
 	}
 
diff --git a/Engine/Source/Runtime/Engine/Private/Animation/BlendSpaceBase.cpp b/Engine/Source/Runtime/Engine/Private/Animation/BlendSpaceBase.cpp
--- a/Engine/Source/Runtime/Engine/Private/Animation/BlendSpaceBase.cpp
+++ b/Engine/Source/Runtime/Engine/Private/Animation/BlendSpaceBase.cpp
@@ -160,7 +160,7 @@ void UBlendSpaceBase::TickAssetPlayer(FAnimTickRecord& Instance, struct FAnimNot
 			if ( !ClampedInput.Equals(BlendInput) ) 
 			{
 				// apply speed change if you want, 
-				if (AxisToScale == BSA_X && !BlendParameters[0].WrapInput)
+				if (AxisToScale == BSA_X && !BlendParameters[0].bWrapInput)
 				{
 					if (ClampedInput.X != 0.f)
 					{
@@ -169,7 +169,7 @@ void UBlendSpaceBase::TickAssetPlayer(FAnimTickRecord& Instance, struct FAnimNot
 				}
 				else if (AxisToScale == BSA_Y)
 				{
-					if (ClampedInput.Y != 0.f && !BlendParameters[1].WrapInput)
+					if (ClampedInput.Y != 0.f && !BlendParameters[1].bWrapInput)
 					{
 						FilterMultiplier *= BlendInput.Y / ClampedInput.Y;
 					}
@@ -1149,7 +1149,7 @@ FVector UBlendSpaceBase::GetClampedBlendInput(const FVector& BlendInput)  const
 	FVector AdjustedInput = BlendInput;
 	for (int iAxis = 0; iAxis != 3; ++iAxis)
 	{
-		if (!BlendParameters[iAxis].WrapInput)
+		if (!BlendParameters[iAxis].bWrapInput)
 		{
 			AdjustedInput[iAxis] = FMath::Clamp(AdjustedInput[iAxis], BlendParameters[iAxis].Min, BlendParameters[iAxis].Max);
 		}
@@ -1162,7 +1162,7 @@ FVector UBlendSpaceBase::GetClampedAndWrappedBlendInput(const FVector& BlendInpu
 	FVector AdjustedInput = BlendInput;
 	for (int iAxis = 0; iAxis != 3; ++iAxis)
 	{
-		if (BlendParameters[iAxis].WrapInput)
+		if (BlendParameters[iAxis].bWrapInput)
 		{
 			AdjustedInput[iAxis] = FMath::Wrap(AdjustedInput[iAxis], BlendParameters[iAxis].Min, BlendParameters[iAxis].Max);
 		}
