Added Strata node special color.

diff --git a/Engine/Source/Editor/UnrealEd/Private/MaterialGraphNode.cpp b/Engine/Source/Editor/UnrealEd/Private/MaterialGraphNode.cpp
--- a/Engine/Source/Editor/UnrealEd/Private/MaterialGraphNode.cpp
+++ b/Engine/Source/Editor/UnrealEd/Private/MaterialGraphNode.cpp
@@ -38,6 +38,7 @@
 #include "Materials/MaterialExpressionNamedReroute.h"
 #include "Materials/MaterialExpressionReroute.h"
 #include "Materials/MaterialExpressionCurveAtlasRowParameter.h"
+#include "Materials/MaterialExpressionStrata.h"
 #include "MaterialEditorUtilities.h"
 #include "MaterialEditorActions.h"
 #include "GraphEditorActions.h"
@@ -315,6 +316,10 @@ FLinearColor UMaterialGraphNode::GetNodeTitleColor() const
 			return FColor( 0, 128, 128 );
 		}
 	}
+	else if (const UMaterialExpressionStrataBSDF* StrataBSDF = Cast<UMaterialExpressionStrataBSDF>(MaterialExpression))
+	{
+		return FColor(181, 29, 230);
+	}
 
 	// Assume that most material expressions act like pure functions and don't affect anything else
 	return Settings->PureFunctionCallNodeTitleColor;
diff --git a/Engine/Source/Runtime/Engine/Classes/Materials/MaterialExpressionStrata.h b/Engine/Source/Runtime/Engine/Classes/Materials/MaterialExpressionStrata.h
--- a/Engine/Source/Runtime/Engine/Classes/Materials/MaterialExpressionStrata.h
+++ b/Engine/Source/Runtime/Engine/Classes/Materials/MaterialExpressionStrata.h
@@ -24,19 +24,14 @@ extern int32 CompileStrataBlendFunction(FMaterialCompiler* Compiler, const int32
 // BSDF nodes
 
 // This would be needed to for a common node interface and weight input and normal too?
-/*UCLASS(MinimalAPI, collapsecategories, hidecategories = Object)
+UCLASS(MinimalAPI, collapsecategories, hidecategories = Object, Abstract)
 class UMaterialExpressionStrataBSDF : public UMaterialExpression
 {
 	GENERATED_UCLASS_BODY()
-
-	UPROPERTY()
-	FExpressionInput Weight;
-
-	// Normal?
-}*/
+};
 
 UCLASS(MinimalAPI, collapsecategories, hidecategories = Object)
-class UMaterialExpressionStrataSlabBSDF : public UMaterialExpression
+class UMaterialExpressionStrataSlabBSDF : public UMaterialExpressionStrataBSDF
 {
 	GENERATED_UCLASS_BODY()
 
@@ -120,7 +115,7 @@ class UMaterialExpressionStrataSlabBSDF : public UMaterialExpression
 };
 
 UCLASS(MinimalAPI, collapsecategories, hidecategories = Object)
-class UMaterialExpressionStrataSheenBSDF : public UMaterialExpression
+class UMaterialExpressionStrataSheenBSDF : public UMaterialExpressionStrataBSDF
 {
 	GENERATED_UCLASS_BODY()
 
@@ -155,7 +150,7 @@ class UMaterialExpressionStrataSheenBSDF : public UMaterialExpression
 };
 
 UCLASS(MinimalAPI, collapsecategories, hidecategories = Object)
-class UMaterialExpressionStrataVolumetricFogCloudBSDF : public UMaterialExpression
+class UMaterialExpressionStrataVolumetricFogCloudBSDF : public UMaterialExpressionStrataBSDF
 {
 	GENERATED_UCLASS_BODY()
 
@@ -196,7 +191,7 @@ class UMaterialExpressionStrataVolumetricFogCloudBSDF : public UMaterialExpressi
 };
 
 UCLASS(MinimalAPI, collapsecategories, hidecategories = Object)
-class UMaterialExpressionStrataUnlitBSDF : public UMaterialExpression
+class UMaterialExpressionStrataUnlitBSDF : public UMaterialExpressionStrataBSDF
 {
 	GENERATED_UCLASS_BODY()
 
@@ -225,7 +220,7 @@ class UMaterialExpressionStrataUnlitBSDF : public UMaterialExpression
 };
 
 UCLASS(MinimalAPI, collapsecategories, hidecategories = Object)
-class UMaterialExpressionStrataHairBSDF : public UMaterialExpression
+class UMaterialExpressionStrataHairBSDF : public UMaterialExpressionStrataBSDF
 {
 	GENERATED_UCLASS_BODY()
 		
@@ -284,7 +279,7 @@ class UMaterialExpressionStrataHairBSDF : public UMaterialExpression
 };
 
 UCLASS(MinimalAPI, collapsecategories, hidecategories = Object)
-class UMaterialExpressionStrataSingleLayerWaterBSDF : public UMaterialExpression
+class UMaterialExpressionStrataSingleLayerWaterBSDF : public UMaterialExpressionStrataBSDF
 {
 	GENERATED_UCLASS_BODY()
 
@@ -372,7 +367,7 @@ class UMaterialExpressionStrataSingleLayerWaterBSDF : public UMaterialExpression
 // Operator nodes
 
 UCLASS(MinimalAPI, collapsecategories, hidecategories = Object)
-class UMaterialExpressionStrataHorizontalMixing : public UMaterialExpression
+class UMaterialExpressionStrataHorizontalMixing : public UMaterialExpressionStrataBSDF
 {
 	GENERATED_UCLASS_BODY()
 		
@@ -407,7 +402,7 @@ class UMaterialExpressionStrataHorizontalMixing : public UMaterialExpression
 };
 
 UCLASS(MinimalAPI, collapsecategories, hidecategories = Object)
-class UMaterialExpressionStrataVerticalLayering : public UMaterialExpression
+class UMaterialExpressionStrataVerticalLayering : public UMaterialExpressionStrataBSDF
 {
 	GENERATED_UCLASS_BODY()
 
@@ -436,7 +431,7 @@ class UMaterialExpressionStrataVerticalLayering : public UMaterialExpression
 };
 
 UCLASS(MinimalAPI, collapsecategories, hidecategories = Object)
-class UMaterialExpressionStrataAdd : public UMaterialExpression
+class UMaterialExpressionStrataAdd : public UMaterialExpressionStrataBSDF
 {
 	GENERATED_UCLASS_BODY()
 
@@ -464,8 +459,8 @@ class UMaterialExpressionStrataAdd : public UMaterialExpression
 	//~ End UMaterialExpression Interface
 };
 
-UCLASS(MinimalAPI, DisplayName = "Strata Multiply Weight", collapsecategories, hidecategories = Object)
-class UMaterialExpressionStrataMultiply : public UMaterialExpression
+UCLASS(MinimalAPI, collapsecategories, hidecategories = Object)
+class UMaterialExpressionStrataMultiply : public UMaterialExpressionStrataBSDF
 {
 	GENERATED_UCLASS_BODY()
 
diff --git a/Engine/Source/Runtime/Engine/Private/Materials/MaterialExpressions.cpp b/Engine/Source/Runtime/Engine/Private/Materials/MaterialExpressions.cpp
--- a/Engine/Source/Runtime/Engine/Private/Materials/MaterialExpressions.cpp
+++ b/Engine/Source/Runtime/Engine/Private/Materials/MaterialExpressions.cpp
@@ -19647,6 +19647,20 @@ static int32 CompileWithDefaultTangentWS(class FMaterialCompiler* Compiler, FExp
 
 #endif // WITH_EDITOR
 
+UMaterialExpressionStrataBSDF::UMaterialExpressionStrataBSDF(const FObjectInitializer& ObjectInitializer)
+	: Super(ObjectInitializer)
+{
+	struct FConstructorStatics
+	{
+		FText NAME_Strata;
+		FConstructorStatics() : NAME_Strata(LOCTEXT("Strata BSDFs", "Strata BSDFs")) { }
+	};
+	static FConstructorStatics ConstructorStatics;
+#if WITH_EDITORONLY_DATA
+	MenuCategories.Add(ConstructorStatics.NAME_Strata);
+#endif
+}
+
 UMaterialExpressionStrataSlabBSDF::UMaterialExpressionStrataSlabBSDF(const FObjectInitializer& ObjectInitializer)
 	: Super(ObjectInitializer)
 {
