Remove no longer needed Gulbrandsen functions. We will have a new Specular/Metal parameterisation

diff --git a/Engine/Shaders/Private/BRDF.ush b/Engine/Shaders/Private/BRDF.ush
--- a/Engine/Shaders/Private/BRDF.ush
+++ b/Engine/Shaders/Private/BRDF.ush
@@ -444,14 +444,6 @@ float3 F_Schlick( float3 SpecularColor, float VoH )
 	
 }
 
-float3 F_Schlick(float3 SpecularColor, float3 EdgeTint, float VoH)
-{
-	float Fc = Pow5(1 - VoH);
-	// Anything less than 2% is physically impossible and is instead considered to be shadowing
-	return saturate(50.0 * SpecularColor.g) * Fc * EdgeTint + (1 - Fc) * SpecularColor;
-
-}
-
 float3 F_Fresnel( float3 SpecularColor, float VoH )
 {
 	float3 SpecularColorSqrt = sqrt( clamp( float3(0, 0, 0), float3(0.99, 0.99, 0.99), SpecularColor ) );
@@ -559,12 +551,6 @@ half3 EnvBRDFApprox( half3 SpecularColor, half Roughness, half NoV )
 	return SpecularColor * AB.x + AB.y;
 }
 
-half3 EnvBRDFApprox(half3 SpecularColor, float3 EdgeTint, half Roughness, half NoV)
-{
-	half2 AB = EnvBRDFApproxLazarov(SpecularColor, Roughness, NoV);
-	return SpecularColor * AB.x + EdgeTint * AB.y;
-}
-
 half EnvBRDFApproxNonmetal( half Roughness, half NoV )
 {
 	// Same as EnvBRDFApprox( 0.04, Roughness, NoV )
diff --git a/Engine/Shaders/Private/ShadingCommon.ush b/Engine/Shaders/Private/ShadingCommon.ush
--- a/Engine/Shaders/Private/ShadingCommon.ush
+++ b/Engine/Shaders/Private/ShadingCommon.ush
@@ -113,64 +113,3 @@ float MakeRoughnessSafe(float Roughness, float MinRoughness=0.001f)
 {
 	return clamp(Roughness, MinRoughness, 1.0f);
 }
-
-// [Gulbrandsen, "Artist Friendly Metallic Fresnel"]
-
-// Compute F0 and edge tint from complex index of refraction N = n + ik. (absolute refractive index + extinction coefficients)
-void ComputeF0AndEdgeTintFromIOR(in float3 n, in float3 k, inout float3 Reflectivity, inout float3 EdgeTint)
-{
-	const float3 k2 = Square(k);
-
-	Reflectivity =
-		(Square(n - 1) + k2) /
-		(Square(n + 1) + k2);
-
-	const float3 SqrtR = sqrt(Reflectivity);
-	const float3 Common0 = (1 + SqrtR) / (1 - SqrtR);
-	EdgeTint = (Common0 - n) / (Common0 + (1 - SqrtR) / (1 + SqrtR));
-}
-
-float3 ComputeF0AndEdgeTintFromComplexIOR_Reflectivity(in float3 n, in float3 k)
-{
-	float3 Reflectivity = 0;
-	float3 EdgeTint = 0;
-	ComputeF0AndEdgeTintFromIOR(n, k, Reflectivity, EdgeTint);
-	return Reflectivity;
-}
-
-float3 ComputeF0AndEdgeTintFromComplexIOR_EdgeTint(in float3 n, in float3 k)
-{
-	float3 Reflectivity = 0;
-	float3 EdgeTint = 0;
-	ComputeF0AndEdgeTintFromIOR(n, k, Reflectivity, EdgeTint);
-	return EdgeTint;
-}
-
-// Compute complex index of refraction N = n + ik from F0 and edge tint.
-void ComputeComplexIORFromF0AndEdgeTint(in float3 Reflectivity, in float3 EdgeTint, inout float3 n, inout float3 k)
-{
-	const float3 g = EdgeTint;
-	const float3 r = clamp(Reflectivity, 0.0, 0.999);
-	const float3 SqrtR = sqrt(r);
-
-	n = g * (1 - r) / (1 + r) + (1 - g) * (1 + SqrtR) / (1 - SqrtR);
-
-	const float3 k2 = ((n + 1) * (n + 1) * r - (n - 1) * (n - 1)) / (1 - r);
-	k = sqrt(k2);
-}
-
-float3 ComputeComplexIORFromF0AndEdgeTint_IOR(in float3 Reflectivity, in float3 EdgeTint)
-{
-	float3 n = 0;
-	float3 k = 0;
-	ComputeComplexIORFromF0AndEdgeTint(Reflectivity, EdgeTint, n, k);
-	return n;
-}
-
-float3 ComputeComplexIORFromF0AndEdgeTint_Extinction(in float3 Reflectivity, in float3 EdgeTint)
-{
-	float3 n = 0;
-	float3 k = 0;
-	ComputeComplexIORFromF0AndEdgeTint(Reflectivity, EdgeTint, n, k);
-	return k;
-}
\ No newline at end of file
diff --git a/Engine/Source/Runtime/Engine/Classes/Materials/MaterialExpressionStrata.h b/Engine/Source/Runtime/Engine/Classes/Materials/MaterialExpressionStrata.h
--- a/Engine/Source/Runtime/Engine/Classes/Materials/MaterialExpressionStrata.h
+++ b/Engine/Source/Runtime/Engine/Classes/Materials/MaterialExpressionStrata.h
@@ -493,61 +493,6 @@ class UMaterialExpressionStrataMultiply : public UMaterialExpressionStrataBSDF
 ///////////////////////////////////////////////////////////////////////////////
 // Utilities
 
-UCLASS(MinimalAPI, collapsecategories, hidecategories = Object)
-class UMaterialExpressionStrataArtisticIOR : public UMaterialExpression
-{
-	GENERATED_UCLASS_BODY()
-
-	/**
-	 * Reflectivity when view direction is perpendicular to the surface, also known as F0 (type = float3, unit = unitless)
-	 */
-	UPROPERTY()
-	FExpressionInput Reflectivity;
-
-	/**
-	 * Reflectivity when the view direction is tangent to the surface (type = float3, unit = unitless)
-	 */
-	UPROPERTY()
-	FExpressionInput EdgeColor;
-
-	//~ Begin UMaterialExpression Interface
-#if WITH_EDITOR
-	virtual int32 Compile(class FMaterialCompiler* Compiler, int32 OutputIndex) override;
-	virtual void GetCaption(TArray<FString>& OutCaptions) const override;
-	virtual uint32 GetOutputType(int32 OutputIndex) override;
-	virtual uint32 GetInputType(int32 InputIndex) override;
-#endif
-	//~ End UMaterialExpression Interface
-};
-
-UCLASS(MinimalAPI, collapsecategories, hidecategories = Object)
-class UMaterialExpressionStrataPhysicalIOR : public UMaterialExpression
-{
-	GENERATED_UCLASS_BODY()
-
-	/**
-	* The index of refraction of the surface (type = float, unit = unitless)
-	*/
-	UPROPERTY()
-	FExpressionInput IOR;
-
-	/**
-	 * The rate at which light is absorbed or scattered by the medium. Mean Free Path = 1 / Extinction. (type = float3, unit = 1/m)
-	 */
-	UPROPERTY()
-	FExpressionInput Extinction;
-
-	//~ Begin UMaterialExpression Interface
-#if WITH_EDITOR
-	virtual int32 Compile(class FMaterialCompiler* Compiler, int32 OutputIndex) override;
-	virtual void GetCaption(TArray<FString>& OutCaptions) const override;
-	virtual uint32 GetOutputType(int32 OutputIndex) override;
-	virtual uint32 GetInputType(int32 InputIndex) override;
-#endif
-	//~ End UMaterialExpression Interface
-};
-
-
 UCLASS(MinimalAPI, collapsecategories, hidecategories = Object)
 class UMaterialExpressionStrataAnisotropyToRoughness : public UMaterialExpression
 {
diff --git a/Engine/Source/Runtime/Engine/Private/Materials/HLSLMaterialTranslator.cpp b/Engine/Source/Runtime/Engine/Private/Materials/HLSLMaterialTranslator.cpp
--- a/Engine/Source/Runtime/Engine/Private/Materials/HLSLMaterialTranslator.cpp
+++ b/Engine/Source/Runtime/Engine/Private/Materials/HLSLMaterialTranslator.cpp
@@ -7677,34 +7677,6 @@ int32 FHLSLMaterialTranslator::StrataMultiply(int32 A, int32 Weight)
 	);
 }
 
-int32 FHLSLMaterialTranslator::StrataArtisticIOR(int32 Reflectivity, int32 EdgeColor, int32 OutputIndex)
-{
-	if (Reflectivity == INDEX_NONE || EdgeColor == INDEX_NONE)
-	{
-		return INDEX_NONE;
-	}
-	return AddCodeChunk(
-		MCT_Float3, 
-		OutputIndex == 0 ? TEXT("ComputeComplexIORFromF0AndEdgeTint_IOR(%s, %s)") : TEXT("ComputeComplexIORFromF0AndEdgeTint_Extinction(%s, %s)"),
-		*GetParameterCode(Reflectivity),
-		*GetParameterCode(EdgeColor)
-	);
-}
-
-int32 FHLSLMaterialTranslator::StrataPhysicalIOR(int32 IOR, int32 Extinction, int32 OutputIndex)
-{
-	if (IOR == INDEX_NONE || Extinction == INDEX_NONE)
-	{
-		return INDEX_NONE;
-	}
-	return AddCodeChunk(
-		MCT_Float3,
-		OutputIndex == 0 ? TEXT("ComputeF0AndEdgeTintFromComplexIOR_Reflectivity(%s, %s)") : TEXT("ComputeF0AndEdgeTintFromComplexIOR_EdgeTint(%s, %s)"),
-		*GetParameterCode(IOR),
-		*GetParameterCode(Extinction)
-	);
-}
-
 int32 FHLSLMaterialTranslator::StrataAnisotropyToRoughness(int32 Roughness, int32 Anisotropy, int32 OutputIndex)
 {
 	if (Roughness == INDEX_NONE)
diff --git a/Engine/Source/Runtime/Engine/Private/Materials/HLSLMaterialTranslator.h b/Engine/Source/Runtime/Engine/Private/Materials/HLSLMaterialTranslator.h
--- a/Engine/Source/Runtime/Engine/Private/Materials/HLSLMaterialTranslator.h
+++ b/Engine/Source/Runtime/Engine/Private/Materials/HLSLMaterialTranslator.h
@@ -808,8 +808,6 @@ protected:
 	virtual int32 StrataVerticalLayering(int32 Top, int32 Base) override;
 	virtual int32 StrataAdd(int32 A, int32 B) override;
 	virtual int32 StrataMultiply(int32 A, int32 Weight) override;
-	virtual int32 StrataArtisticIOR(int32 Reflectivity, int32 EdgeColor, int32 OutputIndex) override;
-	virtual int32 StrataPhysicalIOR(int32 IOR, int32 Extinction, int32 OutputIndex) override;
 	virtual int32 StrataAnisotropyToRoughness(int32 Roughness, int32 Anisotropy, int32 OutputIndex) override;
 
 	virtual void StrataCompilationInfoRegisterCodeChunk(int32 CodeChunk, FStrataMaterialCompilationInfo& StrataMaterialCompilationInfo) override;
diff --git a/Engine/Source/Runtime/Engine/Private/Materials/MaterialExpressions.cpp b/Engine/Source/Runtime/Engine/Private/Materials/MaterialExpressions.cpp
--- a/Engine/Source/Runtime/Engine/Private/Materials/MaterialExpressions.cpp
+++ b/Engine/Source/Runtime/Engine/Private/Materials/MaterialExpressions.cpp
@@ -20543,101 +20543,6 @@ void UMaterialExpressionStrataMultiply::GatherStrataMaterialInfo(FStrataMaterial
 
 
 
-UMaterialExpressionStrataArtisticIOR::UMaterialExpressionStrataArtisticIOR(const FObjectInitializer& ObjectInitializer)
-	: Super(ObjectInitializer)
-{
-	struct FConstructorStatics
-	{
-		FText NAME_Strata;
-		FConstructorStatics() : NAME_Strata(LOCTEXT("Strata Helpers", "Strata Helpers")) { }
-	};
-	static FConstructorStatics ConstructorStatics;
-#if WITH_EDITORONLY_DATA
-	MenuCategories.Add(ConstructorStatics.NAME_Strata);
-
-	bShowOutputNameOnPin = true;
-
-	Outputs.Reset();
-	Outputs.Add(FExpressionOutput(TEXT("IOR")));
-	Outputs.Add(FExpressionOutput(TEXT("Extinction")));
-#endif
-}
-
-#if WITH_EDITOR
-int32 UMaterialExpressionStrataArtisticIOR::Compile(class FMaterialCompiler* Compiler, int32 OutputIndex)
-{
-	return Compiler->StrataArtisticIOR(
-		Reflectivity.GetTracedInput().Expression	? Reflectivity.Compile(Compiler): Compiler->Constant3(0.947f, 0.776f, 0.371f),	// Default to Gold
-		EdgeColor.GetTracedInput().Expression		? EdgeColor.Compile(Compiler)	: Compiler->Constant3(1.000f, 0.982f, 0.753f),	// Default to Gold
-		OutputIndex);
-}
-
-void UMaterialExpressionStrataArtisticIOR::GetCaption(TArray<FString>& OutCaptions) const
-{
-	OutCaptions.Add(TEXT("Strata Artistic IOR"));
-}
-
-uint32 UMaterialExpressionStrataArtisticIOR::GetOutputType(int32 OutputIndex)
-{
-	return MCT_Float3;
-}
-
-uint32 UMaterialExpressionStrataArtisticIOR::GetInputType(int32 InputIndex)
-{
-	return MCT_Float3;
-}
-#endif // WITH_EDITOR
-
-
-
-UMaterialExpressionStrataPhysicalIOR::UMaterialExpressionStrataPhysicalIOR(const FObjectInitializer& ObjectInitializer)
-	: Super(ObjectInitializer)
-{
-	struct FConstructorStatics
-	{
-		FText NAME_Strata;
-		FConstructorStatics() : NAME_Strata(LOCTEXT("Strata Helpers", "Strata Helpers")) { }
-	};
-	static FConstructorStatics ConstructorStatics;
-#if WITH_EDITORONLY_DATA
-	MenuCategories.Add(ConstructorStatics.NAME_Strata);
-
-	bShowOutputNameOnPin = true;
-
-	Outputs.Reset();
-	Outputs.Add(FExpressionOutput(TEXT("Reflectivity")));
-	Outputs.Add(FExpressionOutput(TEXT("EdgeColor")));
-#endif
-}
-
-#if WITH_EDITOR
-int32 UMaterialExpressionStrataPhysicalIOR::Compile(class FMaterialCompiler* Compiler, int32 OutputIndex)
-{
-	return Compiler->StrataPhysicalIOR(
-		IOR.GetTracedInput().Expression			? IOR.Compile(Compiler)			: Compiler->Constant3(0.18f, 0.42f, 1.37f),// Default to Gold
-		Extinction.GetTracedInput().Expression	? Extinction.Compile(Compiler)	: Compiler->Constant3(3.41f, 2.35f, 1.77f),// Default to Gold
-		OutputIndex);
-}
-
-void UMaterialExpressionStrataPhysicalIOR::GetCaption(TArray<FString>& OutCaptions) const
-{
-	OutCaptions.Add(TEXT("Strata Physical IOR"));
-}
-
-uint32 UMaterialExpressionStrataPhysicalIOR::GetOutputType(int32 OutputIndex)
-{
-	return MCT_Float3;
-}
-
-uint32 UMaterialExpressionStrataPhysicalIOR::GetInputType(int32 InputIndex)
-{
-	return MCT_Float3;
-}
-#endif // WITH_EDITOR
-
-
-
-
 UMaterialExpressionStrataAnisotropyToRoughness::UMaterialExpressionStrataAnisotropyToRoughness(const FObjectInitializer& ObjectInitializer)
 	: Super(ObjectInitializer)
 {
diff --git a/Engine/Source/Runtime/Engine/Public/MaterialCompiler.h b/Engine/Source/Runtime/Engine/Public/MaterialCompiler.h
--- a/Engine/Source/Runtime/Engine/Public/MaterialCompiler.h
+++ b/Engine/Source/Runtime/Engine/Public/MaterialCompiler.h
@@ -402,8 +402,6 @@ public:
 	virtual int32 StrataVerticalLayering(int32 Top, int32 Base) = 0;
 	virtual int32 StrataAdd(int32 A, int32 B) = 0;
 	virtual int32 StrataMultiply(int32 A, int32 Weight) = 0;
-	virtual int32 StrataArtisticIOR(int32 Reflectivity, int32 EdgeColor, int32 OutputIndex) = 0;
-	virtual int32 StrataPhysicalIOR(int32 IOR, int32 Extinction, int32 OutputIndex) = 0;
 	virtual int32 StrataAnisotropyToRoughness(int32 Roughness, int32 Anisotropy, int32 OutputIndex) = 0;
 
 	virtual void StrataCompilationInfoRegisterCodeChunk(int32 CodeChunk, FStrataMaterialCompilationInfo& StrataMaterialCompilationInfo) = 0;
@@ -858,16 +856,6 @@ public:
 		return Compiler->StrataMultiply(A, Weight);
 	}
 
-	virtual int32 StrataArtisticIOR(int32 Reflectivity, int32 EdgeColor, int32 OutputIndex) override
-	{
-		return Compiler->StrataArtisticIOR(Reflectivity, EdgeColor, OutputIndex);
-	}
-
-	virtual int32 StrataPhysicalIOR(int32 IOR, int32 Extinction, int32 OutputIndex) override
-	{
-		return Compiler->StrataPhysicalIOR(IOR, Extinction, OutputIndex);
-	}
-
 	virtual int32 StrataAnisotropyToRoughness(int32 Roughness, int32 Anisotropy, int32 OutputIndex) override
 	{
 		return Compiler->StrataAnisotropyToRoughness(Roughness, Anisotropy, OutputIndex);
