Avoid 10ms delay when waiting on streaming if possible
Add some Insights tag in texture to better highlight some time sink

diff --git a/Engine/Source/Runtime/Engine/Private/StreamableRenderAsset.cpp b/Engine/Source/Runtime/Engine/Private/StreamableRenderAsset.cpp
--- a/Engine/Source/Runtime/Engine/Private/StreamableRenderAsset.cpp
+++ b/Engine/Source/Runtime/Engine/Private/StreamableRenderAsset.cpp
@@ -243,6 +243,8 @@ bool UStreamableRenderAsset::IsFullyStreamedIn()
 
 void UStreamableRenderAsset::WaitForPendingInitOrStreaming(bool bWaitForLODTransition, bool bSendCompletionEvents)
 {
+	TRACE_CPUPROFILER_EVENT_SCOPE(UStreamableRenderAsset::WaitForPendingInitOrStreaming);
+
 	while (HasPendingInitOrStreaming(bWaitForLODTransition))
 	{
 		ensure(!IsAssetStreamingSuspended());
@@ -251,8 +253,13 @@ void UStreamableRenderAsset::WaitForPendingInitOrStreaming(bool bWaitForLODTrans
 		TickStreaming(bSendCompletionEvents);
 		// Make sure any render commands are executed, in particular things like InitRHI, or asset updates on the render thread.
 		FlushRenderingCommands();
-		// Give some time increment so that LOD transition can complete, and also for the gamethread to give room for streaming async tasks.
-		FPlatformProcess::Sleep(RENDER_ASSET_STREAMING_SLEEP_DT);
+
+		// Most of the time, sleeping is not required, so avoid loosing a whole quantum (10ms on W10Pro) unless stricly necessary.
+		if (HasPendingInitOrStreaming(bWaitForLODTransition))
+		{
+			// Give some time increment so that LOD transition can complete, and also for the gamethread to give room for streaming async tasks.
+			FPlatformProcess::Sleep(RENDER_ASSET_STREAMING_SLEEP_DT);
+		}
 	}
 }
 
diff --git a/Engine/Source/Runtime/Engine/Private/Texture.cpp b/Engine/Source/Runtime/Engine/Private/Texture.cpp
--- a/Engine/Source/Runtime/Engine/Private/Texture.cpp
+++ b/Engine/Source/Runtime/Engine/Private/Texture.cpp
@@ -262,6 +262,7 @@ bool UTexture::CanEditChange(const FProperty* InProperty) const
 
 void UTexture::PostEditChangeProperty(FPropertyChangedEvent& PropertyChangedEvent)
 {
+	TRACE_CPUPROFILER_EVENT_SCOPE(UTexture::PostEditChangeProperty);
 	Super::PostEditChangeProperty(PropertyChangedEvent);
 
 	SetLightingGuid();
@@ -901,7 +902,7 @@ FStreamableRenderResourceState UTexture::GetResourcePostInitState(FTexturePlatfo
 		// Ensure NumMipsInTail is within valid range to safeguard on the above expressions. 
 		const int32 NumMipsInTail = FMath::Clamp<int32>(PlatformData->GetNumMipsInTail(), 1, NumMips);
 
-		// Bias is not allowed to shrink the mip count bellow NumMipsInTail.
+		// Bias is not allowed to shrink the mip count below NumMipsInTail.
 		NumRequestedMips = FMath::Max<int32>(NumMips - ResourceLODBias, NumMipsInTail);
 
 		// If trying to load optional mips, check if the first resource mip is available.
diff --git a/Engine/Source/Runtime/Engine/Private/Texture2D.cpp b/Engine/Source/Runtime/Engine/Private/Texture2D.cpp
--- a/Engine/Source/Runtime/Engine/Private/Texture2D.cpp
+++ b/Engine/Source/Runtime/Engine/Private/Texture2D.cpp
@@ -662,6 +662,8 @@ void UTexture2D::GetAssetRegistryTags(TArray<FAssetRegistryTag>& OutTags) const
 
 void UTexture2D::UpdateResource()
 {
+	TRACE_CPUPROFILER_EVENT_SCOPE(UTexture2D::UpdateResource);
+
 	WaitForPendingInitOrStreaming();
 
 #if WITH_EDITOR
