Fixed Strata anisotropic lighting that got broken at some point.

diff --git a/Engine/Source/Runtime/Engine/Private/Materials/MaterialExpressions.cpp b/Engine/Source/Runtime/Engine/Private/Materials/MaterialExpressions.cpp
--- a/Engine/Source/Runtime/Engine/Private/Materials/MaterialExpressions.cpp
+++ b/Engine/Source/Runtime/Engine/Private/Materials/MaterialExpressions.cpp
@@ -19678,15 +19678,18 @@ UMaterialExpressionStrataSlabBSDF::UMaterialExpressionStrataSlabBSDF(const FObje
 #if WITH_EDITOR
 int32 UMaterialExpressionStrataSlabBSDF::Compile(class FMaterialCompiler* Compiler, int32 OutputIndex)
 {
-	int32 NormalCodeChunk = CompileWithDefaultNormalWS(Compiler, Normal);
-
-	bool bDefaultTangentIsUsed = false;
-	int32 TangentCodeChunk = CompileWithDefaultTangentWS(Compiler, Tangent, &bDefaultTangentIsUsed);
-	uint8 SharedNormalIndex = bDefaultTangentIsUsed ? StrataCompilationInfoCreateSharedNormal(Compiler, NormalCodeChunk) : StrataCompilationInfoCreateSharedNormal(Compiler, NormalCodeChunk, TangentCodeChunk);
-
 	int32 RoughnessXCodeChunk = CompileWithDefaultFloat1(Compiler, RoughnessX, 0.5f);
 	// If not plugged in, RoughnessYCodeChunk is set to RoughnessXCodeChunk to get an isotropic behavior
 	int32 RoughnessYCodeChunk = CompileWithDefaultCodeChunk(Compiler, RoughnessY, RoughnessXCodeChunk);
+	// As long as both roughness are potentially different, we must take it into account in our encoding.
+	// We also cannot ignore the tangent when using the default Tangent because GetTangentBasis
+	// used in StrataGetBSDFSharedBasis cannot be relied on for smooth tangent used for lighting on any mesh.
+	bool bAnisotropyPotentiallyUsed = RoughnessXCodeChunk != RoughnessYCodeChunk;
+
+	int32 NormalCodeChunk = CompileWithDefaultNormalWS(Compiler, Normal);
+
+	int32 TangentCodeChunk = CompileWithDefaultTangentWS(Compiler, Tangent);
+	uint8 SharedNormalIndex = bAnisotropyPotentiallyUsed ? StrataCompilationInfoCreateSharedNormal(Compiler, NormalCodeChunk, TangentCodeChunk) : StrataCompilationInfoCreateSharedNormal(Compiler, NormalCodeChunk);
 
 	int32 SSSProfileCodeChunk = INDEX_NONE;
 	const bool bHasScattering = SubsurfaceProfile != nullptr;
