Ensure that the transient resampled animation data is always generated during PostLoad in the Editor. Modifications to the model (curves) are not guaranteed to cause this to happen otherwise and will leave the character in ref-pose until re-compressed.

diff --git a/Engine/Source/Runtime/Engine/Private/Animation/AnimSequence.cpp b/Engine/Source/Runtime/Engine/Private/Animation/AnimSequence.cpp
--- a/Engine/Source/Runtime/Engine/Private/Animation/AnimSequence.cpp
+++ b/Engine/Source/Runtime/Engine/Private/Animation/AnimSequence.cpp
@@ -854,6 +854,9 @@ void UAnimSequence::PostLoad()
 	}
 
 #if WITH_EDITOR
+	// Generate transient animation data, used to fallback on whenever the animation modified and not yet compressed (what previously used to be the raw data)
+	ResampleAnimationTrackData();
+
 	static bool ForcedRecompressionSetting = FAnimationUtils::GetForcedRecompressionSetting();
 
 	if (ForcedRecompressionSetting)
@@ -1079,7 +1082,7 @@ void UAnimSequence::PostEditChangeProperty(FPropertyChangedEvent& PropertyChange
 	// @Todo fix me: This is temporary fix to make sure they always have compressed data
 	ValidateModel();
 	if (DataModel->GetNumBoneTracks() && bNeedPostProcess)
-		{
+	{
 		// QQ need new function here to queue async compression
 		PRAGMA_DISABLE_DEPRECATION_WARNINGS
 		MarkRawDataAsModified(false);
@@ -5487,7 +5490,7 @@ void UAnimSequence::PopulateModel()
 
 	USkeleton* TargetSkeleton = GetSkeleton();
 	UE::Anim::CopyCurveDataToModel(CurveData, TargetSkeleton, Controller);
-
+	
 	const int32 NumTracks = SequenceTracks.Num();
 	for (int32 TrackIndex = 0; TrackIndex < NumTracks; ++TrackIndex)
 	{
