Fix potential duplicated actors using ForceExternalActorLevelReference.

diff --git a/Engine/Source/Runtime/Core/Public/UObject/UE5MainStreamObjectVersion.h b/Engine/Source/Runtime/Core/Public/UObject/UE5MainStreamObjectVersion.h
--- a/Engine/Source/Runtime/Core/Public/UObject/UE5MainStreamObjectVersion.h
+++ b/Engine/Source/Runtime/Core/Public/UObject/UE5MainStreamObjectVersion.h
@@ -82,6 +82,9 @@ struct CORE_API FUE5MainStreamObjectVersion
 		// Fix WorldPartitionActorDesc serialization archive not persistent
 		WorldPartitionActorDescSerializeArchivePersistent,
 
+		// Fix potentially duplicated actors when using ForceExternalActorLevelReference
+		FixForceExternalActorLevelReferenceDuplicates,
+
 		// -----<new versions can be added above this line>-------------------------------------------------
 		VersionPlusOne,
 		LatestVersion = VersionPlusOne - 1
diff --git a/Engine/Source/Runtime/Engine/Private/Level.cpp b/Engine/Source/Runtime/Engine/Private/Level.cpp
--- a/Engine/Source/Runtime/Engine/Private/Level.cpp
+++ b/Engine/Source/Runtime/Engine/Private/Level.cpp
@@ -7,6 +7,7 @@ Level.cpp: Level-related functions
 #include "Engine/Level.h"
 #include "Misc/ScopedSlowTask.h"
 #include "UObject/RenderingObjectVersion.h"
+#include "UObject/UE5MainStreamObjectVersion.h"
 #include "Templates/UnrealTemplate.h"
 #include "UObject/Package.h"
 #include "EngineStats.h"
@@ -401,6 +402,7 @@ void ULevel::Serialize( FArchive& Ar )
 	Ar.UsingCustomVersion(FReleaseObjectVersion::GUID);
 	Ar.UsingCustomVersion(FRenderingObjectVersion::GUID);
 	Ar.UsingCustomVersion(FFortniteMainBranchObjectVersion::GUID);
+	Ar.UsingCustomVersion(FUE5MainStreamObjectVersion::GUID);
 
 	if (Ar.IsLoading())
 	{
@@ -533,6 +535,24 @@ void ULevel::Serialize( FArchive& Ar )
 		Ar << Dummy3;
 	}
 
+	if ( IsUsingExternalActors() && Ar.IsLoading() && Ar.CustomVer(FUE5MainStreamObjectVersion::GUID) < FUE5MainStreamObjectVersion::FixForceExternalActorLevelReferenceDuplicates )
+	{
+		TSet<AActor*> DupActors;
+		for (int ActorIndex = 0; ActorIndex < Actors.Num(); ++ActorIndex)
+		{
+			if (AActor * Actor = Actors[ActorIndex])
+			{
+				bool bAlreadyInSet = false;
+				DupActors.Add(Actor, &bAlreadyInSet);
+
+				if (bAlreadyInSet)
+				{
+					Actors[ActorIndex] = nullptr;
+				}
+			}
+		}
+	}
+
 	// Mark archive and package as containing a map if we're serializing to disk.
 	if( !HasAnyFlags( RF_ClassDefaultObject ) && Ar.IsPersistent() )
 	{
@@ -618,20 +638,24 @@ void ULevel::AddLoadedActor(AActor* Actor)
 	check(Actor->GetLevel() == this);
 	check(!Actor->IsPendingKill());
 
-	Actors.Add(Actor);
-	ActorsForGC.Add(Actor);
-
-	// If components from actors belonging to this level were already registered, do the same for the newly loaded actor.
-	// Otherwise, the new actor components will be registered later once the world initialization is completed, through UpdateWorldComponents()
-	if (bAreComponentsCurrentlyRegistered)
+	int32 ActorIndex;
+	if (!Actors.Find(Actor, ActorIndex))
 	{
-		Actor->RegisterAllComponents();
-		Actor->RerunConstructionScripts();
-		GetWorld()->UpdateCullDistanceVolumes(Actor);
-		Actor->MarkComponentsRenderStateDirty();
-	}
+		Actors.Add(Actor);
+		ActorsForGC.Add(Actor);
 
-	OnLoadedActorAddedToLevelEvent.Broadcast(*Actor);
+		// If components from actors belonging to this level were already registered, do the same for the newly loaded actor.
+		// Otherwise, the new actor components will be registered later once the world initialization is completed, through UpdateWorldComponents()
+		if (bAreComponentsCurrentlyRegistered)
+		{
+			Actor->RegisterAllComponents();
+			Actor->RerunConstructionScripts();
+			GetWorld()->UpdateCullDistanceVolumes(Actor);
+			Actor->MarkComponentsRenderStateDirty();
+		}
+
+		OnLoadedActorAddedToLevelEvent.Broadcast(*Actor);
+	}
 }
 
 void ULevel::RemoveLoadedActor(AActor* Actor)
