Clarify naming of world size setting for World Partition auto loading of cells

diff --git a/Engine/Source/Editor/WorldPartitionEditor/Private/WorldPartitionEditorModule.cpp b/Engine/Source/Editor/WorldPartitionEditor/Private/WorldPartitionEditorModule.cpp
--- a/Engine/Source/Editor/WorldPartitionEditor/Private/WorldPartitionEditorModule.cpp
+++ b/Engine/Source/Editor/WorldPartitionEditor/Private/WorldPartitionEditorModule.cpp
@@ -141,9 +141,9 @@ void FWorldPartitionEditorModule::SetConversionPromptEnabled(bool bEnabled)
 	GetMutableDefault<UWorldPartitionEditorSettings>()->bEnableConversionPrompt = bEnabled;
 }
 
-float FWorldPartitionEditorModule::GetMinimumWorldSize() const
+float FWorldPartitionEditorModule::GetAutoCellLoadingMaxWorldSize() const
 {
-	return GetDefault<UWorldPartitionEditorSettings>()->MinimumWorldSize;
+	return GetDefault<UWorldPartitionEditorSettings>()->AutoCellLoadingMaxWorldSize;
 }
 
 bool FWorldPartitionEditorModule::ConvertMap(const FString& InLongPackageName)
@@ -203,7 +203,7 @@ bool FWorldPartitionEditorModule::ConvertMap(const FString& InLongPackageName)
 
 			// Try to provide complete Path, if we can't try with project name
 			FString ProjectPath = FPaths::IsProjectFilePathSet() ? FPaths::GetProjectFilePath() : FApp::GetProjectName();
-			
+
 			uint32 ProcessID;
 			FString Arguments = FString::Printf(TEXT("\"%s\" %s"), *ProjectPath, *DefaultConvertOptions->ToCommandletArgs());
 			ProcessHandle = FPlatformProcess::CreateProc(*CurrentExecutableName, *Arguments, true, false, false, &ProcessID, 0, nullptr, nullptr);
diff --git a/Engine/Source/Editor/WorldPartitionEditor/Public/WorldPartition/WorldPartitionEditorSettings.h b/Engine/Source/Editor/WorldPartitionEditor/Public/WorldPartition/WorldPartitionEditorSettings.h
--- a/Engine/Source/Editor/WorldPartitionEditor/Public/WorldPartition/WorldPartitionEditorSettings.h
+++ b/Engine/Source/Editor/WorldPartitionEditor/Public/WorldPartition/WorldPartitionEditorSettings.h
@@ -28,6 +28,6 @@ public:
 	UPROPERTY(config, EditAnywhere, Category = MapConversion, meta = (ToolTip = "Commandlet class to use for World Parition conversion"))
 	TSubclassOf<UWorldPartitionConvertCommandlet> CommandletClass;
 
-	UPROPERTY(config, EditAnywhere, Category = WorldPartition, meta = (ToolTip = "This option is used to automatically load all cells when the world is smaller than this value"))
-	float MinimumWorldSize = 100000.0f;
+	UPROPERTY(config, EditAnywhere, Category = WorldPartition, meta = (ToolTip = "Automatically load all cells when the world is smaller than this value"))
+	float AutoCellLoadingMaxWorldSize = 100000.0f;
 };
\ No newline at end of file
diff --git a/Engine/Source/Editor/WorldPartitionEditor/Public/WorldPartitionEditorModule.h b/Engine/Source/Editor/WorldPartitionEditor/Public/WorldPartitionEditorModule.h
--- a/Engine/Source/Editor/WorldPartitionEditor/Public/WorldPartitionEditorModule.h
+++ b/Engine/Source/Editor/WorldPartitionEditor/Public/WorldPartitionEditorModule.h
@@ -43,7 +43,7 @@ public:
 	/**
 	 *
 	 */
-	virtual float GetMinimumWorldSize() const override;
+	virtual float GetAutoCellLoadingMaxWorldSize() const override;
 
 	/**
 	 * 
diff --git a/Engine/Source/Runtime/Engine/Private/WorldPartition/WorldPartition.cpp b/Engine/Source/Runtime/Engine/Private/WorldPartition/WorldPartition.cpp
--- a/Engine/Source/Runtime/Engine/Private/WorldPartition/WorldPartition.cpp
+++ b/Engine/Source/Runtime/Engine/Private/WorldPartition/WorldPartition.cpp
@@ -262,15 +262,16 @@ void UWorldPartition::Initialize(UWorld* InWorld, const FTransform& InTransform)
 			// Load the always loaded cell, don't call LoadCells to avoid creating a transaction
 			UpdateLoadingEditorCell(EditorHash->GetAlwaysLoadedCell(), true);
 
-			// If the world is smaller than the project setting's minimum world size, load all cells
+			// Autoload all cells if the world is smaller than the project setting's value
 			IWorldPartitionEditorModule& WorldPartitionEditorModule = FModuleManager::LoadModuleChecked<IWorldPartitionEditorModule>("WorldPartitionEditor");
-			const float MinimumWorldSize = WorldPartitionEditorModule.GetMinimumWorldSize();
-			const bool bWorldIsSmallerThanMinimumWorldSize = GetWorldBounds().GetVolume() <= (MinimumWorldSize * MinimumWorldSize * MinimumWorldSize);
+			const float AutoCellLoadingMaxWorldSize = WorldPartitionEditorModule.GetAutoCellLoadingMaxWorldSize();
+			FVector WorldSize = GetWorldBounds().GetSize();
+			const bool bItsASmallWorld = WorldSize.X <= AutoCellLoadingMaxWorldSize && WorldSize.Y <= AutoCellLoadingMaxWorldSize && WorldSize.Z <= AutoCellLoadingMaxWorldSize;
 			
 			// When loading a subworld, load all actors
 			const bool bWorldIsSubPartition = !IsMainWorldPartition();
 
-			if (bWorldIsSubPartition || bWorldIsSmallerThanMinimumWorldSize)
+			if (bWorldIsSubPartition || bItsASmallWorld)
 			{
 				EditorHash->ForEachCell([this](UWorldPartitionEditorCell* Cell)
 				{
diff --git a/Engine/Source/Runtime/Engine/Public/WorldPartition/IWorldPartitionEditorModule.h b/Engine/Source/Runtime/Engine/Public/WorldPartition/IWorldPartitionEditorModule.h
--- a/Engine/Source/Runtime/Engine/Public/WorldPartition/IWorldPartitionEditorModule.h
+++ b/Engine/Source/Runtime/Engine/Public/WorldPartition/IWorldPartitionEditorModule.h
@@ -15,6 +15,6 @@ public:
 	virtual bool IsWorldPartitionEnabled() = 0;
 	virtual bool IsConversionPromptEnabled() const = 0;
 	virtual void SetConversionPromptEnabled(bool bEnabled) = 0;
-	virtual float GetMinimumWorldSize() const = 0;
+	virtual float GetAutoCellLoadingMaxWorldSize() const = 0;
 	virtual bool ConvertMap(const FString& InLongPackageName) = 0;
 };
\ No newline at end of file
