UE-106482 - Add missing include

diff --git a/Engine/Source/Runtime/Engine/Classes/EdGraph/EdGraphSchema.h b/Engine/Source/Runtime/Engine/Classes/EdGraph/EdGraphSchema.h
--- a/Engine/Source/Runtime/Engine/Classes/EdGraph/EdGraphSchema.h
+++ b/Engine/Source/Runtime/Engine/Classes/EdGraph/EdGraphSchema.h
@@ -9,6 +9,7 @@
 #include "EdGraph/EdGraphPin.h"
 #include "AssetData.h"
 #include "UObject/ObjectKey.h"
+#include "Input/Reply.h"
 #include "EdGraphSchema.generated.h"
 
 class FSlateRect;
