Added converstion of material anisotrpy to Strata Slab

diff --git a/Engine/Source/Runtime/Engine/Private/Materials/Material.cpp b/Engine/Source/Runtime/Engine/Private/Materials/Material.cpp
--- a/Engine/Source/Runtime/Engine/Private/Materials/Material.cpp
+++ b/Engine/Source/Runtime/Engine/Private/Materials/Material.cpp
@@ -75,6 +75,7 @@
 #include "UObject/UE5MainStreamObjectVersion.h"
 #include "Materials/MaterialExpressionAdd.h"
 #include "Materials/MaterialExpressionDivide.h"
+#include "Materials/MaterialExpressionComponentMask.h"
 #include "Materials/MaterialExpressionSingleLayerWaterMaterialOutput.h"
 #include "Materials/MaterialExpressionStrata.h"
 #include "Materials/StrataMaterial.h"
@@ -3625,11 +3626,38 @@ void UMaterial::ConvertMaterialToStrataMaterial()
 		else if (MaterialDomain == MD_Surface && ShadingModel == MSM_DefaultLit)
 		{
 			UMaterialExpressionStrataSlabBSDF* SlabBSDF = NewObject<UMaterialExpressionStrataSlabBSDF>(this);
+
+			if (Anisotropy.IsConnected())
+			{
+				UMaterialExpressionStrataAnisotropyToRoughness* AnisotropyToRoughness = NewObject<UMaterialExpressionStrataAnisotropyToRoughness>(this);
+				MoveConnectionTo(Roughness, AnisotropyToRoughness, 0);	// Roughness
+				CopyConnectionTo(Anisotropy, AnisotropyToRoughness, 1);	// Anisotropy
+
+				UMaterialExpressionComponentMask* RoughnessesMaskX = NewObject<UMaterialExpressionComponentMask>(this);
+				RoughnessesMaskX->GetInput(0)->Connect(0, AnisotropyToRoughness);
+				RoughnessesMaskX->R = 1;
+				RoughnessesMaskX->G = 0;
+				RoughnessesMaskX->B = 0;
+				RoughnessesMaskX->A = 0;
+
+				UMaterialExpressionComponentMask* RoughnessesMaskY = NewObject<UMaterialExpressionComponentMask>(this);
+				RoughnessesMaskY->GetInput(0)->Connect(0, AnisotropyToRoughness);
+				RoughnessesMaskY->R = 0;
+				RoughnessesMaskY->G = 1;
+				RoughnessesMaskY->B = 0;
+				RoughnessesMaskY->A = 0;
+
+				SlabBSDF->RoughnessX.Connect(0, RoughnessesMaskX);
+				SlabBSDF->RoughnessY.Connect(0, RoughnessesMaskY);
+			}
+			else
+			{
+				MoveConnectionTo(Roughness, SlabBSDF, 3);		// RoughnessX
+			}
+
 			MoveConnectionTo(BaseColor, SlabBSDF, 0);		// BaseColor
 			MoveConnectionTo(Metallic, SlabBSDF, 1);		// Metallic
 			MoveConnectionTo(Specular, SlabBSDF, 2);		// Specular
-			MoveConnectionTo(Roughness, SlabBSDF, 3);		// RoughnessX
-															// STRATA_TODO RoughnessY
 			MoveConnectionTo(Normal, SlabBSDF, 5);			// Normal
 			MoveConnectionTo(Tangent, SlabBSDF, 6);			// Tangent
 			// Opacity can remain on the end point node
@@ -3643,7 +3671,6 @@ void UMaterial::ConvertMaterialToStrataMaterial()
 			// STRATA_TODO AmbientOcclusion should be removed and put on the BSDF node
 			// STRATA_TODO Refraction
 			// STRATA_TODO PixelDepthOffset
-			// STRATA_TODO Anisotropy
 			// STRATA_TODO EmissiveColor
 
 			FrontMaterial.Connect(0, SlabBSDF);
