Remove swapping out rich curve keys in ::ReadjustTimeRange as this removes all non-time/value key information. Replacing it with 0-tolerance RemoveRedundantKeys instead which should deal with duplicate and overlapping keys.

diff --git a/Engine/Source/Runtime/Engine/Private/Curves/RichCurve.cpp b/Engine/Source/Runtime/Engine/Private/Curves/RichCurve.cpp
--- a/Engine/Source/Runtime/Engine/Private/Curves/RichCurve.cpp
+++ b/Engine/Source/Runtime/Engine/Private/Curves/RichCurve.cpp
@@ -764,14 +764,7 @@ void FRichCurve::ReadjustTimeRange(float NewMinTimeRange, float NewMaxTimeRange,
 		}
 	}
 
-	// now remove all redundant key
-	TArray<FRichCurveKey> NewKeys;
-	Exchange(NewKeys, Keys);
-
-	for(int32 KeyIndex=0; KeyIndex<NewKeys.Num(); ++KeyIndex)
-	{
-		UpdateOrAddKey(NewKeys[KeyIndex].Time, NewKeys[KeyIndex].Value);
-	}
+	RemoveRedundantKeys(0.f);
 
 	// now cull out all out of range 
 	float MinTime, MaxTime;
