Bumping animation DDC key to clean frosty anim problems

diff --git a/Engine/Source/Runtime/Engine/Private/Animation/AnimCompressionDerivedData.cpp b/Engine/Source/Runtime/Engine/Private/Animation/AnimCompressionDerivedData.cpp
--- a/Engine/Source/Runtime/Engine/Private/Animation/AnimCompressionDerivedData.cpp
+++ b/Engine/Source/Runtime/Engine/Private/Animation/AnimCompressionDerivedData.cpp
@@ -35,7 +35,7 @@ const TCHAR* FDerivedDataAnimationCompression::GetVersionString() const
 	// This is a version string that mimics the old versioning scheme. If you
 	// want to bump this version, generate a new guid using VS->Tools->Create GUID and
 	// return it here. Ex.
-	return TEXT("3B32894C1C3A45A883A84C1CB2F349F0");
+	return TEXT("565E584FC51D4ACB9C17A9F9C54CB545");	
 }
 
 bool FDerivedDataAnimationCompression::Build( TArray<uint8>& OutDataArray )
