Log mesh distance field resolutions when listing mesh distance fields

diff --git a/Engine/Source/Runtime/Engine/Private/DistanceFieldAtlas.cpp b/Engine/Source/Runtime/Engine/Private/DistanceFieldAtlas.cpp
--- a/Engine/Source/Runtime/Engine/Private/DistanceFieldAtlas.cpp
+++ b/Engine/Source/Runtime/Engine/Private/DistanceFieldAtlas.cpp
@@ -272,6 +272,7 @@ struct FMeshDistanceFieldStats
 {
 	size_t MemoryBytes;
 	float ResolutionScale;
+	FIntVector Resolution;
 	UStaticMesh* Mesh;
 };
 
@@ -287,6 +288,7 @@ void FDistanceFieldVolumeTextureAtlas::ListMeshDistanceFields() const
 		FMeshDistanceFieldStats Stats;
 		size_t AtlasMemory = Texture->VolumeData.Size.X * Texture->VolumeData.Size.Y * Texture->VolumeData.Size.Z * FormatSize;
 		size_t BackingMemory = Texture->VolumeData.CompressedDistanceFieldVolume.Num() * Texture->VolumeData.CompressedDistanceFieldVolume.GetTypeSize();
+		Stats.Resolution = Texture->VolumeData.Size;
 		Stats.MemoryBytes = AtlasMemory + BackingMemory;
 		Stats.Mesh = Texture->GetStaticMesh();
 #if WITH_EDITORONLY_DATA
@@ -316,12 +318,12 @@ void FDistanceFieldVolumeTextureAtlas::ListMeshDistanceFields() const
 	}
 
 	UE_LOG(LogStaticMesh, Log, TEXT("Dumping mesh distance fields for %u meshes, total %.1fMb"), GatheredStats.Num(), TotalMemory / 1024.0f / 1024.0f);
-	UE_LOG(LogStaticMesh, Log, TEXT("   Memory Mb, Scale, Name, Path"));
+	UE_LOG(LogStaticMesh, Log, TEXT("   Memory Mb, Resolution, Scale, Name, Path"));
 
 	for (int32 EntryIndex = 0; EntryIndex < GatheredStats.Num(); EntryIndex++)
 	{
 		const FMeshDistanceFieldStats& MeshStats = GatheredStats[EntryIndex];
-		UE_LOG(LogStaticMesh, Log, TEXT("   %.2f, %.1f, %s, %s"), MeshStats.MemoryBytes / 1024.0f / 1024.0f, MeshStats.ResolutionScale, *MeshStats.Mesh->GetName(), *MeshStats.Mesh->GetPathName());
+		UE_LOG(LogStaticMesh, Log, TEXT("   %.2f, %dx%dx%d, %.1f, %s, %s"), MeshStats.MemoryBytes / 1024.0f / 1024.0f, MeshStats.Resolution.X, MeshStats.Resolution.Y, MeshStats.Resolution.Z, MeshStats.ResolutionScale, *MeshStats.Mesh->GetName(), *MeshStats.Mesh->GetPathName());
 	}
 }
 
