Analytic partial derivatives, initial checkin (but disabled).

diff --git a/Engine/Shaders/Private/MaterialTemplate.ush b/Engine/Shaders/Private/MaterialTemplate.ush
--- a/Engine/Shaders/Private/MaterialTemplate.ush
+++ b/Engine/Shaders/Private/MaterialTemplate.ush
@@ -2002,6 +2002,25 @@ float3 DecodeSceneColorForMaterialNode(float2 ScreenUV)
 
 #endif // SHADING_PATH_DEFERRED
 
+
+// DERIV_BASE_VALUE is a disgusting macro to manage backwards compatibility while changing the generated materials. The existing nodes are not "aware"
+// of partial derivatives, and will use sprintf() to put the same line in both CalcPixelMaterialInputs() and CalcPixelMaterialInputsAnalyticDerivatives().
+// If we have a line like this in CalcPixelMaterialInputs():
+//     float Local0 = ...;
+// It will look like this in CalcPixelMaterialInputsAnalyticDerivatives()
+//     FloatDeriv Local0 = ...;
+//
+// That's ok, since the Local0 line is aware of derivatives. But if we have a line that isn't aware then we would want the CalcPixelMaterialInputs() version to be:
+//     float Local1 = Local0 + 3.2;
+// But the CalcPixelMaterialInputsAnalyticDerivatives() version would be:
+//     float Local1 = Local0.Value + 3.2;
+//
+// It's not possible to emit two different versions (with and without ".Value") unless we change every single material node. So the workaround is to always emit
+// DERIV_BASE_VALUE, and in one version #define it to ".Value" and in the other to "". That way, we can emit the same code to both functions. Ideally, once every
+// single node is at least aware of derivatives (and can output two different versions) then we can remove this #define and #undef for DERIV_BASE_VALUE
+
+#define DERIV_BASE_VALUE 
+
 // Uniform material expressions.
 %s
 
@@ -2359,6 +2378,9 @@ float3 CalculateAnisotropyTangent(FMaterialPixelParameters Parameters, FPixelMat
 	return Tangent;
 }
 
+// Be very, very careful about changing CalcPixelMaterialInputs() or CalcPixelMaterialInputsAnalyticDerivatives(). They apply the same basic calculation,
+// but CalcPixelMaterialInputs() uses finite difference derivates from hardware whereas CalcPixelMaterialInputsAnalyticDerivatives() uses analytic
+// derivatives. If you change anything, make sure that both functions stay identical. 
 void CalcPixelMaterialInputs(in out FMaterialPixelParameters Parameters, in out FPixelMaterialInputs PixelMaterialInputs)
 {
 	// Initial calculations (required for Normal)
@@ -2414,6 +2436,66 @@ void CalcPixelMaterialInputs(in out FMaterialPixelParameters Parameters, in out
 #endif
 }
 
+#undef DERIV_BASE_VALUE
+#define DERIV_BASE_VALUE .Value
+
+void CalcPixelMaterialInputsAnalyticDerivatives(in out FMaterialPixelParameters Parameters, in out FPixelMaterialInputs PixelMaterialInputs)
+{
+	// Initial calculations (required for Normal)
+%s
+	// The Normal is a special case as it might have its own expressions and also be used to calculate other inputs, so perform the assignment here
+%s
+
+#if TEMPLATE_USES_STRATA
+	Parameters.SharedNormals = StrataInitialiseSharedNormals();
+#endif
+
+	// Note that here MaterialNormal can be in world space or tangent space
+	float3 MaterialNormal = GetMaterialNormal(Parameters, PixelMaterialInputs);
+
+#if MATERIAL_TANGENTSPACENORMAL
+#if SIMPLE_FORWARD_SHADING
+	Parameters.WorldNormal = float3(0, 0, 1);
+#endif
+
+#if FEATURE_LEVEL >= FEATURE_LEVEL_SM4
+	// Mobile will rely on only the final normalize for performance
+	MaterialNormal = normalize(MaterialNormal);
+#endif
+
+	// normalizing after the tangent space to world space conversion improves quality with sheared bases (UV layout to WS causes shrearing)
+	// use full precision normalize to avoid overflows
+	Parameters.WorldNormal = TransformTangentNormalToWorld(Parameters.TangentToWorld, MaterialNormal);
+
+#else //MATERIAL_TANGENTSPACENORMAL
+
+	Parameters.WorldNormal = normalize(MaterialNormal);
+
+#endif //MATERIAL_TANGENTSPACENORMAL
+
+#if MATERIAL_TANGENTSPACENORMAL
+	// flip the normal for backfaces being rendered with a two-sided material
+	Parameters.WorldNormal *= Parameters.TwoSidedSign;
+#endif
+
+	Parameters.ReflectionVector = ReflectionAboutCustomWorldNormal(Parameters, Parameters.WorldNormal, false);
+
+#if !PARTICLE_SPRITE_FACTORY
+	Parameters.Particle.MotionBlurFade = 1.0f;
+#endif // !PARTICLE_SPRITE_FACTORY
+
+	// Now the rest of the inputs
+%s
+
+#if MATERIAL_USES_ANISOTROPY
+	Parameters.WorldTangent = CalculateAnisotropyTangent(Parameters, PixelMaterialInputs);
+#else
+	Parameters.WorldTangent = 0;
+#endif
+}
+
+#undef DERIV_BASE_VALUE
+
 // Programmatically set the line number after all the material inputs which have a variable number of line endings
 // This allows shader error line numbers after this point to be the same regardless of which material is being compiled
 #line %s
@@ -2635,7 +2717,12 @@ void CalcMaterialParametersEx(
 #endif
 
 	// Now that we have all the pixel-related parameters setup, calculate the Material Input/Attributes and Normal
+	// need to change this #if 1 to be some kind of nanite-related flag.
+#if 1
 	CalcPixelMaterialInputs(Parameters, PixelMaterialInputs);
+#else
+	CalcPixelMaterialInputsAnalyticDerivatives(Parameters, PixelMaterialInputs);
+#endif
 }
 
 // convenience function to setup CalcMaterialParameters assuming we don't support TranslatedWorldPositionExcludingShaderOffsets
diff --git a/Engine/Source/Runtime/Engine/Private/Materials/HLSLMaterialTranslator.cpp b/Engine/Source/Runtime/Engine/Private/Materials/HLSLMaterialTranslator.cpp
--- a/Engine/Source/Runtime/Engine/Private/Materials/HLSLMaterialTranslator.cpp
+++ b/Engine/Source/Runtime/Engine/Private/Materials/HLSLMaterialTranslator.cpp
@@ -7,6 +7,18 @@
 
 class Error;
 
+static TAutoConsoleVariable<int32> CVarIsAnalyticDerivEnabled(
+	TEXT("r.MaterialEditor.IsAnalyticDerivEnabled"),
+	0,
+	TEXT("True when analytic derivative code generation is enabled.\n"),
+	ECVF_Default);
+
+static bool IsAnalyticDerivEnabled()
+{
+	int32 Enabled = CVarIsAnalyticDerivEnabled.GetValueOnAnyThread();
+	return Enabled != 0;
+}
+
 #if WITH_EDITORONLY_DATA
 
 /** @return the number of components in a vector type. */
@@ -49,6 +61,573 @@ static inline int32 SwizzleComponentToIndex(TCHAR Component)
 	}
 }
 
+static const TCHAR * GetDerivBaseType(int Index)
+{
+	const TCHAR* Ret = TEXT("");
+	switch(Index)
+	{
+	case 0:
+		Ret = TEXT("FloatDeriv");
+		break;
+	case 1:
+		Ret = TEXT("FloatDeriv2");
+		break;
+	case 2:
+		Ret = TEXT("FloatDeriv3");
+		break;
+	case 3:
+		Ret = TEXT("FloatDeriv4");
+		break;
+	}
+	return Ret;
+}
+
+static const TCHAR * GetDerivRawType(int Index)
+{
+	const TCHAR* Ret = TEXT("");
+	switch(Index)
+	{
+	case 0:
+		Ret = TEXT("float");
+		break;
+	case 1:
+		Ret = TEXT("float2");
+		break;
+	case 2:
+		Ret = TEXT("float3");
+		break;
+	case 3:
+		Ret = TEXT("float4");
+		break;
+	}
+	return Ret;
+}
+
+
+FMaterialDerivativeAutogen::FMaterialDerivativeAutogen()
+{
+	Reset();
+}
+
+void FMaterialDerivativeAutogen::Reset()
+{
+	for (int32 Index = 0; Index < 4; Index++)
+	{
+		bCastDerivFromRawIsEnabled[Index] = false;
+		bCastDerivFromDerivIsEnabled[Index] = false;
+		bFullConstructorEnabled[Index] = false;
+		bCompactConstructorEnabled[Index] = false;
+	}
+
+	for (int32 Op = 0; Op < EF2_Num; Op++)
+	{
+		for (int32 Index = 0; Index < 4; Index++)
+		{
+			bFunc2OpIsEnabled[Op][Index] = false;
+		}
+	}
+}
+
+// 0: Float
+// 1: Float2
+// 2: Float3
+// 3: Float4
+// -1: Everything else
+static int32 GetDerivTypeIndex(EMaterialValueType ValueType)
+{
+	int32 Ret = INDEX_NONE;
+	if (ValueType == MCT_Float1 || ValueType == MCT_Float)
+	{
+		Ret = 0;
+	}
+	else if (ValueType == MCT_Float2)
+	{
+		Ret = 1;
+	}
+	else if (ValueType == MCT_Float3)
+	{
+		Ret = 2;
+	}
+	else if (ValueType == MCT_Float4)
+	{
+		Ret = 3;
+	}
+	else
+	{
+		Ret = -1; // redundant, but leaving it here so we can set breakpoints
+	}
+	return Ret;
+}
+
+
+static EMaterialValueType GetMaterialTypFromDerivTypeIndex(int32 Index)
+{
+	switch(Index)
+	{
+	case 0:
+		return MCT_Float;
+	case 1:
+		return MCT_Float2;
+	case 2:
+		return MCT_Float3;
+	case 3:
+		return MCT_Float4;
+	default:
+		check(0);
+		break;
+	}
+
+	return (EMaterialValueType)0; // invalid, should be a Float 1/2/3/4, break at the check(0);
+}
+
+// given a string, convert it from type to type
+FString FMaterialDerivativeAutogen::CoerceValueRaw(const FString& Token, int32 SrcType, EMaterialExpressionAnalyticDerivativeStatus SrcStatus, int32 DstType)
+{
+	FString Ret = Token;
+
+	// If the original value is a derivative, grab the raw value
+	if (SrcStatus == MEADS_Valid)
+	{
+		Ret = Ret + TEXT(".Value");
+	}
+
+	if (SrcType == DstType)
+	{
+		// no op
+	}
+	else
+	{
+		check(SrcType == 0); // can only coerce a float1
+		if (DstType == 0)
+		{
+			Ret = FString::Printf( TEXT("MaterialFloat(%s)"), *Ret);
+		}
+		else if (DstType == 1)
+		{
+			Ret = FString::Printf( TEXT("MaterialFloat2(%s,%s)"), *Ret, *Ret);
+		}
+		else if (DstType == 2)
+		{
+			Ret = FString::Printf( TEXT("MaterialFloat3(%s,%s,%s)"), *Ret, *Ret, *Ret);
+		}
+		else if (DstType == 3)
+		{
+			Ret = FString::Printf( TEXT("MaterialFloat4(%s,%s,%s,%s)"), *Ret, *Ret, *Ret, *Ret);
+		}
+		else
+		{
+			check(0);
+		}
+	}
+	return Ret;
+}
+
+// given a string, convert it from type to type
+FString FMaterialDerivativeAutogen::CoerceValueDeriv(const FString& Token, int32 SrcType, EMaterialExpressionAnalyticDerivativeStatus SrcStatus, int32 DstType)
+{
+	FString Ret = Token;
+
+	// If it's valid, then it's already a type. Otherwise, we need to convert it from raw to deriv
+	if (SrcStatus != MEADS_Valid)
+	{
+		FString SrcDerivType = GetDerivBaseType(SrcType);
+		bCastDerivFromRawIsEnabled[SrcType] = true;
+		Ret = TEXT("Make") + SrcDerivType + TEXT("(") + Ret + TEXT(")");
+	}
+
+	if (SrcType == DstType)
+	{
+		// no op
+	}
+	else
+	{
+		check(SrcType == 0); // can only coerce a float/float1
+		if (DstType == 0)
+		{
+			// should never happen
+			check(0);
+		}
+		else if (DstType == 1)
+		{
+			bCastDerivFromDerivIsEnabled[1] = true;
+			Ret = FString::Printf( TEXT("ExpandFloatDeriv2(%s)"), *Ret);
+		}
+		else if (DstType == 2)
+		{
+			bCastDerivFromDerivIsEnabled[2] = true;
+			Ret = FString::Printf( TEXT("ExpandFloatDeriv3(%s)"), *Ret);
+		}
+		else if (DstType == 3)
+		{
+			bCastDerivFromDerivIsEnabled[3] = true;
+			Ret = FString::Printf( TEXT("ExpandFloatDeriv4(%s)"), *Ret);
+		}
+		else
+		{
+			check(0);
+		}
+	}
+	return Ret;
+}
+
+FString FMaterialDerivativeAutogen::MakeDerivConstructor(const FString& Value, const FString& Ddx, const FString& Ddy, int32 DstType)
+{
+	check(DstType >= 0);
+	check(DstType < 4);
+
+	bFullConstructorEnabled[DstType] = true;
+	FString TypeName = GetDerivBaseType(DstType);
+	FString Ret = TEXT("Construct") + TypeName + TEXT("(") + Value + TEXT(",") + Ddx + TEXT(",") + Ddy + TEXT(")");
+	return Ret;
+}
+
+FString FMaterialDerivativeAutogen::MakeDerivConstructorFromFinite(const FString& Value, int32 DstType)
+{
+	check(DstType >= 0);
+	check(DstType < 4);
+
+	bCompactConstructorEnabled[DstType] = true;
+
+	FString TypeName = GetDerivBaseType(DstType);
+	FString Ret = TEXT("ConstructCompact") + TypeName + TEXT("(") + Value + TEXT(")");
+	return Ret;
+}
+
+
+int32 FMaterialDerivativeAutogen::GenerateExpressionFunc2(FHLSLMaterialTranslator& Translator, EFunc2 Op, int32 LhsCode, int32 RhsCode)
+{
+	if (LhsCode == INDEX_NONE || RhsCode == INDEX_NONE)
+	{
+		return INDEX_NONE;
+	}
+
+	const FShaderCodeChunk& LhsChunk = Translator.AtParameterCodeChunk(LhsCode);
+	const FShaderCodeChunk& RhsChunk = Translator.AtParameterCodeChunk(RhsCode);
+
+	int32 LhsTypeIndex = GetDerivTypeIndex(LhsChunk.Type);
+	int32 RhsTypeIndex = GetDerivTypeIndex(RhsChunk.Type);
+
+	int32 DstTypeIndex = FMath::Max<int32>(LhsTypeIndex,RhsTypeIndex);
+
+	EMaterialExpressionAnalyticDerivativeStatus DstStatus = MEADS_NotValid;
+
+	// if the initial type is different from the output type, then it's only valid if type is 0 (float). We can convert
+	// a float to a type with more components, but for example, we can't implicitly convert a float2 to a float3/float4.
+	if (LhsTypeIndex != DstTypeIndex && LhsTypeIndex != 0)
+	{
+		return INDEX_NONE;
+	}
+
+	if (RhsTypeIndex != DstTypeIndex && RhsTypeIndex != 0)
+	{
+		return INDEX_NONE;
+	}
+
+	// Rules for derivatives:
+	// 1. If either the LHS or RHS is Not Valid or Not Aware, then the derivative is not valid. Run scalar route.
+	// 2. If both LHS and RHS are known to be Zero, then run raw code, and specify a known zero status.
+	// 3. If both LHS and RHS are Valid derivatives, then run deriv path.
+	// 4. If one is Valid and the other is known Zero, then promote the Zero to Valid, and run deriv path.
+
+	bool bUseScalarVersion = false;
+	bool bIsDerivValidZero = true;
+	bool bMakeDerivLhs = false;
+	bool bMakeDerivRhs = false;
+	if (LhsChunk.DerivativeStatus == MEADS_NotAware || LhsChunk.DerivativeStatus == MEADS_NotValid ||
+		RhsChunk.DerivativeStatus == MEADS_NotAware || RhsChunk.DerivativeStatus == MEADS_NotValid)
+	{
+		// use scalar version as a fallback
+		bUseScalarVersion = true;
+		// derivative is not valid
+		bIsDerivValidZero = false;
+
+		// We output status as invalid, since one of the parameters is either not aware or not valid
+		DstStatus = MEADS_NotValid;
+
+	}
+	else if (LhsChunk.DerivativeStatus == MEADS_Zero && RhsChunk.DerivativeStatus == MEADS_Zero)
+	{
+		// use scalar version
+		bUseScalarVersion = true;
+		// since we know both incoming values have derivatives of zero, we know the output is zero
+		bIsDerivValidZero = true;
+		// we know that the value is zero
+		DstStatus = MEADS_Zero;
+	}
+	else
+	{
+		check(LhsChunk.DerivativeStatus == MEADS_Zero || LhsChunk.DerivativeStatus == MEADS_Valid);
+		check(RhsChunk.DerivativeStatus == MEADS_Zero || RhsChunk.DerivativeStatus == MEADS_Valid);
+
+		// use deriv version
+		bUseScalarVersion = false;
+		// we will be calculating a valid derivative, so this value doesn't matter, but make it false for consistency
+		bIsDerivValidZero = false;
+
+		// if the lhs has a derivitive of zero, and rhs is non-zero, convert lhs from a scalar type to deriv type
+		if (LhsChunk.DerivativeStatus == MEADS_Zero)
+		{
+			bMakeDerivLhs = true;
+		}
+
+		// if the rhs has a derivative of zero, and lhs is non-zero, convert rhs from a scalar type to deriv type
+		if (RhsChunk.DerivativeStatus == MEADS_Zero)
+		{
+			bMakeDerivRhs = true;
+		}
+
+		// derivative results will be valid
+		DstStatus = MEADS_Valid;
+	}
+
+	FString DstTokens[CompiledPDV_MAX];
+
+	for (int32 Index = 0; Index < CompiledPDV_MAX; Index++)
+	{
+		ECompiledPartialDerivativeVariation Variation = (ECompiledPartialDerivativeVariation)Index;
+
+		FString LhsToken = Translator.GetParameterCodeDeriv(LhsCode,Variation);
+		FString RhsToken = Translator.GetParameterCodeDeriv(RhsCode,Variation);
+
+		// The token is the symbol name. If we are in finite mode, that's all we have to do. But if
+		// we are in analytic mode, we may need to get the value.
+		if (Index == CompiledPDV_Analytic)
+		{
+			LhsToken = CoerceValueRaw(LhsToken,LhsTypeIndex,LhsChunk.DerivativeStatus,DstTypeIndex);
+			RhsToken = CoerceValueRaw(RhsToken,RhsTypeIndex,RhsChunk.DerivativeStatus,DstTypeIndex);
+		}
+
+		FString DstToken;
+		// just generate a type
+		switch(Op)
+		{
+		case EF2_Add:
+			DstToken = TEXT("(") + LhsToken + TEXT(" + ") + RhsToken + TEXT(")");
+			break;
+		case EF2_Sub:
+			DstToken = TEXT("(") + LhsToken + TEXT(" - ") + RhsToken + TEXT(")");
+			break;
+		case EF2_Mul:
+			DstToken = TEXT("(") + LhsToken + TEXT(" * ") + RhsToken + TEXT(")");
+			break;
+		case EF2_Div:
+			DstToken = TEXT("(") + LhsToken + TEXT(" / ") + RhsToken + TEXT(")");
+			break;
+		}
+
+		DstTokens[Index] = DstToken;
+	}
+
+	if (!bUseScalarVersion)
+	{
+		FString LhsToken = Translator.GetParameterCodeDeriv(LhsCode,CompiledPDV_Analytic);
+		FString RhsToken = Translator.GetParameterCodeDeriv(RhsCode,CompiledPDV_Analytic);
+
+		LhsToken = CoerceValueDeriv(LhsToken,LhsTypeIndex,LhsChunk.DerivativeStatus,DstTypeIndex);
+		RhsToken = CoerceValueDeriv(RhsToken,RhsTypeIndex,RhsChunk.DerivativeStatus,DstTypeIndex);
+
+		check(Op < EF2_Num);
+		bFunc2OpIsEnabled[Op][DstTypeIndex] = true;
+
+		FString DstToken;
+		switch(Op)
+		{
+		case EF2_Add:
+			DstToken = TEXT("AddDeriv(") + LhsToken + TEXT(",") + RhsToken + TEXT(")");
+			break;
+		case EF2_Sub:
+			DstToken = TEXT("SubDeriv(") + LhsToken + TEXT(",") + RhsToken + TEXT(")");
+			break;
+		case EF2_Mul:
+			DstToken = TEXT("MulDeriv(") + LhsToken + TEXT(",") + RhsToken + TEXT(")");
+			break;
+		case EF2_Div:
+			DstToken = TEXT("DivDeriv(") + LhsToken + TEXT(",") + RhsToken + TEXT(")");
+			break;
+		}
+
+		DstTokens[CompiledPDV_Analytic] = DstToken;
+	}
+
+	const uint64 Hash = CityHash64((char*)*DstTokens[CompiledPDV_FiniteDifferences], DstTokens[CompiledPDV_FiniteDifferences].Len() * sizeof(TCHAR));
+	EMaterialValueType DstMatType = GetMaterialTypFromDerivTypeIndex(DstTypeIndex);
+
+	int32 Ret = Translator.AddCodeChunkInnerDeriv(Hash,*DstTokens[CompiledPDV_FiniteDifferences],*DstTokens[CompiledPDV_Analytic],DstMatType,false /*?*/, DstStatus);
+	return Ret;
+}
+
+FString FMaterialDerivativeAutogen::GenerateUsedFunctions(FHLSLMaterialTranslator& Translator)
+{
+	FString Ret;
+
+	// the basic structs (FloatDeriv, FloatDeriv2, FloatDeriv3, FloatDeriv4)
+	// it's not worth keeping track of all the times these are used, just make them.
+	for (int32 Index = 0; Index < 4; Index++)
+	{
+		FString BaseName = GetDerivBaseType(Index);
+		FString FieldName = GetDerivRawType(Index);
+
+		Ret += TEXT("struct ") + BaseName + LINE_TERMINATOR;
+		Ret += TEXT("{") LINE_TERMINATOR;
+		Ret += TEXT("\t") + FieldName + TEXT(" Value;") LINE_TERMINATOR;
+		Ret += TEXT("\t") + FieldName + TEXT(" Ddx;") LINE_TERMINATOR;
+		Ret += TEXT("\t") + FieldName + TEXT(" Ddy;") LINE_TERMINATOR;
+		Ret += TEXT("};") LINE_TERMINATOR;
+		Ret += TEXT("") LINE_TERMINATOR;
+	}
+
+	// For building from scalars. Used when then previous value is known to have a derivative of zero for uniforms. I.e.
+	// float3 Local0 = .... 
+	// FloatDeriv3 Local1 = MakeFloatDeriv3(Local0);
+	for (int32 Index = 0; Index < 4; Index++)
+	{
+		if (bCastDerivFromRawIsEnabled[Index])
+		{
+			FString BaseName = GetDerivBaseType(Index);
+			FString FieldName = GetDerivRawType(Index);
+
+			Ret += BaseName + TEXT(" Make") + BaseName + TEXT("(") + FieldName + TEXT(" Value)") LINE_TERMINATOR;
+			Ret += TEXT("{") LINE_TERMINATOR;
+			Ret += TEXT("\t") + BaseName + TEXT(" Ret;") LINE_TERMINATOR;
+			Ret += TEXT("\tRet.Value = Value;") LINE_TERMINATOR;
+			Ret += TEXT("\tRet.Ddx = 0;") LINE_TERMINATOR;
+			Ret += TEXT("\tRet.Ddy = 0;") LINE_TERMINATOR;
+			Ret += TEXT("\treturn Ret;") LINE_TERMINATOR;
+			Ret += TEXT("}") LINE_TERMINATOR;
+			Ret += TEXT("") LINE_TERMINATOR;
+		}
+	}
+
+	// Conversion from FloatDeriv to FloatDeriv2/3/4
+	for (int32 Index = 0; Index < 4; Index++)
+	{
+		if (bCastDerivFromDerivIsEnabled[Index])
+		{
+			check(Index >= 1);
+
+			FString BaseName = GetDerivBaseType(Index);
+			FString ScalarName = GetDerivBaseType(0);
+
+			Ret += BaseName + TEXT(" Expand") + BaseName + TEXT("(") + ScalarName + TEXT(" Src)") LINE_TERMINATOR;
+			Ret += TEXT("{") LINE_TERMINATOR;
+			Ret += TEXT("\t") + BaseName + TEXT(" Ret;") LINE_TERMINATOR;
+			Ret += TEXT("\tRet.Value = Src.Value;") LINE_TERMINATOR;
+			Ret += TEXT("\tRet.Ddx = Src.Ddx;") LINE_TERMINATOR;
+			Ret += TEXT("\tRet.Ddy = Src.Ddy;") LINE_TERMINATOR;
+			Ret += TEXT("\treturn Ret;") LINE_TERMINATOR;
+			Ret += TEXT("}") LINE_TERMINATOR;
+			Ret += TEXT("") LINE_TERMINATOR;
+		}
+	}
+
+	for (int32 Index = 0; Index < 4; Index++)
+	{
+		if (bFullConstructorEnabled[Index])
+		{
+			FString BaseName = GetDerivBaseType(Index);
+			FString FieldName = GetDerivRawType(Index);
+
+			Ret += BaseName + TEXT(" Construct") + BaseName + TEXT("(") + FieldName + TEXT(" InValue,") + FieldName + TEXT(" InDdx,") + FieldName + TEXT(" InDdy)") LINE_TERMINATOR;
+			Ret += TEXT("{") LINE_TERMINATOR;
+			Ret += TEXT("\t") + BaseName + TEXT(" Ret;") LINE_TERMINATOR;
+			Ret += TEXT("\tRet.Value = InValue;") LINE_TERMINATOR;
+			Ret += TEXT("\tRet.Ddx = InDdx;") LINE_TERMINATOR;
+			Ret += TEXT("\tRet.Ddy = InDdy;") LINE_TERMINATOR;
+			Ret += TEXT("\treturn Ret;") LINE_TERMINATOR;
+			Ret += TEXT("}") LINE_TERMINATOR;
+			Ret += TEXT("") LINE_TERMINATOR;
+		}
+	}
+
+	for (int32 Index = 0; Index < 4; Index++)
+	{
+		if (bCompactConstructorEnabled[Index])
+		{
+			FString BaseName = GetDerivBaseType(Index);
+			FString FieldName = GetDerivRawType(Index);
+
+			Ret += BaseName + TEXT(" ConstructCompact") + BaseName + TEXT("(") + FieldName + TEXT(" InValue)") LINE_TERMINATOR;
+			Ret += TEXT("{") LINE_TERMINATOR;
+			Ret += TEXT("\t") + BaseName + TEXT(" Ret;") LINE_TERMINATOR;
+			Ret += TEXT("\tRet.Value = InValue;") LINE_TERMINATOR;
+			Ret += TEXT("\tRet.Ddx = ddx(InValue);") LINE_TERMINATOR;
+			Ret += TEXT("\tRet.Ddy = ddy(InValue);") LINE_TERMINATOR;
+			Ret += TEXT("\treturn Ret;") LINE_TERMINATOR;
+			Ret += TEXT("}") LINE_TERMINATOR;
+			Ret += TEXT("") LINE_TERMINATOR;
+		}
+	}
+
+	// Func2s
+	for (int32 Op = 0; Op < EF2_Num; Op++)
+	{
+		for (int32 Index = 0; Index < 4; Index++)
+		{
+			if (bFunc2OpIsEnabled[Op][Index])
+			{
+				FString BaseName = GetDerivBaseType(Index);
+				FString FieldName = GetDerivRawType(Index);
+
+				switch(Op)
+				{
+				case EF2_Add:
+					Ret += BaseName + TEXT(" AddDeriv(") + BaseName + TEXT(" Lhs, ") + BaseName + TEXT(" Rhs)") LINE_TERMINATOR;
+					Ret += TEXT("{") LINE_TERMINATOR;
+					Ret += TEXT("\t") + BaseName + TEXT(" Ret;") LINE_TERMINATOR;
+					Ret += TEXT("\tRet.Value = Lhs.Value + Rhs.Value;") LINE_TERMINATOR;
+					Ret += TEXT("\tRet.Ddx = Lhs.Ddx + Rhs.Ddx;") LINE_TERMINATOR;
+					Ret += TEXT("\tRet.Ddy = Lhs.Ddy + Rhs.Ddy;") LINE_TERMINATOR;
+					Ret += TEXT("\treturn Ret;") LINE_TERMINATOR;
+					Ret += TEXT("}") LINE_TERMINATOR;
+					Ret += TEXT("") LINE_TERMINATOR;
+					break;
+				case EF2_Sub:
+					Ret += BaseName + TEXT(" SubDeriv(") + BaseName + TEXT(" Lhs, ") + BaseName + TEXT(" Rhs)") LINE_TERMINATOR;
+					Ret += TEXT("{") LINE_TERMINATOR;
+					Ret += TEXT("\t") + BaseName + TEXT(" Ret;") LINE_TERMINATOR;
+					Ret += TEXT("\tRet.Value = Lhs.Value - Rhs.Value;") LINE_TERMINATOR;
+					Ret += TEXT("\tRet.Ddx = Lhs.Ddx - Rhs.Ddx;") LINE_TERMINATOR;
+					Ret += TEXT("\tRet.Ddy = Lhs.Ddy - Rhs.Ddy;") LINE_TERMINATOR;
+					Ret += TEXT("\treturn Ret;") LINE_TERMINATOR;
+					Ret += TEXT("}") LINE_TERMINATOR;
+					Ret += TEXT("") LINE_TERMINATOR;
+					break;
+				case EF2_Mul:
+					Ret += BaseName + TEXT(" MulDeriv(") + BaseName + TEXT(" Lhs, ") + BaseName + TEXT(" Rhs)") LINE_TERMINATOR;
+					Ret += TEXT("{") LINE_TERMINATOR;
+					Ret += TEXT("\t") + BaseName + TEXT(" Ret;") LINE_TERMINATOR;
+					Ret += TEXT("\tRet.Value = Lhs.Value * Rhs.Value;") LINE_TERMINATOR;
+					Ret += TEXT("\tRet.Ddx = Lhs.Ddx * Rhs.Value + Lhs.Value * Rhs.Ddx;") LINE_TERMINATOR;
+					Ret += TEXT("\tRet.Ddy = Lhs.Ddy * Rhs.Value + Lhs.Value * Rhs.Ddy;") LINE_TERMINATOR;
+					Ret += TEXT("\treturn Ret;") LINE_TERMINATOR;
+					Ret += TEXT("}") LINE_TERMINATOR;
+					Ret += TEXT("") LINE_TERMINATOR;
+					break;
+				case EF2_Div:
+					Ret += BaseName + TEXT(" DivDeriv(") + BaseName + TEXT(" Lhs, ") + BaseName + TEXT(" Rhs)") LINE_TERMINATOR;
+					Ret += TEXT("{") LINE_TERMINATOR;
+					Ret += TEXT("\t") + BaseName + TEXT(" Ret;") LINE_TERMINATOR;
+					Ret += TEXT("\tRet.Value = Lhs.Value * Rhs.Value;") LINE_TERMINATOR;
+					Ret += TEXT("\tRet.Ddx = (Lhs.Ddx * Rhs.Value - Lhs.Value * Rhs.Ddx) / (Rhs.Value * Rhs.Value);") LINE_TERMINATOR;
+					Ret += TEXT("\tRet.Ddy = (Lhs.Ddy * Rhs.Value - Lhs.Value * Rhs.Ddy) / (Rhs.Value * Rhs.Value);") LINE_TERMINATOR;
+					Ret += TEXT("\treturn Ret;") LINE_TERMINATOR;
+					Ret += TEXT("}") LINE_TERMINATOR;
+					Ret += TEXT("") LINE_TERMINATOR;
+					break;
+				default:
+					check(0);
+					break;
+				}
+			}
+		}
+	}
+
+	return Ret;
+}
+
+
+
+
 FHLSLMaterialTranslator::FHLSLMaterialTranslator(FMaterial* InMaterial,
 	FMaterialCompilationOutput& InMaterialCompilationOutput,
 	const FStaticParameterSet& InStaticParameters,
@@ -1008,73 +1587,81 @@ bool FHLSLMaterialTranslator::Translate()
 		MaterialCompilationOutput.NumUsedUVScalars = GetNumUserTexCoords() * 2;
 		MaterialCompilationOutput.NumUsedCustomInterpolatorScalars = CurrentCustomVertexInterpolatorOffset;
 
-		// Do Normal Chunk first
+		for (int32 VariationIter = 0; VariationIter < CompiledPDV_MAX; VariationIter++)
 		{
-			GetFixedParameterCode(
-				0,
-				NormalCodeChunkEnd,
-				Chunk[MP_Normal],
-				SharedPropertyCodeChunks[NormalShaderFrequency],
-				TranslatedCodeChunkDefinitions[MP_Normal],
-				TranslatedCodeChunks[MP_Normal]);
+			ECompiledPartialDerivativeVariation Variation = (ECompiledPartialDerivativeVariation)VariationIter;
 
-			// Always gather MP_Normal definitions as they can be shared by other properties
-			if (TranslatedCodeChunkDefinitions[MP_Normal].IsEmpty())
+			// Do Normal Chunk first
 			{
-				TranslatedCodeChunkDefinitions[MP_Normal] = GetDefinitions(SharedPropertyCodeChunks[NormalShaderFrequency], 0, NormalCodeChunkEnd);
+				GetFixedParameterCode(
+					0,
+					NormalCodeChunkEnd,
+					Chunk[MP_Normal],
+					SharedPropertyCodeChunks[NormalShaderFrequency],
+					DerivativeVariations[Variation].TranslatedCodeChunkDefinitions[MP_Normal],
+					DerivativeVariations[Variation].TranslatedCodeChunks[MP_Normal],
+					Variation);
+
+				// Always gather MP_Normal definitions as they can be shared by other properties
+				if (DerivativeVariations[Variation].TranslatedCodeChunkDefinitions[MP_Normal].IsEmpty())
+				{
+					DerivativeVariations[Variation].TranslatedCodeChunkDefinitions[MP_Normal] = GetDefinitions(SharedPropertyCodeChunks[NormalShaderFrequency], 0, NormalCodeChunkEnd, Variation);
+				}
 			}
-		}
 
-		// Now the rest, skipping Normal
-		for(uint32 PropertyId = 0; PropertyId < MP_MAX; ++PropertyId)
-		{
-			if (PropertyId == MP_MaterialAttributes || PropertyId == MP_Normal || PropertyId == MP_CustomOutput)
+			// Now the rest, skipping Normal
+			for(uint32 PropertyId = 0; PropertyId < MP_MAX; ++PropertyId)
 			{
-				continue;
-			}
+				if (PropertyId == MP_MaterialAttributes || PropertyId == MP_Normal || PropertyId == MP_CustomOutput)
+				{
+					continue;
+				}
 
-			const EShaderFrequency PropertyShaderFrequency = FMaterialAttributeDefinitionMap::GetShaderFrequency((EMaterialProperty)PropertyId);
+				const EShaderFrequency PropertyShaderFrequency = FMaterialAttributeDefinitionMap::GetShaderFrequency((EMaterialProperty)PropertyId);
 
-			int32 StartChunk = 0;
-			if (PropertyShaderFrequency == NormalShaderFrequency && SharedPixelProperties[PropertyId])
-			{
-				// When processing shared properties, do not generate the code before the Normal was generated as those are already handled
-				StartChunk = NormalCodeChunkEnd;
-			}
+				int32 StartChunk = 0;
+				if (PropertyShaderFrequency == NormalShaderFrequency && SharedPixelProperties[PropertyId])
+				{
+					// When processing shared properties, do not generate the code before the Normal was generated as those are already handled
+					StartChunk = NormalCodeChunkEnd;
+				}
 
-			GetFixedParameterCode(
-				StartChunk,
-				SharedPropertyCodeChunks[PropertyShaderFrequency].Num(),
-				Chunk[PropertyId],
-				SharedPropertyCodeChunks[PropertyShaderFrequency],
-				TranslatedCodeChunkDefinitions[PropertyId],
-				TranslatedCodeChunks[PropertyId]);
-		}
+				GetFixedParameterCode(
+					StartChunk,
+					SharedPropertyCodeChunks[PropertyShaderFrequency].Num(),
+					Chunk[PropertyId],
+					SharedPropertyCodeChunks[PropertyShaderFrequency],
+					DerivativeVariations[Variation].TranslatedCodeChunkDefinitions[PropertyId],
+					DerivativeVariations[Variation].TranslatedCodeChunks[PropertyId],
+					Variation);
 
-		for(uint32 PropertyId = MP_MAX; PropertyId < CompiledMP_MAX; ++PropertyId)
-		{
-			switch(PropertyId)
+			}
+
+			for(uint32 PropertyId = MP_MAX; PropertyId < CompiledMP_MAX; ++PropertyId)
 			{
-			case CompiledMP_EmissiveColorCS:
-			    if (bCompileForComputeShader)
-				{
-					GetFixedParameterCode(Chunk[PropertyId], SharedPropertyCodeChunks[SF_Compute], TranslatedCodeChunkDefinitions[PropertyId], TranslatedCodeChunks[PropertyId]);
-				}
-				break;
-			case CompiledMP_PrevWorldPositionOffset:
+				switch(PropertyId)
 				{
-					GetFixedParameterCode(Chunk[PropertyId], SharedPropertyCodeChunks[SF_Vertex], TranslatedCodeChunkDefinitions[PropertyId], TranslatedCodeChunks[PropertyId]);
+				case CompiledMP_EmissiveColorCS:
+					if (bCompileForComputeShader)
+					{
+						GetFixedParameterCode(Chunk[PropertyId], SharedPropertyCodeChunks[SF_Compute], DerivativeVariations[Variation].TranslatedCodeChunkDefinitions[PropertyId], DerivativeVariations[Variation].TranslatedCodeChunks[PropertyId], Variation);
+					}
+					break;
+				case CompiledMP_PrevWorldPositionOffset:
+					{
+						GetFixedParameterCode(Chunk[PropertyId], SharedPropertyCodeChunks[SF_Vertex], DerivativeVariations[Variation].TranslatedCodeChunkDefinitions[PropertyId], DerivativeVariations[Variation].TranslatedCodeChunks[PropertyId], Variation);
+					}
+					break;
+				default: check(0);
+					break;
 				}
-				break;
-			default: check(0);
-				break;
 			}
-		}
 
-		// Output the implementation for any custom output expressions
-		for (int32 ExpressionIndex = 0; ExpressionIndex < CustomOutputImplementations.Num(); ExpressionIndex++)
-		{
-			ResourcesString += CustomOutputImplementations[ExpressionIndex] + "\r\n\r\n";
+			// Output the implementation for any custom output expressions
+			for (int32 ExpressionIndex = 0; ExpressionIndex < DerivativeVariations[Variation].CustomOutputImplementations.Num(); ExpressionIndex++)
+			{
+				ResourcesString += DerivativeVariations[Variation].CustomOutputImplementations[ExpressionIndex] + "\r\n\r\n";
+			}
 		}
 
 		for (const FMaterialUniformExpression* ScalarExpression : UniformScalarExpressions)
@@ -1657,7 +2244,7 @@ TBitArray<> FHLSLMaterialTranslator::GetVertexInterpolatorsOffsets(FString& Vert
 	return AllocatedCoords;
 }
 
-void FHLSLMaterialTranslator::GetSharedInputsMaterialCode(FString& PixelMembersDeclaration, FString& NormalAssignment, FString& PixelMembersInitializationEpilog)
+void FHLSLMaterialTranslator::GetSharedInputsMaterialCode(FString& PixelMembersDeclaration, FString& NormalAssignment, FString& PixelMembersInitializationEpilog, ECompiledPartialDerivativeVariation DerivativeVariation)
 {
 	{
 		int32 LastProperty = -1;
@@ -1683,22 +2270,22 @@ void FHLSLMaterialTranslator::GetSharedInputsMaterialCode(FString& PixelMembersD
 			// Normal requires its own separate initializer
 			if (Property == MP_Normal)
 			{
-				NormalInitializerValue = FString::Printf(TEXT("\tPixelMaterialInputs.%s = %s;\n"), *PropertyName, *TranslatedCodeChunks[Property]);
+				NormalInitializerValue = FString::Printf(TEXT("\tPixelMaterialInputs.%s = %s;\n"), *PropertyName, *DerivativeVariations[DerivativeVariation].TranslatedCodeChunks[Property]);
 			}
 			else
 			{
-				if (TranslatedCodeChunkDefinitions[Property].Len() > 0)
+				if (DerivativeVariations[DerivativeVariation].TranslatedCodeChunkDefinitions[Property].Len() > 0)
 				{
 					if (LastProperty >= 0)
 					{
 						// Verify that all code chunks have the same contents
-						check(TranslatedCodeChunkDefinitions[Property].Len() == TranslatedCodeChunkDefinitions[LastProperty].Len());
+						check(DerivativeVariations[DerivativeVariation].TranslatedCodeChunkDefinitions[Property].Len() == DerivativeVariations[DerivativeVariation].TranslatedCodeChunkDefinitions[LastProperty].Len());
 					}
 
 					LastProperty = Property;
 				}
 
-				PixelInputInitializerValues += FString::Printf(TEXT("\tPixelMaterialInputs.%s = %s;\n"), *PropertyName, *TranslatedCodeChunks[Property]);
+				PixelInputInitializerValues += FString::Printf(TEXT("\tPixelMaterialInputs.%s = %s;\n"), *PropertyName, *DerivativeVariations[DerivativeVariation].TranslatedCodeChunks[Property]);
 			}
 
 			PixelMembersDeclaration += FString::Printf(TEXT("\t%s %s;\n"), HLSLTypeString(Type), *PropertyName);
@@ -1707,7 +2294,7 @@ void FHLSLMaterialTranslator::GetSharedInputsMaterialCode(FString& PixelMembersD
 		NormalAssignment = NormalInitializerValue;
 		if (LastProperty != -1)
 		{
-			PixelMembersInitializationEpilog += TranslatedCodeChunkDefinitions[LastProperty] + TEXT("\n");
+			PixelMembersInitializationEpilog += DerivativeVariations[DerivativeVariation].TranslatedCodeChunkDefinitions[LastProperty] + TEXT("\n");
 		}
 
 		if (CodeChunkToStrataSharedNormal.Num() > 0)
@@ -1777,22 +2364,38 @@ FString FHLSLMaterialTranslator::GetMaterialShaderCode()
 	LazyPrintf.PushParam(*MaterialAttributesDeclaration);
 
 	// Stores the shared shader results member declarations
-	FString PixelMembersDeclaration;
+	FString PixelMembersDeclaration[CompiledPDV_MAX];
 
-	FString NormalAssignment;
+	FString NormalAssignment[CompiledPDV_MAX];
 
 	// Stores the code to initialize all inputs after MP_Normal
-	FString PixelMembersSetupAndAssignments;
+	FString PixelMembersSetupAndAssignments[CompiledPDV_MAX];
 
-	GetSharedInputsMaterialCode(PixelMembersDeclaration, NormalAssignment, PixelMembersSetupAndAssignments);
+	for (int32 Index = 0; Index < CompiledPDV_MAX; Index++)
+	{
+		GetSharedInputsMaterialCode(PixelMembersDeclaration[Index], NormalAssignment[Index], PixelMembersSetupAndAssignments[Index], (ECompiledPartialDerivativeVariation)Index);
+	}
 
-	LazyPrintf.PushParam(*PixelMembersDeclaration);
+	// PixelMembersDeclaration should be the same for all variations, but might change in the future. There are cases where work is shared
+	// between the pixel and vertex shader, but with Nanite all work has to be moved into the pixel shader, which means we will want
+	// different inputs. But for now, we are keeping them the same.
+	LazyPrintf.PushParam(*PixelMembersDeclaration[CompiledPDV_FiniteDifferences]);
 
-	LazyPrintf.PushParam(*ResourcesString);
+	{
+		FString DerivativeHelpers = DerivativeAutogen.GenerateUsedFunctions(*this);
+		FString DerivativeHelpersAndResources = DerivativeHelpers + ResourcesString;
+		//LazyPrintf.PushParam(*ResourcesString);
+		LazyPrintf.PushParam(*DerivativeHelpersAndResources);
+	}
+
+	// Anything used bye the GenerationFunctionCode() like WorldPositionOffset shouldn't be using texures, right?
+	// Let those use the standard finite differences textures, since they should be the same. If we actually want
+	// those to handle texture reads properly, we'll have to make extra versions.
+	ECompiledPartialDerivativeVariation BaseDerivativeVariation = CompiledPDV_FiniteDifferences;
 
 	if (bCompileForComputeShader)
 	{
-		LazyPrintf.PushParam(*GenerateFunctionCode(CompiledMP_EmissiveColorCS));
+		LazyPrintf.PushParam(*GenerateFunctionCode(CompiledMP_EmissiveColorCS, BaseDerivativeVariation));
 	}
 	else
 	{
@@ -1815,15 +2418,15 @@ FString FHLSLMaterialTranslator::GetMaterialShaderCode()
 
 	LazyPrintf.PushParam(*FString::Printf(TEXT("return %.5f"), Material->GetOpacityMaskClipValue()));
 
-	LazyPrintf.PushParam(*GenerateFunctionCode(MP_WorldPositionOffset));
-	LazyPrintf.PushParam(*GenerateFunctionCode(CompiledMP_PrevWorldPositionOffset));
-	LazyPrintf.PushParam(*GenerateFunctionCode(MP_WorldDisplacement));
+	LazyPrintf.PushParam(*GenerateFunctionCode(MP_WorldPositionOffset, BaseDerivativeVariation));
+	LazyPrintf.PushParam(*GenerateFunctionCode(CompiledMP_PrevWorldPositionOffset, BaseDerivativeVariation));
+	LazyPrintf.PushParam(*GenerateFunctionCode(MP_WorldDisplacement, BaseDerivativeVariation));
 	LazyPrintf.PushParam(*FString::Printf(TEXT("return %.5f"), Material->GetMaxDisplacement()));
-	LazyPrintf.PushParam(*GenerateFunctionCode(MP_TessellationMultiplier));
-	LazyPrintf.PushParam(*GenerateFunctionCode(MP_CustomData0));
-	LazyPrintf.PushParam(*GenerateFunctionCode(MP_CustomData1));
+	LazyPrintf.PushParam(*GenerateFunctionCode(MP_TessellationMultiplier, BaseDerivativeVariation));
+	LazyPrintf.PushParam(*GenerateFunctionCode(MP_CustomData0, BaseDerivativeVariation));
+	LazyPrintf.PushParam(*GenerateFunctionCode(MP_CustomData1, BaseDerivativeVariation));
 
-	// Print custom texture coordinate assignments
+	// Print custom texture coordinate assignments, should be fine with regular derivatives
 	FString CustomUVAssignments;
 
 	int32 LastProperty = -1;
@@ -1831,18 +2434,18 @@ FString FHLSLMaterialTranslator::GetMaterialShaderCode()
 	{
 		if (CustomUVIndex == 0)
 		{
-			CustomUVAssignments += TranslatedCodeChunkDefinitions[MP_CustomizedUVs0 + CustomUVIndex];
+			CustomUVAssignments += DerivativeVariations[BaseDerivativeVariation].TranslatedCodeChunkDefinitions[MP_CustomizedUVs0 + CustomUVIndex];
 		}
 
-		if (TranslatedCodeChunkDefinitions[MP_CustomizedUVs0 + CustomUVIndex].Len() > 0)
+		if (DerivativeVariations[BaseDerivativeVariation].TranslatedCodeChunkDefinitions[MP_CustomizedUVs0 + CustomUVIndex].Len() > 0)
 		{
 			if (LastProperty >= 0)
 			{
-				check(TranslatedCodeChunkDefinitions[LastProperty].Len() == TranslatedCodeChunkDefinitions[MP_CustomizedUVs0 + CustomUVIndex].Len());
+				check(DerivativeVariations[BaseDerivativeVariation].TranslatedCodeChunkDefinitions[LastProperty].Len() == DerivativeVariations[BaseDerivativeVariation].TranslatedCodeChunkDefinitions[MP_CustomizedUVs0 + CustomUVIndex].Len());
 			}
 			LastProperty = MP_CustomizedUVs0 + CustomUVIndex;
 		}
-		CustomUVAssignments += FString::Printf(TEXT("\tOutTexCoords[%u] = %s;") LINE_TERMINATOR, CustomUVIndex, *TranslatedCodeChunks[MP_CustomizedUVs0 + CustomUVIndex]);
+		CustomUVAssignments += FString::Printf(TEXT("\tOutTexCoords[%u] = %s;") LINE_TERMINATOR, CustomUVIndex, *DerivativeVariations[BaseDerivativeVariation].TranslatedCodeChunks[MP_CustomizedUVs0 + CustomUVIndex]);
 	}
 
 	LazyPrintf.PushParam(*CustomUVAssignments);
@@ -1884,11 +2487,16 @@ FString FHLSLMaterialTranslator::GetMaterialShaderCode()
 
 	LazyPrintf.PushParam(*CustomInterpolatorAssignments);
 
-	// Initializers required for Normal
-	LazyPrintf.PushParam(*TranslatedCodeChunkDefinitions[MP_Normal]);
-	LazyPrintf.PushParam(*NormalAssignment);
-	// Finally the rest of common code followed by assignment into each input
-	LazyPrintf.PushParam(*PixelMembersSetupAndAssignments);
+	for (int32 Iter = 0; Iter < CompiledPDV_MAX; Iter++)
+	{
+		ECompiledPartialDerivativeVariation Variation = (ECompiledPartialDerivativeVariation)Iter;
+
+		// Initializers required for Normal
+		LazyPrintf.PushParam(*DerivativeVariations[Variation].TranslatedCodeChunkDefinitions[MP_Normal]);
+		LazyPrintf.PushParam(*NormalAssignment[Variation]);
+		// Finally the rest of common code followed by assignment into each input
+		LazyPrintf.PushParam(*PixelMembersSetupAndAssignments[Variation]);
+	}
 
 	LazyPrintf.PushParam(*FString::Printf(TEXT("%u"),MaterialTemplateLineNumber));
 
@@ -1926,24 +2534,85 @@ bool FHLSLMaterialTranslator::IsMaterialPropertyUsed(EMaterialProperty Property,
 				&& (NumComponents < 3 || Value.B == ReferenceValue.B)
 				&& (NumComponents < 4 || Value.A == ReferenceValue.A))
 			{
-				bPropertyUsed = false;
+				bPropertyUsed = false;
+			}
+		}
+	}
+
+	return bPropertyUsed;
+}
+
+// only used by GetMaterialShaderCode()
+// @param Index ECompiledMaterialProperty or EMaterialProperty
+FString FHLSLMaterialTranslator::GenerateFunctionCode(uint32 Index, ECompiledPartialDerivativeVariation Variation) const
+{
+	check(Index < CompiledMP_MAX);
+	return DerivativeVariations[Variation].TranslatedCodeChunkDefinitions[Index] + TEXT("	return ") + DerivativeVariations[Variation].TranslatedCodeChunks[Index] + TEXT(";");
+}
+
+
+const FShaderCodeChunk& FHLSLMaterialTranslator::AtParameterCodeChunk(int32 Index) const
+{
+	check(Index != INDEX_NONE);
+	checkf(Index >= 0 && Index < CurrentScopeChunks->Num(), TEXT("Index %d/%d, Platform=%d"), Index, CurrentScopeChunks->Num(), (int)Platform);
+	const FShaderCodeChunk& CodeChunk = (*CurrentScopeChunks)[Index];
+	return CodeChunk;
+}
+
+// Similar to GetParameterCode, but has no default, and is derivative aware. Making it a separate function in case it needs to diverge,
+// but after looking at it, it has the same logic as GetParameterCode() for now.
+FString FHLSLMaterialTranslator::GetParameterCodeDeriv(int32 Index, ECompiledPartialDerivativeVariation Variation)
+{
+	// In the case of a uniform expression, both finite and deriv versions are the same (raw floats, known zero)
+
+	checkf(Index >= 0 && Index < CurrentScopeChunks->Num(), TEXT("Index %d/%d, Platform=%d"), Index, CurrentScopeChunks->Num(), (int)Platform);
+	const FShaderCodeChunk& CodeChunk = (*CurrentScopeChunks)[Index];
+	if((CodeChunk.UniformExpression && CodeChunk.UniformExpression->IsConstant()) || CodeChunk.bInline)
+	{
+		// Constant uniform expressions and code chunks which are marked to be inlined are accessed via Definition
+		return CodeChunk.AtDefinition(Variation);
+	}
+	else
+	{
+		if (CodeChunk.UniformExpression)
+		{
+			// If the code chunk has a uniform expression, create a new code chunk to access it
+			const int32 AccessedIndex = AccessUniformExpression(Index);
+			const FShaderCodeChunk& AccessedCodeChunk = (*CurrentScopeChunks)[AccessedIndex];
+			if(AccessedCodeChunk.bInline)
+			{
+				// Handle the accessed code chunk being inlined
+				return AccessedCodeChunk.AtDefinition(Variation);
 			}
+			// Return the symbol used to reference this code chunk
+			check(AccessedCodeChunk.SymbolName.Len() > 0);
+			return AccessedCodeChunk.SymbolName;
 		}
-	}
 
-	return bPropertyUsed;
+		// Return the symbol used to reference this code chunk
+		check(CodeChunk.SymbolName.Len() > 0);
+		return CodeChunk.SymbolName;
+	}
 }
 
-// only used by GetMaterialShaderCode()
-// @param Index ECompiledMaterialProperty or EMaterialProperty
-FString FHLSLMaterialTranslator::GenerateFunctionCode(uint32 Index) const
+FString FHLSLMaterialTranslator::GetParameterCode(int32 Index, const TCHAR* Default)
 {
-	check(Index < CompiledMP_MAX);
-	return TranslatedCodeChunkDefinitions[Index] + TEXT("	return ") + TranslatedCodeChunks[Index] + TEXT(";");
+	FString Ret = GetParameterCodeRaw(Index,Default);
+
+	check(Index >= 0 && Index < CurrentScopeChunks->Num());
+	const FShaderCodeChunk&	CodeChunk = (*CurrentScopeChunks)[Index];
+
+	// inline types are forbidden from having an analytical derivative
+	if (CodeChunk.DerivativeStatus == MEADS_Valid) // && !CodeChunk.bInline)
+	{
+		Ret = Ret + TEXT(" DERIV_BASE_VALUE");
+	}
+
+	return Ret;
 }
 
 // GetParameterCode
-FString FHLSLMaterialTranslator::GetParameterCode(int32 Index, const TCHAR* Default)
+FString FHLSLMaterialTranslator::GetParameterCodeRaw(int32 Index, const TCHAR* Default)
 {
 	if(Index == INDEX_NONE && Default)
 	{
@@ -1955,7 +2624,7 @@ FString FHLSLMaterialTranslator::GetParameterCode(int32 Index, const TCHAR* Defa
 	if((CodeChunk.UniformExpression && CodeChunk.UniformExpression->IsConstant()) || CodeChunk.bInline)
 	{
 		// Constant uniform expressions and code chunks which are marked to be inlined are accessed via Definition
-		return CodeChunk.Definition;
+		return CodeChunk.DefinitionFinite;
 	}
 	else
 	{
@@ -1967,7 +2636,7 @@ FString FHLSLMaterialTranslator::GetParameterCode(int32 Index, const TCHAR* Defa
 			if(AccessedCodeChunk.bInline)
 			{
 				// Handle the accessed code chunk being inlined
-				return AccessedCodeChunk.Definition;
+				return AccessedCodeChunk.DefinitionFinite;
 			}
 			// Return the symbol used to reference this code chunk
 			check(AccessedCodeChunk.SymbolName.Len() > 0);
@@ -2002,7 +2671,7 @@ uint64 FHLSLMaterialTranslator::GetParameterHash(int32 Index)
 }
 
 /** Creates a string of all definitions needed for the given material input. */
-FString FHLSLMaterialTranslator::GetDefinitions(TArray<FShaderCodeChunk>& CodeChunks, int32 StartChunk, int32 EndChunk) const
+FString FHLSLMaterialTranslator::GetDefinitions(TArray<FShaderCodeChunk>& CodeChunks, int32 StartChunk, int32 EndChunk, ECompiledPartialDerivativeVariation Variation) const
 {
 	FString Definitions;
 	for (int32 ChunkIndex = StartChunk; ChunkIndex < EndChunk; ChunkIndex++)
@@ -2011,15 +2680,19 @@ FString FHLSLMaterialTranslator::GetDefinitions(TArray<FShaderCodeChunk>& CodeCh
 		// Uniform expressions (both constant and variable) and inline expressions don't have definitions.
 		if (!CodeChunk.UniformExpression && !CodeChunk.bInline)
 		{
-			Definitions += CodeChunk.Definition;
+			Definitions += CodeChunk.AtDefinition(Variation);
 		}
 	}
 	return Definitions;
 }
 
 // GetFixedParameterCode
-void FHLSLMaterialTranslator::GetFixedParameterCode(int32 StartChunk, int32 EndChunk, int32 ResultIndex, TArray<FShaderCodeChunk>& CodeChunks, FString& OutDefinitions, FString& OutValue)
+void FHLSLMaterialTranslator::GetFixedParameterCode(int32 StartChunk, int32 EndChunk, int32 ResultIndex, TArray<FShaderCodeChunk>& CodeChunks, FString& OutDefinitions, FString& OutValue, ECompiledPartialDerivativeVariation OriginalVariation)
 {
+	// Only allow the analytic variation for pixel shaders.
+	ECompiledPartialDerivativeVariation Variation = OriginalVariation;
+
+	// This function is hardcoded to finite differences for now.
 	if (ResultIndex != INDEX_NONE)
 	{
 		checkf(ResultIndex >= 0 && ResultIndex < CodeChunks.Num(), TEXT("Index out of range %d/%d [%s]"), ResultIndex, CodeChunks.Num(), *Material->GetFriendlyName());
@@ -2028,15 +2701,15 @@ void FHLSLMaterialTranslator::GetFixedParameterCode(int32 StartChunk, int32 EndC
 		{
 			// Handle a constant uniform expression being the only code chunk hooked up to a material input
 			const FShaderCodeChunk& ResultChunk = CodeChunks[ResultIndex];
-			OutValue = ResultChunk.Definition;
+			OutValue = ResultChunk.AtDefinition(Variation);
 		}
 		else
 		{
 			const FShaderCodeChunk& ResultChunk = CodeChunks[ResultIndex];
 			// Combine the definition lines and the return statement
 			check(ResultChunk.bInline || ResultChunk.SymbolName.Len() > 0);
-			OutDefinitions = GetDefinitions(CodeChunks, StartChunk, EndChunk);
-			OutValue = ResultChunk.bInline ? ResultChunk.Definition : ResultChunk.SymbolName;
+			OutDefinitions = GetDefinitions(CodeChunks, StartChunk, EndChunk, Variation);
+			OutValue = ResultChunk.bInline ? ResultChunk.AtDefinition(Variation) : ResultChunk.SymbolName;
 		}
 	}
 	else
@@ -2045,9 +2718,9 @@ void FHLSLMaterialTranslator::GetFixedParameterCode(int32 StartChunk, int32 EndC
 	}
 }
 
-void FHLSLMaterialTranslator::GetFixedParameterCode(int32 ResultIndex, TArray<FShaderCodeChunk>& CodeChunks, FString& OutDefinitions, FString& OutValue)
+void FHLSLMaterialTranslator::GetFixedParameterCode(int32 ResultIndex, TArray<FShaderCodeChunk>& CodeChunks, FString& OutDefinitions, FString& OutValue, ECompiledPartialDerivativeVariation Variation)
 {
-	GetFixedParameterCode(0, CodeChunks.Num(), ResultIndex, CodeChunks, OutDefinitions, OutValue);
+	GetFixedParameterCode(0, CodeChunks.Num(), ResultIndex, CodeChunks, OutDefinitions, OutValue, Variation);
 }
 
 /** Used to get a user friendly type from EMaterialValueType */
@@ -2100,6 +2773,31 @@ const TCHAR* FHLSLMaterialTranslator::HLSLTypeString(EMaterialValueType Type) co
 	};
 }
 
+/** Used to get an HLSL type from EMaterialValueType */
+const TCHAR* FHLSLMaterialTranslator::HLSLTypeStringDeriv(EMaterialValueType Type, EMaterialExpressionAnalyticDerivativeStatus DerivativeStatus) const
+{
+	switch(Type)
+	{
+	case MCT_Float1:				return (DerivativeStatus == MEADS_Valid) ? TEXT("FloatDeriv") : TEXT("MaterialFloat");
+	case MCT_Float2:				return (DerivativeStatus == MEADS_Valid) ? TEXT("FloatDeriv2") : TEXT("MaterialFloat2");
+	case MCT_Float3:				return (DerivativeStatus == MEADS_Valid) ? TEXT("FloatDeriv3") : TEXT("MaterialFloat3");
+	case MCT_Float4:				return (DerivativeStatus == MEADS_Valid) ? TEXT("FloatDeriv4") : TEXT("MaterialFloat4");
+	case MCT_Float:					return (DerivativeStatus == MEADS_Valid) ? TEXT("FloatDeriv") : TEXT("MaterialFloat");
+	case MCT_Texture2D:				return TEXT("texture2D");
+	case MCT_TextureCube:			return TEXT("textureCube");
+	case MCT_Texture2DArray:		return TEXT("texture2DArray");
+	case MCT_VolumeTexture:			return TEXT("volumeTexture");
+	case MCT_StaticBool:			return TEXT("static bool");
+	case MCT_MaterialAttributes:	return TEXT("FMaterialAttributes");
+	case MCT_TextureExternal:		return TEXT("TextureExternal");
+	case MCT_TextureVirtual:		return TEXT("TextureVirtual");
+	case MCT_VTPageTableResult:		return TEXT("VTPageTableResult");
+	case MCT_ShadingModel:			return TEXT("uint");
+	case MCT_Strata:				return TEXT("FStrataData");
+	default:						return TEXT("unknown");
+	};
+}
+
 int32 FHLSLMaterialTranslator::NonPixelShaderExpressionError()
 {
 	return Errorf(TEXT("Invalid node used in vertex/hull/domain shader input!"));
@@ -2167,7 +2865,7 @@ int32 FHLSLMaterialTranslator::AddCodeChunkInner(uint64 Hash, const TCHAR* Forma
 	{
 		const int32 CodeIndex = CurrentScopeChunks->Num();
 		// Adding an inline code chunk, the definition will be the code to inline
-		new(*CurrentScopeChunks) FShaderCodeChunk(Hash, FormattedCode,TEXT(""),Type,true);
+		new(*CurrentScopeChunks) FShaderCodeChunk(Hash, FormattedCode,FormattedCode,TEXT(""),Type,true);
 		return CodeIndex;
 	}
 	// Can only create temporaries for certain types
@@ -2186,9 +2884,11 @@ int32 FHLSLMaterialTranslator::AddCodeChunkInner(uint64 Hash, const TCHAR* Forma
 		// Allocate a local variable name
 		const FString SymbolName = CreateSymbolName(TEXT("Local"));
 		// Construct the definition string which stores the result in a temporary and adds a newline for readability
-		const FString LocalVariableDefinition = FString("	") + HLSLTypeString(Type) + TEXT(" ") + SymbolName + TEXT(" = ") + FormattedCode + TEXT(";") + LINE_TERMINATOR;
+		const FString LocalVariableDefinitionFinite = FString("	") + HLSLTypeString(Type) + TEXT(" ") + SymbolName + TEXT(" = ") + FormattedCode + TEXT(";") + LINE_TERMINATOR;
+		// Construct the definition string which stores the result in a temporary and adds a newline for readability
+		const FString LocalVariableDefinitionAnalytic = FString("	") + HLSLTypeString(Type) + TEXT(" ") + SymbolName + TEXT(" = ") + FormattedCode + TEXT(";") + LINE_TERMINATOR;
 		// Adding a code chunk that creates a local variable
-		new(*CurrentScopeChunks) FShaderCodeChunk(Hash, *LocalVariableDefinition,SymbolName,Type,false);
+		new(*CurrentScopeChunks) FShaderCodeChunk(Hash, *LocalVariableDefinitionFinite, *LocalVariableDefinitionAnalytic,SymbolName,Type,false);
 		return CodeIndex;
 	}
 	else
@@ -2207,6 +2907,71 @@ int32 FHLSLMaterialTranslator::AddCodeChunkInner(uint64 Hash, const TCHAR* Forma
 	}
 }
 
+/** Adds an already formatted inline or referenced code chunk, and notes the derivative status. */
+int32 FHLSLMaterialTranslator::AddCodeChunkInnerDeriv(uint64 Hash, const TCHAR* FormattedCodeFinite, const TCHAR* FormattedCodeAnalytic, EMaterialValueType Type, bool bLegacyInlined, EMaterialExpressionAnalyticDerivativeStatus DerivativeStatus)
+{
+	// If we have partial derivatives, we can't inline the code. Nearly all nodes use GetParameterCode() to grab the value, and then
+	// sprintf that string into the new expression. Unfortunately, that means that nodes have to temporarily use the exact same string for both
+	// the analytic and the finite differences path. The solution there is to #define DERIV_BASE_VALUE to ".Value" or "", so that both strings
+	// can be the same, and the .Value is optionally added to the right path.
+	//
+	// However, that only works on symbols, but does not work on inline expressions. So we have to disable inlining in all analytic derivative options
+	// until all nodes are converted over, and can output a separate Finite and Analytic version, even if that node treats the output derivative
+	// as invalid.
+	bool bInlined = false;
+	check(bAllowCodeChunkGeneration);
+
+	if (Type == MCT_Unknown)
+	{
+		return INDEX_NONE;
+	}
+
+	if (bInlined)
+	{
+		const int32 CodeIndex = CurrentScopeChunks->Num();
+		// Adding an inline code chunk, the definition will be the code to inline
+		new(*CurrentScopeChunks) FShaderCodeChunk(Hash, FormattedCodeFinite,FormattedCodeAnalytic,TEXT(""),Type,true,DerivativeStatus);
+		return CodeIndex;
+	}
+	// Can only create temporaries for certain types
+	else if ((Type & (MCT_Float | MCT_VTPageTableResult)) || Type == MCT_ShadingModel || Type == MCT_MaterialAttributes || Type == MCT_Strata)
+	{
+		// Check for existing
+		for (int32 i = 0; i < CurrentScopeChunks->Num(); ++i)
+		{
+			if ((*CurrentScopeChunks)[i].Hash == Hash)
+			{
+				return i;
+			}
+		}
+
+		const int32 CodeIndex = CurrentScopeChunks->Num();
+		// Allocate a local variable name
+		const FString SymbolName = CreateSymbolName(TEXT("Local"));
+		// Construct the definition string which stores the result in a temporary and adds a newline for readability
+		const FString LocalVariableDefinitionFinite = FString("	") + HLSLTypeString(Type) + TEXT(" ") + SymbolName + TEXT(" = ") + FormattedCodeFinite + TEXT(";") + LINE_TERMINATOR;
+		// Analytic version too
+		const FString LocalVariableDefinitionAnalytic = FString("	") + HLSLTypeStringDeriv(Type,DerivativeStatus) + TEXT(" ") + SymbolName + TEXT(" = (") + FormattedCodeAnalytic + TEXT("); // deriv") + LINE_TERMINATOR;
+		// Adding a code chunk that creates a local variable
+		new(*CurrentScopeChunks) FShaderCodeChunk(Hash, *LocalVariableDefinitionFinite, *LocalVariableDefinitionAnalytic,SymbolName,Type,false,DerivativeStatus);
+		return CodeIndex;
+	}
+	else
+	{
+		if (Type & MCT_Texture)
+		{
+			return Errorf(TEXT("Operation not supported on a Texture"));
+		}
+
+		if (Type == MCT_StaticBool)
+		{
+			return Errorf(TEXT("Operation not supported on a Static Bool"));
+		}
+
+		return INDEX_NONE;
+	}
+}
+
 /** 
 	* Constructs the formatted code chunk and creates a new local variable definition from it. 
 	* This should be used over AddInlinedCodeChunk when the code chunk adds actual instructions, and especially when calling a function.
@@ -2411,12 +3176,13 @@ int32 FHLSLMaterialTranslator::AddUniformExpressionInner(uint64 Hash, FMaterialU
 
 	const int32 ReturnIndex = CurrentScopeChunks->Num();
 	// Create a new code chunk for the uniform expression
-	new(*CurrentScopeChunks) FShaderCodeChunk(Hash, UniformExpression, FormattedCode, Type);
+	// Note that uniforms have a known-zero derivative
+	new(*CurrentScopeChunks) FShaderCodeChunk(Hash, UniformExpression, FormattedCode, FormattedCode, Type, MEADS_Zero);
 
 	if (!bFoundExistingExpression)
 	{
 		// Add an entry to the material-wide list of uniform expressions
-		new(UniformExpressions) FShaderCodeChunk(Hash, UniformExpression, FormattedCode, Type);
+		new(UniformExpressions) FShaderCodeChunk(Hash, UniformExpression, FormattedCode, FormattedCode, Type, MEADS_Zero);
 	}
 
 	return ReturnIndex;
@@ -2491,6 +3257,8 @@ int32 FHLSLMaterialTranslator::AccessUniformExpression(int32 Index)
 	}
 	else if(CodeChunk.Type & MCT_Float)
 	{
+		check(CodeChunk.DerivativeStatus != MEADS_Valid);
+
 		const TCHAR* Mask;
 		switch(CodeChunk.Type)
 		{
@@ -2561,6 +3329,7 @@ FString FHLSLMaterialTranslator::CoerceParameter(int32 Index,EMaterialValueType
 	}
 	else
 	{
+		// The calling function wants a value (no derivs). So if it is a deriv type, return the .Value field. Otherwise, just return.
 		if( (CodeChunk.Type & DestType) && (CodeChunk.Type & MCT_Float) )
 		{
 			switch( DestType )
@@ -2579,7 +3348,7 @@ FString FHLSLMaterialTranslator::CoerceParameter(int32 Index,EMaterialValueType
 		}
 		else
 		{
-			Errorf(TEXT("Coercion failed: %s: %s -> %s"),*CodeChunk.Definition,DescribeType(CodeChunk.Type),DescribeType(DestType));
+			Errorf(TEXT("Coercion failed: %s: %s -> %s"),*CodeChunk.DefinitionFinite,DescribeType(CodeChunk.Type),DescribeType(DestType));
 			return TEXT("");
 		}
 	}
@@ -2981,6 +3750,10 @@ int32 FHLSLMaterialTranslator::ValidCast(int32 Code, EMaterialValueType DestType
 	}
 	else if((SourceType & MCT_Float) && (DestType & MCT_Float))
 	{
+		check(Code >= 0 && Code < CurrentScopeChunks->Num());
+		const FShaderCodeChunk& CodeChunk = (*CurrentScopeChunks)[Code];
+		check(CodeChunk.DerivativeStatus != MEADS_Valid);
+
 		const uint32 NumSourceComponents = GetNumComponents(SourceType);
 		const uint32 NumDestComponents = GetNumComponents(DestType);
 
@@ -3064,6 +3837,11 @@ int32 FHLSLMaterialTranslator::ForceCast(int32 Code, EMaterialValueType DestType
 	}
 	else if((SourceType & MCT_Float) && (DestType & MCT_Float))
 	{
+		check(Code >= 0 && Code < CurrentScopeChunks->Num());
+		const FShaderCodeChunk& CodeChunk = (*CurrentScopeChunks)[Code];
+		check(CodeChunk.DerivativeStatus != MEADS_Valid);
+		// this code is not legal with partial derivatives because we can't just add .rgb to the end
+
 		const uint32 NumSourceComponents = GetNumComponents(SourceType);
 		const uint32 NumDestComponents = GetNumComponents(DestType);
 
@@ -4221,13 +4999,25 @@ int32 FHLSLMaterialTranslator::TextureCoordinate(uint32 CoordinateIndex, bool Un
 		SampleCode = TEXT("Parameters.TexCoords[%u].xy");
 	}
 
-	// Note: inlining is important so that on GLES devices, where half precision is used in the pixel shader, 
-	// The UV does not get assigned to a half temporary in cases where the texture sample is done directly from interpolated UVs
-	return AddInlinedCodeChunk(
-			MCT_Float2,
-			*SampleCode,
-			CoordinateIndex
-			);
+	if (IsAnalyticDerivEnabled() && !UnMirrorU && !UnMirrorV)
+	{
+		FString SampleCodeFinite = FString::Printf(TEXT("Parameters.TexCoords[%u].xy"),CoordinateIndex);
+		FString SampleCodeAnalytic = DerivativeAutogen.MakeDerivConstructorFromFinite(SampleCodeFinite,1); // 1 = float2
+
+		const uint64 Hash = CityHash64((char*)*SampleCodeAnalytic, SampleCodeAnalytic.Len() * sizeof(TCHAR));
+		const int32 CodeIndex = AddCodeChunkInnerDeriv(Hash,*SampleCodeFinite,*SampleCodeAnalytic,MCT_Float2,false,MEADS_Valid);
+		return CodeIndex;
+	}
+	else
+	{
+		// Note: inlining is important so that on GLES devices, where half precision is used in the pixel shader, 
+		// The UV does not get assigned to a half temporary in cases where the texture sample is done directly from interpolated UVs
+		return AddInlinedCodeChunk(
+				MCT_Float2,
+				*SampleCode,
+				CoordinateIndex
+				);
+	}
 }
 
 static const TCHAR* GetVTAddressMode(TextureAddress Address)
@@ -4385,6 +5175,78 @@ uint32 FHLSLMaterialTranslator::AcquireVTStackIndex(
 	return StackIndex;
 }
 
+
+static FString ApplySamplerType(const FString& InSampleCode, EMaterialSamplerType SamplerType)
+{
+	FString DstSampleCode;
+
+	switch( SamplerType )
+	{
+	case SAMPLERTYPE_External:
+		DstSampleCode = FString::Printf(TEXT("ProcessMaterialExternalTextureLookup(%s)"), *InSampleCode);
+		break;
+
+	case SAMPLERTYPE_Color:
+		DstSampleCode = FString::Printf( TEXT("ProcessMaterialColorTextureLookup(%s)"), *InSampleCode );
+		break;
+	case SAMPLERTYPE_VirtualColor:
+		// has a mobile specific workaround
+		DstSampleCode = FString::Printf( TEXT("ProcessMaterialVirtualColorTextureLookup(%s)"), *InSampleCode );
+		break;
+
+	case SAMPLERTYPE_LinearColor:
+	case SAMPLERTYPE_VirtualLinearColor:
+		DstSampleCode = FString::Printf(TEXT("ProcessMaterialLinearColorTextureLookup(%s)"), *InSampleCode);
+		break;
+
+	case SAMPLERTYPE_Alpha:
+	case SAMPLERTYPE_VirtualAlpha:
+	case SAMPLERTYPE_DistanceFieldFont:
+		// Sampling a single channel texture in D3D9 gives: (G,G,G)
+		// Sampling a single channel texture in D3D11 gives: (G,0,0)
+		// This replication reproduces the D3D9 behavior in all cases.
+		DstSampleCode = FString::Printf( TEXT("(%s).rrrr"), *InSampleCode );
+		break;
+
+	case SAMPLERTYPE_Grayscale:
+	case SAMPLERTYPE_VirtualGrayscale:
+		// Sampling a greyscale texture in D3D9 gives: (G,G,G)
+		// Sampling a greyscale texture in D3D11 gives: (G,0,0)
+		// This replication reproduces the D3D9 behavior in all cases.
+		DstSampleCode = FString::Printf( TEXT("ProcessMaterialGreyscaleTextureLookup((%s).r).rrrr"), *InSampleCode );
+		break;
+
+	case SAMPLERTYPE_LinearGrayscale:
+	case SAMPLERTYPE_VirtualLinearGrayscale:
+		// Sampling a greyscale texture in D3D9 gives: (G,G,G)
+		// Sampling a greyscale texture in D3D11 gives: (G,0,0)
+		// This replication reproduces the D3D9 behavior in all cases.
+		DstSampleCode = FString::Printf(TEXT("ProcessMaterialLinearGreyscaleTextureLookup((%s).r).rrrr"), *InSampleCode);
+		break;
+
+	case SAMPLERTYPE_Normal:
+	case SAMPLERTYPE_VirtualNormal:
+		// Normal maps need to be unpacked in the pixel shader.
+		DstSampleCode = FString::Printf( TEXT("UnpackNormalMap(%s)"), *InSampleCode );
+		break;
+
+	case SAMPLERTYPE_Masks:
+	case SAMPLERTYPE_VirtualMasks:
+		DstSampleCode = InSampleCode;
+		break;
+
+	case SAMPLERTYPE_Data:
+		DstSampleCode = InSampleCode;
+		break;
+
+	default:
+		check(0);
+		break;
+	}
+
+	return DstSampleCode;
+}
+
 int32 FHLSLMaterialTranslator::TextureSample(
 	int32 TextureIndex,
 	int32 CoordinateIndex,
@@ -4487,6 +5349,161 @@ int32 FHLSLMaterialTranslator::TextureSample(
 		AutomaticViewMipBias = false;
 	}
 
+	FString DstLines[CompiledPDV_MAX];
+
+	// We're going to have to completely restructure this function in order to work with analytical partial derivatives. To make the
+	// transition easier, we are going to explicitly only support a subset of functionality, and if we are outside of that subset,
+	// fallback to the existing code.
+	if (Material->GetMaterialDomain() == MD_Surface &&
+		ShaderFrequency == SF_Pixel &&
+		!bVirtualTexture &&
+		TextureType == MCT_Texture2D &&
+		//TextureReferenceIndex == INDEX_NONE &&
+		MipValueMode == TMVM_None &&
+		SamplerSource == SSM_FromTextureAsset)
+	{
+		FString TextureName = CoerceParameter(TextureIndex, MCT_Texture2D);
+
+		EMaterialValueType UVsType = (TextureType == MCT_TextureCube || TextureType == MCT_Texture2DArray || TextureType == MCT_VolumeTexture) ? MCT_Float3 : MCT_Float2;
+		const FShaderCodeChunk& UvChunk = AtParameterCodeChunk(CoordinateIndex);
+
+		const FString TokenUVsFinite = GetParameterCodeDeriv(CoordinateIndex,CompiledPDV_FiniteDifferences);
+		const FString TokenUVsAnalytic = GetParameterCodeDeriv(CoordinateIndex,CompiledPDV_Analytic);
+		check(TokenUVsFinite.Len() > 0);
+		check(TokenUVsAnalytic.Len() > 0);
+
+		FString SamplerStateCode;
+		bool RequiresManualViewMipBias = AutomaticViewMipBias;
+
+		// Instead of concatenating %s printf data, just build the full string one piece at a time
+		if (SamplerSource == SSM_FromTextureAsset)
+		{
+			SamplerStateCode = TextureName + TEXT("Sampler");
+		}
+		else
+		{
+			// add more options later, but just do this one for now
+			check(0);
+		}
+
+		FString SampleInstruction;
+		if (TextureType == MCT_TextureCube)
+		{
+			SampleInstruction += TEXT("TextureCubeSample");
+		}
+		else if (TextureType == MCT_Texture2DArray)
+		{
+			SampleInstruction += TEXT("Texture2DArraySample");
+		}
+		else if (TextureType == MCT_VolumeTexture)
+		{
+			SampleInstruction += TEXT("Texture3DSample");
+		}
+		else if (TextureType == MCT_TextureExternal)
+		{
+			SampleInstruction += TEXT("TextureExternalSample");
+		}
+		else if (bVirtualTexture)
+		{
+			SampleInstruction += TEXT("TextureVirtualSample");
+		}
+		else if (TextureType == MCT_Texture2D)
+		{
+			SampleInstruction += TEXT("Texture2DSample");
+		}
+		else
+		{
+			check(0);
+		}
+
+		//DstLines[CompiledPDV_FiniteDifferences] = SampleInstruction;
+		//DstLines[CompiledPDV_Analytic] = SampleInstruction + TEXT("Grad");
+
+		// FiniteDifferences
+		{
+			// In the case of finite differences, just use the raw token
+			FString CurrLine;
+			if (MipValueMode == TMVM_None)
+			{
+				CurrLine = SampleInstruction + TEXT("(") + TextureName + TEXT(",") + SamplerStateCode + TEXT(",") + TokenUVsFinite + TEXT(")");
+			}
+			else
+			{
+				check(0);
+			}
+			DstLines[CompiledPDV_FiniteDifferences] = CurrLine;
+		}
+
+		// Analytic
+		{
+			FString DdxUVs;
+			FString DdyUVs;
+			FString ValueUVs;
+
+			FString ZeroFloat = (UVsType == MCT_Float2) ? TEXT("float2(0,0)") : TEXT("float3(0,0,0)");
+
+			switch(UvChunk.DerivativeStatus)
+			{
+			case MEADS_NotAware:
+			case MEADS_NotValid:
+				// not aware or not valid, use ddx/ddy to get our best estimate.
+				DdxUVs = TEXT("ddx(") + TokenUVsAnalytic + TEXT(")");
+				DdyUVs = TEXT("ddy(") + TokenUVsAnalytic + TEXT(")");
+				ValueUVs = TokenUVsAnalytic;
+				break;
+			case MEADS_Zero:
+				// if we get here we have something like a constant being used as a UV.
+				DdxUVs = ZeroFloat;
+				DdyUVs = ZeroFloat;
+				ValueUVs = TokenUVsAnalytic;
+				break;
+			case MEADS_Valid:
+				DdxUVs = TokenUVsAnalytic + TEXT(".Ddx");
+				DdyUVs = TokenUVsAnalytic + TEXT(".Ddy");
+				ValueUVs = TokenUVsAnalytic + TEXT(".Value");
+				break;
+			default:
+				check(0);
+				break;
+			}
+
+			FString CurrLine;
+			if (MipValueMode == TMVM_None)
+			{
+				CurrLine = FString::Printf(TEXT("%sGrad(%s,%s,%s,%s,%s)"),
+					*SampleInstruction,
+					*TextureName,
+					*SamplerStateCode,
+					*ValueUVs,
+					*DdxUVs,
+					*DdyUVs);
+			}
+			else
+			{
+				check(0);
+			}
+			DstLines[CompiledPDV_Analytic] = CurrLine;
+		}
+
+		// This code is copy and pasted from below, but it's the same code for both Finite and Analytic, since it happens after the texture sample
+		for (int32 Index = 0; Index < CompiledPDV_MAX; Index++)
+		{
+			FString SampleCode = DstLines[Index];
+
+			SampleCode = ApplySamplerType(SampleCode, SamplerType);
+			DstLines[Index] = SampleCode;
+		}
+
+		const uint64 Hash = CityHash64((char*)*DstLines[CompiledPDV_FiniteDifferences], DstLines[CompiledPDV_FiniteDifferences].Len() * sizeof(TCHAR));
+
+		if (IsAnalyticDerivEnabled())
+		{
+			// always set status to not valid, although we could add an option to do 3 samples if it was important
+			int32 Ret = AddCodeChunkInnerDeriv(Hash,*DstLines[CompiledPDV_FiniteDifferences],*DstLines[CompiledPDV_Analytic],MCT_Float4,false, MEADS_NotValid);
+			return Ret;
+		}
+	}
+
 	FString SamplerStateCode;
 	bool RequiresManualViewMipBias = AutomaticViewMipBias;
 
@@ -4623,63 +5640,7 @@ int32 FHLSLMaterialTranslator::TextureSample(
 		}
 	}
 
-	switch( SamplerType )
-	{
-		case SAMPLERTYPE_External:
-			SampleCode = FString::Printf(TEXT("ProcessMaterialExternalTextureLookup(%s)"), *SampleCode);
-			break;
-
-		case SAMPLERTYPE_Color:
-			SampleCode = FString::Printf( TEXT("ProcessMaterialColorTextureLookup(%s)"), *SampleCode );
-			break;
-		case SAMPLERTYPE_VirtualColor:
-			// has a mobile specific workaround
-			SampleCode = FString::Printf( TEXT("ProcessMaterialVirtualColorTextureLookup(%s)"), *SampleCode );
-			break;
-
-		case SAMPLERTYPE_LinearColor:
-		case SAMPLERTYPE_VirtualLinearColor:
-			SampleCode = FString::Printf(TEXT("ProcessMaterialLinearColorTextureLookup(%s)"), *SampleCode);
-		break;
-
-		case SAMPLERTYPE_Alpha:
-		case SAMPLERTYPE_VirtualAlpha:
-		case SAMPLERTYPE_DistanceFieldFont:
-			// Sampling a single channel texture in D3D9 gives: (G,G,G)
-			// Sampling a single channel texture in D3D11 gives: (G,0,0)
-			// This replication reproduces the D3D9 behavior in all cases.
-			SampleCode = FString::Printf( TEXT("(%s).rrrr"), *SampleCode );
-			break;
-			
-		case SAMPLERTYPE_Grayscale:
-		case SAMPLERTYPE_VirtualGrayscale:
-			// Sampling a greyscale texture in D3D9 gives: (G,G,G)
-			// Sampling a greyscale texture in D3D11 gives: (G,0,0)
-			// This replication reproduces the D3D9 behavior in all cases.
-			SampleCode = FString::Printf( TEXT("ProcessMaterialGreyscaleTextureLookup((%s).r).rrrr"), *SampleCode );
-			break;
-
-		case SAMPLERTYPE_LinearGrayscale:
-		case SAMPLERTYPE_VirtualLinearGrayscale:
-			// Sampling a greyscale texture in D3D9 gives: (G,G,G)
-			// Sampling a greyscale texture in D3D11 gives: (G,0,0)
-			// This replication reproduces the D3D9 behavior in all cases.
-			SampleCode = FString::Printf(TEXT("ProcessMaterialLinearGreyscaleTextureLookup((%s).r).rrrr"), *SampleCode);
-			break;
-
-		case SAMPLERTYPE_Normal:
-		case SAMPLERTYPE_VirtualNormal:
-			// Normal maps need to be unpacked in the pixel shader.
-			SampleCode = FString::Printf( TEXT("UnpackNormalMap(%s)"), *SampleCode );
-			break;
-
-		case SAMPLERTYPE_Masks:
-		case SAMPLERTYPE_VirtualMasks:
-			break;
-
-		case SAMPLERTYPE_Data:
-			break;
-	}
+	SampleCode = ApplySamplerType(SampleCode, SamplerType);
 
 	FString TextureName;
 	int32 VirtualTextureIndex = INDEX_NONE;
@@ -5719,7 +6680,14 @@ int32 FHLSLMaterialTranslator::Add(int32 A,int32 B)
 	}
 	else
 	{
-		return AddCodeChunkWithHash(Hash, GetArithmeticResultType(A,B),TEXT("(%s + %s)"),*GetParameterCode(A),*GetParameterCode(B));
+		if (IsAnalyticDerivEnabled())
+		{
+			return DerivativeAutogen.GenerateExpressionFunc2(*this,FMaterialDerivativeAutogen::EF2_Add,A,B);
+		}
+		else
+		{
+			return AddCodeChunkWithHash(Hash, GetArithmeticResultType(A,B),TEXT("(%s + %s)"),*GetParameterCode(A),*GetParameterCode(B));
+		}
 	}
 }
 
@@ -5737,7 +6705,14 @@ int32 FHLSLMaterialTranslator::Sub(int32 A,int32 B)
 	}
 	else
 	{
-		return AddCodeChunkWithHash(Hash, GetArithmeticResultType(A,B),TEXT("(%s - %s)"),*GetParameterCode(A),*GetParameterCode(B));
+		if (IsAnalyticDerivEnabled())
+		{
+			return DerivativeAutogen.GenerateExpressionFunc2(*this,FMaterialDerivativeAutogen::EF2_Sub,A,B);
+		}
+		else
+		{
+			return AddCodeChunkWithHash(Hash, GetArithmeticResultType(A,B),TEXT("(%s - %s)"),*GetParameterCode(A),*GetParameterCode(B));
+		}
 	}
 }
 
@@ -5755,7 +6730,14 @@ int32 FHLSLMaterialTranslator::Mul(int32 A,int32 B)
 	}
 	else
 	{
-		return AddCodeChunkWithHash(Hash, GetArithmeticResultType(A,B),TEXT("(%s * %s)"),*GetParameterCode(A),*GetParameterCode(B));
+		if (IsAnalyticDerivEnabled())
+		{
+			return DerivativeAutogen.GenerateExpressionFunc2(*this,FMaterialDerivativeAutogen::EF2_Mul,A,B);
+		}
+		else
+		{
+			return AddCodeChunkWithHash(Hash, GetArithmeticResultType(A,B),TEXT("(%s * %s)"),*GetParameterCode(A),*GetParameterCode(B));
+		}
 	}
 }
 
@@ -5773,7 +6755,14 @@ int32 FHLSLMaterialTranslator::Div(int32 A,int32 B)
 	}
 	else
 	{
-		return AddCodeChunkWithHash(Hash, GetArithmeticResultType(A,B),TEXT("(%s / %s)"),*GetParameterCode(A),*GetParameterCode(B));
+		if (IsAnalyticDerivEnabled())
+		{
+			return DerivativeAutogen.GenerateExpressionFunc2(*this,FMaterialDerivativeAutogen::EF2_Div,A,B);
+		}
+		else
+		{
+			return AddCodeChunkWithHash(Hash, GetArithmeticResultType(A,B),TEXT("(%s / %s)"),*GetParameterCode(A),*GetParameterCode(B));
+		}
 	}
 }
 
@@ -8042,17 +9031,22 @@ int32 FHLSLMaterialTranslator::CustomOutput(class UMaterialExpressionCustomOutpu
 	FString Definitions;
 	FString Body;
 
+	// For now just grab the finite differences version, and use both for finite and analytic. Should fix later.
 	if ((*CurrentScopeChunks)[OutputCode].UniformExpression && !(*CurrentScopeChunks)[OutputCode].UniformExpression->IsConstant())
 	{
 		Body = GetParameterCode(OutputCode);
 	}
 	else
 	{
-		GetFixedParameterCode(OutputCode, *CurrentScopeChunks, Definitions, Body);
+		GetFixedParameterCode(OutputCode, *CurrentScopeChunks, Definitions, Body, CompiledPDV_FiniteDifferences);
 	}
 
-	FString ImplementationCode = FString::Printf(TEXT("%s %s%d(FMaterial%sParameters Parameters)\r\n{\r\n%s return %s;\r\n}\r\n"), *OutputTypeString, *Custom->GetFunctionName(), OutputIndex, ShaderFrequency == SF_Vertex ? TEXT("Vertex") : TEXT("Pixel"), *Definitions, *Body);
-	CustomOutputImplementations.Add(ImplementationCode);
+	FString ImplementationCodeFinite = FString::Printf(TEXT("%s %s%d(FMaterial%sParameters Parameters)\r\n{\r\n%s return %s;\r\n}\r\n"), *OutputTypeString, *Custom->GetFunctionName(), OutputIndex, ShaderFrequency == SF_Vertex ? TEXT("Vertex") : TEXT("Pixel"), *Definitions, *Body);
+	DerivativeVariations[CompiledPDV_FiniteDifferences].CustomOutputImplementations.Add(ImplementationCodeFinite);
+
+	// FIXEM: for now just skipping these, as they create invalid code.
+	//FString ImplementationCodeAnalytic = FString::Printf(TEXT("%s %s%d_Analytic(FMaterial%sParameters Parameters)\r\n{\r\n%s return %s;\r\n}\r\n"), *OutputTypeString, *Custom->GetFunctionName(), OutputIndex, ShaderFrequency == SF_Vertex ? TEXT("Vertex") : TEXT("Pixel"), *Definitions, *Body);
+	//DerivativeVariations[CompiledPDV_Analytic].CustomOutputImplementations.Add(ImplementationCodeAnalytic);
 
 	// return value is not used
 	return INDEX_NONE;
diff --git a/Engine/Source/Runtime/Engine/Private/Materials/HLSLMaterialTranslator.h b/Engine/Source/Runtime/Engine/Private/Materials/HLSLMaterialTranslator.h
--- a/Engine/Source/Runtime/Engine/Private/Materials/HLSLMaterialTranslator.h
+++ b/Engine/Source/Runtime/Engine/Private/Materials/HLSLMaterialTranslator.h
@@ -53,6 +53,8 @@ class Error;
 
 #if WITH_EDITORONLY_DATA
 
+class FHLSLMaterialTranslator;
+
 enum EMaterialExpressionVisitResult
 {
 	MVR_CONTINUE,
@@ -66,6 +68,22 @@ public:
 	virtual EMaterialExpressionVisitResult Visit(UMaterialExpression* InExpression) = 0;
 };
 
+/**
+* For a node, the known information of the partial derivatives.
+* MEADS_NotAware - This node is made by a function that has no knowledge of analytic partial derivatives.
+* MEADS_Zero - This node is aware of partial derivatives, and knows that it's value is zero, as is the case for uniform parameters.
+* MEADS_Valid - This node is aware of partial derivatives, and knows that it has a calculated value.
+* MEADS_NotValid - This node is aware of partial derivatives, and knows that one of its source inputs is not partial derivative aware, and therefore its value is not to be used.
+**/
+enum EMaterialExpressionAnalyticDerivativeStatus
+{
+	MEADS_NotAware,
+	MEADS_Zero,
+	MEADS_Valid,
+	MEADS_NotValid,
+	MEADS_MAX,
+};
+
 struct FShaderCodeChunk
 {
 	/**
@@ -77,8 +95,13 @@ struct FShaderCodeChunk
 	 * Definition string of the code chunk. 
 	 * If !bInline && !UniformExpression || UniformExpression->IsConstant(), this is the definition of a local variable named by SymbolName.
 	 * Otherwise if bInline || (UniformExpression && UniformExpression->IsConstant()), this is a code expression that needs to be inlined.
+	 * This string uses hardware finite differences.
 	 */
-	FString Definition;
+	FString DefinitionFinite;
+	/** 
+	* Definition string of the code chunk, but with analytic partial derivatives.
+	*/
+	FString DefinitionAnalytic;
 	/** 
 	 * Name of the local variable used to reference this code chunk. 
 	 * If bInline || UniformExpression, there will be no symbol name and Definition should be used directly instead.
@@ -89,24 +112,44 @@ struct FShaderCodeChunk
 	EMaterialValueType Type;
 	/** Whether the code chunk should be inlined or not.  If true, SymbolName is empty and Definition contains the code to inline. */
 	bool bInline;
+	/** The status of partial derivatives at this expression. **/
+	EMaterialExpressionAnalyticDerivativeStatus DerivativeStatus;
+
+	FString& AtDefinition(ECompiledPartialDerivativeVariation Variation)
+	{
+		checkf(Variation == CompiledPDV_FiniteDifferences || Variation == CompiledPDV_Analytic, TEXT("Invalid partial derivative variation: %d"), Variation);
+		return Variation == CompiledPDV_FiniteDifferences ? DefinitionFinite : DefinitionAnalytic;
+	}
+
+	const FString& AtDefinition(ECompiledPartialDerivativeVariation Variation) const
+	{
+		checkf(Variation == CompiledPDV_FiniteDifferences || Variation == CompiledPDV_Analytic, TEXT("Invalid partial derivative variation: %d"), Variation);
+		return Variation == CompiledPDV_FiniteDifferences ? DefinitionFinite : DefinitionAnalytic;
+	}
 
 	/** Ctor for creating a new code chunk with no associated uniform expression. */
-	FShaderCodeChunk(uint64 InHash, const TCHAR* InDefinition,const FString& InSymbolName,EMaterialValueType InType,bool bInInline):
+	FShaderCodeChunk(uint64 InHash, const TCHAR* InDefinitionFinite,const TCHAR* InDefinitionAnalytic,const FString& InSymbolName,EMaterialValueType InType,bool bInInline,
+					 EMaterialExpressionAnalyticDerivativeStatus InDerivativeStatus = MEADS_NotAware):
 		Hash(InHash),
-		Definition(InDefinition),
+		DefinitionFinite(InDefinitionFinite),
+		DefinitionAnalytic(InDefinitionAnalytic),
 		SymbolName(InSymbolName),
 		UniformExpression(NULL),
 		Type(InType),
-		bInline(bInInline)
+		bInline(bInInline),
+		DerivativeStatus(InDerivativeStatus)
 	{}
 
 	/** Ctor for creating a new code chunk with a uniform expression. */
-	FShaderCodeChunk(uint64 InHash, FMaterialUniformExpression* InUniformExpression,const TCHAR* InDefinition,EMaterialValueType InType):
+	FShaderCodeChunk(uint64 InHash, FMaterialUniformExpression* InUniformExpression,const TCHAR* InDefinitionFinite,const TCHAR* InDefinitionAnalytic,EMaterialValueType InType,
+					 EMaterialExpressionAnalyticDerivativeStatus InDerivativeStatus = MEADS_NotAware):
 		Hash(InHash),
-		Definition(InDefinition),
+		DefinitionFinite(InDefinitionFinite),
+		DefinitionAnalytic(InDefinitionAnalytic),
 		UniformExpression(InUniformExpression),
 		Type(InType),
-		bInline(false)
+		bInline(false),
+		DerivativeStatus(InDerivativeStatus)
 	{}
 };
 
@@ -139,9 +182,75 @@ struct FMaterialCustomExpressionEntry
 	TArray<int32> OutputCodeIndex;
 };
 
+struct FMaterialDerivativeAutogen
+{
+	// functions with one parameter
+	enum EFunc1
+	{
+		EF1_Sin,
+		EF1_Cos,
+		EF1_Num
+	};
+
+	enum EFunc2
+	{
+		EF2_Add,
+		EF2_Sub,
+		EF2_Mul,
+		EF2_Div,
+		EF2_Num
+	};
+
+	FMaterialDerivativeAutogen();
+
+	void Reset();
+
+	int32 GenerateExpressionFunc2(FHLSLMaterialTranslator& Translator, EFunc2 Op, int32 LhsCode, int32 RhsCode);
+	FString GenerateUsedFunctions(FHLSLMaterialTranslator& Translator);
+
+	FString CoerceValueRaw(const FString& Token, int32 SrcType, EMaterialExpressionAnalyticDerivativeStatus SrcStatus, int32 DstType);
+	FString CoerceValueDeriv(const FString& Token, int32 SrcType, EMaterialExpressionAnalyticDerivativeStatus SrcStatus, int32 DstType);
+
+	FString MakeDerivConstructor(const FString& Value, const FString& Ddx, const FString& Ddy, int32 DstType);
+
+	FString MakeDerivConstructorFromFinite(const FString& Value, int32 DstType);
+
+	// this could be packed easily into a bitmask, but keeping it as bools for clarity, and it's pipeline so the cost is irrelevant
+	bool bFunc2OpIsEnabled[EF2_Num][4];
+
+	// For example, used when we cast from a float4 to FloatDeriv4
+	bool bCastDerivFromRawIsEnabled[4];
+
+	// Used when we explicitly create a derivative with values. I.e. MakeFloatDeriv2(float2 Value, float2 Ddx, float2 Ddy)
+	bool bFullConstructorEnabled[4];
+
+	// Used when we explicitly create a derivative from ddx/ddy. I.e. MakeFloatDeriv2(float2 Value) // derivs from ddx(Value) and ddy(Value)
+	bool bCompactConstructorEnabled[4];
+
+	// For example, used when we cast from a FloatDeriv to FloatDeriv4. Note that bCastDerivFromDerivIsEnabled[0] will always be false,
+	// as we don't needto cast from a FloatDeriv to FloatDeriv, but leaving it to avoid have +1/-1s everywhere.
+	bool bCastDerivFromDerivIsEnabled[4];
+
+};
+
+struct FMaterialDerivativeVariation
+{
+	/** Code chunk definitions corresponding to each of the material inputs, only initialized after Translate has been called. */
+	FString TranslatedCodeChunkDefinitions[CompiledMP_MAX];
+
+	/** Code chunks corresponding to each of the material inputs, only initialized after Translate has been called. */
+	FString TranslatedCodeChunks[CompiledMP_MAX];
+
+	/** Any custom output function implementations */
+	TArray<FString> CustomOutputImplementations;
+
+};
+
 class FHLSLMaterialTranslator : public FMaterialCompiler
 {
 protected:
+	/** Data that is different for the finite and anlytical partial derivative options. **/
+	FMaterialDerivativeVariation DerivativeVariations[CompiledPDV_MAX];
 
 	/** The shader frequency of the current material property being compiled. */
 	EShaderFrequency ShaderFrequency;
@@ -159,6 +268,9 @@ protected:
 	// List of Shared pixel properties. Used to share generated code
 	bool SharedPixelProperties[CompiledMP_MAX];
 
+	/** Stores the resource declarations */
+	FString ResourcesString;
+
 	/* Stack that tracks compiler state specific to the function currently being compiled. */
 	TArray<FMaterialFunctionCompileState*> FunctionStacks[SF_NumFrequencies];
 
@@ -174,18 +286,9 @@ protected:
 	/** Feature level being compiled for. */
 	ERHIFeatureLevel::Type FeatureLevel;
 
-	/** Code chunk definitions corresponding to each of the material inputs, only initialized after Translate has been called. */
-	FString TranslatedCodeChunkDefinitions[CompiledMP_MAX];
-
-	/** Code chunks corresponding to each of the material inputs, only initialized after Translate has been called. */
-	FString TranslatedCodeChunks[CompiledMP_MAX];
-
 	/** Line number of the #line in MaterialTemplate.usf */
 	int32 MaterialTemplateLineNumber;
 
-	/** Stores the resource declarations */
-	FString ResourcesString;
-
 	/** Contents of the MaterialTemplate.usf file */
 	FString MaterialTemplate;
 
@@ -211,9 +314,6 @@ protected:
 	//TMap<UMaterialExpressionCustom*, int32> CachedCustomExpressions;
 	TArray<FMaterialCustomExpressionEntry> CustomExpressions;
 
-	/** Any custom output function implementations */
-	TArray<FString> CustomOutputImplementations;
-
 	/** Custom vertex interpolators */
 	TArray<UMaterialExpressionVertexInterpolator*> CustomVertexInterpolators;
 	/** Index to assign to next vertex interpolator. */
@@ -229,6 +329,9 @@ protected:
 	TArray<FString>* CompileErrorsSink;
 	TArray<UMaterialExpression*>* CompileErrorExpressionsSink;
 
+	/** Keeps track of which variations of analytic derivative functions are used, and generates the code during translation. **/
+	FMaterialDerivativeAutogen DerivativeAutogen;
+
 	/** Whether the translation succeeded. */
 	uint32 bSuccess : 1;
 	/** Whether the compute shader material inputs were compiled. */
@@ -368,30 +471,40 @@ public:
 	// Assign custom interpolators to slots, packing them as much as possible in unused slots.
 	TBitArray<> GetVertexInterpolatorsOffsets(FString& VertexInterpolatorsOffsetsDefinitionCode) const;
 
-	void GetSharedInputsMaterialCode(FString& PixelMembersDeclaration, FString& NormalAssignment, FString& PixelMembersInitializationEpilog);
+	void GetSharedInputsMaterialCode(FString& PixelMembersDeclaration, FString& NormalAssignment, FString& PixelMembersInitializationEpilog, ECompiledPartialDerivativeVariation DerivativeVariation);
 
 	FString GetMaterialShaderCode();
 
+	const FShaderCodeChunk& AtParameterCodeChunk(int32 Index) const;
+
 protected:
 
 	bool IsMaterialPropertyUsed(EMaterialProperty Property, int32 PropertyChunkIndex, const FLinearColor& ReferenceValue, int32 NumComponents) const;
 
 	// only used by GetMaterialShaderCode()
 	// @param Index ECompiledMaterialProperty or EMaterialProperty
-	FString GenerateFunctionCode(uint32 Index) const;
+	FString GenerateFunctionCode(uint32 Index, ECompiledPartialDerivativeVariation Variation) const;
 
-	// GetParameterCode
+	// GetParameterCode, with DERIV_BASE_VALUE if necessary
 	virtual FString GetParameterCode(int32 Index, const TCHAR* Default = 0);
 
+	// GetParameterCode, no DERIV_BASE_VALUE
+	virtual FString GetParameterCodeRaw(int32 Index, const TCHAR* Default = 0);
+
+public:
+	// Must always be valid
+	virtual FString GetParameterCodeDeriv(int32 Index, ECompiledPartialDerivativeVariation Variation);
+protected:
+
 	uint64 GetParameterHash(int32 Index);
 
 	/** Creates a string of all definitions needed for the given material input. */
-	FString GetDefinitions(TArray<FShaderCodeChunk>& CodeChunks, int32 StartChunk, int32 EndChunk) const;
+	FString GetDefinitions(TArray<FShaderCodeChunk>& CodeChunks, int32 StartChunk, int32 EndChunk, ECompiledPartialDerivativeVariation Variation) const;
 
 	// GetFixedParameterCode
-	void GetFixedParameterCode(int32 StartChunk, int32 EndChunk, int32 ResultIndex, TArray<FShaderCodeChunk>& CodeChunks, FString& OutDefinitions, FString& OutValue);
+	void GetFixedParameterCode(int32 StartChunk, int32 EndChunk, int32 ResultIndex, TArray<FShaderCodeChunk>& CodeChunks, FString& OutDefinitions, FString& OutValue, ECompiledPartialDerivativeVariation Variation);
 
-	void GetFixedParameterCode(int32 ResultIndex, TArray<FShaderCodeChunk>& CodeChunks, FString& OutDefinitions, FString& OutValue);
+	void GetFixedParameterCode(int32 ResultIndex, TArray<FShaderCodeChunk>& CodeChunks, FString& OutDefinitions, FString& OutValue, ECompiledPartialDerivativeVariation Variation);
 
 	/** Used to get a user friendly type from EMaterialValueType */
 	const TCHAR* DescribeType(EMaterialValueType Type) const;
@@ -399,6 +512,9 @@ protected:
 	/** Used to get an HLSL type from EMaterialValueType */
 	const TCHAR* HLSLTypeString(EMaterialValueType Type) const;
 
+	/** Used to get an HLSL type from EMaterialValueType */
+	const TCHAR* HLSLTypeStringDeriv(EMaterialValueType Type, EMaterialExpressionAnalyticDerivativeStatus DerivativeStatus) const;
+
 	int32 NonPixelShaderExpressionError();
 
 	int32 ErrorUnlessFeatureLevelSupported(ERHIFeatureLevel::Type RequiredFeatureLevel);
@@ -414,6 +530,12 @@ protected:
 	/** Adds an already formatted inline or referenced code chunk */
 	int32 AddCodeChunkInner(uint64 Hash, const TCHAR* FormattedCode, EMaterialValueType Type, bool bInlined);
 
+public:
+	/** Adds an already formatted inline or referenced code chunk, and notes the derivative status. */
+	int32 AddCodeChunkInnerDeriv(uint64 Hash, const TCHAR* FormattedCodeFinite, const TCHAR* FormattedCodeAnalytic, EMaterialValueType Type, bool bInlined, EMaterialExpressionAnalyticDerivativeStatus DerivativeStatus);
+
+protected:
+
 	/** 
 	 * Constructs the formatted code chunk and creates a new local variable definition from it. 
 	 * This should be used over AddInlinedCodeChunk when the code chunk adds actual instructions, and especially when calling a function.
diff --git a/Engine/Source/Runtime/Engine/Public/MaterialShared.h b/Engine/Source/Runtime/Engine/Public/MaterialShared.h
--- a/Engine/Source/Runtime/Engine/Public/MaterialShared.h
+++ b/Engine/Source/Runtime/Engine/Public/MaterialShared.h
@@ -1365,6 +1365,18 @@ enum ECompiledMaterialProperty
 	CompiledMP_MAX
 };
 
+/** 
+* Enum that contains entries for the ways that material properties can be compiled with partial derivative calculations.
+* Standard material shaders using the automatic hardware FiniteDifferences, Nanite uses Analytic, and we will likely
+* need to add a separate one for raytracing at a later date.
+*/
+enum ECompiledPartialDerivativeVariation
+{
+	CompiledPDV_FiniteDifferences,
+	CompiledPDV_Analytic,
+	CompiledPDV_MAX
+};
+
 /**
  * Uniquely identifies a material expression output. 
  * Used by the material compiler to keep track of which output it is compiling.
