Make LevelBounds actor always loaded for partitioned worlds so it doesn't contribute to world bounds

diff --git a/Engine/Source/Runtime/Engine/Classes/Engine/LevelBounds.h b/Engine/Source/Runtime/Engine/Classes/Engine/LevelBounds.h
--- a/Engine/Source/Runtime/Engine/Classes/Engine/LevelBounds.h
+++ b/Engine/Source/Runtime/Engine/Classes/Engine/LevelBounds.h
@@ -66,6 +66,9 @@ class ENGINE_API ALevelBounds
 
 	/** Update level bounds immediately so the bounds are accurate when returning. Use only when needed because updating the bounds is slow */
 	void UpdateLevelBoundsImmediately();
+
+	/** Make the level bounds always loaded for partitioned worlds */
+	virtual EActorGridPlacement GetDefaultGridPlacement() const { return EActorGridPlacement::AlwaysLoaded; }
 	
 protected:
 	/** FTickableGameObject interface */
