Remove compile-time warnings, these seem to have partially slipped through

diff --git a/Engine/Source/Editor/Persona/Private/AnimTimeline/AnimModel_AnimSequenceBase.cpp b/Engine/Source/Editor/Persona/Private/AnimTimeline/AnimModel_AnimSequenceBase.cpp
--- a/Engine/Source/Editor/Persona/Private/AnimTimeline/AnimModel_AnimSequenceBase.cpp
+++ b/Engine/Source/Editor/Persona/Private/AnimTimeline/AnimModel_AnimSequenceBase.cpp
@@ -425,7 +425,6 @@ void FAnimModel_AnimSequenceBase::RemoveSelectedCurves()
 
 	if(bDeletedCurve)
 	{
-		REFACTOR_V2_WARNING("V2 why is post edit change needed here?");
 		AnimSequenceBase->PostEditChange();
 
 		if (GetPreviewScene()->GetPreviewMeshComponent()->PreviewInstance != nullptr)
diff --git a/Engine/Source/Editor/Persona/Private/AnimTimeline/AnimTimelineTrack_FloatCurve.cpp b/Engine/Source/Editor/Persona/Private/AnimTimeline/AnimTimelineTrack_FloatCurve.cpp
--- a/Engine/Source/Editor/Persona/Private/AnimTimeline/AnimTimelineTrack_FloatCurve.cpp
+++ b/Engine/Source/Editor/Persona/Private/AnimTimeline/AnimTimelineTrack_FloatCurve.cpp
@@ -191,7 +191,6 @@ void FAnimTimelineTrack_FloatCurve::RemoveCurve()
 			UAnimDataController* Controller = AnimSequenceBase->GetController();
 			Controller->RemoveCurve(CurveId);
 
-			REFACTOR_V2_WARNING("V2 what is the reasoning behind this call?");
 			AnimSequenceBase->PostEditChange();
 		}
 	}
diff --git a/Engine/Source/Editor/Persona/Private/AnimTimeline/SAnimTimeline.cpp b/Engine/Source/Editor/Persona/Private/AnimTimeline/SAnimTimeline.cpp
--- a/Engine/Source/Editor/Persona/Private/AnimTimeline/SAnimTimeline.cpp
+++ b/Engine/Source/Editor/Persona/Private/AnimTimeline/SAnimTimeline.cpp
@@ -780,7 +780,6 @@ void SAnimTimeline::OnColumnFillCoefficientChanged(float FillCoefficient, int32
 
 void SAnimTimeline::HandleKeyComplete()
 {
-	REFACTOR_V2_WARNING("V2 remove this function and verify that behaviour remained the same");
 	Model.Pin()->RefreshTracks();
 }
 
diff --git a/Engine/Source/Editor/Persona/Private/SAnimCompositePanel.cpp b/Engine/Source/Editor/Persona/Private/SAnimCompositePanel.cpp
--- a/Engine/Source/Editor/Persona/Private/SAnimCompositePanel.cpp
+++ b/Engine/Source/Editor/Persona/Private/SAnimCompositePanel.cpp
@@ -209,7 +209,6 @@ void SAnimCompositePanel::PostUndoRedo()
 		RegisterActiveTimer(0.f, FWidgetActiveTimerDelegate::CreateSP(this, &SAnimCompositePanel::TriggerRebuildPanel));
 	}
 
-	REFACTOR_V2_WARNING("V2 why is this call necessary? Is this handled through notify callbacks now?");
 	// when undo or redo happens, we still have to recalculate length, so we can't rely on sequence length changes or not
 	if (Composite->GetPlayLength())
 	{
diff --git a/Engine/Source/Editor/Persona/Private/SAnimMontagePanel.cpp b/Engine/Source/Editor/Persona/Private/SAnimMontagePanel.cpp
--- a/Engine/Source/Editor/Persona/Private/SAnimMontagePanel.cpp
+++ b/Engine/Source/Editor/Persona/Private/SAnimMontagePanel.cpp
@@ -582,7 +582,6 @@ void SAnimMontagePanel::PostRedo( bool bSuccess )
 void SAnimMontagePanel::PostRedoUndo()
 {
 	// when undo or redo happens, we still have to recalculate length, so we can't rely on sequence length changes or not
-	REFACTOR_V2_WARNING("V2 is this still necessary or is this handled through notify callbacks by now?");
 	if (Montage->GetPlayLength())
 	{
 		Montage->SetCompositeLength(MINIMUM_ANIMATION_LENGTH);
diff --git a/Engine/Source/Runtime/Engine/Classes/Animation/AnimSequenceBase.h b/Engine/Source/Runtime/Engine/Classes/Animation/AnimSequenceBase.h
--- a/Engine/Source/Runtime/Engine/Classes/Animation/AnimSequenceBase.h
+++ b/Engine/Source/Runtime/Engine/Classes/Animation/AnimSequenceBase.h
@@ -31,12 +31,6 @@ class UAnimDataModel;
 class UAnimDataController;
 enum class EAnimDataModelNotifyType : uint8;
 
-#if 0
-#define REFACTOR_V2_WARNING(text) EMIT_CUSTOM_WARNING(text)
-#else
-#define REFACTOR_V2_WARNING(text)
-#endif
-
 UCLASS(abstract, BlueprintType)
 class ENGINE_API UAnimSequenceBase : public UAnimationAsset
 {
diff --git a/Engine/Source/Runtime/Engine/Private/Animation/AnimSequence.cpp b/Engine/Source/Runtime/Engine/Private/Animation/AnimSequence.cpp
--- a/Engine/Source/Runtime/Engine/Private/Animation/AnimSequence.cpp
+++ b/Engine/Source/Runtime/Engine/Private/Animation/AnimSequence.cpp
@@ -1076,7 +1076,6 @@ void UAnimSequence::PostEditChangeProperty(FPropertyChangedEvent& PropertyChange
 		ResampleAnimationTrackData();
 	}
 
-	REFACTOR_V2_WARNING("V2 we should not rely on PostEditChange to check for compressed data");
 	const bool bNeedPostProcess = !IsCompressedDataValid() || bAdditiveSettingsChanged || bCompressionAffectingSettingsChanged || bShouldResample;
 
 	// @Todo fix me: This is temporary fix to make sure they always have compressed data
diff --git a/Engine/Source/Runtime/Engine/Private/Animation/AnimSequenceBase.cpp b/Engine/Source/Runtime/Engine/Private/Animation/AnimSequenceBase.cpp
--- a/Engine/Source/Runtime/Engine/Private/Animation/AnimSequenceBase.cpp
+++ b/Engine/Source/Runtime/Engine/Private/Animation/AnimSequenceBase.cpp
@@ -842,7 +842,6 @@ void UAnimSequenceBase::EvaluateCurveData(FBlendedCurve& OutCurve, float Current
 {
 	CSV_SCOPED_TIMING_STAT(Animation, EvaluateCurveData);
 
-	REFACTOR_V2_WARNING("V2 Curve evaluation goes through original curve data structure. This contains copied data from the model. Goal is to have compressed curve data within AnimSequenceBase so that model would be editor / raw data. And compressed data is used at runtime.");
 	PRAGMA_DISABLE_DEPRECATION_WARNINGS
 	RawCurveData.EvaluateCurveData(OutCurve, CurrentTime);
 	PRAGMA_ENABLE_DEPRECATION_WARNINGS
@@ -854,13 +853,11 @@ float UAnimSequenceBase::EvaluateCurveData(SmartName::UID_Type CurveUID, float C
 	const FFloatCurve* Curve = (const FFloatCurve*)RawCurveData.GetCurveData(CurveUID, ERawCurveTrackTypes::RCT_Float);
 	PRAGMA_ENABLE_DEPRECATION_WARNINGS
 
-	REFACTOR_V2_WARNING("V2 Curve evaluation goes through original curve data structure. This contains copied data from the model. Goal is to have compressed curve data within AnimSequenceBase so that model would be editor / raw data. And compressed data is used at runtime.");
 	return Curve != nullptr ? Curve->Evaluate(CurrentTime) : 0.0f;
 }
 
 const FRawCurveTracks& UAnimSequenceBase::GetCurveData() const
 {
-	REFACTOR_V2_WARNING("V2 Curve evaluation goes through original curve data structure. This contains copied data from the model. Goal is to have compressed curve data within AnimSequenceBase so that model would be editor / raw data. And compressed data is used at runtime.");
 	PRAGMA_DISABLE_DEPRECATION_WARNINGS
 	return RawCurveData;
 	PRAGMA_ENABLE_DEPRECATION_WARNINGS
