Ensure that during the upgrade path of populating the DataModel we never consume sequence length changes

diff --git a/Engine/Source/Runtime/Engine/Classes/Animation/AnimSequenceBase.h b/Engine/Source/Runtime/Engine/Classes/Animation/AnimSequenceBase.h
--- a/Engine/Source/Runtime/Engine/Classes/Animation/AnimSequenceBase.h
+++ b/Engine/Source/Runtime/Engine/Classes/Animation/AnimSequenceBase.h
@@ -292,6 +292,9 @@ protected:
 	UPROPERTY(VisibleAnywhere, BlueprintReadOnly, Category = "Animation Model")
 	UAnimDataModel* DataModel;
 
+	/** Flag set whenever the data-model is initially populated (during upgrade path) */
+	bool bPopulatingDataModel;
+
 	/** UAnimDataController instance set to operate on DataModel */
 	UPROPERTY(VisibleAnywhere, BlueprintReadOnly, Transient, DuplicateTransient, Category = "Animation Model")
 	UAnimDataController* Controller;
diff --git a/Engine/Source/Runtime/Engine/Private/Animation/AnimSequence.cpp b/Engine/Source/Runtime/Engine/Private/Animation/AnimSequence.cpp
--- a/Engine/Source/Runtime/Engine/Private/Animation/AnimSequence.cpp
+++ b/Engine/Source/Runtime/Engine/Private/Animation/AnimSequence.cpp
@@ -423,6 +423,7 @@ UAnimSequence::UAnimSequence(const FObjectInitializer& ObjectInitializer)
 	, bUseRawDataOnly(!FPlatformProperties::RequiresCookedData())
 #if WITH_EDITOR
 	, bCompressionInProgress(false)
+	, bBlockCompressionRequests(false)
 #endif
 {
 	RateScale = 1.0;
@@ -5567,6 +5568,11 @@ void UAnimSequence::OnModelModified(const EAnimDataModelNotifyType& NotifyType,
 
 	auto HandleLengthChanged = [this](float NewLength, float OldLength, float T0, float T1)
 	{
+		if (bPopulatingDataModel)
+		{
+			return;
+		}
+
 		if (NewLength > OldLength)
 		{
 			const float InsertTime = T0;
diff --git a/Engine/Source/Runtime/Engine/Private/Animation/AnimSequenceBase.cpp b/Engine/Source/Runtime/Engine/Private/Animation/AnimSequenceBase.cpp
--- a/Engine/Source/Runtime/Engine/Private/Animation/AnimSequenceBase.cpp
+++ b/Engine/Source/Runtime/Engine/Private/Animation/AnimSequenceBase.cpp
@@ -28,6 +28,7 @@ UAnimSequenceBase::UAnimSequenceBase(const FObjectInitializer& ObjectInitializer
 	, RateScale(1.0f)
 #if WITH_EDITORONLY_DATA
 	, DataModel(nullptr)
+	, bPopulatingDataModel(false)
 	, Controller(nullptr)
 #endif // WITH_EDITORONLY_DATA
 {
@@ -93,7 +94,9 @@ void UAnimSequenceBase::PostLoad()
 		GetController();
 		if (bRequiresModelPopulation)
 		{
+			bPopulatingDataModel = true;
 			PopulateModel();
+			bPopulatingDataModel = false;
 		}
 	}
 #endif // WITH_EDITORONLY_DATA
@@ -915,6 +918,11 @@ void UAnimSequenceBase::OnModelModified(const EAnimDataModelNotifyType& NotifyTy
 
 	auto HandleLengthChange = [&](float NewLength, float OldLength, float T0, float T1)
 	{
+		if (bPopulatingDataModel)
+		{
+			return;
+		}
+		
 		const float StartRemoveTime = T0;
 		const float EndRemoveTime = T1;
 
