Fix condition to check whether or not a notify is entirely within a trimmed time region

diff --git a/Engine/Source/Runtime/Engine/Private/Animation/AnimSequenceBase.cpp b/Engine/Source/Runtime/Engine/Private/Animation/AnimSequenceBase.cpp
--- a/Engine/Source/Runtime/Engine/Private/Animation/AnimSequenceBase.cpp
+++ b/Engine/Source/Runtime/Engine/Private/Animation/AnimSequenceBase.cpp
@@ -991,7 +991,7 @@ void UAnimSequenceBase::OnModelModified(const EAnimDataModelNotifyType& NotifyTy
 					NewDuration = NotifyDuration;
 
 					// If Notify is inside of the trimmed time frame, zero out the duration
-					if (NotifyEnd >= StartRemoveTime && NotifyEnd <= EndRemoveTime)
+					if (CurrentTime >= StartRemoveTime && NotifyEnd <= EndRemoveTime)
 					{
 						// small number @todo see if there is define for this
 						NewDuration = DataModel->GetFrameRate().AsInterval();
