Replaced SetSequenceLength undo/redo action with ResizePlayLength, as it was missing the context for insert/removed time window (T0/T1)
Added missing tests for Resize and new ones for ResizePlayLength + undo/redo action
Matched behaviour of resizing transform curves with float curves (oversight)

diff --git a/Engine/Source/Runtime/Engine/Classes/Animation/AnimData/AnimDataController.h b/Engine/Source/Runtime/Engine/Classes/Animation/AnimData/AnimDataController.h
--- a/Engine/Source/Runtime/Engine/Classes/Animation/AnimData/AnimDataController.h
+++ b/Engine/Source/Runtime/Engine/Classes/Animation/AnimData/AnimDataController.h
@@ -111,8 +111,21 @@ public:
 	UFUNCTION(BlueprintCallable, Category = AnimationData)
 	void SetPlayLength(float Length, bool bShouldTransact = true);
 
+	/*** Sets the total play-able length in seconds. Broadcasts a EAnimDataModelNotifyType::SequenceLengthChanged notify if successful.
+	* T0 and T1 are expected to represent the window of time that was either added or removed. E.g. for insertion T0 indicates the time
+	* at which additional time starts and T1 were it ends. For removal T0 indicates the time at which time should be started to remove, and T1 indicates the end. Giving a total of T1 - T0 added or removed length.
+	* The number of frames and keys for the provided length is recalculated according to the current value of UAnimDataModel::FrameRate.
+	* @param	Length				Total new play-able length value, has to be positive and non-zero
+	* @param	T0					Point between 0 and Length at which the change in time starts
+	* @param	T1					Point between 0 and Length at which the change in time ends
+	* @param	bShouldTransact		Whether or not any undo-redo changes should be generated
+	*/
+	UFUNCTION(BlueprintCallable, Category = AnimationData)
+	void ResizePlayLength(float NewLength, float T0, float T1, bool bShouldTransact = true);
+
 	/**
-	* Sets the total play-able length in seconds. Broadcasts a EAnimDataModelNotifyType::SequenceLengthChanged notify if successful.
+	* Sets the total play-able length in seconds and resizes curves. Broadcasts EAnimDataModelNotifyType::SequenceLengthChanged
+	* and EAnimDataModelNotifyType::CurveChanged notifies if successful.
 	* T0 and T1 are expected to represent the window of time that was either added or removed. E.g. for insertion T0 indicates the time
 	* at which additional time starts and T1 were it ends. For removal T0 indicates the time at which time should be started to remove, and T1 indicates the end. Giving a total of T1 - T0 added or removed length.
 	* The number of frames and keys for the provided length is recalculated according to the current value of UAnimDataModel::FrameRate.
@@ -436,9 +449,6 @@ private:
 	/** Resizes the curve data stored on the model according to the provided new length and time at which to insert or remove time */
 	void ResizeCurves(float NewLength, bool bInserted, float T0, float T1, bool bShouldTransact = true);
 
-	/** Helper functionality to apply any changes in the model's playlength*/
-	void SetPlayLength_Internal(float NewLength, float T0, float T1, bool bShouldTransact);
-	   	
 	/** Ensures that a valid model is currently targeted */
 	void ValidateModel() const;
 
diff --git a/Engine/Source/Runtime/Engine/Private/Animation/AnimData/AnimDataControllerActions.h b/Engine/Source/Runtime/Engine/Private/Animation/AnimData/AnimDataControllerActions.h
--- a/Engine/Source/Runtime/Engine/Private/Animation/AnimData/AnimDataControllerActions.h
+++ b/Engine/Source/Runtime/Engine/Private/Animation/AnimData/AnimDataControllerActions.h
@@ -112,18 +112,20 @@ protected:
 	FRawAnimSequenceTrack TrackData;
 };
 
-class FSetSequenceLengthAction : public FAnimDataBaseAction
+class FResizePlayLengthAction : public FAnimDataBaseAction
 {
 public:
-	explicit FSetSequenceLengthAction(const UAnimDataModel* InModel);
-	virtual ~FSetSequenceLengthAction() {}
+	explicit FResizePlayLengthAction(const UAnimDataModel* InModel, float t0, float t1);
+	virtual ~FResizePlayLengthAction() {}
 protected:
-	FSetSequenceLengthAction() {}
+	FResizePlayLengthAction() {}
 	virtual TUniquePtr<FChange> ExecuteInternal(UAnimDataModel* Model, UAnimDataController* Controller) override;
 	virtual FString ToStringInternal() const override;
 
 protected:
 	float Length;
+	float T0;
+	float T1;
 };
 
 class FSetFrameRateAction : public FAnimDataBaseAction
diff --git a/Engine/Source/Runtime/Engine/Private/Animation/AnimDataController.cpp b/Engine/Source/Runtime/Engine/Private/Animation/AnimDataController.cpp
--- a/Engine/Source/Runtime/Engine/Private/Animation/AnimDataController.cpp
+++ b/Engine/Source/Runtime/Engine/Private/Animation/AnimDataController.cpp
@@ -140,18 +140,14 @@ void UAnimDataController::SetPlayLength(float Length, bool bShouldTransact /*= t
 {
 	ValidateModel();
 
-	if (!FMath::IsNearlyZero(Length) && Length > 0.f)
-	{
-		if (Length != Model->GetPlayLength())
-		{
-			CONDITIONAL_TRANSACTION(LOCTEXT("SetPlayLength", "Setting Play Length"));
-			SetPlayLength_Internal(Length, 0.f, Model->PlayLength, bShouldTransact);
-		}
-	}
-	else
-	{
-		ReportWarningf(LOCTEXT("InvalidPlayLengthWarning", "Invalid play length value provided: {0} seconds"), FText::AsNumber(Length));
-	}
+	// Calculate whether or new play length is shorter or longer than current, set-up T0; T1 accordingly
+	// Assumption is made that time is always added or removed at/from the end
+	// Added: T0 = current length, T1 = new length
+	// Removed: T0 = current length - removed length, T1 = current length
+	const float Delta = FMath::Abs(Length - Model->PlayLength);
+	const float T0 = Length > Model->PlayLength ? Model->PlayLength : Model->PlayLength - Delta;
+	const float T1 = Length > Model->PlayLength ? Length : Model->PlayLength;
+	ResizePlayLength(Length, T0, T1, bShouldTransact);	
 }
 
 void UAnimDataController::Resize(float Length, float T0, float T1, bool bShouldTransact /*= true*/)
@@ -171,7 +167,7 @@ void UAnimDataController::Resize(float Length, float T0, float T1, bool bShouldT
 				{
 					CONDITIONAL_BRACKET(LOCTEXT("ResizeModel", "Resizing Animation Data"));
 					const bool bInserted = Length > Model->PlayLength;
-					SetPlayLength_Internal(Length, T0, T1, bShouldTransact);
+					ResizePlayLength(Length, T0, T1, bShouldTransact);
 					ResizeCurves(Length, bInserted, T0, T1, bShouldTransact);
 				}
 				else
@@ -1155,21 +1151,54 @@ void UAnimDataController::ValidateModel() const
 	checkf(Model != nullptr, TEXT("Invalid Model"));
 }
 
-void UAnimDataController::SetPlayLength_Internal(float NewLength, float T0, float T1, bool bShouldTransact)
+void UAnimDataController::ResizePlayLength(float Length, float T0, float T1, bool bShouldTransact)
 {
-	FSequenceLengthChangedPayload Payload;
-	Payload.T0 = T0;
-	Payload.T1 = T1;
-	Payload.PreviousLength = Model->PlayLength;
+	const TRange<float> PlayRange(TRange<float>::BoundsType::Inclusive(0.f), TRange<float>::BoundsType::Inclusive(Model->PlayLength));
+	if (!FMath::IsNearlyZero(Length) && Length > 0.f)
+	{
+		if (Length != Model->PlayLength)
+		{
+			// Ensure that T0 is within the curent play range
+			if (PlayRange.Contains(T0))
+			{
+				// Ensure that the start and end length of either removal or insertion are valid
+				if (T0 < T1)
+				{
+					CONDITIONAL_TRANSACTION(LOCTEXT("ResizePlayLength", "Resizing Play Length"));
 
-	CONDITIONAL_ACTION(UE::Anim::FSetSequenceLengthAction, Model);
+					FSequenceLengthChangedPayload Payload;
+					Payload.T0 = T0;
+					Payload.T1 = T1;
+					Payload.PreviousLength = Model->PlayLength;
 
-	Model->PlayLength = NewLength;
+					CONDITIONAL_ACTION(UE::Anim::FResizePlayLengthAction, Model, T0, T1);
 
-	Model->NumberOfFrames = Model->FrameRate.AsFrameTime(Model->PlayLength).RoundToFrame().Value;
-	Model->NumberOfKeys = Model->NumberOfFrames + 1;
+					Model->PlayLength = Length;
+
+					Model->NumberOfFrames = Model->FrameRate.AsFrameTime(Model->PlayLength).RoundToFrame().Value;
+					Model->NumberOfKeys = Model->NumberOfFrames + 1;
 	
-	Model->Notify<FSequenceLengthChangedPayload>(EAnimDataModelNotifyType::SequenceLengthChanged, Payload);
+					Model->Notify<FSequenceLengthChangedPayload>(EAnimDataModelNotifyType::SequenceLengthChanged, Payload);
+				}
+				else
+				{
+					ReportErrorf(LOCTEXT("InvalidEndTimeError", "Invalid T1, smaller that T0 value: T0 {0}, T1 {1}"), FText::AsNumber(T0), FText::AsNumber(T1));
+				}
+			}
+			else
+			{
+				ReportErrorf(LOCTEXT("InvalidStartTimeError", "Invalid T0, not within existing play range: T0 {0}, Play Length {1}"), FText::AsNumber(T0), FText::AsNumber(Model->PlayLength));
+			}
+		}
+		else
+		{
+			ReportWarningf(LOCTEXT("SamePlayLengthWarning", "New play length is same as existing one: {0} seconds"), FText::AsNumber(Length));
+		}
+	}
+	else
+	{
+		ReportErrorf(LOCTEXT("InvalidPlayLengthError", "Invalid play length value provided: {0} seconds"), FText::AsNumber(Length));
+	}
 }
 
 void UAnimDataController::ReportWarning(const FText& InMessage) const
@@ -1442,7 +1471,22 @@ void UAnimDataController::ResizeCurves(float NewLength, bool bInserted, float T0
 
 	for (FTransformCurve& Curve : Model->CurveData.TransformCurves)
 	{
-		Curve.Resize(NewLength, bInserted, T0, T1);
+		FTransformCurve ResizedCurve = Curve;
+		for (int32 SubCurveIndex = 0; SubCurveIndex < 3; ++SubCurveIndex)
+		{
+			const ETransformCurveChannel Channel = (ETransformCurveChannel)SubCurveIndex;
+			FVectorCurve& SubCurve = *ResizedCurve.GetVectorCurveByIndex(SubCurveIndex);
+			for (int32 ChannelIndex = 0; ChannelIndex < 3; ++ChannelIndex)
+			{
+				const EVectorCurveChannel Axis = (EVectorCurveChannel)ChannelIndex;
+				FAnimationCurveIdentifier TargetCurveIdentifier = FAnimationCurveIdentifier(Curve.Name, ERawCurveTrackTypes::RCT_Transform);
+				UAnimationCurveIdentifierExtensions::GetTransformChildCurveIdentifier(TargetCurveIdentifier, Channel, Axis);
+				
+				FRichCurve& ChannelCurve = SubCurve.FloatCurves[ChannelIndex];
+				ChannelCurve.ReadjustTimeRange(0, NewLength, bInserted, T0, T1);
+				SetCurveKeys(TargetCurveIdentifier, ChannelCurve.GetConstRefOfKeys(), bShouldTransact);
+			}
+		}
 	}
 }
 
diff --git a/Engine/Source/Runtime/Engine/Private/Animation/AnimDataControllerActions.cpp b/Engine/Source/Runtime/Engine/Private/Animation/AnimDataControllerActions.cpp
--- a/Engine/Source/Runtime/Engine/Private/Animation/AnimDataControllerActions.cpp
+++ b/Engine/Source/Runtime/Engine/Private/Animation/AnimDataControllerActions.cpp
@@ -119,21 +119,21 @@ FString FSetTrackKeysAction::ToStringInternal() const
 	return FText::Format(LOCTEXT("SetTrackKeysAction_Description", "Setting keys for animation bone track '{0}'."), FText::FromName(Name)).ToString();
 }
 
-FSetSequenceLengthAction::FSetSequenceLengthAction(const UAnimDataModel* InModel)
+FResizePlayLengthAction::FResizePlayLengthAction(const UAnimDataModel* InModel, float t0, float t1) : T0(t0), T1(t1)
 {
 	Length = InModel->GetPlayLength();
 }
 
-TUniquePtr<FChange> FSetSequenceLengthAction::ExecuteInternal(UAnimDataModel* Model, UAnimDataController* Controller)
+TUniquePtr<FChange> FResizePlayLengthAction::ExecuteInternal(UAnimDataModel* Model, UAnimDataController* Controller)
 {
-	TUniquePtr<FChange> InverseAction = MakeUnique<FSetSequenceLengthAction>(Model);
-	Controller->SetPlayLength(Length, false);
+	TUniquePtr<FChange> InverseAction = MakeUnique<FResizePlayLengthAction>(Model, T0, T1);
+	Controller->ResizePlayLength(Length, T0, T1, false);
 	return InverseAction;
 }
 
-FString FSetSequenceLengthAction::ToStringInternal() const
+FString FResizePlayLengthAction::ToStringInternal() const
 {
-	return FText::Format(LOCTEXT("SetSequenceLengthAction_Description", "Setting play length to {0}."), FText::AsNumber(Length)).ToString();
+	return FText::Format(LOCTEXT("ResizePlayLengthAction_Description", "Resizing play length to {0}."), FText::AsNumber(Length)).ToString();
 }
 
 FSetFrameRateAction::FSetFrameRateAction(const UAnimDataModel* InModel)
