Remove forced resample and synchronous recompression from DataModel upgrade-path. Any resampling and subsequent (async) compression should already be handled/kicked off when the PopulateDataModel bracket is closed

diff --git a/Engine/Source/Runtime/Engine/Private/Animation/AnimSequence.cpp b/Engine/Source/Runtime/Engine/Private/Animation/AnimSequence.cpp
--- a/Engine/Source/Runtime/Engine/Private/Animation/AnimSequence.cpp
+++ b/Engine/Source/Runtime/Engine/Private/Animation/AnimSequence.cpp
@@ -5535,9 +5535,6 @@ void UAnimSequence::PopulateModel()
 	DeleteDeprecatedRawAnimationData();
 	
 	Controller->CloseBracket();
-
-	ResampleAnimationTrackData();
-	RequestSyncAnimRecompression(false);
 }
 
 void UAnimSequence::OnModelModified(const EAnimDataModelNotifyType& NotifyType, UAnimDataModel* Model, const FAnimDataModelNotifPayload& Payload)
