Fixed non-deterministic texture cubes from the reflection capture component

diff --git a/Engine/Source/Runtime/Engine/Private/Components/ReflectionCaptureComponent.cpp b/Engine/Source/Runtime/Engine/Private/Components/ReflectionCaptureComponent.cpp
--- a/Engine/Source/Runtime/Engine/Private/Components/ReflectionCaptureComponent.cpp
+++ b/Engine/Source/Runtime/Engine/Private/Components/ReflectionCaptureComponent.cpp
@@ -693,6 +693,7 @@ void GenerateEncodedHDRTextureCube(UMapBuildDataRegistry* Registry, FReflectionC
 			TSF_BGRA8,
 			TemporaryEncodedHDRCapturedData.GetData()
 		);
+		TextureCube->Source.UseHashAsGuid();
 		// the loader can suggest a compression setting
 		TextureCube->LODGroup = TEXTUREGROUP_World;
 
