World Partition Handles
- Move all unittest specific code to actor descriptors instead of handles.
- Fix actor descriptors GetActor function to return pending kill actors as it was before recent changes.
- Avoid creating an actor desriptor reference when removing loaded actors, which could trigger an attemp to load a deleted actor.
- Make actor descriptors loading code more resilient for missing actors (shouldn't happen).

diff --git a/Engine/Source/Runtime/Engine/Private/WorldPartition/WorldPartitionActorDesc.cpp b/Engine/Source/Runtime/Engine/Private/WorldPartition/WorldPartitionActorDesc.cpp
--- a/Engine/Source/Runtime/Engine/Private/WorldPartition/WorldPartitionActorDesc.cpp
+++ b/Engine/Source/Runtime/Engine/Private/WorldPartition/WorldPartitionActorDesc.cpp
@@ -207,12 +207,19 @@ FBox FWorldPartitionActorDesc::GetBounds() const
 
 bool FWorldPartitionActorDesc::IsLoaded() const
 {
+#if WITH_DEV_AUTOMATION_TESTS
+	if (GIsAutomationTesting)
+	{
+		return HardRefCount > 0;
+	}
+#endif
+
 	return ActorPtr.IsValid();
 }
 
 AActor* FWorldPartitionActorDesc::GetActor() const
 {
-	return ActorPtr.Get();
+	return ActorPtr.Get(true);
 }
 
 AActor* FWorldPartitionActorDesc::Load() const
@@ -277,22 +284,22 @@ void FWorldPartitionActorDesc::Unload()
 
 void FWorldPartitionActorDesc::RegisterActor()
 {
-	AActor* Actor = GetActor();
-	check(Actor);
-
-	ApplyActorTransform(WorldPartition->InstanceTransform);
-	Actor->GetLevel()->AddLoadedActor(Actor);
+	if (AActor* Actor = GetActor())
+	{
+		ApplyActorTransform(WorldPartition->InstanceTransform);
+		Actor->GetLevel()->AddLoadedActor(Actor);
+	}
 }
 
 void FWorldPartitionActorDesc::UnregisterActor()
 {
-	AActor* Actor = GetActor();
-	check(Actor);
-
-	if (!Actor->IsPendingKill())
+	if (AActor* Actor = GetActor())
 	{
-		Actor->GetLevel()->RemoveLoadedActor(Actor);
-		ApplyActorTransform(WorldPartition->InstanceTransform.Inverse());
+		if (!Actor->IsPendingKill())
+		{
+			Actor->GetLevel()->RemoveLoadedActor(Actor);
+			ApplyActorTransform(WorldPartition->InstanceTransform.Inverse());
+		}
 	}
 }
 #endif
\ No newline at end of file
diff --git a/Engine/Source/Runtime/Engine/Private/WorldPartition/WorldPartitionEditorCell.cpp b/Engine/Source/Runtime/Engine/Private/WorldPartition/WorldPartitionEditorCell.cpp
--- a/Engine/Source/Runtime/Engine/Private/WorldPartition/WorldPartitionEditorCell.cpp
+++ b/Engine/Source/Runtime/Engine/Private/WorldPartition/WorldPartitionEditorCell.cpp
@@ -49,7 +49,17 @@ void UWorldPartitionEditorCell::RemoveActor(const FWorldPartitionHandle& ActorHa
 {
 	check(ActorHandle.IsValid());
 	verify(Actors.Remove(ActorHandle));
-	LoadedActors.Remove(ActorHandle);
+
+	// Don't call LoadedActors.Remove(ActorHandle) right away, as this will create a temporary reference and might try to load
+	// a deleted actor. This is a temporary workaround.
+	for (const FWorldPartitionReference& ActorReference: LoadedActors)
+	{
+		if (ActorReference == ActorHandle)
+		{
+			LoadedActors.Remove(ActorHandle);
+			break;
+		}
+	}
 }
 #endif
 
diff --git a/Engine/Source/Runtime/Engine/Private/WorldPartition/WorldPartitionHandle.cpp b/Engine/Source/Runtime/Engine/Private/WorldPartition/WorldPartitionHandle.cpp
--- a/Engine/Source/Runtime/Engine/Private/WorldPartition/WorldPartitionHandle.cpp
+++ b/Engine/Source/Runtime/Engine/Private/WorldPartition/WorldPartitionHandle.cpp
@@ -17,13 +17,6 @@ TUniquePtr<FWorldPartitionActorDesc>* FWorldPartitionHandleUtils::GetActorDesc(U
 
 bool FWorldPartitionHandleUtils::IsActorDescLoaded(FWorldPartitionActorDesc* ActorDesc)
 {
-#if WITH_DEV_AUTOMATION_TESTS
-	if (GIsAutomationTesting)
-	{
-		return ActorDesc->GetHardRefCount() > 0;
-	}
-#endif
-
 	return ActorDesc->IsLoaded();
 }
 
@@ -41,27 +34,17 @@ void FWorldPartitionReferenceImpl::IncRefCount(FWorldPartitionActorDesc* ActorDe
 {
 	if (ActorDesc->IncHardRefCount() == 1)
 	{
-		AActor* Actor = ActorDesc->Load();
-		check(Actor || GIsAutomationTesting);
-
-		if (Actor)
-		{
-			ActorDesc->RegisterActor();
-		}
+		ActorDesc->Load();
+		ActorDesc->RegisterActor();
 	}
 }
 
 void FWorldPartitionReferenceImpl::DecRefCount(FWorldPartitionActorDesc* ActorDesc)
 {
-	if (!ActorDesc->DecHardRefCount())
+	if (ActorDesc->DecHardRefCount() == 0)
 	{
-		// The only case where an actor can be unloaded while still holding a reference 
-		// is when it was manually deleted in the editor.
-		if (ActorDesc->IsLoaded())
-		{
-			ActorDesc->UnregisterActor();
-			ActorDesc->Unload();
-		}
+		ActorDesc->UnregisterActor();
+		ActorDesc->Unload();
 	}
 }
 #endif
\ No newline at end of file
