LevelInstance: Fix new sub level instance being grayed out while editing parent

diff --git a/Engine/Source/Runtime/Engine/Private/LevelInstance/LevelInstanceLevelStreaming.cpp b/Engine/Source/Runtime/Engine/Private/LevelInstance/LevelInstanceLevelStreaming.cpp
--- a/Engine/Source/Runtime/Engine/Private/LevelInstance/LevelInstanceLevelStreaming.cpp
+++ b/Engine/Source/Runtime/Engine/Private/LevelInstance/LevelInstanceLevelStreaming.cpp
@@ -99,14 +99,14 @@ ULevelStreamingLevelInstance* ULevelStreamingLevelInstance::LoadInstance(ALevelI
 						ResetLoaders(LevelActor->GetExternalPackage());
 						LevelActor->GetPackage()->SetFlags(RF_Transient);
 					}		
+
+					LevelActor->PushSelectionToProxies();
+					LevelActor->PushLevelInstanceEditingStateToProxies(LevelInstanceActor->IsInEditingLevelInstance());
 				}
 			}
 
 			// Create special actor that will handle selection and transform
 			ALevelInstanceEditorInstanceActor::Create(LevelInstanceActor, Level);
-
-			// Make sure selection is reflected after load
-			LevelInstanceActor->PushSelectionToProxies();
 		}
 #endif
 		return LevelStreaming;
