Fix animations no longer playing back in control rig preview scenes

Control rig dynamically links multiple proxies together, so the single node animnode ends up registering tick records with the wrong proxy. Going via the scoped sync method like every other instance fixes this.

diff --git a/Engine/Source/Runtime/Engine/Private/Animation/AnimSingleNodeInstanceProxy.cpp b/Engine/Source/Runtime/Engine/Private/Animation/AnimSingleNodeInstanceProxy.cpp
--- a/Engine/Source/Runtime/Engine/Private/Animation/AnimSingleNodeInstanceProxy.cpp
+++ b/Engine/Source/Runtime/Engine/Private/Animation/AnimSingleNodeInstanceProxy.cpp
@@ -10,6 +10,7 @@
 #include "AnimEncoding.h"
 #include "Animation/AnimTrace.h"
 #include "Animation/AnimationPoseData.h"
+#include "Animation/AnimSyncScope.h"
 
 FAnimSingleNodeInstanceProxy::~FAnimSingleNodeInstanceProxy()
 {
@@ -523,10 +524,12 @@ void FAnimNode_SingleNode::Update_AnyThread(const FAnimationUpdateContext& Conte
 
 	if(Proxy->CurrentAsset != NULL)
 	{
+		UE::Anim::FAnimSyncGroupScope& SyncScope = Context.GetMessageChecked<UE::Anim::FAnimSyncGroupScope>();
+
 		if (UBlendSpaceBase* BlendSpace = Cast<UBlendSpaceBase>(Proxy->CurrentAsset))
 		{
 			FAnimTickRecord TickRecord(BlendSpace, Proxy->BlendSpaceInput, Proxy->BlendSampleData, Proxy->BlendFilter, Proxy->bLooping, NewPlayRate, 1.f, /*inout*/ Proxy->CurrentTime, Proxy->MarkerTickRecord);
-			Proxy->AddTickRecord(TickRecord);
+			SyncScope.AddTickRecord(TickRecord);
 
 			TRACE_ANIM_TICK_RECORD(Context, TickRecord);
 #if WITH_EDITORONLY_DATA
@@ -536,7 +539,7 @@ void FAnimNode_SingleNode::Update_AnyThread(const FAnimationUpdateContext& Conte
 		else if (UAnimSequence* Sequence = Cast<UAnimSequence>(Proxy->CurrentAsset))
 		{
 			FAnimTickRecord TickRecord(Sequence, Proxy->bLooping, NewPlayRate, 1.f, /*inout*/ Proxy->CurrentTime, Proxy->MarkerTickRecord);
-			Proxy->AddTickRecord(TickRecord);
+			SyncScope.AddTickRecord(TickRecord);
 
 			TRACE_ANIM_TICK_RECORD(Context, TickRecord);
 
@@ -553,7 +556,7 @@ void FAnimNode_SingleNode::Update_AnyThread(const FAnimationUpdateContext& Conte
 		else if (UAnimStreamable* Streamable = Cast<UAnimStreamable>(Proxy->CurrentAsset))
 		{
 			FAnimTickRecord TickRecord(Streamable, Proxy->bLooping, NewPlayRate, 1.f, /*inout*/ Proxy->CurrentTime, Proxy->MarkerTickRecord);
-			Proxy->AddTickRecord(TickRecord);
+			SyncScope.AddTickRecord(TickRecord);
 
 			TRACE_ANIM_TICK_RECORD(Context, TickRecord);
 
@@ -570,7 +573,7 @@ void FAnimNode_SingleNode::Update_AnyThread(const FAnimationUpdateContext& Conte
 		else if(UAnimComposite* Composite = Cast<UAnimComposite>(Proxy->CurrentAsset))
 		{
 			FAnimTickRecord TickRecord(Composite, Proxy->bLooping, NewPlayRate, 1.f, /*inout*/ Proxy->CurrentTime, Proxy->MarkerTickRecord);
-			Proxy->AddTickRecord(TickRecord);
+			SyncScope.AddTickRecord(TickRecord);
 
 			TRACE_ANIM_TICK_RECORD(Context, TickRecord);
 
@@ -609,7 +612,7 @@ void FAnimNode_SingleNode::Update_AnyThread(const FAnimationUpdateContext& Conte
 		else if (UPoseAsset* PoseAsset = Cast<UPoseAsset>(Proxy->CurrentAsset))
 		{
 			FAnimTickRecord TickRecord(PoseAsset, 1.f);
-			Proxy->AddTickRecord(TickRecord);
+			SyncScope.AddTickRecord(TickRecord);
 
 			TRACE_ANIM_TICK_RECORD(Context, TickRecord);
 		}
