Converted canvas rendering to RDG.

diff --git a/Engine/Plugins/Runtime/HairStrands/Source/HairStrandsCore/Private/GroomDebug.cpp b/Engine/Plugins/Runtime/HairStrands/Source/HairStrandsCore/Private/GroomDebug.cpp
--- a/Engine/Plugins/Runtime/HairStrands/Source/HairStrandsCore/Private/GroomDebug.cpp
+++ b/Engine/Plugins/Runtime/HairStrands/Source/HairStrandsCore/Private/GroomDebug.cpp
@@ -68,29 +68,6 @@ const TCHAR* ToString(EWorldType::Type Type)
 
 ///////////////////////////////////////////////////////////////////////////////////////////////////
 
-namespace GroomDebug
-{
-	// Internal helper class for FCanvas to be able to get screen size
-	class FRenderTargetTemp : public FRenderTarget
-	{
-		const FTexture2DRHIRef Texture;
-		const FIntPoint SizeXY;
-
-	public:
-		FRenderTargetTemp(const FIntRect& ViewRect, const FTexture2DRHIRef InTexture) 
-		: Texture(InTexture)
-		, SizeXY(ViewRect.Size())
-		{}
-
-		virtual FIntPoint GetSizeXY() const override { return SizeXY; }
-
-		virtual const FTexture2DRHIRef& GetRenderTargetTexture() const override
-		{
-			return Texture;
-		}
-	};
-}
-
 bool IsHairStrandsSkinCacheEnable();
 
 static void GetGroomInterpolationData(
@@ -824,21 +801,13 @@ void RunHairStrandsDebug(
 		PassParameters->View = View.ViewUniformBuffer;
 		PassParameters->RenderTargets[0] = FRenderTargetBinding(SceneColorTexture, ERenderTargetLoadAction::ELoad, 0);
 
-		const FSceneView* LocalView = &View;
-		GraphBuilder.AddPass(
-			RDG_EVENT_NAME("HairStrandsMeshProjectionMeshDebug"),
-			PassParameters,
-			ERDGPassFlags::Raster,
-			[LocalView, Viewport, WorldType, Instances, SceneColorTexture](FRHICommandListImmediate& RHICmdList)
+		const FSceneViewFamily& ViewFamily = *View.Family;
+		FCanvas& Canvas = *FCanvas::Create(GraphBuilder, SceneColorTexture, nullptr, ViewFamily.CurrentRealTime, ViewFamily.CurrentWorldTime, ViewFamily.DeltaWorldTime, View.FeatureLevel);
+		Canvas.SetRenderTargetRect(Viewport);
+
 		{
 			const float YStep = 14;
 			float ClusterY = 68;
-
-			// Component part of the clusters
-			GroomDebug::FRenderTargetTemp TempRenderTarget(Viewport, FTexture2DRHIRef(SceneColorTexture->GetRHI()->GetTexture2D()));
-			FCanvas Canvas(&TempRenderTarget, nullptr, LocalView->Family->CurrentRealTime, LocalView->Family->CurrentWorldTime, LocalView->Family->DeltaWorldTime, LocalView->FeatureLevel);
-			Canvas.SetRenderTargetRect(Viewport);
-
 			float X = 20;
 			float Y = ClusterY;
 			const FLinearColor InactiveColor(0.5, 0.5, 0.5);
@@ -880,11 +849,10 @@ void RunHairStrandsDebug(
 			}
 
 			const bool bFlush = false;
-			const bool bInsideRenderPass = true;
-			Canvas.Flush_RenderThread(RHICmdList, bFlush, bInsideRenderPass);
+			Canvas.Flush_RenderThread(GraphBuilder, bFlush);
 
 			ClusterY = Y;
-		});
+		}
 	}
 
 	if (HairDebugMode == EHairDebugMode::MeshProjection)
diff --git a/Engine/Source/Developer/MaterialBaking/Private/MaterialRenderItem.cpp b/Engine/Source/Developer/MaterialBaking/Private/MaterialRenderItem.cpp
--- a/Engine/Source/Developer/MaterialBaking/Private/MaterialRenderItem.cpp
+++ b/Engine/Source/Developer/MaterialBaking/Private/MaterialRenderItem.cpp
@@ -7,6 +7,7 @@
 #include "StaticMeshAttributes.h"
 #include "DynamicMeshBuilder.h"
 #include "MeshPassProcessor.h"
+#include "CanvasRender.h"
 
 #define SHOW_WIREFRAME_MESH 0
 
@@ -25,7 +26,7 @@ FMeshMaterialRenderItem::FMeshMaterialRenderItem(
 	LCI = new FMeshRenderInfo(InMeshSettings->LightMap, nullptr, nullptr, InMeshSettings->LightmapResourceCluster);
 }
 
-bool FMeshMaterialRenderItem::Render_RenderThread(FRHICommandListImmediate& RHICmdList, FMeshPassProcessorRenderState& DrawRenderState, const FCanvas* Canvas)
+bool FMeshMaterialRenderItem::Render_RenderThread(FCanvasRenderContext& RenderContext, FMeshPassProcessorRenderState& DrawRenderState, const FCanvas* Canvas)
 {
 	checkSlow(ViewFamily && MeshSettings && MaterialRenderProxy);
 	// current render target set for the canvas
@@ -57,20 +58,20 @@ bool FMeshMaterialRenderItem::Render_RenderThread(FRHICommandListImmediate& RHIC
 		LocalDrawRenderState.SetBlendState(TStaticBlendState<CW_RGBA>::GetRHI());
 		LocalDrawRenderState.SetDepthStencilState(TStaticDepthStencilState<false, CF_Always>::GetRHI());
 
-		QueueMaterial(RHICmdList, LocalDrawRenderState, &View);
+		QueueMaterial(RenderContext, LocalDrawRenderState, &View);
 	}
 
 	return true;
 }
 
-bool FMeshMaterialRenderItem::Render_GameThread(const FCanvas* Canvas, FRenderThreadScope& RenderScope)
+bool FMeshMaterialRenderItem::Render_GameThread(const FCanvas* Canvas, FCanvasRenderThreadScope& RenderScope)
 {
 	RenderScope.EnqueueRenderCommand(
-		[this, Canvas](FRHICommandListImmediate& RHICmdList)
+		[this, Canvas](FCanvasRenderContext& RenderContext)
 		{
 			// Render_RenderThread uses its own render state
 			FMeshPassProcessorRenderState DummyRenderState;
-			Render_RenderThread(RHICmdList, DummyRenderState, Canvas);
+			Render_RenderThread(RenderContext, DummyRenderState, Canvas);
 		}
 	);
 
@@ -106,7 +107,7 @@ FMeshMaterialRenderItem::~FMeshMaterialRenderItem()
 	);
 }
 
-void FMeshMaterialRenderItem::QueueMaterial(FRHICommandListImmediate& RHICmdList, FMeshPassProcessorRenderState& DrawRenderState, const FSceneView* View)
+void FMeshMaterialRenderItem::QueueMaterial(FCanvasRenderContext& RenderContext, FMeshPassProcessorRenderState& DrawRenderState, const FSceneView* View)
 {
 	TRACE_CPUPROFILER_EVENT_SCOPE(FMeshMaterialRenderItem::QueueMaterial)
 
@@ -143,7 +144,7 @@ void FMeshMaterialRenderItem::QueueMaterial(FRHICommandListImmediate& RHICmdList
 	}
 
 	// Bake the material out to a tile
-	GetRendererModule().DrawTileMesh(RHICmdList, DrawRenderState, *View, MeshElement, false /*bIsHitTesting*/, FHitProxyId());
+	GetRendererModule().DrawTileMesh(RenderContext, DrawRenderState, *View, MeshElement, false /*bIsHitTesting*/, FHitProxyId());
 }
 
 void FMeshMaterialRenderItem::PopulateWithQuadData()
diff --git a/Engine/Source/Developer/MaterialBaking/Private/MaterialRenderItem.h b/Engine/Source/Developer/MaterialBaking/Private/MaterialRenderItem.h
--- a/Engine/Source/Developer/MaterialBaking/Private/MaterialRenderItem.h
+++ b/Engine/Source/Developer/MaterialBaking/Private/MaterialRenderItem.h
@@ -53,15 +53,15 @@ public:
 	virtual ~FMeshMaterialRenderItem();
 
 	/** Begin FCanvasBaseRenderItem overrides */
-	virtual bool Render_RenderThread(FRHICommandListImmediate& RHICmdList, FMeshPassProcessorRenderState& DrawRenderState, const FCanvas* Canvas) final;
-	virtual bool Render_GameThread(const FCanvas* Canvas, FRenderThreadScope& RenderScope) final;
+	virtual bool Render_RenderThread(FCanvasRenderContext& RenderContext, FMeshPassProcessorRenderState& DrawRenderState, const FCanvas* Canvas) final;
+	virtual bool Render_GameThread(const FCanvas* Canvas, FCanvasRenderThreadScope& RenderScope) final;
 	/** End FCanvasBaseRenderItem overrides */
 
 	/** Populate vertices and indices according to available mesh data and otherwise uses simple quad */
 	void GenerateRenderData();
 protected:
 	/** Enqueues the current material to be rendered */
-	void QueueMaterial(FRHICommandListImmediate& RHICmdList, FMeshPassProcessorRenderState& DrawRenderState, const FSceneView* View);
+	void QueueMaterial(FCanvasRenderContext& RenderContext, FMeshPassProcessorRenderState& DrawRenderState, const FSceneView* View);
 	
 	/** Helper functions to populate render data using either mesh data or a simple quad */
 	void PopulateWithQuadData();
diff --git a/Engine/Source/Developer/MaterialUtilities/Private/MeshRendering.cpp b/Engine/Source/Developer/MaterialUtilities/Private/MeshRendering.cpp
--- a/Engine/Source/Developer/MaterialUtilities/Private/MeshRendering.cpp
+++ b/Engine/Source/Developer/MaterialUtilities/Private/MeshRendering.cpp
@@ -22,7 +22,7 @@
 #include "SceneView.h"
 #include "MeshBatch.h"
 #include "CanvasItem.h"
-#include "CanvasTypes.h"
+#include "CanvasRender.h"
 #include "LocalVertexFactory.h"
 #include "Rendering/SkeletalMeshLODRenderData.h"
 #include "MeshPassProcessor.h"
@@ -468,7 +468,7 @@ public:
 		return 2;
 	}
 
-	static void RenderMaterial(FRHICommandListImmediate& RHICmdList, FMeshPassProcessorRenderState& DrawRenderState, const class FSceneView& View, FRenderData& Data)
+	static void RenderMaterial(FCanvasRenderContext& RenderContext, FMeshPassProcessorRenderState& DrawRenderState, const class FSceneView& View, FRenderData& Data)
 	{
 		// Check if material is TwoSided - single-sided materials should be rendered with normal and reverse
 		// triangle corner orders, to avoid problems with inside-out meshes or mesh parts. Note:
@@ -506,8 +506,8 @@ public:
 		DynamicMeshBuilder.AddVertices(Verts);
 		DynamicMeshBuilder.AddTriangles(Indices);
 
-		FMeshBatch MeshElement;
-		FMeshBuilderOneFrameResources OneFrameResource;
+		FMeshBatch& MeshElement = *RenderContext.Alloc<FMeshBatch>();
+		FMeshBuilderOneFrameResources& OneFrameResource = *RenderContext.Alloc<FMeshBuilderOneFrameResources>();
 		DynamicMeshBuilder.GetMeshElement(FMatrix::Identity, Data.MaterialRenderProxy, SDPG_Foreground, true, false, 0, OneFrameResource, MeshElement);
 
 		check(OneFrameResource.IsValidForRendering());
@@ -520,10 +520,10 @@ public:
 		MeshElement.bWireframe = true;
 #endif
 
-		GetRendererModule().DrawTileMesh(RHICmdList, DrawRenderState, View, MeshElement, false /*bIsHitTesting*/, FHitProxyId());
+		GetRendererModule().DrawTileMesh(RenderContext, DrawRenderState, View, MeshElement, false /*bIsHitTesting*/, FHitProxyId());
 	}
 
-	virtual bool Render_RenderThread(FRHICommandListImmediate& RHICmdList, FMeshPassProcessorRenderState& DrawRenderState, const FCanvas* Canvas)
+	virtual bool Render_RenderThread(FCanvasRenderContext& RenderContext, FMeshPassProcessorRenderState& DrawRenderState, const FCanvas* Canvas)
 	{
 		checkSlow(Data);
 		// current render target set for the canvas
@@ -545,21 +545,19 @@ public:
 
 		FSceneView* View = new FSceneView(ViewInitOptions);
 
-		RenderMaterial(RHICmdList, DrawRenderState, *View, *Data);
+		RenderMaterial(RenderContext, DrawRenderState, *View, *Data);
+
+		RenderContext.DeferredDelete(View);
 
-		delete View;
-		if (Canvas->GetAllowedModes() & FCanvas::Allow_DeleteOnRender)
-		{
-			delete Data;
-		}
 		if (Canvas->GetAllowedModes() & FCanvas::Allow_DeleteOnRender)
 		{
-			Data = NULL;
+			RenderContext.DeferredDelete(Data);
+			Data = nullptr;
 		}
 		return true;
 	}
 
-	virtual bool Render_GameThread(const FCanvas* Canvas, FRenderThreadScope& RenderScope)
+	virtual bool Render_GameThread(const FCanvas* Canvas, FCanvasRenderThreadScope& RenderScope)
 	{
 		checkSlow(Data);
 		// current render target set for the canvas
@@ -595,24 +593,24 @@ public:
 
 		FDrawMaterialParameters Parameters = DrawMaterialParameters;
 		RenderScope.EnqueueRenderCommand(
-			[Parameters](FRHICommandListImmediate& RHICmdList)
+			[Parameters](FCanvasRenderContext& RenderContext)
 		{
 			FMeshPassProcessorRenderState DrawRenderState;
 
 			// disable depth test & writes
 			DrawRenderState.SetDepthStencilState(TStaticDepthStencilState<false, CF_Always>::GetRHI());
 
-			RenderMaterial(RHICmdList, DrawRenderState, *Parameters.View, *Parameters.RenderData);
+			RenderMaterial(RenderContext, DrawRenderState, *Parameters.View, *Parameters.RenderData);
 
-			delete Parameters.View;
+			RenderContext.DeferredDelete(Parameters.View);
 			if (Parameters.AllowedCanvasModes & FCanvas::Allow_DeleteOnRender)
 			{
-				delete Parameters.RenderData;
+				RenderContext.DeferredDelete(Parameters.RenderData);
 			}
 		});
 		if (Canvas->GetAllowedModes() & FCanvas::Allow_DeleteOnRender)
 		{
-			Data = NULL;
+			Data = nullptr;
 		}
 		return true;
 	}
diff --git a/Engine/Source/Editor/UnrealEd/Private/Lightmass/LightmassLandscapeRender.cpp b/Engine/Source/Editor/UnrealEd/Private/Lightmass/LightmassLandscapeRender.cpp
--- a/Engine/Source/Editor/UnrealEd/Private/Lightmass/LightmassLandscapeRender.cpp
+++ b/Engine/Source/Editor/UnrealEd/Private/Lightmass/LightmassLandscapeRender.cpp
@@ -13,7 +13,7 @@
 #include "SceneView.h"
 #include "PrimitiveUniformShaderParameters.h"
 #include "LocalVertexFactory.h"
-#include "CanvasTypes.h"
+#include "CanvasRender.h"
 #include "MeshBatch.h"
 
 #include "LandscapeProxy.h"
@@ -121,28 +121,21 @@ void RenderLandscapeMaterialForLightmass(const FLandscapeStaticLightingMesh* Lan
 			
 			if (OneFrameResource.IsValidForRendering())
 			{
-				// Set the RHI render target.
-				RHICmdList.Transition(FRHITransitionInfo(RenderTarget->GetRenderTargetTexture(), ERHIAccess::Unknown, ERHIAccess::RTV));
+				const FIntRect RTViewRect = FIntRect(0, 0, RenderTarget->GetRenderTargetTexture()->GetSizeX(), RenderTarget->GetRenderTargetTexture()->GetSizeY());
 
-				FRHIRenderPassInfo RPInfo(RenderTarget->GetRenderTargetTexture(), ERenderTargetActions::Load_Store);
-				RHICmdList.BeginRenderPass(RPInfo, TEXT("CanvasFlushSetup"));
-				{
-					const FIntRect RTViewRect = FIntRect(0, 0, RenderTarget->GetRenderTargetTexture()->GetSizeX(), RenderTarget->GetRenderTargetTexture()->GetSizeY());
-
-					// set viewport to RT size
-					RHICmdList.SetViewport(RTViewRect.Min.X, RTViewRect.Min.Y, 0.0f, RTViewRect.Max.X, RTViewRect.Max.Y, 1.0f);
+				FRDGBuilder GraphBuilder(RHICmdList, RDG_EVENT_NAME("LightmassLandscapeMaterial"));
+				FSceneView View(ViewInitOptions);
 
-					FSceneView View(ViewInitOptions);
+				FMeshPassProcessorRenderState DrawRenderState;
 
-					FMeshPassProcessorRenderState DrawRenderState;
+				// disable depth test & writes
+				DrawRenderState.SetDepthStencilState(TStaticDepthStencilState<false, CF_Always>::GetRHI());
 
-					// disable depth test & writes
-					DrawRenderState.SetDepthStencilState(TStaticDepthStencilState<false, CF_Always>::GetRHI());
+				//SCOPED_DRAW_EVENT(RHICmdList, RenderLandscapeMaterialToTexture);
+				FCanvasRenderContext RenderContext(GraphBuilder, RenderTarget->GetRenderTargetTexture(GraphBuilder), RTViewRect, FIntRect(0, 0, 0, 0));
+				GetRendererModule().DrawTileMesh(RenderContext, DrawRenderState, View, Mesh, false, FHitProxyId());
 
-					//SCOPED_DRAW_EVENT(RHICmdList, RenderLandscapeMaterialToTexture);
-					GetRendererModule().DrawTileMesh(RHICmdList, DrawRenderState, View, Mesh, false, FHitProxyId());
-				}
-				RHICmdList.EndRenderPass();
+				GraphBuilder.Execute();
 			}
 		});
 	FlushRenderingCommands();
diff --git a/Engine/Source/Runtime/Engine/Private/Slate/DebugCanvas.cpp b/Engine/Source/Runtime/Engine/Private/Slate/DebugCanvas.cpp
--- a/Engine/Source/Runtime/Engine/Private/Slate/DebugCanvas.cpp
+++ b/Engine/Source/Runtime/Engine/Private/Slate/DebugCanvas.cpp
@@ -87,7 +87,7 @@ void FDebugCanvasDrawer::HandleReleaseFontResources(const class FSlateFontCache&
 		ENQUEUE_RENDER_COMMAND(FlushFontResourcesCommand)(
 			[this](FRHICommandListImmediate& RHICmdList)
 		{
-			RenderThreadCanvas->Flush_RenderThread(RHICmdList, true, false);
+			RenderThreadCanvas->Flush_RenderThread(RHICmdList, true);
 		});
 
 		FlushRenderingCommands();
diff --git a/Engine/Source/Runtime/Engine/Private/TileRendering.cpp b/Engine/Source/Runtime/Engine/Private/TileRendering.cpp
--- a/Engine/Source/Runtime/Engine/Private/TileRendering.cpp
+++ b/Engine/Source/Runtime/Engine/Private/TileRendering.cpp
@@ -12,7 +12,7 @@
 #include "PackedNormal.h"
 #include "LocalVertexFactory.h"
 #include "SceneView.h"
-#include "CanvasTypes.h"
+#include "CanvasRender.h"
 #include "MeshBatch.h"
 #include "RendererInterface.h"
 #include "SceneUtils.h"
@@ -163,7 +163,7 @@ void FCanvasTileRendererItem::FRenderData::ReleaseTileMesh()
 }
 
 void FCanvasTileRendererItem::FRenderData::RenderTiles(
-	FRHICommandListImmediate& RHICmdList,
+	FCanvasRenderContext& RenderContext,
 	FMeshPassProcessorRenderState& DrawRenderState,
 	const FSceneView& View,
 	bool bIsHitTesting,
@@ -172,8 +172,8 @@ void FCanvasTileRendererItem::FRenderData::RenderTiles(
 {
 	check(IsInRenderingThread());
 
-	SCOPED_GPU_STAT(RHICmdList, CanvasDrawTile);
-	SCOPED_DRAW_EVENTF(RHICmdList, CanvasDrawTile, *MaterialRenderProxy->GetIncompleteMaterialWithFallback(GMaxRHIFeatureLevel).GetFriendlyName());
+	RDG_GPU_STAT_SCOPE(RenderContext.GraphBuilder, CanvasDrawTile);
+	RDG_EVENT_SCOPE(RenderContext.GraphBuilder, "%s", *MaterialRenderProxy->GetIncompleteMaterialWithFallback(GMaxRHIFeatureLevel).GetFriendlyName());
 	TRACE_CPUPROFILER_EVENT_SCOPE(CanvasDrawTile);
 	QUICK_SCOPE_CYCLE_COUNTER(STAT_CanvasDrawTile)
 
@@ -189,13 +189,16 @@ void FCanvasTileRendererItem::FRenderData::RenderTiles(
 		Mesh.MaterialRenderProxy = MaterialRenderProxy;
 		Mesh.Elements[0].FirstIndex = CanvasTileIndexCount * TileIdx;
 
-		RendererModule.DrawTileMesh(RHICmdList, DrawRenderState, View, Mesh, bIsHitTesting, Tile.HitProxyId, bUse128bitRT);
+		RendererModule.DrawTileMesh(RenderContext, DrawRenderState, View, Mesh, bIsHitTesting, Tile.HitProxyId, bUse128bitRT);
 	}
 
-	ReleaseTileMesh();
+	AddPass(RenderContext.GraphBuilder, [this](FRHICommandList&)
+	{
+		ReleaseTileMesh();
+	});
 }
 
-bool FCanvasTileRendererItem::Render_RenderThread(FRHICommandListImmediate& RHICmdList, FMeshPassProcessorRenderState& DrawRenderState, const FCanvas* Canvas)
+bool FCanvasTileRendererItem::Render_RenderThread(FCanvasRenderContext& RenderContext, FMeshPassProcessorRenderState& DrawRenderState, const FCanvas* Canvas)
 {
 	float CurrentRealTime = 0.f;
 	float CurrentWorldTime = 0.f;
@@ -212,7 +215,7 @@ bool FCanvasTileRendererItem::Render_RenderThread(FRHICommandListImmediate& RHIC
 
 	const FRenderTarget* CanvasRenderTarget = Canvas->GetRenderTarget();
 
-	TUniquePtr<const FSceneViewFamily> ViewFamily = MakeUnique<const FSceneViewFamily>(FSceneViewFamily::ConstructionValues(
+	const FSceneViewFamily& ViewFamily = *RenderContext.Alloc<FSceneViewFamily>(FSceneViewFamily::ConstructionValues(
 		CanvasRenderTarget,
 		nullptr,
 		FEngineShowFlags(ESFIM_Game))
@@ -223,7 +226,7 @@ bool FCanvasTileRendererItem::Render_RenderThread(FRHICommandListImmediate& RHIC
 
 	// make a temporary view
 	FSceneViewInitOptions ViewInitOptions;
-	ViewInitOptions.ViewFamily = ViewFamily.Get();
+	ViewInitOptions.ViewFamily = &ViewFamily;
 	ViewInitOptions.SetViewRectangle(ViewRect);
 	ViewInitOptions.ViewOrigin = FVector::ZeroVector;
 	ViewInitOptions.ViewRotationMatrix = FMatrix::Identity;
@@ -231,21 +234,22 @@ bool FCanvasTileRendererItem::Render_RenderThread(FRHICommandListImmediate& RHIC
 	ViewInitOptions.BackgroundColor = FLinearColor::Black;
 	ViewInitOptions.OverlayColor = FLinearColor::White;
 
-	TUniquePtr<const FSceneView> View = MakeUnique<const FSceneView>(ViewInitOptions);
+	const FSceneView& View = *RenderContext.Alloc<const FSceneView>(ViewInitOptions);
 
 	const bool bNeedsToSwitchVerticalAxis = RHINeedsToSwitchVerticalAxis(Canvas->GetShaderPlatform()) && Canvas->GetAllowSwitchVerticalAxis(); 
 
-	Data->RenderTiles(RHICmdList, DrawRenderState, *View, Canvas->IsHitTesting(), bNeedsToSwitchVerticalAxis);
+	Data->RenderTiles(RenderContext, DrawRenderState, View, Canvas->IsHitTesting(), bNeedsToSwitchVerticalAxis);
 
 	if (Canvas->GetAllowedModes() & FCanvas::Allow_DeleteOnRender)
 	{
+		RenderContext.DeferredRelease(MoveTemp(Data));
 		Data = nullptr;
 	}
 
 	return true;
 }
 
-bool FCanvasTileRendererItem::Render_GameThread(const FCanvas* Canvas, FRenderThreadScope& RenderScope)
+bool FCanvasTileRendererItem::Render_GameThread(const FCanvas* Canvas, FCanvasRenderThreadScope& RenderScope)
 {
 	float CurrentRealTime = 0.f;
 	float CurrentWorldTime = 0.f;
@@ -298,17 +302,18 @@ bool FCanvasTileRendererItem::Render_GameThread(const FCanvas* Canvas, FRenderTh
 
 		RenderScope.EnqueueRenderCommand(
 			[LocalData = Data, View, bIsHitTesting, bNeedsToSwitchVerticalAxis, bRequiresExplicit128bitRT]
-		(FRHICommandListImmediate& RHICmdList)
+			(FCanvasRenderContext& RenderContext) mutable
 		{
 			FMeshPassProcessorRenderState DrawRenderState;
 
 			// disable depth test & writes
 			DrawRenderState.SetDepthStencilState(TStaticDepthStencilState<false, CF_Always>::GetRHI());
 
-			LocalData->RenderTiles(RHICmdList, DrawRenderState, *View, bIsHitTesting, bNeedsToSwitchVerticalAxis, bRequiresExplicit128bitRT);
+			LocalData->RenderTiles(RenderContext, DrawRenderState, *View, bIsHitTesting, bNeedsToSwitchVerticalAxis, bRequiresExplicit128bitRT);
 
-			delete View->Family;
-			delete View;
+			RenderContext.DeferredRelease(MoveTemp(LocalData));
+			RenderContext.DeferredDelete(View->Family);
+			RenderContext.DeferredDelete(View);
 		});
 
 		if (bDeleteOnRender)
diff --git a/Engine/Source/Runtime/Engine/Private/TriangleRendering.cpp b/Engine/Source/Runtime/Engine/Private/TriangleRendering.cpp
--- a/Engine/Source/Runtime/Engine/Private/TriangleRendering.cpp
+++ b/Engine/Source/Runtime/Engine/Private/TriangleRendering.cpp
@@ -12,7 +12,7 @@
 #include "PackedNormal.h"
 #include "LocalVertexFactory.h"
 #include "SceneView.h"
-#include "CanvasTypes.h"
+#include "CanvasRender.h"
 #include "MeshBatch.h"
 #include "RendererInterface.h"
 #include "SceneUtils.h"
@@ -130,7 +130,7 @@ void FCanvasTriangleRendererItem::FRenderData::ReleaseTriangleMesh()
 }
 
 void FCanvasTriangleRendererItem::FRenderData::RenderTriangles(
-	FRHICommandListImmediate& RHICmdList,
+	FCanvasRenderContext& RenderContext,
 	FMeshPassProcessorRenderState& DrawRenderState,
 	const FSceneView& View,
 	bool bIsHitTesting,
@@ -142,7 +142,7 @@ void FCanvasTriangleRendererItem::FRenderData::RenderTriangles(
 
 	InitTriangleMesh(View, bNeedsToSwitchVerticalAxis);
 
-	SCOPED_DRAW_EVENTF(RHICmdList, CanvasDrawTriangles, *MaterialRenderProxy->GetIncompleteMaterialWithFallback(GMaxRHIFeatureLevel).GetFriendlyName());
+	RDG_EVENT_SCOPE(RenderContext.GraphBuilder, "%s", *MaterialRenderProxy->GetIncompleteMaterialWithFallback(GMaxRHIFeatureLevel).GetFriendlyName());
 
 	for (int32 TriIdx = 0; TriIdx < Triangles.Num(); TriIdx++)
 	{
@@ -153,13 +153,16 @@ void FCanvasTriangleRendererItem::FRenderData::RenderTriangles(
 		MeshBatch.MaterialRenderProxy = MaterialRenderProxy;
 		MeshBatch.Elements[0].FirstIndex = 3 * TriIdx;
 
-		GetRendererModule().DrawTileMesh(RHICmdList, DrawRenderState, View, MeshBatch, bIsHitTesting, Tri.HitProxyId);
+		GetRendererModule().DrawTileMesh(RenderContext, DrawRenderState, View, MeshBatch, bIsHitTesting, Tri.HitProxyId);
 	}
 
-	ReleaseTriangleMesh();
+	AddPass(RenderContext.GraphBuilder, [this](FRHICommandList&)
+	{
+		ReleaseTriangleMesh();
+	});
 }
 
-bool FCanvasTriangleRendererItem::Render_RenderThread(FRHICommandListImmediate& RHICmdList, FMeshPassProcessorRenderState& DrawRenderState, const FCanvas* Canvas)
+bool FCanvasTriangleRendererItem::Render_RenderThread(FCanvasRenderContext& RenderContext, FMeshPassProcessorRenderState& DrawRenderState, const FCanvas* Canvas)
 {
 	float CurrentRealTime = 0.f;
 	float CurrentWorldTime = 0.f;
@@ -176,7 +179,7 @@ bool FCanvasTriangleRendererItem::Render_RenderThread(FRHICommandListImmediate&
 
 	const FRenderTarget* CanvasRenderTarget = Canvas->GetRenderTarget();
 
-	TUniquePtr<const FSceneViewFamily> ViewFamily = MakeUnique<const FSceneViewFamily>(FSceneViewFamily::ConstructionValues(
+	const FSceneViewFamily& ViewFamily = *RenderContext.Alloc<const FSceneViewFamily>(FSceneViewFamily::ConstructionValues(
 		CanvasRenderTarget,
 		nullptr,
 		FEngineShowFlags(ESFIM_Game))
@@ -187,7 +190,7 @@ bool FCanvasTriangleRendererItem::Render_RenderThread(FRHICommandListImmediate&
 
 	// make a temporary view
 	FSceneViewInitOptions ViewInitOptions;
-	ViewInitOptions.ViewFamily = ViewFamily.Get();
+	ViewInitOptions.ViewFamily = &ViewFamily;
 	ViewInitOptions.SetViewRectangle(ViewRect);
 	ViewInitOptions.ViewOrigin = FVector::ZeroVector;
 	ViewInitOptions.ViewRotationMatrix = FMatrix::Identity;
@@ -195,21 +198,22 @@ bool FCanvasTriangleRendererItem::Render_RenderThread(FRHICommandListImmediate&
 	ViewInitOptions.BackgroundColor = FLinearColor::Black;
 	ViewInitOptions.OverlayColor = FLinearColor::White;
 
-	TUniquePtr<const FSceneView> View = MakeUnique<const FSceneView>(ViewInitOptions);
+	const FSceneView& View = *RenderContext.Alloc<const FSceneView>(ViewInitOptions);
 
 	const bool bNeedsToSwitchVerticalAxis = RHINeedsToSwitchVerticalAxis(Canvas->GetShaderPlatform()) && Canvas->GetAllowSwitchVerticalAxis();
 
-	Data->RenderTriangles(RHICmdList, DrawRenderState, *View, Canvas->IsHitTesting(), bNeedsToSwitchVerticalAxis);
+	Data->RenderTriangles(RenderContext, DrawRenderState, View, Canvas->IsHitTesting(), bNeedsToSwitchVerticalAxis);
 
 	if (Canvas->GetAllowedModes() & FCanvas::Allow_DeleteOnRender)
 	{
+		RenderContext.DeferredRelease(MoveTemp(Data));
 		Data = nullptr;
 	}
 
 	return true;
 }
 
-bool FCanvasTriangleRendererItem::Render_GameThread(const FCanvas* Canvas, FRenderThreadScope& RenderScope)
+bool FCanvasTriangleRendererItem::Render_GameThread(const FCanvas* Canvas, FCanvasRenderThreadScope& RenderScope)
 {
 	float CurrentRealTime = 0.f;
 	float CurrentWorldTime = 0.f;
@@ -252,17 +256,18 @@ bool FCanvasTriangleRendererItem::Render_GameThread(const FCanvas* Canvas, FRend
 	const bool bDeleteOnRender = Canvas->GetAllowedModes() & FCanvas::Allow_DeleteOnRender;
 
 	RenderScope.EnqueueRenderCommand(
-		[LocalData = Data, View, bIsHitTesting, bNeedsToSwitchVerticalAxis](FRHICommandListImmediate& RHICmdList)
+		[LocalData = Data, View, bIsHitTesting, bNeedsToSwitchVerticalAxis](FCanvasRenderContext& RenderContext) mutable
 	{
 		FMeshPassProcessorRenderState DrawRenderState;
 
 		// disable depth test & writes
 		DrawRenderState.SetDepthStencilState(TStaticDepthStencilState<false, CF_Always>::GetRHI());
 
-		LocalData->RenderTriangles(RHICmdList, DrawRenderState, *View, bIsHitTesting, bNeedsToSwitchVerticalAxis);
+		LocalData->RenderTriangles(RenderContext, DrawRenderState, *View, bIsHitTesting, bNeedsToSwitchVerticalAxis);
 
-		delete View->Family;
-		delete View;
+		RenderContext.DeferredRelease(MoveTemp(LocalData));
+		RenderContext.DeferredDelete(View->Family);
+		RenderContext.DeferredDelete(View);
 	});
 
 	if (Canvas->GetAllowedModes() & FCanvas::Allow_DeleteOnRender)
diff --git a/Engine/Source/Runtime/Engine/Private/UnrealClient.cpp b/Engine/Source/Runtime/Engine/Private/UnrealClient.cpp
--- a/Engine/Source/Runtime/Engine/Private/UnrealClient.cpp
+++ b/Engine/Source/Runtime/Engine/Private/UnrealClient.cpp
@@ -28,6 +28,7 @@
 #include "Templates/UniquePtr.h"
 #include "Elements/Framework/TypedElementList.h"
 #include "EngineUtils.h"
+#include "RenderGraphUtils.h"
 
 DEFINE_LOG_CATEGORY_STATIC(LogClient, Log, All);
 
@@ -272,6 +273,11 @@ const FTexture2DRHIRef& FRenderTarget::GetRenderTargetTexture() const
 	return RenderTargetTextureRHI;
 }
 
+FRDGTextureRef FRenderTarget::GetRenderTargetTexture(FRDGBuilder& GraphBuilder) const
+{
+	return RegisterExternalTexture(GraphBuilder, RenderTargetTextureRHI, TEXT("RenderTarget"));
+}
+
 FUnorderedAccessViewRHIRef FRenderTarget::GetRenderTargetUAV() const
 {
 	return FUnorderedAccessViewRHIRef();
diff --git a/Engine/Source/Runtime/Engine/Private/UserInterface/Canvas.cpp b/Engine/Source/Runtime/Engine/Private/UserInterface/Canvas.cpp
--- a/Engine/Source/Runtime/Engine/Private/UserInterface/Canvas.cpp
+++ b/Engine/Source/Runtime/Engine/Private/UserInterface/Canvas.cpp
@@ -27,6 +27,7 @@
 #include "Debug/ReporterGraph.h"
 #include "Fonts/FontMeasure.h"
 #include "EngineModule.h"
+#include "CanvasRender.h"
 
 DEFINE_LOG_CATEGORY_STATIC(LogCanvas, Log, All);
 
@@ -41,6 +42,32 @@ DEFINE_STAT(STAT_Canvas_AddTileRenderTime);
 DEFINE_STAT(STAT_Canvas_AddTriangleRenderTime);
 DEFINE_STAT(STAT_Canvas_NumBatchesCreated);
 
+class FRDGCanvasRenderTarget final : public FRenderTarget
+{
+public:
+	FRDGCanvasRenderTarget(FRDGTextureRef InRDGTexture)
+		: RDGTexture(InRDGTexture)
+	{}
+
+	FIntPoint GetSizeXY() const override
+	{
+		return RDGTexture->Desc.Extent;
+	}
+
+	const FTexture2DRHIRef& GetRenderTargetTexture() const override
+	{
+		static FTexture2DRHIRef NullRef;
+		return NullRef;
+	}
+
+	FRDGTextureRef GetRenderTargetTexture(FRDGBuilder&) const override
+	{
+		return RDGTexture;
+	}
+
+	FRDGTextureRef RDGTexture;
+};
+
 FCanvasWordWrapper::FCanvasWordWrapper()
 	: GraphemeBreakIterator(FBreakIterator::CreateCharacterBoundaryIterator())
 	, LineBreakIterator(FBreakIterator::CreateLineBreakIterator())
@@ -234,6 +261,13 @@ int32 FCanvasWordWrapper::FindEndOfLastWholeGraphemeCluster(const int32 InStartI
 	return BreakIndex;
 }
 
+FCanvas* FCanvas::Create(FRDGBuilder& GraphBuilder, FRDGTextureRef InRenderTarget, FHitProxyConsumer* InHitProxyConsumer, float InRealTime, float InWorldTime, float InWorldDeltaTime, ERHIFeatureLevel::Type InFeatureLevel, float InDPIScale)
+{
+	return GraphBuilder.AllocObject<FCanvas>(
+		GraphBuilder.AllocObject<FRDGCanvasRenderTarget>(InRenderTarget),
+		InHitProxyConsumer, InRealTime, InWorldTime, InWorldDeltaTime, InFeatureLevel, InDPIScale);
+}
+
 FCanvas::FCanvas(FRenderTarget* InRenderTarget, FHitProxyConsumer* InHitProxyConsumer, UWorld* InWorld, ERHIFeatureLevel::Type InFeatureLevel, ECanvasDrawMode InDrawMode, float InDPIScale)
 :	ViewRect(0,0,0,0)
 ,	ScissorRect(0,0,0,0)
@@ -392,45 +426,53 @@ FMatrix FCanvas::CalcProjectionMatrix(uint32 ViewSizeX, uint32 ViewSizeY, float
 	}
 }
 
-bool FCanvasBatchedElementRenderItem::Render_RenderThread(FRHICommandListImmediate& RHICmdList, FMeshPassProcessorRenderState& DrawRenderState, const FCanvas* Canvas)
+bool FCanvasBatchedElementRenderItem::Render_RenderThread(FCanvasRenderContext& RenderContext, FMeshPassProcessorRenderState& DrawRenderState, const FCanvas* Canvas)
 {
 	checkSlow(Data);
 	bool bDirty = false;
+
 	if (Data->BatchedElements.HasPrimsToDraw())
 	{
 		bDirty = true;
 
-		// current render target set for the canvas
-		const FRenderTarget* CanvasRenderTarget = Canvas->GetRenderTarget();
-		float Gamma = 1.0f / CanvasRenderTarget->GetDisplayGamma();
-		if (Data->Texture && Data->Texture->bIgnoreGammaConversions)
+		RenderContext.AddPass(
+			RDG_EVENT_NAME("CanvasBatchedElements"),
+			[LocalData = Data, DrawRenderState, Canvas](FRHICommandListImmediate& RHICmdList)
 		{
-			Gamma = 1.0f;
-		}
+			// current render target set for the canvas
+			const FRenderTarget* CanvasRenderTarget = Canvas->GetRenderTarget();
+			float Gamma = 1.0f / CanvasRenderTarget->GetDisplayGamma();
+			if (LocalData->Texture && LocalData->Texture->bIgnoreGammaConversions)
+			{
+				Gamma = 1.0f;
+			}
 
-		bool bNeedsToSwitchVerticalAxis = RHINeedsToSwitchVerticalAxis(Canvas->GetShaderPlatform()) && Canvas->GetAllowSwitchVerticalAxis();
+			bool bNeedsToSwitchVerticalAxis = RHINeedsToSwitchVerticalAxis(Canvas->GetShaderPlatform()) && Canvas->GetAllowSwitchVerticalAxis();
 
-		// draw batched items
-		Data->BatchedElements.Draw(
-			RHICmdList,
-			DrawRenderState,
-			Canvas->GetFeatureLevel(),
-			bNeedsToSwitchVerticalAxis,
-			FBatchedElements::CreateProxySceneView(Data->Transform.GetMatrix(), FIntRect(0, 0, CanvasRenderTarget->GetSizeXY().X, CanvasRenderTarget->GetSizeXY().Y)),
-			Canvas->IsHitTesting(),
-			Gamma
+			// draw batched items
+			LocalData->BatchedElements.Draw(
+				RHICmdList,
+				DrawRenderState,
+				Canvas->GetFeatureLevel(),
+				bNeedsToSwitchVerticalAxis,
+				FBatchedElements::CreateProxySceneView(LocalData->Transform.GetMatrix(), FIntRect(0, 0, CanvasRenderTarget->GetSizeXY().X, CanvasRenderTarget->GetSizeXY().Y)),
+				Canvas->IsHitTesting(),
+				Gamma
 			);
+		});
 	}
+
 	if (Canvas->GetAllowedModes() & FCanvas::Allow_DeleteOnRender)
 	{
 		// delete data since we're done rendering it
-		delete Data;
-		Data = NULL;
+		RenderContext.DeferredDelete(Data);
+		Data = nullptr;
 	}
+
 	return bDirty;
 }
 
-bool FCanvasBatchedElementRenderItem::Render_GameThread(const FCanvas* Canvas, FRenderThreadScope& RenderScope)
+bool FCanvasBatchedElementRenderItem::Render_GameThread(const FCanvas* Canvas, FCanvasRenderThreadScope& RenderScope)
 {	
 	checkSlow(Data);
 	bool bDirty=false;
@@ -474,7 +516,8 @@ bool FCanvasBatchedElementRenderItem::Render_GameThread(const FCanvas* Canvas, F
 			Canvas->GetFeatureLevel(),
 			Canvas->GetShaderPlatform()
 		};
-		RenderScope.EnqueueRenderCommand(
+		RenderScope.AddPass(
+			TEXT("CanvasBatchedElements"),
 			[DrawParameters](FRHICommandList& RHICmdList)
 		{
 			FSceneView SceneView = FBatchedElements::CreateProxySceneView(DrawParameters.RenderData->Transform.GetMatrix(), FIntRect(0, 0, DrawParameters.ViewportSizeX, DrawParameters.ViewportSizeY));
@@ -508,8 +551,52 @@ bool FCanvasBatchedElementRenderItem::Render_GameThread(const FCanvas* Canvas, F
 	return bDirty;
 }
 
+FCanvasRenderContext::FCanvasRenderContext(FRDGBuilder& InGraphBuilder, const FCanvas& Canvas)
+	: GraphBuilder(InGraphBuilder)
+	, RenderTarget(Canvas.RenderTarget->GetRenderTargetTexture(InGraphBuilder))
+	, ViewportRect(Canvas.ViewRect)
+	, ScissorRect(Canvas.ScissorRect)
+{
+	if (ViewportRect.Area() <= 0 || Canvas.bScaledToRenderTarget)
+	{
+		ViewportRect = FIntRect(FIntPoint::ZeroValue, RenderTarget->Desc.Extent);
+	}
+}
+
+FCanvasRenderContext::FCanvasRenderContext(FRDGBuilder& InGraphBuilder, FRDGTextureRef InRenderTarget, FIntRect InViewportRect, FIntRect InScissorRect)
+	: GraphBuilder(InGraphBuilder)
+	, RenderTarget(InRenderTarget)
+	, ViewportRect(InViewportRect)
+	, ScissorRect(InScissorRect)
+{}
+
+FCanvasRenderThreadScope::FCanvasRenderThreadScope(const FCanvas& InCanvas)
+	: Canvas(InCanvas)
+{
+	RenderCommands = new RenderCommandFunctionArray;
+}
+
+FCanvasRenderThreadScope::~FCanvasRenderThreadScope()
+{
+	RenderCommandFunctionArray* RenderCommandArray = RenderCommands;
+
+	ENQUEUE_RENDER_COMMAND(DispatchCanvasRenderCommands)(
+		[RenderCommandArray, LocalCanvas = Canvas](FRHICommandListImmediate& RHICmdList)
+	{
+		GetRendererModule().InitializeSystemTextures(RHICmdList);
+		FMemMark MemMark(FMemStack::Get());
+		FRDGBuilder GraphBuilder(RHICmdList, RDG_EVENT_NAME("CanvasRenderThreadScope"));
+		FCanvasRenderContext RenderContext(GraphBuilder, LocalCanvas);
 
+		for (uint32 Index = 0, Count = RenderCommandArray->Num(); Index < Count; ++Index)
+		{
+			(*RenderCommandArray)[Index](RenderContext);
+		}
 
+		GraphBuilder.Execute();
+		delete RenderCommandArray;
+	});
+}
 
 FCanvas::FCanvasSortElement& FCanvas::GetSortElement(int32 DepthSortKey)
 {
@@ -646,7 +733,27 @@ FCanvas::~FCanvas()
 	ClearBatchesToRender();
 }
 
-void FCanvas::Flush_RenderThread(FRHICommandListImmediate& RHICmdList, bool bForce, bool bInsideRenderPass)
+
+void FCanvas::Flush_RenderThread(FRHICommandListImmediate& RHICmdList, bool bForce)
+{
+	if (!(AllowedModes & Allow_Flush) && !bForce)
+	{
+		return;
+	}
+
+	// no need to set the render target if we aren't going to draw anything to it!
+	if (SortedElements.Num() == 0)
+	{
+		return;
+	}
+
+	FMemMark MemMark(FMemStack::Get());
+	FRDGBuilder GraphBuilder(RHICmdList);
+	Flush_RenderThread(GraphBuilder, bForce);
+	GraphBuilder.Execute();
+}
+
+void FCanvas::Flush_RenderThread(FRDGBuilder& GraphBuilder, bool bForce)
 {
 	SCOPE_CYCLE_COUNTER(STAT_Canvas_FlushTime);
 
@@ -678,85 +785,55 @@ void FCanvas::Flush_RenderThread(FRHICommandListImmediate& RHICmdList, bool bFor
 	// sort the array of FCanvasSortElement entries so that higher sort keys render first (back-to-front)
 	SortedElements.Sort(FCompareFCanvasSortElement());
 
-	SCOPED_DRAW_EVENT(RHICmdList, CanvasFlush);
+	RDG_EVENT_SCOPE(GraphBuilder, "CanvasFlush");
 
-	if (!bInsideRenderPass)
-	{
-		GetRendererModule().InitializeSystemTextures(RHICmdList);
-		const FTexture2DRHIRef& RenderTargetTexture = RenderTarget->GetRenderTargetTexture();
+	GetRendererModule().InitializeSystemTextures(GraphBuilder.RHICmdList);
 
-		check(IsValidRef(RenderTargetTexture));
-		check(RHICmdList.IsOutsideRenderPass());
+	FCanvasRenderContext RenderContext(GraphBuilder, *this);
 
-		FRHIRenderPassInfo RPInfo(RenderTargetTexture, ERenderTargetActions::Load_Store);
+	// Set the RHI render target.
+	if (IsUsingInternalTexture())
+	{
+		AddClearRenderTargetPass(GraphBuilder, RenderContext.GetRenderTarget());
+	}
 
-		// Set the RHI render target.
-		if (IsUsingInternalTexture())
-		{
-			RPInfo.ColorRenderTargets[0].Action = ERenderTargetActions::Clear_Store;
-		}
+	// Disable depth test & writes
+	FMeshPassProcessorRenderState DrawRenderState;
+	DrawRenderState.SetDepthStencilState(TStaticDepthStencilState<false, CF_Always>::GetRHI());
 
-		TransitionRenderPassTargets(RHICmdList, RPInfo);
-		RHICmdList.BeginRenderPass(RPInfo, TEXT("CanvasRenderThread"));
-	}
-	else
+	if (ViewRect.Area() <= 0)
 	{
-		check(RHICmdList.IsInsideRenderPass());
+		ViewRect = FIntRect(FIntPoint::ZeroValue, RenderTarget->GetSizeXY());
 	}
 
+	// iterate over the FCanvasSortElements in sorted order and render all the batched items for each entry
+	for (int32 Idx = 0; Idx < SortedElements.Num(); Idx++)
 	{
-		// Disable depth test & writes
-		FMeshPassProcessorRenderState DrawRenderState;
-		DrawRenderState.SetDepthStencilState(TStaticDepthStencilState<false, CF_Always>::GetRHI());
-
-		if (ViewRect.Area() <= 0)
-		{
-			ViewRect = FIntRect(FIntPoint::ZeroValue, RenderTarget->GetSizeXY());
-		}
-
-		// set viewport to RT size
-		RHICmdList.SetViewport(ViewRect.Min.X, ViewRect.Min.Y, 0.0f, ViewRect.Max.X, ViewRect.Max.Y, 1.0f);
-
-		// Set scissor rect if valid applied
-		if (ScissorRect.Area() > 0)
-		{
-			RHICmdList.SetScissorRect(true, ScissorRect.Min.X, ScissorRect.Min.Y, ScissorRect.Max.X, ScissorRect.Max.Y);
-		}
-
-		// iterate over the FCanvasSortElements in sorted order and render all the batched items for each entry
-		for (int32 Idx = 0; Idx < SortedElements.Num(); Idx++)
+		FCanvasSortElement& SortElement = SortedElements[Idx];
+		for (int32 BatchIdx = 0; BatchIdx < SortElement.RenderBatchArray.Num(); BatchIdx++)
 		{
-			FCanvasSortElement& SortElement = SortedElements[Idx];
-			for (int32 BatchIdx = 0; BatchIdx < SortElement.RenderBatchArray.Num(); BatchIdx++)
+			FCanvasBaseRenderItem* RenderItem = SortElement.RenderBatchArray[BatchIdx];
+			if (RenderItem)
 			{
-				FCanvasBaseRenderItem* RenderItem = SortElement.RenderBatchArray[BatchIdx];
-				if (RenderItem)
+				// mark current render target as dirty since we are drawing to it
+				bRenderTargetDirty |= RenderItem->Render_RenderThread(RenderContext, DrawRenderState, this);
+				if (AllowedModes & Allow_DeleteOnRender)
 				{
-					// mark current render target as dirty since we are drawing to it
-					bRenderTargetDirty |= RenderItem->Render_RenderThread(RHICmdList, DrawRenderState, this);
-					if (AllowedModes & Allow_DeleteOnRender)
-					{
-						delete RenderItem;
-					}
+					RenderContext.DeferredDelete(RenderItem);
 				}
 			}
-			if (AllowedModes & Allow_DeleteOnRender)
-			{
-				SortElement.RenderBatchArray.Empty();
-			}
 		}
 		if (AllowedModes & Allow_DeleteOnRender)
 		{
-			// empty the array of FCanvasSortElement entries after finished with rendering	
-			SortedElements.Empty();
-			SortedElementLookupMap.Empty();
-			LastElementIndex = INDEX_NONE;
+			SortElement.RenderBatchArray.Empty();
 		}
 	}
-
-	if (!bInsideRenderPass)
+	if (AllowedModes & Allow_DeleteOnRender)
 	{
-		RHICmdList.EndRenderPass();
+		// empty the array of FCanvasSortElement entries after finished with rendering	
+		SortedElements.Empty();
+		SortedElementLookupMap.Empty();
+		LastElementIndex = INDEX_NONE;
 	}
 }
 
@@ -795,58 +872,12 @@ void FCanvas::Flush_GameThread(bool bForce)
 	// sort the array of FCanvasSortElement entries so that higher sort keys render first (back-to-front)
 	SortedElements.Sort(FCompareFCanvasSortElement());
 	
-	const FIntPoint RenderTargetSize = RenderTarget->GetSizeXY();
-	if (ViewRect.Area() <= 0)
-	{
-		ViewRect = FIntRect(FIntPoint::ZeroValue, RenderTargetSize);
-	}
-
-	struct FCanvasFlushParameters
-	{
-		FIntRect ViewRect;
-		FIntRect ScissorRect;
-		const FRenderTarget* CanvasRenderTarget;
-		bool bIsScaledToRenderTarget;
-	};
-	FCanvasFlushParameters FlushParameters =
-	{
-		ViewRect,
-		ScissorRect,
-		RenderTarget,
-		IsScaledToRenderTarget()
-	};
 	bool bEmitCanvasDrawEvents = GetEmitDrawEvents();
 
 	// Only create these render commands if we actually have something to draw.
 	if(SortedElements.Num() > 0 && !GUsingNullRHI)
 	{
-		FRenderThreadScope RenderThreadScope;
-		RenderThreadScope.EnqueueRenderCommand(
-			[FlushParameters](FRHICommandListImmediate& RHICmdList)
-		{
-			GetRendererModule().InitializeSystemTextures(RHICmdList);
-			check(RHICmdList.IsOutsideRenderPass());
-
-			// Set the RHI render target.
-			FRHIRenderPassInfo RPInfo(FlushParameters.CanvasRenderTarget->GetRenderTargetTexture(), ERenderTargetActions::Load_Store);
-			TransitionRenderPassTargets(RHICmdList, RPInfo);
-			RHICmdList.BeginRenderPass(RPInfo, TEXT("Canvas_GameThreadFlush"));
-
-			FIntRect ViewportRect = FlushParameters.ViewRect;
-			FIntRect ScissorRectParam = FlushParameters.ScissorRect;
-			if (FlushParameters.bIsScaledToRenderTarget)
-			{
-				FIntPoint CanvasSize = FlushParameters.CanvasRenderTarget->GetSizeXY();
-				ViewportRect = FIntRect(0, 0, CanvasSize.X, CanvasSize.Y);
-			}
-			else if (ScissorRectParam.Area() > 0)
-			{
-				RHICmdList.SetScissorRect(true, ScissorRectParam.Min.X, ScissorRectParam.Min.Y, ScissorRectParam.Max.X, ScissorRectParam.Max.Y);
-			}
-
-			// set viewport to RT size
-			RHICmdList.SetViewport(ViewportRect.Min.X, ViewportRect.Min.Y, 0.0f, ViewportRect.Max.X, ViewportRect.Max.Y, 1.0f);
-		});
+		FCanvasRenderThreadScope RenderThreadScope(*this);
 
 		// iterate over the FCanvasSortElements in sorted order and render all the batched items for each entry
 		for( int32 Idx=0; Idx < SortedElements.Num(); Idx++ )
@@ -861,7 +892,7 @@ void FCanvas::Flush_GameThread(bool bForce)
 					bRenderTargetDirty |= RenderItem->Render_GameThread(this, RenderThreadScope);
 					if( AllowedModes & Allow_DeleteOnRender )
 					{
-						delete RenderItem;
+						RenderThreadScope.DeferredDelete(RenderItem);
 					}
 				}
 			}
@@ -870,13 +901,6 @@ void FCanvas::Flush_GameThread(bool bForce)
 				SortElement.RenderBatchArray.Empty();
 			}
 		}
-
-		RenderThreadScope.EnqueueRenderCommand(
-			[FlushParameters](FRHICommandListImmediate& RHICmdList)
-		{
-			RHICmdList.EndRenderPass();
-			RHICmdList.Transition(FRHITransitionInfo(FlushParameters.CanvasRenderTarget->GetRenderTargetTexture(), ERHIAccess::RTV, ERHIAccess::SRVMask));
-		});
 	}
 	
 	if( AllowedModes & Allow_DeleteOnRender )
diff --git a/Engine/Source/Runtime/Engine/Public/CanvasRender.h b/Engine/Source/Runtime/Engine/Public/CanvasRender.h
new file mode 100644
--- /dev/null
+++ b/Engine/Source/Runtime/Engine/Public/CanvasRender.h
@@ -0,0 +1,129 @@
+// Copyright Epic Games, Inc. All Rights Reserved.
+
+#pragma once
+
+#include "CanvasTypes.h"
+#include "RenderGraphBuilder.h"
+
+class FCanvasRenderContext
+{
+public:
+	ENGINE_API FCanvasRenderContext(FRDGBuilder& GraphBuilder, const FCanvas& Canvas);
+	ENGINE_API FCanvasRenderContext(FRDGBuilder& GraphBuilder, FRDGTextureRef Texture, FIntRect ViewportRect, FIntRect ScissorRect);
+
+	template <typename ExecuteLambdaType, typename ParameterStructType>
+	void AddPass(FRDGEventName&& PassName, const ParameterStructType* PassParameters, ExecuteLambdaType&& ExecuteLambda)
+	{
+		GraphBuilder.AddPass(Forward<FRDGEventName>(PassName), PassParameters, ERDGPassFlags::Raster,
+			[LocalScissorRect = ScissorRect, LocalViewportRect = ViewportRect, LocalExecuteLambda = Forward<ExecuteLambdaType&&>(ExecuteLambda)](FRHICommandListImmediate& RHICmdList)
+		{
+			if (LocalScissorRect.Area() > 0)
+			{
+				RHICmdList.SetScissorRect(true, LocalScissorRect.Min.X, LocalScissorRect.Min.Y, LocalScissorRect.Max.X, LocalScissorRect.Max.Y);
+			}
+			RHICmdList.SetViewport(LocalViewportRect.Min.X, LocalViewportRect.Min.Y, 0.0f, LocalViewportRect.Max.X, LocalViewportRect.Max.Y, 1.0f);
+			LocalExecuteLambda(RHICmdList);
+		});
+	}
+
+	template <typename ExecuteLambdaType>
+	void AddPass(FRDGEventName&& PassName, ExecuteLambdaType&& ExecuteLambda)
+	{
+		FRenderTargetParameters* PassParameters = GraphBuilder.AllocParameters<FRenderTargetParameters>();
+		PassParameters->RenderTargets[0] = FRenderTargetBinding(RenderTarget, ERenderTargetLoadAction::ELoad);
+		AddPass(Forward<FRDGEventName&&>(PassName), PassParameters, Forward<ExecuteLambdaType&&>(ExecuteLambda));
+	}
+
+	template <typename T, typename... TArgs>
+	T* Alloc(TArgs&&... Args)
+	{
+		return GraphBuilder.AllocObject<T>(Forward<TArgs&&>(Args)...);
+	}
+
+	template <typename T>
+	void DeferredRelease(TSharedPtr<T>&& Ptr)
+	{
+		// Hold the reference until completion of the graph.
+		Alloc<TSharedPtr<T>>(Forward<TSharedPtr<T>&&>(Ptr));
+	}
+
+	template <typename T>
+	void DeferredDelete(const T* Ptr)
+	{
+		struct FDeleter
+		{
+			FDeleter(const T* InPtr)
+				: Ptr(InPtr)
+			{}
+
+			~FDeleter()
+			{
+				delete Ptr;
+			}
+
+			const T* Ptr;
+		};
+
+		Alloc<FDeleter>(Ptr);
+	}
+
+	FRDGTextureRef GetRenderTarget() const
+	{
+		return RenderTarget;
+	}
+
+	FIntRect GetViewportRect() const
+	{
+		return ViewportRect;
+	}
+
+	FIntRect GetScissorRect() const
+	{
+		return ScissorRect;
+	}
+
+	FRDGBuilder& GraphBuilder;
+
+private:
+	FRDGTextureRef RenderTarget;
+	FIntRect ViewportRect;
+	FIntRect ScissorRect;
+};
+
+class FCanvasRenderThreadScope
+{
+	using RenderCommandFunction = TFunction<void(FCanvasRenderContext&)>;
+	using RenderCommandFunctionArray = TArray<RenderCommandFunction>;
+public:
+	ENGINE_API FCanvasRenderThreadScope(const FCanvas& Canvas);
+	ENGINE_API ~FCanvasRenderThreadScope();
+
+	void EnqueueRenderCommand(RenderCommandFunction&& Lambda)
+	{
+		RenderCommands->Add(MoveTemp(Lambda));
+	}
+
+	template <typename ExecuteLambdaType>
+	void AddPass(const TCHAR* PassName, ExecuteLambdaType&& Lambda)
+	{
+		EnqueueRenderCommand(
+			[PassName, InLambda = Forward<ExecuteLambdaType&&>(Lambda)]
+		(FCanvasRenderContext& RenderContext) mutable
+		{
+			RenderContext.AddPass(RDG_EVENT_NAME("%s", PassName), Forward<ExecuteLambdaType&&>(InLambda));
+		});
+	}
+
+	template <typename T>
+	void DeferredDelete(const T* Ptr)
+	{
+		EnqueueRenderCommand([Ptr](FCanvasRenderContext& RenderContext)
+		{
+			RenderContext.DeferredDelete(Ptr);
+		});
+	}
+
+private:
+	const FCanvas& Canvas;
+	RenderCommandFunctionArray* RenderCommands;
+};
\ No newline at end of file
diff --git a/Engine/Source/Runtime/Engine/Public/CanvasTypes.h b/Engine/Source/Runtime/Engine/Public/CanvasTypes.h
--- a/Engine/Source/Runtime/Engine/Public/CanvasTypes.h
+++ b/Engine/Source/Runtime/Engine/Public/CanvasTypes.h
@@ -15,6 +15,8 @@
 #include "StaticMeshResources.h"
 #include "CanvasTypes.generated.h"
 
+class FCanvasRenderContext;
+class FCanvasRenderThreadScope;
 class FCanvasItem;
 class FMaterialRenderProxy;
 class IBreakIterator;
@@ -226,6 +228,8 @@ public:
 		CDM_ImmediateDrawing
 	};
 
+	ENGINE_API static FCanvas* Create(FRDGBuilder& GraphBuilder, FRDGTextureRef InRenderTarget, FHitProxyConsumer* InHitProxyConsumer, float InRealTime, float InWorldTime, float InWorldDeltaTime, ERHIFeatureLevel::Type InFeatureLevel, float InDPIScale = 1.0f);
+
 	/**
 	* Constructor.
 	*/
@@ -273,7 +277,8 @@ public:
 	* @param bInsideRenderPass - Set to true if flushing inside a render pass (e.g. Render Graph pass).
 	*	This will skip creating a render pass internally, and assert if the command list is not in a render pass.
 	*/
-	ENGINE_API void Flush_RenderThread(FRHICommandListImmediate& RHICmdList, bool bForce = false, bool bInsideRenderPass = false);
+	ENGINE_API void Flush_RenderThread(FRHICommandListImmediate& RHICmdList, bool bForce = false);
+	ENGINE_API void Flush_RenderThread(FRDGBuilder& GraphBuilder, bool bForce = false);
 
 	/**
 	* Sends a message to the rendering thread to draw the batched elements.
@@ -591,6 +596,7 @@ public:
 	TSharedPtr<FCanvasWordWrapper> WordWrapper;
 
 private:
+
 	/** Stack of SortKeys. All rendering is done using the top most sort key */
 	TArray<int32> DepthSortKeyStack;	
 	/** Stack of matrices. Bottom most entry is the canvas projection */
@@ -774,10 +780,9 @@ public:
 	*/
 	ENGINE_API FCanvasSortElement& GetSortElement(int32 DepthSortKey);
 
+	friend class FCanvasRenderContext;
 };
 
-
-
 /**
 * Base interface for canvas items which can be batched for rendering
 */
@@ -790,11 +795,11 @@ public:
 	/**
 	* Renders the canvas item
 	*
+	* @param RenderContext - the canvas render context to submit render passes to.
 	* @param Canvas - canvas currently being rendered
-	* @param RHICmdList - command list to use
 	* @return true if anything rendered
 	*/
-	virtual bool Render_RenderThread(FRHICommandListImmediate& RHICmdList, FMeshPassProcessorRenderState& DrawRenderState, const FCanvas* Canvas) = 0;
+	virtual bool Render_RenderThread(FCanvasRenderContext& RenderContext, FMeshPassProcessorRenderState& DrawRenderState, const FCanvas* Canvas) = 0;
 	
 	/**
 	* Renders the canvas item
@@ -802,7 +807,7 @@ public:
 	* @param Canvas - canvas currently being rendered
 	* @return true if anything rendered
 	*/
-	virtual bool Render_GameThread(const FCanvas* Canvas, FRenderThreadScope& RenderScope) = 0;
+	virtual bool Render_GameThread(const FCanvas* Canvas, FCanvasRenderThreadScope& RenderScope) = 0;
 	
 	/**
 	* FCanvasBatchedElementRenderItem instance accessor
@@ -873,7 +878,7 @@ public:
 	* @param RHICmdList - command list to use
 	* @return true if anything rendered
 	*/
-	virtual bool Render_RenderThread(FRHICommandListImmediate& RHICmdList, FMeshPassProcessorRenderState& DrawRenderState, const FCanvas* Canvas) override;
+	virtual bool Render_RenderThread(FCanvasRenderContext& RenderContext, FMeshPassProcessorRenderState& DrawRenderState, const FCanvas* Canvas) override;
 	
 	/**
 	* Renders the canvas item.
@@ -882,7 +887,7 @@ public:
 	* @param Canvas - canvas currently being rendered
 	* @return true if anything rendered
 	*/
-	virtual bool Render_GameThread(const FCanvas* Canvas, FRenderThreadScope& RenderScope) override;
+	virtual bool Render_GameThread(const FCanvas* Canvas, FCanvasRenderThreadScope& RenderScope) override;
 
 	/**
 	* Determine if this is a matching set by comparing texture,blendmode,elementype,transform. All must match
@@ -996,7 +1001,7 @@ public:
 	* @param RHICmdList - command list to use
 	* @return true if anything rendered
 	*/
-	virtual bool Render_RenderThread(FRHICommandListImmediate& RHICmdList, FMeshPassProcessorRenderState& DrawRenderState, const FCanvas* Canvas) override;
+	virtual bool Render_RenderThread(FCanvasRenderContext& RenderContext, FMeshPassProcessorRenderState& DrawRenderState, const FCanvas* Canvas) override;
 
 	/**
 	* Renders the canvas item.
@@ -1005,7 +1010,7 @@ public:
 	* @param Canvas - canvas currently being rendered
 	* @return true if anything rendered
 	*/
-	virtual bool Render_GameThread(const FCanvas* Canvas, FRenderThreadScope& RenderScope) override;
+	virtual bool Render_GameThread(const FCanvas* Canvas, FCanvasRenderThreadScope& RenderScope) override;
 
 	/**
 	* Determine if this is a matching set by comparing material,transform. All must match
@@ -1071,7 +1076,7 @@ private:
 			const FCanvas::FTransformEntry& InTransform);
 
 		void RenderTiles(
-			FRHICommandListImmediate& RHICmdList,
+			FCanvasRenderContext& RenderContext,
 			FMeshPassProcessorRenderState& DrawRenderState,
 			const FSceneView& View,
 			bool bIsHitTesting,
@@ -1153,7 +1158,7 @@ public:
 	* @param RHICmdList - command list to use
 	* @return true if anything rendered
 	*/
-	virtual bool Render_RenderThread(FRHICommandListImmediate& RHICmdList, FMeshPassProcessorRenderState& DrawRenderState, const FCanvas* Canvas) override;
+	virtual bool Render_RenderThread(FCanvasRenderContext& RenderContext, FMeshPassProcessorRenderState& DrawRenderState, const FCanvas* Canvas) override;
 
 	/**
 	* Renders the canvas item.
@@ -1162,7 +1167,7 @@ public:
 	* @param Canvas - canvas currently being rendered
 	* @return true if anything rendered
 	*/
-	virtual bool Render_GameThread(const FCanvas* Canvas, FRenderThreadScope& RenderScope) override;
+	virtual bool Render_GameThread(const FCanvas* Canvas, FCanvasRenderThreadScope& RenderScope) override;
 
 	/**
 	* Determine if this is a matching set by comparing material,transform. All must match
@@ -1259,7 +1264,7 @@ private:
 		}
 
 		void RenderTriangles(
-			FRHICommandListImmediate& RHICmdList,
+			FCanvasRenderContext& RenderContext,
 			FMeshPassProcessorRenderState& DrawRenderState,
 			const FSceneView& View,
 			bool bIsHitTesting,
diff --git a/Engine/Source/Runtime/Engine/Public/UnrealClient.h b/Engine/Source/Runtime/Engine/Public/UnrealClient.h
--- a/Engine/Source/Runtime/Engine/Public/UnrealClient.h
+++ b/Engine/Source/Runtime/Engine/Public/UnrealClient.h
@@ -15,6 +15,7 @@
 #include "RHI.h"
 #include "RenderResource.h"
 #include "HitProxies.h"
+#include "RenderGraphDefinitions.h"
 
 class FCanvas;
 class FViewport;
@@ -39,6 +40,11 @@ public:
 	ENGINE_API virtual const FTexture2DRHIRef& GetRenderTargetTexture() const;
 	ENGINE_API virtual FUnorderedAccessViewRHIRef GetRenderTargetUAV() const;
 
+	/**
+	 * Returns a valid RDG texture for this render target.
+	 */
+	ENGINE_API virtual FRDGTextureRef GetRenderTargetTexture(FRDGBuilder& GraphBuilder) const;
+
 	// Properties.
 	virtual FIntPoint GetSizeXY() const = 0;
 
diff --git a/Engine/Source/Runtime/RenderCore/Public/RendererInterface.h b/Engine/Source/Runtime/RenderCore/Public/RendererInterface.h
--- a/Engine/Source/Runtime/RenderCore/Public/RendererInterface.h
+++ b/Engine/Source/Runtime/RenderCore/Public/RendererInterface.h
@@ -19,6 +19,7 @@
 #include "RenderGraphDefinitions.h"
 
 class FCanvas;
+class FCanvasRenderContext;
 class FMaterial;
 class FSceneInterface;
 class FSceneView;
@@ -737,7 +738,7 @@ public:
 	virtual void InitializeSystemTextures(FRHICommandListImmediate& RHICmdList) = 0;
 
 	/** Draws a tile mesh element with the specified view. */
-	virtual void DrawTileMesh(FRHICommandListImmediate& RHICmdList, struct FMeshPassProcessorRenderState& DrawRenderState, const FSceneView& View, FMeshBatch& Mesh, bool bIsHitTesting, const class FHitProxyId& HitProxyId, bool bUse128bitRT = false) = 0;
+	virtual void DrawTileMesh(FCanvasRenderContext& RenderContext, struct FMeshPassProcessorRenderState& DrawRenderState, const FSceneView& View, FMeshBatch& Mesh, bool bIsHitTesting, const class FHitProxyId& HitProxyId, bool bUse128bitRT = false) = 0;
 
 	virtual const TSet<FSceneInterface*>& GetAllocatedScenes() = 0;
 
diff --git a/Engine/Source/Runtime/Renderer/Private/Renderer.cpp b/Engine/Source/Runtime/Renderer/Private/Renderer.cpp
--- a/Engine/Source/Runtime/Renderer/Private/Renderer.cpp
+++ b/Engine/Source/Runtime/Renderer/Private/Renderer.cpp
@@ -34,6 +34,7 @@
 #include "ScreenSpaceDenoise.h"
 #include "VT/VirtualTextureSystem.h"
 #include "PostProcess/TemporalAA.h"
+#include "CanvasRender.h"
 DEFINE_LOG_CATEGORY(LogRenderer);
 
 IMPLEMENT_MODULE(FRendererModule, Renderer);
@@ -97,15 +98,16 @@ BEGIN_SHADER_PARAMETER_STRUCT(FDrawTileMeshPassParameters, )
 	SHADER_PARAMETER_RDG_UNIFORM_BUFFER(FTranslucentBasePassUniformParameters, TranslucentBasePass)
 	SHADER_PARAMETER_RDG_UNIFORM_BUFFER(FOpaqueBasePassUniformParameters, OpaqueBasePass)
 	SHADER_PARAMETER_RDG_UNIFORM_BUFFER(FMobileBasePassUniformParameters, MobileBasePass)
+	RENDER_TARGET_BINDING_SLOTS()
 END_SHADER_PARAMETER_STRUCT()
 
-void FRendererModule::DrawTileMesh(FRHICommandListImmediate& RHICmdList, FMeshPassProcessorRenderState& DrawRenderState, const FSceneView& SceneView, FMeshBatch& Mesh, bool bIsHitTesting, const FHitProxyId& HitProxyId, bool bUse128bitRT)
+void FRendererModule::DrawTileMesh(FCanvasRenderContext& RenderContext, FMeshPassProcessorRenderState& DrawRenderState, const FSceneView& SceneView, FMeshBatch& Mesh, bool bIsHitTesting, const FHitProxyId& HitProxyId, bool bUse128bitRT)
 {
 	if (!GUsingNullRHI)
 	{
 		// Create an FViewInfo so we can initialize its RHI resources
 		//@todo - reuse this view for multiple tiles, this is going to be slow for each tile
-		FViewInfo View(&SceneView);
+		FViewInfo& View = *RenderContext.Alloc<FViewInfo>(&SceneView);
 		View.ViewRect = View.UnscaledViewRect;
 
 		const auto FeatureLevel = View.GetFeatureLevel();
@@ -199,16 +201,15 @@ void FRendererModule::DrawTileMesh(FRHICommandListImmediate& RHICmdList, FMeshPa
 		//get the blend mode of the material
 		const EBlendMode MaterialBlendMode = Mesh.MaterialRenderProxy->GetIncompleteMaterialWithFallback(FeatureLevel).GetBlendMode();
 
-		GSystemTextures.InitializeTextures(RHICmdList, FeatureLevel);
-		FMemMark Mark(FMemStack::Get());
-		FRDGBuilder GraphBuilder(RHICmdList);
+		FRDGBuilder& GraphBuilder = RenderContext.GraphBuilder;
 
-		// Initialize the RDG read-only system textures.
-		FRDGSystemTextures::Create(GraphBuilder);
-
-		const ERDGPassFlags PassFlags = ERDGPassFlags::Raster | ERDGPassFlags::SkipRenderPass | ERDGPassFlags::NeverCull;
+		if (!FRDGSystemTextures::IsValid(GraphBuilder))
+		{
+			FRDGSystemTextures::Create(GraphBuilder);
+		}
 
 		auto* PassParameters = GraphBuilder.AllocParameters<FDrawTileMeshPassParameters>();
+		PassParameters->RenderTargets[0] = FRenderTargetBinding(RenderContext.GetRenderTarget(), ERenderTargetLoadAction::ELoad);
 		PassParameters->View = View.ViewUniformBuffer;
 		PassParameters->ReflectionCapture = EmptyReflectionCaptureUniformBuffer;
 
@@ -224,10 +225,10 @@ void FRendererModule::DrawTileMesh(FRHICommandListImmediate& RHICmdList, FMeshPa
 			{
 				PassParameters->DebugViewMode = CreateDebugViewModePassUniformBuffer(GraphBuilder, View, nullptr);
 
-				GraphBuilder.AddPass({}, PassParameters, PassFlags,
-					[Scene, &View, &Mesh](FRHICommandListImmediate& InRHICmdList)
+				RenderContext.AddPass({}, PassParameters,
+					[Scene, &View, &Mesh](FRHICommandListImmediate& RHICmdList)
 				{
-					DrawDynamicMeshPass(View, InRHICmdList, [&](FMeshPassDrawListContext* InDrawListContext)
+					DrawDynamicMeshPass(View, RHICmdList, [&](FMeshPassDrawListContext* InDrawListContext)
 					{
 						FDebugViewModeMeshProcessor PassMeshProcessor(
 							Scene,
@@ -248,10 +249,10 @@ void FRendererModule::DrawTileMesh(FRHICommandListImmediate& RHICmdList, FMeshPa
 			{
 				PassParameters->TranslucentBasePass = CreateTranslucentBasePassUniformBuffer(GraphBuilder, View);
 
-				GraphBuilder.AddPass({}, PassParameters, PassFlags,
-					[Scene, &View, &Mesh, &DrawRenderState, bUse128bitRT](FRHICommandListImmediate& InRHICmdList)
+				RenderContext.AddPass({}, PassParameters,
+					[Scene, &View, &Mesh, DrawRenderState, bUse128bitRT](FRHICommandListImmediate& RHICmdList)
 				{
-					DrawDynamicMeshPass(View, InRHICmdList, [&](FDynamicPassMeshDrawListContext* DynamicMeshPassContext)
+					DrawDynamicMeshPass(View, RHICmdList, [&](FDynamicPassMeshDrawListContext* DynamicMeshPassContext)
 					{
 						FBasePassMeshProcessor PassMeshProcessor(
 							Scene,
@@ -271,20 +272,23 @@ void FRendererModule::DrawTileMesh(FRHICommandListImmediate& RHICmdList, FMeshPa
 			{
 				PassParameters->MobileBasePass = CreateMobileBasePassUniformBuffer(GraphBuilder, View, EMobileBasePass::Translucent);
 
-				DrawDynamicMeshPass(View, RHICmdList,
-					[Scene, &View, &DrawRenderState, &Mesh](FDynamicPassMeshDrawListContext* DynamicMeshPassContext)
+				RenderContext.AddPass({}, PassParameters,
+					[Scene, &View, DrawRenderState, &Mesh](FRHICommandListImmediate& RHICmdList)
 				{
-					FMobileBasePassMeshProcessor PassMeshProcessor(
-						Scene,
-						View.GetFeatureLevel(),
-						&View,
-						DrawRenderState,
-						DynamicMeshPassContext,
-						FMobileBasePassMeshProcessor::EFlags::None,
-						ETranslucencyPass::TPT_AllTranslucency);
-					
-					const uint64 DefaultBatchElementMask = ~0ull;
-					PassMeshProcessor.AddMeshBatch(Mesh, DefaultBatchElementMask, nullptr);
+					DrawDynamicMeshPass(View, RHICmdList, [&](FDynamicPassMeshDrawListContext* DynamicMeshPassContext)
+					{
+						FMobileBasePassMeshProcessor PassMeshProcessor(
+							Scene,
+							View.GetFeatureLevel(),
+							&View,
+							DrawRenderState,
+							DynamicMeshPassContext,
+							FMobileBasePassMeshProcessor::EFlags::None,
+							ETranslucencyPass::TPT_AllTranslucency);
+
+						const uint64 DefaultBatchElementMask = ~0ull;
+						PassMeshProcessor.AddMeshBatch(Mesh, DefaultBatchElementMask, nullptr);
+					});
 				});
 			}
 		}
@@ -300,18 +304,21 @@ void FRendererModule::DrawTileMesh(FRHICommandListImmediate& RHICmdList, FMeshPa
 				ensureMsgf(HitProxyId == Mesh.BatchHitProxyId, TEXT("Only Mesh.BatchHitProxyId is used for hit testing."));
 
 #if WITH_EDITOR
-				DrawDynamicMeshPass(View, RHICmdList,
-					[Scene, &View, &DrawRenderState, &Mesh](FDynamicPassMeshDrawListContext* DynamicMeshPassContext)
+				RenderContext.AddPass({}, PassParameters,
+					[Scene, &View, DrawRenderState, &Mesh](FRHICommandListImmediate& RHICmdList)
 				{
-					FHitProxyMeshProcessor PassMeshProcessor(
-						Scene,
-						&View,
-						false,
-						DrawRenderState,
-						DynamicMeshPassContext);
-
-					const uint64 DefaultBatchElementMask = ~0ull;
-					PassMeshProcessor.AddMeshBatch(Mesh, DefaultBatchElementMask, nullptr);
+					DrawDynamicMeshPass(View, RHICmdList, [&](FDynamicPassMeshDrawListContext* DynamicMeshPassContext)
+					{
+						FHitProxyMeshProcessor PassMeshProcessor(
+							Scene,
+							&View,
+							false,
+							DrawRenderState,
+							DynamicMeshPassContext);
+
+						const uint64 DefaultBatchElementMask = ~0ull;
+						PassMeshProcessor.AddMeshBatch(Mesh, DefaultBatchElementMask, nullptr);
+					});
 				});
 #endif
 			}
@@ -321,10 +328,10 @@ void FRendererModule::DrawTileMesh(FRHICommandListImmediate& RHICmdList, FMeshPa
 				{
 					PassParameters->OpaqueBasePass = CreateOpaqueBasePassUniformBuffer(GraphBuilder, View);
 
-					GraphBuilder.AddPass({}, PassParameters, PassFlags,
-						[Scene, &View, &DrawRenderState, &Mesh, bUse128bitRT](FRHICommandListImmediate& InRHICmdList)
+					RenderContext.AddPass({}, PassParameters,
+						[Scene, &View, DrawRenderState, &Mesh, bUse128bitRT](FRHICommandListImmediate& RHICmdList)
 					{
-						DrawDynamicMeshPass(View, InRHICmdList,
+						DrawDynamicMeshPass(View, RHICmdList,
 							[&](FDynamicPassMeshDrawListContext* DynamicMeshPassContext)
 						{
 							FBasePassMeshProcessor PassMeshProcessor(
@@ -343,26 +350,27 @@ void FRendererModule::DrawTileMesh(FRHICommandListImmediate& RHICmdList, FMeshPa
 				else // Mobile
 				{
 					PassParameters->MobileBasePass = CreateMobileBasePassUniformBuffer(GraphBuilder, View, EMobileBasePass::Opaque);
-					
-					DrawDynamicMeshPass(View, RHICmdList,
-						[Scene, &View, &DrawRenderState, &Mesh](FDynamicPassMeshDrawListContext* DynamicMeshPassContext)
+
+					RenderContext.AddPass({}, PassParameters,
+						[Scene, &View, DrawRenderState, &Mesh](FRHICommandListImmediate& RHICmdList)
 					{
-						FMobileBasePassMeshProcessor PassMeshProcessor(
-							Scene,
-							View.GetFeatureLevel(),
-							&View,
-							DrawRenderState,
-							DynamicMeshPassContext,
-							FMobileBasePassMeshProcessor::EFlags::CanReceiveCSM);
-						
-						const uint64 DefaultBatchElementMask = ~0ull;
-						PassMeshProcessor.AddMeshBatch(Mesh, DefaultBatchElementMask, nullptr);
+						DrawDynamicMeshPass(View, RHICmdList, [&](FDynamicPassMeshDrawListContext* DynamicMeshPassContext)
+						{
+							FMobileBasePassMeshProcessor PassMeshProcessor(
+								Scene,
+								View.GetFeatureLevel(),
+								&View,
+								DrawRenderState,
+								DynamicMeshPassContext,
+								FMobileBasePassMeshProcessor::EFlags::CanReceiveCSM);
+
+							const uint64 DefaultBatchElementMask = ~0ull;
+							PassMeshProcessor.AddMeshBatch(Mesh, DefaultBatchElementMask, nullptr);
+						});
 					});
 				}
 			}
 		}
-
-		GraphBuilder.Execute();
 	}
 }
 
diff --git a/Engine/Source/Runtime/Renderer/Private/RendererModule.h b/Engine/Source/Runtime/Renderer/Private/RendererModule.h
--- a/Engine/Source/Runtime/Renderer/Private/RendererModule.h
+++ b/Engine/Source/Runtime/Renderer/Private/RendererModule.h
@@ -10,7 +10,6 @@
 #include "RendererInterface.h"
 #include "Shader.h"
 
-class FCanvas;
 class FHitProxyId;
 class FLightCacheInterface;
 class FMaterial;
@@ -47,7 +46,7 @@ public:
 	virtual uint32 GetNumDynamicLightsAffectingPrimitive(const FPrimitiveSceneInfo* PrimitiveSceneInfo,const FLightCacheInterface* LCI) override;
 	virtual void OnWorldCleanup(UWorld* World, bool bSessionEnded, bool bCleanupResources, bool bWorldChanged) override;
 	virtual void InitializeSystemTextures(FRHICommandListImmediate& RHICmdList);
-	virtual void DrawTileMesh(FRHICommandListImmediate& RHICmdList, FMeshPassProcessorRenderState& DrawRenderState, const FSceneView& View, FMeshBatch& Mesh, bool bIsHitTesting, const FHitProxyId& HitProxyId, bool bUse128bitRT = false) override;
+	virtual void DrawTileMesh(FCanvasRenderContext& RenderContext, FMeshPassProcessorRenderState& DrawRenderState, const FSceneView& View, FMeshBatch& Mesh, bool bIsHitTesting, const FHitProxyId& HitProxyId, bool bUse128bitRT = false) override;
 	virtual void DebugLogOnCrash() override;
 	virtual void GPUBenchmark(FSynthBenchmarkResults& InOut, float WorkScale) override;
 	virtual void ExecVisualizeTextureCmd(const FString& Cmd) override;
diff --git a/Engine/Source/Runtime/Renderer/Private/ScreenPass.h b/Engine/Source/Runtime/Renderer/Private/ScreenPass.h
--- a/Engine/Source/Runtime/Renderer/Private/ScreenPass.h
+++ b/Engine/Source/Runtime/Renderer/Private/ScreenPass.h
@@ -537,28 +537,6 @@ void RENDERER_API AddDrawTexturePass(
 	FScreenPassTexture Input,
 	FScreenPassRenderTarget Output);
 
-/** Helper function render a canvas to an output texture. Must be called within a render pass with Output as the render target. */
-template <typename TFunction>
-void DrawCanvasPass(
-	FRHICommandListImmediate& RHICmdList,
-	const FViewInfo& View,
-	FScreenPassTexture Output,
-	TFunction&& Function)
-{
-	check(Output.IsValid());
-
-	const FSceneViewFamily& ViewFamily = *View.Family;
-	FRenderTargetTemp TempRenderTarget(static_cast<FRHITexture2D*>(Output.Texture->GetRHI()), Output.ViewRect.Size());
-	FCanvas Canvas(&TempRenderTarget, nullptr, ViewFamily.CurrentRealTime, ViewFamily.CurrentWorldTime, ViewFamily.DeltaWorldTime, View.GetFeatureLevel());
-	Canvas.SetRenderTargetRect(Output.ViewRect);
-
-	Function(Canvas);
-
-	const bool bFlush = false;
-	const bool bInsideRenderPass = true;
-	Canvas.Flush_RenderThread(RHICmdList, bFlush, bInsideRenderPass);
-}
-
 template <typename TFunction>
 FORCEINLINE void AddRenderTargetPass(
 	FRDGBuilder& GraphBuilder,
@@ -579,10 +557,16 @@ FORCEINLINE void AddDrawCanvasPass(
 	FScreenPassRenderTarget Output,
 	TFunction Function)
 {
-	AddRenderTargetPass(GraphBuilder, MoveTemp(PassName), Output, [Output, &View, Function](FRHICommandListImmediate& RHICmdList)
-	{
-		DrawCanvasPass(RHICmdList, View, Output, Function);
-	});
+	check(Output.IsValid());
+
+	const FSceneViewFamily& ViewFamily = *View.Family;
+	FCanvas& Canvas = *FCanvas::Create(GraphBuilder, Output.Texture, nullptr, ViewFamily.CurrentRealTime, ViewFamily.CurrentWorldTime, ViewFamily.DeltaWorldTime, View.GetFeatureLevel());
+	Canvas.SetRenderTargetRect(Output.ViewRect);
+
+	Function(Canvas);
+
+	const bool bFlush = false;
+	Canvas.Flush_RenderThread(GraphBuilder, bFlush);
 }
 
 enum class EDownsampleDepthFilter
diff --git a/Engine/Source/Runtime/Renderer/Private/SystemTextures.cpp b/Engine/Source/Runtime/Renderer/Private/SystemTextures.cpp
--- a/Engine/Source/Runtime/Renderer/Private/SystemTextures.cpp
+++ b/Engine/Source/Runtime/Renderer/Private/SystemTextures.cpp
@@ -46,6 +46,11 @@ const FRDGSystemTextures& FRDGSystemTextures::Get(FRDGBuilder& GraphBuilder)
 	return *SystemTextures;
 }
 
+bool FRDGSystemTextures::IsValid(FRDGBuilder& GraphBuilder)
+{
+	return GraphBuilder.Blackboard.Get<FRDGSystemTextures>() != nullptr;
+}
+
 /** The global render targets used for scene rendering. */
 TGlobalResource<FSystemTextures> GSystemTextures;
 
diff --git a/Engine/Source/Runtime/Renderer/Private/SystemTextures.h b/Engine/Source/Runtime/Renderer/Private/SystemTextures.h
--- a/Engine/Source/Runtime/Renderer/Private/SystemTextures.h
+++ b/Engine/Source/Runtime/Renderer/Private/SystemTextures.h
@@ -17,6 +17,9 @@ struct FRDGSystemTextures
 	/** Returns the instance for the builder. Must be called after Create. */
 	RENDERER_API static const FRDGSystemTextures& Get(FRDGBuilder& GraphBuilder);
 
+	/** Returns whether the system textures have been created. */
+	RENDERER_API static bool IsValid(FRDGBuilder& GraphBuilder);
+
 	FRDGTextureRef White{};
 	FRDGTextureRef Black{};
 	FRDGTextureRef BlackAlphaOne{};
diff --git a/Engine/Source/Runtime/SlateRHIRenderer/Private/SlateRHIRenderingPolicy.cpp b/Engine/Source/Runtime/SlateRHIRenderer/Private/SlateRHIRenderingPolicy.cpp
--- a/Engine/Source/Runtime/SlateRHIRenderer/Private/SlateRHIRenderingPolicy.cpp
+++ b/Engine/Source/Runtime/SlateRHIRenderer/Private/SlateRHIRenderingPolicy.cpp
@@ -721,6 +721,7 @@ void FSlateRHIRenderingPolicy::DrawElements(
 				RPInfo.DepthStencilRenderTarget.Action = EDepthStencilTargetActions::DontLoad_DontStore;
 				RPInfo.DepthStencilRenderTarget.ExclusiveDepthStencil = FExclusiveDepthStencil::DepthNop_StencilNop;
 			}
+			TransitionRenderPassTargets(RHICmdList, RPInfo);
 			RHICmdList.BeginRenderPass(RPInfo, TEXT("RestartingSlateDrawElements"));
 
 			// Something may have messed with the viewport size so set it back to the full target.
