LevelInstance: Fixed issue where attempting to overwrite a packed level instance bp while saving would fail.

diff --git a/Engine/Source/Runtime/Engine/Private/LevelInstance/Packed/PackedLevelInstanceBuilder.cpp b/Engine/Source/Runtime/Engine/Private/LevelInstance/Packed/PackedLevelInstanceBuilder.cpp
--- a/Engine/Source/Runtime/Engine/Private/LevelInstance/Packed/PackedLevelInstanceBuilder.cpp
+++ b/Engine/Source/Runtime/Engine/Private/LevelInstance/Packed/PackedLevelInstanceBuilder.cpp
@@ -30,6 +30,10 @@
 #include "AssetToolsModule.h"
 #include "Factories/BlueprintFactory.h"
 
+#include "IContentBrowserSingleton.h"
+#include "ContentBrowserModule.h"
+#include "EditorDirectories.h"
+
 #define LOCTEXT_NAMESPACE "FPackedLevelInstanceBuilder"
 
 void FPackedLevelInstanceBuilderContext::ClusterLevelActor(AActor* InActor)
@@ -67,7 +71,37 @@ const FString& FPackedLevelInstanceBuilder::GetPackedBPSuffix()
 	return BPSuffix;
 }
 
-UBlueprint* FPackedLevelInstanceBuilder::CreatePackedLevelInstanceBlueprint(const FString& InPackagePath, const FString& InAssetName, bool bInCompile)
+UBlueprint* FPackedLevelInstanceBuilder::CreatePackedLevelInstanceBlueprintWithDialog(const FString& InPackagePath, const FString& InAssetName, bool bInCompile)
+{
+	FSaveAssetDialogConfig SaveAssetDialogConfig;
+	SaveAssetDialogConfig.DialogTitleOverride = LOCTEXT("SaveAssetDialogTitle", "Save Asset As");
+	SaveAssetDialogConfig.DefaultPath = InPackagePath;
+	SaveAssetDialogConfig.DefaultAssetName = InAssetName;
+	SaveAssetDialogConfig.ExistingAssetPolicy = ESaveAssetDialogExistingAssetPolicy::AllowButWarn;
+
+	FContentBrowserModule& ContentBrowserModule = FModuleManager::LoadModuleChecked<FContentBrowserModule>("ContentBrowser");
+	FString SaveObjectPath = ContentBrowserModule.Get().CreateModalSaveAssetDialog(SaveAssetDialogConfig);
+
+	if (!SaveObjectPath.IsEmpty())
+	{
+		TSoftObjectPtr<UBlueprint> ExistingBPAsset(SaveObjectPath);
+
+		if (UBlueprint* BP = ExistingBPAsset.LoadSynchronous())
+		{
+			return BP;
+		}
+		
+		const FString SavePackageName = FPackageName::ObjectPathToPackageName(SaveObjectPath);
+		const FString SavePackagePath = FPaths::GetPath(SavePackageName);
+		const FString SaveAssetName = FPaths::GetBaseFilename(SavePackageName);
+
+		return CreatePackedLevelInstanceBlueprint(SaveAssetName, SavePackagePath, bInCompile);
+	}
+
+	return nullptr;
+}
+
+UBlueprint* FPackedLevelInstanceBuilder::CreatePackedLevelInstanceBlueprint(const FString& InSaveAssetName, const FString& InPackagePath, bool bInCompile)
 {
 	IAssetTools& AssetTools = FAssetToolsModule::GetModule().Get();
 
@@ -75,19 +109,25 @@ UBlueprint* FPackedLevelInstanceBuilder::CreatePackedLevelInstanceBlueprint(cons
 	BlueprintFactory->ParentClass = APackedLevelInstance::StaticClass();
 	BlueprintFactory->bSkipClassPicker = true;
 
-	if (UBlueprint* NewBP = Cast<UBlueprint>(AssetTools.CreateAssetWithDialog(InAssetName, InPackagePath, UBlueprint::StaticClass(), BlueprintFactory, FName("Create LevelInstance Blueprint"))))
+	FEditorDelegates::OnConfigureNewAssetProperties.Broadcast(BlueprintFactory);
+	if (BlueprintFactory->ConfigureProperties())
 	{
-		APackedLevelInstance* CDO = CastChecked<APackedLevelInstance>(NewBP->GeneratedClass->GetDefaultObject());
-		CDO->BlueprintAsset = NewBP;
+		FEditorDirectories::Get().SetLastDirectory(ELastDirectory::NEW_ASSET, InPackagePath);
 
-		if (bInCompile)
+		if (UBlueprint* NewBP = Cast<UBlueprint>(AssetTools.CreateAsset(InSaveAssetName, InPackagePath, UBlueprint::StaticClass(), BlueprintFactory, FName("Create LevelInstance Blueprint"))))
 		{
-			FKismetEditorUtilities::CompileBlueprint(NewBP, EBlueprintCompileOptions::SkipGarbageCollection);
-		}
+			APackedLevelInstance* CDO = CastChecked<APackedLevelInstance>(NewBP->GeneratedClass->GetDefaultObject());
+			CDO->BlueprintAsset = NewBP;
 
-		AssetTools.SyncBrowserToAssets(TArray<UObject*>{ NewBP });
+			if (bInCompile)
+			{
+				FKismetEditorUtilities::CompileBlueprint(NewBP, EBlueprintCompileOptions::None);
+			}
 
-		return NewBP;
+			AssetTools.SyncBrowserToAssets(TArray<UObject*>{ NewBP });
+
+			return NewBP;
+		}
 	}
 
 	return nullptr;
@@ -332,14 +372,6 @@ bool FPackedLevelInstanceBuilder::CreateOrUpdateBlueprintFromPacked(APackedLevel
 	if (!InBlueprintAsset.IsNull())
 	{
 		BP = InBlueprintAsset.LoadSynchronous();
-		if (BP && BP->SimpleConstructionScript)
-		{
-			TArray<USCS_Node*> AllNodes = BP->SimpleConstructionScript->GetAllNodes();
-			for (USCS_Node* Node : AllNodes)
-			{
-				BP->SimpleConstructionScript->RemoveNodeAndPromoteChildren(Node);
-			}
-		}
 	}
 
 	if (!BP)
@@ -351,10 +383,19 @@ bool FPackedLevelInstanceBuilder::CreateOrUpdateBlueprintFromPacked(APackedLevel
 		FString PackagePath = LongPackageName.Mid(0, LastSlashIndex == INDEX_NONE ? MAX_int32 : LastSlashIndex);
 		FString AssetName = InActor->GetWorldAsset().GetAssetName() + GetPackedBPSuffix();
 		const bool bCompile = false;
-		BP = CreatePackedLevelInstanceBlueprint(PackagePath, AssetName, bCompile);
+
+		BP = CreatePackedLevelInstanceBlueprintWithDialog(PackagePath, AssetName, bCompile);
 	}
-		
-	if (BP == nullptr)
+
+	if (BP && BP->SimpleConstructionScript)
+	{
+		TArray<USCS_Node*> AllNodes = BP->SimpleConstructionScript->GetAllNodes();
+		for (USCS_Node* Node : AllNodes)
+		{
+			BP->SimpleConstructionScript->RemoveNodeAndPromoteChildren(Node);
+		}
+	}
+	else
 	{
 		return false;
 	}
diff --git a/Engine/Source/Runtime/Engine/Public/LevelInstance/Packed/PackedLevelInstanceBuilder.h b/Engine/Source/Runtime/Engine/Public/LevelInstance/Packed/PackedLevelInstanceBuilder.h
--- a/Engine/Source/Runtime/Engine/Public/LevelInstance/Packed/PackedLevelInstanceBuilder.h
+++ b/Engine/Source/Runtime/Engine/Public/LevelInstance/Packed/PackedLevelInstanceBuilder.h
@@ -42,9 +42,11 @@ public:
 	/* Update existing Blueprint */
 	void UpdateBlueprint(UBlueprint* Blueprint);
 
-	/* Creates a new APackedLevelInstance Blueprint using InPackagePath/InAssetName as hint for path */
 	static const FString& GetPackedBPSuffix();
-	static UBlueprint* CreatePackedLevelInstanceBlueprint(const FString& InPackagePath, const FString& InAssetName, bool bInCompile);
+	/* Creates a new APackedLevelInstance Blueprint using InPackagePath/InAssetName as hint for path. Prompts the user to input the final asset name. */
+	static UBlueprint* CreatePackedLevelInstanceBlueprintWithDialog(const FString& InPackagePath, const FString& InAssetName, bool bInCompile);
+	/* Creates a new APackedLevelInstance Blueprint using InPackagePath/InAssetName */
+	static UBlueprint* CreatePackedLevelInstanceBlueprint(const FString& InSaveAssetName, const FString& InPackagePath, bool bInCompile);
 	
 private:
 	/* Create/Updates a APackedLevelInstance Blueprint from InPackedActor (will overwrite existing asset or show a dialog if InBlueprintAsset is null) */
