Add GUID to FVirtualTextureBuiltData, and to VT chunk data.
Since VT data is split across multiple DDC keys, this GUID ensures we're always using data that was uploaded to the DDC at the same time
In theory this shouldn't be a requirement, but we currently have some non-determinism when building VT chunk data...this means that the chunk data may not match the built data if they were created at separate times

diff --git a/Engine/Source/Runtime/Engine/Private/TextureDerivedData.cpp b/Engine/Source/Runtime/Engine/Private/TextureDerivedData.cpp
--- a/Engine/Source/Runtime/Engine/Private/TextureDerivedData.cpp
+++ b/Engine/Source/Runtime/Engine/Private/TextureDerivedData.cpp
@@ -53,7 +53,7 @@
 
 // This GUID is mixed into DDC version for virtual textures only, this allows updating DDC version for VT without invalidating DDC for all textures
 // This is useful during development, but once large numbers of VT are present in shipped content, it will have the same problem as TEXTURE_DERIVEDDATA_VER
-#define TEXTURE_VT_DERIVEDDATA_VER	TEXT("744516E4C74F4867BDAC2E6D49596219")
+#define TEXTURE_VT_DERIVEDDATA_VER	TEXT("39E589A391EF40459CA8F45359851DF3")
 
 #if ENABLE_COOK_STATS
 namespace TextureCookStats
@@ -1061,12 +1061,9 @@ bool FTexturePlatformData::TryInlineMipData(int32 FirstMipToLoad, UTexture* Text
 				}
 				if (bLoadedFromDDC)
 				{
-					int32 ChunkSize = 0;
-					FMemoryReader Ar(TempData, /*bIsPersistent=*/ true);
-					Ar << ChunkSize;
 					Chunk.BulkData.Lock(LOCK_READ_WRITE);
-					void* ChunkData = Chunk.BulkData.Realloc(ChunkSize);
-					Ar.Serialize(ChunkData, ChunkSize);
+					void* ChunkData = Chunk.BulkData.Realloc(TempData.Num());
+					FMemory::Memcpy(ChunkData, TempData.GetData(), TempData.Num());
 					Chunk.BulkData.Unlock();
 					Chunk.DerivedDataKey.Empty();
 				}
diff --git a/Engine/Source/Runtime/Engine/Private/VT/VirtualTextureBuiltData.cpp b/Engine/Source/Runtime/Engine/Private/VT/VirtualTextureBuiltData.cpp
--- a/Engine/Source/Runtime/Engine/Private/VT/VirtualTextureBuiltData.cpp
+++ b/Engine/Source/Runtime/Engine/Private/VT/VirtualTextureBuiltData.cpp
@@ -50,6 +50,7 @@ void FVirtualTextureBuiltData::Serialize(FArchive& Ar, UObject* Owner, int32 Fir
 	bool bCooked = Ar.IsCooking();
 	Ar << bCooked;
 
+	Ar << VersionGuid;
 	Ar << NumLayers;
 	Ar << WidthInBlocks;
 	Ar << HeightInBlocks;
@@ -187,6 +188,8 @@ bool FVirtualTextureBuiltData::ValidateData(FStringView const& InDDCDebugContext
 	TArray<uint8> UncompressedResult;
 	TArray<uint8> ChunkDataDDC;
 
+	const FString TextureName(InDDCDebugContext);
+
 	bool bResult = true;
 	for (int32 ChunkIndex = 0; bResult && ChunkIndex < Chunks.Num(); ++ChunkIndex)
 	{
@@ -206,15 +209,29 @@ bool FVirtualTextureBuiltData::ValidateData(FStringView const& InDDCDebugContext
 		{
 			ChunkDataDDC.Reset();
 			const bool bDDCResult = GetDerivedDataCacheRef().GetSynchronous(*Chunk.DerivedDataKey, ChunkDataDDC, InDDCDebugContext);
-			check(bDDCResult);
-			check(ChunkDataDDC.Num() > 4);
-			ChunkData = ChunkDataDDC.GetData() + 4;
-			ChunkDataSize = (uint32)(ChunkDataDDC.Num() - 4);
+			if (!bDDCResult)
+			{
+				UE_LOG(LogTexture, Error, TEXT("Virtual Texture %s failed to retrieve DDC data (%s) for chunk %d"), *TextureName, *Chunk.DerivedDataKey, ChunkIndex);
+				bResult = false;
+				break;
+			}
+
+			ChunkData = ChunkDataDDC.GetData();
+			ChunkDataSize = (uint32)ChunkDataDDC.Num();
 		}
 #endif // WITH_EDITORONLY_DATA
 
-		if (!ChunkData)
+		if (!ChunkData || ChunkDataSize < sizeof(FVirtualTextureChunkHeader))
 		{
+			UE_LOG(LogTexture, Error, TEXT("Virtual Texture %s has invalid size %d for chunk %d"), *TextureName, ChunkDataSize, ChunkIndex);
+			bResult = false;
+			break;
+		}
+
+		FVirtualTextureChunkHeader* ChunkHeader = (FVirtualTextureChunkHeader*)ChunkData;
+		if (ChunkHeader->VersionGuid != VersionGuid)
+		{
+			UE_LOG(LogTexture, Error, TEXT("Virtual Texture %s has mismatched GUID for chunk %d"), *TextureName, ChunkIndex);
 			bResult = false;
 			break;
 		}
@@ -223,6 +240,7 @@ bool FVirtualTextureBuiltData::ValidateData(FStringView const& InDDCDebugContext
 		FSHA1::HashBuffer(ChunkData, ChunkDataSize, Hash.Hash);
 		if (Hash != Chunk.BulkDataHash)
 		{
+			UE_LOG(LogTexture, Error, TEXT("Virtual Texture %s has invalid hash for chunk %d"), *TextureName, ChunkIndex);
 			bResult = false;
 			break;
 		}
@@ -254,6 +272,7 @@ bool FVirtualTextureBuiltData::ValidateData(FStringView const& InDDCDebugContext
 							const bool bUncompressResult = FCompression::UncompressMemory(NAME_Zlib, UncompressedResult.GetData(), PackedOutputSize, &ChunkData[TileOffset], CompressedTileSize);
 							if (!bUncompressResult)
 							{
+								UE_LOG(LogTexture, Error, TEXT("Virtual Texture %s failed to validate compression for chunk %d"), *TextureName, ChunkIndex);
 								bResult = false;
 								break;
 							}
@@ -312,21 +331,17 @@ uint32 FVirtualTextureDataChunk::StoreInDerivedDataCache(const FString& InDerive
 	int32 BulkDataSizeInBytes = BulkData.GetBulkDataSize();
 	check(BulkDataSizeInBytes > 0);
 
-	TArray<uint8> DerivedData;
-	FMemoryWriter Ar(DerivedData, /*bIsPersistent=*/ true);
-	Ar << BulkDataSizeInBytes;
 	{
-		void* BulkChunkData = BulkData.Lock(LOCK_READ_ONLY);
-		Ar.Serialize(BulkChunkData, BulkDataSizeInBytes);
+		const uint8* BulkChunkData = (uint8*)BulkData.Lock(LOCK_READ_ONLY);
+		GetDerivedDataCacheRef().Put(*InDerivedDataKey, MakeArrayView(BulkChunkData, BulkDataSizeInBytes), TextureName, bReplaceExistingDDC);
 		BulkData.Unlock();
 	}
-	const uint32 Result = DerivedData.Num();
-	GetDerivedDataCacheRef().Put(*InDerivedDataKey, DerivedData, TextureName, bReplaceExistingDDC);
+
 	DerivedDataKey = InDerivedDataKey;
 	ShortenKey(DerivedDataKey, ShortDerivedDataKey);
 
 	// remove the actual bulkdata so when we serialize the owning FVirtualTextureBuiltData, this is actually serializing only the meta data
 	BulkData.RemoveBulkData();
-	return Result;
+	return BulkDataSizeInBytes;
 }
 #endif // WITH_EDITORONLY_DATA
diff --git a/Engine/Source/Runtime/Engine/Private/VT/VirtualTextureBuiltData.h b/Engine/Source/Runtime/Engine/Private/VT/VirtualTextureBuiltData.h
--- a/Engine/Source/Runtime/Engine/Private/VT/VirtualTextureBuiltData.h
+++ b/Engine/Source/Runtime/Engine/Private/VT/VirtualTextureBuiltData.h
@@ -28,6 +28,11 @@ enum class EVirtualTextureCodec : uint8
 	Max,			// Add new codecs before this entry
 };
 
+struct FVirtualTextureChunkHeader
+{
+	FGuid VersionGuid;
+};
+
 struct FVirtualTextureDataChunk
 {
 	FByteBulkData BulkData;
@@ -68,6 +73,7 @@ struct FVirtualTextureDataChunk
 
 struct FVirtualTextureBuiltData
 {
+	FGuid VersionGuid;
 	uint32 NumLayers;
 	uint32 NumMips;
 	uint32 Width; // Width of the texture in pixels. Note the physical width may be larger due to tiling
diff --git a/Engine/Source/Runtime/Engine/Private/VT/VirtualTextureChunkDDCCache.cpp b/Engine/Source/Runtime/Engine/Private/VT/VirtualTextureChunkDDCCache.cpp
--- a/Engine/Source/Runtime/Engine/Private/VT/VirtualTextureChunkDDCCache.cpp
+++ b/Engine/Source/Runtime/Engine/Private/VT/VirtualTextureChunkDDCCache.cpp
@@ -111,15 +111,14 @@ public:
 			return;
 		}
 
-		if (Results.Num() <= 4)
+		if (Results.Num() <= sizeof(FVirtualTextureChunkHeader))
 		{
 			UE_LOG(LogVTDiskCache, Error, TEXT("VT DDC data is too small %d (key: %s)"), Results.Num(), *Chunk->DerivedDataKey);
 			return;
 		}
 
-		// skip size embedded in DDC entry
-		const uint8* ChunkData = Results.GetData() + 4;
-		const int32 ChunkDataSize = Results.Num() - 4;
+		const uint8* ChunkData = Results.GetData();
+		const int32 ChunkDataSize = Results.Num();
 
 		FSHAHash Hash;
 		FSHA1::HashBuffer(ChunkData, ChunkDataSize, Hash.Hash);
diff --git a/Engine/Source/Runtime/Engine/Private/VT/VirtualTextureDataBuilder.cpp b/Engine/Source/Runtime/Engine/Private/VT/VirtualTextureDataBuilder.cpp
--- a/Engine/Source/Runtime/Engine/Private/VT/VirtualTextureDataBuilder.cpp
+++ b/Engine/Source/Runtime/Engine/Private/VT/VirtualTextureDataBuilder.cpp
@@ -334,6 +334,7 @@ void FVirtualTextureDataBuilder::Build(const FTextureSourceData& InSourceData, c
 	//NOTE: OutData may point to a previously build data so it is important to
 	//properly initialize all fields and not assume this is a freshly constructed object
 
+	OutData.VersionGuid = FGuid::NewGuid();
 	OutData.TileBorderSize = BuildSettingsLayer0.VirtualTextureBorderSize;
 	OutData.TileSize = TileSize;
 	OutData.NumLayers = NumLayers;
@@ -716,7 +717,7 @@ void FVirtualTextureDataBuilder::PushDataToChunk(const TArray<FVTSourceTileEntry
 {
 	const int32 NumLayers = SourceLayers.Num();
 
-	uint32 TotalSize = 0u;
+	uint32 TotalSize = sizeof(FVirtualTextureChunkHeader);
 	for (int32 Layer = 0; Layer < NumLayers; ++Layer)
 	{
 		TotalSize += LayerData[Layer].CodecPayload.Num();
@@ -733,6 +734,12 @@ void FVirtualTextureDataBuilder::PushDataToChunk(const TArray<FVTSourceTileEntry
 	uint8* NewChunkData = (uint8*)BulkData.Realloc(TotalSize);
 	uint32 ChunkOffset = 0u;
 
+	// Header for the chunk
+	FVirtualTextureChunkHeader* Header = (FVirtualTextureChunkHeader*)NewChunkData;
+	Header->VersionGuid = OutData.VersionGuid;
+
+	ChunkOffset += sizeof(FVirtualTextureChunkHeader);
+
 	// codec payloads
 	for (int32 Layer = 0; Layer < NumLayers; ++Layer)
 	{
