Prototype Wrapped Object Pointers (ObjectHandle/ObjectPtr)

Call-site changes for non-scalar wrapped object pointer upgrades in Engine + Shootergame

diff --git a/Engine/Source/Runtime/Engine/Private/HLODProxy.cpp b/Engine/Source/Runtime/Engine/Private/HLODProxy.cpp
--- a/Engine/Source/Runtime/Engine/Private/HLODProxy.cpp
+++ b/Engine/Source/Runtime/Engine/Private/HLODProxy.cpp
@@ -662,7 +662,7 @@ bool UHLODProxy::SetHLODBakingTransform(const FTransform& InTransform)
 {
 	bool bChanged = false;
 
-	for (TMap<UHLODProxyDesc*, FHLODProxyMesh>::TIterator ItHLODActor = HLODActors.CreateIterator(); ItHLODActor; ++ItHLODActor)
+	for (auto ItHLODActor = HLODActors.CreateIterator(); ItHLODActor; ++ItHLODActor)
 	{
 		UHLODProxyDesc* HLODProxyDesc = ItHLODActor.Key();
 		if (!HLODProxyDesc->HLODBakingTransform.Equals(InTransform))
