Made "debug thread.." console commands' callstacks easier to QA. Return addresses can often be the address of the next statement after a function call, especially if the call isn't expected to return a value or it's discarded. For calls at the end of a function the return addresses can often resolve to unexpected places (due to the compiler inlining for example). This can cause confusing callstacks that are difficult to read/understand.

diff --git a/Engine/Source/Runtime/Engine/Private/UnrealEngine.cpp b/Engine/Source/Runtime/Engine/Private/UnrealEngine.cpp
--- a/Engine/Source/Runtime/Engine/Private/UnrealEngine.cpp
+++ b/Engine/Source/Runtime/Engine/Private/UnrealEngine.cpp
@@ -8480,6 +8480,20 @@ FORCENOINLINE void StackOverflowFunction(int32* DummyArg)
 	}
 }
 
+// This function is used to add dummy work at the end of "debug thread.."
+// console commands' implementations. The intention is to better condition
+// return addresses. Return addresses can often be the address of the next
+// statement after a function call, especially if the call isn't expected to
+// return a value or the value's discarded. For calls at the end of a function
+// the return addresses can often resolve to unexpected places (due to the
+// compiler inlining for example). This can cause confusing callstacks that are
+// difficult to test and observe. Adding a call to this after keeps the failing
+// address more to more appropriate values.
+FORCENOINLINE static void DebugDummyWorkAfterFailure()
+{
+	UE_LOG(LogEngine, Log, TEXT("A dummy log message"));
+}
+
 bool UEngine::PerformError(const TCHAR* Cmd, FOutputDevice& Ar)
 {
 #if !UE_BUILD_SHIPPING
@@ -8556,6 +8570,7 @@ bool UEngine::PerformError(const TCHAR* Cmd, FOutputDevice& Ar)
 				UE_LOG(LogEngine, Warning, TEXT("Printed warning to log."));
 				FGenericCrashContext::SetCrashTrigger(ECrashTrigger::Debug);
 				UE_LOG(LogEngine, Fatal, TEXT("Crashing the worker thread at your request"));
+				DebugDummyWorkAfterFailure();
 			}
 		};
 
@@ -8581,6 +8596,7 @@ bool UEngine::PerformError(const TCHAR* Cmd, FOutputDevice& Ar)
 				UE_LOG(LogEngine, Warning, TEXT("Printed warning to log."));
 				FGenericCrashContext::SetCrashTrigger(ECrashTrigger::Debug);
 				check(!"Crashing a worker thread via check(0) at your request");
+				DebugDummyWorkAfterFailure();
 			}
 		};
 
@@ -8640,6 +8656,7 @@ bool UEngine::PerformError(const TCHAR* Cmd, FOutputDevice& Ar)
 					// with an Int 3 being handled by the exception handler.
 					const bool CrashingTheWorkerThreadAtYourRequest = false;
 					check(CrashingTheWorkerThreadAtYourRequest);
+					DebugDummyWorkAfterFailure();
 				}
 			);
 		};
@@ -8677,6 +8694,7 @@ bool UEngine::PerformError(const TCHAR* Cmd, FOutputDevice& Ar)
 						FPlatformProcess::Sleep(0);
 					}
 				} while (true);
+				DebugDummyWorkAfterFailure();
 			}
 		};
 
@@ -8717,6 +8735,7 @@ bool UEngine::PerformError(const TCHAR* Cmd, FOutputDevice& Ar)
 						FPlatformProcess::Sleep(0);
 					}
 				} while (true);
+				DebugDummyWorkAfterFailure();
 			}
 		};
 
@@ -8734,6 +8753,7 @@ bool UEngine::PerformError(const TCHAR* Cmd, FOutputDevice& Ar)
 			{
 				UE_LOG(LogEngine, Warning, TEXT("Printed warning to log."));
 				ensure(0);
+				DebugDummyWorkAfterFailure();
 			}
 		};
 
@@ -8758,6 +8778,7 @@ bool UEngine::PerformError(const TCHAR* Cmd, FOutputDevice& Ar)
 				UE_LOG(LogEngine, Warning, TEXT("Printed warning to log."));
 				FGenericCrashContext::SetCrashTrigger(ECrashTrigger::Debug);
 				LowLevelFatalError(TEXT("FError::LowLevelFatal test"));
+				DebugDummyWorkAfterFailure();
 			}
 		};
 
@@ -8829,7 +8850,7 @@ bool UEngine::PerformError(const TCHAR* Cmd, FOutputDevice& Ar)
 	}
 	else if (FParse::Command(&Cmd, TEXT("CRTINVALID")))
 	{
-	FGenericCrashContext::SetCrashTrigger(ECrashTrigger::Debug);
+		FGenericCrashContext::SetCrashTrigger(ECrashTrigger::Debug);
 		CauseCrtError();
 		return true;
 	}
