Moved functional test code out of lambdas and into methods. Methods have more human-friend names and give callstacks that QA can verify with less error.

diff --git a/Engine/Source/Runtime/Engine/Private/UnrealEngine.cpp b/Engine/Source/Runtime/Engine/Private/UnrealEngine.cpp
--- a/Engine/Source/Runtime/Engine/Private/UnrealEngine.cpp
+++ b/Engine/Source/Runtime/Engine/Private/UnrealEngine.cpp
@@ -8499,66 +8499,79 @@ bool UEngine::PerformError(const TCHAR* Cmd, FOutputDevice& Ar)
 #if !UE_BUILD_SHIPPING
 	if (FParse::Command(&Cmd, TEXT("RENDERCRASH")))
 	{
-		ENQUEUE_RENDER_COMMAND(CauseRenderThreadCrash)(
-			[](FRHICommandList& RHICmdList)
+		struct FRender
+		{
+			static void Crash(FRHICommandList&)
 			{
 				UE_LOG(LogEngine, Warning, TEXT("Printed warning to log."));
 				FGenericCrashContext::SetCrashTrigger(ECrashTrigger::Debug);
 				UE_LOG(LogEngine, Fatal, TEXT("Crashing the renderthread at your request"));
-			});
+				DebugDummyWorkAfterFailure();
+			}
+		};
+
+		ENQUEUE_RENDER_COMMAND(CauseRenderThreadCrash)(&FRender::Crash);
 		return true;
 	}
 	if (FParse::Command(&Cmd, TEXT("RENDERCHECK")))
 	{
 		struct FRender
 		{
-			static void Check()
+			static void Check(FRHICommandList&)
 			{				
 				UE_LOG(LogEngine, Warning, TEXT("Printed warning to log."));
 				FGenericCrashContext::SetCrashTrigger(ECrashTrigger::Debug);
 				check(!"Crashing the renderthread via check(0) at your request");
+				DebugDummyWorkAfterFailure();
 			}
 		};
-		ENQUEUE_RENDER_COMMAND(CauseRenderThreadCrash)(
-			[](FRHICommandList& RHICmdList)
-			{
-				FRender::Check();
-			});
+		ENQUEUE_RENDER_COMMAND(CauseRenderThreadCrash)(&FRender::Check);
 		return true;
 	}
 	if (FParse::Command(&Cmd, TEXT("RENDERGPF")))
 	{
-		ENQUEUE_RENDER_COMMAND(CauseRenderThreadCrash)(
-			[](FRHICommandList& RHICmdList)
+		struct FRender
+		{
+			static void GPF(FRHICommandList&)
 			{
 				UE_LOG(LogEngine, Warning, TEXT("Printed warning to log."));
 				FGenericCrashContext::SetCrashTrigger(ECrashTrigger::Debug);
 				*(int32 *)3 = 123;
-			});
+			}
+		};
+		ENQUEUE_RENDER_COMMAND(CauseRenderThreadCrash)(&FRender::GPF);
 		return true;
 	}
 	if (FParse::Command(&Cmd, TEXT("RENDERFATAL")))
 	{
-		ENQUEUE_RENDER_COMMAND(CauseRenderThreadCrash)(
-			[](FRHICommandList& RHICmdList)
+		struct FRender
+		{
+			static void Fatal(FRHICommandList&)
 			{
 				UE_LOG(LogEngine, Warning, TEXT("Printed warning to log."));
 				FGenericCrashContext::SetCrashTrigger(ECrashTrigger::Debug);
 				LowLevelFatalError(TEXT("FError::LowLevelFatal test"));
-			});
+				DebugDummyWorkAfterFailure();
+			}
+		};
+		ENQUEUE_RENDER_COMMAND(CauseRenderThreadCrash)(&FRender::Fatal);
 		return true;
 	}
 	if (FParse::Command(&Cmd, TEXT("RENDERENSURE")))
 	{
-		ENQUEUE_RENDER_COMMAND(CauseRenderThreadEnsure)(
-			[](FRHICommandList& RHICmdList)
+		struct FRender
+		{
+			static void Ensure(FRHICommandList&)
 			{
 				UE_LOG(LogEngine, Warning, TEXT("Printed warning to log."));
-			if (!ensure(0))
-			{
-				UE_LOG(LogEngine, Warning, TEXT("Ensure condition failed (this is the expected behavior)."));
+				if (!ensure(0))
+				{
+					UE_LOG(LogEngine, Warning, TEXT("Ensure condition failed (this is the expected behavior)."));
+				}
+				DebugDummyWorkAfterFailure();
 			}
-			});
+		};
+		ENQUEUE_RENDER_COMMAND(CauseRenderThreadEnsure)(&FRender::Ensure);
 		return true;
 	}
 	if (FParse::Command(&Cmd, TEXT("THREADCRASH")))
@@ -8641,24 +8654,26 @@ bool UEngine::PerformError(const TCHAR* Cmd, FOutputDevice& Ar)
 
 	if (FParse::Command(&Cmd, TEXT("THREADEDCHECKS")))
 	{
+		struct FThreaded
+		{
+			static void Checks()
+			{
+				FGenericCrashContext::SetCrashTrigger(ECrashTrigger::Debug);
+
+				// We now test if assert behavior is deterministic when called from different threads or
+				// if one thread could end up with a different behavior due to a race condition.
+				// We expect that the result will always be the assert being reported and not end-up
+				// with an Int 3 being handled by the exception handler.
+				const bool CrashingTheWorkerThreadAtYourRequest = false;
+				check(CrashingTheWorkerThreadAtYourRequest);
+				DebugDummyWorkAfterFailure();
+			}
+		};
+
 		UE_LOG(LogEngine, Warning, TEXT("Queuing several multi-thread checks."));
 		for (int32 Index = 0; Index < 16; ++Index)
 		{
-			Async(
-				EAsyncExecution::ThreadPool,
-				[]()
-				{
-					FGenericCrashContext::SetCrashTrigger(ECrashTrigger::Debug);
-
-					// We now test if assert behavior is deterministic when called from different threads or
-					// if one thread could end up with a different behavior due to a race condition.
-					// We expect that the result will always be the assert being reported and not end-up
-					// with an Int 3 being handled by the exception handler.
-					const bool CrashingTheWorkerThreadAtYourRequest = false;
-					check(CrashingTheWorkerThreadAtYourRequest);
-					DebugDummyWorkAfterFailure();
-				}
-			);
+			Async(EAsyncExecution::ThreadPool, &FThreaded::Checks);
 		};
 
 		return true;
@@ -9023,18 +9038,27 @@ bool UEngine::PerformError(const TCHAR* Cmd, FOutputDevice& Ar)
 	}
 	else if (FParse::Command(&Cmd, TEXT("AUDIOGPF")))
 	{
-		FAudioThread::RunCommandOnAudioThread([]()
+		struct FAudio
 		{
-			*(int32 *)3 = 123;
-		}, TStatId());
+			static void GPF()
+			{
+				*(int32 *)3 = 123;
+			}
+		};
+		FAudioThread::RunCommandOnAudioThread(&FAudio::GPF, TStatId());
 		return true;
 	}
 	else if (FParse::Command(&Cmd, TEXT("AUDIOCHECK")))
 	{
-		FAudioThread::RunCommandOnAudioThread([]()
+		struct FAudio
 		{
-			check(!"Crashing the audio thread via check(0) at your request");
-		}, TStatId());
+			static void Check()
+			{
+				check(!"Crashing the audio thread via check(0) at your request");
+				DebugDummyWorkAfterFailure();
+			}
+		}
+		FAudioThread::RunCommandOnAudioThread(&FAudio::Check, TStatId());
 		return true;
 	}
 	else if (FParse::Command(&Cmd, TEXT("CRASHCONTEXTWRITER")))
