Missed a ;. Oops.

diff --git a/Engine/Source/Runtime/Engine/Private/UnrealEngine.cpp b/Engine/Source/Runtime/Engine/Private/UnrealEngine.cpp
--- a/Engine/Source/Runtime/Engine/Private/UnrealEngine.cpp
+++ b/Engine/Source/Runtime/Engine/Private/UnrealEngine.cpp
@@ -9057,7 +9057,7 @@ bool UEngine::PerformError(const TCHAR* Cmd, FOutputDevice& Ar)
 				check(!"Crashing the audio thread via check(0) at your request");
 				DebugDummyWorkAfterFailure();
 			}
-		}
+		};
 		FAudioThread::RunCommandOnAudioThread(&FAudio::Check, TStatId());
 		return true;
 	}
