Fixed nullptr crash in canvas rendering.

#rnx

diff --git a/Engine/Source/Runtime/Engine/Private/UnrealClient.cpp b/Engine/Source/Runtime/Engine/Private/UnrealClient.cpp
--- a/Engine/Source/Runtime/Engine/Private/UnrealClient.cpp
+++ b/Engine/Source/Runtime/Engine/Private/UnrealClient.cpp
@@ -275,7 +275,7 @@ const FTexture2DRHIRef& FRenderTarget::GetRenderTargetTexture() const
 
 FRDGTextureRef FRenderTarget::GetRenderTargetTexture(FRDGBuilder& GraphBuilder) const
 {
-	return RegisterExternalTexture(GraphBuilder, RenderTargetTextureRHI, TEXT("RenderTarget"));
+	return RegisterExternalTexture(GraphBuilder, GetRenderTargetTexture(), TEXT("RenderTarget"));
 }
 
 FUnorderedAccessViewRHIRef FRenderTarget::GetRenderTargetUAV() const
