Improved some labels and tooltips in the blend spaces.
This is related to the per bone interpolation.

diff --git a/Engine/Source/Runtime/Engine/Classes/Animation/BlendSpaceBase.h b/Engine/Source/Runtime/Engine/Classes/Animation/BlendSpaceBase.h
--- a/Engine/Source/Runtime/Engine/Classes/Animation/BlendSpaceBase.h
+++ b/Engine/Source/Runtime/Engine/Classes/Animation/BlendSpaceBase.h
@@ -200,7 +200,14 @@ struct FPerBoneInterpolation
 	UPROPERTY(EditAnywhere, Category=FPerBoneInterpolation)
 	FBoneReference BoneReference;
 
-	UPROPERTY(EditAnywhere, Category=FPerBoneInterpolation)
+	/**
+	* This is the speed at which we interpolate towards the target weights for this specific bone, measured in 'how many times per second' we can get to the target.
+	* A value of 0 means it would instantly set itself to the target value, while a value of one means it will take one second to get there.
+	* A value of 2 would mean it goes there in twice the speed of a second, so in half a second, a value of 3 would mean in a third of a second, and so on.
+	* Smaller values mean slower interpolation speeds.
+	* This value overrides the global interpolation speed, so the global interpolation speed has no impact anymore on the interpolation speed of this bone.
+	*/
+	UPROPERTY(EditAnywhere, Category=FPerBoneInterpolation, meta=(DisplayName="Interpolation Speed"))
 	float InterpolationSpeedPerSec;
 
 	FPerBoneInterpolation()
@@ -446,11 +453,14 @@ public:
 	FInterpolationParameter	InterpolationParam[3];
 
 	/**
-	* Target weight interpolation. When target samples are set, how fast you'd like to get to target. Improve target blending.
-	* i.e. for locomotion, if you interpolate input, when you move from left to right rapidly, you'll interpolate through forward, but if you use target weight interpolation,
-	* you'll skip forward, but interpolate between left to right
+	* This is the speed at which we interpolate towards the target weights, measured in 'how many times per second' we can get to the target.
+	* A value of 0 means it would instantly set itself to the target value, while a value of one means it will take one second to get there.
+	* A value of 2 would mean it goes there in twice the speed of a second, so in half a second, a value of 3 would mean in a third of a second, and so on.
+	* Smaller values mean slower interpolation speeds.
+	* Imagine we have a blend space for locomotion, moving left, forward and right. Now if you interpolate the inputs of the blend space itself, from one extreme to the other, you will
+	* go from left, to forward, to right. As an alternative, by setting this global interpolation speed to a value higher than zero, it will go directly from left to right, without going through moving forward first.
 	*/
-	UPROPERTY(EditAnywhere, Category = SampleInterpolation)
+	UPROPERTY(EditAnywhere, Category = SampleInterpolation, meta = (DisplayName = "Global Interpolation Speed"))
 	float TargetWeightInterpolationSpeedPerSec;
 
 	/** The current mode used by the blendspace to decide which animation notifies to fire. Valid options are:
@@ -464,9 +474,11 @@ public:
 protected:
 
 	/**
-	* Define target weight interpolation per bone. This will blend in different speed per each bone setting
+	* Per bone interpolation speed settings. 
+	* These act as overrides to the global interpolation speed. 
+	* This means the global interpolation speed does not impact these bones.
 	*/
-	UPROPERTY(EditAnywhere, Category = SampleInterpolation)
+	UPROPERTY(EditAnywhere, Category = SampleInterpolation, meta = (DisplayName="Per Bone Overrides"))
 	TArray<FPerBoneInterpolation> PerBoneBlend;
 
 	/** Track index to get marker data from. Samples are tested for the suitability of marker based sync
