LevelInstance: Packed level instances now create multiple ISMs for different collision states of the same SMCs.

diff --git a/Engine/Source/Runtime/Engine/Private/LevelInstance/Packed/LevelInstanceISMPacker.cpp b/Engine/Source/Runtime/Engine/Private/LevelInstance/Packed/LevelInstanceISMPacker.cpp
--- a/Engine/Source/Runtime/Engine/Private/LevelInstance/Packed/LevelInstanceISMPacker.cpp
+++ b/Engine/Source/Runtime/Engine/Private/LevelInstance/Packed/LevelInstanceISMPacker.cpp
@@ -57,6 +57,11 @@ void FLevelInstanceISMPacker::PackActors(FPackedLevelInstanceBuilderContext& InC
 		PackComponent->SetMaterial(MaterialIndex, ISMCluster->Materials[MaterialIndex]);
 	}
 
+	if (!ISMCluster->bCollisionEnabled)
+	{
+		PackComponent->SetCollisionEnabled(ECollisionEnabled::NoCollision);
+	}
+
 	for (UActorComponent* Component : InComponents)
 	{
 		// If we have a ISM we need to add all instances
@@ -92,6 +97,7 @@ FLevelInstanceISMPackerCluster::FLevelInstanceISMPackerCluster(FLevelInstancePac
 	, bReceivesDecals(InComponent->bReceivesDecals)
 	, bCastShadow(InComponent->CastShadow)
 	, bVisibleInRayTracing(InComponent->bVisibleInRayTracing)
+	, bCollisionEnabled(InComponent->IsCollisionEnabled())
 {
 	
 }
@@ -108,6 +114,7 @@ uint32 FLevelInstanceISMPackerCluster::ComputeHash() const
 	Hash = FCrc::TypeCrc32(bReceivesDecals, Hash);
 	Hash = FCrc::TypeCrc32(bCastShadow, Hash);
 	Hash = FCrc::TypeCrc32(bVisibleInRayTracing, Hash);
+	Hash = FCrc::TypeCrc32(bCollisionEnabled, Hash);
 
 	return Hash;
 }
diff --git a/Engine/Source/Runtime/Engine/Public/LevelInstance/Packed/LevelInstanceISMPacker.h b/Engine/Source/Runtime/Engine/Public/LevelInstance/Packed/LevelInstanceISMPacker.h
--- a/Engine/Source/Runtime/Engine/Public/LevelInstance/Packed/LevelInstanceISMPacker.h
+++ b/Engine/Source/Runtime/Engine/Public/LevelInstance/Packed/LevelInstanceISMPacker.h
@@ -40,6 +40,7 @@ public:
 	bool bReceivesDecals;
 	bool bCastShadow;
 	bool bVisibleInRayTracing;
+	bool bCollisionEnabled;
 };
 
 #endif
\ No newline at end of file
