Added editor only flag to World::InitializationValues to support the use of movement component in non-game worlds #review-15165556 [REDACTED_HANDLE] [REDACTED_HANDLE] #fyi [REDACTED_HANDLE]

diff --git a/Engine/Source/Runtime/Engine/Classes/Engine/World.h b/Engine/Source/Runtime/Engine/Classes/Engine/World.h
--- a/Engine/Source/Runtime/Engine/Classes/Engine/World.h
+++ b/Engine/Source/Runtime/Engine/Classes/Engine/World.h
@@ -1148,6 +1148,9 @@ public:
 	*/
 	uint8 bEnableTraceCollision:1;
 
+	/** Special flag to enable movement component in non game worlds (see UMovementComponent::OnRegister) */
+	uint8 bForceUseMovementComponentInNonGameWorld:1;
+
 	/** If True, overloaded method IsNameStableForNetworking will always return true. */
 	uint8 bIsNameStableForNetworking:1;
 #endif
@@ -2751,6 +2754,7 @@ public:
 			, bCreateAISystem(true)
 			, bShouldSimulatePhysics(true)
 			, bEnableTraceCollision(false)
+			, bForceUseMovementComponentInNonGameWorld(false)
 			, bTransactional(true)
 			, bCreateFXSystem(true)
 		{
@@ -2780,6 +2784,9 @@ public:
 		/** Are collision trace calls valid within this world. */
 		uint32 bEnableTraceCollision:1;
 
+		/** Special flag to enable movement component in non game worlds (see UMovementComponent::OnRegister) */
+		uint32 bForceUseMovementComponentInNonGameWorld:1;
+
 		/** Should actions performed to objects in this world be saved to the transaction buffer. */
 		uint32 bTransactional:1;
 
@@ -2797,6 +2804,7 @@ public:
 		InitializationValues& CreateAISystem(const bool bCreate) { bCreateAISystem = bCreate; return *this; }
 		InitializationValues& ShouldSimulatePhysics(const bool bInShouldSimulatePhysics) { bShouldSimulatePhysics = bInShouldSimulatePhysics; return *this; }
 		InitializationValues& EnableTraceCollision(const bool bInEnableTraceCollision) { bEnableTraceCollision = bInEnableTraceCollision; return *this; }
+		InitializationValues& ForceUseMovementComponentInNonGameWorld(const bool bInForceUseMovementComponentInNonGameWorld) { bForceUseMovementComponentInNonGameWorld = bInForceUseMovementComponentInNonGameWorld; return *this; }
 		InitializationValues& SetTransactional(const bool bInTransactional) { bTransactional = bInTransactional; return *this; }
 		InitializationValues& CreateFXSystem(const bool bCreate) { bCreateFXSystem = bCreate; return *this; }
 		InitializationValues& SetDefaultGameMode(TSubclassOf<class AGameModeBase> GameMode) { DefaultGameMode = GameMode; return *this; }
diff --git a/Engine/Source/Runtime/Engine/Private/Components/MovementComponent.cpp b/Engine/Source/Runtime/Engine/Private/Components/MovementComponent.cpp
--- a/Engine/Source/Runtime/Engine/Private/Components/MovementComponent.cpp
+++ b/Engine/Source/Runtime/Engine/Private/Components/MovementComponent.cpp
@@ -146,7 +146,11 @@ void UMovementComponent::OnRegister()
 	}
 
 	const UWorld* MyWorld = GetWorld();
-	if (MyWorld && MyWorld->IsGameWorld())
+	if (MyWorld && (MyWorld->IsGameWorld()
+#if WITH_EDITOR
+	|| MyWorld->bForceUseMovementComponentInNonGameWorld
+#endif
+	))
 	{
 		PlaneConstraintNormal = PlaneConstraintNormal.GetSafeNormal();
 
diff --git a/Engine/Source/Runtime/Engine/Private/PreviewScene.cpp b/Engine/Source/Runtime/Engine/Private/PreviewScene.cpp
--- a/Engine/Source/Runtime/Engine/Private/PreviewScene.cpp
+++ b/Engine/Source/Runtime/Engine/Private/PreviewScene.cpp
@@ -42,7 +42,8 @@ FPreviewScene::FPreviewScene(FPreviewScene::ConstructionValues CVS)
 										.CreateAISystem(false)
 										.ShouldSimulatePhysics(CVS.bShouldSimulatePhysics)
 										.SetTransactional(CVS.bTransactional)
-										.SetDefaultGameMode(CVS.DefaultGameMode));
+										.SetDefaultGameMode(CVS.DefaultGameMode)
+										.ForceUseMovementComponentInNonGameWorld(CVS.bForceUseMovementComponentInNonGameWorld));
 
 	FURL URL = FURL();
 	//URL += TEXT("?SpectatorOnly=1");
diff --git a/Engine/Source/Runtime/Engine/Private/World.cpp b/Engine/Source/Runtime/Engine/Private/World.cpp
--- a/Engine/Source/Runtime/Engine/Private/World.cpp
+++ b/Engine/Source/Runtime/Engine/Private/World.cpp
@@ -1530,6 +1530,7 @@ void UWorld::InitWorld(const InitializationValues IVS)
 
 	#if WITH_EDITOR
 		bEnableTraceCollision = IVS.bEnableTraceCollision;
+		bForceUseMovementComponentInNonGameWorld = IVS.bForceUseMovementComponentInNonGameWorld;
 	#endif
 
 
diff --git a/Engine/Source/Runtime/Engine/Public/PreviewScene.h b/Engine/Source/Runtime/Engine/Public/PreviewScene.h
--- a/Engine/Source/Runtime/Engine/Public/PreviewScene.h
+++ b/Engine/Source/Runtime/Engine/Public/PreviewScene.h
@@ -31,6 +31,7 @@ public:
 			, bShouldSimulatePhysics(false)
 			, bTransactional(true)
 			, bEditor(true)
+			, bForceUseMovementComponentInNonGameWorld(false)
 		{
 		}
 
@@ -44,6 +45,7 @@ public:
 		uint32 bShouldSimulatePhysics:1;
 		uint32 bTransactional:1;
 		uint32 bEditor:1;
+		uint32 bForceUseMovementComponentInNonGameWorld:1;
 
 		TSubclassOf<class AGameModeBase> DefaultGameMode;
 		class UGameInstance* OwningGameInstance = nullptr;
@@ -59,6 +61,7 @@ public:
 		ConstructionValues& ShouldSimulatePhysics(const bool bInShouldSimulatePhysics) { bShouldSimulatePhysics = bInShouldSimulatePhysics; return *this; }
 		ConstructionValues& SetTransactional(const bool bInTransactional) { bTransactional = bInTransactional; return *this; }
 		ConstructionValues& SetEditor(const bool bInEditor) { bEditor = bInEditor; return *this; }
+		ConstructionValues& ForceUseMovementComponentInNonGameWorld(const bool bInForceUseMovementComponentInNonGameWorld) { bForceUseMovementComponentInNonGameWorld = bInForceUseMovementComponentInNonGameWorld; return *this; }
 
 		ConstructionValues& SetDefaultGameMode(TSubclassOf<class AGameModeBase> GameMode) { DefaultGameMode = GameMode; return *this; }
 		ConstructionValues& SetOwningGameInstance(class UGameInstance* InGameInstance) { OwningGameInstance = InGameInstance; return *this; }
