Skip calls to PostStaticMeshCompilation for pending async mesh compilations when the engine is exiting.

diff --git a/Engine/Source/Runtime/Engine/Private/StaticMeshCompiler.cpp b/Engine/Source/Runtime/Engine/Private/StaticMeshCompiler.cpp
--- a/Engine/Source/Runtime/Engine/Private/StaticMeshCompiler.cpp
+++ b/Engine/Source/Runtime/Engine/Private/StaticMeshCompiler.cpp
@@ -190,7 +190,7 @@ void FStaticMeshCompilingManager::PostCompilation(UStaticMesh* StaticMesh)
 	using namespace StaticMeshCompilingManagerImpl;
 	
 	// If AsyncTask is null here, the task got canceled so we don't need to do anything
-	if (StaticMesh->AsyncTask)
+	if (StaticMesh->AsyncTask && !IsEngineExitRequested())
 	{
 		check(IsInGameThread());
 		TRACE_CPUPROFILER_EVENT_SCOPE(PostCompilation);
