Using RVT streaming mips in PIE by default.

diff --git a/Engine/Source/Runtime/Engine/Private/Components/RuntimeVirtualTextureComponent.cpp b/Engine/Source/Runtime/Engine/Private/Components/RuntimeVirtualTextureComponent.cpp
--- a/Engine/Source/Runtime/Engine/Private/Components/RuntimeVirtualTextureComponent.cpp
+++ b/Engine/Source/Runtime/Engine/Private/Components/RuntimeVirtualTextureComponent.cpp
@@ -84,7 +84,10 @@ void URuntimeVirtualTextureComponent::DestroyRenderState_Concurrent()
 
 void URuntimeVirtualTextureComponent::Invalidate(FBoxSphereBounds const& InWorldBounds)
 {
-	GetScene()->InvalidateRuntimeVirtualTexture(this, InWorldBounds);
+	if (GetScene() != nullptr)
+	{
+		GetScene()->InvalidateRuntimeVirtualTexture(this, InWorldBounds);
+	}
 }
 
 FBoxSphereBounds URuntimeVirtualTextureComponent::CalcBounds(const FTransform& LocalToWorld) const
@@ -149,7 +152,7 @@ bool URuntimeVirtualTextureComponent::IsStreamingTextureValid() const
 bool URuntimeVirtualTextureComponent::IsStreamingLowMips() const
 {
 #if WITH_EDITOR
-	if (!bUseStreamingLowMipsInEditor)
+	if (!bUseStreamingLowMipsInEditor && GIsEditor && !GetWorld()->IsPlayInEditor())
 	{
 		return false;
 	}
