Implemented support for GPU skinned meshes exposed to GPU scene instance culling and rendering into virtual shadow maps.

diff --git a/Engine/Shaders/Private/GpuSkinVertexFactory.ush b/Engine/Shaders/Private/GpuSkinVertexFactory.ush
--- a/Engine/Shaders/Private/GpuSkinVertexFactory.ush
+++ b/Engine/Shaders/Private/GpuSkinVertexFactory.ush
@@ -104,6 +104,11 @@ struct FVertexFactoryInput
 
 	/** Per vertex color */
 	float4 Color : ATTRIBUTE13;
+
+#if VF_USE_PRIMITIVE_SCENE_DATA
+	uint InstanceIdOffset : ATTRIBUTE16;
+	uint DrawInstanceId : SV_InstanceID;
+#endif
 };
 
 #if !GPUSKIN_MORPH_BLEND && GPUSKIN_APEX_CLOTH
@@ -196,7 +201,7 @@ FMaterialPixelParameters GetMaterialPixelParameters(FVertexFactoryInterpolantsVS
 	// GetMaterialPixelParameters is responsible for fully initializing the result
 	FMaterialPixelParameters Result = MakeInitializedMaterialPixelParameters();
 
-#if NUM_TEX_COORD_INTERPOLATORS	
+#if NUM_TEX_COORD_INTERPOLATORS
 	UNROLL
 	for(int CoordinateIndex = 0;CoordinateIndex < NUM_TEX_COORD_INTERPOLATORS;CoordinateIndex++)
 	{
@@ -234,6 +239,15 @@ STRONG_TYPE Buffer<float4> PreviousBoneMatrices;
 // Cache data to avoid multiple calculation 
 struct FVertexFactoryIntermediates
 {
+#if VF_USE_PRIMITIVE_SCENE_DATA
+	FSceneDataIntermediates SceneData;
+#endif
+
+	float4 InvNonUniformScaleAndDeterminantSign;
+	float4x4 LocalToWorld;
+	float4x4 WorldToLocal;
+	float4x4 PrevLocalToWorld;
+
 	// Blend Matrix (used for position/tangents)
 	FBoneMatrix BlendMatrix;
 
@@ -263,7 +277,7 @@ FMaterialVertexParameters GetMaterialVertexParameters(FVertexFactoryInput Input,
 	FMaterialVertexParameters Result = (FMaterialVertexParameters)0;
 	Result.WorldPosition = WorldPosition;
 	Result.VertexColor = Intermediates.Color;
-	Result.TangentToWorld = mul(TangentToLocal, GetLocalToWorld3x3());
+	Result.TangentToWorld = mul(TangentToLocal, (float3x3)Intermediates.LocalToWorld);
 	Result.PreSkinnedPosition = Intermediates.UnpackedPosition.xyz;
 	Result.PreSkinnedNormal = Input.TangentZ.xyz;
 
@@ -271,7 +285,7 @@ FMaterialVertexParameters GetMaterialVertexParameters(FVertexFactoryInput Input,
 	Result.PostSkinOffset = Input.PostSkinOffset;
 
 	// Assumes no instacing
-	Result.PrevFrameLocalToWorld = Primitive.PreviousLocalToWorld;
+	Result.PrevFrameLocalToWorld = Intermediates.PrevLocalToWorld;
 
 #if NUM_MATERIAL_TEXCOORDS_VERTEX
 	for(int CoordinateIndex = 0; CoordinateIndex < NUM_MATERIAL_TEXCOORDS_VERTEX; CoordinateIndex++)
@@ -653,8 +667,8 @@ float3x3 SkinTangents( FVertexFactoryInput Input, FVertexFactoryIntermediates In
 		    TangentToLocal[2] = normalize(NormalPosition - Intermediates.SimulatedPosition);
 
 		    // cloth data are all in world space so need to change into local space
-		    TangentToLocal[0] = mul(TangentToLocal[0], (half3x3)Primitive.WorldToLocal);
-		    TangentToLocal[2] = mul(TangentToLocal[2], (half3x3)Primitive.WorldToLocal);
+		    TangentToLocal[0] = mul(TangentToLocal[0], (half3x3)Intermediates.WorldToLocal);
+		    TangentToLocal[2] = mul(TangentToLocal[2], (half3x3)Intermediates.WorldToLocal);
         }
 	    else
 	    {
@@ -692,7 +706,21 @@ FVertexFactoryIntermediates GetVertexFactoryIntermediates(FVertexFactoryInput In
 {
 	FVertexFactoryIntermediates Intermediates;
 	
+#if VF_USE_PRIMITIVE_SCENE_DATA
+	Intermediates.SceneData = GetSceneDataIntermediates(Input.InstanceIdOffset, Input.DrawInstanceId);
+	Intermediates.InvNonUniformScaleAndDeterminantSign = Intermediates.SceneData.InstanceData.InvNonUniformScaleAndDeterminantSign;
+	Intermediates.LocalToWorld = Intermediates.SceneData.InstanceData.LocalToWorld;
+	Intermediates.WorldToLocal = Intermediates.SceneData.InstanceData.WorldToLocal;
+	Intermediates.PrevLocalToWorld = Intermediates.SceneData.InstanceData.PrevLocalToWorld;
+#else
+	Intermediates.InvNonUniformScaleAndDeterminantSign = Primitive.InvNonUniformScaleAndDeterminantSign;
+	Intermediates.LocalToWorld = Primitive.LocalToWorld;
+	Intermediates.WorldToLocal = Primitive.WorldToLocal;
+	Intermediates.PrevLocalToWorld = Primitive.PreviousLocalToWorld;
+#endif
+
 	Intermediates.UnpackedPosition = UnpackedPosition(Input);
+
 #if GPUSKIN_APEX_CLOTH
 	SetupClothVertex(Input.ClothVertexID, Intermediates.ClothVertexInfluences);
 	if( IsSimulatedVertex(Intermediates.ClothVertexInfluences[0]) )
@@ -725,7 +753,7 @@ float4 CalcWorldPosition(FVertexFactoryInput Input, FVertexFactoryIntermediates
 	{
 		// skinned positions
 		float4 SkinnedPosition = float4(SkinPosition(Input, Intermediates), 1);
-		float4 TransformedSkinPos = mul(SkinnedPosition, Primitive.LocalToWorld);
+		float4 TransformedSkinPos = mul(SkinnedPosition, Intermediates.LocalToWorld);
 
 		// Intermediates.SimulatedPosition is a clothing position in world coord
 		float3 BlendedPos = lerp(TransformedSkinPos.xyz, Intermediates.SimulatedPosition, ClothBlendWeight);
@@ -734,7 +762,11 @@ float4 CalcWorldPosition(FVertexFactoryInput Input, FVertexFactoryIntermediates
 	else
 #endif
 	{
+#if VF_USE_PRIMITIVE_SCENE_DATA
+		return TransformLocalToTranslatedWorld(SkinPosition(Input, Intermediates), Intermediates.SceneData.InstanceData);
+#else
 		return TransformLocalToTranslatedWorld(SkinPosition(Input, Intermediates));
+#endif
 	}
 }
 
@@ -766,10 +798,10 @@ float3 VertexFactoryGetPositionForVertexLighting(FVertexFactoryInput Input, FVer
 
 void CalcTangentToWorld(FVertexFactoryInput Input, FVertexFactoryIntermediates Intermediates, out float3 TangentToWorld0, out float4 TangentToWorld2)
 {
-	float3x3 LocalToWorld = GetLocalToWorld3x3();
+	float3x3 LocalToWorld = (float3x3)Intermediates.LocalToWorld;
 
 	// Remove scaling.
-	half3 InvScale = Primitive.InvNonUniformScaleAndDeterminantSign.xyz;
+	half3 InvScale = Intermediates.InvNonUniformScaleAndDeterminantSign.xyz;
 	LocalToWorld[0] *= InvScale.x;
 	LocalToWorld[1] *= InvScale.y;
 	LocalToWorld[2] *= InvScale.z;
@@ -777,7 +809,7 @@ void CalcTangentToWorld(FVertexFactoryInput Input, FVertexFactoryIntermediates I
 	float3x3 TangentToWorld = mul(Intermediates.TangentToLocal, LocalToWorld);
 
 	TangentToWorld0 = TangentToWorld[0];
-	TangentToWorld2 = float4(TangentToWorld[2], Input.TangentZ.w * Primitive.InvNonUniformScaleAndDeterminantSign.w);
+	TangentToWorld2 = float4(TangentToWorld[2], Input.TangentZ.w * Intermediates.InvNonUniformScaleAndDeterminantSign.w);
 }
 
 float3 VertexFactoryGetWorldNormal(FVertexFactoryInput Input, FVertexFactoryIntermediates Intermediates)
@@ -825,7 +857,7 @@ FVertexFactoryInterpolantsVSToPS VertexFactoryGetInterpolantsVSToPS(FVertexFacto
 // @return previous translated world position
 float4 VertexFactoryGetPreviousWorldPosition(FVertexFactoryInput Input, FVertexFactoryIntermediates Intermediates)
 {
-	float4x4 PreviousLocalToWorldTranslated = Primitive.PreviousLocalToWorld;
+	float4x4 PreviousLocalToWorldTranslated = Intermediates.PrevLocalToWorld;
 
 	PreviousLocalToWorldTranslated[3][0] += ResolvedView.PrevPreViewTranslation.x;
 	PreviousLocalToWorldTranslated[3][1] += ResolvedView.PrevPreViewTranslation.y;
diff --git a/Engine/Source/Runtime/Engine/Private/GPUSkinVertexFactory.cpp b/Engine/Source/Runtime/Engine/Private/GPUSkinVertexFactory.cpp
--- a/Engine/Source/Runtime/Engine/Private/GPUSkinVertexFactory.cpp
+++ b/Engine/Source/Runtime/Engine/Private/GPUSkinVertexFactory.cpp
@@ -61,6 +61,12 @@ static FBoneMatricesUniformShaderParameters GBoneUniformStruct;
 
 #define NANITE_VAL_FALSE false,
 
+#if defined(GPUCULL_TODO)
+#define PRIMITIVE_ID_SUPPORT true,
+#else
+#define PRIMITIVE_ID_SUPPORT false,
+#endif
+
 #define IMPLEMENT_GPUSKINNING_VERTEX_FACTORY_TYPE_INTERNAL(FactoryClass, ShaderFilename,bUsedWithMaterials,bSupportsStaticLighting,bSupportsDynamicLighting,bPrecisePrevWorldPos,bSupportsPositionOnly) \
 	template <GPUSkinBoneInfluenceType BoneInfluenceType> FVertexFactoryType FactoryClass<BoneInfluenceType>::StaticType( \
 	BoneInfluenceType == DefaultBoneInfluence ? TEXT(#FactoryClass) TEXT("Default") : TEXT(#FactoryClass) TEXT("Unlimited"), \
@@ -71,7 +77,7 @@ static FBoneMatricesUniformShaderParameters GBoneUniformStruct;
 	bPrecisePrevWorldPos, \
 	bSupportsPositionOnly, \
 	false, \
-	false, \
+	PRIMITIVE_ID_SUPPORT \
 	NANITE_VAL_FALSE \
 	IMPLEMENT_VERTEX_FACTORY_VTABLE(FactoryClass<BoneInfluenceType>) \
 	); \
@@ -392,23 +398,31 @@ bool TGPUSkinVertexFactory<BoneInfluenceType>::ShouldCompilePermutation(const FV
 	return ((Parameters.MaterialParameters.bIsUsedWithSkeletalMesh && (BoneInfluenceType != UnlimitedBoneInfluence || bUnlimitedBoneInfluences)) || Parameters.MaterialParameters.bIsSpecialEngineMaterial);
 }
 
-
 template <GPUSkinBoneInfluenceType BoneInfluenceType>
 void TGPUSkinVertexFactory<BoneInfluenceType>::ModifyCompilationEnvironment(const FVertexFactoryShaderPermutationParameters& Parameters, FShaderCompilerEnvironment& OutEnvironment )
 {
 	FVertexFactory::ModifyCompilationEnvironment(Parameters, OutEnvironment);
-	const int32 MaxGPUSkinBones = GetFeatureLevelMaxNumberOfBones(GetMaxSupportedFeatureLevel(Parameters.Platform));
+
+	const FStaticFeatureLevel MaxSupportedFeatureLevel = GetMaxSupportedFeatureLevel(Parameters.Platform);
+	const bool bUseGPUScene = UseGPUScene(Parameters.Platform, MaxSupportedFeatureLevel);
+	const bool bSupportsPrimitiveIdStream = Parameters.VertexFactoryType->SupportsPrimitiveIdStream();
+	const int32 MaxGPUSkinBones = GetFeatureLevelMaxNumberOfBones(MaxSupportedFeatureLevel);
+
 	OutEnvironment.SetDefine(TEXT("MAX_SHADER_BONES"), MaxGPUSkinBones);
+
 	{
-		bool bLimit2BoneInfluences = (CVarGPUSkinLimit2BoneInfluences.GetValueOnAnyThread() != 0);
+		const bool bLimit2BoneInfluences = (CVarGPUSkinLimit2BoneInfluences.GetValueOnAnyThread() != 0);
 		OutEnvironment.SetDefine(TEXT("GPUSKIN_LIMIT_2BONE_INFLUENCES"), (bLimit2BoneInfluences ? 1 : 0));
 	}
+
 	OutEnvironment.SetDefine(TEXT("GPUSKIN_USE_BONES_SRV_BUFFER"), SupportsBonesBufferSRV(Parameters.Platform) ? 1 : 0);
 	OutEnvironment.SetDefine(TEXT("GPUSKIN_UNLIMITED_BONE_INFLUENCE"), BoneInfluenceType == UnlimitedBoneInfluence ? 1 : 0);
 
 	OutEnvironment.SetDefine(TEXT("GPU_SKINNED_MESH_FACTORY"), 1);
-}
 
+	OutEnvironment.SetDefine(TEXT("VF_SUPPORTS_PRIMITIVE_SCENE_DATA"), bSupportsPrimitiveIdStream && bUseGPUScene);
+	OutEnvironment.SetDefine(TEXT("VF_GPU_SCENE_TEXTURE"), bSupportsPrimitiveIdStream && bUseGPUScene && GPUSceneUseTexture2D(Parameters.Platform));
+}
 
 template<GPUSkinBoneInfluenceType BoneInfluenceType>
 void TGPUSkinVertexFactory<BoneInfluenceType>::CopyDataTypeForPassthroughFactory(FGPUSkinPassthroughVertexFactory* PassthroughVertexFactory)
@@ -432,38 +446,38 @@ void TGPUSkinVertexFactory<BoneInfluenceType>::CopyDataTypeForPassthroughFactory
 }
 
 /**
-* Add the decl elements for the streams
-* @param InData - type with stream components
-* @param OutElements - vertex decl list to modify
+* Add the vertex declaration elements for the streams.
+* @param InData - Type with stream components.
+* @param OutElements - Vertex declaration list to modify.
 */
 template <GPUSkinBoneInfluenceType BoneInfluenceType>
 void TGPUSkinVertexFactory<BoneInfluenceType>::AddVertexElements(FDataType& InData, FVertexDeclarationElementList& OutElements)
 {
-	// position decls
-	OutElements.Add(AccessStreamComponent(InData.PositionComponent,0));
+	// Position
+	OutElements.Add(AccessStreamComponent(InData.PositionComponent, 0));
 
-	// tangent basis vector decls
-	OutElements.Add(AccessStreamComponent(InData.TangentBasisComponents[0],1));
-	OutElements.Add(AccessStreamComponent(InData.TangentBasisComponents[1],2));
+	// Tangent basis vector
+	OutElements.Add(AccessStreamComponent(InData.TangentBasisComponents[0], 1));
+	OutElements.Add(AccessStreamComponent(InData.TangentBasisComponents[1], 2));
 
-	// texture coordinate decls
-	if(InData.TextureCoordinates.Num())
+	// Texture coordinates
+	if (InData.TextureCoordinates.Num())
 	{
 		const uint8 BaseTexCoordAttribute = 5;
-		for(int32 CoordinateIndex = 0;CoordinateIndex < InData.TextureCoordinates.Num();CoordinateIndex++)
+		for (int32 CoordinateIndex = 0; CoordinateIndex < InData.TextureCoordinates.Num(); ++CoordinateIndex)
 		{
 			OutElements.Add(AccessStreamComponent(
 				InData.TextureCoordinates[CoordinateIndex],
 				BaseTexCoordAttribute + CoordinateIndex
-				));
+			));
 		}
 
-		for(int32 CoordinateIndex = InData.TextureCoordinates.Num();CoordinateIndex < MAX_TEXCOORDS;CoordinateIndex++)
+		for (int32 CoordinateIndex = InData.TextureCoordinates.Num(); CoordinateIndex < MAX_TEXCOORDS; ++CoordinateIndex)
 		{
 			OutElements.Add(AccessStreamComponent(
 				InData.TextureCoordinates[InData.TextureCoordinates.Num() - 1],
 				BaseTexCoordAttribute + CoordinateIndex
-				));
+			));
 		}
 	}
 
@@ -473,32 +487,33 @@ void TGPUSkinVertexFactory<BoneInfluenceType>::AddVertexElements(FDataType& InDa
 		Data.ColorIndexMask = 0;
 	}
 
-	// Account for the possibility that the mesh has no vertex colors
-	if( InData.ColorComponent.VertexBuffer )
+	// Vertex color - account for the possibility that the mesh has no vertex colors
+	if (InData.ColorComponent.VertexBuffer)
 	{
 		OutElements.Add(AccessStreamComponent(InData.ColorComponent, 13));
 	}
 	else
 	{
-		//If the mesh has no color component, set the null color buffer on a new stream with a stride of 0.
-		//This wastes 4 bytes of bandwidth per vertex, but prevents having to compile out twice the number of vertex factories.
+		// If the mesh has no color component, set the null color buffer on a new stream with a stride of 0.
+		// This wastes 4 bytes of memory per vertex, but prevents having to compile out twice the number of vertex factories.
 		FVertexStreamComponent NullColorComponent(&GNullColorVertexBuffer, 0, 0, VET_Color, EVertexStreamUsage::ManualFetch);
 		OutElements.Add(AccessStreamComponent(NullColorComponent, 13));
 	}
 
 	if (BoneInfluenceType == UnlimitedBoneInfluence)
 	{
+		// Blend offset count
 		OutElements.Add(AccessStreamComponent(InData.BlendOffsetCount, 3));
 	}
 	else
 	{
-		// bone indices decls
-		OutElements.Add(AccessStreamComponent(InData.BoneIndices,3));
+		// Bone indices
+		OutElements.Add(AccessStreamComponent(InData.BoneIndices, 3));
 
-		// bone weights decls
-		OutElements.Add(AccessStreamComponent(InData.BoneWeights,4));
+		// Bone weights
+		OutElements.Add(AccessStreamComponent(InData.BoneWeights, 4));
 
-		// Extra bone indices & weights decls
+		// Extra bone indices & weights
 		if (GetNumBoneInfluences() > MAX_INFLUENCES_PER_STREAM)
 		{
 			OutElements.Add(AccessStreamComponent(InData.ExtraBoneIndices, 14));
@@ -511,6 +526,7 @@ void TGPUSkinVertexFactory<BoneInfluenceType>::AddVertexElements(FDataType& InDa
 		}
 	}
 
+	// Pre-skinning offsets
 	if (InData.PreSkinningOffsets.VertexBuffer)
 	{
 		OutElements.Add(AccessStreamComponent(InData.PreSkinningOffsets, 11));
@@ -521,6 +537,7 @@ void TGPUSkinVertexFactory<BoneInfluenceType>::AddVertexElements(FDataType& InDa
 		OutElements.Add(AccessStreamComponent(NullDeltaComponent, 11));
 	}
 
+	// Post-skinning offsets
 	if (InData.PostSkinningOffsets.VertexBuffer)
 	{
 		OutElements.Add(AccessStreamComponent(InData.PostSkinningOffsets, 12));
@@ -530,6 +547,14 @@ void TGPUSkinVertexFactory<BoneInfluenceType>::AddVertexElements(FDataType& InDa
 		FVertexStreamComponent NullDeltaComponent(&GNullVertexBuffer, 0, 0, VET_Float3);
 		OutElements.Add(AccessStreamComponent(NullDeltaComponent, 12));
 	}
+
+	// Primitive Id
+	if (GetType()->SupportsPrimitiveIdStream() && UseGPUScene(GMaxRHIShaderPlatform, GMaxRHIFeatureLevel))
+	{
+		// When the VF is used for rendering in normal mesh passes, this vertex buffer and offset will be overridden
+		OutElements.Add(AccessStreamComponent(FVertexStreamComponent(&GPrimitiveIdDummy, 0, 0, PrimitiveIdStreamStride, VET_UInt, EVertexStreamUsage::Instancing), 16));
+		SetPrimitiveIdStreamIndex(EVertexInputStreamType::Default, OutElements.Last().StreamIndex);
+	}
 }
 
 /**
@@ -815,7 +840,11 @@ IMPLEMENT_VERTEX_FACTORY_PARAMETER_TYPE(FGPUSkinPassthroughVertexFactory, SF_Ver
 IMPLEMENT_VERTEX_FACTORY_PARAMETER_TYPE(FGPUSkinPassthroughVertexFactory, SF_RayHitGroup, FGPUSkinVertexPassthroughFactoryShaderParameters);
 IMPLEMENT_VERTEX_FACTORY_PARAMETER_TYPE(FGPUSkinPassthroughVertexFactory, SF_Compute, FGPUSkinVertexPassthroughFactoryShaderParameters);
 #endif // RHI_RAYTRACING
+#if defined(GPUCULL_TODO) // Supports primitive ID stream
+IMPLEMENT_VERTEX_FACTORY_TYPE_EX(FGPUSkinPassthroughVertexFactory, "/Engine/Private/LocalVertexFactory.ush", true, false, true, false, false, false, true, false);
+#else
 IMPLEMENT_VERTEX_FACTORY_TYPE(FGPUSkinPassthroughVertexFactory, "/Engine/Private/LocalVertexFactory.ush", true, false, true, false, false);
+#endif
 
 /*-----------------------------------------------------------------------------
 TGPUSkinMorphVertexFactory
diff --git a/Engine/Source/Runtime/Engine/Private/SkeletalMesh.cpp b/Engine/Source/Runtime/Engine/Private/SkeletalMesh.cpp
--- a/Engine/Source/Runtime/Engine/Private/SkeletalMesh.cpp
+++ b/Engine/Source/Runtime/Engine/Private/SkeletalMesh.cpp
@@ -5090,7 +5090,25 @@ FSkeletalMeshSceneProxy::FSkeletalMeshSceneProxy(const USkinnedMeshComponent* Co
 	}
 
 	// Skip primitive uniform buffer if we will be using local vertex factory which gets it's data from GPUScene.
+	// GPUCULL_TODO: Move to base class - all should follow the same pattern here
+#if defined(GPUCULL_TODO)
+	bVFRequiresPrimitiveUniformBuffer = !UseGPUScene(GMaxRHIShaderPlatform, FeatureLevel);
+	bSupportsInstanceDataBuffer = true;
+#else
 	bVFRequiresPrimitiveUniformBuffer = !((bIsCPUSkinned || bRenderStatic) && UseGPUScene(GMaxRHIShaderPlatform, FeatureLevel));
+#endif
+
+#if defined(GPUCULL_TODO)
+	// 100% generic default-instance
+	Instances.SetNum(1);
+	FPrimitiveInstance& Instance = Instances[0];
+	Instance.PrimitiveId = ~uint32(0);
+	Instance.InstanceToLocal.SetIdentity();
+	Instance.LocalToInstance.SetIdentity();
+	Instance.LocalToWorld.SetIdentity();
+	Instance.RenderBounds = GetLocalBounds();
+	Instance.LocalBounds = Instance.RenderBounds;
+#endif
 
 #if RHI_RAYTRACING
 	if (IsRayTracingEnabled())
@@ -6048,6 +6066,17 @@ void FSkeletalMeshSceneProxy::OnTransformChanged()
 {
 	// OnTransformChanged is called on the following frame after FSkeletalMeshObject::Update(), thus omit '+ 1' to frame number.
 	MeshObject->RefreshClothingTransforms(GetLocalToWorld(), GetScene().GetFrameNumber());
+
+#if defined(GPUCULL_TODO)
+	// Update the default-instance
+	// GPUCULL_TODO: really making more sense to move to the base proxy, or auto-generate default instance on GPU-scene update.
+	if (Instances.Num() == 1)
+	{
+	    FPrimitiveInstance& Instance = Instances[0];
+	    Instance.RenderBounds = GetLocalBounds();
+	    Instance.LocalBounds = Instance.RenderBounds;
+	}
+#endif
 }
 
 
diff --git a/Engine/Source/Runtime/Engine/Public/SkeletalMeshTypes.h b/Engine/Source/Runtime/Engine/Public/SkeletalMeshTypes.h
--- a/Engine/Source/Runtime/Engine/Public/SkeletalMeshTypes.h
+++ b/Engine/Source/Runtime/Engine/Public/SkeletalMeshTypes.h
@@ -206,6 +206,16 @@ public:
 	const TArray<FMatrix>& GetMeshObjectReferenceToLocalMatrices() const;
 	const TIndirectArray<FSkeletalMeshLODRenderData>& GetSkeletalMeshRenderDataLOD() const;
 
+	virtual const TArray<FPrimitiveInstance>* GetPrimitiveInstances() const
+	{
+		return &Instances;
+	}
+
+	virtual TArray<FPrimitiveInstance>* GetPrimitiveInstances()
+	{
+		return &Instances;
+	}
+
 protected:
 	AActor* Owner;
 	class FSkeletalMeshObject* MeshObject;
@@ -285,6 +295,9 @@ protected:
 	/** The primitive's pre-skinned local space bounds. */
 	FBoxSphereBounds PreSkinnedLocalBounds;
 
+	// Instances (going to contain just one).
+	TArray<FPrimitiveInstance> Instances;
+
 #if !(UE_BUILD_SHIPPING || UE_BUILD_TEST)
 	/** The color we draw this component in if drawing debug bones */
 	TOptional<FLinearColor> DebugDrawColor;
