FSoundWaveProxy refactor to get USoundWave off the audio render thread.

diff --git a/Engine/Source/Runtime/AudioMixer/Private/AudioMixerBlueprintLibrary.cpp b/Engine/Source/Runtime/AudioMixer/Private/AudioMixerBlueprintLibrary.cpp
--- a/Engine/Source/Runtime/AudioMixer/Private/AudioMixerBlueprintLibrary.cpp
+++ b/Engine/Source/Runtime/AudioMixer/Private/AudioMixerBlueprintLibrary.cpp
@@ -532,7 +532,7 @@ void UAudioMixerBlueprintLibrary::PrimeSoundForPlayback(USoundWave* SoundWave, c
 	}
 	else
 	{
-		IStreamingManager::Get().GetAudioStreamingManager().RequestChunk(SoundWave, 1, [OnLoadCompletion, SoundWave](EAudioChunkLoadResult InResult) 
+		IStreamingManager::Get().GetAudioStreamingManager().RequestChunk(SoundWave->GetThisAsProxy(), 1, [OnLoadCompletion, SoundWave](EAudioChunkLoadResult InResult)
 		{
 			AsyncTask(ENamedThreads::GameThread, [OnLoadCompletion, SoundWave, InResult]() {
 				if (InResult == EAudioChunkLoadResult::Completed || InResult == EAudioChunkLoadResult::AlreadyLoaded)
diff --git a/Engine/Source/Runtime/AudioMixer/Private/AudioMixerBuffer.cpp b/Engine/Source/Runtime/AudioMixer/Private/AudioMixerBuffer.cpp
--- a/Engine/Source/Runtime/AudioMixer/Private/AudioMixerBuffer.cpp
+++ b/Engine/Source/Runtime/AudioMixer/Private/AudioMixerBuffer.cpp
@@ -313,7 +313,7 @@ namespace Audio
 
 		FSoundQualityInfo QualityInfo = { 0 };
 
-		Buffer->DecompressionState = AudioDevice->CreateCompressedAudioInfo(InWave);
+		Buffer->DecompressionState = AudioDevice->CreateCompressedAudioInfo(InWave->GetThisAsProxy());
 
 		// Get the header information of our compressed format
 		if (Buffer->DecompressionState->StreamCompressedInfo(InWave, &QualityInfo))
@@ -366,7 +366,7 @@ namespace Audio
 			InWave->InitAudioResource(AudioDevice->GetRuntimeFormat(InWave));
 		}
 
-		Buffer->DecompressionState = AudioDevice->CreateCompressedAudioInfo(InWave);
+		Buffer->DecompressionState = AudioDevice->CreateCompressedAudioInfo(InWave->GetThisAsProxy());
 
 		if (Buffer->DecompressionState)
 		{
diff --git a/Engine/Source/Runtime/AudioMixer/Private/AudioMixerSourceBuffer.cpp b/Engine/Source/Runtime/AudioMixer/Private/AudioMixerSourceBuffer.cpp
--- a/Engine/Source/Runtime/AudioMixer/Private/AudioMixerSourceBuffer.cpp
+++ b/Engine/Source/Runtime/AudioMixer/Private/AudioMixerSourceBuffer.cpp
@@ -368,7 +368,7 @@ namespace Audio
 		if (InDecoder && InDecoder->HasError())
 		{			
 			UE_LOG(LogAudioMixer, Warning, TEXT("Decoder Error, stopping source [%s]"), 
-				*GetNameSafe(InDecoder->GetStreamingSoundWave()));
+				*InDecoder->GetStreamingSoundWave().GetFName().ToString());
 			bHasError = true;
 			bBufferFinished = true;
 			return false;	
diff --git a/Engine/Source/Runtime/Engine/Classes/Sound/SoundWave.h b/Engine/Source/Runtime/Engine/Classes/Sound/SoundWave.h
--- a/Engine/Source/Runtime/Engine/Classes/Sound/SoundWave.h
+++ b/Engine/Source/Runtime/Engine/Classes/Sound/SoundWave.h
@@ -16,6 +16,7 @@
 #include "Serialization/BulkDataBuffer.h"
 #include "Sound/SoundGroups.h"
 #include "Sound/SoundWaveLoadingBehavior.h"
+#include "UObject/ObjectKey.h"
 #include "AudioMixerTypes.h"
 #include "AudioCompressionSettings.h"
 #include "PerPlatformProperties.h"
@@ -66,8 +67,6 @@ struct FStreamedAudioChunk
 	/** Bulk data if stored in the package. */
 	FByteBulkData BulkData;
 
-	/** This is set by the audio stream cache to speed up lookup. Maps directly to the location into the cache chunk table that we loaded this chunk into. */
-	uint64 CacheLookupID = InvalidAudioStreamCacheLookupID;
 
 	/** Serialization. */
 	void Serialize(FArchive& Ar, UObject* Owner, int32 ChunkIndex);
@@ -91,7 +90,7 @@ struct FStreamedAudioChunk
  * Platform-specific data used streaming audio at runtime.
  */
 USTRUCT()
-struct FStreamedAudioPlatformData
+struct ENGINE_API FStreamedAudioPlatformData
 {
 	GENERATED_USTRUCT_BODY()
 
@@ -133,7 +132,7 @@ struct FStreamedAudioPlatformData
 	void Cache(class USoundWave& InSoundWave, const FPlatformAudioCookOverrides* CompressionOverrides, FName AudioFormatName, uint32 InFlags);
 	void FinishCache();
 	bool IsFinishedCache() const;
-	ENGINE_API bool TryInlineChunkData();
+	bool TryInlineChunkData();
 	bool AreDerivedChunksAvailable() const;
 #endif // WITH_EDITORONLY_DATA
 
@@ -283,21 +282,83 @@ struct ISoundWaveClient
 	virtual void OnFinishDestroy(class USoundWave* Wave) = 0;
 };
 
-class FSoundWaveProxy : public Audio::TProxyData<FSoundWaveProxy>
-{
-	IMPL_AUDIOPROXY_CLASS(FSoundWaveProxy)
 
+class ENGINE_API FSoundWaveProxy : public Audio::TProxyData<FSoundWaveProxy>
+{
 public:
-	// For now, this is just a weak object pointer.
-	// However, in the future this should be a thread-safe handle to whatever data is needed for this wave.
-	TWeakObjectPtr<USoundWave> SoundWavePtr;
+	IMPL_AUDIOPROXY_CLASS(FSoundWaveProxy);
 
-	FSoundWaveProxy(USoundWave* InWave)
-		: SoundWavePtr(InWave)
-	{
-	}
+	explicit FSoundWaveProxy(USoundWave* InWave);
+	FSoundWaveProxy(const FSoundWaveProxy& Other);
+
+	// UObject Interface cached data getters
+	const FName& GetFName() const { return NameCached; }
+	const FString& GetFullName() const { return FullNameCached; }
+	const FString& GetPathName() const { return PathNameCached; }
+	const FName& GetRuntimeFormat() const { return RuntimeFormat; }
+	const FObjectKey& GetFObjectKey() const { return SoundWaveKeyCached; };
+
+	float GetSampleRate() const { return SampleRate; }
+	uint32 GetNumChannels() const { return NumChannels; }
+	uint32 GetNumChunks() const;
+	uint32 GetSizeOfChunk(uint32 ChunkIndex) const;
+
+	void ReleaseCompressedAudio();
+
+	bool IsStreaming() const { return bIsStreaming; }
+	bool IsSeekableStreaming() const { return bIsStreaming && bSeekableStreaming; }
+	bool IsRetainingAudio() const { return FirstChunk.IsValid(); }
+	bool WasLoadingBehaviorOverridden() const { return bLoadingBehaviorOverridden; }
+	ESoundWaveLoadingBehavior GetLoadingBehavior() const { return LoadingBehavior; }
+
+	TArrayView<const uint8> GetZerothChunk(bool bForImmediatePlayback = false) const;
+	bool GetChunkData(int32 ChunkIndex, uint8** OutChunkData, bool bMakeSureChunkIsLoaded = false);
+
+	bool IsZerothChunkDataLoaded() const;
+	const TArrayView<uint8> GetZerothChunkDataView() const;
+	void EnsureZerothChunkIsLoaded();
+
+#if WITH_EDITOR
+	int32 GetCurrentChunkRevision() const;
+#endif // #if WITH_EDITOR
+
+	const FStreamedAudioChunk& GetChunk(uint32 ChunkIndex) const;
+	int32 GetChunkFromDDC(int32 ChunkIndex, uint8** OutChunkData, bool bMakeSureChunkIsLoaded);
+	FString GetDerivedDataKey() const;
+
+	/** Zeroth Chunk of audio for sources that use Load On Demand. */
+	FBulkDataBuffer<uint8> ZerothChunkData;
+
+	/** The streaming derived data for this sound on this platform. */
+	TUniquePtr<FStreamedAudioPlatformData> RunningPlatformData{ nullptr };
+
+private:
+	FName NameCached;
+	FString FullNameCached;
+	FString PathNameCached;
+	FObjectKey SoundWaveKeyCached;
+	FName RuntimeFormat{ "FSoundWaveProxy_InvalidFormat" };
+
+	float SampleRate;
+	uint32 NumChannels;
+	uint32 NumChunks;
+	FAudioChunkHandle FirstChunk;
+	ESoundWaveLoadingBehavior LoadingBehavior;
+
+	bool bLoadingBehaviorOverridden;
+	bool bIsStreaming;
+	bool bSeekableStreaming;
+	bool bShouldUseStreamCaching;
+
+#if WITH_EDITOR
+	TSharedPtr<FThreadSafeCounter> CurrentChunkRevision{ nullptr };
+#endif // #if WITH_EDITOR
+
+
+	friend class USoundWave;
 };
 
+
 UCLASS(hidecategories=Object, editinlinenew, BlueprintType)
 class ENGINE_API USoundWave : public USoundBase, public IAudioProxyDataFactory
 {
@@ -305,101 +366,102 @@ class ENGINE_API USoundWave : public USoundBase, public IAudioProxyDataFactory
 
 public:
 	/** Platform agnostic compression quality. 1..100 with 1 being best compression and 100 being best quality. */
-	UPROPERTY(EditAnywhere, Category="Format|Quality", meta=(DisplayName = "Compression", ClampMin = "1", ClampMax = "100"), AssetRegistrySearchable)
-	int32 CompressionQuality;
+	UPROPERTY(EditAnywhere, Category = "Format|Quality", meta = (DisplayName = "Compression", ClampMin = "1", ClampMax = "100"), AssetRegistrySearchable)
+		int32 CompressionQuality;
 
 	/** Priority of this sound when streaming (lower priority streams may not always play) */
-	UPROPERTY(EditAnywhere, Category="Playback|Streaming", meta=(ClampMin=0))
-	int32 StreamingPriority;
+	UPROPERTY(EditAnywhere, Category = "Playback|Streaming", meta = (ClampMin = 0))
+		int32 StreamingPriority;
 
 	/** Quality of sample rate conversion for platforms that opt into resampling during cook. The sample rate for each enumeration is definable per platform in platform target settings. */
 	UPROPERTY(EditAnywhere, Category = "Format|Quality")
-	ESoundwaveSampleRateSettings SampleRateQuality;
+		ESoundwaveSampleRateSettings SampleRateQuality;
 
 	/** Type of buffer this wave uses. Set once on load */
 	TEnumAsByte<EDecompressionType> DecompressionType;
 
-	UPROPERTY(EditAnywhere, Category=Sound, meta=(DisplayName="Group"))
-	TEnumAsByte<ESoundGroup> SoundGroup;
+	UPROPERTY(EditAnywhere, Category = Sound, meta = (DisplayName = "Group"))
+		TEnumAsByte<ESoundGroup> SoundGroup;
 
 	/** If set, when played directly (not through a sound cue) the wave will be played looping. */
-	UPROPERTY(EditAnywhere, Category=Sound, AssetRegistrySearchable)
-	uint8 bLooping:1;
+	UPROPERTY(EditAnywhere, Category = Sound, AssetRegistrySearchable)
+		uint8 bLooping : 1;
 
 	/** Whether this sound can be streamed to avoid increased memory usage. If using Stream Caching, use Loading Behavior instead to control memory usage. */
-	UPROPERTY(EditAnywhere, Category="Playback|Streaming", meta = (DisplayName = "Force Streaming"))
-	uint8 bStreaming:1;
+	UPROPERTY(EditAnywhere, Category = "Playback|Streaming", meta = (DisplayName = "Force Streaming"))
+		uint8 bStreaming : 1;
 
 	/** Whether this sound supports seeking. This requires recooking with a codec which supports seekability and streaming. */
 	UPROPERTY(EditAnywhere, Category = "Playback|Streaming", meta = (DisplayName = "Seekable", EditCondition = "bStreaming"))
-	uint8 bSeekableStreaming:1;
+		uint8 bSeekableStreaming : 1;
 
+// Loading behavior members are lazily initialized in const getters
 	/** Specifies how and when compressed audio data is loaded for asset if stream caching is enabled. */
 	UPROPERTY(EditAnywhere, Category = "Loading", meta = (DisplayName = "Loading Behavior Override"))
-	ESoundWaveLoadingBehavior LoadingBehavior;
+	mutable ESoundWaveLoadingBehavior LoadingBehavior;
 
 	/** Set to true if LoadingBehavior was inherited from a SoundCue. This is useful for debugging/logging */
-	uint8 bLoadingBehaviorOverridden:1;
+	mutable uint8 bLoadingBehaviorOverridden : 1;
 
 	/** Set to true for programmatically generated audio. */
-	uint8 bProcedural:1;
-	
+	uint8 bProcedural : 1;
+
 	/** Set to true if fade is required when sound is abruptly stopped. */
 	uint8 bRequiresStopFade:1;
 
 	/** Set to true of this is a bus sound source. This will result in the sound wave not generating audio for itself, but generate audio through instances. Used only in audio mixer. */
-	uint8 bIsSourceBus:1;
+	uint8 bIsSourceBus : 1;
 
 	/** Set to true for procedural waves that can be processed asynchronously. */
-	uint8 bCanProcessAsync:1;
+	uint8 bCanProcessAsync : 1;
 
 	/** Whether to free the resource data after it has been uploaded to the hardware */
-	uint8 bDynamicResource:1;
+	uint8 bDynamicResource : 1;
 
 	/** If set to true if this sound is considered to contain mature/adult content. */
-	UPROPERTY(EditAnywhere, BlueprintReadOnly, Category=Subtitles, AssetRegistrySearchable)
-	uint8 bMature:1;
+	UPROPERTY(EditAnywhere, BlueprintReadOnly, Category = Subtitles, AssetRegistrySearchable)
+		uint8 bMature : 1;
 
 	/** If set to true will disable automatic generation of line breaks - use if the subtitles have been split manually. */
-	UPROPERTY(EditAnywhere, BlueprintReadOnly, Category=Subtitles )
-	uint8 bManualWordWrap:1;
+	UPROPERTY(EditAnywhere, BlueprintReadOnly, Category = Subtitles)
+		uint8 bManualWordWrap : 1;
 
 	/** If set to true the subtitles display as a sequence of single lines as opposed to multiline. */
-	UPROPERTY(EditAnywhere, BlueprintReadOnly, Category=Subtitles )
-	uint8 bSingleLine:1;
+	UPROPERTY(EditAnywhere, BlueprintReadOnly, Category = Subtitles)
+		uint8 bSingleLine : 1;
 
 #if WITH_EDITORONLY_DATA
 	UPROPERTY()
-	uint8 bVirtualizeWhenSilent_DEPRECATED:1;
+		uint8 bVirtualizeWhenSilent_DEPRECATED : 1;
 #endif // WITH_EDITORONLY_DATA
 
-	/** Whether or not this source is ambisonics file format. If set, sound always uses the 
+	/** Whether or not this source is ambisonics file format. If set, sound always uses the
 	  * 'Master Ambisonics Submix' as set in the 'Audio' category of Project Settings'
 	  * and ignores submix if provided locally or in the referenced SoundClass. */
 	UPROPERTY(EditAnywhere, Category = Format)
-	uint8 bIsAmbisonics : 1;
+		uint8 bIsAmbisonics : 1;
 
 	/** Whether this SoundWave was decompressed from OGG. */
 	uint8 bDecompressedFromOgg : 1;
 
 #if WITH_EDITOR
 	/** The current revision of our compressed audio data. Used to tell when a chunk in the cache is stale. */
-	FThreadSafeCounter CurrentChunkRevision;
+	TSharedPtr<FThreadSafeCounter> CurrentChunkRevision{ MakeShared<FThreadSafeCounter>() };
 #endif
 
 private:
 
 	// This is set to false on initialization, then set to true on non-editor platforms when we cache appropriate sample rate.
-	uint8 bCachedSampleRateFromPlatformSettings:1;
+	uint8 bCachedSampleRateFromPlatformSettings : 1;
 
 	// This is set when SetSampleRate is called to invalidate our cached sample rate while not re-parsing project settings.
-	uint8 bSampleRateManuallyReset:1;
+	uint8 bSampleRateManuallyReset : 1;
 
 #if WITH_EDITOR
 	// Whether this was previously cooked with stream caching enabled.
-	uint8 bWasStreamCachingEnabledOnLastCook:1;
+	uint8 bWasStreamCachingEnabledOnLastCook : 1;
 	// Whether this asset is loaded from cooked data.
-	uint8 bLoadedFromCookedData:1;
+	uint8 bLoadedFromCookedData : 1;
 #endif // !WITH_EDITOR
 
 	enum class ESoundWaveResourceState : uint8
@@ -418,74 +480,74 @@ public:
 #if WITH_EDITORONLY_DATA
 	/** Specify a sound to use for the baked analysis. Will default to this USoundWave if not set. */
 	UPROPERTY(EditAnywhere, Category = "Analysis")
-	USoundWave* OverrideSoundToUseForAnalysis;
+		USoundWave* OverrideSoundToUseForAnalysis;
 
 	/**
 		Whether or not we should treat the sound wave used for analysis (this or the override) as looping while performing analysis.
 		A looping sound may include the end of the file for inclusion in analysis for envelope and FFT analysis.
 	*/
 	UPROPERTY(EditAnywhere, Category = "Analysis")
-	uint8 TreatFileAsLoopingForAnalysis:1;
+		uint8 TreatFileAsLoopingForAnalysis : 1;
 
 	/** Whether or not to enable cook-time baked FFT analysis. */
 	UPROPERTY(EditAnywhere, Category = "Analysis|FFT")
-	uint8 bEnableBakedFFTAnalysis : 1;
+		uint8 bEnableBakedFFTAnalysis : 1;
 
 	/** Whether or not to enable cook-time amplitude envelope analysis. */
 	UPROPERTY(EditAnywhere, Category = "Analysis|Envelope")
-	uint8 bEnableAmplitudeEnvelopeAnalysis : 1;
+		uint8 bEnableAmplitudeEnvelopeAnalysis : 1;
 
 	/** The FFT window size to use for fft analysis. */
 	UPROPERTY(EditAnywhere, Category = "Analysis|FFT", meta = (EditCondition = "bEnableBakedFFTAnalysis"))
-	ESoundWaveFFTSize FFTSize;
+		ESoundWaveFFTSize FFTSize;
 
 	/** How many audio frames analyze at a time. */
 	UPROPERTY(EditAnywhere, Category = "Analysis|FFT", meta = (EditCondition = "bEnableBakedFFTAnalysis", ClampMin = "512", UIMin = "512"))
-	int32 FFTAnalysisFrameSize;
+		int32 FFTAnalysisFrameSize;
 
 	/** Attack time in milliseconds of the spectral envelope follower. */
 	UPROPERTY(EditAnywhere, Category = "Analysis|FFT", meta = (EditCondition = "bEnableBakedFFTAnalysis", ClampMin = "0", UIMin = "0"))
-	int32 FFTAnalysisAttackTime;
+		int32 FFTAnalysisAttackTime;
 
 	/** Release time in milliseconds of the spectral envelope follower. */
 	UPROPERTY(EditAnywhere, Category = "Analysis|FFT", meta = (EditCondition = "bEnableBakedFFTAnalysis", ClampMin = "0", UIMin = "0"))
-	int32 FFTAnalysisReleaseTime;
+		int32 FFTAnalysisReleaseTime;
 
 	/** How many audio frames to average a new envelope value. Larger values use less memory for audio envelope data but will result in lower envelope accuracy. */
 	UPROPERTY(EditAnywhere, Category = "Analysis|Envelope", meta = (EditCondition = "bEnableAmplitudeEnvelopeAnalysis", ClampMin = "512", UIMin = "512"))
-	int32 EnvelopeFollowerFrameSize;
+		int32 EnvelopeFollowerFrameSize;
 
 	/** The attack time in milliseconds. Describes how quickly the envelope analyzer responds to increasing amplitudes. */
 	UPROPERTY(EditAnywhere, Category = "Analysis|Envelope", meta = (EditCondition = "bEnableAmplitudeEnvelopeAnalysis", ClampMin = "0", UIMin = "0"))
-	int32 EnvelopeFollowerAttackTime;
+		int32 EnvelopeFollowerAttackTime;
 
 	/** The release time in milliseconds. Describes how quickly the envelope analyzer responds to decreasing amplitudes. */
 	UPROPERTY(EditAnywhere, Category = "Analysis|Envelope", meta = (EditCondition = "bEnableAmplitudeEnvelopeAnalysis", ClampMin = "0", UIMin = "0"))
-	int32 EnvelopeFollowerReleaseTime;
+		int32 EnvelopeFollowerReleaseTime;
 #endif // WITH_EDITORONLY_DATA
 
 	/** The frequencies (in hz) to analyze when doing baked FFT analysis. */
 	UPROPERTY(EditAnywhere, Category = "Analysis|FFT", meta = (EditCondition = "bEnableBakedFFTAnalysis"))
-	TArray<float> FrequenciesToAnalyze;
+		TArray<float> FrequenciesToAnalyze;
 
 	/** The cooked spectral time data. */
 	UPROPERTY()
-	TArray<FSoundWaveSpectralTimeData> CookedSpectralTimeData;
+		TArray<FSoundWaveSpectralTimeData> CookedSpectralTimeData;
 
 	/** The cooked cooked envelope data. */
 	UPROPERTY()
-	TArray<FSoundWaveEnvelopeTimeData> CookedEnvelopeTimeData;
+		TArray<FSoundWaveEnvelopeTimeData> CookedEnvelopeTimeData;
 
 	/** Helper function to get interpolated cooked FFT data for a given time value. */
 	bool GetInterpolatedCookedFFTDataForTime(float InTime, uint32& InOutLastIndex, TArray<FSoundWaveSpectralData>& OutData, bool bLoop);
 	bool GetInterpolatedCookedEnvelopeDataForTime(float InTime, uint32& InOutLastIndex, float& OutAmplitude, bool bLoop);
 
-	/** If stream caching is enabled, allows the user to retain a strong handle to the first chunk of audio in the cache. 
+	/** If stream caching is enabled, allows the user to retain a strong handle to the first chunk of audio in the cache.
 	 *  Please note that this USoundWave is NOT guaranteed to be still alive when OnLoadCompleted is called.
 	 */
 	void GetHandleForChunkOfAudio(TFunction<void(FAudioChunkHandle&&)> OnLoadCompleted, bool bForceSync = false, int32 ChunkIndex = 1, ENamedThreads::Type CallbackThread = ENamedThreads::GameThread);
 
-	/** If stream caching is enabled, set this sound wave to retain a strong handle to its first chunk. 
+	/** If stream caching is enabled, set this sound wave to retain a strong handle to its first chunk.
 	 *  If not called on the game thread, bForceSync must be true.
 	*/
 	void RetainCompressedAudio(bool bForceSync = false);
@@ -501,19 +563,6 @@ public:
 	 */
 	void OverrideLoadingBehavior(ESoundWaveLoadingBehavior InLoadingBehavior);
 
-	/**
-	 * This is called by the audio stream cache once we put a chunk of compressed data in the chunk table. 
-	 * It allows us to look up the chunk directly, rather than searching the cache linearly.
-	 */
-	void SetCacheLookupIDForChunk(uint32 InChunkIndex, uint64 InCacheLookupID);
-
-	/**
-	 * This is called by the audio stream cache when we try to get a chunk of compressed data.
-	 * It allows us to look up the chunk directly, rather than searching the cache linearly.
-	 * Returns InvalidAudioStreamCacheLookupID if the chunk hasn't been added to the cache.
-	 */
-	uint64 GetCacheLookupIDForChunk(uint32 InChunkIndex);
-
 	/** Returns the loading behavior we should use for this sound wave.
 	 *  If this is called within Serialize(), this should be called with bCheckSoundClasses = false,
 	 *  Since there is no guarantee that the deserialized USoundClasses have been resolved yet.
@@ -522,7 +571,7 @@ public:
 
 	/** Use this to override how much audio data is loaded when this USoundWave is loaded. */
 	UPROPERTY(EditAnywhere, AdvancedDisplay, Category = "Loading")
-	int32 InitialChunkSize;
+		int32 InitialChunkSize;
 
 private:
 
@@ -543,39 +592,39 @@ private:
 
 	// We cache a soundwave's loading behavior on the first call to USoundWave::GetLoadingBehaviorForWave(true);
 	// Caches resolved loading behavior from the SoundClass graph. Must be called on the game thread.
-	void CacheInheritedLoadingBehavior();
-	ESoundWaveLoadingBehavior CachedSoundWaveLoadingBehavior;
+	void CacheInheritedLoadingBehavior() const;
+	mutable ESoundWaveLoadingBehavior CachedSoundWaveLoadingBehavior{ ESoundWaveLoadingBehavior::Uninitialized };
 
 public:
 
 	/** A localized version of the text that is actually spoken phonetically in the audio. */
-	UPROPERTY(EditAnywhere, BlueprintReadOnly, Category=Subtitles )
-	FString SpokenText;
+	UPROPERTY(EditAnywhere, BlueprintReadOnly, Category = Subtitles)
+		FString SpokenText;
 
 	/** The priority of the subtitle. */
-	UPROPERTY(EditAnywhere, BlueprintReadOnly, Category=Subtitles)
-	float SubtitlePriority;
+	UPROPERTY(EditAnywhere, BlueprintReadOnly, Category = Subtitles)
+		float SubtitlePriority;
 
 	/** Playback volume of sound 0 to 1 - Default is 1.0. */
-	UPROPERTY(Category=Sound, meta=(ClampMin = "0.0"), EditAnywhere)
-	float Volume;
+	UPROPERTY(Category = Sound, meta = (ClampMin = "0.0"), EditAnywhere)
+		float Volume;
 
 	/** Playback pitch for sound. */
-	UPROPERTY(Category=Sound, meta=(ClampMin = "0.125", ClampMax = "4.0"), EditAnywhere)
-	float Pitch;
+	UPROPERTY(Category = Sound, meta = (ClampMin = "0.125", ClampMax = "4.0"), EditAnywhere)
+		float Pitch;
 
 	/** Number of channels of multichannel data; 1 or 2 for regular mono and stereo files */
-	UPROPERTY(Category=Info, AssetRegistrySearchable, VisibleAnywhere)
-	int32 NumChannels;
+	UPROPERTY(Category = Info, AssetRegistrySearchable, VisibleAnywhere)
+		int32 NumChannels;
 
 #if WITH_EDITORONLY_DATA
 	/** Offsets into the bulk data for the source wav data */
 	UPROPERTY()
-	TArray<int32> ChannelOffsets;
+		TArray<int32> ChannelOffsets;
 
 	/** Sizes of the bulk data for the source wav data */
 	UPROPERTY()
-	TArray<int32> ChannelSizes;
+		TArray<int32> ChannelSizes;
 
 #endif // WITH_EDITORONLY_DATA
 
@@ -583,7 +632,7 @@ protected:
 
 	/** Cached sample rate for displaying in the tools */
 	UPROPERTY(Category = Info, AssetRegistrySearchable, VisibleAnywhere)
-	int32 SampleRate;
+		int32 SampleRate;
 
 public:
 
@@ -600,24 +649,24 @@ public:
 	 * Subtitle cues.  If empty, use SpokenText as the subtitle.  Will often be empty,
 	 * as the contents of the subtitle is commonly identical to what is spoken.
 	 */
-	UPROPERTY(EditAnywhere, BlueprintReadOnly, Category=Subtitles)
-	TArray<struct FSubtitleCue> Subtitles;
+	UPROPERTY(EditAnywhere, BlueprintReadOnly, Category = Subtitles)
+		TArray<struct FSubtitleCue> Subtitles;
 
 #if WITH_EDITORONLY_DATA
 	/** Provides contextual information for the sound to the translator. */
-	UPROPERTY(EditAnywhere, Category=Subtitles )
-	FString Comment;
+	UPROPERTY(EditAnywhere, Category = Subtitles)
+		FString Comment;
 
 #endif // WITH_EDITORONLY_DATA
 
 #if WITH_EDITORONLY_DATA
 	UPROPERTY()
-	FString SourceFilePath_DEPRECATED;
+		FString SourceFilePath_DEPRECATED;
 	UPROPERTY()
-	FString SourceFileTimestamp_DEPRECATED;
+		FString SourceFileTimestamp_DEPRECATED;
 
-	UPROPERTY(VisibleAnywhere, Instanced, Category=ImportSettings)
-	class UAssetImportData* AssetImportData;
+	UPROPERTY(VisibleAnywhere, Instanced, Category = ImportSettings)
+		class UAssetImportData* AssetImportData;
 
 #endif // WITH_EDITORONLY_DATA
 
@@ -625,11 +674,11 @@ protected:
 
 	/** Curves associated with this sound wave */
 	UPROPERTY(EditAnywhere, Category = SoundWave, AdvancedDisplay)
-	class UCurveTable* Curves;
+		class UCurveTable* Curves;
 
 	/** Hold a reference to our internal curve so we can switch back to it if we want to */
 	UPROPERTY()
-	class UCurveTable* InternalCurves;
+		class UCurveTable* InternalCurves;
 
 	/** Potential strong handle to the first chunk of audio data. Can be released via ReleaseCompressedAudioData. */
 	FAudioChunkHandle FirstChunk;
@@ -662,9 +711,6 @@ public:
 	FOwnedBulkDataPtr* OwnedBulkDataPtr{ nullptr };
 	const uint8* ResourceData{ nullptr };
 
-	/** Zeroth Chunk of audio for sources that use Load On Demand. */
-	FBulkDataBuffer<uint8> ZerothChunkData;
-
 	/** Uncompressed wav data 16 bit in mono or stereo - stereo not allowed for multichannel data */
 	FByteBulkData RawData;
 
@@ -673,6 +719,12 @@ public:
 
 	FFormatContainer CompressedFormatData;
 
+	/** Zeroth Chunk of audio for sources that use Load On Demand. */
+	FBulkDataBuffer<uint8> ZerothChunkData;
+
+	/** The streaming derived data for this sound on this platform. */
+	FStreamedAudioPlatformData* RunningPlatformData{ nullptr };
+
 #if WITH_EDITORONLY_DATA
 	TMap<FName, uint32> AsyncLoadingDataFormats;
 
@@ -680,15 +732,11 @@ public:
 	FCriticalSection RawDataCriticalSection;
 
 #endif // WITH_EDITORONLY_DATA
-
-	/** The streaming derived data for this sound on this platform. */
-	FStreamedAudioPlatformData* RunningPlatformData;
-
 	/** cooked streaming platform data for this sound */
 	TSortedMap<FString, FStreamedAudioPlatformData*> CookedPlatformData;
 
 	//~ Begin UObject Interface.
-	virtual void Serialize( FArchive& Ar ) override;
+	virtual void Serialize(FArchive& Ar) override;
 	virtual void PostInitProperties() override;
 	virtual bool IsReadyForFinishDestroy() override;
 	virtual void FinishDestroy() override;
@@ -715,7 +763,7 @@ public:
 
 	//~ Begin USoundBase Interface.
 	virtual bool IsPlayable() const override;
-	virtual void Parse( class FAudioDevice* AudioDevice, const UPTRINT NodeWaveInstanceHash, FActiveSound& ActiveSound, const FSoundParseParameters& ParseParams, TArray<FWaveInstance*>& WaveInstances ) override;
+	virtual void Parse(class FAudioDevice* AudioDevice, const UPTRINT NodeWaveInstanceHash, FActiveSound& ActiveSound, const FSoundParseParameters& ParseParams, TArray<FWaveInstance*>& WaveInstances) override;
 	virtual float GetDuration() override;
 	virtual float GetSubtitlePriority() const override;
 	virtual bool SupportsSubtitles() const override;
@@ -733,7 +781,7 @@ public:
 
 	// Called when the procedural sound wave is done generating on the render thread. Only used in the audio mixer and when bProcedural is true..
 	virtual void OnEndGenerate() {};
-	
+
 	void AddPlayingSource(const FSoundWaveClientPtr& Source);
 	void RemovePlayingSource(const FSoundWaveClientPtr& Source);
 
@@ -752,11 +800,11 @@ public:
 	}
 
 	bool IsGeneratingAudio() const
-	{		
+	{
 		bool bIsGeneratingAudio = false;
 		FScopeLock Lock(&SourcesPlayingCs);
 		bIsGeneratingAudio = SourcesPlaying.Num() > 0;
-		
+
 		return bIsGeneratingAudio;
 	}
 
@@ -794,7 +842,7 @@ public:
 	/**
 	 * Copy the compressed audio data from the bulk data
 	 */
-	virtual void InitAudioResource( FByteBulkData& CompressedData );
+	virtual void InitAudioResource(FByteBulkData& CompressedData);
 
 	/**
 	 * Copy the compressed audio data from derived data cache
@@ -971,9 +1019,9 @@ public:
 	/**
 	 * Begins caching platform data in the background for the platform requested
 	 */
-	virtual void BeginCacheForCookedPlatformData(  const ITargetPlatform *TargetPlatform ) override;
+	virtual void BeginCacheForCookedPlatformData(const ITargetPlatform* TargetPlatform) override;
 
-	virtual bool IsCachedCookedPlatformDataLoaded( const ITargetPlatform* TargetPlatform ) override;
+	virtual bool IsCachedCookedPlatformDataLoaded(const ITargetPlatform* TargetPlatform) override;
 
 	/**
 	 * Clear all the cached cooked platform data which we have accumulated with BeginCacheForCookedPlatformData calls
@@ -981,11 +1029,11 @@ public:
 	 */
 	virtual void ClearAllCachedCookedPlatformData() override;
 
-	virtual void ClearCachedCookedPlatformData( const ITargetPlatform* TargetPlatform ) override;
+	virtual void ClearCachedCookedPlatformData(const ITargetPlatform* TargetPlatform) override;
 
 	virtual void WillNeverCacheCookedPlatformDataAgain() override;
 
-	uint32 bNeedsThumbnailGeneration:1;
+	uint32 bNeedsThumbnailGeneration : 1;
 #endif // WITH_EDITOR
 
 	/**
@@ -1026,6 +1074,11 @@ public:
 		return (ESoundWavePrecacheState)PrecacheState.GetValue();
 	}
 
+	const FSoundWaveProxy& GetThisAsProxy();
+
+private:
+	TUniquePtr<FSoundWaveProxy> InternalProxy{ nullptr };
+	friend class FSoundWaveProxy;
 };
 
 
diff --git a/Engine/Source/Runtime/Engine/Private/ADPCMAudioInfo.cpp b/Engine/Source/Runtime/Engine/Private/ADPCMAudioInfo.cpp
--- a/Engine/Source/Runtime/Engine/Private/ADPCMAudioInfo.cpp
+++ b/Engine/Source/Runtime/Engine/Private/ADPCMAudioInfo.cpp
@@ -472,18 +472,18 @@ void FADPCMAudioInfo::ProcessSeekRequest()
 	}
 }
 
-bool FADPCMAudioInfo::StreamCompressedInfoInternal(USoundWave* Wave, struct FSoundQualityInfo* QualityInfo)
+bool FADPCMAudioInfo::StreamCompressedInfoInternal(const FSoundWaveProxy& InWaveProxy, struct FSoundQualityInfo* QualityInfo)
 {
 	FScopeLock ScopeLock(&CurCompressedChunkHandleCriticalSection);
 
 	check(QualityInfo);
 
-	check(StreamingSoundWave == Wave);
+	check(StreamingSoundWave.Get() == &InWaveProxy);
 
 	CurrentChunkIndex = 0;
 
 	// Get the first chunk of audio data (should already be loaded)
-	uint8 const* ChunkData = GetLoadedChunk(Wave, CurrentChunkIndex, CurrentChunkDataSize);
+	uint8 const* ChunkData = GetLoadedChunk(InWaveProxy, CurrentChunkIndex, CurrentChunkDataSize);
 
 	if (ChunkData == nullptr)
 	{
@@ -507,7 +507,7 @@ bool FADPCMAudioInfo::StreamCompressedInfoInternal(USoundWave* Wave, struct FSou
 	if (((uint32)SampleDataOffset) >= CurrentChunkDataSize)
 	{
 		++CurrentChunkIndex;
-		ChunkData = GetLoadedChunk(Wave, CurrentChunkIndex, CurrentChunkDataSize);
+		ChunkData = GetLoadedChunk(InWaveProxy, CurrentChunkIndex, CurrentChunkDataSize);
 		FirstChunkSampleDataIndex = CurrentChunkIndex;
 		FirstChunkSampleDataOffset = 0;
 	}
@@ -679,7 +679,7 @@ bool FADPCMAudioInfo::StreamCompressedData(uint8* Destination, bool bLooping, ui
 					}
 
 					// Request the next chunk of data from the streaming engine
-					CurCompressedChunkData = GetLoadedChunk(StreamingSoundWave, CurrentChunkIndex, CurrentChunkDataSize);
+					CurCompressedChunkData = GetLoadedChunk(*StreamingSoundWave, CurrentChunkIndex, CurrentChunkDataSize);
 
 					if(CurCompressedChunkData == nullptr)
 					{
@@ -874,7 +874,7 @@ bool FADPCMAudioInfo::StreamCompressedData(uint8* Destination, bool bLooping, ui
 				}
 
 				// Request the next chunk of data from the streaming engine
-				CurCompressedChunkData = GetLoadedChunk(StreamingSoundWave, CurrentChunkIndex, CurrentChunkDataSize);
+				CurCompressedChunkData = GetLoadedChunk(*StreamingSoundWave, CurrentChunkIndex, CurrentChunkDataSize);
 
 				if(CurCompressedChunkData == nullptr)
 				{
@@ -985,7 +985,7 @@ bool FADPCMAudioInfo::ReleaseStreamChunk(bool bBlockUntilReleased)
 	}
 }
 
-const uint8* FADPCMAudioInfo::GetLoadedChunk(USoundWave* InSoundWave, uint32 ChunkIndex, uint32& OutChunkSize)
+const uint8* FADPCMAudioInfo::GetLoadedChunk(const FSoundWaveProxy& InSoundWave, uint32 ChunkIndex, uint32& OutChunkSize)
 {
 	if (ChunkIndex != 0 && FMath::RandRange(0.0f, 1.0f) < ChanceForIntentionalChunkMissCVar)
 	{
@@ -994,19 +994,14 @@ const uint8* FADPCMAudioInfo::GetLoadedChunk(USoundWave* InSoundWave, uint32 Chu
 		return nullptr;
 	}
 
-	if (!InSoundWave || ChunkIndex >= InSoundWave->GetNumChunks())
+	if (ChunkIndex >= InSoundWave.GetNumChunks())
 	{
-		if(InSoundWave)
-		{
-			UE_LOG(LogAudio, Verbose, TEXT("Error calling GetLoadedChunk on wave with %d chunks. ChunkIndex: %d. Name: %s"), InSoundWave->GetNumChunks(), ChunkIndex, *InSoundWave->GetFullName());
-		}
-		
 		OutChunkSize = 0;
 		return nullptr;
 	}
 	else if (ChunkIndex == 0)
 	{
-		TArrayView<const uint8> ZerothChunk = InSoundWave->GetZerothChunk(true);
+		TArrayView<const uint8> ZerothChunk = InSoundWave.GetZerothChunk(true);
 		OutChunkSize = ZerothChunk.Num();
 		PreviouslyRequestedChunkIndex = ChunkIndex;
 		return ZerothChunk.GetData();
diff --git a/Engine/Source/Runtime/Engine/Private/AudioDecompress.cpp b/Engine/Source/Runtime/Engine/Private/AudioDecompress.cpp
--- a/Engine/Source/Runtime/Engine/Private/AudioDecompress.cpp
+++ b/Engine/Source/Runtime/Engine/Private/AudioDecompress.cpp
@@ -142,14 +142,14 @@ void IStreamedCompressedInfo::ExpandFile(uint8* DstBuffer, struct FSoundQualityI
 	}
 }
 
-bool IStreamedCompressedInfo::StreamCompressedInfoInternal(USoundWave* Wave, struct FSoundQualityInfo* QualityInfo)
+bool IStreamedCompressedInfo::StreamCompressedInfoInternal(const FSoundWaveProxy& InWaveProxy, struct FSoundQualityInfo* QualityInfo)
 {
-	check(StreamingSoundWave == Wave);
+	check(StreamingSoundWave.Get() == &InWaveProxy);
 
 	// Get the first chunk of audio data (should always be loaded)
 	CurrentChunkIndex = 0;
 	uint32 ChunkSize = 0;
-	const uint8* FirstChunk = GetLoadedChunk(StreamingSoundWave, CurrentChunkIndex, ChunkSize);
+	const uint8* FirstChunk = GetLoadedChunk(*StreamingSoundWave, CurrentChunkIndex, ChunkSize);
 
 	bIsStreaming = true;
 	if (FirstChunk)
@@ -180,7 +180,7 @@ bool IStreamedCompressedInfo::StreamCompressedData(uint8* Destination, bool bLoo
 
 	SCOPE_CYCLE_COUNTER(STAT_AudioStreamedDecompressTime);
 
-	UE_LOG(LogAudio, Log, TEXT("Streaming compressed data from SoundWave'%s' - Chunk %d, Offset %d"), *StreamingSoundWave->GetName(), CurrentChunkIndex, SrcBufferOffset);
+	UE_LOG(LogAudio, Log, TEXT("Streaming compressed data from SoundWave'%s' - Chunk %d, Offset %d"), *StreamingSoundWave->GetFName().ToString(), CurrentChunkIndex, SrcBufferOffset);
 
 	// Write out any PCM data that was decoded during the last request
 	uint32 RawPCMOffset = WriteFromDecodedPCM(Destination, BufferSize);
@@ -189,7 +189,7 @@ bool IStreamedCompressedInfo::StreamCompressedData(uint8* Destination, bool bLoo
 	if (SrcBufferData == NULL)
 	{
 		uint32 ChunkSize = 0;
-		SrcBufferData = GetLoadedChunk(StreamingSoundWave, CurrentChunkIndex, ChunkSize);
+		SrcBufferData = GetLoadedChunk(*StreamingSoundWave, CurrentChunkIndex, ChunkSize);
 		if (SrcBufferData)
 		{
 			bPrintChunkFailMessage = true;
@@ -201,7 +201,7 @@ bool IStreamedCompressedInfo::StreamCompressedData(uint8* Destination, bool bLoo
 			// Still not loaded, zero remainder of current buffer
 			if (bPrintChunkFailMessage)
 			{
-				UE_LOG(LogAudio, Verbose, TEXT("Chunk %d not loaded from streaming manager for SoundWave '%s'. Likely due to stall on game thread."), CurrentChunkIndex, *StreamingSoundWave->GetName());
+				UE_LOG(LogAudio, Verbose, TEXT("Chunk %d not loaded from streaming manager for SoundWave '%s'. Likely due to stall on game thread."), CurrentChunkIndex, *StreamingSoundWave->GetFName().ToString());
 				bPrintChunkFailMessage = false;
 			}
 			ZeroBuffer(Destination + RawPCMOffset, BufferSize - RawPCMOffset);
@@ -242,7 +242,7 @@ bool IStreamedCompressedInfo::StreamCompressedData(uint8* Destination, bool bLoo
 			if (SrcBufferOffset >= SrcBufferDataSize - SrcBufferPadding)
 			{
 				// Special case for the last chunk of audio
-				if (CurrentChunkIndex == StreamingSoundWave->RunningPlatformData->NumChunks - 1)
+				if (CurrentChunkIndex == StreamingSoundWave->GetNumChunks() - 1)
 				{
 					// check whether all decoded PCM was written
 					if (LastPCMByteSize == 0)
@@ -274,10 +274,10 @@ bool IStreamedCompressedInfo::StreamCompressedData(uint8* Destination, bool bLoo
 					SrcBufferOffset = 0;
 				}
 
-				SrcBufferData = GetLoadedChunk(StreamingSoundWave, CurrentChunkIndex, SrcBufferDataSize);
+				SrcBufferData = GetLoadedChunk(*StreamingSoundWave, CurrentChunkIndex, SrcBufferDataSize);
 				if (SrcBufferData)
 				{
-					UE_LOG(LogAudio, Log, TEXT("Incremented current chunk from SoundWave'%s' - Chunk %d, Offset %d"), *StreamingSoundWave->GetName(), CurrentChunkIndex, SrcBufferOffset);
+					UE_LOG(LogAudio, Log, TEXT("Incremented current chunk from SoundWave'%s' - Chunk %d, Offset %d"), *StreamingSoundWave->GetFName().ToString(), CurrentChunkIndex, SrcBufferOffset);
 				}
 				else
 				{
@@ -357,21 +357,16 @@ uint32 IStreamedCompressedInfo::ZeroBuffer(uint8* Destination, uint32 BufferSize
 }
 
 
-const uint8* IStreamedCompressedInfo::GetLoadedChunk(USoundWave* InSoundWave, uint32 ChunkIndex, uint32& OutChunkSize)
+const uint8* IStreamedCompressedInfo::GetLoadedChunk(const FSoundWaveProxy& InSoundWave, uint32 ChunkIndex, uint32& OutChunkSize)
 {
-	if (!InSoundWave || ChunkIndex >= InSoundWave->GetNumChunks())
+	if (ChunkIndex >= InSoundWave.GetNumChunks())
 	{
-		if(InSoundWave)
-		{
-			UE_LOG(LogAudio, Verbose, TEXT("Error calling GetLoadedChunk on wave with %d chunks. ChunkIndex: %d. Name: %s"), InSoundWave->GetNumChunks(), ChunkIndex, *InSoundWave->GetFullName());
-		}
-		
 		OutChunkSize = 0;
 		return nullptr;
 	}
 	else if (ChunkIndex == 0)
 	{
-		TArrayView<const uint8> ZerothChunk = InSoundWave->GetZerothChunk(true);
+		TArrayView<const uint8> ZerothChunk = InSoundWave.GetZerothChunk(true);
 		OutChunkSize = ZerothChunk.Num();
 		return ZerothChunk.GetData();
 	}
@@ -651,7 +646,7 @@ FAsyncAudioDecompressWorker::FAsyncAudioDecompressWorker(USoundWave* InWave, int
 	}
 	else if (GEngine && GEngine->GetMainAudioDevice())
 	{
-		AudioInfo = GEngine->GetMainAudioDevice()->CreateCompressedAudioInfo(Wave);
+		AudioInfo = GEngine->GetMainAudioDevice()->CreateCompressedAudioInfo(Wave->GetThisAsProxy());
 	}
 }
 
@@ -755,3 +750,22 @@ bool ShouldUseBackgroundPoolFor_FAsyncRealtimeAudioTask()
 }
 
 // end
+
+bool ICompressedAudioInfo::StreamCompressedInfo(USoundWave* Wave, FSoundQualityInfo* QualityInfo)
+{
+	// Create and cache proxy object
+	Audio::FProxyDataInitParams Params;
+	Params.NameOfFeatureRequestingProxy = "ICompressedAudioInfo";
+
+	TUniquePtr<Audio::IProxyData> ProxyData = Wave->CreateNewProxyData(Params);
+	StreamingSoundWave = MakeUnique<FSoundWaveProxy>(MoveTemp(ProxyData->GetAs<FSoundWaveProxy>()));
+
+	return StreamCompressedInfoInternal(*StreamingSoundWave, QualityInfo);
+}
+
+bool ICompressedAudioInfo::StreamCompressedInfo(const FSoundWaveProxy& Wave, FSoundQualityInfo* QualityInfo)
+{
+	// Create our own copy of the proxy object
+	StreamingSoundWave = MakeUnique<FSoundWaveProxy>(Wave);
+	return StreamCompressedInfoInternal(*StreamingSoundWave, QualityInfo);
+}
diff --git a/Engine/Source/Runtime/Engine/Private/AudioStreaming.cpp b/Engine/Source/Runtime/Engine/Private/AudioStreaming.cpp
--- a/Engine/Source/Runtime/Engine/Private/AudioStreaming.cpp
+++ b/Engine/Source/Runtime/Engine/Private/AudioStreaming.cpp
@@ -127,9 +127,9 @@ void FStreamingWaveData::FreeResources()
 	}
 }
 
-bool FStreamingWaveData::Initialize(USoundWave* InSoundWave, FLegacyAudioStreamingManager* InAudioStreamingManager)
+bool FStreamingWaveData::Initialize(const FSoundWaveProxy& InSoundWave, FLegacyAudioStreamingManager* InAudioStreamingManager)
 {
-	if (!InSoundWave || !InSoundWave->RunningPlatformData->Chunks.Num())
+	if (!InSoundWave.GetNumChunks())
 	{
 #if WITH_EDITOR
 		UE_LOG(LogAudio, Display, TEXT("Failed to initialize streaming wave data due to lack of serialized stream chunks. Error during stream cooking."));
@@ -139,7 +139,7 @@ bool FStreamingWaveData::Initialize(USoundWave* InSoundWave, FLegacyAudioStreami
 		return false;
 	}
 
-	SoundWave = InSoundWave;
+	SoundWave = MakeUnique<FSoundWaveProxy>(InSoundWave);
 	AudioStreamingManager = InAudioStreamingManager;
 
 	// Always get the first chunk of data so we can play immediately
@@ -149,8 +149,9 @@ bool FStreamingWaveData::Initialize(USoundWave* InSoundWave, FLegacyAudioStreami
 	// Prepare 4 chunks of streaming wave data in loaded chunks array
 	LoadedChunks.Reset(4);
 
-	const int32 DataSize = SoundWave->RunningPlatformData->Chunks[0].DataSize;
-	const int32 AudioDataSize = SoundWave->RunningPlatformData->Chunks[0].AudioDataSize;
+	const FStreamedAudioChunk& Chunk = SoundWave->GetChunk(0);
+	const int32 DataSize = Chunk.DataSize;
+	const int32 AudioDataSize = Chunk.AudioDataSize;
 	const int32 FirstLoadedChunkIndex = AddNewLoadedChunk(DataSize, AudioDataSize);
 
 	FLoadedAudioChunk* FirstChunk = &LoadedChunks[FirstLoadedChunkIndex];
@@ -172,11 +173,6 @@ bool FStreamingWaveData::Initialize(USoundWave* InSoundWave, FLegacyAudioStreami
 
 bool FStreamingWaveData::UpdateStreamingStatus()
 {
-	if (!SoundWave)
-	{
-		return false;
-	}
-
 	bool	bHasPendingRequestInFlight = true;
 	int32	RequestStatus = PendingChunkChangeRequestStatus.GetValue();
 	TArray<uint32> IndicesToLoad;
@@ -199,7 +195,7 @@ bool FStreamingWaveData::UpdateStreamingStatus()
 				{
 					LogString += FString::Printf(TEXT(", %d"), IndicesToLoad[Index]);
 				}
-				LogString += FString::Printf(TEXT(" from SoundWave'%s'"), *SoundWave->GetName());
+				LogString += FString::Printf(TEXT(" from SoundWave'%s'"), *SoundWave->GetFName().ToString());
 				UE_LOG(LogAudio, Log, TEXT("%s"), *LogString);
 			}
 
@@ -278,7 +274,7 @@ void FStreamingWaveData::BeginPendingRequests(const TArray<uint32>& IndicesToLoa
 		{
 			LogString += FString::Printf(TEXT(", %d"), IndicesToLoad[Index]);
 		}
-		LogString += FString::Printf(TEXT(" from SoundWave'%s'"), *SoundWave->GetName());
+		LogString += FString::Printf(TEXT(" from SoundWave'%s'"), *SoundWave->GetFName().ToString());
 		UE_LOG(LogAudio, Log, TEXT("%s"), *LogString);
 	}
 
@@ -308,7 +304,7 @@ void FStreamingWaveData::BeginPendingRequests(const TArray<uint32>& IndicesToLoa
 		// Set off all IO Requests
 		for (auto Index : IndicesToLoad)
 		{
-			const FStreamedAudioChunk& Chunk = SoundWave->RunningPlatformData->Chunks[Index];
+			const FStreamedAudioChunk& Chunk = SoundWave->GetChunk(Index);
 			int32 ChunkSize = Chunk.DataSize;
 
 			int32 LoadedChunkStorageIndex = AddNewLoadedChunk(ChunkSize, Chunk.AudioDataSize);
@@ -344,7 +340,7 @@ void FStreamingWaveData::BeginPendingRequests(const TArray<uint32>& IndicesToLoa
 					ChunkSize,
 					&PendingChunkChangeRequestStatus,
 					MoveTemp(NullOnLoadCompletedCallback)
-					);
+				);
 				PendingAsyncStreamDerivedChunkTasks.Add(Task);
 				Task->StartBackgroundTask();
 			}
@@ -354,7 +350,7 @@ void FStreamingWaveData::BeginPendingRequests(const TArray<uint32>& IndicesToLoa
 				check(!ChunkStorage->IORequest); // Make sure we do not already have a request
 				check(!ChunkStorage->Data); // Make sure we do not already have data
 				check(Chunk.BulkData.GetBulkDataSize() == ChunkStorage->DataSize); // Make sure that the bulkdata size matches
-				
+
 				FBulkDataIORequestCallBack AsyncFileCallBack =
 					[this, LoadedChunkStorageIndex](bool bWasCancelled, IBulkDataIORequest* Req)
 				{
@@ -505,7 +501,7 @@ void FLegacyAudioStreamingManager::OnAsyncFileCallback(FStreamingWaveData* Strea
 
 		// Grab the loaded chunk memory ptr since it will be invalid as soon as this callback finishes
 		NewAudioChunkLoadResult->DataResults = Mem;
-		
+
 		// The chunk index to load the results into
 		NewAudioChunkLoadResult->LoadedAudioChunkIndex = LoadedAudioChunkIndex;
 
@@ -585,15 +581,15 @@ void FLegacyAudioStreamingManager::UpdateResourceStreaming(float DeltaTime, bool
 				if (SoundBuffer)
 				{
 					int32 SourceChunk = SoundBuffer->GetCurrentChunkIndex();
-					if (SourceChunk >= 0 && SourceChunk < Wave->RunningPlatformData->NumChunks)
+					if (SourceChunk >= 0 && SourceChunk < (int32)Wave->GetNumChunks())
 					{
 						WaveRequest.RequiredIndices.AddUnique(SourceChunk);
-						WaveRequest.RequiredIndices.AddUnique((SourceChunk + 1) % Wave->RunningPlatformData->NumChunks);
+						WaveRequest.RequiredIndices.AddUnique((SourceChunk + 1) % Wave->GetNumChunks());
 						WaveRequest.bPrioritiseRequest = true;
 					}
 					else
 					{
-						UE_LOG(LogAudio, Log, TEXT("Invalid chunk request curIndex=%d numChunks=%d\n"), SourceChunk, Wave->RunningPlatformData->NumChunks);
+						UE_LOG(LogAudio, Log, TEXT("Invalid chunk request curIndex=%d numChunks=%d\n"), SourceChunk, Wave->GetNumChunks());
 					}
 				}
 			}
@@ -602,34 +598,30 @@ void FLegacyAudioStreamingManager::UpdateResourceStreaming(float DeltaTime, bool
 
 	for (ICompressedAudioInfo* Decoder : CompressedAudioInfos)
 	{
-		USoundWave* SoundWave = Decoder->GetStreamingSoundWave();
-		if (SoundWave)
+		const FSoundWaveProxy& SoundWave = Decoder->GetStreamingSoundWave();
+		FStreamingWaveData** WaveDataPtr = StreamingSoundWaves.Find(SoundWave.GetFObjectKey());
+		if (WaveDataPtr && (*WaveDataPtr)->PendingChunkChangeRequestStatus.GetValue() == AudioState_ReadyFor_Requests)
 		{
-			FStreamingWaveData** WaveDataPtr = StreamingSoundWaves.Find(SoundWave);
-			if (WaveDataPtr && (*WaveDataPtr)->PendingChunkChangeRequestStatus.GetValue() == AudioState_ReadyFor_Requests)
+			FStreamingWaveData* WaveData = *WaveDataPtr;
+			// Request the chunk the source is using and the one after that
+			FWaveRequest& WaveRequest = GetWaveRequest(SoundWave.GetFObjectKey());
+			int32 SourceChunk = Decoder->GetCurrentChunkIndex();
+			if (SourceChunk >= 0 && SourceChunk < (int32)SoundWave.GetNumChunks())
 			{
-				FStreamingWaveData* WaveData = *WaveDataPtr;
-				// Request the chunk the source is using and the one after that
-				FWaveRequest& WaveRequest = GetWaveRequest(SoundWave);
-				int32 SourceChunk = Decoder->GetCurrentChunkIndex();
-				if (SourceChunk >= 0 && SourceChunk < SoundWave->RunningPlatformData->NumChunks)
-				{
-					WaveRequest.RequiredIndices.AddUnique(SourceChunk);
-					WaveRequest.RequiredIndices.AddUnique((SourceChunk + 1) % SoundWave->RunningPlatformData->NumChunks);
-					WaveRequest.bPrioritiseRequest = true;
-				}
-				else
-				{
-					UE_LOG(LogAudio, Log, TEXT("Invalid chunk request curIndex=%d numChunks=%d\n"), SourceChunk, SoundWave->RunningPlatformData->NumChunks);
-				}
+				WaveRequest.RequiredIndices.AddUnique(SourceChunk);
+				WaveRequest.RequiredIndices.AddUnique((SourceChunk + 1) % SoundWave.GetNumChunks());
+				WaveRequest.bPrioritiseRequest = true;
+			}
+			else
+			{
+				UE_LOG(LogAudio, Log, TEXT("Invalid chunk request curIndex=%d numChunks=%d\n"), SourceChunk, SoundWave.GetNumChunks());
 			}
 		}
 	}
 
 	for (auto Iter = WaveRequests.CreateIterator(); Iter; ++Iter)
 	{
-		USoundWave* Wave = Iter.Key();
-		FStreamingWaveData* WaveData = StreamingSoundWaves.FindRef(Wave);
+		FStreamingWaveData* WaveData = StreamingSoundWaves.FindRef(Iter.Key());
 
 		if (WaveData && WaveData->PendingChunkChangeRequestStatus.GetValue() == AudioState_ReadyFor_Requests)
 		{
@@ -707,18 +699,18 @@ void FLegacyAudioStreamingManager::NotifyLevelOffset(class ULevel* Level, const
 {
 }
 
-void FLegacyAudioStreamingManager::AddStreamingSoundWave(USoundWave* SoundWave)
+void FLegacyAudioStreamingManager::AddStreamingSoundWave(const FSoundWaveProxy& SoundWave)
 {
-	if (FPlatformProperties::SupportsAudioStreaming() && SoundWave->IsStreaming(nullptr))
+	if (FPlatformProperties::SupportsAudioStreaming() && SoundWave.IsStreaming())
 	{
 		FScopeLock Lock(&CriticalSection);
 
-		if (StreamingSoundWaves.FindRef(SoundWave) == nullptr)
+		if (StreamingSoundWaves.FindRef(SoundWave.GetFObjectKey()) == nullptr)
 		{
 			FStreamingWaveData* NewStreamingWaveData = new FStreamingWaveData;
 			if (NewStreamingWaveData->Initialize(SoundWave, this))
 			{
-				StreamingSoundWaves.Add(SoundWave, NewStreamingWaveData);
+				StreamingSoundWaves.Add(SoundWave.GetFObjectKey(), NewStreamingWaveData);
 			}
 			else
 			{
@@ -729,13 +721,13 @@ void FLegacyAudioStreamingManager::AddStreamingSoundWave(USoundWave* SoundWave)
 	}
 }
 
-void FLegacyAudioStreamingManager::RemoveStreamingSoundWave(USoundWave* SoundWave)
+void FLegacyAudioStreamingManager::RemoveStreamingSoundWave(const FSoundWaveProxy& SoundWave)
 {
 	FScopeLock Lock(&CriticalSection);
-	FStreamingWaveData* WaveData = StreamingSoundWaves.FindRef(SoundWave);
+	FStreamingWaveData* WaveData = StreamingSoundWaves.FindRef(SoundWave.GetFObjectKey());
 	if (WaveData)
 	{
-		StreamingSoundWaves.Remove(SoundWave);
+		StreamingSoundWaves.Remove(SoundWave.GetFObjectKey());
 
 		// Free the resources of the streaming wave data. This blocks pending IO requests
 		WaveData->FreeResources();
@@ -757,7 +749,7 @@ void FLegacyAudioStreamingManager::RemoveStreamingSoundWave(USoundWave* SoundWav
 		}
 		delete WaveData;
 	}
-	WaveRequests.Remove(SoundWave);
+	WaveRequests.Remove(SoundWave.GetFObjectKey());
 }
 
 void FLegacyAudioStreamingManager::AddDecoder(ICompressedAudioInfo* InCompressedAudioInfo)
@@ -772,16 +764,16 @@ void FLegacyAudioStreamingManager::RemoveDecoder(ICompressedAudioInfo* InCompres
 	CompressedAudioInfos.Remove(InCompressedAudioInfo);
 }
 
-bool FLegacyAudioStreamingManager::IsManagedStreamingSoundWave(const USoundWave* SoundWave) const
+bool FLegacyAudioStreamingManager::IsManagedStreamingSoundWave(const FSoundWaveProxy&  SoundWave) const
 {
 	FScopeLock Lock(&CriticalSection);
-	return StreamingSoundWaves.FindRef(SoundWave) != NULL;
+	return StreamingSoundWaves.FindRef(SoundWave.GetFObjectKey()) != NULL;
 }
 
-bool FLegacyAudioStreamingManager::IsStreamingInProgress(const USoundWave* SoundWave)
+bool FLegacyAudioStreamingManager::IsStreamingInProgress(const FSoundWaveProxy&  SoundWave)
 {
 	FScopeLock Lock(&CriticalSection);
-	FStreamingWaveData* WaveData = StreamingSoundWaves.FindRef(SoundWave);
+	FStreamingWaveData* WaveData = StreamingSoundWaves.FindRef(SoundWave.GetFObjectKey());
 	if (WaveData)
 	{
 		return WaveData->UpdateStreamingStatus();
@@ -879,7 +871,7 @@ void FLegacyAudioStreamingManager::RemoveStreamingSoundSource(FSoundSource* Soun
 
 		// Make sure there is a request so that unused chunks
 		// can be cleared if this was the last playing instance
-		GetWaveRequest(WaveInstance->WaveData);
+		GetWaveRequest(FObjectKey(WaveInstance->WaveData));
 		StreamingSoundSources.Remove(SoundSource);
 	}
 }
@@ -890,13 +882,13 @@ bool FLegacyAudioStreamingManager::IsManagedStreamingSoundSource(const FSoundSou
 	return StreamingSoundSources.FindByKey(SoundSource) != NULL;
 }
 
-bool FLegacyAudioStreamingManager::RequestChunk(USoundWave* SoundWave, uint32 ChunkIndex, TFunction<void(EAudioChunkLoadResult)> OnLoadCompleted, ENamedThreads::Type ThreadToCallOnLoadCompletedOn, bool bForImmediatePlayback)
+bool FLegacyAudioStreamingManager::RequestChunk(const FSoundWaveProxy&, uint32 ChunkIndex, TFunction<void(EAudioChunkLoadResult)> OnLoadCompleted, ENamedThreads::Type ThreadToCallOnLoadCompletedOn, bool bForImmediatePlayback)
 {
 	UE_LOG(LogAudio, Warning, TEXT("RequestChunk is only supported in Stream Caching."));
 	return false;
 }
 
-FAudioChunkHandle FLegacyAudioStreamingManager::GetLoadedChunk(const USoundWave* SoundWave, uint32 ChunkIndex, bool bBlockForLoad, bool bForImmediatePlayback) const
+FAudioChunkHandle FLegacyAudioStreamingManager::GetLoadedChunk(const FSoundWaveProxy&  SoundWave, uint32 ChunkIndex, bool bBlockForLoad, bool bForImmediatePlayback) const
 {
 	// Check for the spoof of failing to load a stream chunk
 	if (SpoofFailedStreamChunkLoad > 0)
@@ -910,7 +902,7 @@ FAudioChunkHandle FLegacyAudioStreamingManager::GetLoadedChunk(const USoundWave*
 		return FAudioChunkHandle();
 	}
 
-	const FStreamingWaveData* WaveData = StreamingSoundWaves.FindRef(SoundWave);
+	const FStreamingWaveData* WaveData = StreamingSoundWaves.FindRef(SoundWave.GetFObjectKey());
 	if (WaveData)
 	{
 		if (WaveData->LoadedChunkIndices.Contains(ChunkIndex))
@@ -920,7 +912,7 @@ FAudioChunkHandle FLegacyAudioStreamingManager::GetLoadedChunk(const USoundWave*
 				if (WaveData->LoadedChunks[Index].Index == ChunkIndex)
 				{
 					CriticalSection.Unlock();
-					return BuildChunkHandle(WaveData->LoadedChunks[Index].Data, WaveData->LoadedChunks[Index].AudioDataSize, SoundWave, SoundWave->GetFName(), ChunkIndex, InvalidAudioStreamCacheLookupID);
+					return BuildChunkHandle(WaveData->LoadedChunks[Index].Data, WaveData->LoadedChunks[Index].AudioDataSize, SoundWave, SoundWave.GetFName(), ChunkIndex, InvalidAudioStreamCacheLookupID);
 				}
 			}
 		}
@@ -930,13 +922,13 @@ FAudioChunkHandle FLegacyAudioStreamingManager::GetLoadedChunk(const USoundWave*
 	return FAudioChunkHandle();
 }
 
-FWaveRequest& FLegacyAudioStreamingManager::GetWaveRequest(USoundWave* SoundWave)
+FWaveRequest& FLegacyAudioStreamingManager::GetWaveRequest(FObjectKey Key)
 {
-	FWaveRequest* WaveRequest = WaveRequests.Find(SoundWave);
+	FWaveRequest* WaveRequest = WaveRequests.Find(Key);
 	if (!WaveRequest)
 	{
 		// Setup the new request so it always asks for chunk 0
-		WaveRequest = &WaveRequests.Add(SoundWave);
+		WaveRequest = &WaveRequests.Add(Key);
 		WaveRequest->RequiredIndices.AddUnique(0);
 		WaveRequest->bPrioritiseRequest = false;
 	}
diff --git a/Engine/Source/Runtime/Engine/Private/AudioStreamingCache.cpp b/Engine/Source/Runtime/Engine/Private/AudioStreamingCache.cpp
--- a/Engine/Source/Runtime/Engine/Private/AudioStreamingCache.cpp
+++ b/Engine/Source/Runtime/Engine/Private/AudioStreamingCache.cpp
@@ -236,6 +236,63 @@ static FAutoConsoleCommand GDisableProfilingAudioCacheCommand(
 );
 
 
+FAudioChunkCache::FChunkKey::FChunkKey(const FChunkKey& Other)
+	: SoundWaveProxyPtr(Other.SoundWaveProxyPtr.IsValid() ? MakeUnique<FSoundWaveProxy>(*Other.SoundWaveProxyPtr) : nullptr)
+	, SoundWaveName(SoundWaveProxyPtr->GetFName())
+	, ChunkIndex(Other.ChunkIndex)
+	, ObjectKey(Other.ObjectKey)
+#if WITH_EDITOR
+	, ChunkRevision(Other.ChunkRevision)
+#endif // #if WITH_EDITOR
+{
+	SoundWaveProxyPtr->ReleaseCompressedAudio();
+}
+
+
+FAudioChunkCache::FChunkKey& FAudioChunkCache::FChunkKey::operator=(const FChunkKey& Other)
+{
+	SoundWaveProxyPtr.Reset();
+
+	if (Other.SoundWaveProxyPtr.IsValid())
+	{
+		SoundWaveProxyPtr = MakeUnique<FSoundWaveProxy>(*Other.SoundWaveProxyPtr);
+	}
+
+	SoundWaveName = Other.SoundWaveName;
+	ChunkIndex = Other.ChunkIndex;
+	ObjectKey = Other.ObjectKey;
+
+#if WITH_EDITOR
+	ChunkRevision = Other.ChunkRevision;
+#endif // #if WITH_EDITOR
+
+
+	return *this;
+}
+
+FAudioChunkCache::FChunkKey::FChunkKey(const FSoundWaveProxy& InSoundWave, uint32 InChunkIndex
+#if WITH_EDITOR
+	, uint32 InChunkRevision
+#endif // #if WITH_EDITOR
+)
+
+	: SoundWaveProxyPtr(MakeUnique<FSoundWaveProxy>(InSoundWave))
+	, SoundWaveName(InSoundWave.GetFName())
+	, ChunkIndex(InChunkIndex)
+	, ObjectKey(InSoundWave.GetFObjectKey())
+#if WITH_EDITOR
+	, ChunkRevision(InChunkRevision)
+#endif // #if WITH_EDITOR
+{
+	SoundWaveProxyPtr->ReleaseCompressedAudio();
+}
+
+FAudioChunkCache::FChunkKey::FChunkKey(FName InSoundWaveName, uint32 InChunkIndex)
+	: SoundWaveName(InSoundWaveName)
+	, ChunkIndex(InChunkIndex)
+{
+}
+
 bool FAudioChunkCache::FChunkKey::operator==(const FChunkKey& Other) const
 {
 	if (UseObjectKeyInChunkKeyComparisonsCVar != 0)
@@ -257,6 +314,59 @@ bool FAudioChunkCache::FChunkKey::operator==(const FChunkKey& Other) const
 
 }
 
+bool FAudioChunkCache::FChunkKey::IsChunkStale()
+{
+	return !SoundWaveProxyPtr.IsValid()
+#if WITH_EDITOR
+		|| (SoundWaveProxyPtr->GetCurrentChunkRevision() != ChunkRevision)
+#endif // #if WITH_EDITOR
+		;
+}
+
+ESoundWaveLoadingBehavior FAudioChunkCache::FChunkKey::GetLoadingBehavior() const
+{
+	check(SoundWaveProxyPtr.IsValid());
+	return SoundWaveProxyPtr->GetLoadingBehavior();
+}
+
+FStreamedAudioChunk& FAudioChunkCache::FChunkKey::GetChunk(uint32 InChunkIndex) const
+{
+	check(SoundWaveProxyPtr.IsValid());
+	check((ChunkIndex < (uint32)SoundWaveProxyPtr->GetNumChunks()));
+
+	check(SoundWaveProxyPtr->RunningPlatformData);
+
+	// This function shouldn't be called on audio marked "ForceInline."
+	ensureMsgf(SoundWaveProxyPtr->RunningPlatformData, TEXT("Calling GetNumChunks on a FSoundWaveProxy without RunnigPlatformData is not allowed! SoundWave: %s - %s")
+		, *SoundWaveProxyPtr->GetFName().ToString(), EnumToString(SoundWaveProxyPtr->GetLoadingBehavior()));
+
+	return SoundWaveProxyPtr->RunningPlatformData->Chunks[ChunkIndex];
+}
+
+uint32 FAudioChunkCache::FChunkKey::GetNumChunks() const
+{
+	check(SoundWaveProxyPtr.IsValid());
+	return SoundWaveProxyPtr->GetNumChunks();
+}
+
+bool FAudioChunkCache::FChunkKey::IsRetainingAudio() const
+{
+	check(SoundWaveProxyPtr.IsValid());
+	return SoundWaveProxyPtr->IsRetainingAudio();
+}
+
+void FAudioChunkCache::FChunkKey::ReleaseCompressedAudio()
+{
+	check(SoundWaveProxyPtr.IsValid());
+	return SoundWaveProxyPtr->ReleaseCompressedAudio();
+}
+
+bool FAudioChunkCache::FChunkKey::WasLoadingBehaviorOverridden() const
+{
+	check(SoundWaveProxyPtr.IsValid());
+	return SoundWaveProxyPtr->WasLoadingBehaviorOverridden();
+}
+
 FCachedAudioStreamingManager::FCachedAudioStreamingManager(const FCachedAudioStreamingManagerParams& InitParams)
 {
 	LLM_SCOPE(ELLMTag::AudioStreamCache);
@@ -325,12 +435,12 @@ void FCachedAudioStreamingManager::NotifyLevelOffset(class ULevel* Level, const
 	// Unused.
 }
 
-void FCachedAudioStreamingManager::AddStreamingSoundWave(USoundWave* SoundWave)
+void FCachedAudioStreamingManager::AddStreamingSoundWave(const FSoundWaveProxy& SoundWave)
 {
 	// Unused.
 }
 
-void FCachedAudioStreamingManager::RemoveStreamingSoundWave(USoundWave* SoundWave)
+void FCachedAudioStreamingManager::RemoveStreamingSoundWave(const FSoundWaveProxy& SoundWave)
 {
 	// Unused.
 }
@@ -345,14 +455,14 @@ void FCachedAudioStreamingManager::RemoveDecoder(ICompressedAudioInfo* InCompres
 	//Unused.
 }
 
-bool FCachedAudioStreamingManager::IsManagedStreamingSoundWave(const USoundWave* SoundWave) const
+bool FCachedAudioStreamingManager::IsManagedStreamingSoundWave(const FSoundWaveProxy&  SoundWave) const
 {
 	// Unused. The concept of a sound wave being "managed" doesn't apply here.
 	checkf(false, TEXT("Not Implemented!"));
 	return true;
 }
 
-bool FCachedAudioStreamingManager::IsStreamingInProgress(const USoundWave* SoundWave)
+bool FCachedAudioStreamingManager::IsStreamingInProgress(const FSoundWaveProxy&  SoundWave)
 {
 	// This function is used in USoundWave cleanup.
 	// Since this manager owns the binary data we are streaming off of,
@@ -389,7 +499,7 @@ bool FCachedAudioStreamingManager::IsManagedStreamingSoundSource(const FSoundSou
 	return true;
 }
 
-FAudioChunkHandle FCachedAudioStreamingManager::GetLoadedChunk(const USoundWave* SoundWave, uint32 ChunkIndex, bool bBlockForLoad, bool bForImmediatePlayback) const
+FAudioChunkHandle FCachedAudioStreamingManager::GetLoadedChunk(const FSoundWaveProxy& SoundWave, uint32 ChunkIndex, bool bBlockForLoad, bool bForImmediatePlayback) const
 {
 	LLM_SCOPE(ELLMTag::AudioStreamCache);
 	bBlockForLoad |= (ForceBlockForLoadCVar != 0);
@@ -401,151 +511,73 @@ FAudioChunkHandle FCachedAudioStreamingManager::GetLoadedChunk(const USoundWave*
 		// With this code, the zeroth chunk should never get hit.
 		checkf(ChunkIndex != 0, TEXT("Decoder tried to access the zeroth chunk through the streaming manager. Use USoundWave::GetZerothChunk() instead."));
 
-		// TODO:  See if we can avoid non-const calls on the USoundWave here.
-		USoundWave* MutableWave = const_cast<USoundWave*>(SoundWave);
-		const FAudioChunkCache::FChunkKey ChunkKey =
-		{
-			  MutableWave
-			, SoundWave->GetFName()
+		FAudioChunkCache::FChunkKey ChunkKey(
+			  SoundWave
 			, ChunkIndex
-			, FObjectKey(MutableWave)
 #if WITH_EDITOR
-			, (uint32)SoundWave->CurrentChunkRevision.GetValue()
+			, (uint32)SoundWave.GetCurrentChunkRevision()
 #endif
-		};
+		);
 
 		if (!FAudioChunkCache::IsKeyValid(ChunkKey))
 		{
-			UE_LOG(LogAudioStreamCaching, Warning, TEXT("Invalid Chunk Index %d Requested for Wave %s!"), ChunkIndex, *SoundWave->GetName());
+			UE_LOG(LogAudioStreamCaching, Warning, TEXT("Invalid Chunk Index %d Requested for Wave %s!"), ChunkIndex, *SoundWave.GetFName().ToString());
 			return FAudioChunkHandle();
 		}
 
 		// The function call below increments the reference count to the internal chunk.
-		uint64 LookupIDForChunk = ChunkKey.SoundWave->GetCacheLookupIDForChunk(ChunkKey.ChunkIndex);
+		uint64 LookupIDForChunk = Cache->GetCacheLookupIDForChunk(ChunkKey);
 		TArrayView<uint8> LoadedChunk = Cache->GetChunk(ChunkKey, bBlockForLoad, (bForImmediatePlayback || bBlockForLoad), LookupIDForChunk);
-		
+
 		// Ensure that, if we requested a synchronous load of this chunk, we didn't fail to load said chunk.
-		UE_CLOG(bBlockForLoad && !LoadedChunk.GetData(), LogAudioStreamCaching, Display, TEXT("Synchronous load of chunk index %d for SoundWave %s failed to return any data. Likely because the cache was blown."), ChunkIndex, *SoundWave->GetName());
+		UE_CLOG(bBlockForLoad && !LoadedChunk.GetData(), LogAudioStreamCaching, Display, TEXT("Synchronous load of chunk index %d for SoundWave %s failed to return any data. Likely because the cache was blown."), ChunkIndex, *SoundWave.GetFName().ToString());
 
 		// Set the updated cache offset for this chunk index.
-		ChunkKey.SoundWave->SetCacheLookupIDForChunk(ChunkIndex, LookupIDForChunk);
+		Cache->SetCacheLookupIDForChunk(ChunkKey, LookupIDForChunk);
 
-		UE_CLOG(!bBlockForLoad && !LoadedChunk.GetData(), LogAudioStreamCaching, Verbose, TEXT("GetLoadedChunk called for chunk index %d of SoundWave %s when audio was not loaded yet. This will result in latency."), ChunkIndex, *SoundWave->GetName());
+		UE_CLOG(!bBlockForLoad && !LoadedChunk.GetData(), LogAudioStreamCaching, Verbose, TEXT("GetLoadedChunk called for chunk index %d of SoundWave %s when audio was not loaded yet. This will result in latency."), ChunkIndex, *SoundWave.GetFName().ToString());
 
 		// Finally, if there's a chunk after this in the sound, request that it is in the cache.
 		const int32 NextChunk = GetNextChunkIndex(SoundWave, ChunkIndex);
 
-		if (NextChunk != INDEX_NONE) 
+		if (NextChunk != INDEX_NONE)
 		{
-			const FAudioChunkCache::FChunkKey NextChunkKey = 
-			{ 
-				  MutableWave 
-				, SoundWave->GetFName() 
-				, ((uint32)NextChunk) 
-				, FObjectKey(MutableWave)
+			FAudioChunkCache::FChunkKey NextChunkKey(
+				SoundWave
+				, ((uint32)NextChunk)
 #if WITH_EDITOR
-				, (uint32)SoundWave->CurrentChunkRevision.GetValue()
+				, (uint32)SoundWave.GetCurrentChunkRevision()
 #endif
-			};
+			);
 
 			uint64 LookupIDForNextChunk = Cache->AddOrTouchChunk(NextChunkKey, [](EAudioChunkLoadResult) {}, ENamedThreads::AnyThread, false);
 			if (LookupIDForNextChunk == InvalidAudioStreamCacheLookupID)
 			{
-				// this bool is true while we are waiting on the game thread to reset chunk handles owned by USoundWaves
-				static FThreadSafeBool bCacheCurrentlyBlown = false;
-
-				if (!bCacheCurrentlyBlown)
-				{
-					bCacheCurrentlyBlown = true;
-					Cache->IncrementCacheOverflowCounter();
-
-					UE_LOG(LogAudioStreamCaching, Warning, TEXT("Cache overflow!!! couldn't load chunk %d for sound %s!"), ChunkIndex, *SoundWave->GetName());
-
-					// gather SoundWaves to release compressed data on:
-					TArray<FObjectKey> SoundWavesToRelease;
-
-					if (NumSoundWavesToClearOnCacheOverflowCVar > 0)
-					{
-						SoundWavesToRelease = Cache->GetLeastRecentlyUsedRetainedSoundWaves(NumSoundWavesToClearOnCacheOverflowCVar);
-					}
-
-					AsyncTask(ENamedThreads::GameThread, [WavesToRelease = MoveTemp(SoundWavesToRelease)]() mutable
-					{
-						// Here we optionally capture the state of the cache when we overflowed:
-#if !UE_BUILD_SHIPPING && !UE_BUILD_TEST
-						if (SaveAudioMemReportOnCacheOverflowCVar && GEngine && GEngine->GetMainAudioDevice())
-						{
-							GEngine->GetMainAudioDevice()->Exec(nullptr, TEXT("audiomemreport"));
-						}
-#endif // !UE_BUILD_SHIPPING && !UE_BUILD_TEST
-
-
-						int32 NumChunksReleased = 0;
-
-						for (TObjectIterator<USoundWave> It; It; ++It)
-						{
-							USoundWave* Wave = *It;
-							if (Wave && Wave->IsRetainingAudio())
-							{
-								// If we have a specific list of sound waves to release, check if this sound wave is in it.
-								if (WavesToRelease.Num())
-								{
-									for (int32 Index = 0; Index < WavesToRelease.Num(); Index++)
-									{
-										if (WavesToRelease[Index] == FObjectKey(Wave))
-										{
-											Wave->ReleaseCompressedAudio();
-											WavesToRelease.RemoveAtSwap(Index);
-											NumChunksReleased++;
-											break;
-										}
-									}
-
-									// If we've found every wave we're going to release, break out of the soundwave iterator.
-									if (WavesToRelease.Num() == 0)
-									{
-										break;
-									}
-								}
-								else
-								{
-									// Otherwise, we release all compressed audio by default.
-									Wave->ReleaseCompressedAudio();
-									NumChunksReleased++;
-								}
-							}
-						}
-
-						UE_LOG(LogAudioStreamCaching, Warning, TEXT("Removed %d retained sounds from the stream cache."), NumChunksReleased);
-
-						bCacheCurrentlyBlown = false;
-					});
-				}
+				UE_LOG(LogAudioStreamCaching, Warning, TEXT("Failed to add follow-up chunk for %s. This should not happen with our current TrimMemory() logic"), *SoundWave.GetFName().ToString());
 			}
 			else
 			{
-				NextChunkKey.SoundWave->SetCacheLookupIDForChunk(NextChunkKey.ChunkIndex, LookupIDForNextChunk);
+				Cache->SetCacheLookupIDForChunk(NextChunkKey, LookupIDForNextChunk);
 			}
 		}
 
-		return BuildChunkHandle(LoadedChunk.GetData(), LoadedChunk.Num(), SoundWave, SoundWave->GetFName(), ChunkIndex, LookupIDForChunk);
+		return BuildChunkHandle(LoadedChunk.GetData(), LoadedChunk.Num(), SoundWave, SoundWave.GetFName(), ChunkIndex, LookupIDForChunk);
 	}
 	else
 	{
-		ensureMsgf(false, TEXT("Failed to find cache for wave %s. Are you sure this is a streaming wave?"), *SoundWave->GetName());
+		ensureMsgf(false, TEXT("Failed to find cache for wave %s. Are you sure this is a streaming wave?"), *SoundWave.GetFName().ToString());
 		return FAudioChunkHandle();
 	}
 }
 
-FAudioChunkCache* FCachedAudioStreamingManager::GetCacheForWave(const USoundWave* InSoundWave) const
+FAudioChunkCache* FCachedAudioStreamingManager::GetCacheForWave(const FSoundWaveProxy&  InSoundWave) const
 {
 	LLM_SCOPE(ELLMTag::AudioStreamCache);
-	check(InSoundWave);
 
 	// We only cache chunks beyond the zeroth chunk of audio (which is inlined directly on the asset)
-	if (InSoundWave->RunningPlatformData && InSoundWave->RunningPlatformData->Chunks.Num() > 1)
+	if (InSoundWave.GetNumChunks() > 1)
 	{
-		const int32 SoundWaveChunkSize = InSoundWave->RunningPlatformData->Chunks[1].AudioDataSize;
+		const int32 SoundWaveChunkSize = InSoundWave.GetChunk(1).AudioDataSize;
 		return GetCacheForChunkSize(SoundWaveChunkSize);
 	}
 	else
@@ -561,7 +593,7 @@ FAudioChunkCache* FCachedAudioStreamingManager::GetCacheForChunkSize(uint32 InCh
 	for (int32 CacheIndex = 0; CacheIndex < CacheArray.Num(); CacheIndex++)
 	{
 		check(CacheArray[CacheIndex].MaxChunkSize >= 0);
-		if (InChunkSize <= ((uint32) CacheArray[CacheIndex].MaxChunkSize))
+		if (InChunkSize <= ((uint32)CacheArray[CacheIndex].MaxChunkSize))
 		{
 			return const_cast<FAudioChunkCache*>(&CacheArray[CacheIndex]);
 		}
@@ -574,22 +606,22 @@ FAudioChunkCache* FCachedAudioStreamingManager::GetCacheForChunkSize(uint32 InCh
 	return nullptr;
 }
 
-int32 FCachedAudioStreamingManager::GetNextChunkIndex(const USoundWave* InSoundWave, uint32 CurrentChunkIndex) const
+int32 FCachedAudioStreamingManager::GetNextChunkIndex(const FSoundWaveProxy&  InSoundWave, uint32 CurrentChunkIndex) const
 {
 	LLM_SCOPE(ELLMTag::AudioStreamCache);
-	check(InSoundWave);
+
 	// TODO: Figure out a way to tell whether this wave is looping or not. For now we always prime the first chunk
 	// during the playback of the last chunk.
-	// if(bNotLooping) return ((int32) CurrentChunkIndex) < (InSoundWave->RunningPlatformData->Chunks.Num() - 1);
-	
-	const int32 NumChunksTotal = InSoundWave->GetNumChunks();
+	// if(bNotLooping) return ((int32) CurrentChunkIndex) < (InSoundWave->GetNumChunks() - 1);
+
+	const int32 NumChunksTotal = InSoundWave.GetNumChunks();
 	if (NumChunksTotal <= 2)
 	{
 		// If there's only one chunk to cache (besides the zeroth chunk, which is inlined),
 		// We don't need to load anything.
 		return INDEX_NONE;
 	}
-	else if(CurrentChunkIndex == (NumChunksTotal - 1))
+	else if (CurrentChunkIndex == (NumChunksTotal - 1))
 	{
 		// if we're on the last chunk, load the first chunk after the zeroth chunk.
 		return 1;
@@ -607,19 +639,15 @@ void FCachedAudioStreamingManager::AddReferenceToChunk(const FAudioChunkHandle&
 	FAudioChunkCache* Cache = GetCacheForChunkSize(InHandle.CachedDataNumBytes);
 	check(Cache);
 
-	USoundWave* MutableWave = const_cast<USoundWave*>(InHandle.CorrespondingWave);
-	FAudioChunkCache::FChunkKey ChunkKey =
-	{
-		  MutableWave
-		, InHandle.CorrespondingWaveName
-		, ((uint32) InHandle.ChunkIndex)
-		, FObjectKey(MutableWave)
+	const FAudioChunkCache::FChunkKey ChunkKey(
+		  *InHandle.CorrespondingWave
+		, ((uint32)InHandle.ChunkIndex)
 #if WITH_EDITOR
 		, InHandle.ChunkGeneration
 #endif
-	};
+	);
 
-	Cache->AddNewReferenceToChunk(ChunkKey, InHandle.CacheLookupID);
+	Cache->AddNewReferenceToChunk(ChunkKey);
 }
 
 void FCachedAudioStreamingManager::RemoveReferenceToChunk(const FAudioChunkHandle& InHandle)
@@ -628,47 +656,42 @@ void FCachedAudioStreamingManager::RemoveReferenceToChunk(const FAudioChunkHandl
 	FAudioChunkCache* Cache = GetCacheForChunkSize(InHandle.CachedDataNumBytes);
 	check(Cache);
 
-	USoundWave* MutableWave = const_cast<USoundWave*>(InHandle.CorrespondingWave);
-	FAudioChunkCache::FChunkKey ChunkKey =
-	{
-		  MutableWave
-		, InHandle.CorrespondingWaveName
+	const FAudioChunkCache::FChunkKey ChunkKey(
+		  *InHandle.CorrespondingWave
 		, ((uint32)InHandle.ChunkIndex)
-		, FObjectKey(MutableWave)
 #if WITH_EDITOR
 		, InHandle.ChunkGeneration
 #endif
-	};
+	);
 
-	Cache->RemoveReferenceToChunk(ChunkKey, InHandle.CacheLookupID);
+	Cache->RemoveReferenceToChunk(ChunkKey);
 }
 
-bool FCachedAudioStreamingManager::RequestChunk(USoundWave* SoundWave, uint32 ChunkIndex, TFunction<void(EAudioChunkLoadResult)> OnLoadCompleted, ENamedThreads::Type ThreadToCallOnLoadCompletedOn, bool bForImmediatePlayback)
+bool FCachedAudioStreamingManager::RequestChunk(const FSoundWaveProxy& SoundWave, uint32 ChunkIndex, TFunction<void(EAudioChunkLoadResult)> OnLoadCompleted, ENamedThreads::Type ThreadToCallOnLoadCompletedOn, bool bForImmediatePlayback)
 {
 	LLM_SCOPE(ELLMTag::AudioStreamCache);
 	FAudioChunkCache* Cache = GetCacheForWave(SoundWave);
 	if (Cache)
 	{
-		FAudioChunkCache::FChunkKey ChunkKey = 
-		{
+		FAudioChunkCache::FChunkKey ChunkKey(
 			SoundWave
-		  , SoundWave->GetFName()
-		  , ChunkIndex
-		  , FObjectKey(SoundWave)
+			, ChunkIndex
 #if WITH_EDITOR
-		  , (uint32)SoundWave->CurrentChunkRevision.GetValue()
+			, (uint32)SoundWave.GetCurrentChunkRevision()
 #endif
-		};
+		);
 
 		uint64 LookupIDForChunk = Cache->AddOrTouchChunk(ChunkKey, OnLoadCompleted, ThreadToCallOnLoadCompletedOn, bForImmediatePlayback);
-		SoundWave->SetCacheLookupIDForChunk(ChunkIndex, LookupIDForChunk);
+
+		// Update the map entry through the streaming cache
+		Cache->SetCacheLookupIDForChunk(ChunkKey, LookupIDForChunk);
 		return LookupIDForChunk != InvalidAudioStreamCacheLookupID;
 	}
 	else
 	{
 		// This can hit if an out of bounds chunk was requested, or the zeroth chunk was requested from the streaming manager.
-		ensureMsgf(false, TEXT("GetCacheForWave failed for SoundWave %s!"), *SoundWave->GetName());
-		SoundWave->SetCacheLookupIDForChunk(ChunkIndex, InvalidAudioStreamCacheLookupID);
+		ensureMsgf(false, TEXT("GetCacheForWave failed for SoundWave %s!"), *SoundWave.GetFName().ToString());
+
 		return false;
 	}
 }
@@ -705,7 +728,7 @@ uint64 FAudioChunkCache::AddOrTouchChunk(const FChunkKey& InKey, TFunction<void(
 	{
 		MemoryLimitBytes = ((uint64)(StreamCacheSizeOverrideMBCVar * 1024)) * 1024;
 	}
-	
+
 	if (!IsKeyValid(InKey))
 	{
 		ensure(false);
@@ -715,9 +738,9 @@ uint64 FAudioChunkCache::AddOrTouchChunk(const FChunkKey& InKey, TFunction<void(
 
 	FScopeLock ScopeLock(&CacheMutationCriticalSection);
 
-	const uint64 LookupIDForChunk = InKey.SoundWave->GetCacheLookupIDForChunk(InKey.ChunkIndex);
-	FCacheElement* FoundElement = FindElementForKey(InKey, LookupIDForChunk);
-	
+	const uint64 LookupIDForChunk = GetCacheLookupIDForChunk(InKey);
+	FCacheElement* FoundElement = FindElementForKey(InKey);
+
 	if (FoundElement)
 	{
 		TouchElement(FoundElement);
@@ -730,10 +753,9 @@ uint64 FAudioChunkCache::AddOrTouchChunk(const FChunkKey& InKey, TFunction<void(
 		FoundElement->DebugInfo.NumTimesTouched++;
 
 		// Recursing in no longer needed at this point since the inherited loading behavior has already been cached by the time this information is needed
-		const bool bRecurseSoundClasses = false;
-		FoundElement->DebugInfo.LoadingBehavior = InKey.SoundWave->GetLoadingBehavior(bRecurseSoundClasses);
+		FoundElement->DebugInfo.LoadingBehavior = InKey.GetLoadingBehavior();
 
-		FoundElement->DebugInfo.bLoadingBehaviorExternallyOverriden = InKey.SoundWave->bLoadingBehaviorOverridden;
+		FoundElement->DebugInfo.bLoadingBehaviorExternallyOverriden = InKey.WasLoadingBehaviorOverridden();
 #endif
 
 		return FoundElement->CacheLookupID;
@@ -751,20 +773,18 @@ uint64 FAudioChunkCache::AddOrTouchChunk(const FChunkKey& InKey, TFunction<void(
 #if DEBUG_STREAM_CACHE
 		CacheElement->DebugInfo.bWasCacheMiss = bNeededForPlayback;
 
-		// Recursing in no longer needed at this point since the inherited loading behavior has already been cached by the time this information is needed
-		const bool bRecurseSoundClasses = false;
-		CacheElement->DebugInfo.LoadingBehavior = InKey.SoundWave->GetLoadingBehavior(bRecurseSoundClasses);
+		CacheElement->DebugInfo.LoadingBehavior = InKey.GetLoadingBehavior();
 
-		CacheElement->DebugInfo.bLoadingBehaviorExternallyOverriden = InKey.SoundWave->bLoadingBehaviorOverridden;
+		CacheElement->DebugInfo.bLoadingBehaviorExternallyOverriden = InKey.WasLoadingBehaviorOverridden();
 #endif
-		const FStreamedAudioChunk& Chunk = InKey.SoundWave->RunningPlatformData->Chunks[InKey.ChunkIndex];
+		const FStreamedAudioChunk& Chunk = InKey.GetChunk(InKey.ChunkIndex);
 		int32 ChunkDataSize = Chunk.AudioDataSize;
 
 		if (TrimCacheWhenOverBudgetCVar != 0 && (MemoryCounterBytes + ChunkDataSize) > MemoryLimitBytes)
 		{
 			uint64 MemoryToTrim = 0;
 			if (MemoryLimitTrimPercentageCVar > 0.0f)
-			{	
+			{
 				MemoryToTrim = MemoryLimitBytes * FMath::Min(MemoryLimitTrimPercentageCVar, 1.0f);
 			}
 			else
@@ -780,7 +800,7 @@ uint64 FAudioChunkCache::AddOrTouchChunk(const FChunkKey& InKey, TFunction<void(
 		if (bNeededForPlayback && (bLogCacheMisses || AlwaysLogCacheMissesCVar))
 		{
 			// We missed 
-			const uint32 TotalNumChunksInWave = InKey.SoundWave->GetNumChunks();
+			const uint32 TotalNumChunksInWave = InKey.GetNumChunks();
 
 			FCacheMissInfo CacheMissInfo = { InKey.SoundWaveName, InKey.ChunkIndex, TotalNumChunksInWave, false };
 			CacheMissQueue.Enqueue(MoveTemp(CacheMissInfo));
@@ -793,7 +813,7 @@ uint64 FAudioChunkCache::AddOrTouchChunk(const FChunkKey& InKey, TFunction<void(
 TArrayView<uint8> FAudioChunkCache::GetChunk(const FChunkKey& InKey, bool bBlockForLoadCompletion, bool bNeededForPlayback, uint64& OutCacheOffset)
 {
 	FScopeLock ScopeLock(&CacheMutationCriticalSection);
-	FCacheElement* FoundElement = FindElementForKey(InKey, OutCacheOffset);
+	FCacheElement* FoundElement = FindElementForKey(InKey);
 	if (FoundElement)
 	{
 		OutCacheOffset = FoundElement->CacheLookupID;
@@ -826,15 +846,15 @@ TArrayView<uint8> FAudioChunkCache::GetChunk(const FChunkKey& InKey, bool bBlock
 			UE_LOG(LogAudioStreamCaching, Display, TEXT("GetChunk failed to find an available chunk slot in the cache, likely because the cache is blown."));
 			return TArrayView<uint8>();
 		}
-		
+
 		OutCacheOffset = FoundElement->CacheLookupID;
 
 		if (bBlockForLoadCompletion)
 		{
-			FStreamedAudioChunk& Chunk = InKey.SoundWave->RunningPlatformData->Chunks[InKey.ChunkIndex];
+			FStreamedAudioChunk& Chunk = InKey.GetChunk(InKey.ChunkIndex);
 			int32 ChunkAudioDataSize = Chunk.AudioDataSize;
 #if DEBUG_STREAM_CACHE
-			FoundElement->DebugInfo.NumTotalChunks = InKey.SoundWave->GetNumChunks() - 1;
+			FoundElement->DebugInfo.NumTotalChunks = InKey.GetNumChunks() - 1;
 			FoundElement->DebugInfo.TimeLoadStarted = FPlatformTime::Seconds();
 #endif
 			MemoryCounterBytes -= FoundElement->ChunkDataSize;
@@ -888,7 +908,7 @@ TArrayView<uint8> FAudioChunkCache::GetChunk(const FChunkKey& InKey, bool bBlock
 		if (bLogCacheMisses && !bBlockForLoadCompletion)
 		{
 			// Chunks missing. Log this as a miss.
-			const uint32 TotalNumChunksInWave = InKey.SoundWave->GetNumChunks();
+			const uint32 TotalNumChunksInWave = InKey.GetNumChunks();
 			FCacheMissInfo CacheMissInfo = { InKey.SoundWaveName, InKey.ChunkIndex, TotalNumChunksInWave, false };
 			CacheMissQueue.Enqueue(MoveTemp(CacheMissInfo));
 		}
@@ -897,10 +917,10 @@ TArrayView<uint8> FAudioChunkCache::GetChunk(const FChunkKey& InKey, bool bBlock
 	}
 }
 
-void FAudioChunkCache::AddNewReferenceToChunk(const FChunkKey& InKey, uint64 ChunkOffset)
+void FAudioChunkCache::AddNewReferenceToChunk(const FChunkKey& InKey)
 {
 	FScopeLock ScopeLock(&CacheMutationCriticalSection);
-	FCacheElement* FoundElement = FindElementForKey(InKey, ChunkOffset);
+	FCacheElement* FoundElement = FindElementForKey(InKey);
 	if (ensure(FoundElement))
 	{
 		// If this value is ever negative, it means that we're decrementing more than we're incrementing:
@@ -909,10 +929,10 @@ void FAudioChunkCache::AddNewReferenceToChunk(const FChunkKey& InKey, uint64 Chu
 	}
 }
 
-void FAudioChunkCache::RemoveReferenceToChunk(const FChunkKey& InKey, uint64 ChunkOffset)
+void FAudioChunkCache::RemoveReferenceToChunk(const FChunkKey& InKey)
 {
 	FScopeLock ScopeLock(&CacheMutationCriticalSection);
-	FCacheElement* FoundElement = FindElementForKey(InKey, ChunkOffset);
+	FCacheElement* FoundElement = FindElementForKey(InKey);
 	if (ensure(FoundElement))
 	{
 		// If this value is ever less than 1 when we hit this code, it means that we're decrementing more than we're incrementing:
@@ -994,10 +1014,10 @@ uint64 FAudioChunkCache::TrimMemory(uint64 BytesToFree, bool bInAllowRetainedChu
 		ElementToStopAt = MostRecentElement->LessRecentElement;
 		while (CurrentElement != ElementToStopAt && BytesFreed < BytesToFree)
 		{
-			if (CurrentElement->ChunkDataSize != 0 && CurrentElement->Key.SoundWave && CurrentElement->Key.SoundWave->IsRetainingAudio())
+			if (CurrentElement->ChunkDataSize != 0 && CurrentElement->Key.IsSoundWaveValid() && CurrentElement->Key.IsRetainingAudio())
 			{
 				// Directly release the retained audio (TODO: this is on the audio thread right?)
-				CurrentElement->Key.SoundWave->ReleaseCompressedAudio();
+				CurrentElement->Key.ReleaseCompressedAudio();
 				if (CurrentElement->CanEvictChunk())
 				{
 					uint32 ChunkSize = CurrentElement->ChunkDataSize;
@@ -1013,7 +1033,7 @@ uint64 FAudioChunkCache::TrimMemory(uint64 BytesToFree, bool bInAllowRetainedChu
 					CurrentElement->ChunkData = nullptr;
 					CurrentElement->ChunkDataSize = 0;
 					CurrentElement->Key = FChunkKey();
-					
+
 #if DEBUG_STREAM_CACHE
 					// Reset debug info:
 					CurrentElement->DebugInfo.Reset();
@@ -1123,16 +1143,7 @@ FString FAudioChunkCache::FlushCacheMissLog()
 		ConcatenatedCacheMisses.AppendInt(CacheMissInfo.ChunkIndex);
 		ConcatenatedCacheMisses.Append(TEXT("\n"));
 
-		FChunkKey Chunk =
-		{
-			  nullptr
-			, CacheMissInfo.SoundWaveName
-			, CacheMissInfo.ChunkIndex
-			, FObjectKey()
-#if WITH_EDITOR
-			, 0
-#endif
-		};
+		FChunkKey Chunk(CacheMissInfo.SoundWaveName, CacheMissInfo.ChunkIndex);
 
 		int32& MissCount = CacheMissCount.FindOrAdd(Chunk);
 		MissCount++;
@@ -1184,19 +1195,27 @@ FString FAudioChunkCache::FlushCacheMissLog()
 	return TopChunkMissesLog + TEXT("\n") + ConcatenatedCacheMisses;
 }
 
-FAudioChunkCache::FCacheElement* FAudioChunkCache::FindElementForKey(const FChunkKey& InKey, uint64 CacheOffset)
+FAudioChunkCache::FCacheElement* FAudioChunkCache::FindElementForKey(const FChunkKey& InKey)
 {
 	FScopeLock ScopeLock(&CacheMutationCriticalSection);
 
-	// If we have a known cache offset, access that chunk directly.
-	if (CacheOffset != InvalidAudioStreamCacheLookupID)
+	if (CacheLookupIdMap.Contains(InKey))
 	{
-		check(CacheOffset < CachePool.Num());
+		const uint64 CacheOffset = CacheLookupIdMap[InKey];
 
-		// Finally, sanity check that the key is still the same.
-		if (CachePool[CacheOffset].Key == InKey)
+		// If we have a known cache offset, access that chunk directly.
+		if (CacheOffset != InvalidAudioStreamCacheLookupID)
 		{
-			return &CachePool[CacheOffset];
+			check(CacheOffset < CachePool.Num());
+
+			// Finally, sanity check that the key is still the same.
+			if (CachePool[CacheOffset].Key == InKey)
+			{
+				return &CachePool[CacheOffset];
+			}
+
+			UE_LOG(LogAudioStreamCaching, Display, TEXT("Cache Offset [%i] currently stores chunk for Soundwave: %s -- (looking for Soundwave %s)")
+				, CacheOffset, *CachePool[CacheOffset].Key.SoundWaveName.ToString(), *InKey.SoundWaveName.ToString());
 		}
 	}
 
@@ -1381,7 +1400,7 @@ FAudioChunkCache::FCacheElement* FAudioChunkCache::InsertChunk(const FChunkKey&
 		}
 	}
 
-	InKey.SoundWave->SetCacheLookupIDForChunk(InKey.ChunkIndex, CacheElement->CacheLookupID);
+	SetCacheLookupIDForChunk(InKey, CacheElement->CacheLookupID);
 	return CacheElement;
 }
 
@@ -1559,7 +1578,7 @@ void FAudioChunkCache::KickOffAsyncLoad(FCacheElement* CacheElement, const FChun
 {
 	check(CacheElement);
 
-	const FStreamedAudioChunk& Chunk = InKey.SoundWave->RunningPlatformData->Chunks[InKey.ChunkIndex];
+	const FStreamedAudioChunk& Chunk = InKey.GetChunk(InKey.ChunkIndex);
 	int32 ChunkDataSize = Chunk.AudioDataSize;
 
 	EAsyncIOPriorityAndFlags AsyncIOPriority = GetAsyncPriorityForChunk(InKey, bNeededForPlayback);
@@ -1570,16 +1589,16 @@ void FAudioChunkCache::KickOffAsyncLoad(FCacheElement* CacheElement, const FChun
 		LLM_SCOPE(ELLMTag::AudioStreamCacheCompressedData);
 
 		// Reallocate our chunk data This allows us to shrink if possible.
-		CacheElement->ChunkData = (uint8*) FMemory::Realloc(CacheElement->ChunkData, Chunk.AudioDataSize);
+		CacheElement->ChunkData = (uint8*)FMemory::Realloc(CacheElement->ChunkData, Chunk.AudioDataSize);
 		CacheElement->ChunkDataSize = Chunk.AudioDataSize;
 	}
 
 	MemoryCounterBytes += CacheElement->ChunkDataSize;
 
 #if DEBUG_STREAM_CACHE
-	CacheElement->DebugInfo.NumTotalChunks = InKey.SoundWave->GetNumChunks() - 1;
-	CacheElement->DebugInfo.LoadingBehavior = InKey.SoundWave->GetLoadingBehavior(false);
-	CacheElement->DebugInfo.bLoadingBehaviorExternallyOverriden = InKey.SoundWave->bLoadingBehaviorOverridden;
+	CacheElement->DebugInfo.NumTotalChunks = InKey.GetNumChunks() - 1;
+	CacheElement->DebugInfo.LoadingBehavior = InKey.GetLoadingBehavior();
+	CacheElement->DebugInfo.bLoadingBehaviorExternallyOverriden = InKey.WasLoadingBehaviorOverridden();
 #endif
 
 	// In editor, we retrieve from the DDC. In non-editor situations, we read the chunk async from the pak file.
@@ -1674,7 +1693,7 @@ void FAudioChunkCache::KickOffAsyncLoad(FCacheElement* CacheElement, const FChun
 #if DEBUG_STREAM_CACHE
 		CacheElement->DebugInfo.TimeLoadStarted = FPlatformTime::Seconds();
 #endif
-		
+
 		CacheElement->ReadRequest = nullptr;
 		IBulkDataIORequest* LocalReadRequest = Chunk.BulkData.CreateStreamingRequest(0, ChunkDataSize, AsyncIOPriority | AIOP_FLAG_DONTCACHE, &AsyncFileCallBack, CacheElement->ChunkData);
 		if (!LocalReadRequest)
@@ -1771,9 +1790,30 @@ void FAudioChunkCache::ExecuteOnLoadCompleteCallback(EAudioChunkLoadResult Resul
 
 bool FAudioChunkCache::IsKeyValid(const FChunkKey& InKey)
 {
-	return InKey.ChunkIndex < TNumericLimits<uint32>::Max() && ((int32)InKey.ChunkIndex) < InKey.SoundWave->RunningPlatformData->Chunks.Num();
+	return InKey.ChunkIndex < TNumericLimits<uint32>::Max() && ((int32)InKey.ChunkIndex) < InKey.GetNumChunks();
 }
 
+uint64 FAudioChunkCache::GetCacheLookupIDForChunk(const FChunkKey& InChunkKey) const
+{
+	const uint64* ID = CacheLookupIdMap.Find(InChunkKey);
+
+	if (ID)
+	{
+		return *ID;
+	}
+	else
+	{
+		return InvalidAudioStreamCacheLookupID;
+	}
+}
+
+void FAudioChunkCache::SetCacheLookupIDForChunk(const FChunkKey& InChunkKey, uint64 InCacheLookupID)
+{
+	uint64& ID = CacheLookupIdMap.FindOrAdd(InChunkKey);
+	ID = InCacheLookupID;
+}
+
+
 #include "UnrealEngine.h"
 
 int32 FCachedAudioStreamingManager::RenderStatAudioStreaming(UWorld* World, FViewport* Viewport, FCanvas* Canvas, int32 X, int32 Y, const FVector* ViewLocation, const FRotator* ViewRotation)
@@ -1937,7 +1977,7 @@ TPair<int, int> FAudioChunkCache::DebugDisplayLegacy(UWorld* World, FViewport* V
 
 #if WITH_EDITOR
 		// TODO: Worry about whether the sound wave is alive here. In most editor cases this is ok because the soundwave will always be loaded, but this may not be the case in the future.
-		bIsStaleChunk = (CurrentElement->Key.SoundWave == nullptr) || (CurrentElement->Key.SoundWave->CurrentChunkRevision.GetValue() != CurrentElement->Key.ChunkRevision);
+		bIsStaleChunk = CurrentElement->Key.IsChunkStale();
 #endif
 
 		const bool bWasTrimmed = CurrentElement->ChunkDataSize == 0;
@@ -2072,7 +2112,7 @@ FString FAudioChunkCache::DebugPrint()
 
 #if WITH_EDITOR
 		// TODO: Worry about whether the sound wave is alive here. In most editor cases this is ok because the soundwave will always be loaded, but this may not be the case in the future.
-		bIsStaleChunk = (CurrentElement->Key.SoundWave == nullptr) || (CurrentElement->Key.SoundWave->CurrentChunkRevision.GetValue() != CurrentElement->Key.ChunkRevision);
+		bIsStaleChunk = CurrentElement->Key.IsChunkStale();
 #endif
 
 		const bool bWasTrimmed = CurrentElement->ChunkDataSize == 0;
@@ -2172,7 +2212,7 @@ TPair<int, int> FAudioChunkCache::DebugDisplay(UWorld* World, FViewport* Viewpor
 	int32 NumOther = 0;
 
 
-	for(int i = 0; i < ChunksInUse; ++i)
+	for (int i = 0; i < ChunksInUse; ++i)
 	{
 		const FAudioChunkCache::FCacheElement* CurrentElement = &CachePool[i];
 
@@ -2265,19 +2305,19 @@ TPair<int, int> FAudioChunkCache::DebugDisplay(UWorld* World, FViewport* Viewpor
 	Y += 24;
 
 	// gather cache composition as percentages
-	float NumChunks = NumRetainedAndPlaying 
-					+ NumRetained 
-					+ NumPrimedAndPlaying 
-					+ NumPrimed 
-					+ NumRetainedAndPlayingCacheMiss
-					+ NumRetainedCacheMiss
-					+ NumPrimedAndPlayingCacheMiss
-					+ NumPrimedCacheMiss
-					+ NumLODAndPlaying 
-					+ NumLOD 
-					+ NumTrimmed
-					+ NumLoadInProgress
-					+ NumOther;
+	float NumChunks = NumRetainedAndPlaying
+		+ NumRetained
+		+ NumPrimedAndPlaying
+		+ NumPrimed
+		+ NumRetainedAndPlayingCacheMiss
+		+ NumRetainedCacheMiss
+		+ NumPrimedAndPlayingCacheMiss
+		+ NumPrimedCacheMiss
+		+ NumLODAndPlaying
+		+ NumLOD
+		+ NumTrimmed
+		+ NumLoadInProgress
+		+ NumOther;
 
 	if (NumChunks == 0)
 	{
@@ -2285,7 +2325,7 @@ TPair<int, int> FAudioChunkCache::DebugDisplay(UWorld* World, FViewport* Viewpor
 	}
 
 	// Draw the composition bar
-	const int32 BarWidth = 0.5f * (Canvas->GetParentCanvasSize().X - 2*X);
+	const int32 BarWidth = 0.5f * (Canvas->GetParentCanvasSize().X - 2 * X);
 	const int32 BarHeight = 20;
 	const int32 BarPad = BarHeight / 7;
 
@@ -2303,19 +2343,19 @@ TPair<int, int> FAudioChunkCache::DebugDisplay(UWorld* World, FViewport* Viewpor
 	const float PercentageLoadInProgress = (NumLoadInProgress / NumChunks);
 	const float PercentageOther = (NumOther / NumChunks);
 
-	const int32 BarWidthRetainedAndPlaying			= PercentageRetainedAndPlaying * BarWidth;
-	const int32 BarWidthRetained					= PercentageRetained * BarWidth;
-	const int32 BarWidthPrimedAndPlaying			= PercentagePrimedAndPlaying * BarWidth;
-	const int32 BarWidthPrimed						= PercentagePrimed * BarWidth;
+	const int32 BarWidthRetainedAndPlaying = PercentageRetainedAndPlaying * BarWidth;
+	const int32 BarWidthRetained = PercentageRetained * BarWidth;
+	const int32 BarWidthPrimedAndPlaying = PercentagePrimedAndPlaying * BarWidth;
+	const int32 BarWidthPrimed = PercentagePrimed * BarWidth;
 	const int32 BarWidthRetainedAndPlayingCacheMiss = PercentageRetainedAndPlayingCacheMiss * BarWidth;
-	const int32 BarWidthRetainedCacheMiss			= PercentageRetainedCacheMiss * BarWidth;
-	const int32 BarWidthPrimedAndPlayingCacheMiss	= PercentagePrimedAndPlayingCacheMiss * BarWidth;
-	const int32 BarWidthPrimedCacheMiss				= PercentagePrimedCacheMiss * BarWidth;
-	const int32 BarWidthLODAndPlaying				= PercentageLODAndPlaying * BarWidth;
-	const int32 BarWidthLOD							= PercentageLOD * BarWidth;
-	const int32 BarWidthTrimmed						= PercentageTrimmed * BarWidth;
-	const int32 BarWidthLoadInProgress				= PercentageLoadInProgress * BarWidth;
-	const int32 BarWidthOther						= PercentageOther * BarWidth;
+	const int32 BarWidthRetainedCacheMiss = PercentageRetainedCacheMiss * BarWidth;
+	const int32 BarWidthPrimedAndPlayingCacheMiss = PercentagePrimedAndPlayingCacheMiss * BarWidth;
+	const int32 BarWidthPrimedCacheMiss = PercentagePrimedCacheMiss * BarWidth;
+	const int32 BarWidthLODAndPlaying = PercentageLODAndPlaying * BarWidth;
+	const int32 BarWidthLOD = PercentageLOD * BarWidth;
+	const int32 BarWidthTrimmed = PercentageTrimmed * BarWidth;
+	const int32 BarWidthLoadInProgress = PercentageLoadInProgress * BarWidth;
+	const int32 BarWidthOther = PercentageOther * BarWidth;
 
 
 	// Draw color key
@@ -2377,7 +2417,7 @@ TPair<int, int> FAudioChunkCache::DebugDisplay(UWorld* World, FViewport* Viewpor
 	const int32 TotalRetainedAndPlaying = BarWidthRetainedAndPlaying + BarWidthRetainedAndPlayingCacheMiss;
 	Canvas->DrawTile(CurrHorzOffset, CurrVertOffset, TotalRetainedAndPlaying, BarHeight, 0, 0, 0, 0, ColorRetainedAndPlaying);
 	CurrHorzOffset += TotalRetainedAndPlaying;
-	
+
 	// (retained)
 	const int32 TotalRetained = BarWidthRetained + BarWidthRetainedCacheMiss;
 	Canvas->DrawTile(CurrHorzOffset, CurrVertOffset, TotalRetained, BarHeight, 0, 0, 0, 0, ColorRetained);
@@ -2416,7 +2456,7 @@ TPair<int, int> FAudioChunkCache::DebugDisplay(UWorld* World, FViewport* Viewpor
 	Y = (CurrVertOffset + 24);
 
 	// Draw the body of our display depending on the CVAR
-	TPair<int, int> Size(X,Y);
+	TPair<int, int> Size(X, Y);
 	if (DebugViewCVar == 0)
 	{
 		Size = DebugDisplayLegacy(World, Viewport, Canvas, X, Y + 2 * BarPad, ViewLocation, ViewRotation);
@@ -2454,7 +2494,7 @@ TPair<int, int> FAudioChunkCache::DebugVisualDisplay(UWorld* World, FViewport* V
 	// More detailed info about individual chunks here:
 	const int32 TileSize = 3;
 	const int32 TilePadding = 2;
-	const int32 MaxWidth = 0.5f * (Canvas->GetParentCanvasSize().X -2*X);
+	const int32 MaxWidth = 0.5f * (Canvas->GetParentCanvasSize().X - 2 * X);
 
 	int32 CurrentXOffset = 0;
 
@@ -2487,7 +2527,7 @@ TPair<int, int> FAudioChunkCache::DebugVisualDisplay(UWorld* World, FViewport* V
 
 #if WITH_EDITOR
 		// TODO: Worry about whether the sound wave is alive here. In most editor cases this is ok because the soundwave will always be loaded, but this may not be the case in the future.
-		bIsStaleChunk = (CurrentElement->Key.SoundWave == nullptr) || (CurrentElement->Key.SoundWave->CurrentChunkRevision.GetValue() != CurrentElement->Key.ChunkRevision);
+		bIsStaleChunk = CurrentElement->Key.IsChunkStale();
 #endif
 		const bool bWasTrimmed = CurrentElement->ChunkDataSize == 0;
 
diff --git a/Engine/Source/Runtime/Engine/Private/ContentStreaming.cpp b/Engine/Source/Runtime/Engine/Private/ContentStreaming.cpp
--- a/Engine/Source/Runtime/Engine/Private/ContentStreaming.cpp
+++ b/Engine/Source/Runtime/Engine/Private/ContentStreaming.cpp
@@ -1399,22 +1399,20 @@ FAudioChunkHandle::FAudioChunkHandle()
 	, CorrespondingWave(nullptr)
 	, CorrespondingWaveName()
 	, ChunkIndex(INDEX_NONE)
-	, CacheLookupID(InvalidAudioStreamCacheLookupID)
 #if WITH_EDITOR
 	, ChunkGeneration(INDEX_NONE)
 #endif
 {
 }
 
-FAudioChunkHandle::FAudioChunkHandle(const uint8* InData, uint32 NumBytes, const USoundWave* InSoundWave, const FName& SoundWaveName, uint32 InChunkIndex, uint64 InCacheLookupID)
+FAudioChunkHandle::FAudioChunkHandle(const uint8* InData, uint32 NumBytes, const FSoundWaveProxy&  InSoundWave, const FName& SoundWaveName, uint32 InChunkIndex, uint64 InCacheLookupID)
 	: CachedData(InData)
 	, CachedDataNumBytes(NumBytes)
-	, CorrespondingWave(InSoundWave)
+	, CorrespondingWave(MakeUnique<FSoundWaveProxy>(InSoundWave))
 	, CorrespondingWaveName(SoundWaveName)
 	, ChunkIndex(InChunkIndex)
-	, CacheLookupID(InCacheLookupID)
 #if WITH_EDITOR
-	, ChunkGeneration(InSoundWave->CurrentChunkRevision.GetValue())
+	, ChunkGeneration(InSoundWave.GetCurrentChunkRevision())
 #endif
 {
 }
@@ -1441,10 +1439,9 @@ FAudioChunkHandle& FAudioChunkHandle::operator=(FAudioChunkHandle&& Other)
 
 	CachedData = Other.CachedData;
 	CachedDataNumBytes = Other.CachedDataNumBytes;
-	CorrespondingWave = Other.CorrespondingWave;
+	CorrespondingWave = MoveTemp(Other.CorrespondingWave);
 	CorrespondingWaveName = Other.CorrespondingWaveName;
 	ChunkIndex = Other.ChunkIndex;
-	CacheLookupID = Other.CacheLookupID;
 #if WITH_EDITOR
 	ChunkGeneration = Other.ChunkGeneration;
 #endif
@@ -1456,7 +1453,6 @@ FAudioChunkHandle& FAudioChunkHandle::operator=(FAudioChunkHandle&& Other)
 	Other.CorrespondingWave = nullptr;
 	Other.CorrespondingWaveName = FName();
 	Other.ChunkIndex = INDEX_NONE;
-	Other.CacheLookupID = InvalidAudioStreamCacheLookupID;
 #if WITH_EDITOR
 	Other.ChunkGeneration = INDEX_NONE;
 #endif
@@ -1470,14 +1466,19 @@ FAudioChunkHandle& FAudioChunkHandle::operator=(const FAudioChunkHandle& Other)
 	if (IsValid())
 	{
 		IStreamingManager::Get().GetAudioStreamingManager().RemoveReferenceToChunk(*this);
+		CorrespondingWave.Reset();
 	}
 
 	CachedData = Other.CachedData;
 	CachedDataNumBytes = Other.CachedDataNumBytes;
-	CorrespondingWave = Other.CorrespondingWave;
+
+	if (Other.CorrespondingWave.IsValid())
+	{
+		CorrespondingWave = MakeUnique<FSoundWaveProxy>(*Other.CorrespondingWave);
+	}
+
 	CorrespondingWaveName = Other.CorrespondingWaveName;
 	ChunkIndex = Other.ChunkIndex;
-	CacheLookupID = Other.CacheLookupID;
 #if WITH_EDITOR
 	ChunkGeneration = Other.ChunkGeneration;
 #endif
@@ -1517,10 +1518,10 @@ bool FAudioChunkHandle::IsValid() const
 #if WITH_EDITOR
 bool FAudioChunkHandle::IsStale() const
 {
-	if (CorrespondingWave != nullptr)
+	if (CorrespondingWave.IsValid())
 	{
 		// NOTE: While this is currently safe in editor, there's no guarantee the USoundWave will be kept alive during the lifecycle of this chunk handle.
-		return ChunkGeneration != CorrespondingWave->CurrentChunkRevision.GetValue();
+		return ChunkGeneration != CorrespondingWave->GetCurrentChunkRevision();
 	}
 	else
 	{
@@ -1529,7 +1530,7 @@ bool FAudioChunkHandle::IsStale() const
 }
 #endif
 
-FAudioChunkHandle IAudioStreamingManager::BuildChunkHandle(const uint8* InData, uint32 NumBytes, const USoundWave* InSoundWave, const FName& SoundWaveName, uint32 InChunkIndex, uint64 InCacheLookupID)
+FAudioChunkHandle IAudioStreamingManager::BuildChunkHandle(const uint8* InData, uint32 NumBytes, const FSoundWaveProxy&  InSoundWave, const FName& SoundWaveName, uint32 InChunkIndex, uint64 InCacheLookupID)
 {
 	return FAudioChunkHandle(InData, NumBytes, InSoundWave, SoundWaveName, InChunkIndex, InCacheLookupID);
 }
diff --git a/Engine/Source/Runtime/Engine/Private/GameplayStatics.cpp b/Engine/Source/Runtime/Engine/Private/GameplayStatics.cpp
--- a/Engine/Source/Runtime/Engine/Private/GameplayStatics.cpp
+++ b/Engine/Source/Runtime/Engine/Private/GameplayStatics.cpp
@@ -1724,7 +1724,7 @@ void UGameplayStatics::PrimeSound(USoundBase* InSound)
 
 		if (InSoundWave->GetNumChunks() > 1)
 		{
-			IStreamingManager::Get().GetAudioStreamingManager().RequestChunk(InSoundWave, 1, [](EAudioChunkLoadResult) {});
+			IStreamingManager::Get().GetAudioStreamingManager().RequestChunk(InSoundWave->GetThisAsProxy(), 1, [](EAudioChunkLoadResult) {});
 		}
 	}
 }
diff --git a/Engine/Source/Runtime/Engine/Private/SoundWave.cpp b/Engine/Source/Runtime/Engine/Private/SoundWave.cpp
--- a/Engine/Source/Runtime/Engine/Private/SoundWave.cpp
+++ b/Engine/Source/Runtime/Engine/Private/SoundWave.cpp
@@ -536,7 +536,7 @@ void USoundWave::Serialize( FArchive& Ar )
 				// Prime first chunk of audio.
 				CachedSoundWaveLoadingBehavior = ESoundWaveLoadingBehavior::PrimeOnLoad;
 				bLoadingBehaviorOverridden = true;
-				IStreamingManager::Get().GetAudioStreamingManager().RequestChunk(this, 1, [](EAudioChunkLoadResult) {});
+				IStreamingManager::Get().GetAudioStreamingManager().RequestChunk(GetThisAsProxy(), 1, [](EAudioChunkLoadResult) {});
 			}
 		}
 	}
@@ -982,7 +982,7 @@ void USoundWave::InvalidateCompressedData(bool bFreeResources, bool bRebuildStre
 	if (bRebuildStreamingChunks)
 	{
 		CachePlatformData();
-		CurrentChunkRevision.Increment();
+		CurrentChunkRevision->Increment();
 	}
 
 
@@ -1026,7 +1026,6 @@ void USoundWave::PostLoad()
 	}
 
 	CacheInheritedLoadingBehavior();
-
 	
 	// If our loading behavior is defined by a sound class, we need to update whether this sound wave actually needs to retain its audio data or not.
 	ESoundWaveLoadingBehavior ActualLoadingBehavior = GetLoadingBehavior();
@@ -1049,7 +1048,7 @@ void USoundWave::PostLoad()
 			
 			if (bShouldPrime && bHasMultipleChunks)
 			{
-				IStreamingManager::Get().GetAudioStreamingManager().RequestChunk(this, 1, [](EAudioChunkLoadResult) {});
+				IStreamingManager::Get().GetAudioStreamingManager().RequestChunk(GetThisAsProxy(), 1, [](EAudioChunkLoadResult) {});
 			}
 		}
 
@@ -1110,7 +1109,10 @@ void USoundWave::PostLoad()
 #if WITH_EDITORONLY_DATA
 		FinishCachePlatformData();
 #endif // #if WITH_EDITORONLY_DATA
-		IStreamingManager::Get().GetAudioStreamingManager().AddStreamingSoundWave(this);
+		if (!ShouldUseStreamCaching())
+		{
+			IStreamingManager::Get().GetAudioStreamingManager().AddStreamingSoundWave(GetThisAsProxy());
+		}
 	}
 
 	const bool bDoesSoundWaveHaveStreamingAudioData = HasStreamingChunks();
@@ -1133,6 +1135,13 @@ void USoundWave::PostLoad()
 
 	INC_FLOAT_STAT_BY(STAT_AudioBufferTime, Duration);
 	INC_FLOAT_STAT_BY(STAT_AudioBufferTimeChannels, NumChannels * Duration);
+
+
+	// overwrite cached proxy now that post-load is complete.
+	Audio::FProxyDataInitParams Params;
+	Params.NameOfFeatureRequestingProxy = "USoundWave";
+	InternalProxy = MakeUnique<FSoundWaveProxy>(MoveTemp(CreateNewProxyData(Params)->GetAs<FSoundWaveProxy>()));
+	InternalProxy->ReleaseCompressedAudio(); // release handles if there are any
 }
 
 void USoundWave::EnsureZerothChunkIsLoaded()
@@ -1908,7 +1917,7 @@ bool USoundWave::IsReadyForFinishDestroy()
 		}
 	}
 
-	const bool bIsStreamingInProgress = IStreamingManager::Get().GetAudioStreamingManager().IsStreamingInProgress(this);
+	const bool bIsStreamingInProgress = IStreamingManager::Get().GetAudioStreamingManager().IsStreamingInProgress(GetThisAsProxy());
 
 	check(GetPrecacheState() != ESoundWavePrecacheState::InProgress);
 
@@ -1960,7 +1969,7 @@ void USoundWave::FinishDestroy()
 
 	if (FApp::CanEverRenderAudio())
 	{
-		IStreamingManager::Get().GetAudioStreamingManager().RemoveStreamingSoundWave(this);
+		IStreamingManager::Get().GetAudioStreamingManager().RemoveStreamingSoundWave(GetThisAsProxy());
 	}
 }
 
@@ -2330,14 +2339,14 @@ TArrayView<const uint8> USoundWave::GetZerothChunk(bool bForImmediatePlayback)
 		if (GetNumChunks() > 1)
 		{
 			// Prime first chunk for playback.
-			IStreamingManager::Get().GetAudioStreamingManager().RequestChunk(this, 1, [](EAudioChunkLoadResult InResult) {}, ENamedThreads::AnyThread, bForImmediatePlayback);
+			IStreamingManager::Get().GetAudioStreamingManager().RequestChunk(GetThisAsProxy(), 1, [](EAudioChunkLoadResult InResult) {}, ENamedThreads::AnyThread, bForImmediatePlayback);
 		}
 
 		return ZerothChunkData.GetView();
 	}
 	else
 	{
-		FAudioChunkHandle ChunkHandle = IStreamingManager::Get().GetAudioStreamingManager().GetLoadedChunk(this, 0);
+		FAudioChunkHandle ChunkHandle = IStreamingManager::Get().GetAudioStreamingManager().GetLoadedChunk(GetThisAsProxy(), 0);
 		return TArrayView<const uint8>(ChunkHandle.GetData(), ChunkHandle.Num());
 	}
 }
@@ -2399,7 +2408,7 @@ void USoundWave::UpdatePlatformData()
 	if (IsStreaming(nullptr))
 	{
 		// Make sure there are no pending requests in flight.
-		while (IStreamingManager::Get().GetAudioStreamingManager().IsStreamingInProgress(this))
+		while (IStreamingManager::Get().GetAudioStreamingManager().IsStreamingInProgress(GetThisAsProxy()))
 		{
 			// Give up timeslice.
 			FPlatformProcess::Sleep(0);
@@ -2407,16 +2416,16 @@ void USoundWave::UpdatePlatformData()
 
 #if WITH_EDITORONLY_DATA
 		// Temporarily remove from streaming manager to release currently used data chunks
-		IStreamingManager::Get().GetAudioStreamingManager().RemoveStreamingSoundWave(this);
+		IStreamingManager::Get().GetAudioStreamingManager().RemoveStreamingSoundWave(GetThisAsProxy());
 		// Recache platform data if the source has changed.
 		CachePlatformData();
 		// Add back to the streaming manager to reload first chunk
-		IStreamingManager::Get().GetAudioStreamingManager().AddStreamingSoundWave(this);
+		IStreamingManager::Get().GetAudioStreamingManager().AddStreamingSoundWave(GetThisAsProxy());
 #endif
 	}
 	else
 	{
-		IStreamingManager::Get().GetAudioStreamingManager().RemoveStreamingSoundWave(this);
+		IStreamingManager::Get().GetAudioStreamingManager().RemoveStreamingSoundWave(GetThisAsProxy());
 	}
 }
 
@@ -2499,6 +2508,20 @@ bool USoundWave::GetChunkData(int32 ChunkIndex, uint8** OutChunkData, bool bMake
 	return true;
 }
 
+const FSoundWaveProxy& USoundWave::GetThisAsProxy()
+{
+	if (!InternalProxy.IsValid())
+	{
+		// PostLoad hasn't been called. The InternalProxy created will be overritten by PostLoad().
+		Audio::FProxyDataInitParams Params;
+		Params.NameOfFeatureRequestingProxy = "USoundWave";
+		InternalProxy = MakeUnique<FSoundWaveProxy>(MoveTemp(CreateNewProxyData(Params)->GetAs<FSoundWaveProxy>()));
+	}
+
+	check(InternalProxy.IsValid());
+	return *InternalProxy;
+}
+
 uint32 USoundWave::GetInterpolatedCookedFFTDataForTimeInternal(float InTime, uint32 StartingIndex, TArray<FSoundWaveSpectralData>& OutData, bool bLoop)
 {
 	// Find the two entries on either side of the input time
@@ -2723,7 +2746,7 @@ void USoundWave::GetHandleForChunkOfAudio(TFunction<void(FAudioChunkHandle&&)> O
 	else if (bForceSync)
 	{
 		// For sync cases, we call GetLoadedChunk with bBlockForLoad = true, then execute the callback immediately.
-		FAudioChunkHandle ChunkHandle = IStreamingManager::Get().GetAudioStreamingManager().GetLoadedChunk(this, ChunkIndex, true);
+		FAudioChunkHandle ChunkHandle = IStreamingManager::Get().GetAudioStreamingManager().GetLoadedChunk(GetThisAsProxy(), ChunkIndex, true);
 		OnLoadCompleted(MoveTemp(ChunkHandle));
 	}
 	else
@@ -2731,7 +2754,7 @@ void USoundWave::GetHandleForChunkOfAudio(TFunction<void(FAudioChunkHandle&&)> O
 		TWeakObjectPtr<USoundWave> WeakThis = MakeWeakObjectPtr(this);
 
 		// For async cases, we call RequestChunk and request the loaded chunk in the completion callback.
-		IStreamingManager::Get().GetAudioStreamingManager().RequestChunk(this, ChunkIndex, [ThreadToDispatchCallbackOn, WeakThis, OnLoadCompleted, ChunkIndex, CallbackThread](EAudioChunkLoadResult LoadResult)
+		IStreamingManager::Get().GetAudioStreamingManager().RequestChunk(GetThisAsProxy(), ChunkIndex, [ThreadToDispatchCallbackOn, WeakThis, OnLoadCompleted, ChunkIndex, CallbackThread](EAudioChunkLoadResult LoadResult)
 		{
 			auto DispatchOnLoadCompletedCallback = [ThreadToDispatchCallbackOn, OnLoadCompleted, CallbackThread](FAudioChunkHandle&& InHandle)
 			{
@@ -2753,7 +2776,7 @@ void USoundWave::GetHandleForChunkOfAudio(TFunction<void(FAudioChunkHandle&&)> O
 			if (WeakThis.IsValid() && (LoadResult == EAudioChunkLoadResult::Completed || LoadResult == EAudioChunkLoadResult::AlreadyLoaded))
 			{
 				USoundWave* ThisSoundWave = WeakThis.Get();
-				FAudioChunkHandle ChunkHandle = IStreamingManager::Get().GetAudioStreamingManager().GetLoadedChunk(ThisSoundWave, ChunkIndex, (BlockOnChunkLoadCompletionCVar != 0));
+				FAudioChunkHandle ChunkHandle = IStreamingManager::Get().GetAudioStreamingManager().GetLoadedChunk(ThisSoundWave->GetThisAsProxy(), ChunkIndex, (BlockOnChunkLoadCompletionCVar != 0));
 
 				// If we hit this, something went wrong in GetLoadedChunk.
 				if (!ChunkHandle.IsValid())
@@ -2789,7 +2812,7 @@ void USoundWave::RetainCompressedAudio(bool bForceSync /*= false*/)
 	}
 	else if (bForceSync)
 	{
-		FirstChunk = IStreamingManager::Get().GetAudioStreamingManager().GetLoadedChunk(this, 1, true);
+		FirstChunk = IStreamingManager::Get().GetAudioStreamingManager().GetLoadedChunk(GetThisAsProxy(), 1, true);
 		UE_CLOG(!FirstChunk.IsValid(), LogAudio, Display, TEXT("First chunk was invalid after synchronous load in RetainCompressedAudio(). This was likely because the cache was blown. Sound: %s"), *GetFullName());
 	}
 	else
@@ -2862,28 +2885,12 @@ void USoundWave::OverrideLoadingBehavior(ESoundWaveLoadingBehavior InLoadingBeha
 		}
 		else if (InLoadingBehavior == ESoundWaveLoadingBehavior::PrimeOnLoad && GetNumChunks() > 1)
 		{
-			IStreamingManager::Get().GetAudioStreamingManager().RequestChunk(this, 1, [](EAudioChunkLoadResult) {});
+			IStreamingManager::Get().GetAudioStreamingManager().RequestChunk(GetThisAsProxy(), 1, [](EAudioChunkLoadResult) {});
 		}
 	}
 }
 
-void USoundWave::SetCacheLookupIDForChunk(uint32 InChunkIndex, uint64 InCacheLookupID)
-{
-	check(FPlatformCompressionUtilities::IsCurrentPlatformUsingStreamCaching());
-	check(RunningPlatformData && InChunkIndex < ((uint32)RunningPlatformData->Chunks.Num()));
-
-	RunningPlatformData->Chunks[InChunkIndex].CacheLookupID = InCacheLookupID;
-}
-
-uint64 USoundWave::GetCacheLookupIDForChunk(uint32 InChunkIndex)
-{
-	check(FPlatformCompressionUtilities::IsCurrentPlatformUsingStreamCaching());
-	check(RunningPlatformData && (InChunkIndex < (uint32)RunningPlatformData->Chunks.Num()));
-
-	return RunningPlatformData->Chunks[InChunkIndex].CacheLookupID;
-}
-
-void USoundWave::CacheInheritedLoadingBehavior()
+void USoundWave::CacheInheritedLoadingBehavior() const
 {
 	check(IsInGameThread());
 
@@ -2956,11 +2963,227 @@ ESoundWaveLoadingBehavior USoundWave::GetLoadingBehavior(bool bCheckSoundClasses
 			return (ESoundWaveLoadingBehavior)FMath::Clamp<int32>(SoundWaveDefaultLoadingBehaviorCVar, 0, (int32)ESoundWaveLoadingBehavior::LoadOnDemand);
 		}
 	}
-	else if (CachedSoundWaveLoadingBehavior != ESoundWaveLoadingBehavior::Uninitialized)
+	else if (CachedSoundWaveLoadingBehavior == ESoundWaveLoadingBehavior::Uninitialized)
 	{
-		return CachedSoundWaveLoadingBehavior;
+		CacheInheritedLoadingBehavior();
 	}
 
-	// We should never reach this. See check at the beginning of this function.
-	return ESoundWaveLoadingBehavior::Inherited;
+	return CachedSoundWaveLoadingBehavior;
+}
+
+FSoundWaveProxy::FSoundWaveProxy(USoundWave* InWave)
+	: ZerothChunkData(InWave->ZerothChunkData)
+	, NameCached(InWave->GetFName())
+	, FullNameCached(InWave->GetFullName())
+	, PathNameCached(InWave->GetPathName())
+	, SoundWaveKeyCached(FObjectKey(InWave))
+	, SampleRate(InWave->GetSampleRateForCurrentPlatform())
+	, NumChannels(InWave->NumChannels)
+	, NumChunks(0)
+	, FirstChunk(InWave->FirstChunk)
+	, LoadingBehavior(InWave->CachedSoundWaveLoadingBehavior)
+	, bLoadingBehaviorOverridden(InWave->bLoadingBehaviorOverridden)
+	, bIsStreaming(InWave->IsStreaming(nullptr))
+	, bSeekableStreaming(InWave->bSeekableStreaming)
+	, bShouldUseStreamCaching(InWave->ShouldUseStreamCaching())
+#if WITH_EDITOR
+	, CurrentChunkRevision(InWave->CurrentChunkRevision)
+#endif // #if WITH_EDITOR
+{
+	// cache the runtime format for this wave
+	if (GEngine)
+	{
+		FAudioDevice* LocalAudioDevice = GEngine->GetMainAudioDeviceRaw();
+		if (LocalAudioDevice)
+		{
+			RuntimeFormat = LocalAudioDevice->GetRuntimeFormat(InWave);
+		}
+	}
+
+	if (InWave->RunningPlatformData)
+	{
+		check(LoadingBehavior != ESoundWaveLoadingBehavior::ForceInline);
+		RunningPlatformData = MakeUnique<FStreamedAudioPlatformData>(*InWave->RunningPlatformData);
+		NumChunks = InWave->GetNumChunks();
+	}
+
+//	UE_LOG(LogAudio, Display, TEXT("MAXDEBUG: FSoundWaveProxy(USoundWave* InWave)"));
+}
+
+FSoundWaveProxy::FSoundWaveProxy(const FSoundWaveProxy& Other)
+	: ZerothChunkData(Other.ZerothChunkData)
+	, RunningPlatformData(Other.RunningPlatformData? MakeUnique<FStreamedAudioPlatformData>(*Other.RunningPlatformData) : nullptr)
+	, NameCached(Other.NameCached)
+	, FullNameCached(Other.FullNameCached)
+	, PathNameCached(Other.PathNameCached)
+	, SoundWaveKeyCached(Other.SoundWaveKeyCached)
+	, RuntimeFormat(Other.RuntimeFormat)
+	, SampleRate(Other.SampleRate)
+	, NumChannels(Other.NumChannels)
+	, NumChunks(Other.NumChunks)
+	, FirstChunk(Other.FirstChunk)
+	, LoadingBehavior(Other.LoadingBehavior)
+	, bLoadingBehaviorOverridden(Other.bLoadingBehaviorOverridden)
+	, bIsStreaming(Other.bIsStreaming)
+	, bSeekableStreaming(Other.bSeekableStreaming)
+	, bShouldUseStreamCaching(Other.bShouldUseStreamCaching)
+#if WITH_EDITOR
+	, CurrentChunkRevision(Other.CurrentChunkRevision)
+#endif // #if WITH_EDITOR
+{
+//	UE_LOG(LogAudio, Display, TEXT("MAXDEBUG: FSoundWaveProxy(const FSoundWaveProxy& Other)"));
+}
+
+void FSoundWaveProxy::ReleaseCompressedAudio()
+{
+	FirstChunk = FAudioChunkHandle();
+}
+
+uint32 FSoundWaveProxy::GetSizeOfChunk(uint32 ChunkIndex) const
+{
+	check(ChunkIndex < GetNumChunks());
+	return GetChunk(ChunkIndex).AudioDataSize;
+}
+
+TArrayView<const uint8> FSoundWaveProxy::GetZerothChunk(bool bForImmediatePlayback) const
+{
+	if (bShouldUseStreamCaching)
+	{
+		check(IsZerothChunkDataLoaded());
+
+		if (GetNumChunks() > 1)
+		{
+			// Prime first chunk for playback
+			IStreamingManager::Get().GetAudioStreamingManager().RequestChunk(*this, 1, [](EAudioChunkLoadResult InResult) {}, ENamedThreads::AnyThread, bForImmediatePlayback);
+		}
+
+		return GetZerothChunkDataView();
+	}
+	else
+	{
+		FAudioChunkHandle ChunkHandle = IStreamingManager::Get().GetAudioStreamingManager().GetLoadedChunk(*this, 0);
+		return TArrayView<const uint8>(ChunkHandle.GetData(), ChunkHandle.Num());
+	}
+}
+
+bool FSoundWaveProxy::GetChunkData(int32 ChunkIndex, uint8** OutChunkData, bool bMakeSureChunkIsLoaded)
+{
+	if (GetChunkFromDDC(ChunkIndex, OutChunkData, bMakeSureChunkIsLoaded) == 0)
+	{
+#if WITH_EDITORONLY_DATA
+		// Unable to load chunks from the cache. Rebuild the sound and attempt to precache it.
+		UE_LOG(LogAudio, Display, TEXT("GetChunkData failed, rebuilding %s"), *GetPathName());
+
+		//		ForceRebuildPlatformData();
+		if (GetChunkFromDDC(ChunkIndex, OutChunkData, bMakeSureChunkIsLoaded) == 0)
+		{
+			UE_LOG(LogAudio, Warning, TEXT("Failed to build sound %s."), *GetPathName());
+		}
+		else
+		{
+			// Succeeded after rebuilding platform data
+			return true;
+		}
+#else
+		// Failed to find the SoundWave chunk in the cooked package.
+		UE_LOG(LogAudio, Warning, TEXT("GetChunkData failed while streaming. Ensure the following file is cooked: %s"), *GetPathName());
+#endif // #if WITH_EDITORONLY_DATA
+		return false;
+	}
+	return true;
+}
+
+
+bool FSoundWaveProxy::IsZerothChunkDataLoaded() const
+{
+	return ZerothChunkData.GetView().Num() > 0;
+}
+
+const TArrayView<uint8> FSoundWaveProxy::GetZerothChunkDataView() const
+{
+	return ZerothChunkData.GetView();
+}
+
+void FSoundWaveProxy::EnsureZerothChunkIsLoaded()
+{
+
+#ifdef WITH_EDITOR
+	// If we're running the editor, we'll need to retrieve the chunked audio from the DDC:
+	uint8* TempChunkBuffer = nullptr;
+	int32 ChunkSizeInBytes = GetChunkFromDDC(0, &TempChunkBuffer, true);
+	// Since we block for the DDC in the previous call we should always have the chunk loaded.
+	if (ChunkSizeInBytes == 0)
+	{
+		return;
+	}
+
+	ZerothChunkData.Reset(TempChunkBuffer, ChunkSizeInBytes);
+
+#else // WITH_EDITOR
+	// Otherwise, the zeroth chunk is cooked out to SharedRuntimeData::RunningPlatformData, and we just need to retrieve it.
+	check(GetNumChunks() > 0);
+	FStreamedAudioChunk& ZerothChunk = GetChunk(0);
+	// Some sanity checks to ensure that the bulk size set up
+	UE_CLOG(ZerothChunk.BulkData.GetBulkDataSize() != ZerothChunk.DataSize, LogAudio, Warning, TEXT("Bulk data serialized out had a mismatched size with the DataSize field. Soundwave: %s Bulk Data Reported Size: %d Bulk Data Actual Size: %ld"), *GetFullName(), ZerothChunk.DataSize, ZerothChunk.BulkData.GetBulkDataSize());
+
+	ZerothChunkData = ZerothChunk.BulkData.GetCopyAsBuffer(ZerothChunk.AudioDataSize, true);
+#endif // WITH_EDITOR
+}
+
+uint32 FSoundWaveProxy::GetNumChunks() const
+{
+	// This function shouldn't be called on audio marked "ForceInline."
+	ensureMsgf(RunningPlatformData, TEXT("Calling GetNumChunks on a FSoundWaveProxy without RunnigPlatformData is not allowed! SoundWave: %s - %s")
+		, *GetFName().ToString(), EnumToString(LoadingBehavior));
+
+	if (RunningPlatformData)
+	{
+		return RunningPlatformData->Chunks.Num();
+	}
+	
+	else
+	{
+		return 0;
+	}
+}
+
+#if WITH_EDITOR
+int32 FSoundWaveProxy::GetCurrentChunkRevision() const
+{
+	return CurrentChunkRevision->GetValue();
+}
+#endif // #if WITH_EDITOR
+
+const FStreamedAudioChunk& FSoundWaveProxy::GetChunk(uint32 ChunkIndex) const
+{
+	// This function shouldn't be called on audio marked "ForceInline."
+	ensureMsgf(RunningPlatformData, TEXT("Calling GetNumChunks on a FSoundWaveProxy without RunnigPlatformData is not allowed! SoundWave: %s - %s")
+		, *GetFName().ToString(), EnumToString(LoadingBehavior));
+
+	check((ChunkIndex < (uint32)GetNumChunks()));
+	check(RunningPlatformData);
+	return RunningPlatformData->Chunks[ChunkIndex];
+}
+
+int32 FSoundWaveProxy::GetChunkFromDDC(int32 ChunkIndex, uint8** OutChunkData, bool bMakeSureChunkIsLoaded)
+{
+	// This function shouldn't be called on audio marked "ForceInline."
+	ensureMsgf(RunningPlatformData, TEXT("Calling GetNumChunks on a FSoundWaveProxy without RunnigPlatformData is not allowed! SoundWave: %s - %s")
+		, *GetFName().ToString(), EnumToString(LoadingBehavior));
+
+	check(RunningPlatformData);
+	return RunningPlatformData->GetChunkFromDDC(ChunkIndex, OutChunkData, bMakeSureChunkIsLoaded);
+}
+
+#if WITH_EDITORONLY_DATA
+FString FSoundWaveProxy::GetDerivedDataKey() const
+{
+	check(RunningPlatformData);
+
+	// This function shouldn't be called on audio marked "ForceInline."
+	ensureMsgf(RunningPlatformData, TEXT("Calling GetNumChunks on a FSoundWaveProxy without RunnigPlatformData is not allowed! SoundWave: %s - %s")
+		, *GetFName().ToString(), EnumToString(LoadingBehavior));
+
+	return RunningPlatformData->DerivedDataKey;
 }
+#endif // #if WITH_EDITORONLY_DATA
\ No newline at end of file
diff --git a/Engine/Source/Runtime/Engine/Private/VorbisAudioInfo.cpp b/Engine/Source/Runtime/Engine/Private/VorbisAudioInfo.cpp
--- a/Engine/Source/Runtime/Engine/Private/VorbisAudioInfo.cpp
+++ b/Engine/Source/Runtime/Engine/Private/VorbisAudioInfo.cpp
@@ -206,12 +206,12 @@ size_t FVorbisAudioInfo::ReadStreaming(void *Ptr, uint32 Size )
 			}
 			else
 			{
-				CurrentStreamingChunkData = GetLoadedChunk(StreamingSoundWave, CurrentStreamingChunkIndex, CurrentStreamingChunksSize);
+				CurrentStreamingChunkData = GetLoadedChunk(*StreamingSoundWave, CurrentStreamingChunkIndex, CurrentStreamingChunksSize);
 			}
 
 			if (CurrentStreamingChunkData)
 			{
-				check(CurrentStreamingChunkIndex < StreamingSoundWave->RunningPlatformData->Chunks.Num());
+				check(CurrentStreamingChunkIndex < (int32)StreamingSoundWave->GetNumChunks());
 				CurrentBufferChunkOffset = 0;
 			}
 		}
@@ -541,11 +541,11 @@ void FVorbisAudioInfo::EnableHalfRate( bool HalfRate )
 	ov_halfrate(&VFWrapper->vf, int32(HalfRate));
 }
 
-bool FVorbisAudioInfo::StreamCompressedInfoInternal(USoundWave* Wave, struct FSoundQualityInfo* QualityInfo)
+bool FVorbisAudioInfo::StreamCompressedInfoInternal(const FSoundWaveProxy& InWaveProxy, struct FSoundQualityInfo* QualityInfo)
 {
 	if (!bDllLoaded)
 	{
-		UE_LOG(LogAudio, Error, TEXT("FVorbisAudioInfo::StreamCompressedInfoInternal failed to parse header due to vorbis DLL not being loaded for sound '%s'."), *Wave->GetName());
+		UE_LOG(LogAudio, Error, TEXT("FVorbisAudioInfo::StreamCompressedInfoInternal failed to parse header due to vorbis DLL not being loaded for sound '%s'."), *InWaveProxy.GetFName().ToString());
 		return false;
 	}
 
@@ -555,7 +555,7 @@ bool FVorbisAudioInfo::StreamCompressedInfoInternal(USoundWave* Wave, struct FSo
 
 	if (!VFWrapper)
 	{
-		UE_LOG(LogAudio, Error, TEXT("FVorbisAudioInfo::StreamCompressedInfoInternal failed due to no vorbis wrapper for sound '%s'."), *Wave->GetName());
+		UE_LOG(LogAudio, Error, TEXT("FVorbisAudioInfo::StreamCompressedInfoInternal failed due to no vorbis wrapper for sound '%s'."), *InWaveProxy.GetFName().ToString());
 		return false;
 	}
 
@@ -564,7 +564,6 @@ bool FVorbisAudioInfo::StreamCompressedInfoInternal(USoundWave* Wave, struct FSo
 	SrcBufferData = NULL;
 	SrcBufferDataSize = 0;
 	BufferOffset = 0;
-	StreamingSoundWave = Wave;
 
 	Callbacks.read_func = OggReadStreaming;
 	Callbacks.close_func = OggCloseStreaming;
@@ -574,28 +573,24 @@ bool FVorbisAudioInfo::StreamCompressedInfoInternal(USoundWave* Wave, struct FSo
 	bHeaderParsed = GetCompressedInfoCommon(&Callbacks, QualityInfo);
 	if (!bHeaderParsed)
 	{
-		UE_LOG(LogAudio, Error, TEXT("FVorbisAudioInfo::StreamCompressedInfoInternal failed to parse header for '%s'."), *Wave->GetName());
+		UE_LOG(LogAudio, Error, TEXT("FVorbisAudioInfo::StreamCompressedInfoInternal failed to parse header for '%s'."), *InWaveProxy.GetFName().ToString());
 	}
 	
 
 	return bHeaderParsed;
 }
 
-const uint8* FVorbisAudioInfo::GetLoadedChunk(USoundWave* InSoundWave, uint32 ChunkIndex, uint32& OutChunkSize)
+// TODO: Make InSoundWave a const ref after fixing up FSoundWaveProxy::GetZerothChunk() to be const
+const uint8* FVorbisAudioInfo::GetLoadedChunk(FSoundWaveProxy& InSoundWave, uint32 ChunkIndex, uint32& OutChunkSize)
 {
-	if (!InSoundWave || ChunkIndex >= InSoundWave->GetNumChunks())
+	if (ChunkIndex >= InSoundWave.GetNumChunks())
 	{
-		if(InSoundWave)
-		{
-			UE_LOG(LogAudio, Verbose, TEXT("Error calling GetLoadedChunk on wave with %d chunks. ChunkIndex: %d. Name: %s"), InSoundWave->GetNumChunks(), ChunkIndex, *InSoundWave->GetFullName());
-		}
-		
 		OutChunkSize = 0;
 		return nullptr;
 	}
 	else if (ChunkIndex == 0)
 	{
-		TArrayView<const uint8> ZerothChunk = InSoundWave->GetZerothChunk(true);
+		TArrayView<const uint8> ZerothChunk = InSoundWave.GetZerothChunk(true);
 		OutChunkSize = ZerothChunk.Num();
 		return ZerothChunk.GetData();
 	}
diff --git a/Engine/Source/Runtime/Engine/Public/ADPCMAudioInfo.h b/Engine/Source/Runtime/Engine/Public/ADPCMAudioInfo.h
--- a/Engine/Source/Runtime/Engine/Public/ADPCMAudioInfo.h
+++ b/Engine/Source/Runtime/Engine/Public/ADPCMAudioInfo.h
@@ -106,7 +106,7 @@ public:
 
 	// Additional overrides for streaming
 	virtual bool SupportsStreaming() const override {return true;}
-	virtual bool StreamCompressedInfoInternal(USoundWave* Wave, struct FSoundQualityInfo* QualityInfo) override;
+	virtual bool StreamCompressedInfoInternal(const FSoundWaveProxy& InWaveProxy, struct FSoundQualityInfo* QualityInfo) override;
 	virtual bool StreamCompressedData(uint8* Destination, bool bLooping, uint32 BufferSize) override;
 	virtual int32 GetCurrentChunkIndex() const override
 	{
@@ -138,7 +138,7 @@ private:
 	// Wrapper function that returns a pointer to the currently used compressed data.
 	// If a non-zero chunk is requested, this function also aquires a reference to that chunk
 	// until we move on to a different chunk.
-	const uint8* GetLoadedChunk(USoundWave* InSoundWave, uint32 ChunkIndex, uint32& OutChunkSize);
+	const uint8* GetLoadedChunk(const FSoundWaveProxy& InSoundWave, uint32 ChunkIndex, uint32& OutChunkSize);
 
 	int32 NumConsecutiveReadFailiures;
 
diff --git a/Engine/Source/Runtime/Engine/Public/AudioDecompress.h b/Engine/Source/Runtime/Engine/Public/AudioDecompress.h
--- a/Engine/Source/Runtime/Engine/Public/AudioDecompress.h
+++ b/Engine/Source/Runtime/Engine/Public/AudioDecompress.h
@@ -117,12 +117,8 @@ public:
 	* @param	Wave			Wave that will be read from to retrieve necessary chunk
 	* @param	QualityInfo		Quality Info (to be filled out)
 	*/
-	bool StreamCompressedInfo(USoundWave* Wave, struct FSoundQualityInfo* QualityInfo)
-	{
-		StreamingSoundWave = Wave;
-
-		return StreamCompressedInfoInternal(Wave, QualityInfo);
-	}
+	ENGINE_API bool StreamCompressedInfo(USoundWave* Wave, struct FSoundQualityInfo* QualityInfo);
+	ENGINE_API bool StreamCompressedInfo(const FSoundWaveProxy& Wave, struct FSoundQualityInfo* QualityInfo);
 
 	/**
 	*  Returns true if a non-recoverable error has occurred.
@@ -131,7 +127,7 @@ public:
 
 protected:
 	/** Internal override implemented by subclasses. */
-	virtual bool StreamCompressedInfoInternal(USoundWave* Wave, struct FSoundQualityInfo* QualityInfo) = 0;
+	virtual bool StreamCompressedInfoInternal(const FSoundWaveProxy& InWaveProxy, struct FSoundQualityInfo* QualityInfo) = 0;
 
 public:
 
@@ -139,7 +135,7 @@ public:
 	* Decompresses streamed data to raw PCM data.
 	*
 	* @param	Destination	where to place the decompressed sound
-	* @param	bLooping	whether to loop the sound by seeking to the start, or pad the buffer with zeroes
+	* @param	bLooping	whether to loop the sound by seeking to the start, or pad the buffer with zeros
 	* @param	BufferSize	number of bytes of PCM data to create
 	*
 	* @return	bool		true if the end of the data was reached (for both single shot and looping sounds)
@@ -157,10 +153,10 @@ public:
 	virtual int32 GetCurrentChunkOffset() const {return -1;}
 
 	/** Return the streaming sound wave used by this decoder. Returns nullptr if there is not a streaming sound wave. */
-	virtual USoundWave* GetStreamingSoundWave() { return StreamingSoundWave; }
+	virtual const FSoundWaveProxy& GetStreamingSoundWave() { return *StreamingSoundWave; }
 
 protected:
-	USoundWave* StreamingSoundWave;
+	TUniquePtr<FSoundWaveProxy> StreamingSoundWave;
 };
 
 /** Struct used to store the results of a decode operation. **/
@@ -201,7 +197,7 @@ public:
 	virtual bool UsesVorbisChannelOrdering() const override { return false; }
 	virtual int GetStreamBufferSize() const override { return  MONO_PCM_BUFFER_SIZE; }
 	virtual bool SupportsStreaming() const override { return true; }
-	virtual bool StreamCompressedInfoInternal(USoundWave* Wave, FSoundQualityInfo* QualityInfo) override;
+	virtual bool StreamCompressedInfoInternal(const FSoundWaveProxy& InWaveProxy, FSoundQualityInfo* QualityInfo) override;
 	virtual bool StreamCompressedData(uint8* Destination, bool bLooping, uint32 BufferSize) override;
 	virtual int32 GetCurrentChunkIndex() const override { return CurrentChunkIndex; }
 	virtual int32 GetCurrentChunkOffset() const override { return SrcBufferOffset; }
@@ -271,7 +267,7 @@ protected:
 	 * @param[out] OutChunkSize the size of the chunk.
 	 * @return a pointer to the chunk if it's loaded, nullptr otherwise.
 	 */
-	const uint8* GetLoadedChunk(USoundWave* InSoundWave, uint32 ChunkIndex, uint32& OutChunkSize);
+	const uint8* GetLoadedChunk(const FSoundWaveProxy& InSoundWave, uint32 ChunkIndex, uint32& OutChunkSize);
 
 	/** bool set before ParseHeader. Whether we are streaming a file or not. */
 	bool bIsStreaming;
diff --git a/Engine/Source/Runtime/Engine/Public/AudioDevice.h b/Engine/Source/Runtime/Engine/Public/AudioDevice.h
--- a/Engine/Source/Runtime/Engine/Public/AudioDevice.h
+++ b/Engine/Source/Runtime/Engine/Public/AudioDevice.h
@@ -1042,6 +1042,7 @@ public:
 
 	/** Creates a Compressed audio info class suitable for decompressing this SoundWave */
 	virtual ICompressedAudioInfo* CreateCompressedAudioInfo(USoundWave* SoundWave) { return nullptr; }
+	virtual ICompressedAudioInfo* CreateCompressedAudioInfo(const FSoundWaveProxy& SoundWave) { return nullptr; }
 
 	/**
 	 * Check for errors and output a human readable string
diff --git a/Engine/Source/Runtime/Engine/Public/AudioStreaming.h b/Engine/Source/Runtime/Engine/Public/AudioStreaming.h
--- a/Engine/Source/Runtime/Engine/Public/AudioStreaming.h
+++ b/Engine/Source/Runtime/Engine/Public/AudioStreaming.h
@@ -17,7 +17,7 @@ AudioStreaming.h: Definitions of classes used for audio streaming.
 #include "HAL/ThreadSafeBool.h"
 
 class FSoundSource;
-class USoundWave;
+class FSoundWaveProxy;
 class ICompressedAudioInfo;
 struct FLegacyAudioStreamingManager;
 struct FWaveInstance;
@@ -139,7 +139,7 @@ struct FStreamingWaveData final
 	 *
 	 * @param SoundWave	The SoundWave we are managing
 	 */
-	bool Initialize(USoundWave* SoundWave, FLegacyAudioStreamingManager* InStreamingManager);
+	bool Initialize(const FSoundWaveProxy& SoundWave, FLegacyAudioStreamingManager* InStreamingManager);
 
 	/**
 	 * Updates the streaming status of the sound wave and performs finalization when appropriate. The function returns
@@ -199,7 +199,7 @@ private:
 
 public:
 	/** SoundWave this streaming data is for */
-	USoundWave* SoundWave;
+	TUniquePtr<FSoundWaveProxy> SoundWave;
 
 	/** Thread-safe counter indicating the audio streaming state. */
 	mutable FThreadSafeCounter	PendingChunkChangeRequestStatus;
@@ -263,18 +263,18 @@ struct FLegacyAudioStreamingManager : public IAudioStreamingManager
 	// End IStreamingManager interface
 
 	// IAudioStreamingManager interface
-	virtual void AddStreamingSoundWave(USoundWave* SoundWave) override;
-	virtual void RemoveStreamingSoundWave(USoundWave* SoundWave) override;
+	virtual void AddStreamingSoundWave(const FSoundWaveProxy& SoundWave) override;
+	virtual void RemoveStreamingSoundWave(const FSoundWaveProxy& SoundWave) override;
 	virtual void AddDecoder(ICompressedAudioInfo* CompressedAudioInfo) override;
 	virtual void RemoveDecoder(ICompressedAudioInfo* CompressedAudioInfo) override;
-	virtual bool IsManagedStreamingSoundWave(const USoundWave* SoundWave) const override;
-	virtual bool IsStreamingInProgress(const USoundWave* SoundWave) override;
+	virtual bool IsManagedStreamingSoundWave(const FSoundWaveProxy&  SoundWave) const override;
+	virtual bool IsStreamingInProgress(const FSoundWaveProxy&  SoundWave) override;
 	virtual bool CanCreateSoundSource(const FWaveInstance* WaveInstance) const override;
 	virtual void AddStreamingSoundSource(FSoundSource* SoundSource) override;
 	virtual void RemoveStreamingSoundSource(FSoundSource* SoundSource) override;
 	virtual bool IsManagedStreamingSoundSource(const FSoundSource* SoundSource) const override;
-	virtual bool RequestChunk(USoundWave* SoundWave, uint32 ChunkIndex, TFunction<void(EAudioChunkLoadResult)> OnLoadCompleted, ENamedThreads::Type ThreadToCallOnLoadCompletedOn, bool bForImmediatePlayback = false) override;
-	virtual FAudioChunkHandle GetLoadedChunk(const USoundWave* SoundWave, uint32 ChunkIndex, bool bBlockForLoad = false, bool bForImmediatePlayback = false) const override;
+	virtual bool RequestChunk(const FSoundWaveProxy& SoundWave, uint32 ChunkIndex, TFunction<void(EAudioChunkLoadResult)> OnLoadCompleted, ENamedThreads::Type ThreadToCallOnLoadCompletedOn, bool bForImmediatePlayback = false) override;
+	virtual FAudioChunkHandle GetLoadedChunk(const FSoundWaveProxy&  SoundWave, uint32 ChunkIndex, bool bBlockForLoad = false, bool bForImmediatePlayback = false) const override;
 	virtual uint64 TrimMemory(uint64 NumBytesToFree) override;
 	virtual int32 RenderStatAudioStreaming(UWorld* World, FViewport* Viewport, FCanvas* Canvas, int32 X, int32 Y, const FVector* ViewLocation, const FRotator* ViewRotation) override;
 	virtual FString GenerateMemoryReport() override;
@@ -299,16 +299,16 @@ protected:
 	 * @param SoundWave		SoundWave we want request for
 	 * @return Existing or new request structure
 	 */
-	FWaveRequest& GetWaveRequest(USoundWave* SoundWave);
+	FWaveRequest& GetWaveRequest(FObjectKey Key);
 
 	/** Sound Waves being managed. */
-	TMap<USoundWave*, FStreamingWaveData*> StreamingSoundWaves;
+	TMap<FObjectKey, FStreamingWaveData*> StreamingSoundWaves;
 
 	/** Sound Sources being managed. */
 	TArray<FSoundSource*>	StreamingSoundSources;
 
 	/** Map of requests to make next time sound waves are ready */
-	TMap<USoundWave*, FWaveRequest> WaveRequests;
+	TMap<FObjectKey, FWaveRequest> WaveRequests;
 
 	/** Results of async loading audio chunks. */
 	TArray<FASyncAudioChunkLoadResult*> AsyncAudioStreamChunkResults;
diff --git a/Engine/Source/Runtime/Engine/Public/AudioStreamingCache.h b/Engine/Source/Runtime/Engine/Public/AudioStreamingCache.h
--- a/Engine/Source/Runtime/Engine/Public/AudioStreamingCache.h
+++ b/Engine/Source/Runtime/Engine/Public/AudioStreamingCache.h
@@ -31,17 +31,54 @@ class FAudioChunkCache
 public:
 	struct FChunkKey
 	{
-		USoundWave* SoundWave = nullptr;
-		FName SoundWaveName = FName();
+	public:
+		FChunkKey() {};
+
+		FChunkKey(const FChunkKey& Other);
+
+		// used when flushing the cache miss log
+		FChunkKey(FName InSoundWaveName, uint32 InChunkIndex);
+
+		FChunkKey(
+			  const FSoundWaveProxy& InSoundWave
+			, uint32 InChunkIndex
+#if WITH_EDITOR
+			, uint32 InChunkRevision = 0
+#endif // #if WITH_EDITOR
+		);
+
+		FChunkKey& operator=(const FChunkKey& Other);
+
+
+	private:
+		mutable TUniquePtr<FSoundWaveProxy> SoundWaveProxyPtr{ nullptr };
+
+	public:
+		mutable FName SoundWaveName = FName(); // mutable for GetTypeHash()
 		uint32 ChunkIndex = INDEX_NONE;
 		FObjectKey ObjectKey = FObjectKey();
 
 #if WITH_EDITOR
 		// This is used in the editor to invalidate stale compressed chunks.
-		uint32 ChunkRevision = INDEX_NONE;
+		uint32 ChunkRevision = 0;
 #endif
 		inline bool operator==(const FChunkKey& Other) const;
-		
+
+		bool IsChunkStale();
+		bool IsSoundWaveValid() const { return SoundWaveProxyPtr.IsValid(); }
+		// Forward FSoundWaveProxy interface
+		ESoundWaveLoadingBehavior GetLoadingBehavior() const;
+		FStreamedAudioChunk& GetChunk(uint32 InChunkIndex) const;
+		uint32 GetNumChunks() const;
+		bool IsRetainingAudio() const;
+		void ReleaseCompressedAudio();
+		bool WasLoadingBehaviorOverridden() const;
+
+		/** Hash function */
+		friend uint32 GetTypeHash(const FChunkKey& InChunkKey)
+		{
+			return HashCombine(InChunkKey.SoundWaveName.GetNumber(), InChunkKey.ChunkIndex);
+		}
 	};
 
 	FAudioChunkCache(uint32 InMaxChunkSize, uint32 NumChunks, uint64 InMemoryLimitInBytes);
@@ -58,8 +95,8 @@ public:
 	TArrayView<uint8> GetChunk(const FChunkKey& InKey, bool bBlockForLoadCompletion, bool bNeededForPlayback, uint64& InOutCacheLookupID);
 
 	// add an additional reference for a chunk.
-	void AddNewReferenceToChunk(const FChunkKey& InKey, uint64 InCacheLookupID);
-	void RemoveReferenceToChunk(const FChunkKey& InKey, uint64 InCacheLookupID);
+	void AddNewReferenceToChunk(const FChunkKey& InKey);
+	void RemoveReferenceToChunk(const FChunkKey& InKey);
 
 	// Evict all sounds from the cache.
 	void ClearCache();
@@ -97,6 +134,10 @@ public:
 	// Static helper function to make sure a chunk is withing the bounds of a USoundWave.
 	static bool IsKeyValid(const FChunkKey& InKey);
 
+	// interface with the cache id lookup map
+	uint64 GetCacheLookupIDForChunk(const FChunkKey& InChunkKey) const;
+	void SetCacheLookupIDForChunk(const FChunkKey& InChunkKey, uint64 InCacheLookupID);
+
 	const int32 MaxChunkSize;
 
 	// This is for debugging purposes only. Prints the elements in the cache from most recently used to least.
@@ -320,6 +361,10 @@ private:
 	// Critical section: only used when we are modifying element positions in the cache. This only happens in TouchElement, EvictLeastRecentChunk, and TrimMemory.
 	// Individual cache elements should be thread safe to access.
 	FCriticalSection CacheMutationCriticalSection;
+
+	// Map that USoundWaves, FSoundWaveProxys, FChunkKeys, and FAudioChunkHandles can use to
+	// quickly lookup where their chunks are currently stored in the cache
+	TMap<FChunkKey, uint64> CacheLookupIdMap;
 	 
 	// This struct is used for logging cache misses.
 	struct FCacheMissInfo
@@ -338,7 +383,7 @@ private:
 
 	// Returns cached element if it exists in our cache, nullptr otherwise.
 	// If the index of the element is already known, it can be used here to avoid searching the cache.
-	FCacheElement* FindElementForKey(const FChunkKey& InKey, uint64 CacheLookupID = InvalidAudioStreamCacheLookupID);
+	FCacheElement* FindElementForKey(const FChunkKey& InKey);
 	FCacheElement* LinearSearchCacheForElement(const FChunkKey& InKey);
 	FCacheElement* LinearSearchChunkArrayForElement(const FChunkKey& InKey);
 
@@ -422,12 +467,12 @@ public:
 	// End IStreamingManager interface
 
 	// IAudioStreamingManager interface (unused functions)
-	virtual void AddStreamingSoundWave(USoundWave* SoundWave) override;
-	virtual void RemoveStreamingSoundWave(USoundWave* SoundWave) override;
+	virtual void AddStreamingSoundWave(const FSoundWaveProxy& SoundWave) override;
+	virtual void RemoveStreamingSoundWave(const FSoundWaveProxy& SoundWave) override;
 	virtual void AddDecoder(ICompressedAudioInfo* CompressedAudioInfo) override;
 	virtual void RemoveDecoder(ICompressedAudioInfo* CompressedAudioInfo) override;
-	virtual bool IsManagedStreamingSoundWave(const USoundWave* SoundWave) const override;
-	virtual bool IsStreamingInProgress(const USoundWave* SoundWave) override;
+	virtual bool IsManagedStreamingSoundWave(const FSoundWaveProxy&  SoundWave) const override;
+	virtual bool IsStreamingInProgress(const FSoundWaveProxy&  SoundWave) override;
 	virtual bool CanCreateSoundSource(const FWaveInstance* WaveInstance) const override;
 	virtual void AddStreamingSoundSource(FSoundSource* SoundSource) override;
 	virtual void RemoveStreamingSoundSource(FSoundSource* SoundSource) override;
@@ -435,14 +480,15 @@ public:
 	// End IAudioStreamingManager interface (unused)
 
 	// IAudioStreamingManager interface (used functions)
-	virtual bool RequestChunk(USoundWave* SoundWave, uint32 ChunkIndex, TFunction<void(EAudioChunkLoadResult)> OnLoadCompleted, ENamedThreads::Type ThreadToCallOnLoadCompletedOn, bool bForImmediatePlayback = false) override;
-	virtual FAudioChunkHandle GetLoadedChunk(const USoundWave* SoundWave, uint32 ChunkIndex, bool bBlockForLoad = false, bool bForImmediatePlayback = false) const override;
+	virtual bool RequestChunk(const FSoundWaveProxy& SoundWave, uint32 ChunkIndex, TFunction<void(EAudioChunkLoadResult)> OnLoadCompleted, ENamedThreads::Type ThreadToCallOnLoadCompletedOn, bool bForImmediatePlayback = false) override;
+	virtual FAudioChunkHandle GetLoadedChunk(const FSoundWaveProxy&  SoundWave, uint32 ChunkIndex, bool bBlockForLoad = false, bool bForImmediatePlayback = false) const override;
 	virtual uint64 TrimMemory(uint64 NumBytesToFree) override;
 	virtual int32 RenderStatAudioStreaming(UWorld* World, FViewport* Viewport, FCanvas* Canvas, int32 X, int32 Y, const FVector* ViewLocation, const FRotator* ViewRotation) override;
 	virtual FString GenerateMemoryReport() override;
 	virtual void SetProfilingMode(bool bEnabled) override;
 	// End IAudioStreamingManager interface
 
+
 protected:
 
 	// These are used to reference count consumers of audio chunks.
@@ -454,13 +500,13 @@ protected:
 	 * based on the size of this sound wave's chunk,
 	 * or nullptr if MemoryLoadOnDemand is disabled.
 	 */
-	FAudioChunkCache* GetCacheForWave(const USoundWave* InSoundWave) const;
+	FAudioChunkCache* GetCacheForWave(const FSoundWaveProxy&  InSoundWave) const;
 	FAudioChunkCache* GetCacheForChunkSize(uint32 InChunkSize) const;
 
 	/**
 	 * Returns the next chunk to kick off a load for, or INDEX_NONE if there is only one chunk to cache.
 	 */
-	int32 GetNextChunkIndex(const USoundWave* InSoundWave, uint32 CurrentChunkIndex) const;
+	int32 GetNextChunkIndex(const FSoundWaveProxy&  InSoundWave, uint32 CurrentChunkIndex) const;
 
 	/** Audio chunk caches. These are set up on initialization. */
 	TArray<FAudioChunkCache> CacheArray;
@@ -468,8 +514,3 @@ protected:
 	
 
 };
-
-inline int32 GetTypeHash(const FAudioChunkCache::FChunkKey& InKey)
-{
-	return HashCombine(InKey.SoundWaveName.GetNumber(), InKey.ChunkIndex);
-}
diff --git a/Engine/Source/Runtime/Engine/Public/ContentStreaming.h b/Engine/Source/Runtime/Engine/Public/ContentStreaming.h
--- a/Engine/Source/Runtime/Engine/Public/ContentStreaming.h
+++ b/Engine/Source/Runtime/Engine/Public/ContentStreaming.h
@@ -14,7 +14,7 @@
 class AActor;
 class FSoundSource;
 class UPrimitiveComponent;
-class USoundWave;
+class FSoundWaveProxy;
 class ICompressedAudioInfo;
 class UTexture2D;
 struct FRenderAssetStreamingManager;
@@ -35,7 +35,7 @@ class USkeletalMesh;
 class ULandscapeLODStreamingProxy;
 class UStreamableRenderAsset;
 class FSoundSource;
-class USoundWave;
+class FSoundWaveProxy;
 struct FWaveInstance;
 struct FRenderAssetStreamingManager;
 
@@ -117,21 +117,17 @@ public:
 
 private:
 	// This constructor should only be called by an implementation of IAudioStreamingManager.
-	FAudioChunkHandle(const uint8* InData, uint32 NumBytes, const USoundWave* InSoundWave, const FName& SoundWaveName, uint32 InChunkIndex, uint64 InCacheLookupID);
+	FAudioChunkHandle(const uint8* InData, uint32 NumBytes, const FSoundWaveProxy&  InSoundWave, const FName& SoundWaveName, uint32 InChunkIndex, uint64 InCacheLookupID);
 
 	const uint8*  CachedData;
 	int32 CachedDataNumBytes;
 
-	const USoundWave* CorrespondingWave;
+	TUniquePtr<FSoundWaveProxy> CorrespondingWave;
 	FName CorrespondingWaveName;
 
 	// The index of this chunk in the sound wave's full set of chunks of compressed audio.
 	int32 ChunkIndex;
 
-	// This ID can be used to access the element this handle is for directly,
-	// rather than linearly searching the cache. This should only be used by the stream cache itself.
-	uint64 CacheLookupID;
-
 #if WITH_EDITOR
 	uint32 ChunkGeneration;
 #endif
@@ -502,10 +498,10 @@ enum class EAudioChunkLoadResult : uint8
 struct IAudioStreamingManager : public IStreamingManager
 {
 	/** Adds a new Sound Wave to the streaming manager. */
-	virtual void AddStreamingSoundWave(USoundWave* SoundWave) = 0;
+	virtual void AddStreamingSoundWave(const FSoundWaveProxy& SoundWave) = 0;
 
 	/** Removes a Sound Wave from the streaming manager. */
-	virtual void RemoveStreamingSoundWave(USoundWave* SoundWave) = 0;
+	virtual void RemoveStreamingSoundWave(const FSoundWaveProxy& SoundWave) = 0;
 
 	/** Adds the decoder to the streaming manager to prevent stream chunks from getting reaped from underneath it */
 	virtual void AddDecoder(ICompressedAudioInfo* CompressedAudioInfo) = 0;
@@ -514,10 +510,10 @@ struct IAudioStreamingManager : public IStreamingManager
 	virtual void RemoveDecoder(ICompressedAudioInfo* CompressedAudioInfo) = 0;
 
 	/** Returns true if this is a Sound Wave that is managed by the streaming manager. */
-	virtual bool IsManagedStreamingSoundWave(const USoundWave* SoundWave) const = 0;
+	virtual bool IsManagedStreamingSoundWave(const FSoundWaveProxy&  SoundWave) const = 0;
 
 	/** Returns true if this Sound Wave is currently streaming a chunk. */
-	virtual bool IsStreamingInProgress(const USoundWave* SoundWave) = 0;
+	virtual bool IsStreamingInProgress(const FSoundWaveProxy&  SoundWave) = 0;
 
 	virtual bool CanCreateSoundSource(const FWaveInstance* WaveInstance) const = 0;
 
@@ -538,7 +534,7 @@ struct IAudioStreamingManager : public IStreamingManager
 	 * @param ThreadToCallOnLoadCompleteOn. Optional specifier for which thread OnLoadCompleted should be called on.
 	 * @param bForImmediatePlaybac if true, this will optionally reprioritize this chunk's load request.
 	 */
-	virtual bool RequestChunk(USoundWave* SoundWave, uint32 ChunkIndex, TFunction<void(EAudioChunkLoadResult)> OnLoadCompleted = [](EAudioChunkLoadResult) {}, ENamedThreads::Type ThreadToCallOnLoadCompletedOn = ENamedThreads::AnyThread, bool bForImmediatePlayback = false) = 0;
+	virtual bool RequestChunk(const FSoundWaveProxy& SoundWave, uint32 ChunkIndex, TFunction<void(EAudioChunkLoadResult)> OnLoadCompleted = [](EAudioChunkLoadResult) {}, ENamedThreads::Type ThreadToCallOnLoadCompletedOn = ENamedThreads::AnyThread, bool bForImmediatePlayback = false) = 0;
 
 	/**
 	 * Gets a pointer to a chunk of audio data
@@ -549,7 +545,7 @@ struct IAudioStreamingManager : public IStreamingManager
 	 * @param bForImmediatePlayback if true, will optionally reprioritize this chunk's load request. See au.streamcaching.PlaybackRequestPriority.
 	 * @return a handle to the loaded chunk. Can return a default constructed FAudioChunkHandle if the chunk is not loaded yet.
 	 */
-	virtual FAudioChunkHandle GetLoadedChunk(const USoundWave* SoundWave, uint32 ChunkIndex,  bool bBlockForLoad = false, bool bForImmediatePlayback = false) const = 0;
+	virtual FAudioChunkHandle GetLoadedChunk(const FSoundWaveProxy&  SoundWave, uint32 ChunkIndex,  bool bBlockForLoad = false, bool bForImmediatePlayback = false) const = 0;
 
 	/**
 	 * This will start evicting elements from the cache until either hit our target of bytes or run out of chunks we can free.
@@ -578,7 +574,7 @@ protected:
 	friend FAudioChunkHandle;
 
 	/** This can be called by implementers of IAudioStreamingManager to construct an FAudioChunkHandle using an otherwise inaccessible constructor. */
-	static FAudioChunkHandle BuildChunkHandle(const uint8* InData, uint32 NumBytes, const USoundWave* InSoundWave, const FName& SoundWaveName, uint32 InChunkIndex, uint64 CacheLookupID);
+	static FAudioChunkHandle BuildChunkHandle(const uint8* InData, uint32 NumBytes, const FSoundWaveProxy&  InSoundWave, const FName& SoundWaveName, uint32 InChunkIndex, uint64 CacheLookupID);
 
 	/**
 	 * This can be used to increment reference counted handles to audio chunks. Called by the copy constructor of FAudioChunkHandle.
diff --git a/Engine/Source/Runtime/Engine/Public/VorbisAudioInfo.h b/Engine/Source/Runtime/Engine/Public/VorbisAudioInfo.h
--- a/Engine/Source/Runtime/Engine/Public/VorbisAudioInfo.h
+++ b/Engine/Source/Runtime/Engine/Public/VorbisAudioInfo.h
@@ -105,10 +105,10 @@ public:
 	// End of ICompressedAudioInfo Interface
 
 protected:
-	virtual bool StreamCompressedInfoInternal(USoundWave* Wave, struct FSoundQualityInfo* QualityInfo) override;
+	virtual bool StreamCompressedInfoInternal(const FSoundWaveProxy& InWaveProxy, struct FSoundQualityInfo* QualityInfo) override;
 
 private:
-	const uint8* GetLoadedChunk(USoundWave* InSoundWave, uint32 ChunkIndex, uint32& OutChunkSize);
+	const uint8* GetLoadedChunk(FSoundWaveProxy& InSoundWave, uint32 ChunkIndex, uint32& OutChunkSize);
 
 	struct FVorbisFileWrapper* VFWrapper;
 	const uint8* SrcBufferData;
