Made DEFINE_ENUM_TO_STRING and DECLARE_ENUM_TO_STRING macros available even when visual log is not enabled. The main purpose being to avoid annoying compilation issues in shipping/test builds.

diff --git a/Engine/Source/Runtime/Engine/Public/VisualLogger/VisualLoggerTypes.h b/Engine/Source/Runtime/Engine/Public/VisualLogger/VisualLoggerTypes.h
--- a/Engine/Source/Runtime/Engine/Public/VisualLogger/VisualLoggerTypes.h
+++ b/Engine/Source/Runtime/Engine/Public/VisualLogger/VisualLoggerTypes.h
@@ -9,18 +9,12 @@ class AActor;
 class UCanvas;
 struct FLogEntryItem;
 
-
-#if ENABLE_VISUAL_LOG
 #define DEFINE_ENUM_TO_STRING(EnumType) FString EnumToString(const EnumType Value) \
 { \
 	static const UEnum* TypeEnum = FindObject<UEnum>(ANY_PACKAGE, TEXT(#EnumType)); \
 	return TypeEnum->GetNameStringByIndex(static_cast<int32>(Value)); \
 }
 #define DECLARE_ENUM_TO_STRING(EnumType) FString EnumToString(const EnumType Value)
-#else
-#define DEFINE_ENUM_TO_STRING(EnumType) 
-#define DECLARE_ENUM_TO_STRING(EnumType)
-#endif // ENABLE_VISUAL_LOG
 
 enum class ECreateIfNeeded : int8
 {
