Added emissive color on strata slab node and unified the way it is process for deferred and forward.
Also unified Emissive to use the EmissiveColor name on pin for UX.

diff --git a/Engine/Shaders/Private/BasePassPixelShader.usf b/Engine/Shaders/Private/BasePassPixelShader.usf
--- a/Engine/Shaders/Private/BasePassPixelShader.usf
+++ b/Engine/Shaders/Private/BasePassPixelShader.usf
@@ -1550,11 +1550,14 @@ void FPixelShaderInOut_MainPS(
 	uint2 PixelPos = In.SvPosition.xy;
 	uint PixelStrataDataByteOffset = GetStrataPixelDataByteOffset(PixelPos, ResolvedView.BufferSizeAndInvSize.x, OpaqueBasePass.Strata.MaxBytesPerPixel);
 
+	float3 EmissiveLuminance = 0.0f;
 	const float3 VertexNormal = MaterialParameters.TangentToWorld[2].xyz;
 	PackStrataOut(
 		OpaqueBasePass.Strata.MaterialLobesBufferUAV, PixelStrataDataByteOffset, 
 		StrataPixelHeader, Strata,
-		MaterialParameters.CameraVector, GBuffer);
+		MaterialParameters.CameraVector, GBuffer, EmissiveLuminance);
+
+	Color += EmissiveLuminance;
 
 	// Since we don't carry SSS on/off through the ShadingModel value (DefaultLit/SubSurfaceProfile), we update subsurface toggle, based on the MFP radius value
 	bSubsurfaceEnable = GBuffer.CustomData.a > 0;
diff --git a/Engine/Shaders/Private/Strata/Strata.ush b/Engine/Shaders/Private/Strata/Strata.ush
--- a/Engine/Shaders/Private/Strata/Strata.ush
+++ b/Engine/Shaders/Private/Strata/Strata.ush
@@ -307,6 +307,7 @@ FStrataData GetInitialisedStrataData()
 FStrataData GetStrataSlabBSDF(
 	float3 BaseColor, float Specular, float Metallic, float2 Roughness,
 	float SSSProfileID, float3 SSSDMFPAlbedo, float SSSDMFPRadius, float SSSRadiusScale,
+	float3 Emissive,
 	uint SharedNormalIndex, inout uint SharedNormalTypes)
 {
 	FStrataData StrataData = GetInitialisedStrataData();
@@ -325,6 +326,7 @@ FStrataData GetStrataSlabBSDF(
 	StrataData.Layers[0].BSDFCount++;
 	BSDF_SETTYPE			(StrataData.Layers[0].BSDFs[0],		STRATA_BSDF_TYPE_SLAB);
 	BSDF_SETNORMALID		(StrataData.Layers[0].BSDFs[0],		SharedNormalIndex);
+	BSDF_SETEMISSIVE		(StrataData.Layers[0].BSDFs[0],		Emissive);
 	BSDF_SETANISOTROPY		(StrataData.Layers[0].BSDFs[0],		bHasAnisotropy);
 	BSDF_SETISTOPLAYER		(StrataData.Layers[0].BSDFs[0],		0);
 	BSDF_SETHASSCATTERING	(StrataData.Layers[0].BSDFs[0],		bHasSSS ? 1 : 0);
@@ -622,6 +624,7 @@ FStrataData StrataVerticalLayering(FStrataData Top, FStrataData Base)
 
 void SanitizeStrataSlab(inout FStrataBSDF BSDF)
 {
+	BSDF_SETEMISSIVE(BSDF, clamp(BSDF_GETEMISSIVE(BSDF), 0.0, Max10BitsFloat));
 	SLAB_BASECOLOR(BSDF) = saturate(SLAB_BASECOLOR(BSDF));
 	SLAB_SPECULAR(BSDF) = saturate(SLAB_SPECULAR(BSDF));
 	SLAB_METALLIC(BSDF) = saturate(SLAB_METALLIC(BSDF));
@@ -1018,8 +1021,10 @@ void PackStrataOut(
 	FStrataPixelHeader StrataPixelHeader,
 	FStrataData Strata,
 	float3 V,
-	inout FGBufferData GBuffer)
+	inout FGBufferData GBuffer,
+	inout float3 EmissiveColor)
 {
+	EmissiveColor = 0.0f;
 
 	const float OpaqueBSDFThroughput = 0.0f;
 	const float FullThroughput = 1.0f;
@@ -1094,6 +1099,8 @@ void PackStrataOut(
 					SSS.DMFPRadius			= SLAB_SSSDMFPRADIUS(BSDF);
 				}
 
+				EmissiveColor += BSDF_GETEMISSIVE(BSDF) * Throughput;
+
 				LayerCoverage	+= BSDFCoverage;
 				LayerThroughput += BSDFCoverage * OpaqueBSDFThroughput; // STRATA_TODO: take into accound SSS transmittance
 				break;
@@ -1147,6 +1154,8 @@ void PackStrataOut(
 					TopLayer.BaseColor += BSDFCoverage * EnvBRDFApprox(F0, SafeRoughness, saturate(ContextNoLight.NoV));
 				}
 
+				EmissiveColor += BSDF_GETEMISSIVE(BSDF) * Throughput;
+
 				// Water is considered opaque, it has a custom process to simulate translucency
 				LayerCoverage += BSDFCoverage;
 				LayerThroughput += BSDFCoverage * OpaqueBSDFThroughput;
diff --git a/Engine/Shaders/Private/Strata/StrataEvaluation.ush b/Engine/Shaders/Private/Strata/StrataEvaluation.ush
--- a/Engine/Shaders/Private/Strata/StrataEvaluation.ush
+++ b/Engine/Shaders/Private/Strata/StrataEvaluation.ush
@@ -131,12 +131,16 @@ void StrataUpdateBSDFContext(inout FStrataBSDFContext BSDFContext, float3 NewL)
 
 struct FStrataEvaluateResult
 {
+	float3 EmissivePathValue;
+
 	float3 DiffusePathValue;
 	float3 SpecularPathValue;
 	float3 TransmissionPathValue;
+
 	float  DiffusePDF;
 	float  SpecularPDF;
 	float  TransmissionPDF;
+
 	float3 Throughput;					// Throughput to the next layer
 	bool   bSubsurface;					// True if we need to separate the subsurface light contribution for the screen space diffusion process.
 	bool   bApplyProjectedSolidAngle;	// True if the Saturate(NoL) factor should be applied or not over the value of the BSDF
@@ -204,6 +208,8 @@ FStrataEvaluateResult StrataEvaluateBSDF(FStrataBSDFContext BSDFContext, FShadow
 			Sample.SpecularPathValue= D * Vis * F;
 			Sample.SpecularPDF		= PDF;
 
+			Sample.EmissivePathValue= BSDF_GETEMISSIVE(BSDFContext.BSDF);
+
 			Sample.Throughput = OpaqueBSDFThroughput;
 			break;
 		}
diff --git a/Engine/Shaders/Private/Strata/StrataForwardLighting.ush b/Engine/Shaders/Private/Strata/StrataForwardLighting.ush
--- a/Engine/Shaders/Private/Strata/StrataForwardLighting.ush
+++ b/Engine/Shaders/Private/Strata/StrataForwardLighting.ush
@@ -66,7 +66,8 @@ float3 StrataForwardLighting(
 
 			// Evalaute the lighting coming from L (it will lso contain the BSDF throughput to the next layer)
 			FStrataEvaluateResult StrataEvaluate = StrataEvaluateBSDF(StrataBSDFContext);
-			Color += (StrataEvaluate.DiffusePathValue + StrataEvaluate.SpecularPathValue) * DirLightGeometryFactor;
+			Color += (StrataEvaluate.DiffusePathValue + StrataEvaluate.SpecularPathValue) * DirLightGeometryFactor 
+					+ StrataEvaluate.EmissivePathValue * OutThroughput;
 
 			LayerCoverage	+= BSDFCoverage;
 			LayerThroughput += BSDFCoverage * StrataEvaluate.Throughput;
diff --git a/Engine/Source/Runtime/Engine/Classes/Materials/MaterialExpressionStrata.h b/Engine/Source/Runtime/Engine/Classes/Materials/MaterialExpressionStrata.h
--- a/Engine/Source/Runtime/Engine/Classes/Materials/MaterialExpressionStrata.h
+++ b/Engine/Source/Runtime/Engine/Classes/Materials/MaterialExpressionStrata.h
@@ -97,6 +97,12 @@ class UMaterialExpressionStrataSlabBSDF : public UMaterialExpressionStrataBSDF
 	UPROPERTY(meta = (DisplayName = "SSS Radius Scale"))
 	FExpressionInput SSSRadiusScale;
 
+	/**
+	 * Emissive color on top of the surface (type = float3, unit = luminance, default = 0)
+	 */
+	UPROPERTY()
+	FExpressionInput EmissiveColor;
+
 	/** SubsurfaceProfile, for Screen Space Subsurface Scattering. The profile needs to be set up on both the Strata diffuse node, and the material node at the moment. */
 	UPROPERTY(EditAnywhere, BlueprintReadOnly, Category = Material, meta = (DisplayName = "Subsurface Profile"))
 	class USubsurfaceProfile* SubsurfaceProfile;
@@ -169,10 +175,10 @@ class UMaterialExpressionStrataVolumetricFogCloudBSDF : public UMaterialExpressi
 	FExpressionInput Extinction;
 
 	/**
-	 * Emissive (type = float3, unit = luminance, default = 0)
+	 * Emissive color of the medium (type = float3, unit = luminance, default = 0)
 	 */
 	UPROPERTY()
-	FExpressionInput Emissive;
+	FExpressionInput EmissiveColor;
 
 	/**
 	 * Ambient occlusion: 1 means no occlusion while 0 means fully occluded. (type = float, unit = unitless, default = 1)
@@ -198,7 +204,7 @@ class UMaterialExpressionStrataUnlitBSDF : public UMaterialExpressionStrataBSDF
 	GENERATED_UCLASS_BODY()
 
 	/**
-	* Emissive (type = float3, unit = Luminance, default = 0)
+	* Emissive color on top of the surface (type = float3, unit = Luminance, default = 0)
 	*/
 	UPROPERTY()
 	FExpressionInput EmissiveColor;
@@ -263,10 +269,10 @@ class UMaterialExpressionStrataHairBSDF : public UMaterialExpressionStrataBSDF
 	FExpressionInput Tangent;
 
 	/**
-	 * Emissive (type = float3, unit = luminance, defaults to 0.0)
+	 * Emissive color on top of the surface (type = float3, unit = luminance, defaults to 0.0)
 	 */
 	UPROPERTY()
-	FExpressionInput Emissive;
+	FExpressionInput EmissiveColor;
 
 	//~ Begin UMaterialExpression Interface
 #if WITH_EDITOR
@@ -316,10 +322,10 @@ class UMaterialExpressionStrataSingleLayerWaterBSDF : public UMaterialExpression
 	FExpressionInput Normal;
 
 	/**
-	 * Emissive RGB (type = float3, unit = luminance, defaults to 0.0)
+	 * Emissive color on top of the surface (type = float3, unit = luminance, defaults to 0.0)
 	 */
 	UPROPERTY()
-	FExpressionInput Emissive;
+	FExpressionInput EmissiveColor;
 
 	/**
 	 * Opacity of the material layered on top of the water (type = float3, unit = unitless, defaults to 0.0)
diff --git a/Engine/Source/Runtime/Engine/Private/Materials/HLSLMaterialTranslator.cpp b/Engine/Source/Runtime/Engine/Private/Materials/HLSLMaterialTranslator.cpp
--- a/Engine/Source/Runtime/Engine/Private/Materials/HLSLMaterialTranslator.cpp
+++ b/Engine/Source/Runtime/Engine/Private/Materials/HLSLMaterialTranslator.cpp
@@ -8527,10 +8527,10 @@ int32 FHLSLMaterialTranslator::StrataCreateAndRegisterNullMaterial()
 	return OutputCodeChunk;
 }
 
-int32 FHLSLMaterialTranslator::StrataSlabBSDF(int32 BaseColor, int32 Reflectivity, int32 Metallic, int32 RoughnessX, int32 RoughnessY, int32 SSSProfileId, int32 SSSDMFPAlbedo, int32 SSSDMFPRadius, int32 SSSRadiusScale, int32 Normal, int32 Tangent, uint8 SharedNormalIndex)
+int32 FHLSLMaterialTranslator::StrataSlabBSDF(int32 BaseColor, int32 Reflectivity, int32 Metallic, int32 RoughnessX, int32 RoughnessY, int32 SSSProfileId, int32 SSSDMFPAlbedo, int32 SSSDMFPRadius, int32 SSSRadiusScale, int32 EmissiveColor, int32 Normal, int32 Tangent, uint8 SharedNormalIndex)
 {
 	return AddCodeChunk(
-		MCT_Strata, TEXT("GetStrataSlabBSDF(%s, %s, %s, float2(%s, %s), %s, %s, %s, %s, %u, Parameters.SharedNormals.NormalTypes) /* Normal = %s ; Tangent = %s */"),
+		MCT_Strata, TEXT("GetStrataSlabBSDF(%s, %s, %s, float2(%s, %s), %s, %s, %s, %s, %s, %u, Parameters.SharedNormals.NormalTypes) /* Normal = %s ; Tangent = %s */"),
 		*GetParameterCode(BaseColor),
 		*GetParameterCode(Reflectivity),
 		*GetParameterCode(Metallic),
@@ -8540,6 +8540,7 @@ int32 FHLSLMaterialTranslator::StrataSlabBSDF(int32 BaseColor, int32 Reflectivit
 		*GetParameterCode(SSSDMFPAlbedo),
 		*GetParameterCode(SSSDMFPRadius),
 		*GetParameterCode(SSSRadiusScale),
+		*GetParameterCode(EmissiveColor),
 		SharedNormalIndex,
 		*GetParameterCode(Normal),
 		*GetParameterCode(Tangent)
@@ -8557,27 +8558,27 @@ int32 FHLSLMaterialTranslator::StrataSheenBSDF(int32 BaseColor, int32 Roughness,
 	);
 }
 
-int32 FHLSLMaterialTranslator::StrataVolumetricFogCloudBSDF(int32 Albedo, int32 Extinction, int32 Emissive, int32 AmbientOcclusion)
+int32 FHLSLMaterialTranslator::StrataVolumetricFogCloudBSDF(int32 Albedo, int32 Extinction, int32 EmissiveColor, int32 AmbientOcclusion)
 {
 	return AddCodeChunk(
 		MCT_Strata, TEXT("GetStrataVolumeFogCloudBSDF(%s, %s, %s, %s)"),
 		*GetParameterCode(Albedo),
 		*GetParameterCode(Extinction),
-		*GetParameterCode(Emissive),
+		*GetParameterCode(EmissiveColor),
 		*GetParameterCode(AmbientOcclusion)
 	);
 }
 
-int32 FHLSLMaterialTranslator::StrataUnlitBSDF(int32 Emissive, int32 TransmittanceColor)
+int32 FHLSLMaterialTranslator::StrataUnlitBSDF(int32 EmissiveColor, int32 TransmittanceColor)
 {
 	return AddCodeChunk(
 		MCT_Strata, TEXT("GetStrataUnlitBSDF(%s, %s)"),
-		*GetParameterCode(Emissive),
+		*GetParameterCode(EmissiveColor),
 		*GetParameterCode(TransmittanceColor)
 	);
 }
 
-int32 FHLSLMaterialTranslator::StrataHairBSDF(int32 BaseColor, int32 Scatter, int32 Specular, int32 Roughness, int32 Backlit, int32 Emissive, int32 Tangent, uint8 SharedNormalIndex)
+int32 FHLSLMaterialTranslator::StrataHairBSDF(int32 BaseColor, int32 Scatter, int32 Specular, int32 Roughness, int32 Backlit, int32 EmissiveColor, int32 Tangent, uint8 SharedNormalIndex)
 {
 	return AddCodeChunk(
 		MCT_Strata, TEXT("GetStrataHairBSDF(%s, %s, %s, %s, %s, %s, %u) /* %s */"),
@@ -8586,14 +8587,14 @@ int32 FHLSLMaterialTranslator::StrataHairBSDF(int32 BaseColor, int32 Scatter, in
 		*GetParameterCode(Specular),
 		*GetParameterCode(Roughness),
 		*GetParameterCode(Backlit),
-		*GetParameterCode(Emissive),
+		*GetParameterCode(EmissiveColor),
 		SharedNormalIndex,
 		*GetParameterCode(Tangent)
 	);
 }
 
 int32 FHLSLMaterialTranslator::StrataSingleLayerWaterBSDF(
-	int32 BaseColor, int32 Metallic, int32 Specular, int32 Roughness, int32 Emissive, int32 TopMaterialOpacity,
+	int32 BaseColor, int32 Metallic, int32 Specular, int32 Roughness, int32 EmissiveColor, int32 TopMaterialOpacity,
 	int32 WaterAlbedo, int32 WaterExtinction, int32 WaterPhaseG, int32 ColorScaleBehindWater, int32 Normal, uint8 SharedNormalIndex)
 {
 	return AddCodeChunk(
@@ -8602,7 +8603,7 @@ int32 FHLSLMaterialTranslator::StrataSingleLayerWaterBSDF(
 		*GetParameterCode(Metallic),
 		*GetParameterCode(Specular),
 		*GetParameterCode(Roughness),
-		*GetParameterCode(Emissive),
+		*GetParameterCode(EmissiveColor),
 		*GetParameterCode(TopMaterialOpacity),
 		*GetParameterCode(WaterAlbedo),
 		*GetParameterCode(WaterExtinction),
diff --git a/Engine/Source/Runtime/Engine/Private/Materials/HLSLMaterialTranslator.h b/Engine/Source/Runtime/Engine/Private/Materials/HLSLMaterialTranslator.h
--- a/Engine/Source/Runtime/Engine/Private/Materials/HLSLMaterialTranslator.h
+++ b/Engine/Source/Runtime/Engine/Private/Materials/HLSLMaterialTranslator.h
@@ -920,12 +920,12 @@ protected:
 
 	// Strata
 	virtual int32 StrataCreateAndRegisterNullMaterial() override;
-	virtual int32 StrataSlabBSDF(int32 BaseColor, int32 Reflectivity, int32 Metallic, int32 RoughnessX, int32 RoughnessY, int32 SSSProfileId, int32 SSSDMFPAlbedo, int32 SSSDMFPRadius, int32 SSSRadiusScale, int32 Normal, int32 Tangent, uint8 SharedNormalIndex) override;
+	virtual int32 StrataSlabBSDF(int32 BaseColor, int32 Reflectivity, int32 Metallic, int32 RoughnessX, int32 RoughnessY, int32 SSSProfileId, int32 SSSDMFPAlbedo, int32 SSSDMFPRadius, int32 EmissiveColor, int32 SSSRadiusScale, int32 Normal, int32 Tangent, uint8 SharedNormalIndex) override;
 	virtual int32 StrataSheenBSDF(int32 BaseColor, int32 Roughness, int32 Normal, uint8 SharedNormalIndex) override;
-	virtual int32 StrataVolumetricFogCloudBSDF(int32 Albedo, int32 Extinction, int32 Emissive, int32 AmbientOcclusion) override;
-	virtual int32 StrataUnlitBSDF(int32 Emissive, int32 TransmittanceColor) override;
-	virtual int32 StrataHairBSDF(int32 BaseColor, int32 Scatter, int32 Specular, int32 Roughness, int32 Backlit, int32 Emissive, int32 Tangent, uint8 SharedNormalIndex) override;
-	virtual int32 StrataSingleLayerWaterBSDF(int32 BaseColor, int32 Metallic, int32 Specular, int32 Roughness, int32 Emissive, int32 TopMaterialOpacity, int32 WaterAlbedo, int32 WaterExtinction, int32 WaterPhaseG, int32 ColorScaleBehindWater, int32 Normal, uint8 SharedNormalIndex) override;
+	virtual int32 StrataVolumetricFogCloudBSDF(int32 Albedo, int32 Extinction, int32 EmissiveColor, int32 AmbientOcclusion) override;
+	virtual int32 StrataUnlitBSDF(int32 EmissiveColor, int32 TransmittanceColor) override;
+	virtual int32 StrataHairBSDF(int32 BaseColor, int32 Scatter, int32 Specular, int32 Roughness, int32 Backlit, int32 EmissiveColor, int32 Tangent, uint8 SharedNormalIndex) override;
+	virtual int32 StrataSingleLayerWaterBSDF(int32 BaseColor, int32 Metallic, int32 Specular, int32 Roughness, int32 EmissiveColor, int32 TopMaterialOpacity, int32 WaterAlbedo, int32 WaterExtinction, int32 WaterPhaseG, int32 ColorScaleBehindWater, int32 Normal, uint8 SharedNormalIndex) override;
 	virtual int32 StrataHorizontalMixing(int32 Foreground, int32 Background, int32 Mix) override;
 	virtual int32 StrataVerticalLayering(int32 Top, int32 Base) override;
 	virtual int32 StrataAdd(int32 A, int32 B) override;
diff --git a/Engine/Source/Runtime/Engine/Private/Materials/Material.cpp b/Engine/Source/Runtime/Engine/Private/Materials/Material.cpp
--- a/Engine/Source/Runtime/Engine/Private/Materials/Material.cpp
+++ b/Engine/Source/Runtime/Engine/Private/Materials/Material.cpp
@@ -3738,7 +3738,7 @@ void UMaterial::ConvertMaterialToStrataMaterial()
 			UMaterialExpressionStrataVolumetricFogCloudBSDF* VolBSDF = NewObject<UMaterialExpressionStrataVolumetricFogCloudBSDF>(this);
 			MoveConnectionTo(BaseColor, VolBSDF, 0);		// Albedo
 			MoveConnectionTo(SubsurfaceColor, VolBSDF, 1);	// Extinction
-			MoveConnectionTo(EmissiveColor, VolBSDF, 2);	// Emissive
+			MoveConnectionTo(EmissiveColor, VolBSDF, 2);	// EmissiveColor
 			MoveConnectionTo(AmbientOcclusion, VolBSDF, 3);	// AmbientOcclusion
 
 			// STRATA_TODO remove the VolumetricAdvancedOutput node and add the input onto FogCloudBSDF even if only used by the cloud renderer?
diff --git a/Engine/Source/Runtime/Engine/Private/Materials/MaterialExpressions.cpp b/Engine/Source/Runtime/Engine/Private/Materials/MaterialExpressions.cpp
--- a/Engine/Source/Runtime/Engine/Private/Materials/MaterialExpressions.cpp
+++ b/Engine/Source/Runtime/Engine/Private/Materials/MaterialExpressions.cpp
@@ -19715,6 +19715,7 @@ int32 UMaterialExpressionStrataSlabBSDF::Compile(class FMaterialCompiler* Compil
 		CompileWithDefaultFloat3(Compiler, DiffuseMeanFreePathAlbedo, 0.0f, 0.0f, 0.0f),
 		CompileWithDefaultFloat1(Compiler, DiffuseMeanFreePathRadius, 0.0f),
 		CompileWithDefaultFloat1(Compiler, SSSRadiusScale, 1.0f),
+		CompileWithDefaultFloat3(Compiler, EmissiveColor, 0.0f, 0.0f, 0.0f),
 		NormalCodeChunk,
 		TangentCodeChunk,
 		SharedNormalIndex);
@@ -19767,6 +19768,9 @@ uint32 UMaterialExpressionStrataSlabBSDF::GetInputType(int32 InputIndex)
 	case 9:
 		return MCT_Float1; // SSS Radius Scale
 		break;
+	case 10:
+		return MCT_Float3; // Emissive Color
+		break;
 	}
 
 	check(false);
@@ -19815,6 +19819,10 @@ FName UMaterialExpressionStrataSlabBSDF::GetInputName(int32 InputIndex) const
 	{
 		return TEXT("SSS Radius Scale");
 	}
+	else if (InputIndex == 10)
+	{
+		return TEXT("Emissive Color");
+	}
 	return TEXT("Unknown");
 }
 
@@ -19938,7 +19946,7 @@ int32 UMaterialExpressionStrataVolumetricFogCloudBSDF::Compile(class FMaterialCo
 	int32 OutputCodeChunk = Compiler->StrataVolumetricFogCloudBSDF(
 		CompileWithDefaultFloat3(Compiler, Albedo, 0.0f, 0.0f, 0.0f),
 		CompileWithDefaultFloat3(Compiler, Extinction, 0.0f, 0.0f, 0.0f),
-		CompileWithDefaultFloat3(Compiler, Emissive, 0.0f, 0.0f, 0.0f),
+		CompileWithDefaultFloat3(Compiler, EmissiveColor, 0.0f, 0.0f, 0.0f),
 		CompileWithDefaultFloat1(Compiler, AmbientOcclusion, 1.0f));
 
 	uint8 FakeSharedNormalIndex = 0;
@@ -20085,7 +20093,7 @@ int32 UMaterialExpressionStrataHairBSDF::Compile(class FMaterialCompiler* Compil
 		CompileWithDefaultFloat1(Compiler, Specular,	0.5f),
 		CompileWithDefaultFloat1(Compiler, Roughness,	0.5f),
 		CompileWithDefaultFloat1(Compiler, Backlit,		0.0f),
-		CompileWithDefaultFloat3(Compiler, Emissive,	1.0f, 0.0f, 0.0f),
+		CompileWithDefaultFloat3(Compiler, EmissiveColor,1.0f, 0.0f, 0.0f),
 		TangentCodeChunk,
 		SharedNormalIndex);
 
@@ -20173,7 +20181,7 @@ int32 UMaterialExpressionStrataSingleLayerWaterBSDF::Compile(class FMaterialComp
 		CompileWithDefaultFloat1(Compiler, Metallic, 0.0f),
 		CompileWithDefaultFloat1(Compiler, Specular, 0.5f),
 		CompileWithDefaultFloat1(Compiler, Roughness, 0.5f),
-		CompileWithDefaultFloat3(Compiler, Emissive, 0.0f, 0.0f, 0.0f),
+		CompileWithDefaultFloat3(Compiler, EmissiveColor, 0.0f, 0.0f, 0.0f),
 		CompileWithDefaultFloat1(Compiler, TopMaterialOpacity, 0.0f),
 		CompileWithDefaultFloat3(Compiler, WaterAlbedo, 0.0f, 0.0f, 0.0f),
 		CompileWithDefaultFloat3(Compiler, WaterExtinction, 0.0f, 0.0f, 0.0f),
@@ -20217,7 +20225,7 @@ uint32 UMaterialExpressionStrataSingleLayerWaterBSDF::GetInputType(int32 InputIn
 		return MCT_Float3; // Normal
 		break;
 	case 5:
-		return MCT_Float3; // Emissive
+		return MCT_Float3; // Emissive Color
 		break;
 	case 6:
 		return MCT_Float1; // TopMaterialOpacity
diff --git a/Engine/Source/Runtime/Engine/Public/MaterialCompiler.h b/Engine/Source/Runtime/Engine/Public/MaterialCompiler.h
--- a/Engine/Source/Runtime/Engine/Public/MaterialCompiler.h
+++ b/Engine/Source/Runtime/Engine/Public/MaterialCompiler.h
@@ -392,12 +392,12 @@ public:
 
 	// Strata
 	virtual int32 StrataCreateAndRegisterNullMaterial() = 0;
-	virtual int32 StrataSlabBSDF(int32 BaseColor, int32 Reflectivity, int32 Metallic, int32 RoughnessX, int32 RoughnessY, int32 SSSProfileId, int32 SSSDMFPAlbedo, int32 SSSDMFPRadius, int32 SSSRadiusScale, int32 Normal, int32 Tangent, uint8 SharedNormalIndex) = 0;
+	virtual int32 StrataSlabBSDF(int32 BaseColor, int32 Reflectivity, int32 Metallic, int32 RoughnessX, int32 RoughnessY, int32 SSSProfileId, int32 SSSDMFPAlbedo, int32 SSSDMFPRadius, int32 SSSRadiusScale, int32 EmissiveColor, int32 Normal, int32 Tangent, uint8 SharedNormalIndex) = 0;
 	virtual int32 StrataSheenBSDF(int32 BaseColor, int32 Roughness, int32 Normal, uint8 SharedNormalIndex) = 0;
-	virtual int32 StrataVolumetricFogCloudBSDF(int32 Albedo, int32 Extinction, int32 Emissive, int32 AmbientOcclusion) = 0;
-	virtual int32 StrataUnlitBSDF(int32 Emissive, int32 TransmittanceColor) = 0;
-	virtual int32 StrataHairBSDF(int32 BaseColor, int32 Scatter, int32 Specular, int32 Roughness, int32 Backlit, int32 Emissive, int32 Tangent, uint8 SharedNormalIndex) = 0;
-	virtual int32 StrataSingleLayerWaterBSDF(int32 BaseColor, int32 Metallic, int32 Specular, int32 Roughness, int32 Emissive, int32 TopMaterialOpacity, int32 WaterAlbedo, int32 WaterExtinction, int32 WaterPhaseG, int32 ColorScaleBehindWater, int32 Normal, uint8 SharedNormalIndex) = 0;
+	virtual int32 StrataVolumetricFogCloudBSDF(int32 Albedo, int32 Extinction, int32 EmissiveColor, int32 AmbientOcclusion) = 0;
+	virtual int32 StrataUnlitBSDF(int32 EmissiveColor, int32 TransmittanceColor) = 0;
+	virtual int32 StrataHairBSDF(int32 BaseColor, int32 Scatter, int32 Specular, int32 Roughness, int32 Backlit, int32 EmissiveColor, int32 Tangent, uint8 SharedNormalIndex) = 0;
+	virtual int32 StrataSingleLayerWaterBSDF(int32 BaseColor, int32 Metallic, int32 Specular, int32 Roughness, int32 EmissiveColor, int32 TopMaterialOpacity, int32 WaterAlbedo, int32 WaterExtinction, int32 WaterPhaseG, int32 ColorScaleBehindWater, int32 Normal, uint8 SharedNormalIndex) = 0;
 	virtual int32 StrataHorizontalMixing(int32 Foreground, int32 Background, int32 Mix) = 0;
 	virtual int32 StrataVerticalLayering(int32 Top, int32 Base) = 0;
 	virtual int32 StrataAdd(int32 A, int32 B) = 0;
@@ -806,9 +806,9 @@ public:
 		return Compiler->StrataCreateAndRegisterNullMaterial();
 	}
 
-	virtual int32 StrataSlabBSDF(int32 BaseColor, int32 Reflectivity, int32 Metallic, int32 RoughnessX, int32 RoughnessY, int32 SSSProfileId, int32 SSSDMFPAlbedo, int32 SSSDMFPRadius, int32 SSSRadiusScale, int32 Normal, int32 Tangent, uint8 SharedNormalIndex) override
+	virtual int32 StrataSlabBSDF(int32 BaseColor, int32 Reflectivity, int32 Metallic, int32 RoughnessX, int32 RoughnessY, int32 SSSProfileId, int32 SSSDMFPAlbedo, int32 SSSDMFPRadius, int32 SSSRadiusScale, int32 EmissiveColor, int32 Normal, int32 Tangent, uint8 SharedNormalIndex) override
 	{
-		return Compiler->StrataSlabBSDF(BaseColor, Reflectivity, Metallic, RoughnessX, RoughnessY, SSSProfileId, SSSDMFPAlbedo, SSSDMFPRadius, SSSRadiusScale, Normal, Tangent, SharedNormalIndex);
+		return Compiler->StrataSlabBSDF(BaseColor, Reflectivity, Metallic, RoughnessX, RoughnessY, SSSProfileId, SSSDMFPAlbedo, SSSDMFPRadius, SSSRadiusScale, EmissiveColor, Normal, Tangent, SharedNormalIndex);
 	}
 
 	virtual int32 StrataSheenBSDF(int32 BaseColor, int32 Roughness, int32 Normal, uint8 SharedNormalIndex) override
@@ -816,24 +816,24 @@ public:
 		return Compiler->StrataSheenBSDF(BaseColor, Roughness, Normal, SharedNormalIndex);
 	}
 
-	virtual int32 StrataVolumetricFogCloudBSDF(int32 Albedo, int32 Extinction, int32 Emissive, int32 AmbientOcclusion) override
+	virtual int32 StrataVolumetricFogCloudBSDF(int32 Albedo, int32 Extinction, int32 EmissiveColor, int32 AmbientOcclusion) override
 	{
-		return Compiler->StrataVolumetricFogCloudBSDF(Albedo, Extinction, Emissive, AmbientOcclusion);
+		return Compiler->StrataVolumetricFogCloudBSDF(Albedo, Extinction, EmissiveColor, AmbientOcclusion);
 	}
 
-	virtual int32 StrataUnlitBSDF(int32 Emissive, int32 TransmittanceColor) override
+	virtual int32 StrataUnlitBSDF(int32 EmissiveColor, int32 TransmittanceColor) override
 	{
-		return Compiler->StrataUnlitBSDF(Emissive, TransmittanceColor);
+		return Compiler->StrataUnlitBSDF(EmissiveColor, TransmittanceColor);
 	}
 
-	virtual int32 StrataHairBSDF(int32 BaseColor, int32 Scatter, int32 Specular, int32 Roughness, int32 Backlit, int32 Emissive, int32 Tangent, uint8 SharedNormalIndex) override
+	virtual int32 StrataHairBSDF(int32 BaseColor, int32 Scatter, int32 Specular, int32 Roughness, int32 Backlit, int32 EmissiveColor, int32 Tangent, uint8 SharedNormalIndex) override
 	{
-		return Compiler->StrataHairBSDF(BaseColor, Scatter, Specular, Roughness, Backlit, Emissive, Tangent, SharedNormalIndex);
+		return Compiler->StrataHairBSDF(BaseColor, Scatter, Specular, Roughness, Backlit, EmissiveColor, Tangent, SharedNormalIndex);
 	}
 
-	virtual int32 StrataSingleLayerWaterBSDF(int32 BaseColor, int32 Metallic, int32 Specular, int32 Roughness, int32 Emissive, int32 TopMaterialOpacity, int32 WaterAlbedo, int32 WaterExtinction, int32 WaterPhaseG, int32 ColorScaleBehindWater, int32 Normal, uint8 SharedNormalIndex) override
+	virtual int32 StrataSingleLayerWaterBSDF(int32 BaseColor, int32 Metallic, int32 Specular, int32 Roughness, int32 EmissiveColor, int32 TopMaterialOpacity, int32 WaterAlbedo, int32 WaterExtinction, int32 WaterPhaseG, int32 ColorScaleBehindWater, int32 Normal, uint8 SharedNormalIndex) override
 	{
-		return Compiler->StrataSingleLayerWaterBSDF(BaseColor, Metallic, Specular, Roughness, Emissive, TopMaterialOpacity, WaterAlbedo, WaterExtinction, WaterPhaseG, ColorScaleBehindWater, Normal, SharedNormalIndex);
+		return Compiler->StrataSingleLayerWaterBSDF(BaseColor, Metallic, Specular, Roughness, EmissiveColor, TopMaterialOpacity, WaterAlbedo, WaterExtinction, WaterPhaseG, ColorScaleBehindWater, Normal, SharedNormalIndex);
 	}
 
 	virtual int32 StrataHorizontalMixing(int32 Foreground, int32 Background, int32 Mix) override
