BlendSpace 2.0: Blendspace Graph Node

diff --git a/Engine/Source/Runtime/Engine/Private/Animation/BlendSpaceBase.cpp b/Engine/Source/Runtime/Engine/Private/Animation/BlendSpaceBase.cpp
--- a/Engine/Source/Runtime/Engine/Private/Animation/BlendSpaceBase.cpp
+++ b/Engine/Source/Runtime/Engine/Private/Animation/BlendSpaceBase.cpp
@@ -115,131 +115,156 @@ void UBlendSpaceBase::PostEditChangeProperty( struct FPropertyChangedEvent& Prop
 }
 #endif // WITH_EDITOR
 
-void UBlendSpaceBase::TickAssetPlayer(FAnimTickRecord& Instance, struct FAnimNotifyQueue& NotifyQueue, FAnimAssetTickContext& Context) const
+
+bool UBlendSpaceBase::UpdateBlendSamples_Internal(const FVector& InBlendSpacePosition, float InDeltaTime, TArray<FBlendSampleData>& InOutOldSampleDataList, TArray<FBlendSampleData>& InOutSampleDataCache) const
 {
-	const float DeltaTime = Context.GetDeltaTime();
-	float MoveDelta = Instance.PlayRateMultiplier * DeltaTime;
+	// For Target weight interpolation, we'll need to save old data, and interpolate to new data
+	TArray<FBlendSampleData>& NewSampleDataList = FBlendSpaceScratchData::Get().NewSampleDataList;
+	check(!NewSampleDataList.Num()); // this must be called non-recursively
 
-	// this happens even if MoveDelta == 0.f. This still should happen if it is being interpolated
-	// since we allow setting position of blendspace, we can't ignore MoveDelta == 0.f
-	// also now we don't have to worry about not following if DeltaTime = 0.f
+	InOutOldSampleDataList.Append(InOutSampleDataCache);
+
+	// @fixme: temporary code to clear the invalid sample data 
+	// related jira: UE-71107
+	for (int32 Index = 0; Index < InOutOldSampleDataList.Num(); ++Index)
 	{
-		// first filter input using blend filter
-		const FVector BlendSpacePosition(Instance.BlendSpace.BlendSpacePositionX, Instance.BlendSpace.BlendSpacePositionY, 0.f);
-		const FVector BlendInput = FilterInput(Instance.BlendSpace.BlendFilter, BlendSpacePosition, DeltaTime);
-		EBlendSpaceAxis AxisToScale = GetAxisToScale();
-		if (AxisToScale != BSA_None)
+		if (!SampleData.IsValidIndex(InOutOldSampleDataList[Index].SampleDataIndex))
 		{
-			float FilterMultiplier = 1.f;
-			// first use multiplier using new blendinput
-			// new filtered input is going to be used for sampling animation
-			// so we'll need to change playrate if you'd like to not slide foot
-			if ( !BlendSpacePosition.Equals(BlendInput) )  
+			InOutOldSampleDataList.RemoveAt(Index);
+			--Index;
+		}
+	}
+		
+	// get sample data based on new input
+	// consolidate all samples and sort them, so that we can handle from biggest weight to smallest
+	InOutSampleDataCache.Reset();
+
+	// get sample data from blendspace
+	bool bSuccessfullySampled = false;
+	if (GetSamplesFromBlendInput(InBlendSpacePosition, NewSampleDataList))
+	{
+		// if target weight interpolation is set
+		if (TargetWeightInterpolationSpeedPerSec > 0.f)
+		{
+			// target weight interpolation
+			if (InterpolateWeightOfSampleData(InDeltaTime, InOutOldSampleDataList, NewSampleDataList, InOutSampleDataCache))
 			{
-				// apply speed change if you want, 
-				if (AxisToScale == BSA_X)
-				{
-					if (BlendInput.X != 0.f)
-					{
-						FilterMultiplier = BlendSpacePosition.X / BlendInput.X;
-					}
-				}
-				else if (AxisToScale == BSA_Y)
-				{
-					if (BlendInput.Y != 0.f)
-					{
-						FilterMultiplier = BlendSpacePosition.Y / BlendInput.Y;
-					}
-				}
+				// now I need to normalize
+				FBlendSampleData::NormalizeDataWeight(InOutSampleDataCache);
 			}
-
-			// Now find if clamped input is different. If different, then apply scale to fit in. This allows
-			// "extrapolation" of the blend space outside of the range by time scaling the animation, which is
-			// appropriate when the specified axis is speed (for example).
-			FVector ClampedInput = GetClampedBlendInput(BlendInput);
-			if ( !ClampedInput.Equals(BlendInput) ) 
+			else
 			{
-				// apply speed change if you want, 
-				if (AxisToScale == BSA_X && !BlendParameters[0].bWrapInput)
-				{
-					if (ClampedInput.X != 0.f)
-					{
-						FilterMultiplier *= BlendInput.X / ClampedInput.X;
-					}
-				}
-				else if (AxisToScale == BSA_Y)
-				{
-					if (ClampedInput.Y != 0.f && !BlendParameters[1].bWrapInput)
-					{
-						FilterMultiplier *= BlendInput.Y / ClampedInput.Y;
-					}
-				}
+				// if interpolation failed, just copy new sample data to sample data
+				InOutSampleDataCache = NewSampleDataList;
 			}
+		}
+		else
+		{
+			// when there is no target weight interpolation, just copy new to target
+			InOutSampleDataCache.Append(NewSampleDataList);
+		}
 
+		bSuccessfullySampled = true;
+	}
 
-			MoveDelta *= FilterMultiplier;
-			UE_LOG(LogAnimation, Log, TEXT("BlendSpace(%s) - BlendInput(%s) : FilteredBlendInput(%s), FilterMultiplier(%0.2f)"), *GetName(), *BlendSpacePosition.ToString(), *BlendInput.ToString(), FilterMultiplier );
-		}
+	NewSampleDataList.Reset();
 
-		check(Instance.BlendSpace.BlendSampleDataCache);
+	return bSuccessfullySampled;
+}
 
-		// For Target weight interpolation, we'll need to save old data, and interpolate to new data
-		TArray<FBlendSampleData>& OldSampleDataList = FBlendSpaceScratchData::Get().OldSampleDataList;
-		TArray<FBlendSampleData>& NewSampleDataList = FBlendSpaceScratchData::Get().NewSampleDataList;
-		check(!OldSampleDataList.Num() && !NewSampleDataList.Num()); // this must be called non-recursively
+bool UBlendSpaceBase::UpdateBlendSamples(const FVector& InBlendSpacePosition, float InDeltaTime, TArray<FBlendSampleData>& InOutSampleDataCache) const
+{
+	TArray<FBlendSampleData>& OldSampleDataList = FBlendSpaceScratchData::Get().OldSampleDataList;
+	check(!OldSampleDataList.Num()); // this must be called non-recursively
+	const bool bResult = UpdateBlendSamples_Internal(InBlendSpacePosition, InDeltaTime, OldSampleDataList, InOutSampleDataCache);
+	OldSampleDataList.Reset();
+	return bResult;
+}
 
-		OldSampleDataList.Append(*Instance.BlendSpace.BlendSampleDataCache);
+void UBlendSpaceBase::TickAssetPlayer(FAnimTickRecord& Instance, struct FAnimNotifyQueue& NotifyQueue, FAnimAssetTickContext& Context) const
+{
+	check(Instance.BlendSpace.BlendSampleDataCache);
 
-		// @fixme: temporary code to clear the invalid sample data 
-		// related jira: UE-71107
-		for (int32 Index = 0; Index < OldSampleDataList.Num(); ++Index)
-		{
-			if (!SampleData.IsValidIndex(OldSampleDataList[Index].SampleDataIndex))
-			{
-				OldSampleDataList.RemoveAt(Index);
-				--Index;
-			}
-		}
-		
-		// get sample data based on new input
-		// consolidate all samples and sort them, so that we can handle from biggest weight to smallest
-		Instance.BlendSpace.BlendSampleDataCache->Reset();
-		// new sample data that will be used for evaluation
-		TArray<FBlendSampleData> & SampleDataList = *Instance.BlendSpace.BlendSampleDataCache;
-
-		// get sample data from blendspace
-		if (GetSamplesFromBlendInput(BlendInput, NewSampleDataList))
+	// Scratch area for old samples
+	TArray<FBlendSampleData>& OldSampleDataList = FBlendSpaceScratchData::Get().OldSampleDataList;
+	check(!OldSampleDataList.Num()); // this must be called non-recursively
+	// new sample data that will be used for evaluation
+	TArray<FBlendSampleData>& SampleDataList = *Instance.BlendSpace.BlendSampleDataCache;
+
+	const float DeltaTime = Context.GetDeltaTime();
+	float MoveDelta = Instance.PlayRateMultiplier * DeltaTime;
+
+	// this happens even if MoveDelta == 0.f. This still should happen if it is being interpolated
+	// since we allow setting position of blendspace, we can't ignore MoveDelta == 0.f
+	// also now we don't have to worry about not following if DeltaTime = 0.f
+	{
+		// first filter input using blend filter
+		const FVector BlendSpacePosition(Instance.BlendSpace.BlendSpacePositionX, Instance.BlendSpace.BlendSpacePositionY, 0.f);
+		const FVector FilteredBlendInput = FilterInput(Instance.BlendSpace.BlendFilter, BlendSpacePosition, DeltaTime);
+
+		if(UpdateBlendSamples_Internal(FilteredBlendInput, DeltaTime, OldSampleDataList, SampleDataList))
 		{
 			float NewAnimLength=0.f;
 			float PreInterpAnimLength = 0.f;
 
-			// if target weight interpolation is set
 			if (TargetWeightInterpolationSpeedPerSec > 0.f)
 			{
-				UE_LOG(LogAnimation, Verbose, TEXT("Target Weight Interpolation: Target Samples "));
 				// recalculate AnimLength based on weight of target animations - this is used for scaling animation later (change speed)
-				PreInterpAnimLength = GetAnimationLengthFromSampleData(NewSampleDataList);
-				UE_LOG(LogAnimation, Verbose, TEXT("BlendSpace(%s) - BlendInput(%s) : PreAnimLength(%0.5f) "), *GetName(), *BlendInput.ToString(), PreInterpAnimLength);
+				PreInterpAnimLength = GetAnimationLengthFromSampleData(*Instance.BlendSpace.BlendSampleDataCache);
+				UE_LOG(LogAnimation, Verbose, TEXT("BlendSpace(%s) - FilteredBlendInput(%s) : PreAnimLength(%0.5f) "), *GetName(), *FilteredBlendInput.ToString(), PreInterpAnimLength);	
+			}
 
-				// target weight interpolation
-				if (InterpolateWeightOfSampleData(DeltaTime, OldSampleDataList, NewSampleDataList, SampleDataList))
+			EBlendSpaceAxis AxisToScale = GetAxisToScale();
+			if (AxisToScale != BSA_None)
+			{
+				float FilterMultiplier = 1.f;
+				// first use multiplier using new blendinput
+				// new filtered input is going to be used for sampling animation
+				// so we'll need to change playrate if you'd like to not slide foot
+				if ( !BlendSpacePosition.Equals(FilteredBlendInput) )  
 				{
-					// now I need to normalize
-					FBlendSampleData::NormalizeDataWeight(SampleDataList);
+					// apply speed change if you want, 
+					if (AxisToScale == BSA_X)
+					{
+						if (FilteredBlendInput.X != 0.f)
+						{
+							FilterMultiplier = BlendSpacePosition.X / FilteredBlendInput.X;
+						}
+					}
+					else if (AxisToScale == BSA_Y)
+					{
+						if (FilteredBlendInput.Y != 0.f)
+						{
+							FilterMultiplier = BlendSpacePosition.Y / FilteredBlendInput.Y;
+						}
+					}
 				}
-				else
+
+				// Now find if clamped input is different. If different, then apply scale to fit in. This allows
+				// "extrapolation" of the blend space outside of the range by time scaling the animation, which is
+				// appropriate when the specified axis is speed (for example).
+				FVector ClampedInput = GetClampedBlendInput(FilteredBlendInput);
+				if ( !ClampedInput.Equals(FilteredBlendInput) ) 
 				{
-					// if interpolation failed, just copy new sample data tto sample data
-					SampleDataList = NewSampleDataList;
+					// apply speed change if you want, 
+					if (AxisToScale == BSA_X && !BlendParameters[0].bWrapInput)
+					{
+						if (ClampedInput.X != 0.f)
+						{
+							FilterMultiplier *= FilteredBlendInput.X / ClampedInput.X;
+						}
+					}
+					else if (AxisToScale == BSA_Y)
+					{
+						if (ClampedInput.Y != 0.f && !BlendParameters[1].bWrapInput)
+						{
+							FilterMultiplier *= FilteredBlendInput.Y / ClampedInput.Y;
+						}
+					}
 				}
 
-				// recalculate AnimLength based on weight of animations
-				UE_LOG(LogAnimation, Verbose, TEXT("Target Weight Interpolation: Interp Samples "));
-			}
-			else
-			{
-				// when there is no target weight interpolation, just copy new to target
-				SampleDataList.Append(NewSampleDataList);
+				MoveDelta *= FilterMultiplier;
+				UE_LOG(LogAnimation, Log, TEXT("BlendSpace(%s) - FilteredBlendInput(%s) : FilteredBlendInput(%s), FilterMultiplier(%0.2f)"), *GetName(), *BlendSpacePosition.ToString(), *FilteredBlendInput.ToString(), FilterMultiplier );
 			}
 
 			bool bCanDoMarkerSync = (SampleIndexWithMarkers != INDEX_NONE) && (Context.IsSingleAnimationContext() || (Instance.bCanUseMarkerSync && Context.CanUseMarkerPosition()));
@@ -289,7 +314,7 @@ void UBlendSpaceBase::TickAssetPlayer(FAnimTickRecord& Instance, struct FAnimNot
 			if (Context.IsLeader())
 			{
 				// advance current time - blend spaces hold normalized time as when dealing with changing anim length it would be possible to go backwards
-				UE_LOG(LogAnimation, Verbose, TEXT("BlendSpace(%s) - BlendInput(%s) : AnimLength(%0.5f) "), *GetName(), *BlendInput.ToString(), NewAnimLength);
+				UE_LOG(LogAnimation, Verbose, TEXT("BlendSpace(%s) - FilteredBlendInput(%s) : AnimLength(%0.5f) "), *GetName(), *FilteredBlendInput.ToString(), NewAnimLength);
 				
 				Context.SetPreviousAnimationPositionRatio(NormalizedCurrentTime);
 
@@ -445,8 +470,8 @@ void UBlendSpaceBase::TickAssetPlayer(FAnimTickRecord& Instance, struct FAnimNot
 				}
 			}
 		}
+
 		OldSampleDataList.Reset();
-		NewSampleDataList.Reset();
 	}
 }
 
@@ -556,7 +581,7 @@ bool UBlendSpaceBase::IsValidAdditiveType(EAdditiveAnimationType AdditiveType) c
 	return false;
 }
 
-void UBlendSpaceBase::ResetToRefPose(FCompactPose& OutPose)
+void UBlendSpaceBase::ResetToRefPose(FCompactPose& OutPose) const
 {
 	if (IsValidAdditive())
 	{
@@ -568,14 +593,25 @@ void UBlendSpaceBase::ResetToRefPose(FCompactPose& OutPose)
 	}
 }
 
-void UBlendSpaceBase::GetAnimationPose(TArray<FBlendSampleData>& BlendSampleDataCache, /*out*/ FCompactPose& OutPose, /*out*/ FBlendedCurve& OutCurve)
+void UBlendSpaceBase::GetAnimationPose(TArray<FBlendSampleData>& BlendSampleDataCache, /*out*/ FCompactPose& OutPose, /*out*/ FBlendedCurve& OutCurve) const
 {
 	FStackCustomAttributes TempAttributes;
 	FAnimationPoseData AnimationPoseData = { OutPose, OutCurve, TempAttributes };
 	GetAnimationPose(BlendSampleDataCache, AnimationPoseData);
 }
 
-void UBlendSpaceBase::GetAnimationPose(TArray<FBlendSampleData>& BlendSampleDataCache, /*out*/ FAnimationPoseData& OutAnimationPoseData)
+void UBlendSpaceBase::GetAnimationPose(TArray<FBlendSampleData>& BlendSampleDataCache, /*out*/ FAnimationPoseData& OutAnimationPoseData) const
+{
+	GetAnimationPose_Internal(BlendSampleDataCache, TArrayView<FPoseLink>(), nullptr, false, OutAnimationPoseData);
+}
+
+void UBlendSpaceBase::GetAnimationPose(TArray<FBlendSampleData>& BlendSampleDataCache, TArrayView<FPoseLink> InPoseLinks, /*out*/ FPoseContext& Output) const
+{
+	FAnimationPoseData AnimationPoseData(Output);
+	GetAnimationPose_Internal(BlendSampleDataCache, InPoseLinks, Output.AnimInstanceProxy, Output.ExpectsAdditivePose(), AnimationPoseData);
+}
+
+void UBlendSpaceBase::GetAnimationPose_Internal(TArray<FBlendSampleData>& BlendSampleDataCache, TArrayView<FPoseLink> InPoseLinks, FAnimInstanceProxy* InProxy, bool bInExpectsAdditivePose, /*out*/ FAnimationPoseData& OutAnimationPoseData) const
 {
 	SCOPE_CYCLE_COUNTER(STAT_BlendSpace_GetAnimPose);
 	FScopeCycleCounterUObject BlendSpaceScope(this);
@@ -589,6 +625,7 @@ void UBlendSpaceBase::GetAnimationPose(TArray<FBlendSampleData>& BlendSampleData
 		return;
 	}
 
+	const bool bNested = InPoseLinks.Num() > 0;
 	const int32 NumPoses = BlendSampleDataCache.Num();
 
 	TArray<FCompactPose, TInlineAllocator<8>> ChildrenPoses;
@@ -610,7 +647,7 @@ void UBlendSpaceBase::GetAnimationPose(TArray<FBlendSampleData>& BlendSampleData
 	}
 
 	// get all child atoms we interested in
-	for(int32 I = 0; I < BlendSampleDataCache.Num(); ++I)
+	for(int32 I = 0; I < NumPoses; ++I)
 	{
 		FCompactPose& Pose = ChildrenPoses[I];
 
@@ -619,22 +656,38 @@ void UBlendSpaceBase::GetAnimationPose(TArray<FBlendSampleData>& BlendSampleData
 			const FBlendSample& Sample = SampleData[BlendSampleDataCache[I].SampleDataIndex];
 			ChildrenWeights[I] = BlendSampleDataCache[I].GetWeight();
 
-			if(Sample.Animation 
-#if WITH_EDITOR
-				// verify if Sample.Animation->GetSkeleton matches
-				&& ensure(Sample.Animation->GetSkeleton() == GetSkeleton())
-#endif // WITH_EDITOR
-			)
+			if(bNested)
 			{
-				const float Time = FMath::Clamp<float>(BlendSampleDataCache[I].Time, 0.f, Sample.Animation->GetPlayLength());
+				check(InPoseLinks.IsValidIndex(BlendSampleDataCache[I].SampleDataIndex));
 
-				FAnimationPoseData ChildAnimationPoseData = { Pose, ChildrenCurves[I], ChildrenAttributes[I] };
-				// first one always fills up the source one
-				Sample.Animation->GetAnimationPose(ChildAnimationPoseData, FAnimExtractContext(Time, true));
+				// Evaluate the linked graphs
+				FPoseContext ChildPoseContext(InProxy, bInExpectsAdditivePose);
+				InPoseLinks[BlendSampleDataCache[I].SampleDataIndex].Evaluate(ChildPoseContext);
+				
+				// Move out poses etc. for blending
+				ChildrenPoses[I] = MoveTemp(ChildPoseContext.Pose);
+				ChildrenCurves[I] = MoveTemp(ChildPoseContext.Curve);
+				ChildrenAttributes[I] = MoveTemp(ChildPoseContext.CustomAttributes);
 			}
 			else
 			{
-				ResetToRefPose(Pose);
+				if(Sample.Animation 
+#if WITH_EDITOR
+					// verify if Sample.Animation->GetSkeleton matches
+					&& ensure(Sample.Animation->GetSkeleton() == GetSkeleton())
+#endif // WITH_EDITOR
+				)
+				{
+					const float Time = FMath::Clamp<float>(BlendSampleDataCache[I].Time, 0.f, Sample.Animation->GetPlayLength());
+
+					FAnimationPoseData ChildAnimationPoseData = { Pose, ChildrenCurves[I], ChildrenAttributes[I] };
+					// first one always fills up the source one
+					Sample.Animation->GetAnimationPose(ChildAnimationPoseData, FAnimExtractContext(Time, true));
+				}
+				else
+				{
+					ResetToRefPose(Pose);
+				}
 			}
 		}
 		else
@@ -705,13 +758,7 @@ bool UBlendSpaceBase::GetSamplesFromBlendInput(const FVector &BlendInput, TArray
 		for(int32 Ind = 0; Ind < GridElement.MAX_VERTICES; ++Ind)
 		{
 			const int32 SampleDataIndex = GridElement.Indices[Ind];		
-			if( SampleData.IsValidIndex(SampleDataIndex) 
-#if WITH_EDITOR // we check these in editor because these could change when editor is running
-				&& SampleData[SampleDataIndex].bIsValid
-				&& SampleData[SampleDataIndex].Animation 
-				&& SampleData[SampleDataIndex].Animation->GetSkeleton() == GetSkeleton()
-#endif // WITH_EDITOR
-				)
+			if(SampleData.IsValidIndex(SampleDataIndex))
 			{
 				int32 Index = OutSampleDataList.AddUnique(SampleDataIndex);
 				FBlendSampleData& NewSampleData = OutSampleDataList[Index];
@@ -730,8 +777,9 @@ bool UBlendSpaceBase::GetSamplesFromBlendInput(const FVector &BlendInput, TArray
 		for (int32 Index2 = Index1 + 1; Index2 < OutSampleDataList.Num(); ++Index2)
 		{
 			FBlendSampleData& SecondSample = OutSampleDataList[Index2];
-			// if they have sample sample, remove the Index2, and get out
-			if (FirstSample.Animation == SecondSample.Animation)
+			// if they have same sample, remove the Index2, and get out
+			if (FirstSample.SampleDataIndex == SecondSample.SampleDataIndex || 
+				(FirstSample.Animation != nullptr && FirstSample.Animation == SecondSample.Animation))
 			{
 				//Calc New Sample Playrate
 				const float TotalWeight = FirstSample.GetWeight() + SecondSample.GetWeight();
@@ -808,8 +856,6 @@ void UBlendSpaceBase::ValidateSampleData()
 	{
 		FBlendSample& Sample = SampleData[SampleIndex];
 
-		Sample.bIsValid = ValidateSampleValue(Sample.SampleValue, SampleIndex) && (Sample.Animation != nullptr);
-
 		// see if same data exists, by same, same values
 		for (int32 ComparisonSampleIndex = SampleIndex + 1; ComparisonSampleIndex < SampleData.Num(); ++ComparisonSampleIndex)
 		{
@@ -822,64 +868,73 @@ void UBlendSpaceBase::ValidateSampleData()
 			}
 		}
 
-		if (Sample.bIsValid)
+		if(IsAsset())
 		{
-			if (Sample.Animation->GetPlayLength() > AnimLength)
+			Sample.bIsValid = ValidateSampleValue(Sample.SampleValue, SampleIndex) && (Sample.Animation != nullptr);
+
+			if (Sample.bIsValid)
 			{
-				// @todo : should apply scale? If so, we'll need to apply also when blend
-				AnimLength = Sample.Animation->GetPlayLength();
-			}
+				if (Sample.Animation->GetPlayLength() > AnimLength)
+				{
+					// @todo : should apply scale? If so, we'll need to apply also when blend
+					AnimLength = Sample.Animation->GetPlayLength();
+				}
 
-			Sample.CachedMarkerDataUpdateCounter = Sample.Animation->GetMarkerUpdateCounter();
+				Sample.CachedMarkerDataUpdateCounter = Sample.Animation->GetMarkerUpdateCounter();
 
-			if (Sample.Animation->AuthoredSyncMarkers.Num() > 0)
-			{
-				static const TFunction<void(TArray<FName>&, TArray<FAnimSyncMarker>&)> PopulateMarkerNameArray = [](TArray<FName>& Pattern, TArray<struct FAnimSyncMarker>& AuthoredSyncMarkers)
+				if (Sample.Animation->AuthoredSyncMarkers.Num() > 0)
 				{
-					Pattern.Reserve(AuthoredSyncMarkers.Num());
-					for (FAnimSyncMarker& Marker : AuthoredSyncMarkers)
+					auto PopulateMarkerNameArray = [](TArray<FName>& Pattern, TArray<struct FAnimSyncMarker>& AuthoredSyncMarkers)
 					{
-						Pattern.Add(Marker.MarkerName);
+						Pattern.Reserve(AuthoredSyncMarkers.Num());
+						for (FAnimSyncMarker& Marker : AuthoredSyncMarkers)
+						{
+							Pattern.Add(Marker.MarkerName);
+						}
+					};
+
+					if (SampleWithMarkers == INDEX_NONE)
+					{
+						SampleWithMarkers = SampleIndex;
 					}
-				};
 
-				if (SampleWithMarkers == INDEX_NONE)
-				{
-					SampleWithMarkers = SampleIndex;
+					if (BlendSpacePattern.MarkerNames.Num() == 0)
+					{
+						PopulateMarkerNameArray(BlendSpacePattern.MarkerNames, Sample.Animation->AuthoredSyncMarkers);
+					}
+					else
+					{
+						TArray<FName> ThisPattern;
+						PopulateMarkerNameArray(ThisPattern, Sample.Animation->AuthoredSyncMarkers);
+						if (!BlendSpacePattern.DoesPatternMatch(ThisPattern))
+						{
+							bAllMarkerPatternsMatch = false;
+						}
+					}
 				}
-
-				if (BlendSpacePattern.MarkerNames.Num() == 0)
+			}		
+			else
+			{
+				if (IsRunningGame())
 				{
-					PopulateMarkerNameArray(BlendSpacePattern.MarkerNames, Sample.Animation->AuthoredSyncMarkers);
+					UE_LOG(LogAnimation, Error, TEXT("[%s : %d] - Missing Sample Animation"), *GetFullName(), SampleIndex + 1);
 				}
 				else
 				{
-					TArray<FName> ThisPattern;
-					PopulateMarkerNameArray(ThisPattern, Sample.Animation->AuthoredSyncMarkers);
-					if (!BlendSpacePattern.DoesPatternMatch(ThisPattern))
-					{
-						bAllMarkerPatternsMatch = false;
-					}
+					static FName NAME_LoadErrors("LoadErrors");
+					FMessageLog LoadErrors(NAME_LoadErrors);
+
+					TSharedRef<FTokenizedMessage> Message = LoadErrors.Error();
+					Message->AddToken(FTextToken::Create(LOCTEXT("EmptyAnimationData1", "The BlendSpace ")));
+					Message->AddToken(FAssetNameToken::Create(GetPathName(), FText::FromString(GetName())));
+					Message->AddToken(FTextToken::Create(LOCTEXT("EmptyAnimationData2", " has sample with no animation. Recommend to remove sample point or set new animation.")));
+					LoadErrors.Notify();
 				}
 			}
-		}		
+		}
 		else
 		{
-			if (IsRunningGame())
-			{
-				UE_LOG(LogAnimation, Error, TEXT("[%s : %d] - Missing Sample Animation"), *GetFullName(), SampleIndex + 1);
-			}
-			else
-			{
-				static FName NAME_LoadErrors("LoadErrors");
-				FMessageLog LoadErrors(NAME_LoadErrors);
-
-				TSharedRef<FTokenizedMessage> Message = LoadErrors.Error();
-				Message->AddToken(FTextToken::Create(LOCTEXT("EmptyAnimationData1", "The BlendSpace ")));
-				Message->AddToken(FAssetNameToken::Create(GetPathName(), FText::FromString(GetName())));
-				Message->AddToken(FTextToken::Create(LOCTEXT("EmptyAnimationData2", " has sample with no animation. Recommend to remove sample point or set new animation.")));
-				LoadErrors.Notify();
-			}
+			Sample.bIsValid = ValidateSampleValue(Sample.SampleValue, SampleIndex);
 		}
 	}
 
@@ -895,6 +950,22 @@ void UBlendSpaceBase::ValidateSampleData()
 	}
 }
 
+bool UBlendSpaceBase::AddSample(const FVector& SampleValue)
+{
+	// We should only be adding samples without a source animation if we are not a standalone asset
+	check(!IsAsset());
+
+	const bool bValidSampleData = ValidateSampleValue(SampleValue);
+
+	if (bValidSampleData)
+	{
+		SampleData.Add(FBlendSample(nullptr, SampleValue, true, bValidSampleData));
+		UpdatePreviewBasePose();
+	}
+
+	return bValidSampleData;
+}
+
 bool UBlendSpaceBase::AddSample(UAnimSequence* AnimationSequence, const FVector& SampleValue)
 {
 	const bool bValidSampleData = ValidateSampleValue(SampleValue) && ValidateAnimationSequence(AnimationSequence);
