Fix missing custom data render target for hair shader material.

diff --git a/Engine/Source/Runtime/Engine/Private/ShaderCompiler/ShaderGenerationUtil.cpp b/Engine/Source/Runtime/Engine/Private/ShaderCompiler/ShaderGenerationUtil.cpp
--- a/Engine/Source/Runtime/Engine/Private/ShaderCompiler/ShaderGenerationUtil.cpp
+++ b/Engine/Source/Runtime/Engine/Private/ShaderCompiler/ShaderGenerationUtil.cpp
@@ -1738,6 +1738,7 @@ static void DetermineUsedMaterialSlots(
 	if (Mat.MATERIAL_SHADINGMODEL_HAIR)
 	{
 		SetStandardGBufferSlots(Slots, bWriteEmissive, bHasTangent, bHasVelocity, bHasStaticLighting);
+		Slots[GBS_CustomData] = true;
 	}
 
 	if (Mat.MATERIAL_SHADINGMODEL_CLOTH)
