Fix Non Unity build: missing include

diff --git a/Engine/Source/Runtime/Engine/Private/LevelInstance/LevelInstanceEditorLevelStreaming.cpp b/Engine/Source/Runtime/Engine/Private/LevelInstance/LevelInstanceEditorLevelStreaming.cpp
--- a/Engine/Source/Runtime/Engine/Private/LevelInstance/LevelInstanceEditorLevelStreaming.cpp
+++ b/Engine/Source/Runtime/Engine/Private/LevelInstance/LevelInstanceEditorLevelStreaming.cpp
@@ -9,6 +9,7 @@
 #include "EditorLevelUtils.h"
 #include "Editor.h"
 #include "Engine/LevelBounds.h"
+#include "GameFramework/WorldSettings.h"
 #endif
 
 #if WITH_EDITOR
diff --git a/Engine/Source/Runtime/Engine/Private/LevelInstance/LevelInstanceLevelStreaming.cpp b/Engine/Source/Runtime/Engine/Private/LevelInstance/LevelInstanceLevelStreaming.cpp
--- a/Engine/Source/Runtime/Engine/Private/LevelInstance/LevelInstanceLevelStreaming.cpp
+++ b/Engine/Source/Runtime/Engine/Private/LevelInstance/LevelInstanceLevelStreaming.cpp
@@ -7,6 +7,7 @@
 #include "Engine/LevelBounds.h"
 #include "Engine/World.h"
 #include "Misc/PackageName.h"
+#include "GameFramework/WorldSettings.h"
 
 #if WITH_EDITOR
 #include "LevelInstance/LevelInstanceEditorInstanceActor.h"
