Fix race condition where async texture could be stuck forever

diff --git a/Engine/Source/Runtime/Engine/Private/Texture2D.cpp b/Engine/Source/Runtime/Engine/Private/Texture2D.cpp
--- a/Engine/Source/Runtime/Engine/Private/Texture2D.cpp
+++ b/Engine/Source/Runtime/Engine/Private/Texture2D.cpp
@@ -908,12 +908,21 @@ FTextureResource* UTexture2D::CreateResource()
 	TRACE_CPUPROFILER_EVENT_SCOPE(UTexture2D::CreateResource)
 
 #if WITH_EDITOR
-	if (PrivatePlatformData && !PrivatePlatformData->IsAsyncWorkComplete())
+	if (PrivatePlatformData)
 	{
-		FTextureCompilingManager::Get().AddTextures({this});
+		if (PrivatePlatformData->IsAsyncWorkComplete())
+		{
+			// Make sure AsyncData has been destroyed in case it still exists to avoid
+			// IsDefaultTexture thinking platform data is still being computed.
+			PrivatePlatformData->FinishCache();
+		}
+		else
+		{
+			FTextureCompilingManager::Get().AddTextures({ this });
 
-		UnlinkStreaming();
-		return new FTexture2DResource(this, GetDefaultTexture2D(this)->GetResource()->GetTexture2DResource());
+			UnlinkStreaming();
+			return new FTexture2DResource(this, GetDefaultTexture2D(this)->GetResource()->GetTexture2DResource());
+		}
 	}
 #endif
 
