Another batch of meaningful Insight trace

diff --git a/Engine/Source/Runtime/Core/Private/Misc/AutomationTest.cpp b/Engine/Source/Runtime/Core/Private/Misc/AutomationTest.cpp
--- a/Engine/Source/Runtime/Core/Private/Misc/AutomationTest.cpp
+++ b/Engine/Source/Runtime/Core/Private/Misc/AutomationTest.cpp
@@ -217,6 +217,8 @@ bool FAutomationTestFramework::ContainsTest( const FString& InTestName ) const
 
 bool FAutomationTestFramework::RunSmokeTests()
 {
+	TRACE_CPUPROFILER_EVENT_SCOPE(FAutomationTestFramework::RunSmokeTests);
+
 	bool bAllSuccessful = true;
 
 	uint32 PreviousRequestedTestFilter = RequestedTestFilter;
diff --git a/Engine/Source/Runtime/Core/Private/Misc/OutputDeviceRedirector.cpp b/Engine/Source/Runtime/Core/Private/Misc/OutputDeviceRedirector.cpp
--- a/Engine/Source/Runtime/Core/Private/Misc/OutputDeviceRedirector.cpp
+++ b/Engine/Source/Runtime/Core/Private/Misc/OutputDeviceRedirector.cpp
@@ -206,6 +206,8 @@ void FOutputDeviceRedirector::InternalFlushThreadedLogs(TLocalOutputDevicesArray
 {	
 	if (BufferedLines.Num())
 	{
+		TRACE_CPUPROFILER_EVENT_SCOPE(FOutputDeviceRedirector::InternalFlushThreadedLogs);
+
 		TArray<FBufferedLine, TInlineAllocator<64>> LocalBufferedLines;
 		FLogAllocator::FBufferLock BufferLock;
 		{
diff --git a/Engine/Source/Runtime/Engine/Private/InstancedStaticMesh.cpp b/Engine/Source/Runtime/Engine/Private/InstancedStaticMesh.cpp
--- a/Engine/Source/Runtime/Engine/Private/InstancedStaticMesh.cpp
+++ b/Engine/Source/Runtime/Engine/Private/InstancedStaticMesh.cpp
@@ -1695,6 +1695,7 @@ void UInstancedStaticMeshComponent::InitInstanceBody(int32 InstanceIdx, FBodyIns
 
 void UInstancedStaticMeshComponent::CreateAllInstanceBodies()
 {
+	TRACE_CPUPROFILER_EVENT_SCOPE(UInstancedStaticMeshComponent::CreateAllInstanceBodies);
 	QUICK_SCOPE_CYCLE_COUNTER(STAT_UInstancedStaticMeshComponent_CreateAllInstanceBodies);
 
 	const int32 NumBodies = PerInstanceSMData.Num();
diff --git a/Engine/Source/Runtime/Engine/Private/Materials/MaterialShader.cpp b/Engine/Source/Runtime/Engine/Private/Materials/MaterialShader.cpp
--- a/Engine/Source/Runtime/Engine/Private/Materials/MaterialShader.cpp
+++ b/Engine/Source/Runtime/Engine/Private/Materials/MaterialShader.cpp
@@ -1780,6 +1780,8 @@ static FHashedName PreprocessedSourceKeyFromName(FName VertexFactoryName, FName
 
 FShader* FMaterialShaderMap::ProcessCompilationResultsForSingleJob(FShaderCompileJob* SingleJob, const FShaderPipelineType* ShaderPipeline, const FSHAHash& MaterialShaderMapHash)
 {
+	TRACE_CPUPROFILER_EVENT_SCOPE(FMaterialShaderMap::ProcessCompilationResultsForSingleJob);
+
 	check(SingleJob);
 	const FShaderCompileJob& CurrentJob = *SingleJob;
 	check(CurrentJob.Id == CompilingId);
diff --git a/Engine/Source/Runtime/Engine/Private/Materials/MaterialShared.cpp b/Engine/Source/Runtime/Engine/Private/Materials/MaterialShared.cpp
--- a/Engine/Source/Runtime/Engine/Private/Materials/MaterialShared.cpp
+++ b/Engine/Source/Runtime/Engine/Private/Materials/MaterialShared.cpp
@@ -2051,6 +2051,7 @@ bool FMaterial::CacheShaders(EShaderPlatform Platform, EMaterialShaderPrecompile
  */
 bool FMaterial::CacheShaders(const FMaterialShaderMapId& ShaderMapId, EShaderPlatform Platform, EMaterialShaderPrecompileMode PrecompileMode, const ITargetPlatform* TargetPlatform)
 {
+	TRACE_CPUPROFILER_EVENT_SCOPE(FMaterial::CacheShaders);
 	UE_CLOG(!ShaderMapId.IsValid(), LogMaterial, Warning, TEXT("Invalid shader map ID caching shaders for '%s', will use default material."), *GetFriendlyName());
 
 	// If we loaded this material with inline shaders, use what was loaded (GameThreadShaderMap) instead of looking in the DDC
diff --git a/Engine/Source/Runtime/Engine/Private/VT/UploadingVirtualTexture.cpp b/Engine/Source/Runtime/Engine/Private/VT/UploadingVirtualTexture.cpp
--- a/Engine/Source/Runtime/Engine/Private/VT/UploadingVirtualTexture.cpp
+++ b/Engine/Source/Runtime/Engine/Private/VT/UploadingVirtualTexture.cpp
@@ -300,6 +300,8 @@ struct FCreateCodecTask
 
 FVTCodecAndStatus FUploadingVirtualTexture::GetCodecForChunk(FGraphEventArray& OutCompletionEvents, uint32 ChunkIndex, EVTRequestPagePriority Priority)
 {
+	TRACE_CPUPROFILER_EVENT_SCOPE(FUploadingVirtualTexture::GetCodecForChunk);
+
 	const FVirtualTextureDataChunk& Chunk = Data->Chunks[ChunkIndex];
 	if (Chunk.CodecPayloadSize == 0u)
 	{
@@ -344,6 +346,8 @@ FVTCodecAndStatus FUploadingVirtualTexture::GetCodecForChunk(FGraphEventArray& O
 
 FVTDataAndStatus FUploadingVirtualTexture::ReadData(FGraphEventArray& OutCompletionEvents, uint32 ChunkIndex, size_t Offset, size_t Size, EVTRequestPagePriority Priority)
 {
+	TRACE_CPUPROFILER_EVENT_SCOPE(FUploadingVirtualTexture::ReadData);
+
 	FVirtualTextureDataChunk& Chunk = Data->Chunks[ChunkIndex];
 	FByteBulkData& BulkData = Chunk.BulkData;
 #if WITH_EDITOR
diff --git a/Engine/Source/Runtime/Engine/Private/VT/VirtualTextureBuiltData.cpp b/Engine/Source/Runtime/Engine/Private/VT/VirtualTextureBuiltData.cpp
--- a/Engine/Source/Runtime/Engine/Private/VT/VirtualTextureBuiltData.cpp
+++ b/Engine/Source/Runtime/Engine/Private/VT/VirtualTextureBuiltData.cpp
@@ -43,6 +43,8 @@ uint32 FVirtualTextureBuiltData::GetNumTileHeaders() const
 
 void FVirtualTextureBuiltData::Serialize(FArchive& Ar, UObject* Owner, int32 FirstMipToSerialize)
 {
+	TRACE_CPUPROFILER_EVENT_SCOPE(FVirtualTextureDataChunk::Serialize);
+
 	check(FirstMipToSerialize == 0 || Ar.IsSaving());
 	const bool bStripMips = (FirstMipToSerialize > 0);
 	uint32 NumChunksToStrip = 0u;
@@ -184,6 +186,8 @@ void FVirtualTextureBuiltData::Serialize(FArchive& Ar, UObject* Owner, int32 Fir
 
 bool FVirtualTextureBuiltData::ValidateData(FStringView const& InDDCDebugContext, bool bValidateCompression) const
 {
+	TRACE_CPUPROFILER_EVENT_SCOPE(FVirtualTextureBuiltData::ValidateCompression);
+
 	const uint32 TilePixelSize = GetPhysicalTileSize();
 	TArray<uint8> UncompressedResult;
 	TArray<uint8> ChunkDataDDC;
@@ -328,6 +332,8 @@ bool FVirtualTextureDataChunk::ShortenKey(const FString& CacheKey, FString& Resu
 
 uint32 FVirtualTextureDataChunk::StoreInDerivedDataCache(const FString& InDerivedDataKey, const FStringView& TextureName, bool bReplaceExistingDDC)
 {
+	TRACE_CPUPROFILER_EVENT_SCOPE(FVirtualTextureDataChunk::StoreInDerivedDataCache);
+
 	int32 BulkDataSizeInBytes = BulkData.GetBulkDataSize();
 	check(BulkDataSizeInBytes > 0);
 
diff --git a/Engine/Source/Runtime/Engine/Private/VT/VirtualTextureChunkDDCCache.cpp b/Engine/Source/Runtime/Engine/Private/VT/VirtualTextureChunkDDCCache.cpp
--- a/Engine/Source/Runtime/Engine/Private/VT/VirtualTextureChunkDDCCache.cpp
+++ b/Engine/Source/Runtime/Engine/Private/VT/VirtualTextureChunkDDCCache.cpp
@@ -56,6 +56,8 @@ public:
 
 	void DoWork()
 	{
+		TRACE_CPUPROFILER_EVENT_SCOPE(FAsyncFillCacheWorker::DoWork);
+
 		IPlatformFile& PlatformFile = FPlatformFileManager::Get().GetPlatformFile();
 		FDerivedDataCacheInterface& DDC = GetDerivedDataCacheRef();
 
@@ -176,6 +178,8 @@ FVirtualTextureChunkDDCCache* GetVirtualTextureChunkDDCCache()
 
 void FVirtualTextureChunkDDCCache::Initialize()
 {
+	TRACE_CPUPROFILER_EVENT_SCOPE(FVirtualTextureChunkDDCCache::Initialize);
+
 	// setup the cache folder
 	auto& PlatformFile = FPlatformFileManager::Get().GetPlatformFile();
 
@@ -238,6 +242,8 @@ void FVirtualTextureChunkDDCCache::UpdateRequests()
 
 bool FVirtualTextureChunkDDCCache::MakeChunkAvailable(struct FVirtualTextureDataChunk* Chunk, bool bAsync, FString& OutChunkFileName, int64& OutOffsetInFile)
 {
+	TRACE_CPUPROFILER_EVENT_SCOPE(FVirtualTextureChunkDDCCache::MakeChunkAvailable);
+
 	const FString CachedFilePath = AbsoluteCachePath / Chunk->ShortDerivedDataKey;
 	const FString TempFilePath = AbsoluteCachePath / FGuid::NewGuid().ToString() + ".tmp";
 
diff --git a/Engine/Source/Runtime/RenderCore/Private/ShaderResource.cpp b/Engine/Source/Runtime/RenderCore/Private/ShaderResource.cpp
--- a/Engine/Source/Runtime/RenderCore/Private/ShaderResource.cpp
+++ b/Engine/Source/Runtime/RenderCore/Private/ShaderResource.cpp
@@ -169,6 +169,8 @@ void FShaderMapResourceCode::AddShaderCompilerOutput(const FShaderCompilerOutput
 
 void FShaderMapResourceCode::AddShaderCode(EShaderFrequency InFrequency, const FSHAHash& InHash, TConstArrayView<uint8> InCode)
 {
+	TRACE_CPUPROFILER_EVENT_SCOPE(FShaderMapResourceCode::AddShaderCode);
+
 	const int32 Index = Algo::LowerBound(ShaderHashes, InHash);
 	if (Index >= ShaderHashes.Num() || ShaderHashes[Index] != InHash)
 	{
