Optimized ULevelInstanceSubsystem::OnActorDeleted, early out for non-editing PackedLevelInstance actors.

diff --git a/Engine/Source/Runtime/Engine/Private/LevelInstance/LevelInstanceSubsystem.cpp b/Engine/Source/Runtime/Engine/Private/LevelInstance/LevelInstanceSubsystem.cpp
--- a/Engine/Source/Runtime/Engine/Private/LevelInstance/LevelInstanceSubsystem.cpp
+++ b/Engine/Source/Runtime/Engine/Private/LevelInstance/LevelInstanceSubsystem.cpp
@@ -945,6 +945,12 @@ void ULevelInstanceSubsystem::OnActorDeleted(AActor* Actor)
 			return;
 		}
 
+		const bool bIsEditingLevelInstance = IsEditingLevelInstance(LevelInstanceActor);
+		if (!bIsEditingLevelInstance && LevelInstanceActor->IsA<APackedLevelInstance>())
+		{
+			return;
+		}
+
 		const bool bAlreadyRooted = LevelInstanceActor->IsRooted();
 		// Unloading LevelInstances leads to GC and Actor can be collected. Add to root temp. It will get collected after the OnActorDeleted callbacks
 		if (!bAlreadyRooted)
@@ -955,7 +961,7 @@ void ULevelInstanceSubsystem::OnActorDeleted(AActor* Actor)
 		FScopedSlowTask SlowTask(0, LOCTEXT("UnloadingLevelInstances", "Unloading Level Instances..."), !GetWorld()->IsGameWorld());
 		SlowTask.MakeDialog();
 		check(!IsEditingLevelInstanceDirty(LevelInstanceActor) && !HasDirtyChildrenLevelInstances(LevelInstanceActor));
-		if (IsEditingLevelInstance(LevelInstanceActor))
+		if (bIsEditingLevelInstance)
 		{
 			CommitLevelInstance(LevelInstanceActor);
 		}
