Make FLightmassWorldInfoSettings a blueprint type and make its properties BlueprintReadWrite in order to be modified via editor utility widgets. The properties can already be modified via python, so this provides a more consistent tooling experience

#rnx

diff --git a/Engine/Source/Runtime/Engine/Classes/GameFramework/WorldSettings.h b/Engine/Source/Runtime/Engine/Classes/GameFramework/WorldSettings.h
--- a/Engine/Source/Runtime/Engine/Classes/GameFramework/WorldSettings.h
+++ b/Engine/Source/Runtime/Engine/Classes/GameFramework/WorldSettings.h
@@ -49,7 +49,7 @@ enum EVolumeLightingMethod
 	VLM_SparseVolumeLightingSamples UMETA(DisplayName = "Sparse Volume Lighting Samples"),
 };
 
-USTRUCT()
+USTRUCT(BlueprintType)
 struct FLightmassWorldInfoSettings
 {
 	GENERATED_USTRUCT_BODY()
@@ -61,7 +61,7 @@ struct FLightmassWorldInfoSettings
 	 * Any levels with a different scale should use this scale to compensate. 
 	 * For large levels it can drastically reduce build times to set this to 2 or 4.
 	 */
-	UPROPERTY(EditAnywhere, Category=LightmassGeneral, AdvancedDisplay, meta=(UIMin = "1.0", UIMax = "4.0"))
+	UPROPERTY(BlueprintReadWrite, EditAnywhere, Category=LightmassGeneral, AdvancedDisplay, meta=(UIMin = "1.0", UIMax = "4.0"))
 	float StaticLightingLevelScale;
 
 	/** 
@@ -70,14 +70,14 @@ struct FLightmassWorldInfoSettings
 	 * Bounce 1 takes the most time to calculate and contributes the most to visual quality, followed by bounce 2.
 	 * Successive bounces don't really affect build times, but have a much lower visual impact, unless the material diffuse colors are close to 1.
 	 */
-	UPROPERTY(EditAnywhere, Category=LightmassGeneral, meta=(UIMin = "1.0", UIMax = "10.0"))
+	UPROPERTY(BlueprintReadWrite, EditAnywhere, Category=LightmassGeneral, meta=(UIMin = "1.0", UIMax = "10.0"))
 	int32 NumIndirectLightingBounces;
 
 	/** 
 	 * Number of skylight and emissive bounces to simulate.  
 	 * Lightmass uses a non-distributable radiosity method for skylight bounces whose cost is proportional to the number of bounces.
 	 */
-	UPROPERTY(EditAnywhere, Category=LightmassGeneral, meta=(UIMin = "1.0", UIMax = "10.0"))
+	UPROPERTY(BlueprintReadWrite, EditAnywhere, Category=LightmassGeneral, meta=(UIMin = "1.0", UIMax = "10.0"))
 	int32 NumSkyLightingBounces;
 
 	/** 
@@ -86,7 +86,7 @@ struct FLightmassWorldInfoSettings
 	 * It can be useful to reduce IndirectLightingSmoothness somewhat (~.75) when increasing quality to get defined indirect shadows.
 	 * Note that this can't affect compression artifacts, UV seams or other texture based artifacts.
 	 */
-	UPROPERTY(EditAnywhere, Category=LightmassGeneral, AdvancedDisplay, meta=(UIMin = "1.0", UIMax = "4.0"))
+	UPROPERTY(BlueprintReadWrite, EditAnywhere, Category=LightmassGeneral, AdvancedDisplay, meta=(UIMin = "1.0", UIMax = "4.0"))
 	float IndirectLightingQuality;
 
 	/** 
@@ -94,18 +94,18 @@ struct FLightmassWorldInfoSettings
 	 * 1 is default smoothness tweaked for a variety of lighting situations.
 	 * Higher values like 3 smooth out the indirect lighting more, but at the cost of indirect shadows losing detail.
 	 */
-	UPROPERTY(EditAnywhere, Category=LightmassGeneral, AdvancedDisplay, meta=(UIMin = "0.5", UIMax = "6.0"))
+	UPROPERTY(BlueprintReadWrite, EditAnywhere, Category=LightmassGeneral, AdvancedDisplay, meta=(UIMin = "0.5", UIMax = "6.0"))
 	float IndirectLightingSmoothness;
 
 	/** 
 	 * Represents a constant color light surrounding the upper hemisphere of the level, like a sky.
 	 * This light source currently does not get bounced as indirect lighting and causes reflection capture brightness to be incorrect.  Prefer using a Static Skylight instead.
 	 */
-	UPROPERTY(EditAnywhere, Category=LightmassGeneral)
+	UPROPERTY(BlueprintReadWrite, EditAnywhere, Category=LightmassGeneral)
 	FColor EnvironmentColor;
 
 	/** Scales EnvironmentColor to allow independent color and brightness controls. */
-	UPROPERTY(EditAnywhere, Category=LightmassGeneral, meta=(UIMin = "0", UIMax = "10"))
+	UPROPERTY(BlueprintReadWrite, EditAnywhere, Category=LightmassGeneral, meta=(UIMin = "0", UIMax = "10"))
 	float EnvironmentIntensity;
 
 	/** Scales the emissive contribution of all materials in the scene.  Currently disabled and should be removed with mesh area lights. */
@@ -113,15 +113,15 @@ struct FLightmassWorldInfoSettings
 	float EmissiveBoost;
 
 	/** Scales the diffuse contribution of all materials in the scene. */
-	UPROPERTY(EditAnywhere, Category=LightmassGeneral, meta=(UIMin = "0.1", UIMax = "6.0"))
+	UPROPERTY(BlueprintReadWrite, EditAnywhere, Category=LightmassGeneral, meta=(UIMin = "0.1", UIMax = "6.0"))
 	float DiffuseBoost;
 
 	/** Technique to use for providing precomputed lighting at all positions inside the Lightmass Importance Volume */
-	UPROPERTY(EditAnywhere, Category=LightmassVolumeLighting)
+	UPROPERTY(BlueprintReadWrite, EditAnywhere, Category=LightmassVolumeLighting)
 	TEnumAsByte<enum EVolumeLightingMethod> VolumeLightingMethod;
 
 	/** If true, AmbientOcclusion will be enabled. */
-	UPROPERTY(EditAnywhere, Category=LightmassOcclusion)
+	UPROPERTY(BlueprintReadWrite, EditAnywhere, Category=LightmassOcclusion)
 	uint8 bUseAmbientOcclusion:1;
 
 	/** 
@@ -130,22 +130,22 @@ struct FLightmassWorldInfoSettings
 	 * Which is useful for blending between material layers on environment assets.
 	 * Be sure to set DirectIlluminationOcclusionFraction and IndirectIlluminationOcclusionFraction to 0 if you only want the PrecomputedAOMask!
 	 */
-	UPROPERTY(EditAnywhere, Category=LightmassOcclusion)
+	UPROPERTY(BlueprintReadWrite, EditAnywhere, Category=LightmassOcclusion)
 	uint8 bGenerateAmbientOcclusionMaterialMask:1;
 
 	/** If true, override normal direct and indirect lighting with just the exported diffuse term. */
-	UPROPERTY(EditAnywhere, Category=LightmassDebug, AdvancedDisplay)
+	UPROPERTY(BlueprintReadWrite, EditAnywhere, Category=LightmassDebug, AdvancedDisplay)
 	uint8 bVisualizeMaterialDiffuse:1;
 
 	/** If true, override normal direct and indirect lighting with just the AO term. */
-	UPROPERTY(EditAnywhere, Category=LightmassDebug, AdvancedDisplay)
+	UPROPERTY(BlueprintReadWrite, EditAnywhere, Category=LightmassDebug, AdvancedDisplay)
 	uint8 bVisualizeAmbientOcclusion:1;
 
 	/** 
 	 * Whether to compress lightmap textures.  Disabling lightmap texture compression will reduce artifacts but increase memory and disk size by 4x.
 	 * Use caution when disabling this.
 	 */
-	UPROPERTY(EditAnywhere, Category=LightmassGeneral, AdvancedDisplay)
+	UPROPERTY(BlueprintReadWrite, EditAnywhere, Category=LightmassGeneral, AdvancedDisplay)
 	uint8 bCompressLightmaps:1;
 
 	/** 
@@ -153,13 +153,13 @@ struct FLightmassWorldInfoSettings
 	 * This setting has a large impact on build times and memory, use with caution.  
 	 * Halving the DetailCellSize can increase memory by up to a factor of 8x.
 	 */
-	UPROPERTY(EditAnywhere, Category=LightmassVolumeLighting, meta=(UIMin = "50", UIMax = "1000"))
+	UPROPERTY(BlueprintReadWrite, EditAnywhere, Category=LightmassVolumeLighting, meta=(UIMin = "50", UIMax = "1000"))
 	float VolumetricLightmapDetailCellSize;
 
 	/** 
 	 * Maximum amount of memory to spend on Volumetric Lightmap Brick data.  High density bricks will be discarded until this limit is met, with bricks furthest from geometry discarded first.
 	 */
-	UPROPERTY(EditAnywhere, Category=LightmassVolumeLighting, meta=(UIMin = "1", UIMax = "500"))
+	UPROPERTY(BlueprintReadWrite, EditAnywhere, Category=LightmassVolumeLighting, meta=(UIMin = "1", UIMax = "500"))
 	float VolumetricLightmapMaximumBrickMemoryMb;
 
 	/** 
@@ -168,34 +168,34 @@ struct FLightmassWorldInfoSettings
 	 * Smoothing can reduce this artifact.  Smoothing is only applied when the ringing artifact is present.
 	 * 0 = no smoothing, 1 = strong smooth (little directionality in lighting).
 	 */
-	UPROPERTY(EditAnywhere, Category=LightmassVolumeLighting, meta=(UIMin = "0", UIMax = "1"))
+	UPROPERTY(BlueprintReadWrite, EditAnywhere, Category=LightmassVolumeLighting, meta=(UIMin = "0", UIMax = "1"))
 	float VolumetricLightmapSphericalHarmonicSmoothing;
 
 	/** 
 	 * Scales the distances at which volume lighting samples are placed.  Volume lighting samples are computed by Lightmass and are used for GI on movable components.
 	 * Using larger scales results in less sample memory usage and reduces Indirect Lighting Cache update times, but less accurate transitions between lighting areas.
 	 */
-	UPROPERTY(EditAnywhere, Category=LightmassVolumeLighting, AdvancedDisplay, meta=(UIMin = "0.1", UIMax = "100.0"))
+	UPROPERTY(BlueprintReadWrite, EditAnywhere, Category=LightmassVolumeLighting, AdvancedDisplay, meta=(UIMin = "0.1", UIMax = "100.0"))
 	float VolumeLightSamplePlacementScale;
 
 	/** How much of the AO to apply to direct lighting. */
-	UPROPERTY(EditAnywhere, Category=LightmassOcclusion, meta=(UIMin = "0", UIMax = "1"))
+	UPROPERTY(BlueprintReadWrite, EditAnywhere, Category=LightmassOcclusion, meta=(UIMin = "0", UIMax = "1"))
 	float DirectIlluminationOcclusionFraction;
 
 	/** How much of the AO to apply to indirect lighting. */
-	UPROPERTY(EditAnywhere, Category=LightmassOcclusion, meta=(UIMin = "0", UIMax = "1"))
+	UPROPERTY(BlueprintReadWrite, EditAnywhere, Category=LightmassOcclusion, meta=(UIMin = "0", UIMax = "1"))
 	float IndirectIlluminationOcclusionFraction;
 
 	/** Higher exponents increase contrast. */
-	UPROPERTY(EditAnywhere, Category=LightmassOcclusion, meta=(UIMin = ".5", UIMax = "8"))
+	UPROPERTY(BlueprintReadWrite, EditAnywhere, Category=LightmassOcclusion, meta=(UIMin = ".5", UIMax = "8"))
 	float OcclusionExponent;
 
 	/** Fraction of samples taken that must be occluded in order to reach full occlusion. */
-	UPROPERTY(EditAnywhere, Category=LightmassOcclusion, meta=(UIMin = "0", UIMax = "1"))
+	UPROPERTY(BlueprintReadWrite, EditAnywhere, Category=LightmassOcclusion, meta=(UIMin = "0", UIMax = "1"))
 	float FullyOccludedSamplesFraction;
 
 	/** Maximum distance for an object to cause occlusion on another object. */
-	UPROPERTY(EditAnywhere, Category=LightmassOcclusion)
+	UPROPERTY(BlueprintReadWrite, EditAnywhere, Category=LightmassOcclusion)
 	float MaxOcclusionDistance;
 
 	FLightmassWorldInfoSettings()
@@ -627,7 +627,7 @@ public:
 	
 	/** LIGHTMASS RELATED SETTINGS **/
 	
-	UPROPERTY(EditAnywhere, Category=Lightmass)
+	UPROPERTY(BlueprintReadWrite, EditAnywhere, Category=Lightmass)
 	struct FLightmassWorldInfoSettings LightmassSettings;
 #endif
 
