Fixed missing comparison in LevelInstanceISMPacker operator =.

diff --git a/Engine/Source/Runtime/Engine/Private/LevelInstance/Packed/LevelInstanceISMPacker.cpp b/Engine/Source/Runtime/Engine/Private/LevelInstance/Packed/LevelInstanceISMPacker.cpp
--- a/Engine/Source/Runtime/Engine/Private/LevelInstance/Packed/LevelInstanceISMPacker.cpp
+++ b/Engine/Source/Runtime/Engine/Private/LevelInstance/Packed/LevelInstanceISMPacker.cpp
@@ -127,7 +127,8 @@ bool FLevelInstanceISMPackerCluster::operator==(const FLevelInstancePackerCluste
 		Materials == ISMOther.Materials &&
 		bReceivesDecals == ISMOther.bReceivesDecals &&
 		bCastShadow == ISMOther.bCastShadow &&
-		bVisibleInRayTracing == ISMOther.bVisibleInRayTracing;
+		bVisibleInRayTracing == ISMOther.bVisibleInRayTracing &&
+		bCollisionEnabled == ISMOther.bCollisionEnabled;
 }
 
 #endif
\ No newline at end of file
