Added DataLayerEditorContext and use it to allow automatic creation of Instanced Foliage Actors associated with proper Data Layers (Foliage painting & PFV).

diff --git a/Engine/Source/Editor/DataLayerEditor/Private/DataLayer/DataLayerPropertyTypeCustomization.cpp b/Engine/Source/Editor/DataLayerEditor/Private/DataLayer/DataLayerPropertyTypeCustomization.cpp
--- a/Engine/Source/Editor/DataLayerEditor/Private/DataLayer/DataLayerPropertyTypeCustomization.cpp
+++ b/Engine/Source/Editor/DataLayerEditor/Private/DataLayer/DataLayerPropertyTypeCustomization.cpp
@@ -1,9 +1,11 @@
 // Copyright Epic Games, Inc. All Rights Reserved.
 
 #include "DataLayerPropertyTypeCustomization.h"
+#include "DataLayer/DataLayerPropertyTypeCustomizationHelper.h"
 #include "DataLayer/DataLayerDragDropOp.h"
 #include "DataLayer/DataLayerEditorSubsystem.h"
 #include "WorldPartition/DataLayer/DataLayer.h"
+#include "WorldPartition/WorldPartitionSubsystem.h"
 #include "Algo/Accumulate.h"
 #include "Modules/ModuleManager.h"
 #include "PropertyHandle.h"
@@ -13,11 +15,9 @@
 #include "Widgets/Input/SButton.h"
 #include "Widgets/Input/SComboButton.h"
 #include "Widgets/Text/STextBlock.h"
-#include "Framework/MultiBox/MultiBoxBuilder.h"
 #include "SDropTarget.h"
 #include "Editor.h"
 #include "EditorFontGlyphs.h"
-#include "LevelEditor.h"
 #include "DetailLayoutBuilder.h"
 #include "DetailCategoryBuilder.h"
 
@@ -36,6 +36,7 @@ void FDataLayerPropertyTypeCustomization::CustomizeHeader(TSharedRef<IPropertyHa
 	.MaxDesiredWidth(TOptional<float>())
 	[
 		SNew(SDropTarget)
+		.IsEnabled_Lambda([]() { return GWorld ? UWorld::HasSubsystem<UWorldPartitionSubsystem>(GWorld) : false; })
 		.OnDrop(this, &FDataLayerPropertyTypeCustomization::OnDrop)
 		.OnAllowDrop(this, &FDataLayerPropertyTypeCustomization::OnVerifyDrag)
 		.OnIsRecognized(this, &FDataLayerPropertyTypeCustomization::OnVerifyDrag)
@@ -90,12 +91,6 @@ void FDataLayerPropertyTypeCustomization::CustomizeHeader(TSharedRef<IPropertyHa
 	];
 }
 
-FText FDataLayerPropertyTypeCustomization::GetDataLayerDescription(const UDataLayer* InDataLayer) const
-{
-	check(InDataLayer);
-	return FText::FromName(InDataLayer->GetDataLayerLabel());
-}
-
 UDataLayer* FDataLayerPropertyTypeCustomization::GetDataLayerFromPropertyHandle() const
 {
 	FName DataLayerName;
@@ -109,46 +104,12 @@ UDataLayer* FDataLayerPropertyTypeCustomization::GetDataLayerFromPropertyHandle(
 
 FText FDataLayerPropertyTypeCustomization::GetDataLayerText() const
 {
-	const UDataLayer* DataLayer = GetDataLayerFromPropertyHandle();
-	return DataLayer ? GetDataLayerDescription(DataLayer) : LOCTEXT("InvalidDataLayerLabel", "<Invalid>");
+	return UDataLayer::GetDataLayerText(GetDataLayerFromPropertyHandle());
 }
 
 TSharedRef<SWidget> FDataLayerPropertyTypeCustomization::OnGetDataLayerMenu()
 {
-	FMenuBuilder MenuBuilder(true, nullptr);
-
-	MenuBuilder.AddMenuEntry(
-		LOCTEXT("OpenDataLayersBrowser", "Browse DataLayers..."),
-		FText(),
-		FSlateIcon(FEditorStyle::GetStyleSetName(), "LevelEditor.Tabs.DataLayers"),
-		FUIAction(
-			FExecuteAction::CreateSP(this, &FDataLayerPropertyTypeCustomization::OpenDataLayerBrowser)
-		)
-	);
-
-	MenuBuilder.BeginSection(FName(), LOCTEXT("ExistingDataLayers", "Existing DataLayers"));
-	{
-		TArray<TWeakObjectPtr<UDataLayer>> AllDataLayers;
-		UDataLayerEditorSubsystem::Get()->AddAllDataLayersTo(AllDataLayers);
-
-		for (const TWeakObjectPtr<UDataLayer>& WeakDataLayer : AllDataLayers)
-		{
-			if (const UDataLayer* DataLayerPtr = WeakDataLayer.Get())
-			{
-				MenuBuilder.AddMenuEntry(
-					GetDataLayerDescription(DataLayerPtr),
-					FText(),
-					FSlateIcon(FEditorStyle::GetStyleSetName(), "DataLayer.Icon16x"),
-					FUIAction(
-						FExecuteAction::CreateSP(this, &FDataLayerPropertyTypeCustomization::AssignDataLayer, DataLayerPtr)
-					)
-				);
-			}
-		}
-	}
-	MenuBuilder.EndSection();
-
-	return MenuBuilder.MakeWidget();
+	return FDataLayerPropertyTypeCustomizationHelper::CreateDataLayerMenu([this](const UDataLayer* DataLayer) { AssignDataLayer(DataLayer); });
 }
 
 EVisibility FDataLayerPropertyTypeCustomization::GetSelectDataLayerVisibility() const
@@ -176,12 +137,6 @@ void FDataLayerPropertyTypeCustomization::AssignDataLayer(const UDataLayer* InDa
 	}
 }
 
-void FDataLayerPropertyTypeCustomization::OpenDataLayerBrowser()
-{
-	FLevelEditorModule& LevelEditorModule = FModuleManager::GetModuleChecked<FLevelEditorModule>("LevelEditor");
-	LevelEditorModule.GetLevelEditorTabManager()->TryInvokeTab(FTabId("LevelEditorDataLayerBrowser"));
-}
-
 FReply FDataLayerPropertyTypeCustomization::OnDrop(TSharedPtr<FDragDropOperation> InDragDrop)
 {
 	if (InDragDrop.IsValid() && InDragDrop->IsOfType<FDataLayerDragDropOp>())
diff --git a/Engine/Source/Editor/DataLayerEditor/Private/DataLayer/DataLayerPropertyTypeCustomization.h b/Engine/Source/Editor/DataLayerEditor/Private/DataLayer/DataLayerPropertyTypeCustomization.h
--- a/Engine/Source/Editor/DataLayerEditor/Private/DataLayer/DataLayerPropertyTypeCustomization.h
+++ b/Engine/Source/Editor/DataLayerEditor/Private/DataLayer/DataLayerPropertyTypeCustomization.h
@@ -19,10 +19,8 @@ struct FDataLayerPropertyTypeCustomization : public IPropertyTypeCustomization
 private:
 
 	void AssignDataLayer(const UDataLayer* InDataLayer);
-	void OpenDataLayerBrowser();
 
 	UDataLayer* GetDataLayerFromPropertyHandle() const;
-	FText GetDataLayerDescription(const UDataLayer* InDataLayer) const;
 	FText GetDataLayerText() const;
 	EVisibility GetSelectDataLayerVisibility() const;
 
@@ -33,4 +31,3 @@ private:
 
 	TSharedPtr<IPropertyHandle> PropertyHandle;
 };
-
diff --git a/Engine/Source/Editor/DataLayerEditor/Private/DataLayer/DataLayerPropertyTypeCustomizationHelper.cpp b/Engine/Source/Editor/DataLayerEditor/Private/DataLayer/DataLayerPropertyTypeCustomizationHelper.cpp
new file mode 100644
--- /dev/null
+++ b/Engine/Source/Editor/DataLayerEditor/Private/DataLayer/DataLayerPropertyTypeCustomizationHelper.cpp
@@ -0,0 +1,51 @@
+// Copyright Epic Games, Inc. All Rights Reserved.
+
+#include "DataLayer/DataLayerPropertyTypeCustomizationHelper.h"
+#include "DataLayer/DataLayerEditorSubsystem.h"
+#include "WorldPartition/DataLayer/DataLayer.h"
+#include "Framework/MultiBox/MultiBoxBuilder.h"
+#include "Modules/ModuleManager.h"
+#include "LevelEditor.h"
+
+#define LOCTEXT_NAMESPACE "DataLayer"
+
+TSharedRef<SWidget> FDataLayerPropertyTypeCustomizationHelper::CreateDataLayerMenu(TFunction<void(const UDataLayer* DataLayer)> OnDataLayerSelectedFunction)
+{
+	FMenuBuilder MenuBuilder(true, nullptr);
+
+	MenuBuilder.AddMenuEntry(
+		LOCTEXT("OpenDataLayersBrowser", "Browse DataLayers..."),
+		FText(),
+		FSlateIcon(FEditorStyle::GetStyleSetName(), "LevelEditor.Tabs.DataLayers"),
+		FUIAction(
+			FExecuteAction::CreateLambda([]()
+			{
+				FLevelEditorModule& LevelEditorModule = FModuleManager::GetModuleChecked<FLevelEditorModule>("LevelEditor");
+				LevelEditorModule.GetLevelEditorTabManager()->TryInvokeTab(FTabId("LevelEditorDataLayerBrowser"));
+			})
+		)
+	);
+
+	MenuBuilder.BeginSection(FName(), LOCTEXT("ExistingDataLayers", "Existing DataLayers"));
+	{
+		TArray<TWeakObjectPtr<UDataLayer>> AllDataLayers;
+		AllDataLayers.Add(nullptr); // This allows to show the "<None>" option
+		UDataLayerEditorSubsystem::Get()->AddAllDataLayersTo(AllDataLayers);
+
+		for (const TWeakObjectPtr<UDataLayer>& WeakDataLayer : AllDataLayers)
+		{
+			const UDataLayer* DataLayerPtr = WeakDataLayer.Get();
+			MenuBuilder.AddMenuEntry(
+				UDataLayer::GetDataLayerText(DataLayerPtr),
+				FText(),
+				FSlateIcon(FEditorStyle::GetStyleSetName(), "DataLayer.Icon16x"),
+				FUIAction(FExecuteAction::CreateLambda([DataLayerPtr, OnDataLayerSelectedFunction]() { OnDataLayerSelectedFunction(DataLayerPtr); }))
+			);
+		}
+	}
+	MenuBuilder.EndSection();
+
+	return MenuBuilder.MakeWidget();
+}
+
+#undef LOCTEXT_NAMESPACE
\ No newline at end of file
diff --git a/Engine/Source/Editor/DataLayerEditor/Public/DataLayer/DataLayerPropertyTypeCustomizationHelper.h b/Engine/Source/Editor/DataLayerEditor/Public/DataLayer/DataLayerPropertyTypeCustomizationHelper.h
new file mode 100644
--- /dev/null
+++ b/Engine/Source/Editor/DataLayerEditor/Public/DataLayer/DataLayerPropertyTypeCustomizationHelper.h
@@ -0,0 +1,14 @@
+// Copyright Epic Games, Inc. All Rights Reserved.
+
+#pragma once
+
+#include "CoreTypes.h"
+#include "Templates/SharedPointer.h"
+#include "Widgets/SWidget.h"
+
+class UDataLayer;
+
+struct DATALAYEREDITOR_API FDataLayerPropertyTypeCustomizationHelper
+{
+	static TSharedRef<SWidget> CreateDataLayerMenu(TFunction<void(const UDataLayer* DataLayer)> OnDataLayerSelectedFunction);
+};
diff --git a/Engine/Source/Editor/FoliageEdit/FoliageEdit.build.cs b/Engine/Source/Editor/FoliageEdit/FoliageEdit.build.cs
--- a/Engine/Source/Editor/FoliageEdit/FoliageEdit.build.cs
+++ b/Engine/Source/Editor/FoliageEdit/FoliageEdit.build.cs
@@ -26,6 +26,7 @@ public class FoliageEdit : ModuleRules
                 "DetailCustomizations",
                 "AssetTools",
                 "Foliage",
+				"DataLayerEditor",
 			}
 		);
 
diff --git a/Engine/Source/Editor/FoliageEdit/Private/FoliageEdMode.cpp b/Engine/Source/Editor/FoliageEdit/Private/FoliageEdMode.cpp
--- a/Engine/Source/Editor/FoliageEdit/Private/FoliageEdMode.cpp
+++ b/Engine/Source/Editor/FoliageEdit/Private/FoliageEdMode.cpp
@@ -39,6 +39,7 @@
 #include "Misc/ScopeExit.h"
 
 #include "ActorPartition/ActorPartitionSubsystem.h"
+#include "WorldPartition/DataLayer/DataLayerSubsystem.h"
 
 //Slate dependencies
 #include "IAssetViewport.h"
@@ -392,6 +393,11 @@ void FEdModeFoliage::Enter()
 	// Load UI settings from config file
 	UISettings.Load();
 
+	// For now, Only Foliage sets DataLayerSubsystem's DataLayerEditorContext, validate it's empty before updating it
+	const UDataLayerSubsystem* DataLayerSubsystem = GetWorld()->GetSubsystem<UDataLayerSubsystem>();
+	check(!DataLayerSubsystem || !DataLayerSubsystem->HasDataLayerEditorContext());
+	SetDataLayerEditorContext(UISettings.GetDataLayer());
+
 	// Bind to editor callbacks
 	FEditorDelegates::NewCurrentLevel.AddSP(this, &FEdModeFoliage::NotifyNewCurrentLevel);
 	FWorldDelegates::LevelAddedToWorld.AddSP(this, &FEdModeFoliage::NotifyLevelAddedToWorld);
@@ -444,6 +450,11 @@ void FEdModeFoliage::Enter()
 /** FEdMode: Called when the mode is exited */
 void FEdModeFoliage::Exit()
 {
+	DataLayerEditorContext.Reset();
+	// For now, Only Foliage sets DataLayerSubsystem's DataLayerEditorContext, validate it's back to be empty
+	const UDataLayerSubsystem* DataLayerSubsystem = GetWorld()->GetSubsystem<UDataLayerSubsystem>();
+	check(!DataLayerSubsystem || !DataLayerSubsystem->HasDataLayerEditorContext());
+
 	if (bToolActive)
 	{
 		EndFoliageBrushTrace();
@@ -506,6 +517,15 @@ void FEdModeFoliage::Exit()
 	FEdMode::Exit();
 }
 
+void FEdModeFoliage::SetDataLayerEditorContext(const FActorDataLayer& DataLayer)
+{
+	// Update Data Layer and update DataLayerSubsystem's DataLayerEditorContext with it
+	UISettings.SetDataLayer(DataLayer);
+	// Since this is a scope, first call Reset on TUniquePtr to delete old context before assigning the new one
+	DataLayerEditorContext.Reset();
+	DataLayerEditorContext.Reset(new FScopeChangeDataLayerEditorContext(GetWorld(), UISettings.GetDataLayer()));
+}
+
 EFoliageEditingState FEdModeFoliage::GetEditingState() const
 {
 	UWorld* World = GetWorld();
diff --git a/Engine/Source/Editor/FoliageEdit/Private/FoliageEdMode.h b/Engine/Source/Editor/FoliageEdit/Private/FoliageEdMode.h
--- a/Engine/Source/Editor/FoliageEdit/Private/FoliageEdMode.h
+++ b/Engine/Source/Editor/FoliageEdit/Private/FoliageEdMode.h
@@ -8,6 +8,7 @@
 #include "UnrealWidgetFwd.h"
 #include "EdMode.h"
 #include "Widgets/Views/SHeaderRow.h"
+#include "WorldPartition/DataLayer/ActorDataLayer.h"
 
 class AInstancedFoliageActor;
 class FCanvas;
@@ -21,6 +22,7 @@ class ULandscapeComponent;
 class UPrimitiveComponent;
 class UStaticMeshComponent;
 struct FViewportClick;
+struct FScopeChangeDataLayerEditorContext;
 
 //
 // Forward declarations.
@@ -83,6 +85,8 @@ struct FFoliageUISettings
 	void SetPaintDensity(float InPaintDensity) { PaintDensity = InPaintDensity; }
 	float GetUnpaintDensity() const { return UnpaintDensity; }
 	void SetUnpaintDensity(float InUnpaintDensity) { UnpaintDensity = InUnpaintDensity; }
+	const FActorDataLayer& GetDataLayer() const { return DataLayer; }
+	void SetDataLayer(const FActorDataLayer& InDataLayer) { DataLayer = InDataLayer; }
 	bool GetFilterLandscape() const { return bFilterLandscape ? true : false; }
 	void SetFilterLandscape(bool InbFilterLandscape) { bFilterLandscape = InbFilterLandscape; }
 	bool GetFilterStaticMesh() const { return bFilterStaticMesh ? true : false; }
@@ -184,6 +188,8 @@ private:
 	float PaintDensity;
 	float UnpaintDensity;
 
+	FActorDataLayer DataLayer;
+
 	bool IsInSingleInstantiationMode;
 	bool IsInQuickSingleInstantiationMode;
 	bool IsInQuickEraseMode;
@@ -571,7 +577,10 @@ public:
 
 	/*Find the relevant foliage actor with the foliage type and run the operation*/
 	static void ForEachFoliageInfo(UWorld* InWorld, const UFoliageType* FoliageType, const FSphere& BrushSphere, TFunctionRef<bool(AInstancedFoliageActor* IFA, FFoliageInfo* FoliageInfo, const UFoliageType* FoliageType)> InOperation);
-			
+	
+	/** Changes the Data Layer editing context */
+	void SetDataLayerEditorContext(const FActorDataLayer& DataLayer);
+
 private:
 
 	void BindCommands();
@@ -713,5 +722,7 @@ private:
 
 	/** Flag to know when we are tracking a transaction in mouse delta */
 	bool bTracking;
+
+	TUniquePtr<FScopeChangeDataLayerEditorContext> DataLayerEditorContext;
 };
 
diff --git a/Engine/Source/Editor/FoliageEdit/Private/ProceduralFoliageEditorLibrary.cpp b/Engine/Source/Editor/FoliageEdit/Private/ProceduralFoliageEditorLibrary.cpp
--- a/Engine/Source/Editor/FoliageEdit/Private/ProceduralFoliageEditorLibrary.cpp
+++ b/Engine/Source/Editor/FoliageEdit/Private/ProceduralFoliageEditorLibrary.cpp
@@ -7,6 +7,8 @@
 #include "Algo/Transform.h"
 #include "Framework/Notifications/NotificationManager.h"
 #include "Widgets/Notifications/SNotificationList.h"
+#include "WorldPartition/DataLayer/WorldDataLayers.h"
+#include "WorldPartition/DataLayer/DataLayerSubsystem.h"
 #include "Misc/ScopedSlowTask.h"
 
 #define LOCTEXT_NAMESPACE "ProceduralFoliageEditorLibrary"
@@ -53,35 +55,26 @@ void UProceduralFoliageEditorLibrary::ResimulateProceduralFoliageComponents(cons
 	for (UProceduralFoliageComponent* Component : ProceduralFoliageComponents)
 	{
 		SlowTask.EnterProgressFrame(1);
-		if (Component && Component->FoliageSpawner)
+		if (Component && Component->ResimulateProceduralFoliage([Component](const TArray<FDesiredFoliageInstance>& DesiredFoliageInstances)
+			{
+				FFoliagePaintingGeometryFilter OverrideGeometryFilter;
+				OverrideGeometryFilter.bAllowLandscape = Component->bAllowLandscape;
+				OverrideGeometryFilter.bAllowStaticMesh = Component->bAllowStaticMesh;
+				OverrideGeometryFilter.bAllowBSP = Component->bAllowBSP;
+				OverrideGeometryFilter.bAllowFoliage = Component->bAllowFoliage;
+				OverrideGeometryFilter.bAllowTranslucent = Component->bAllowTranslucent;
+				FEdModeFoliage::AddInstances(Component->GetWorld(), DesiredFoliageInstances, OverrideGeometryFilter, true);
+			}))
 		{
-			TArray <FDesiredFoliageInstance> DesiredFoliageInstances;
-			if (Component->GenerateProceduralContent(DesiredFoliageInstances))
+			// If no instances were spawned, inform the user
+			if (!Component->HasSpawnedAnyInstances())
 			{
-				if (DesiredFoliageInstances.Num() > 0)
-				{
-					Component->RemoveProceduralContent(false);
-
-					FFoliagePaintingGeometryFilter OverrideGeometryFilter;
-					OverrideGeometryFilter.bAllowLandscape = Component->bAllowLandscape;
-					OverrideGeometryFilter.bAllowStaticMesh = Component->bAllowStaticMesh;
-					OverrideGeometryFilter.bAllowBSP = Component->bAllowBSP;
-					OverrideGeometryFilter.bAllowFoliage = Component->bAllowFoliage;
-					OverrideGeometryFilter.bAllowTranslucent = Component->bAllowTranslucent;
-
-					FEdModeFoliage::AddInstances(Component->GetWorld(), DesiredFoliageInstances, OverrideGeometryFilter, true);
-				}
-
-				// If no instances were spawned, inform the user
-				if (!Component->HasSpawnedAnyInstances())
-				{
-					FNotificationInfo Info(LOCTEXT("NothingSpawned_Notification", "Unable to spawn instances. Ensure a large enough surface exists within the volume."));
-					Info.bUseLargeFont = false;
-					Info.bFireAndForget = true;
-					Info.bUseThrobber = false;
-					Info.bUseSuccessFailIcons = true;
-					FSlateNotificationManager::Get().AddNotification(Info);
-				}
+				FNotificationInfo Info(LOCTEXT("NothingSpawned_Notification", "Unable to spawn instances. Ensure a large enough surface exists within the volume."));
+				Info.bUseLargeFont = false;
+				Info.bFireAndForget = true;
+				Info.bUseThrobber = false;
+				Info.bUseSuccessFailIcons = true;
+				FSlateNotificationManager::Get().AddNotification(Info);
 			}
 		}
 	}
diff --git a/Engine/Source/Editor/FoliageEdit/Private/SFoliageEdit.cpp b/Engine/Source/Editor/FoliageEdit/Private/SFoliageEdit.cpp
--- a/Engine/Source/Editor/FoliageEdit/Private/SFoliageEdit.cpp
+++ b/Engine/Source/Editor/FoliageEdit/Private/SFoliageEdit.cpp
@@ -33,6 +33,10 @@
 
 #include "Widgets/Input/SSpinBox.h"
 #include "Editor/PropertyEditor/Public/VariablePrecisionNumericInterface.h"
+#include "DataLayer/DataLayerEditorSubsystem.h"
+#include "DataLayer/DataLayerPropertyTypeCustomizationHelper.h"
+#include "WorldPartition/DataLayer/DataLayer.h"
+#include "WorldPartition/WorldPartitionSubsystem.h"
 
 #define LOCTEXT_NAMESPACE "FoliageEd_Mode"
 
@@ -222,7 +226,51 @@ void SFoliageEdit::Construct(const FArguments& InArgs)
 							.OnValueChanged(this, &SFoliageEdit::SetEraseDensity)
 							.IsEnabled(this, &SFoliageEdit::IsEnabled_EraseDensity)
 						]
-					]					
+					]
+
+					// Data Layer
+					+ SVerticalBox::Slot()
+					.AutoHeight()
+					[
+						SNew(SHorizontalBox)
+						.ToolTipText(LOCTEXT("DataLayer_Tooltip", "The Data Layer to use for the foliage actors"))
+						.Visibility(this, &SFoliageEdit::GetVisibility_DataLayer)
+
+						+ SHorizontalBox::Slot()
+						.Padding(StandardLeftPadding)
+						.FillWidth(1.0f)
+						.VAlign(VAlign_Center)
+						[
+							SNew(STextBlock)
+							.Text(LOCTEXT("DataLayer_Text", "Data Layer"))
+							.Font(StandardFont)
+						]
+						+ SHorizontalBox::Slot()
+						.Padding(StandardRightPadding)
+						.FillWidth(2.0f)
+						.MaxWidth(100.f)
+						.VAlign(VAlign_Center)
+						[
+							SNew(SComboButton)
+							.OnGetMenuContent_Lambda([this]()
+								{ 
+									return FDataLayerPropertyTypeCustomizationHelper::CreateDataLayerMenu([this](const UDataLayer* DataLayer)
+									{ 
+										FoliageEditMode->SetDataLayerEditorContext(FActorDataLayer(DataLayer ? DataLayer->GetFName() : NAME_None)); 
+									});
+								})
+							.ContentPadding(2)
+							.ButtonContent()
+							[
+								SNew(STextBlock)
+								.Text_Lambda([this]()
+								{
+									const UDataLayer* DataLayer = UDataLayerEditorSubsystem::Get()->GetDataLayerFromName(FoliageEditMode->UISettings.GetDataLayer().Name);
+									return UDataLayer::GetDataLayerText(DataLayer);
+								})
+							]
+						]						
+					]
 					
 					+ SVerticalBox::Slot()
 					.Padding(StandardPadding)
@@ -871,6 +919,15 @@ FText SFoliageEdit::GetActiveToolMessage() const
 	return OutText;
 }
 
+EVisibility SFoliageEdit::GetVisibility_DataLayer() const
+{
+	if (UWorld::HasSubsystem<UWorldPartitionSubsystem>(FoliageEditMode->GetWorld()) && (IsPaintTool() || IsPlaceTool() || IsReapplySettingsTool() || IsPaintFillTool()))
+	{
+		return EVisibility::Visible;
+	}
+	return EVisibility::Collapsed;
+}
+
 void SFoliageEdit::SetRadius(float InRadius)
 {
 	FoliageEditMode->UISettings.SetRadius(InRadius);
diff --git a/Engine/Source/Editor/FoliageEdit/Private/SFoliageEdit.h b/Engine/Source/Editor/FoliageEdit/Private/SFoliageEdit.h
--- a/Engine/Source/Editor/FoliageEdit/Private/SFoliageEdit.h
+++ b/Engine/Source/Editor/FoliageEdit/Private/SFoliageEdit.h
@@ -171,6 +171,9 @@ public:	// BRUSH SETTINGS
 	/** Retrieves the tooltip text for the translucent filter */
 	FText GetTooltipText_Translucent() const;
 
+	/** Checks if the Data Layer should appear. */
+	EVisibility GetVisibility_DataLayer() const;
+
 	/** Checks if the radius spinbox should appear. Dependant on the current tool being used. */
 	EVisibility GetVisibility_Radius() const;
 
diff --git a/Engine/Source/Runtime/Engine/Private/ActorPartition/ActorPartitionSubsystem.cpp b/Engine/Source/Runtime/Engine/Private/ActorPartition/ActorPartitionSubsystem.cpp
--- a/Engine/Source/Runtime/Engine/Private/ActorPartition/ActorPartitionSubsystem.cpp
+++ b/Engine/Source/Runtime/Engine/Private/ActorPartition/ActorPartitionSubsystem.cpp
@@ -4,6 +4,7 @@
 #include "ActorPartition/PartitionActor.h"
 #include "WorldPartition/ActorPartition/PartitionActorDesc.h"
 #include "WorldPartition/WorldPartitionSubsystem.h"
+#include "WorldPartition/DataLayer/DataLayerSubsystem.h"
 #include "Engine/World.h"
 #include "EngineUtils.h"
 
@@ -97,7 +98,7 @@ public:
 		return UActorPartitionSubsystem::FCellCoord(0, 0, 0, SpawnLevel);
 	}
 	
-	APartitionActor* GetActor(const TSubclassOf<APartitionActor>& InActorClass, bool bInCreate, const UActorPartitionSubsystem::FCellCoord& InCellCoord, const FGuid& InGuid, uint32 InGridSize, bool bInBoundsSearch, TFunctionRef<void(APartitionActor*)> InActorCreated) override
+	APartitionActor* GetActor(const FActorPartitionIdentifier& InActorPartitionId, bool bInCreate, const UActorPartitionSubsystem::FCellCoord& InCellCoord, uint32 InGridSize, bool bInBoundsSearch, TFunctionRef<void(APartitionActor*)> InActorCreated) override
 	{
 		check(InCellCoord.Level);
 		
@@ -106,7 +107,7 @@ public:
 		{
 			if (APartitionActor* PartitionActor = Cast<APartitionActor>(Actor))
 			{
-				if (PartitionActor->GetGridGuid() == InGuid)
+				if (PartitionActor->GetGridGuid() == InActorPartitionId.GetGridGuid())
 				{
 					FoundActor = PartitionActor;
 					break;
@@ -118,7 +119,7 @@ public:
 		{
 			FActorSpawnParameters SpawnParams;
 			SpawnParams.OverrideLevel = InCellCoord.Level;
-			FoundActor = CastChecked<APartitionActor>(World->SpawnActor(InActorClass, nullptr, nullptr, SpawnParams));
+			FoundActor = CastChecked<APartitionActor>(World->SpawnActor(InActorPartitionId.GetClass(), nullptr, nullptr, SpawnParams));
 			InActorCreated(FoundActor);
 		}
 
@@ -185,19 +186,20 @@ public:
 		return UActorPartitionSubsystem::FCellCoord::GetCellCoord(GetParams.LocationHint, World->PersistentLevel, GridSize);
 	}
 
-	virtual APartitionActor* GetActor(const TSubclassOf<APartitionActor>& InActorClass, bool bInCreate, const UActorPartitionSubsystem::FCellCoord& InCellCoord, const FGuid& InGuid, uint32 InGridSize, bool bInBoundsSearch, TFunctionRef<void(APartitionActor*)> InActorCreated)
+	virtual APartitionActor* GetActor(const FActorPartitionIdentifier& InActorPartitionId, bool bInCreate, const UActorPartitionSubsystem::FCellCoord& InCellCoord, uint32 InGridSize, bool bInBoundsSearch, TFunctionRef<void(APartitionActor*)> InActorCreated)
 	{
 		APartitionActor* FoundActor = nullptr;
 		bool bUnloadedActorExists = false;
-		auto FindActor = [&FoundActor, &bUnloadedActorExists, InActorClass, InCellCoord, InGuid, InGridSize](const FWorldPartitionActorDesc* ActorDesc)
+		auto FindActor = [&FoundActor, &bUnloadedActorExists, InCellCoord, InActorPartitionId, InGridSize, ThisWorld = World](const FWorldPartitionActorDesc* ActorDesc)
 		{
-			check(ActorDesc->GetActorClass()->IsChildOf(InActorClass));
+			check(ActorDesc->GetActorClass()->IsChildOf(InActorPartitionId.GetClass()));
 			FPartitionActorDesc* PartitionActorDesc = (FPartitionActorDesc*)ActorDesc;
 			if ((PartitionActorDesc->GridIndexX == InCellCoord.X) &&
 				(PartitionActorDesc->GridIndexY == InCellCoord.Y) &&
 				(PartitionActorDesc->GridIndexZ == InCellCoord.Z) &&
 				(PartitionActorDesc->GridSize == InGridSize) &&
-				(PartitionActorDesc->GridGuid == InGuid))
+				(PartitionActorDesc->GridGuid == InActorPartitionId.GetGridGuid()) &&
+				(FDataLayerEditorContext(ThisWorld, PartitionActorDesc->GetDataLayers()).GetHash() == InActorPartitionId.GetDataLayerEditorContextHash()))
 			{
 				AActor* DescActor = ActorDesc->GetActor();
 
@@ -209,7 +211,7 @@ public:
 				}
 
 				FoundActor = CastChecked<APartitionActor>(DescActor);
-				check(FoundActor->GridSize == InGridSize && FoundActor->GetGridGuid() == InGuid);
+				check(FoundActor->GridSize == InGridSize && FoundActor->GetGridGuid() == InActorPartitionId.GetGridGuid());
 				// Typically, this could happen if a PartitionActor was manually deleted (it's now be prevented by APartitionActor::CanDeleteSelectedActor);
 				// If at some point it's decided to be supported, a solution for a found PendingKill PartitionActor is to
 				// call Modify(), Rename the actor to a unique/trash name and behave as if no actor was found
@@ -222,11 +224,11 @@ public:
 		FBox CellBounds = UActorPartitionSubsystem::FCellCoord::GetCellBounds(InCellCoord, InGridSize);
 		if (bInBoundsSearch)
 		{
-			WorldPartition->ForEachIntersectingActorDesc(CellBounds, InActorClass, FindActor);
+			WorldPartition->ForEachIntersectingActorDesc(CellBounds, InActorPartitionId.GetClass(), FindActor);
 		}
 		else
 		{
-			WorldPartition->ForEachActorDesc(InActorClass, FindActor);
+			WorldPartition->ForEachActorDesc(InActorPartitionId.GetClass(), FindActor);
 		}
 				
 		if (bUnloadedActorExists)
@@ -238,17 +240,22 @@ public:
 		{
 			TStringBuilderWithBuffer<TCHAR, NAME_SIZE> ActorNameBuilder;
 
-			ActorNameBuilder += InActorClass->GetName();
+			ActorNameBuilder += InActorPartitionId.GetClass()->GetName();
 			ActorNameBuilder += TEXT("_");
 
-			if (InGuid.IsValid())
+			if (InActorPartitionId.GetGridGuid().IsValid())
 			{
-				ActorNameBuilder += InGuid.ToString(EGuidFormats::Base36Encoded);
+				ActorNameBuilder += InActorPartitionId.GetGridGuid().ToString(EGuidFormats::Base36Encoded);
 				ActorNameBuilder += TEXT("_");
 			}
 
 			ActorNameBuilder += FString::Printf(TEXT("%d_%d_%d"), InCellCoord.X, InCellCoord.Y, InCellCoord.Z);
 
+			if (InActorPartitionId.GetDataLayerEditorContextHash() != FDataLayerEditorContext::EmptyHash)
+			{
+				ActorNameBuilder += FString::Printf(TEXT("_%X"), InActorPartitionId.GetDataLayerEditorContextHash());
+			}
+
 			FActorSpawnParameters SpawnParams;
 			SpawnParams.OverrideLevel = InCellCoord.Level;
 			SpawnParams.Name = *ActorNameBuilder;
@@ -262,14 +269,20 @@ public:
 			}
 						
 			FVector CellCenter(CellBounds.GetCenter());
-			FoundActor = CastChecked<APartitionActor>(World->SpawnActor(InActorClass, &CellCenter, nullptr, SpawnParams));
+			FoundActor = CastChecked<APartitionActor>(World->SpawnActor(InActorPartitionId.GetClass(), &CellCenter, nullptr, SpawnParams));
 			FoundActor->GridSize = InGridSize;
 			FoundActor->bLockLocation = true;
 			
 			InActorCreated(FoundActor);
 
 			// Once actor is created, update its label
-			FoundActor->SetActorLabel(FString::Printf(TEXT("%s_%d_%d_%d"), *InActorClass->GetName(), InCellCoord.X, InCellCoord.Y, InCellCoord.Z));
+			TStringBuilderWithBuffer<TCHAR, NAME_SIZE> ActorLabelBuilder;
+			ActorLabelBuilder += FString::Printf(TEXT("%s_%d_%d_%d"), *InActorPartitionId.GetClass()->GetName(), InCellCoord.X, InCellCoord.Y, InCellCoord.Z);
+			if (InActorPartitionId.GetDataLayerEditorContextHash() != FDataLayerEditorContext::EmptyHash)
+			{
+				ActorLabelBuilder += FString::Printf(TEXT("_%X"), InActorPartitionId.GetDataLayerEditorContextHash());
+			}
+			FoundActor->SetActorLabel(*ActorLabelBuilder);
 		}
 
 		check(FoundActor || !bInCreate);
@@ -363,51 +376,58 @@ APartitionActor* UActorPartitionSubsystem::GetActor(const TSubclassOf<APartition
 {
 	const uint32 GridSize = InGridSize > 0 ? InGridSize : InActorClass->GetDefaultObject<APartitionActor>()->GetDefaultGridSize(GetWorld());
 	
-	TMap<UClass*, TMap<FGuid, TWeakObjectPtr<APartitionActor>>>* ActorsPerClass = PartitionedActors.Find(InCellCoords);
+	UWorld* World = GetWorld();
+	
+	auto WrappedInActorCreated = [World, InActorCreated](APartitionActor* PartitionActor)
+	{
+		if (UDataLayerSubsystem* DataLayerSubsystem = World->GetSubsystem<UDataLayerSubsystem>())
+		{
+			if (const AWorldDataLayers* WorldDataLayers = AWorldDataLayers::Get(World))
+			{
+				for (const FName& DataLayer : DataLayerSubsystem->GetDataLayerEditorContext().GetDataLayers())
+				{
+					PartitionActor->AddDataLayer(WorldDataLayers->GetDataLayerFromName(DataLayer));
+				}
+			}
+		}
+
+		InActorCreated(PartitionActor);
+	};
+
+	UDataLayerSubsystem* DataLayerSubsystem = World->GetSubsystem<UDataLayerSubsystem>();
+	FActorPartitionIdentifier ActorPartitionId(InActorClass, InGuid, DataLayerSubsystem ? DataLayerSubsystem->GetDataLayerEditorContext().GetHash() : FDataLayerEditorContext::EmptyHash);
+	TMap<FActorPartitionIdentifier, TWeakObjectPtr<APartitionActor>>* ActorsPerId = PartitionedActors.Find(InCellCoords);
 	APartitionActor* FoundActor = nullptr;
-	if (!ActorsPerClass)
+	if (!ActorsPerId)
 	{
-		FoundActor = ActorPartition->GetActor(InActorClass, bInCreate, InCellCoords, InGuid, GridSize, bInBoundsSearch, InActorCreated);
+		FoundActor = ActorPartition->GetActor(ActorPartitionId, bInCreate, InCellCoords, GridSize, bInBoundsSearch, WrappedInActorCreated);
 		if (FoundActor)
 		{
-			PartitionedActors.Add(InCellCoords).Add(InActorClass).Add(InGuid, FoundActor);
+			PartitionedActors.Add(InCellCoords).Add(ActorPartitionId, FoundActor);
 		}
 	}
 	else
 	{
-		TMap<FGuid, TWeakObjectPtr<APartitionActor>>* ActorsPerGuid = ActorsPerClass->Find(InActorClass);
-		if (!ActorsPerGuid)
+		TWeakObjectPtr<APartitionActor>* ActorPtr = ActorsPerId->Find(ActorPartitionId);
+		if (!ActorPtr || !ActorPtr->IsValid())
 		{
-			FoundActor = ActorPartition->GetActor(InActorClass, bInCreate, InCellCoords, InGuid, GridSize, bInBoundsSearch, InActorCreated);
+			FoundActor = ActorPartition->GetActor(ActorPartitionId, bInCreate, InCellCoords, GridSize, bInBoundsSearch, WrappedInActorCreated);
 			if (FoundActor)
 			{
-				ActorsPerClass->Add(InActorClass).Add(InGuid, FoundActor);
+				if (!ActorPtr)
+				{
+					ActorsPerId->Add(ActorPartitionId, FoundActor);
+				}
+				else
+				{
+					*ActorPtr = FoundActor;
+				}
 			}
 		}
 		else
 		{
-			TWeakObjectPtr<APartitionActor>* ActorPtr = ActorsPerGuid->Find(InGuid);
-			if (!ActorPtr || !ActorPtr->IsValid())
-			{
-				FoundActor = ActorPartition->GetActor(InActorClass, bInCreate, InCellCoords, InGuid, GridSize, bInBoundsSearch, InActorCreated);
-				if (FoundActor)
-				{
-					if (!ActorPtr)
-					{
-						ActorsPerGuid->Add(InGuid, FoundActor);
-					}
-					else
-					{
-						*ActorPtr = FoundActor;
-					}
-				}
-			}
-			else
-			{
-				FoundActor = ActorPtr->Get();
-			}
+			FoundActor = ActorPtr->Get();
 		}
-
 	}
 	
 	return FoundActor;
diff --git a/Engine/Source/Runtime/Engine/Private/WorldPartition/DataLayer/DataLayer.cpp b/Engine/Source/Runtime/Engine/Private/WorldPartition/DataLayer/DataLayer.cpp
--- a/Engine/Source/Runtime/Engine/Private/WorldPartition/DataLayer/DataLayer.cpp
+++ b/Engine/Source/Runtime/Engine/Private/WorldPartition/DataLayer/DataLayer.cpp
@@ -6,6 +6,8 @@
 
 #include "WorldPartition/DataLayer/DataLayer.h"
 
+#define LOCTEXT_NAMESPACE "DataLayer"
+
 UDataLayer::UDataLayer(const FObjectInitializer& ObjectInitializer)
 : Super(ObjectInitializer)
 #if WITH_EDITORONLY_DATA
@@ -56,4 +58,11 @@ void UDataLayer::SetIsDynamicallyLoadedInEditor(bool bInIsDynamicallyLoadedInEdi
 	}
 }
 
-#endif
\ No newline at end of file
+FText UDataLayer::GetDataLayerText(const UDataLayer* InDataLayer)
+{
+	return InDataLayer ? FText::FromName(InDataLayer->GetDataLayerLabel()) : LOCTEXT("InvalidDataLayerLabel", "<None>");
+}
+
+#endif
+
+#undef LOCTEXT_NAMESPACE
\ No newline at end of file
diff --git a/Engine/Source/Runtime/Engine/Private/WorldPartition/DataLayer/DataLayerEditorContext.cpp b/Engine/Source/Runtime/Engine/Private/WorldPartition/DataLayer/DataLayerEditorContext.cpp
new file mode 100644
--- /dev/null
+++ b/Engine/Source/Runtime/Engine/Private/WorldPartition/DataLayer/DataLayerEditorContext.cpp
@@ -0,0 +1,77 @@
+// Copyright Epic Games, Inc. All Rights Reserved.
+
+#include "WorldPartition/DataLayer/DataLayerEditorContext.h"
+
+#if WITH_EDITOR
+
+#include "WorldPartition/DataLayer/DataLayer.h"
+#include "WorldPartition/DataLayer/WorldDataLayers.h"
+#include "WorldPartition/DataLayer/DataLayerSubsystem.h"
+#include "Engine/World.h"
+
+/*
+ * FDataLayerEditorContext
+ */
+FDataLayerEditorContext::FDataLayerEditorContext(UWorld* InWorld, const TArray<FName>& InDataLayers)
+	: Hash(FDataLayerEditorContext::EmptyHash)
+{
+	const AWorldDataLayers* WorldDataLayers = AWorldDataLayers::Get(InWorld);
+	if (!WorldDataLayers)
+	{
+		return;
+	}
+
+	for (const FName& DataLayerName : InDataLayers)
+	{
+		if (const UDataLayer* DataLayerObject = WorldDataLayers->GetDataLayerFromName(DataLayerName))
+		{
+			DataLayers.AddUnique(DataLayerObject->GetFName());
+		}
+	}
+
+	if (DataLayers.Num())
+	{
+		DataLayers.Sort([](const FName& A, const FName& B) { return A.ToString() < B.ToString(); });
+		for (FName LayerName : DataLayers)
+		{
+			Hash = FCrc::StrCrc32(*LayerName.ToString(), Hash);
+		}
+		check(Hash != FDataLayerEditorContext::EmptyHash);
+	}
+}
+
+/*
+ * FScopeChangeDataLayerEditorContext
+ */
+
+FScopeChangeDataLayerEditorContext::FScopeChangeDataLayerEditorContext(UWorld* InWorld, const FDataLayerEditorContext& InContext)
+	: World(InWorld)
+{
+	Initialize(InContext);
+}
+
+FScopeChangeDataLayerEditorContext::FScopeChangeDataLayerEditorContext(UWorld* InWorld, const FActorDataLayer& InContextDataLayer)
+	: World(InWorld)
+{
+	const AWorldDataLayers* WorldDataLayers = AWorldDataLayers::Get(InWorld);
+	Initialize(FDataLayerEditorContext(InWorld, WorldDataLayers ? WorldDataLayers->GetDataLayerNames({ InContextDataLayer }) : TArray<FName>()));
+}
+
+void FScopeChangeDataLayerEditorContext::Initialize(const FDataLayerEditorContext& InContext)
+{
+	if (const UDataLayerSubsystem* DataLayerSubsystem = World.Get() ? World->GetSubsystem<UDataLayerSubsystem>() : nullptr)
+	{
+		OldContext = DataLayerSubsystem->GetDataLayerEditorContext();
+		DataLayerSubsystem->DataLayerEditorContext = InContext;
+	}
+}
+
+FScopeChangeDataLayerEditorContext::~FScopeChangeDataLayerEditorContext()
+{
+	if (const UDataLayerSubsystem* DataLayerSubsystem = World.Get() ? World->GetSubsystem<UDataLayerSubsystem>() : nullptr)
+	{
+		DataLayerSubsystem->DataLayerEditorContext = OldContext;
+	}
+}
+
+#endif
\ No newline at end of file
diff --git a/Engine/Source/Runtime/Engine/Private/WorldPartition/DataLayer/WorldDataLayers.cpp b/Engine/Source/Runtime/Engine/Private/WorldPartition/DataLayer/WorldDataLayers.cpp
--- a/Engine/Source/Runtime/Engine/Private/WorldPartition/DataLayer/WorldDataLayers.cpp
+++ b/Engine/Source/Runtime/Engine/Private/WorldPartition/DataLayer/WorldDataLayers.cpp
@@ -87,6 +87,35 @@ FName AWorldDataLayers::GenerateUniqueDataLayerLabel(const FName& InDataLayerLab
 	return UniqueNewDataLayerLabel;
 }
 
+TArray<FName> AWorldDataLayers::GetDataLayerNames(const TArray<FActorDataLayer>& InDataLayers) const
+{
+	TArray<FName> OutDataLayerNames;
+	OutDataLayerNames.Reserve(DataLayers.Num());
+
+	for (const UDataLayer* DataLayer : GetDataLayerObjects(InDataLayers))
+	{
+		OutDataLayerNames.Add(DataLayer->GetFName());
+	}
+
+	return OutDataLayerNames;
+}
+
+TArray<const UDataLayer*> AWorldDataLayers::GetDataLayerObjects(const TArray<FActorDataLayer>& InDataLayers) const
+{
+	TArray<const UDataLayer*> OutDataLayers;
+	OutDataLayers.Reserve(DataLayers.Num());
+
+	for (const FActorDataLayer& DataLayer : InDataLayers)
+	{
+		if (const UDataLayer* DataLayerObject = GetDataLayerFromName(DataLayer.Name))
+		{
+			OutDataLayers.AddUnique(DataLayerObject);
+		}
+	}
+
+	return OutDataLayers;
+}
+
 UDataLayer* AWorldDataLayers::CreateDataLayer()
 {
 	Modify();
diff --git a/Engine/Source/Runtime/Engine/Public/ActorPartition/ActorPartitionSubsystem.h b/Engine/Source/Runtime/Engine/Public/ActorPartition/ActorPartitionSubsystem.h
--- a/Engine/Source/Runtime/Engine/Public/ActorPartition/ActorPartitionSubsystem.h
+++ b/Engine/Source/Runtime/Engine/Public/ActorPartition/ActorPartitionSubsystem.h
@@ -6,10 +6,10 @@
 #include "Templates/SubclassOf.h"
 #include "Misc/HashBuilder.h"
 #include "Misc/Guid.h"
+#include "ActorPartition/PartitionActor.h"
 #include "ActorPartitionSubsystem.generated.h"
 
 class FBaseActorPartition;
-class APartitionActor;
 
 #if WITH_EDITOR
 /**
@@ -38,6 +38,35 @@ struct ENGINE_API FActorPartitionGetParams
 	int32 GridSize;
 };
 
+/**
+ * FActorPartitionIdentifier
+ */
+
+struct ENGINE_API FActorPartitionIdentifier
+{
+	FActorPartitionIdentifier(UClass* InClass, const FGuid& InGridGuid, const uint32& InDataLayerEditorContextHash)
+		: Class(InClass)
+		, GridGuid(InGridGuid)
+		, DataLayerEditorContextHash(InDataLayerEditorContextHash)
+	{}
+
+	bool operator == (const FActorPartitionIdentifier& Other) const { return Class == Other.Class && GridGuid == Other.GridGuid && DataLayerEditorContextHash == Other.DataLayerEditorContextHash; }
+
+	friend uint32 GetTypeHash(const FActorPartitionIdentifier& Id)
+	{
+		return GetTypeHash(Id.Class.Get()->GetName()) ^ GetTypeHash(Id.GridGuid) ^ GetTypeHash(Id.DataLayerEditorContextHash);
+	}
+
+	const TSubclassOf<APartitionActor>& GetClass() const { return Class; }
+	const FGuid& GetGridGuid() const { return GridGuid; }
+	uint32 GetDataLayerEditorContextHash() const { return DataLayerEditorContextHash; }
+
+private:
+	TSubclassOf<APartitionActor> Class;
+	FGuid GridGuid;
+	uint32 DataLayerEditorContextHash;
+};
+
 #endif
 
 /**
@@ -146,7 +175,7 @@ private:
 
 	void InitializeActorPartition();
 
-	TMap<FCellCoord, TMap<UClass*, TMap<FGuid, TWeakObjectPtr<APartitionActor>>>> PartitionedActors;
+	TMap<FCellCoord, TMap<FActorPartitionIdentifier, TWeakObjectPtr<APartitionActor>>> PartitionedActors;
 	TUniquePtr<FBaseActorPartition> ActorPartition;
 	
 	FDelegateHandle ActorPartitionHashInvalidatedHandle;
@@ -164,7 +193,7 @@ public:
 	virtual ~FBaseActorPartition() {}
 
 	virtual UActorPartitionSubsystem::FCellCoord GetActorPartitionHash(const FActorPartitionGetParams& GetParams) const = 0;
-	virtual APartitionActor* GetActor(const TSubclassOf<APartitionActor>& InActorClass, bool bInCreate, const UActorPartitionSubsystem::FCellCoord& InCellCoord, const FGuid& InGuid, uint32 InGridSize, bool bInBoundsSearch, TFunctionRef<void(APartitionActor*)> InActorCreated) = 0;
+	virtual APartitionActor* GetActor(const FActorPartitionIdentifier& InActorPartitionId, bool bInCreate, const UActorPartitionSubsystem::FCellCoord& InCellCoord, uint32 InGridSize, bool bInBoundsSearch, TFunctionRef<void(APartitionActor*)> InActorCreated) = 0;
 	virtual void ForEachRelevantActor(const TSubclassOf<APartitionActor>& InActorClass, const FBox& IntersectionBounds, TFunctionRef<bool(APartitionActor*)>InOperation) const = 0;
 
 	DECLARE_EVENT_OneParam(FBaseActorPartition, FOnActorPartitionHashInvalidated, const UActorPartitionSubsystem::FCellCoord&);
diff --git a/Engine/Source/Runtime/Engine/Public/WorldPartition/DataLayer/DataLayer.h b/Engine/Source/Runtime/Engine/Public/WorldPartition/DataLayer/DataLayer.h
--- a/Engine/Source/Runtime/Engine/Public/WorldPartition/DataLayer/DataLayer.h
+++ b/Engine/Source/Runtime/Engine/Public/WorldPartition/DataLayer/DataLayer.h
@@ -22,6 +22,7 @@ public:
 	void SetIsDynamicallyLoaded(bool bInIsDynamicallyLoaded);
 	void SetIsDynamicallyLoadedInEditor(bool bInIsDynamicallyLoadedInEditor);
 	bool IsDynamicallyLoadedInEditor() const { return !IsDynamicallyLoaded() || bIsDynamicallyLoadedInEditor; }
+	static FText GetDataLayerText(const UDataLayer* InDataLayer);
 #endif
 	
 	FName GetDataLayerLabel() const  { return DataLayerLabel; }
diff --git a/Engine/Source/Runtime/Engine/Public/WorldPartition/DataLayer/DataLayerEditorContext.h b/Engine/Source/Runtime/Engine/Public/WorldPartition/DataLayer/DataLayerEditorContext.h
new file mode 100644
--- /dev/null
+++ b/Engine/Source/Runtime/Engine/Public/WorldPartition/DataLayer/DataLayerEditorContext.h
@@ -0,0 +1,45 @@
+// Copyright Epic Games, Inc. All Rights Reserved.
+
+#pragma once
+
+#include "Engine/EngineTypes.h"
+
+#if WITH_EDITOR
+
+struct FActorDataLayer;
+class UWorld;
+
+/**
+ * FDataLayerEditorContext
+ */
+
+struct ENGINE_API FDataLayerEditorContext
+{
+public:
+	static const uint32 EmptyHash = 0;
+	FDataLayerEditorContext() : Hash(FDataLayerEditorContext::EmptyHash) {}
+	FDataLayerEditorContext(UWorld* InWorld, const TArray<FName>& InDataLayers);
+	FORCEINLINE bool IsEmpty() const { return (Hash == FDataLayerEditorContext::EmptyHash) && DataLayers.IsEmpty(); }
+	FORCEINLINE uint32 GetHash() const { return Hash; }
+	FORCEINLINE const TArray<FName>& GetDataLayers() const { return DataLayers; }
+private:
+	uint32 Hash;
+	TArray<FName> DataLayers;
+};
+
+/**
+ * FScopeChangeDataLayerEditorContext
+ */
+
+struct ENGINE_API FScopeChangeDataLayerEditorContext
+{
+	FScopeChangeDataLayerEditorContext(UWorld* InWorld, const FDataLayerEditorContext& InContext);
+	FScopeChangeDataLayerEditorContext(UWorld* InWorld, const FActorDataLayer& InContextDataLayer);
+	~FScopeChangeDataLayerEditorContext();
+private:
+	void Initialize(const FDataLayerEditorContext& InContext);
+
+	TWeakObjectPtr<UWorld> World;
+	FDataLayerEditorContext OldContext;
+};
+#endif
\ No newline at end of file
diff --git a/Engine/Source/Runtime/Engine/Public/WorldPartition/DataLayer/DataLayerSubsystem.h b/Engine/Source/Runtime/Engine/Public/WorldPartition/DataLayer/DataLayerSubsystem.h
--- a/Engine/Source/Runtime/Engine/Public/WorldPartition/DataLayer/DataLayerSubsystem.h
+++ b/Engine/Source/Runtime/Engine/Public/WorldPartition/DataLayer/DataLayerSubsystem.h
@@ -5,6 +5,7 @@
 #include "Subsystems/WorldSubsystem.h"
 #include "WorldPartition/DataLayer/ActorDataLayer.h"
 #include "WorldPartition/DataLayer/DataLayer.h"
+#include "WorldPartition/DataLayer/DataLayerEditorContext.h"
 #include "DataLayerSubsystem.generated.h"
 
 /**
@@ -53,6 +54,16 @@ public:
 
 	void Draw(class UCanvas* Canvas, class APlayerController* PC);
 
+#if WITH_EDITOR
+public:
+	FORCEINLINE const FDataLayerEditorContext& GetDataLayerEditorContext() const { return DataLayerEditorContext; }
+	FORCEINLINE bool HasDataLayerEditorContext() const { return !DataLayerEditorContext.IsEmpty(); }
+
+private:
+	mutable FDataLayerEditorContext DataLayerEditorContext;
+	friend struct FScopeChangeDataLayerEditorContext;
+#endif
+
 private:
 	/** Console command used to toggle activation of a DataLayer */
 	static class FAutoConsoleCommand ToggleDataLayerActivation;
diff --git a/Engine/Source/Runtime/Engine/Public/WorldPartition/DataLayer/WorldDataLayers.h b/Engine/Source/Runtime/Engine/Public/WorldPartition/DataLayer/WorldDataLayers.h
--- a/Engine/Source/Runtime/Engine/Public/WorldPartition/DataLayer/WorldDataLayers.h
+++ b/Engine/Source/Runtime/Engine/Public/WorldPartition/DataLayer/WorldDataLayers.h
@@ -31,6 +31,11 @@ public:
 	//~ Begin AActor Interface
 	virtual bool ForceExternalActorLevelReference() const override { return true; }
 	//~ End AActor Interface
+
+	//~ Begin Helper Functions
+	TArray<const UDataLayer*> GetDataLayerObjects(const TArray<FActorDataLayer>& DataLayers) const;
+	TArray<FName> GetDataLayerNames(const TArray<FActorDataLayer>& DataLayers) const;
+	//~ End Helper Functions
 #endif
 	
 	bool ContainsDataLayer(const UDataLayer* InDataLayer) const;
diff --git a/Engine/Source/Runtime/Foliage/Private/ProceduralFoliageComponent.cpp b/Engine/Source/Runtime/Foliage/Private/ProceduralFoliageComponent.cpp
--- a/Engine/Source/Runtime/Foliage/Private/ProceduralFoliageComponent.cpp
+++ b/Engine/Source/Runtime/Foliage/Private/ProceduralFoliageComponent.cpp
@@ -14,6 +14,10 @@
 
 #if WITH_EDITOR
 #include "WorldPartition/WorldPartition.h"
+#include "WorldPartition/WorldPartitionSubsystem.h"
+#include "WorldPartition/DataLayer/WorldDataLayers.h"
+#include "WorldPartition/DataLayer/DataLayerSubsystem.h"
+#include "WorldPartition/DataLayer/DataLayerEditorContext.h"
 #endif
 
 
@@ -284,6 +288,38 @@ void UProceduralFoliageComponent::PostEditImport()
 	ProceduralGuid = FGuid::NewGuid();
 }
 
+bool UProceduralFoliageComponent::ResimulateProceduralFoliage(TFunctionRef<void(const TArray<FDesiredFoliageInstance>&)> AddInstancesFunc)
+{
+#if WITH_EDITOR
+	if (FoliageSpawner)
+	{
+		TArray<FDesiredFoliageInstance> DesiredFoliageInstances;
+		if (GenerateProceduralContent(DesiredFoliageInstances))
+		{
+			if (DesiredFoliageInstances.Num() > 0)
+			{
+				{
+					// Remove old foliage instances
+					FScopeChangeDataLayerEditorContext ScopeContext(GetWorld(), LastSimulationDataLayer);
+					RemoveProceduralContent(false);
+				}
+
+				{
+					// Add new foliage instances
+					FScopeChangeDataLayerEditorContext ScopeContext(GetWorld(), DataLayer);
+					AddInstancesFunc(DesiredFoliageInstances);
+					LastSimulationDataLayer = DataLayer;
+				}
+			}
+
+			return true;
+		}
+	}
+#endif
+
+	return false;
+}
+
 bool UProceduralFoliageComponent::GenerateProceduralContent(TArray <FDesiredFoliageInstance>& OutInstances)
 {
 #if WITH_EDITOR
@@ -317,4 +353,23 @@ bool UProceduralFoliageComponent::HasSpawnedAnyInstances()
 	return false;
 }
 
+#if WITH_EDITOR
+bool UProceduralFoliageComponent::CanEditChange(const FProperty* InProperty) const
+{
+	if (InProperty->GetFName() == GET_MEMBER_NAME_CHECKED(UProceduralFoliageComponent, DataLayer))
+	{
+		if (!IsTemplate())
+		{
+			const bool bIsPartitionedWorld = UWorld::HasSubsystem<UWorldPartitionSubsystem>(GetWorld());
+			if (!bIsPartitionedWorld)
+			{
+				return false;
+			}
+		}
+	}
+	return Super::CanEditChange(InProperty);
+}
+#endif
+
+
 #undef LOCTEXT_NAMESPACE
diff --git a/Engine/Source/Runtime/Foliage/Public/ProceduralFoliageComponent.h b/Engine/Source/Runtime/Foliage/Public/ProceduralFoliageComponent.h
--- a/Engine/Source/Runtime/Foliage/Public/ProceduralFoliageComponent.h
+++ b/Engine/Source/Runtime/Foliage/Public/ProceduralFoliageComponent.h
@@ -7,6 +7,7 @@
 #include "Misc/Guid.h"
 #include "Components/ActorComponent.h"
 #include "InstancedFoliage.h"
+#include "WorldPartition/DataLayer/ActorDataLayer.h"
 #include "ProceduralFoliageComponent.generated.h"
 
 class AVolume;
@@ -75,11 +76,25 @@ class FOLIAGE_API UProceduralFoliageComponent : public UActorComponent
 	/** Whether to visualize the tiles used for the foliage spawner simulation */
 	UPROPERTY(Category = "ProceduralFoliage", BlueprintReadWrite, EditAnywhere)
 	bool bShowDebugTiles;
+
+	/** Data Layer to assign to foliage actors generated by foliage spawner simulation */
+	UPROPERTY(Category = "ProceduralFoliage", BlueprintReadWrite, EditAnywhere)
+	FActorDataLayer DataLayer;
+#endif
+
+#if WITH_EDITOR
+	virtual bool CanEditChange(const FProperty* InProperty) const;
 #endif
 
 	// UObject interface
 	virtual void PostEditImport() override;
 
+	/**
+	 * Runs the procedural foliage simulation, removes the old result, creates instances with the new result
+	 * @return True if the simulation succeeded
+	 */ 
+	bool ResimulateProceduralFoliage(TFunctionRef<void(const TArray<FDesiredFoliageInstance>&)> AddInstancesFunc);
+
 	/** 
 	 * Runs the procedural foliage simulation to generate a list of desired instances to spawn.
 	 * @return True if the simulation succeeded
@@ -117,4 +132,10 @@ private:
 	
 	UPROPERTY()
 	FGuid ProceduralGuid;
+
+#if WITH_EDITORONLY_DATA
+	UPROPERTY()
+	FActorDataLayer LastSimulationDataLayer;
+	friend class UProceduralFoliageEditorLibrary;
+#endif
 };
