Fixed issue preventing users from creating Packed LevelInstance blueprints.

diff --git a/Engine/Source/Runtime/Engine/Private/LevelInstance/LevelInstanceSubsystem.cpp b/Engine/Source/Runtime/Engine/Private/LevelInstance/LevelInstanceSubsystem.cpp
--- a/Engine/Source/Runtime/Engine/Private/LevelInstance/LevelInstanceSubsystem.cpp
+++ b/Engine/Source/Runtime/Engine/Private/LevelInstance/LevelInstanceSubsystem.cpp
@@ -737,7 +737,7 @@ ALevelInstance* ULevelInstanceSubsystem::CreateLevelInstanceFrom(const TArray<AA
 		FString PackagePath = LongPackageName.Mid(0, LastSlashIndex == INDEX_NONE ? MAX_int32 : LastSlashIndex);
 		FString AssetName = WorldPtr.GetAssetName() + FPackedLevelInstanceBuilder::GetPackedBPSuffix();
 		const bool bCompile = true;
-		if (UBlueprint* NewBP = FPackedLevelInstanceBuilder::CreatePackedLevelInstanceBlueprint(PackagePath, AssetName, bCompile))
+		if (UBlueprint* NewBP = FPackedLevelInstanceBuilder::CreatePackedLevelInstanceBlueprintWithDialog(AssetName, PackagePath, bCompile))
 		{
 			NewLevelInstanceActor = GetWorld()->SpawnActor<APackedLevelInstance>(NewBP->GeneratedClass, SpawnParams);
 			APackedLevelInstance* CDO = Cast<APackedLevelInstance>(NewBP->GeneratedClass->GetDefaultObject());
diff --git a/Engine/Source/Runtime/Engine/Private/LevelInstance/Packed/PackedLevelInstanceBuilder.cpp b/Engine/Source/Runtime/Engine/Private/LevelInstance/Packed/PackedLevelInstanceBuilder.cpp
--- a/Engine/Source/Runtime/Engine/Private/LevelInstance/Packed/PackedLevelInstanceBuilder.cpp
+++ b/Engine/Source/Runtime/Engine/Private/LevelInstance/Packed/PackedLevelInstanceBuilder.cpp
@@ -71,7 +71,7 @@ const FString& FPackedLevelInstanceBuilder::GetPackedBPSuffix()
 	return BPSuffix;
 }
 
-UBlueprint* FPackedLevelInstanceBuilder::CreatePackedLevelInstanceBlueprintWithDialog(const FString& InPackagePath, const FString& InAssetName, bool bInCompile)
+UBlueprint* FPackedLevelInstanceBuilder::CreatePackedLevelInstanceBlueprintWithDialog(const FString& InAssetName, const FString& InPackagePath, bool bInCompile)
 {
 	FSaveAssetDialogConfig SaveAssetDialogConfig;
 	SaveAssetDialogConfig.DialogTitleOverride = LOCTEXT("SaveAssetDialogTitle", "Save Asset As");
@@ -384,7 +384,7 @@ bool FPackedLevelInstanceBuilder::CreateOrUpdateBlueprintFromPacked(APackedLevel
 		FString AssetName = InActor->GetWorldAsset().GetAssetName() + GetPackedBPSuffix();
 		const bool bCompile = false;
 
-		BP = CreatePackedLevelInstanceBlueprintWithDialog(PackagePath, AssetName, bCompile);
+		BP = CreatePackedLevelInstanceBlueprintWithDialog(AssetName, PackagePath, bCompile);
 	}
 
 	if (BP && BP->SimpleConstructionScript)
diff --git a/Engine/Source/Runtime/Engine/Public/LevelInstance/Packed/PackedLevelInstanceBuilder.h b/Engine/Source/Runtime/Engine/Public/LevelInstance/Packed/PackedLevelInstanceBuilder.h
--- a/Engine/Source/Runtime/Engine/Public/LevelInstance/Packed/PackedLevelInstanceBuilder.h
+++ b/Engine/Source/Runtime/Engine/Public/LevelInstance/Packed/PackedLevelInstanceBuilder.h
@@ -44,7 +44,7 @@ public:
 
 	static const FString& GetPackedBPSuffix();
 	/* Creates a new APackedLevelInstance Blueprint using InPackagePath/InAssetName as hint for path. Prompts the user to input the final asset name. */
-	static UBlueprint* CreatePackedLevelInstanceBlueprintWithDialog(const FString& InPackagePath, const FString& InAssetName, bool bInCompile);
+	static UBlueprint* CreatePackedLevelInstanceBlueprintWithDialog(const FString& InAssetName, const FString& InPackagePath, bool bInCompile);
 	/* Creates a new APackedLevelInstance Blueprint using InPackagePath/InAssetName */
 	static UBlueprint* CreatePackedLevelInstanceBlueprint(const FString& InSaveAssetName, const FString& InPackagePath, bool bInCompile);
 	
