Edigrate cl 15145909: Fix utility API to properly collect subobject references from explicitly instanced struct properties with a static array size.

diff --git a/Engine/Source/Runtime/Engine/Private/InstancedReferenceSubobjectHelper.cpp b/Engine/Source/Runtime/Engine/Private/InstancedReferenceSubobjectHelper.cpp
--- a/Engine/Source/Runtime/Engine/Private/InstancedReferenceSubobjectHelper.cpp
+++ b/Engine/Source/Runtime/Engine/Private/InstancedReferenceSubobjectHelper.cpp
@@ -13,8 +13,13 @@ UObject* FInstancedPropertyPath::Resolve(const UObject* Container) const
 		return FindFProperty<FProperty>(CurrentContainerType, SrcProperty->GetFName());
 	};
 
+	auto GetArrayIndex = [&PropChainRef](int32 ChainIndex)->int32
+	{
+		return PropChainRef[ChainIndex].ArrayIndex == INDEX_NONE ? 0 : PropChainRef[ChainIndex].ArrayIndex;
+	};
+
 	const FProperty* CurrentProp = GetProperty(0);
-	const uint8* ValuePtr = (CurrentProp) ? CurrentProp->ContainerPtrToValuePtr<uint8>(Container) : nullptr;
+	const uint8* ValuePtr = (CurrentProp) ? CurrentProp->ContainerPtrToValuePtr<uint8>(Container, GetArrayIndex(0)) : nullptr;
 
 	for (int32 ChainIndex = 1; CurrentProp && ChainIndex < PropertyChain.Num(); ++ChainIndex)
 	{
@@ -68,7 +73,7 @@ UObject* FInstancedPropertyPath::Resolve(const UObject* Container) const
 			}
 			CurrentProp = PropertyChain[ChainIndex].PropertyPtr;
 		}
-		else if (ensure(PropertyChain[ChainIndex].ArrayIndex <= 0))
+		else
 		{
 			if (const FStructProperty* StructProperty = CastField<FStructProperty>(CurrentProp))
 			{
@@ -76,14 +81,14 @@ UObject* FInstancedPropertyPath::Resolve(const UObject* Container) const
 			}
 
 			CurrentProp = GetProperty(ChainIndex);
-			ValuePtr    = (CurrentProp) ? CurrentProp->ContainerPtrToValuePtr<uint8>(ValuePtr) : nullptr;
+			ValuePtr = (CurrentProp) ? CurrentProp->ContainerPtrToValuePtr<uint8>(ValuePtr, GetArrayIndex(ChainIndex)) : nullptr;
 		}
 	}
 
-	const FObjectProperty* TargetPropety = CastField<FObjectProperty>(CurrentProp);
-	if (TargetPropety && TargetPropety->HasAnyPropertyFlags(CPF_InstancedReference))
-	{ 
-		return TargetPropety->GetObjectPropertyValue(ValuePtr);
+	const FObjectProperty* TargetProperty = CastField<FObjectProperty>(CurrentProp);
+	if (TargetProperty && TargetProperty->HasAnyPropertyFlags(CPF_InstancedReference))
+	{
+		return TargetProperty->GetObjectPropertyValue(ValuePtr);
 	}
 	return nullptr;
 }
@@ -93,194 +98,82 @@ void FFindInstancedReferenceSubobjectHelper::GetInstancedSubObjects_Inner(FInsta
 	check(ContainerAddress);
 	const FProperty* TargetProp = PropertyPath.Head();
 
-	if (!TargetProp->HasAnyPropertyFlags(CPF_PersistentInstance | CPF_ContainsInstancedReference))
-	{
-		return;
-	}
-
 	if (const FArrayProperty* ArrayProperty = CastField<const FArrayProperty>(TargetProp))
 	{
-		if (const FStructProperty* InnerStructProperty = CastField<const FStructProperty>(ArrayProperty->Inner))
+		// Exit now if the array doesn't contain any instanced references.
+		if (!ArrayProperty->HasAnyPropertyFlags(CPF_ContainsInstancedReference))
 		{
-			if (const UStruct* Struct = InnerStructProperty->Struct)
-			{
-				FScriptArrayHelper_InContainer ArrayHelper(ArrayProperty, ContainerAddress);
-				for (int32 ElementIndex = 0; ElementIndex < ArrayHelper.Num(); ++ElementIndex)
-				{
-					const uint8* ValueAddress = ArrayHelper.GetRawPtr(ElementIndex);
-					for (FProperty* StructProp = Struct->RefLink; StructProp; StructProp = StructProp->NextRef)
-					{
-						PropertyPath.Push(InnerStructProperty, ElementIndex);
-						GetInstancedSubObjects_Inner(PropertyPath, ValueAddress, OutObjects);
-						PropertyPath.Pop();
-					}
-				}
-			}
+			return;
 		}
-		else if (const FObjectProperty* InnerObjectProperty = CastField<const FObjectProperty>(ArrayProperty->Inner))
+
+		FScriptArrayHelper ArrayHelper(ArrayProperty, ContainerAddress);
+		for (int32 ElementIndex = 0; ElementIndex < ArrayHelper.Num(); ++ElementIndex)
 		{
-			if (InnerObjectProperty->HasAllPropertyFlags(CPF_PersistentInstance))
-			{
-				ensure(InnerObjectProperty->HasAllPropertyFlags(CPF_InstancedReference));
-				FScriptArrayHelper_InContainer ArrayHelper(ArrayProperty, ContainerAddress);
-				for (int32 ElementIndex = 0; ElementIndex < ArrayHelper.Num(); ++ElementIndex)
-				{
-					if (UObject* ObjectValue = InnerObjectProperty->GetObjectPropertyValue(ArrayHelper.GetRawPtr(ElementIndex)))
-					{
-						PropertyPath.Push(InnerObjectProperty, ElementIndex);
-						OutObjects(FInstancedSubObjRef(ObjectValue, PropertyPath));
-						PropertyPath.Pop();
-					}
-				}
-			}
+			const uint8* ValueAddress = ArrayHelper.GetRawPtr(ElementIndex);
+
+			PropertyPath.Push(ArrayProperty->Inner, ElementIndex);
+			GetInstancedSubObjects_Inner(PropertyPath, ValueAddress, OutObjects);
+			PropertyPath.Pop();
 		}
 	}
 	else if (const FMapProperty* MapProperty = CastField<const FMapProperty>(TargetProp))
 	{
-		if (const FStructProperty* KeyStructProperty = CastField<const FStructProperty>(MapProperty->KeyProp))
+		// Exit now if the map doesn't contain any instanced references.
+		if (!MapProperty->HasAnyPropertyFlags(CPF_ContainsInstancedReference))
 		{
-			if (const UStruct* Struct = KeyStructProperty->Struct)
-			{
-				FScriptMapHelper MapHelper(MapProperty, MapProperty->ContainerPtrToValuePtr<void>(ContainerAddress));
-				int32 Num = MapHelper.Num();
-				for (int32 ElementIndex = 0; Num; ++ElementIndex)
-				{
-					if (MapHelper.IsValidIndex(ElementIndex))
-					{
-						const uint8* KeyAddress = MapHelper.GetKeyPtr(ElementIndex);
-						for (FProperty* StructProp = Struct->RefLink; StructProp; StructProp = StructProp->NextRef)
-						{
-							PropertyPath.Push(KeyStructProperty, ElementIndex);
-							GetInstancedSubObjects_Inner(PropertyPath, KeyAddress, OutObjects);
-							PropertyPath.Pop();
-						}
-
-						--Num;
-					}
-				}
-			}
+			return;
 		}
-		else if (const FObjectProperty* KeyObjectProperty = CastField<const FObjectProperty>(MapProperty->KeyProp))
+
+		FScriptMapHelper MapHelper(MapProperty, ContainerAddress);
+		for (int32 ElementIndex = 0; ElementIndex < MapHelper.Num(); ++ElementIndex)
 		{
-			if (KeyObjectProperty->HasAllPropertyFlags(CPF_PersistentInstance))
-			{
-				ensure(KeyObjectProperty->HasAllPropertyFlags(CPF_InstancedReference));
-				FScriptMapHelper MapHelper(MapProperty, MapProperty->ContainerPtrToValuePtr<void>(ContainerAddress));
-				int32 Num = MapHelper.Num();
-				for (int32 ElementIndex = 0; Num; ++ElementIndex)
-				{
-					if (MapHelper.IsValidIndex(ElementIndex))
-					{
-						const uint8* KeyAddress = MapHelper.GetKeyPtr(ElementIndex);
-						if (UObject* ObjectValue = KeyObjectProperty->GetObjectPropertyValue(KeyAddress))
-						{
-							PropertyPath.Push(KeyObjectProperty, ElementIndex);
-							OutObjects(FInstancedSubObjRef(ObjectValue, PropertyPath));
-							PropertyPath.Pop();
-						}
+			const uint8* KeyAddress = MapHelper.GetKeyPtr(ElementIndex);
+			const uint8* ValueAddress = MapHelper.GetValuePtr(ElementIndex);
 
-						--Num;
-					}
-				}
-			}
-		}
+			PropertyPath.Push(MapProperty->KeyProp, ElementIndex);
+			GetInstancedSubObjects_Inner(PropertyPath, KeyAddress, OutObjects);
+			PropertyPath.Pop();
 
-		if (const FStructProperty* ValueStructProperty = CastField<const FStructProperty>(MapProperty->ValueProp))
+			PropertyPath.Push(MapProperty->ValueProp, ElementIndex);
+			GetInstancedSubObjects_Inner(PropertyPath, ValueAddress, OutObjects);
+			PropertyPath.Pop();
+		}
+	}
+	else if (const FSetProperty* SetProperty = CastField<const FSetProperty>(TargetProp))
+	{
+		// Exit now if the set doesn't contain any instanced references.
+		if (!SetProperty->HasAnyPropertyFlags(CPF_ContainsInstancedReference))
 		{
-			if (const UStruct* Struct = ValueStructProperty->Struct)
-			{
-				FScriptMapHelper MapHelper(MapProperty, MapProperty->ContainerPtrToValuePtr<void>(ContainerAddress));
-				int32 Num = MapHelper.Num();
-				for (int32 ElementIndex = 0; Num; ++ElementIndex)
-				{
-					if (MapHelper.IsValidIndex(ElementIndex))
-					{
-						// use of pair pointer is intentional, next call is going to offset from the pair entry:
-						const uint8* ValueAddress = MapHelper.GetPairPtr(ElementIndex);
-						for (FProperty* StructProp = Struct->RefLink; StructProp; StructProp = StructProp->NextRef)
-						{
-							PropertyPath.Push(ValueStructProperty, ElementIndex);
-							GetInstancedSubObjects_Inner(PropertyPath, ValueAddress, OutObjects);
-							PropertyPath.Pop();
-						}
-
-						--Num;
-					}
-				}
-			}
+			return;
 		}
-		else if (const FObjectProperty* ValueObjectProperty = CastField<const FObjectProperty>(MapProperty->ValueProp))
+
+		FScriptSetHelper SetHelper(SetProperty, ContainerAddress);
+		for (int32 ElementIndex = 0; ElementIndex < SetHelper.Num(); ++ElementIndex)
 		{
-			if (ValueObjectProperty->HasAllPropertyFlags(CPF_PersistentInstance))
-			{
-				ensure(ValueObjectProperty->HasAllPropertyFlags(CPF_InstancedReference));
-				FScriptMapHelper MapHelper(MapProperty, MapProperty->ContainerPtrToValuePtr<void>(ContainerAddress));
-				int32 Num = MapHelper.Num();
-				for (int32 ElementIndex = 0; Num; ++ElementIndex)
-				{
-					if (MapHelper.IsValidIndex(ElementIndex))
-					{
-						const uint8* ValueAddress = MapHelper.GetValuePtr(ElementIndex);
-						if (UObject* ObjectValue = ValueObjectProperty->GetObjectPropertyValue(ValueAddress))
-						{
-							PropertyPath.Push(ValueObjectProperty, ElementIndex);
-							OutObjects(FInstancedSubObjRef(ObjectValue, PropertyPath));
-							PropertyPath.Pop();
-						}
+			const uint8* ValueAddress = SetHelper.GetElementPtr(ElementIndex);
 
-						--Num;
-					}
-				}
-			}
+			PropertyPath.Push(SetProperty->ElementProp, ElementIndex);
+			GetInstancedSubObjects_Inner(PropertyPath, ValueAddress, OutObjects);
+			PropertyPath.Pop();
 		}
 	}
-	else if (const FSetProperty* SetProperty = CastField<const FSetProperty>(TargetProp))
+	else if (const FStructProperty* StructProperty = CastField<const FStructProperty>(TargetProp))
 	{
-		if (const FStructProperty* ElementStructProperty = CastField<const FStructProperty>(SetProperty->ElementProp))
+		// Exit early if the struct does not contain any instanced references or if the struct is invalid.
+		if (!StructProperty->HasAnyPropertyFlags(CPF_ContainsInstancedReference) || !StructProperty->Struct)
 		{
-			if (const UStruct* Struct = ElementStructProperty->Struct)
-			{
-				FScriptSetHelper SetHelper(SetProperty, SetProperty->ContainerPtrToValuePtr<void>(ContainerAddress));
-				int32 Num = SetHelper.Num();
-				for (int32 ElementIndex = 0; Num; ++ElementIndex)
-				{
-					if (SetHelper.IsValidIndex(ElementIndex))
-					{
-						const uint8* ElementAddress = SetHelper.GetElementPtr(ElementIndex);
-						for (FProperty* StructProp = Struct->RefLink; StructProp; StructProp = StructProp->NextRef)
-						{
-							PropertyPath.Push(ElementStructProperty, ElementIndex);
-							GetInstancedSubObjects_Inner(PropertyPath, ElementAddress, OutObjects);
-							PropertyPath.Pop();
-						}
-
-						--Num;
-					}
-				}
-			}
+			return;
 		}
-		else if (const FObjectProperty* ElementObjectProperty = CastField<const FObjectProperty>(SetProperty->ElementProp))
+
+		for (FProperty* StructProp = StructProperty->Struct->RefLink; StructProp; StructProp = StructProp->NextRef)
 		{
-			if (ElementObjectProperty->HasAllPropertyFlags(CPF_PersistentInstance))
+			for (int32 ArrayIdx = 0; ArrayIdx < StructProp->ArrayDim; ++ArrayIdx)
 			{
-				ensure(ElementObjectProperty->HasAllPropertyFlags(CPF_InstancedReference));
-				FScriptSetHelper SetHelper(SetProperty, SetProperty->ContainerPtrToValuePtr<void>(ContainerAddress));
-				int32 Num = SetHelper.Num();
-				for (int32 ElementIndex = 0; Num; ++ElementIndex)
-				{
-					if (SetHelper.IsValidIndex(ElementIndex))
-					{
-						const uint8* ElementAddress = SetHelper.GetElementPtr(ElementIndex);
-						if (UObject* ObjectValue = ElementObjectProperty->GetObjectPropertyValue(ElementAddress))
-						{
-							PropertyPath.Push(ElementObjectProperty, ElementIndex);
-							OutObjects(FInstancedSubObjRef(ObjectValue, PropertyPath));
-							PropertyPath.Pop();
-						}
+				const uint8* ValueAddress = StructProp->ContainerPtrToValuePtr<uint8>(ContainerAddress, ArrayIdx);
 
-						--Num;
-					}
-				}
+				PropertyPath.Push(StructProp, ArrayIdx);
+				GetInstancedSubObjects_Inner(PropertyPath, ValueAddress, OutObjects);
+				PropertyPath.Pop();
 			}
 		}
 	}
@@ -289,32 +182,10 @@ void FFindInstancedReferenceSubobjectHelper::GetInstancedSubObjects_Inner(FInsta
 		ensure(TargetProp->HasAllPropertyFlags(CPF_InstancedReference));
 		if (const FObjectProperty* ObjectProperty = CastField<const FObjectProperty>(TargetProp))
 		{
-			for (int32 ArrayIdx = 0; ArrayIdx < ObjectProperty->ArrayDim; ++ArrayIdx)
-			{
-				UObject* ObjectValue = ObjectProperty->GetObjectPropertyValue_InContainer(ContainerAddress, ArrayIdx);
-				if (ObjectValue)
-				{
-					// don't need to push to PropertyPath, since this property is already at its head
-					OutObjects(FInstancedSubObjRef(ObjectValue, PropertyPath));	
-				}
-			}
-		}
-
-		return;
-	}
-	else if (const FStructProperty* StructProperty = CastField<const FStructProperty>(TargetProp))
-	{
-		if (StructProperty->Struct)
-		{
-			for (int32 ArrayIdx = 0; ArrayIdx < StructProperty->ArrayDim; ++ArrayIdx)
+			if (UObject* ObjectValue = ObjectProperty->GetObjectPropertyValue(ContainerAddress))
 			{
-				const uint8* ValueAddress = StructProperty->ContainerPtrToValuePtr<uint8>(ContainerAddress, ArrayIdx);
-				for (FProperty* StructProp = StructProperty->Struct->RefLink; StructProp; StructProp = StructProp->NextRef)
-				{
-					PropertyPath.Push(StructProp, ArrayIdx);
-					GetInstancedSubObjects_Inner(PropertyPath, ValueAddress, OutObjects);
-					PropertyPath.Pop();
-				}
+				// don't need to push to PropertyPath, since this property is already at its head
+				OutObjects(FInstancedSubObjRef(ObjectValue, PropertyPath));
 			}
 		}
 	}
diff --git a/Engine/Source/Runtime/Engine/Public/InstancedReferenceSubobjectHelper.h b/Engine/Source/Runtime/Engine/Public/InstancedReferenceSubobjectHelper.h
--- a/Engine/Source/Runtime/Engine/Public/InstancedReferenceSubobjectHelper.h
+++ b/Engine/Source/Runtime/Engine/Public/InstancedReferenceSubobjectHelper.h
@@ -27,9 +27,9 @@ private:
 
 public:
 	//--------------------------------------------------------------------------
-	FInstancedPropertyPath(FProperty* RootProperty)
+	FInstancedPropertyPath(FProperty* RootProperty, int32 ArrayIndex = INDEX_NONE)
 	{
-		PropertyChain.Add(FPropertyLink(RootProperty));
+		Push(RootProperty, ArrayIndex);
 	}
 
 	//--------------------------------------------------------------------------
@@ -104,8 +104,12 @@ public:
 		const UClass* ContainerClass = Container->GetClass();
 		for (FProperty* Prop = ContainerClass->RefLink; Prop; Prop = Prop->NextRef)
 		{
-			FInstancedPropertyPath RootPropertyPath(Prop);
-			GetInstancedSubObjects_Inner(RootPropertyPath, reinterpret_cast<const uint8*>(Container), [&OutObjects](const FInstancedSubObjRef& Ref){ OutObjects.Add(Ref); });
+			for (int32 ArrayIdx = 0; ArrayIdx < Prop->ArrayDim; ++ArrayIdx)
+			{
+				FInstancedPropertyPath RootPropertyPath(Prop, ArrayIdx);
+				const uint8* ValuePtr = Prop->ContainerPtrToValuePtr<uint8>(Container, ArrayIdx);
+				GetInstancedSubObjects_Inner(RootPropertyPath, ValuePtr, [&OutObjects](const FInstancedSubObjRef& Ref) { OutObjects.Add(Ref); });
+			}
 		}
 	}
 
