Use hash of contents for VT chunk DDC key
- Chunks are always managed by the top level VT build data in the DDC, so the key for each indivual chunk never needs to be derived from any properties
- Using hash of contents makes the system robust if building chunks is non-deterministic
- Downside is potential for bloating the DDC with lots of similar (not not exactly the same) keys for the same logical chunk

diff --git a/Engine/Source/Runtime/Engine/Private/TextureDerivedData.cpp b/Engine/Source/Runtime/Engine/Private/TextureDerivedData.cpp
--- a/Engine/Source/Runtime/Engine/Private/TextureDerivedData.cpp
+++ b/Engine/Source/Runtime/Engine/Private/TextureDerivedData.cpp
@@ -53,7 +53,7 @@
 
 // This GUID is mixed into DDC version for virtual textures only, this allows updating DDC version for VT without invalidating DDC for all textures
 // This is useful during development, but once large numbers of VT are present in shipped content, it will have the same problem as TEXTURE_DERIVEDDATA_VER
-#define TEXTURE_VT_DERIVEDDATA_VER	TEXT("39E589A391EF40459CA8F45359851DF3")
+#define TEXTURE_VT_DERIVEDDATA_VER	TEXT("E541C13E65B24669B01DDE3DF18C7312")
 
 #if ENABLE_COOK_STATS
 namespace TextureCookStats
@@ -649,11 +649,12 @@ uint32 PutDerivedDataInCache(FTexturePlatformData* DerivedData, const FString& D
 		const int32 ChunkCount = DerivedData->VTData->Chunks.Num();
 		for (int32 ChunkIndex = 0; ChunkIndex < ChunkCount; ++ChunkIndex)
 		{
+			FVirtualTextureDataChunk& Chunk = DerivedData->VTData->Chunks[ChunkIndex];
+
 			const FString ChunkDerivedDataKey = FDerivedDataCacheInterface::BuildCacheKey(
-				TEXT("TEXTURE"), TEXTURE_DERIVEDDATA_VER,
-				*FString::Printf(TEXT("%s_VTCHUNK%u"), *DerivedDataKeySuffix, ChunkIndex));
+				TEXT("TEXTURE"), TEXTURE_VT_DERIVEDDATA_VER,
+				*FString::Printf(TEXT("VTCHUNK%s"), *Chunk.BulkDataHash.ToString()));
 
-			FVirtualTextureDataChunk& Chunk = DerivedData->VTData->Chunks[ChunkIndex];
 			TotalBytesPut += Chunk.StoreInDerivedDataCache(ChunkDerivedDataKey, TextureName, bReplaceExistingDDC);
 		}
 	}
diff --git a/Engine/Source/Runtime/Engine/Private/VT/VirtualTextureBuiltData.cpp b/Engine/Source/Runtime/Engine/Private/VT/VirtualTextureBuiltData.cpp
--- a/Engine/Source/Runtime/Engine/Private/VT/VirtualTextureBuiltData.cpp
+++ b/Engine/Source/Runtime/Engine/Private/VT/VirtualTextureBuiltData.cpp
@@ -52,7 +52,6 @@ void FVirtualTextureBuiltData::Serialize(FArchive& Ar, UObject* Owner, int32 Fir
 	bool bCooked = Ar.IsCooking();
 	Ar << bCooked;
 
-	Ar << VersionGuid;
 	Ar << NumLayers;
 	Ar << WidthInBlocks;
 	Ar << HeightInBlocks;
@@ -232,14 +231,6 @@ bool FVirtualTextureBuiltData::ValidateData(FStringView const& InDDCDebugContext
 			break;
 		}
 
-		FVirtualTextureChunkHeader* ChunkHeader = (FVirtualTextureChunkHeader*)ChunkData;
-		if (ChunkHeader->VersionGuid != VersionGuid)
-		{
-			UE_LOG(LogTexture, Error, TEXT("Virtual Texture %s has mismatched GUID for chunk %d"), *TextureName, ChunkIndex);
-			bResult = false;
-			break;
-		}
-
 		FSHAHash Hash;
 		FSHA1::HashBuffer(ChunkData, ChunkDataSize, Hash.Hash);
 		if (Hash != Chunk.BulkDataHash)
diff --git a/Engine/Source/Runtime/Engine/Private/VT/VirtualTextureBuiltData.h b/Engine/Source/Runtime/Engine/Private/VT/VirtualTextureBuiltData.h
--- a/Engine/Source/Runtime/Engine/Private/VT/VirtualTextureBuiltData.h
+++ b/Engine/Source/Runtime/Engine/Private/VT/VirtualTextureBuiltData.h
@@ -30,7 +30,7 @@ enum class EVirtualTextureCodec : uint8
 
 struct FVirtualTextureChunkHeader
 {
-	FGuid VersionGuid;
+	uint32 Unused; // Chunk header is currently not used, but keeping in place for future expansion
 };
 
 struct FVirtualTextureDataChunk
@@ -73,7 +73,6 @@ struct FVirtualTextureDataChunk
 
 struct FVirtualTextureBuiltData
 {
-	FGuid VersionGuid;
 	uint32 NumLayers;
 	uint32 NumMips;
 	uint32 Width; // Width of the texture in pixels. Note the physical width may be larger due to tiling
diff --git a/Engine/Source/Runtime/Engine/Private/VT/VirtualTextureDataBuilder.cpp b/Engine/Source/Runtime/Engine/Private/VT/VirtualTextureDataBuilder.cpp
--- a/Engine/Source/Runtime/Engine/Private/VT/VirtualTextureDataBuilder.cpp
+++ b/Engine/Source/Runtime/Engine/Private/VT/VirtualTextureDataBuilder.cpp
@@ -334,7 +334,6 @@ void FVirtualTextureDataBuilder::Build(const FTextureSourceData& InSourceData, c
 	//NOTE: OutData may point to a previously build data so it is important to
 	//properly initialize all fields and not assume this is a freshly constructed object
 
-	OutData.VersionGuid = FGuid::NewGuid();
 	OutData.TileBorderSize = BuildSettingsLayer0.VirtualTextureBorderSize;
 	OutData.TileSize = TileSize;
 	OutData.NumLayers = NumLayers;
@@ -736,7 +735,7 @@ void FVirtualTextureDataBuilder::PushDataToChunk(const TArray<FVTSourceTileEntry
 
 	// Header for the chunk
 	FVirtualTextureChunkHeader* Header = (FVirtualTextureChunkHeader*)NewChunkData;
-	Header->VersionGuid = OutData.VersionGuid;
+	FMemory::Memzero(*Header);
 
 	ChunkOffset += sizeof(FVirtualTextureChunkHeader);
 
