Made GPU-Scene instance culling toggle with GPU-Scene support (UseGPUScene)
- Added dynamic switches in addition to the ifdef.
- Added ifdef to disable workaround for legacy primitives for cube shadow when not used (enabled with GPU-Scene)
- Fixed some references to PrimitiveId in hair strand shaders

diff --git a/Engine/Shaders/Private/HairStrands/HairStrandsMaterialCommon.ush b/Engine/Shaders/Private/HairStrands/HairStrandsMaterialCommon.ush
--- a/Engine/Shaders/Private/HairStrands/HairStrandsMaterialCommon.ush
+++ b/Engine/Shaders/Private/HairStrands/HairStrandsMaterialCommon.ush
@@ -48,19 +48,22 @@ struct FHairVertexData
 	bool   bIsValid;
 };
 
+// GPUCULL_TODO: this should be switched over to use Instance ID instead of Primitive ID.
 FVertexFactoryInterpolantsVSToPS GetInterpolants(uint PrimitiveId, uint HairPrimitiveId, inout FHairVertexData OutHairVertex)
 {
 	FMaterialVertexParameters Empty = (FMaterialVertexParameters)0;
 
 	FVertexFactoryInput Input;
+#if VF_USE_PRIMITIVE_SCENE_DATA
 #if defined(GPUCULL_TODO)
-	// GPUCULL_TODO: this should be switched over to use instance ID.
-	//Input.PrimitiveId = PrimitiveId;
-	Input.InstanceIdOffset = 0;
+	// GPUCULL_TODO: Making use of dynamic draw workaround, assumes no actual instancing used for these draws.
+	//               Will load the first instance data.
+	Input.InstanceIdOffset = PrimitiveId | (1U << 31U);
 	Input.DrawInstanceId = 0;
-#else //!defined(GPUCULL_TODO)
+#else //!(defined(GPUCULL_TODO)
 	Input.PrimitiveId = PrimitiveId;
 #endif // defined(GPUCULL_TODO)
+#endif // VF_USE_PRIMITIVE_SCENE_DATA 
 	Input.VertexId = HairPrimitiveId * 6; // 6: the vertex factory expect triangles list index. A quad is made of two triangles with indices: 0,1,2,3,4,5.
 
 	FVertexFactoryIntermediates	Intermediates = GetVertexFactoryIntermediates(Input);
@@ -134,8 +137,9 @@ FEvaluateOutput Evaluate(
 
 		Interpolants.TangentToWorld0 = normalize(Interpolants.TangentToWorld0);
 		Interpolants.TangentToWorld2 = normalize(Interpolants.TangentToWorld2);
-
+#if VF_USE_PRIMITIVE_SCENE_DATA
 		Interpolants.PrimitiveId	 = PrimitiveId;
+#endif // VF_USE_PRIMITIVE_SCENE_DATA
 		Interpolants.HairPrimitiveId = HairPrimitiveId0;
 		Interpolants.HairPrimitiveUV = float2(SegmentU, SegmentV);
 	}
diff --git a/Engine/Shaders/Private/HairStrands/HairStrandsVertexFactory.ush b/Engine/Shaders/Private/HairStrands/HairStrandsVertexFactory.ush
--- a/Engine/Shaders/Private/HairStrands/HairStrandsVertexFactory.ush
+++ b/Engine/Shaders/Private/HairStrands/HairStrandsVertexFactory.ush
@@ -182,18 +182,6 @@ struct FVertexFactoryIntermediates
 	float  HairDensity;
 };
 
-#if !defined(GPUCULL_TODO)
-
-uint GetPrimitiveId(FVertexFactoryInput Input)
-{
-#if VF_USE_PRIMITIVE_SCENE_DATA
-	return Input.PrimitiveId;
-#else
-	return 0;
-#endif
-}
-#endif // defined(GPUCULL_TODO)
-
 uint GetHairPrimitiveId(FVertexFactoryInput Input)
 {
 	return GetVertexInfo(Input).HairPrimitiveId;
@@ -350,8 +338,6 @@ half3x3 CalcTangentToWorld(FVertexFactoryIntermediates Intermediates, half3x3 Ta
 	return LocalTangentToWorld;
 }
 
-#if defined(GPUCULL_TODO)
-
 float3 GetTangent(FVertexFactoryIntermediates Intermediates)
 {
 	// Hair shader expec the tangent to be stored in place of the normal. This is what the StrandHairComponent is filling in:
@@ -361,40 +347,20 @@ float3 GetTangent(FVertexFactoryIntermediates Intermediates)
 	return Intermediates.TangentToWorld[2].xyz;
 }
 
-#else //!defined(GPUCULL_TODO)
-
-float3 GetTangent(FVertexFactoryInput Input)
-{
-	FVertexFactoryIntermediates Intermediates;
-	Intermediates.PrimitiveId = GetPrimitiveId(Input);
-	
-	float TangentSign;
-	Intermediates.TangentToLocal = CalcTangentToLocal(Input, TangentSign);
-	Intermediates.TangentToWorld = CalcTangentToWorld(Intermediates,Intermediates.TangentToLocal);
-	Intermediates.TangentToWorldSign = TangentSign * GetPrimitiveData(Intermediates.PrimitiveId).InvNonUniformScaleAndDeterminantSign.w;
-
-	// Hair shader expec the tangent to be stored in place of the normal. This is what the StrandHairComponent is filling in:
-	// [0]: Real normal		| Should be Tangent
-	// [1]: Bitangent		| Should be BiTangent
-	// [2]: Real Tangent	| Should be Normal
-	return Intermediates.TangentToWorld[2].xyz;
-}
-
-#endif // defined(GPUCULL_TODO)
-
 FVertexFactoryIntermediates GetVertexFactoryIntermediates(FVertexFactoryInput Input)
 {
 	FVertexFactoryIntermediates Intermediates;
-#if defined(GPUCULL_TODO)
 #if VF_USE_PRIMITIVE_SCENE_DATA
+#if defined(GPUCULL_TODO)
 	Intermediates.SceneData = GetSceneDataIntermediates(Input.InstanceIdOffset, Input.DrawInstanceId);
 	Intermediates.PrimitiveId = Intermediates.SceneData.InstanceData.PrimitiveId;
-#else
-	Intermediates.PrimitiveId = 0U;
-#endif
 #else //!defined(GPUCULL_TODO)
-	Intermediates.PrimitiveId		= GetPrimitiveId(Input);
+	Intermediates.PrimitiveId = Input.PrimitiveId;
 #endif // defined(GPUCULL_TODO)
+#else // !VF_USE_PRIMITIVE_SCENE_DATA
+	Intermediates.PrimitiveId = 0U;
+#endif // VF_USE_PRIMITIVE_SCENE_DATA
+
 	Intermediates.HairPrimitiveId	= GetHairPrimitiveId(Input);
 	Intermediates.HairPrimitiveUV	= GetSegmentUVCoord(Input);
 	Intermediates.HairDensity		= HairStrandsVF_Density;
@@ -441,14 +407,10 @@ float4 VertexFactoryGetWorldPosition_Visibility(FVertexFactoryInput Input, FHair
 	OutHairPrimitiveId = VertexInfo.HairPrimitiveId;
 
 	FHairControlPoint ControlPoint = GetVertexPosition(VertexInfo);
-#if defined(GPUCULL_TODO)
+
 	FVertexFactoryIntermediates Intermediates = GetVertexFactoryIntermediates(Input);
 	const float4 WorldPosition = CalcWorldPosition(ControlPoint.Position, Intermediates.PrimitiveId);
 	const float3 WorldTangent = GetTangent(Intermediates);
-#else //!defined(GPUCULL_TODO)
-	const float4 WorldPosition = CalcWorldPosition(ControlPoint.Position, GetPrimitiveId(Input));
-	const float3 WorldTangent = GetTangent(Input);
-#endif // defined(GPUCULL_TODO)
 
 	// Minimal radius to snap the strand to a sample/pixel center (to avoid aliasing)
 	const float DistanceToCamera = length(HairViewInfo.TranslatedWorldCameraOrigin - WorldPosition.xyz);
@@ -470,14 +432,9 @@ float4 VertexFactoryGetWorldPosition_Voxelization(FVertexFactoryInput Input, FHa
 	OutHairWorldRadius = ControlPoint.WorldRadius;
 	OutHairDensity = HairStrandsVF_Density;
 
-#if defined(GPUCULL_TODO)
 	FVertexFactoryIntermediates Intermediates = GetVertexFactoryIntermediates(Input);
 	const float4 WorldPosition = CalcWorldPosition(ControlPoint.Position, Intermediates.PrimitiveId);
 	const float3 WorldTangent = GetTangent(Intermediates);
-#else //!defined(GPUCULL_TODO)
-	const float4 WorldPosition = CalcWorldPosition(ControlPoint.Position, GetPrimitiveId(Input));
-	const float3 WorldTangent = GetTangent(Input);
-#endif // defined(GPUCULL_TODO)
 
 	// Minimal radius to snap the strand to a sample/pixel center (to avoid aliasing)
 	const float3 ViewDir = -HairViewInfo.ViewForward;
@@ -495,11 +452,7 @@ float4 VertexFactoryGetWorldPosition(FVertexFactoryInput Input, FVertexFactoryIn
 	const float4 WorldPosition = CalcWorldPosition(VertexPosition, Intermediates.PrimitiveId);
 
 	// Hair shader expect the WorldNormal to be the tangent vector
-#if defined(GPUCULL_TODO)
 	const float3 WorldTangent = GetTangent(Intermediates);
-#else //!defined(GPUCULL_TODO)
-	const float3 WorldTangent = GetTangent(Input);
-#endif // defined(GPUCULL_TODO)
 	return ComputeViewAlignedWorldPosition(Input, WorldTangent, WorldPosition, ControlPoint.WorldRadius, HairViewInfo);
 }
 
diff --git a/Engine/Shaders/Private/LocalVertexFactory.ush b/Engine/Shaders/Private/LocalVertexFactory.ush
--- a/Engine/Shaders/Private/LocalVertexFactory.ush
+++ b/Engine/Shaders/Private/LocalVertexFactory.ush
@@ -951,11 +951,9 @@ FVertexFactoryIntermediates GetVertexFactoryIntermediates(FVertexFactoryInput In
 
 	// PerInstanceParams.z stores a hide/show flag for this instance
 	float SelectedValue = GetInstanceSelected(Intermediates);
-#if defined(GPUCULL_TODO)
-	Intermediates.PerInstanceParams.x = Intermediates.PerInstanceRandom;
-#else //!defined(GPUCULL_TODO)
+	// GPUCULL_TODO: This can't be right
+	//Intermediates.PerInstanceParams.x = Intermediates.PerInstanceRandom;
 	Intermediates.PerInstanceParams.x = GetInstanceRandom(Intermediates);
-#endif // defined(GPUCULL_TODO)
 	float3 InstanceLocation = TransformLocalToWorld(GetInstanceOrigin(Intermediates), Intermediates.PrimitiveId).xyz;
 	Intermediates.PerInstanceParams.y = 1.0 - saturate((length(InstanceLocation + ResolvedView.PreViewTranslation.xyz) - InstancingFadeOutParams.x) * InstancingFadeOutParams.y);
 	// InstancingFadeOutParams.z,w are RenderSelected and RenderDeselected respectively.
diff --git a/Engine/Shaders/Private/ShadowDepthCommon.ush b/Engine/Shaders/Private/ShadowDepthCommon.ush
--- a/Engine/Shaders/Private/ShadowDepthCommon.ush
+++ b/Engine/Shaders/Private/ShadowDepthCommon.ush
@@ -37,9 +37,9 @@ struct FShadowDepthVSToGS : FShadowDepthVSToPS
 {
 	float4 GSPosition : TEXCOORD6;
 	float3 WorldNormal : TEXCOORD8;
-#if defined(GPUCULL_TODO)
+#if defined(GPUCULL_TODO) && ENABLE_FALLBACK_POINTLIGHT_SHADOW_GS
 	nointerpolation uint ViewIndex : TEXCOORD9;
-#endif // defined(GPUCULL_TODO)
+#endif // ENABLE_FALLBACK_POINTLIGHT_SHADOW_GS
 };
 #endif
 
diff --git a/Engine/Shaders/Private/ShadowDepthVertexShader.usf b/Engine/Shaders/Private/ShadowDepthVertexShader.usf
--- a/Engine/Shaders/Private/ShadowDepthVertexShader.usf
+++ b/Engine/Shaders/Private/ShadowDepthVertexShader.usf
@@ -8,6 +8,11 @@
 #define SHADOW_DEPTH_SHADER 1
 #define USE_STENCIL_LOD_DITHER	0
 
+#ifndef ENABLE_NON_NANITE_VSM
+#error "ENABLE_NON_NANITE_VSM should be defined to either 0 or 1!"
+#endif
+
+
 #if ENABLE_NON_NANITE_VSM
 	#define VIRTUAL_SM_ENABLED (!(USING_TESSELLATION || ONEPASS_POINTLIGHT_SHADOW))
 #else //!ENABLE_NON_NANITE_VSM
@@ -250,7 +255,7 @@ bool IsTriangleInFrustum(float4 ClipSpacePositions[3])
 [maxvertexcount(18)]
 void MainOnePassPointLightGS(triangle FShadowDepthVSToGS Input[3], inout TriangleStream<FShadowDepthGSToPS> OutStream)
 {
-#if defined(GPUCULL_TODO)
+#if defined(GPUCULL_TODO) && ENABLE_FALLBACK_POINTLIGHT_SHADOW_GS
 	uint CubeFaceStartIndex = bUseGpuSceneInstancing ? Input[0].ViewIndex : 0;
 	uint CubeFaceEndIndex = bUseGpuSceneInstancing ? Input[0].ViewIndex + 1: 6;
 
@@ -460,9 +465,9 @@ void Main(
 
 	OutPosition = OutParameters.GSPosition = WorldPos;
 	OutParameters.WorldNormal = WorldNormal;
-#if defined(GPUCULL_TODO)
+#if defined(GPUCULL_TODO) && ENABLE_FALLBACK_POINTLIGHT_SHADOW_GS
 	OutParameters.ViewIndex = VertexFactoryGetViewIndex(VFIntermediates);
-#endif // defined(GPUCULL_TODO)
+#endif // !ENABLE_FALLBACK_POINTLIGHT_SHADOW_GS
 
 	#if INTERPOLATE_VF_ATTRIBUTES
 		// Masked materials need texture coords to clip
@@ -604,9 +609,9 @@ void PositionOnlyMain(
 #if ONEPASS_POINTLIGHT_SHADOW
 	OutPosition = OutParameters.GSPosition = WorldPos;
 	OutParameters.WorldNormal = WorldNormal;
-#if defined(GPUCULL_TODO)
+#if defined(GPUCULL_TODO) && ENABLE_FALLBACK_POINTLIGHT_SHADOW_GS
 	OutParameters.ViewIndex = VertexFactoryGetViewIndex(Input);
-#endif // defined(GPUCULL_TODO)
+#endif // !ENABLE_FALLBACK_POINTLIGHT_SHADOW_GS
 
 	#if USING_ONEPASS_POINT_VERTEX_LAYER
 #if defined(GPUCULL_TODO)
diff --git a/Engine/Source/Runtime/Engine/Private/HierarchicalInstancedStaticMesh.cpp b/Engine/Source/Runtime/Engine/Private/HierarchicalInstancedStaticMesh.cpp
--- a/Engine/Source/Runtime/Engine/Private/HierarchicalInstancedStaticMesh.cpp
+++ b/Engine/Source/Runtime/Engine/Private/HierarchicalInstancedStaticMesh.cpp
@@ -893,36 +893,39 @@ public:
 #endif
 	{
 #if defined(GPUCULL_TODO)
-		if (InstancedRenderData.PerInstanceRenderData.IsValid() && InstancedRenderData.PerInstanceRenderData->InstanceBuffer_GameThread.IsValid())
+		if (UseGPUScene(GetScene().GetShaderPlatform(), GetScene().GetFeatureLevel()))
 		{
-			const FStaticMeshInstanceData& StaticMeshInstanceData = *InstancedRenderData.PerInstanceRenderData->InstanceBuffer_GameThread;
-			
-			Instances.SetNum(StaticMeshInstanceData.GetNumInstances());
-			if (Instances.Num() == 0 && ClusterTree.Num() != 0)
+			if (InstancedRenderData.PerInstanceRenderData.IsValid() && InstancedRenderData.PerInstanceRenderData->InstanceBuffer_GameThread.IsValid())
 			{
-				UE_LOG(LogTemp, Display, TEXT("Instances.Num() == 0 && ClusterTree.Num() != 0"));
-			}
+				const FStaticMeshInstanceData& StaticMeshInstanceData = *InstancedRenderData.PerInstanceRenderData->InstanceBuffer_GameThread;
 
-			for (int32 InstanceIndex = 0; InstanceIndex < Instances.Num(); ++InstanceIndex)
+					Instances.SetNum(StaticMeshInstanceData.GetNumInstances());
+					if (Instances.Num() == 0 && ClusterTree.Num() != 0)
+					{
+						UE_LOG(LogTemp, Display, TEXT("Instances.Num() == 0 && ClusterTree.Num() != 0"));
+					}
+
+				for (int32 InstanceIndex = 0; InstanceIndex < Instances.Num(); ++InstanceIndex)
+				{
+					FMatrix InstanceTransform;
+					StaticMeshInstanceData.GetInstanceTransform(InstanceIndex, InstanceTransform);
+					InstanceTransform.M[3][3] = 1.0f;
+
+					FPrimitiveInstance& Instance = Instances[InstanceIndex];
+					Instance.PrimitiveId = ~uint32(0);
+					Instance.InstanceToLocal = InstanceTransform;
+					Instance.LocalToInstance = InstanceTransform.Inverse();
+					// Filled in during GPU-Scene update...
+					Instance.LocalToWorld.SetIdentity();
+					Instance.RenderBounds = InComponent->GetStaticMesh()->GetBounds();
+					Instance.LocalBounds = Instance.RenderBounds.TransformBy(Instance.InstanceToLocal);
+				}
+			}
+			else
 			{
-				FMatrix InstanceTransform;
-				StaticMeshInstanceData.GetInstanceTransform(InstanceIndex, InstanceTransform);
-				InstanceTransform.M[3][3] = 1.0f;
-
-				FPrimitiveInstance& Instance = Instances[InstanceIndex];
-				Instance.PrimitiveId = ~uint32(0);
-				Instance.InstanceToLocal = InstanceTransform;
-				Instance.LocalToInstance = InstanceTransform.Inverse();
-				// Filled in during GPU-Scene update...
-				Instance.LocalToWorld.SetIdentity();
-				Instance.RenderBounds = InComponent->GetStaticMesh()->GetBounds();
-				Instance.LocalBounds = Instance.RenderBounds.TransformBy(Instance.InstanceToLocal);
+				ensure(ClusterTree.Num() == 0);
 			}
 		}
-		else
-		{
-			ensure(ClusterTree.Num() == 0);
-		}
 #endif // defined(GPUCULL_TODO)
 
 		SetupOcclusion(InComponent);
@@ -1351,6 +1354,8 @@ void FHierarchicalStaticMeshSceneProxy::FillDynamicMeshElements(FMeshElementColl
 	int32 FirstLOD = FMath::Max((OnlyLOD < 0) ? 0 : OnlyLOD, static_cast<int32>(this->GetCurrentFirstLODIdx_Internal()));
 	int32 LastLODPlusOne = (OnlyLOD < 0) ? InstancedRenderData.VertexFactories.Num() : (OnlyLOD+1);
 
+	const bool bUseGPUScene = UseGPUScene(GetScene().GetShaderPlatform(), GetScene().GetFeatureLevel());
+
 	for (int32 LODIndex = FirstLOD; LODIndex < LastLODPlusOne; LODIndex++)
 	{
 		const FStaticMeshLODResources& LODModel = RenderData->LODResources[LODIndex];
@@ -1372,189 +1377,195 @@ void FHierarchicalStaticMeshSceneProxy::FillDynamicMeshElements(FMeshElementColl
 
 
 #if defined(GPUCULL_TODO)
-				FMeshBatch& MeshBatch = Collector.AllocateMesh();
-				INC_DWORD_STAT(STAT_FoliageMeshBatches);
-
-				if (!FStaticMeshSceneProxy::GetMeshElement(LODIndex, 0, SectionIndex, GetDepthPriorityGroup(ElementParams.View), ElementParams.BatchRenderSelection[SelectionGroupIndex], true, MeshBatch))
-				{
-					continue;
-				}
-
-				checkSlow(MeshBatch.GetNumPrimitives() > 0);
-				MeshBatch.bCanApplyViewModeOverrides = true;
-				MeshBatch.bUseSelectionOutline = ElementParams.BatchRenderSelection[SelectionGroupIndex];
-				MeshBatch.bUseWireframeSelectionColoring = ElementParams.BatchRenderSelection[SelectionGroupIndex];
-				MeshBatch.bUseAsOccluder = ShouldUseAsOccluder();
-				MeshBatch.VertexFactory = &InstancedRenderData.VertexFactories[LODIndex];
-
-				FMeshBatchElement& MeshBatchElement = MeshBatch.Elements[0];
-				MeshBatchElement.UserData = ElementParams.PassUserData[SelectionGroupIndex];
-				MeshBatchElement.bUserDataIsColorVertexBuffer = false;
-				MeshBatchElement.MaxScreenSize = 1.0;
-				MeshBatchElement.MinScreenSize = 0.0;
-				MeshBatchElement.InstancedLODIndex = LODIndex;
-				MeshBatchElement.InstancedLODRange = bDitherLODEnabled ? 1 : 0;
-				MeshBatchElement.PrimitiveUniformBuffer = GetUniformBuffer();
-
-				int32 TotalInstances = bDitherLODEnabled ? Params.TotalMultipleLODInstances[LODIndex] : Params.TotalSingleLODInstances[LODIndex];
+				if (bUseGPUScene)
 				{
-					const int64 Tris = int64(TotalInstances) * int64(MeshBatchElement.NumPrimitives);
-					TotalTriangles += Tris;
-#if STATS
-					if (GFrameNumberRenderThread_CaptureFoliageRuns == GFrameNumberRenderThread)
-					{
-						if (ElementParams.FinalCullDistance > 9.9E8)
-						{
-							UE_LOG(LogStaticMesh, Display, TEXT("lod:%1d/%1d   sel:%1d   section:%1d/%1d   runs:%4d   inst:%8d   tris:%9lld   cast shadow:%1d   cull:-NONE!!-   shadow:%1d     %s %s"),
-								LODIndex, InstancedRenderData.VertexFactories.Num(), SelectionGroupIndex, SectionIndex, LODModel.Sections.Num(), RunArray.Num() / 2,
-								TotalInstances, Tris, (int)MeshBatch.CastShadow, ElementParams.ShadowFrustum,
-								*StaticMesh->GetPathName(),
-								*MeshBatch.MaterialRenderProxy->GetMaterial(ElementParams.FeatureLevel)->GetFriendlyName());
-						}
-						else
-						{
-							UE_LOG(LogStaticMesh, Display, TEXT("lod:%1d/%1d   sel:%1d   section:%1d/%1d   runs:%4d   inst:%8d   tris:%9lld   cast shadow:%1d   cull:%8.0f   shadow:%1d     %s %s"),
-								LODIndex, InstancedRenderData.VertexFactories.Num(), SelectionGroupIndex, SectionIndex, LODModel.Sections.Num(), RunArray.Num() / 2,
-								TotalInstances, Tris, (int)MeshBatch.CastShadow, ElementParams.FinalCullDistance, ElementParams.ShadowFrustum,
-								*StaticMesh->GetPathName(),
-								*MeshBatch.MaterialRenderProxy->GetMaterial(ElementParams.FeatureLevel)->GetFriendlyName());
-						}
-					}
-#endif // STATS
-				}
-
-				//MeshBatchElement.NumInstances = TotalInstances;
-				// The index was used as an offset, but the dynamic buffer thing uses a resource view to make this not needed (using PrimitiveInstanceDataOffset as a temp. debug help)
-				MeshBatchElement.UserIndex = 0;
-
-				// Note: this call overrides the UserIndex to mean the command index, which is used to fetch the offset to the instance array
-				//Collector.AllocateInstancedBatchArguments(ElementParams.ViewIndex, MeshBatch, PrimitiveInstanceDataOffset, PrimitiveInstanceDataCount, RunArray);
-				
-				// We use this existing hook to send info about the runs over to the visiblemeshbatch
-				MeshBatchElement.NumInstances = RunArray.Num() / 2;
-				MeshBatchElement.InstanceRuns = &RunArray[0];
-				MeshBatchElement.bIsInstanceRuns = true;
-
-				if (TotalTriangles < (int64)CVarMaxTrianglesToRender.GetValueOnRenderThread())
-				{
-					Collector.AddMesh(ElementParams.ViewIndex, MeshBatch);
-				}
-
-#else // !defined(GPUCULL_TODO)
-				int32 NumBatches = 1;
-				int32 CurrentRun = 0;
-				int32 CurrentInstance = 0;
-				int32 RemainingInstances = bDitherLODEnabled ? Params.TotalMultipleLODInstances[LODIndex] : Params.TotalSingleLODInstances[LODIndex];
-				int32 RemainingRuns = RunArray.Num() / 2;
-
-				if (!ElementParams.bUseInstanceRuns)
-				{
-					NumBatches = FMath::DivideAndRoundUp(RemainingRuns, (int32)FInstancedStaticMeshVertexFactory::NumBitsForVisibilityMask());
-				}
-
-#if STATS
-				INC_DWORD_STAT_BY(STAT_FoliageInstances, RemainingInstances);
-#endif
-				bool bDidStats = false;
-				for (int32 BatchIndex = 0; BatchIndex < NumBatches; BatchIndex++)
-				{
-					FMeshBatch& MeshElement = Collector.AllocateMesh();
+					FMeshBatch& MeshBatch = Collector.AllocateMesh();
 					INC_DWORD_STAT(STAT_FoliageMeshBatches);
 
-					if (!FStaticMeshSceneProxy::GetMeshElement(LODIndex, 0, SectionIndex, GetDepthPriorityGroup(ElementParams.View), ElementParams.BatchRenderSelection[SelectionGroupIndex], true, MeshElement))
+					if (!FStaticMeshSceneProxy::GetMeshElement(LODIndex, 0, SectionIndex, GetDepthPriorityGroup(ElementParams.View), ElementParams.BatchRenderSelection[SelectionGroupIndex], true, MeshBatch))
 					{
 						continue;
 					}
-					checkSlow(MeshElement.GetNumPrimitives() > 0);
-
-					MeshElement.VertexFactory = &InstancedRenderData.VertexFactories[LODIndex];
-					FMeshBatchElement& BatchElement0 = MeshElement.Elements[0];
-
-					BatchElement0.UserData = ElementParams.PassUserData[SelectionGroupIndex];
-					BatchElement0.bUserDataIsColorVertexBuffer = false;
-					BatchElement0.MaxScreenSize = 1.0;
-					BatchElement0.MinScreenSize = 0.0;
-					BatchElement0.InstancedLODIndex = LODIndex;
-					BatchElement0.InstancedLODRange = bDitherLODEnabled ? 1 : 0;
-					BatchElement0.PrimitiveUniformBuffer = GetUniformBuffer();
-					MeshElement.bCanApplyViewModeOverrides = true;
-					MeshElement.bUseSelectionOutline = ElementParams.BatchRenderSelection[SelectionGroupIndex];
-					MeshElement.bUseWireframeSelectionColoring = ElementParams.BatchRenderSelection[SelectionGroupIndex];
-					MeshElement.bUseAsOccluder = ShouldUseAsOccluder();
-
-					if (!bDidStats)
+
+					checkSlow(MeshBatch.GetNumPrimitives() > 0);
+					MeshBatch.bCanApplyViewModeOverrides = true;
+					MeshBatch.bUseSelectionOutline = ElementParams.BatchRenderSelection[SelectionGroupIndex];
+					MeshBatch.bUseWireframeSelectionColoring = ElementParams.BatchRenderSelection[SelectionGroupIndex];
+					MeshBatch.bUseAsOccluder = ShouldUseAsOccluder();
+					MeshBatch.VertexFactory = &InstancedRenderData.VertexFactories[LODIndex];
+
+					FMeshBatchElement& MeshBatchElement = MeshBatch.Elements[0];
+					MeshBatchElement.UserData = ElementParams.PassUserData[SelectionGroupIndex];
+					MeshBatchElement.bUserDataIsColorVertexBuffer = false;
+					MeshBatchElement.MaxScreenSize = 1.0;
+					MeshBatchElement.MinScreenSize = 0.0;
+					MeshBatchElement.InstancedLODIndex = LODIndex;
+					MeshBatchElement.InstancedLODRange = bDitherLODEnabled ? 1 : 0;
+					MeshBatchElement.PrimitiveUniformBuffer = GetUniformBuffer();
+
+					int32 TotalInstances = bDitherLODEnabled ? Params.TotalMultipleLODInstances[LODIndex] : Params.TotalSingleLODInstances[LODIndex];
 					{
-						bDidStats = true;
-						int64 Tris = int64(RemainingInstances) * int64(BatchElement0.NumPrimitives);
+						const int64 Tris = int64(TotalInstances) * int64(MeshBatchElement.NumPrimitives);
 						TotalTriangles += Tris;
 #if STATS
 						if (GFrameNumberRenderThread_CaptureFoliageRuns == GFrameNumberRenderThread)
 						{
 							if (ElementParams.FinalCullDistance > 9.9E8)
 							{
-								UE_LOG(LogStaticMesh, Display, TEXT("lod:%1d/%1d   sel:%1d   section:%1d/%1d   runs:%4d   inst:%8d   tris:%9lld   cast shadow:%1d   cull:-NONE!!-   shadow:%1d     %s %s"), 
-									LODIndex, InstancedRenderData.VertexFactories.Num(), SelectionGroupIndex, SectionIndex, LODModel.Sections.Num(), RunArray.Num() / 2, 
-									RemainingInstances, Tris, (int)MeshElement.CastShadow, ElementParams.ShadowFrustum,
+								UE_LOG(LogStaticMesh, Display, TEXT("lod:%1d/%1d   sel:%1d   section:%1d/%1d   runs:%4d   inst:%8d   tris:%9lld   cast shadow:%1d   cull:-NONE!!-   shadow:%1d     %s %s"),
+									LODIndex, InstancedRenderData.VertexFactories.Num(), SelectionGroupIndex, SectionIndex, LODModel.Sections.Num(), RunArray.Num() / 2,
+									TotalInstances, Tris, (int)MeshBatch.CastShadow, ElementParams.ShadowFrustum,
 									*StaticMesh->GetPathName(),
-									*MeshElement.MaterialRenderProxy->GetIncompleteMaterialWithFallback(ElementParams.FeatureLevel).GetFriendlyName());
+									*MeshBatch.MaterialRenderProxy->GetMaterial(ElementParams.FeatureLevel)->GetFriendlyName());
 							}
 							else
 							{
-								UE_LOG(LogStaticMesh, Display, TEXT("lod:%1d/%1d   sel:%1d   section:%1d/%1d   runs:%4d   inst:%8d   tris:%9lld   cast shadow:%1d   cull:%8.0f   shadow:%1d     %s %s"), 
-									LODIndex, InstancedRenderData.VertexFactories.Num(), SelectionGroupIndex, SectionIndex, LODModel.Sections.Num(), RunArray.Num() / 2, 
-									RemainingInstances, Tris, (int)MeshElement.CastShadow, ElementParams.FinalCullDistance, ElementParams.ShadowFrustum,
+								UE_LOG(LogStaticMesh, Display, TEXT("lod:%1d/%1d   sel:%1d   section:%1d/%1d   runs:%4d   inst:%8d   tris:%9lld   cast shadow:%1d   cull:%8.0f   shadow:%1d     %s %s"),
+									LODIndex, InstancedRenderData.VertexFactories.Num(), SelectionGroupIndex, SectionIndex, LODModel.Sections.Num(), RunArray.Num() / 2,
+									TotalInstances, Tris, (int)MeshBatch.CastShadow, ElementParams.FinalCullDistance, ElementParams.ShadowFrustum,
 									*StaticMesh->GetPathName(),
-									*MeshElement.MaterialRenderProxy->GetIncompleteMaterialWithFallback(ElementParams.FeatureLevel).GetFriendlyName());
+									*MeshBatch.MaterialRenderProxy->GetMaterial(ElementParams.FeatureLevel)->GetFriendlyName());
 							}
 						}
-#endif
+#endif // STATS
 					}
-					if (ElementParams.bUseInstanceRuns)
+
+					//MeshBatchElement.NumInstances = TotalInstances;
+					// The index was used as an offset, but the dynamic buffer thing uses a resource view to make this not needed (using PrimitiveInstanceDataOffset as a temp. debug help)
+					MeshBatchElement.UserIndex = 0;
+
+					// Note: this call overrides the UserIndex to mean the command index, which is used to fetch the offset to the instance array
+					//Collector.AllocateInstancedBatchArguments(ElementParams.ViewIndex, MeshBatch, PrimitiveInstanceDataOffset, PrimitiveInstanceDataCount, RunArray);
+
+					// We use this existing hook to send info about the runs over to the visiblemeshbatch
+					MeshBatchElement.NumInstances = RunArray.Num() / 2;
+					MeshBatchElement.InstanceRuns = &RunArray[0];
+					MeshBatchElement.bIsInstanceRuns = true;
+
+					if (TotalTriangles < (int64)CVarMaxTrianglesToRender.GetValueOnRenderThread())
 					{
-						BatchElement0.NumInstances = RunArray.Num() / 2;
-						BatchElement0.InstanceRuns = &RunArray[0];
-						BatchElement0.bIsInstanceRuns = true;
-#if STATS
-						INC_DWORD_STAT_BY(STAT_FoliageRuns, BatchElement0.NumInstances);
-#endif
+						Collector.AddMesh(ElementParams.ViewIndex, MeshBatch);
 					}
-					else
+				}
+				else
+				{
+#endif // defined(GPUCULL_TODO)
+					int32 NumBatches = 1;
+					int32 CurrentRun = 0;
+					int32 CurrentInstance = 0;
+					int32 RemainingInstances = bDitherLODEnabled ? Params.TotalMultipleLODInstances[LODIndex] : Params.TotalSingleLODInstances[LODIndex];
+					int32 RemainingRuns = RunArray.Num() / 2;
+
+					if (!ElementParams.bUseInstanceRuns)
 					{
-						const uint32 NumElementsThisBatch = FMath::Min(RemainingRuns, (int32)FInstancedStaticMeshVertexFactory::NumBitsForVisibilityMask());
+						NumBatches = FMath::DivideAndRoundUp(RemainingRuns, (int32)FInstancedStaticMeshVertexFactory::NumBitsForVisibilityMask());
+					}
 
-						MeshElement.Elements.Reserve(NumElementsThisBatch);
-						check(NumElementsThisBatch);
+#if STATS
+					INC_DWORD_STAT_BY(STAT_FoliageInstances, RemainingInstances);
+#endif
+					bool bDidStats = false;
+					for (int32 BatchIndex = 0; BatchIndex < NumBatches; BatchIndex++)
+					{
+						FMeshBatch& MeshElement = Collector.AllocateMesh();
+						INC_DWORD_STAT(STAT_FoliageMeshBatches);
 
-						for (uint32 InstanceRun = 0; InstanceRun < NumElementsThisBatch; ++InstanceRun)
+						if (!FStaticMeshSceneProxy::GetMeshElement(LODIndex, 0, SectionIndex, GetDepthPriorityGroup(ElementParams.View), ElementParams.BatchRenderSelection[SelectionGroupIndex], true, MeshElement))
 						{
-							FMeshBatchElement* NewBatchElement; 
-
-							if (InstanceRun == 0)
-							{
-								NewBatchElement = &MeshElement.Elements[0];
-							}
-							else
+							continue;
+						}
+						checkSlow(MeshElement.GetNumPrimitives() > 0);
+
+						MeshElement.VertexFactory = &InstancedRenderData.VertexFactories[LODIndex];
+						FMeshBatchElement& BatchElement0 = MeshElement.Elements[0];
+
+						BatchElement0.UserData = ElementParams.PassUserData[SelectionGroupIndex];
+						BatchElement0.bUserDataIsColorVertexBuffer = false;
+						BatchElement0.MaxScreenSize = 1.0;
+						BatchElement0.MinScreenSize = 0.0;
+						BatchElement0.InstancedLODIndex = LODIndex;
+						BatchElement0.InstancedLODRange = bDitherLODEnabled ? 1 : 0;
+						BatchElement0.PrimitiveUniformBuffer = GetUniformBuffer();
+						MeshElement.bCanApplyViewModeOverrides = true;
+						MeshElement.bUseSelectionOutline = ElementParams.BatchRenderSelection[SelectionGroupIndex];
+						MeshElement.bUseWireframeSelectionColoring = ElementParams.BatchRenderSelection[SelectionGroupIndex];
+						MeshElement.bUseAsOccluder = ShouldUseAsOccluder();
+
+						if (!bDidStats)
+						{
+							bDidStats = true;
+							int64 Tris = int64(RemainingInstances) * int64(BatchElement0.NumPrimitives);
+							TotalTriangles += Tris;
+#if STATS
+							if (GFrameNumberRenderThread_CaptureFoliageRuns == GFrameNumberRenderThread)
 							{
-								NewBatchElement = new(MeshElement.Elements) FMeshBatchElement();
-								*NewBatchElement = MeshElement.Elements[0];
+								if (ElementParams.FinalCullDistance > 9.9E8)
+								{
+									UE_LOG(LogStaticMesh, Display, TEXT("lod:%1d/%1d   sel:%1d   section:%1d/%1d   runs:%4d   inst:%8d   tris:%9lld   cast shadow:%1d   cull:-NONE!!-   shadow:%1d     %s %s"),
+										LODIndex, InstancedRenderData.VertexFactories.Num(), SelectionGroupIndex, SectionIndex, LODModel.Sections.Num(), RunArray.Num() / 2,
+										RemainingInstances, Tris, (int)MeshElement.CastShadow, ElementParams.ShadowFrustum,
+										*StaticMesh->GetPathName(),
+										*MeshElement.MaterialRenderProxy->GetIncompleteMaterialWithFallback(ElementParams.FeatureLevel).GetFriendlyName());
+								}
+								else
+								{
+									UE_LOG(LogStaticMesh, Display, TEXT("lod:%1d/%1d   sel:%1d   section:%1d/%1d   runs:%4d   inst:%8d   tris:%9lld   cast shadow:%1d   cull:%8.0f   shadow:%1d     %s %s"),
+										LODIndex, InstancedRenderData.VertexFactories.Num(), SelectionGroupIndex, SectionIndex, LODModel.Sections.Num(), RunArray.Num() / 2,
+										RemainingInstances, Tris, (int)MeshElement.CastShadow, ElementParams.FinalCullDistance, ElementParams.ShadowFrustum,
+										*StaticMesh->GetPathName(),
+										*MeshElement.MaterialRenderProxy->GetIncompleteMaterialWithFallback(ElementParams.FeatureLevel).GetFriendlyName());
+								}
 							}
+#endif
+						}
+						if (ElementParams.bUseInstanceRuns)
+						{
+							BatchElement0.NumInstances = RunArray.Num() / 2;
+							BatchElement0.InstanceRuns = &RunArray[0];
+							BatchElement0.bIsInstanceRuns = true;
+#if STATS
+							INC_DWORD_STAT_BY(STAT_FoliageRuns, BatchElement0.NumInstances);
+#endif
+						}
+						else
+						{
+							const uint32 NumElementsThisBatch = FMath::Min(RemainingRuns, (int32)FInstancedStaticMeshVertexFactory::NumBitsForVisibilityMask());
 
-							const int32 InstanceOffset = RunArray[CurrentRun];
-							NewBatchElement->UserIndex = InstanceOffset;
-							NewBatchElement->NumInstances = 1 + RunArray[CurrentRun + 1] - InstanceOffset;
+							MeshElement.Elements.Reserve(NumElementsThisBatch);
+							check(NumElementsThisBatch);
 
-							if (--RemainingRuns)
+							for (uint32 InstanceRun = 0; InstanceRun < NumElementsThisBatch; ++InstanceRun)
 							{
-								CurrentRun += 2;
-								check(CurrentRun + 1 < RunArray.Num());
+								FMeshBatchElement* NewBatchElement;
+
+								if (InstanceRun == 0)
+								{
+									NewBatchElement = &MeshElement.Elements[0];
+								}
+								else
+								{
+									NewBatchElement = new(MeshElement.Elements) FMeshBatchElement();
+									*NewBatchElement = MeshElement.Elements[0];
+								}
+
+								const int32 InstanceOffset = RunArray[CurrentRun];
+								NewBatchElement->UserIndex = InstanceOffset;
+								NewBatchElement->NumInstances = 1 + RunArray[CurrentRun + 1] - InstanceOffset;
+
+								if (--RemainingRuns)
+								{
+									CurrentRun += 2;
+									check(CurrentRun + 1 < RunArray.Num());
+								}
 							}
 						}
-					}
 
-					if (TotalTriangles < (int64)CVarMaxTrianglesToRender.GetValueOnRenderThread())
-					{
-						Collector.AddMesh(ElementParams.ViewIndex, MeshElement);
+						if (TotalTriangles < (int64)CVarMaxTrianglesToRender.GetValueOnRenderThread())
+						{
+							Collector.AddMesh(ElementParams.ViewIndex, MeshElement);
+						}
 					}
+#if defined(GPUCULL_TODO)
 				}
 #endif // defined(GPUCULL_TODO)
 			}
diff --git a/Engine/Source/Runtime/Engine/Private/InstancedStaticMesh.cpp b/Engine/Source/Runtime/Engine/Private/InstancedStaticMesh.cpp
--- a/Engine/Source/Runtime/Engine/Private/InstancedStaticMesh.cpp
+++ b/Engine/Source/Runtime/Engine/Private/InstancedStaticMesh.cpp
@@ -627,10 +627,17 @@ void FInstancedStaticMeshVertexFactory::ModifyCompilationEnvironment(const FVert
 	}
 
 #if defined(GPUCULL_TODO)
-	// USE_INSTANCE_CULLING - set up additional instancing attributes (basic instancing is the default)
-	OutEnvironment.SetDefine(TEXT("USE_INSTANCE_CULLING"), TEXT("1"));
-#else
-	OutEnvironment.SetDefine(TEXT("USE_INSTANCING"), TEXT("1"));
+	if (UseGPUScene(Parameters.Platform))
+	{
+		// USE_INSTANCE_CULLING - set up additional instancing attributes (basic instancing is the default)
+		OutEnvironment.SetDefine(TEXT("USE_INSTANCE_CULLING"), TEXT("1"));
+	}
+	else
+	{
+#endif
+		OutEnvironment.SetDefine(TEXT("USE_INSTANCING"), TEXT("1"));
+#if defined(GPUCULL_TODO)
+	}
 #endif
 
 	if (IsFeatureLevelSupported(Parameters.Platform, ERHIFeatureLevel::SM5))
@@ -1050,28 +1057,31 @@ void FInstancedStaticMeshSceneProxy::SetupProxy(UInstancedStaticMeshComponent* I
 	UserData_DeselectedInstances.bRenderSelected = false;
 
 #if defined(GPUCULL_TODO)
-	const TArray<int32>& InstanceReorderTable = InComponent->InstanceReorderTable;
-	bSupportsInstanceDataBuffer = true;
-	Instances.SetNum(InComponent->GetInstanceCount());
-	for (int32 InInstanceIndex = 0; InInstanceIndex < Instances.Num(); ++InInstanceIndex)
+	if (UseGPUScene(GetScene().GetShaderPlatform(), GetScene().GetFeatureLevel()))
 	{
-		int32 OutInstanceIndex = InInstanceIndex;
-		if (OutInstanceIndex < InstanceReorderTable.Num())
+		const TArray<int32>& InstanceReorderTable = InComponent->InstanceReorderTable;
+		bSupportsInstanceDataBuffer = true;
+		Instances.SetNum(InComponent->GetInstanceCount());
+		for (int32 InInstanceIndex = 0; InInstanceIndex < Instances.Num(); ++InInstanceIndex)
 		{
-			OutInstanceIndex = InstanceReorderTable[OutInstanceIndex];
-		}
-		FTransform InstanceTransform;
-		InComponent->GetInstanceTransform(InInstanceIndex, InstanceTransform);
+			int32 OutInstanceIndex = InInstanceIndex;
+			if (OutInstanceIndex < InstanceReorderTable.Num())
+			{
+				OutInstanceIndex = InstanceReorderTable[OutInstanceIndex];
+			}
+			FTransform InstanceTransform;
+			InComponent->GetInstanceTransform(InInstanceIndex, InstanceTransform);
 
-		FPrimitiveInstance& Instance = Instances[OutInstanceIndex];
-		Instance.PrimitiveId = ~uint32(0);
-		Instance.InstanceToLocal = InstanceTransform.ToMatrixWithScale();
-		// GPUCULL_TODO: not sure this is needed either - might be better to delegate to later anyway since inverse can then be threaded, plus some platforms might not need it at all.
-		Instance.LocalToInstance = Instance.InstanceToLocal.Inverse();
-		// Filled in during GPU Scene update...
-		Instance.LocalToWorld.SetIdentity();
-		Instance.RenderBounds = InComponent->GetStaticMesh()->GetBounds();
-		Instance.LocalBounds = Instance.RenderBounds.TransformBy(Instance.InstanceToLocal);
+			FPrimitiveInstance& Instance = Instances[OutInstanceIndex];
+			Instance.PrimitiveId = ~uint32(0);
+			Instance.InstanceToLocal = InstanceTransform.ToMatrixWithScale();
+			// GPUCULL_TODO: not sure this is needed either - might be better to delegate to later anyway since inverse can then be threaded, plus some platforms might not need it at all.
+			Instance.LocalToInstance = Instance.InstanceToLocal.Inverse();
+			// Filled in during GPU Scene update...
+			Instance.LocalToWorld.SetIdentity();
+			Instance.RenderBounds = InComponent->GetStaticMesh()->GetBounds();
+			Instance.LocalBounds = Instance.RenderBounds.TransformBy(Instance.InstanceToLocal);
+		}
 	}
 #endif
 }
@@ -2965,7 +2975,7 @@ void UInstancedStaticMeshComponent::InitPerInstanceRenderData(bool InitializeFro
 	ERHIFeatureLevel::Type FeatureLevel = World != nullptr ? World->FeatureLevel.GetValue() : GMaxRHIFeatureLevel;
 
 #if defined(GPUCULL_TODO)
-	bool KeepInstanceBufferCPUAccess = true;//GIsEditor || InRequireCPUAccess || ComponentRequestsCPUAccess(this, FeatureLevel);
+	bool KeepInstanceBufferCPUAccess = UseGPUScene(GetFeatureLevelShaderPlatform(FeatureLevel), FeatureLevel) || GIsEditor || InRequireCPUAccess || ComponentRequestsCPUAccess(this, FeatureLevel);
 #else
 	bool KeepInstanceBufferCPUAccess = GIsEditor || InRequireCPUAccess || ComponentRequestsCPUAccess(this, FeatureLevel);
 #endif
diff --git a/Engine/Source/Runtime/Engine/Private/SkeletalMesh.cpp b/Engine/Source/Runtime/Engine/Private/SkeletalMesh.cpp
--- a/Engine/Source/Runtime/Engine/Private/SkeletalMesh.cpp
+++ b/Engine/Source/Runtime/Engine/Private/SkeletalMesh.cpp
@@ -5092,22 +5092,26 @@ FSkeletalMeshSceneProxy::FSkeletalMeshSceneProxy(const USkinnedMeshComponent* Co
 	// Skip primitive uniform buffer if we will be using local vertex factory which gets it's data from GPUScene.
 	// GPUCULL_TODO: Move to base class - all should follow the same pattern here
 #if defined(GPUCULL_TODO)
-	bVFRequiresPrimitiveUniformBuffer = !UseGPUScene(GMaxRHIShaderPlatform, FeatureLevel);
-	bSupportsInstanceDataBuffer = true;
+	const bool bUseGPUScene = UseGPUScene(GMaxRHIShaderPlatform, FeatureLevel);
+	bVFRequiresPrimitiveUniformBuffer = !bUseGPUScene;
 #else
 	bVFRequiresPrimitiveUniformBuffer = !((bIsCPUSkinned || bRenderStatic) && UseGPUScene(GMaxRHIShaderPlatform, FeatureLevel));
 #endif
 
 #if defined(GPUCULL_TODO)
-	// 100% generic default-instance
-	Instances.SetNum(1);
-	FPrimitiveInstance& Instance = Instances[0];
-	Instance.PrimitiveId = ~uint32(0);
-	Instance.InstanceToLocal.SetIdentity();
-	Instance.LocalToInstance.SetIdentity();
-	Instance.LocalToWorld.SetIdentity();
-	Instance.RenderBounds = GetLocalBounds();
-	Instance.LocalBounds = Instance.RenderBounds;
+	if (bUseGPUScene)
+	{
+		bSupportsInstanceDataBuffer = true;
+		// 100% generic default-instance
+		Instances.SetNum(1);
+		FPrimitiveInstance& Instance = Instances[0];
+		Instance.PrimitiveId = ~uint32(0);
+		Instance.InstanceToLocal.SetIdentity();
+		Instance.LocalToInstance.SetIdentity();
+		Instance.LocalToWorld.SetIdentity();
+		Instance.RenderBounds = GetLocalBounds();
+		Instance.LocalBounds = Instance.RenderBounds;
+	}
 #endif
 
 #if RHI_RAYTRACING
diff --git a/Engine/Source/Runtime/Engine/Private/StaticMeshRender.cpp b/Engine/Source/Runtime/Engine/Private/StaticMeshRender.cpp
--- a/Engine/Source/Runtime/Engine/Private/StaticMeshRender.cpp
+++ b/Engine/Source/Runtime/Engine/Private/StaticMeshRender.cpp
@@ -345,15 +345,18 @@ FStaticMeshSceneProxy::FStaticMeshSceneProxy(UStaticMeshComponent* InComponent,
 
 	AddSpeedTreeWind();
 #if defined(GPUCULL_TODO)
-	Instances.SetNum(1);
-	FPrimitiveInstance& Instance = Instances[0];
-	Instance.PrimitiveId = ~uint32(0);
-	Instance.InstanceToLocal.SetIdentity();
-	Instance.LocalToInstance.SetIdentity();
-	Instance.LocalToWorld.SetIdentity();
-	Instance.RenderBounds = InComponent->GetStaticMesh()->GetBounds();
-	Instance.LocalBounds = Instance.RenderBounds;
-	bSupportsInstanceDataBuffer = true;
+	if (UseGPUScene(GMaxRHIShaderPlatform, FeatureLevel))
+	{
+		Instances.SetNum(1);
+		FPrimitiveInstance& Instance = Instances[0];
+		Instance.PrimitiveId = ~uint32(0);
+		Instance.InstanceToLocal.SetIdentity();
+		Instance.LocalToInstance.SetIdentity();
+		Instance.LocalToWorld.SetIdentity();
+		Instance.RenderBounds = InComponent->GetStaticMesh()->GetBounds();
+		Instance.LocalBounds = Instance.RenderBounds;
+		bSupportsInstanceDataBuffer = true;
+	}
 #endif // defined(GPUCULL_TODO)
 }
 
diff --git a/Engine/Source/Runtime/RenderCore/Public/RenderUtils.h b/Engine/Source/Runtime/RenderCore/Public/RenderUtils.h
--- a/Engine/Source/Runtime/RenderCore/Public/RenderUtils.h
+++ b/Engine/Source/Runtime/RenderCore/Public/RenderUtils.h
@@ -572,6 +572,12 @@ inline bool UseGPUScene(const FStaticShaderPlatform Platform, const FStaticFeatu
 		&& (!FDataDrivenShaderPlatformInfo::IsValid(Platform) || FDataDrivenShaderPlatformInfo::GetSupportsGPUScene(Platform));
 }
 
+
+inline bool UseGPUScene(const FStaticShaderPlatform Platform)
+{
+	return UseGPUScene(Platform, GetMaxSupportedFeatureLevel(Platform));
+}
+
 inline bool ForceSimpleSkyDiffuse(const FStaticShaderPlatform Platform)
 {
 	extern RENDERCORE_API uint64 GSimpleSkyDiffusePlatformMask;
diff --git a/Engine/Source/Runtime/Renderer/Private/DeferredShadingRenderer.cpp b/Engine/Source/Runtime/Renderer/Private/DeferredShadingRenderer.cpp
--- a/Engine/Source/Runtime/Renderer/Private/DeferredShadingRenderer.cpp
+++ b/Engine/Source/Runtime/Renderer/Private/DeferredShadingRenderer.cpp
@@ -1618,7 +1618,7 @@ void FDeferredShadingSceneRenderer::Render(FRDGBuilder& GraphBuilder)
 	// Find the visible primitives.
 	GraphBuilder.RHICmdList.ImmediateFlush(EImmediateFlushType::DispatchToRHIThread);
 
-	FInstanceCullingManager InstanceCullingManager(GInstanceCullingManagerResources);
+	FInstanceCullingManager InstanceCullingManager(GInstanceCullingManagerResources, Scene->GPUScene.IsEnabled());
 
 	bool bDoInitViewAftersPrepass = false;
 	{
diff --git a/Engine/Source/Runtime/Renderer/Private/InstanceCulling/InstanceCullingContext.cpp b/Engine/Source/Runtime/Renderer/Private/InstanceCulling/InstanceCullingContext.cpp
--- a/Engine/Source/Runtime/Renderer/Private/InstanceCulling/InstanceCullingContext.cpp
+++ b/Engine/Source/Runtime/Renderer/Private/InstanceCulling/InstanceCullingContext.cpp
@@ -85,7 +85,7 @@ public:
 
 	static bool ShouldCompilePermutation(const FGlobalShaderPermutationParameters& Parameters)
 	{
-		return IsFeatureLevelSupported(Parameters.Platform, ERHIFeatureLevel::SM5);
+		return UseGPUScene(Parameters.Platform);
 	}
 
 	static void ModifyCompilationEnvironment(const FGlobalShaderPermutationParameters& Parameters, FShaderCompilerEnvironment& OutEnvironment)
@@ -131,7 +131,7 @@ public:
 
 	static bool ShouldCompilePermutation(const FGlobalShaderPermutationParameters& Parameters)
 	{
-		return IsFeatureLevelSupported(Parameters.Platform, ERHIFeatureLevel::SM5);
+		return UseGPUScene(Parameters.Platform);
 	}
 
 	static void ModifyCompilationEnvironment(const FGlobalShaderPermutationParameters& Parameters, FShaderCompilerEnvironment& OutEnvironment)
@@ -169,7 +169,7 @@ public:
 
 	static bool ShouldCompilePermutation(const FGlobalShaderPermutationParameters& Parameters)
 	{
-		return IsFeatureLevelSupported(Parameters.Platform, ERHIFeatureLevel::SM5);
+		return UseGPUScene(Parameters.Platform);
 	}
 
 	static void ModifyCompilationEnvironment(const FGlobalShaderPermutationParameters& Parameters, FShaderCompilerEnvironment& OutEnvironment)
@@ -229,7 +229,7 @@ public:
 
 	static bool ShouldCompilePermutation(const FGlobalShaderPermutationParameters& Parameters)
 	{
-		return IsFeatureLevelSupported(Parameters.Platform, ERHIFeatureLevel::SM5);
+		return UseGPUScene(Parameters.Platform);
 	}
 
 	static void ModifyCompilationEnvironment(const FGlobalShaderPermutationParameters& Parameters, FShaderCompilerEnvironment& OutEnvironment)
diff --git a/Engine/Source/Runtime/Renderer/Private/InstanceCulling/InstanceCullingManager.cpp b/Engine/Source/Runtime/Renderer/Private/InstanceCulling/InstanceCullingManager.cpp
--- a/Engine/Source/Runtime/Renderer/Private/InstanceCulling/InstanceCullingManager.cpp
+++ b/Engine/Source/Runtime/Renderer/Private/InstanceCulling/InstanceCullingManager.cpp
@@ -20,11 +20,13 @@ static TAutoConsoleVariable<int32> CVarCullInstances(
 	ECVF_RenderThreadSafe);
 
 
-#if defined(GPUCULL_TODO)
-
-
 int32 FInstanceCullingManager::RegisterView(const FViewInfo& ViewInfo)
 {
+	if (!bIsEnabled)
+	{
+		return 0;
+	}
+
 	Nanite::FPackedViewParams Params;
 	Params.ViewMatrices = ViewInfo.ViewMatrices;
 	Params.PrevViewMatrices = ViewInfo.PrevViewInfo.ViewMatrices;
@@ -38,10 +40,15 @@ int32 FInstanceCullingManager::RegisterView(const FViewInfo& ViewInfo)
 
 int32 FInstanceCullingManager::RegisterView(const Nanite::FPackedViewParams& Params)
 {
+	if (!bIsEnabled)
+	{
+		return 0;
+	}
 	CullingViews.Add(CreatePackedView(Params));
 	return CullingViews.Num() - 1;
 }
 
+#if defined(GPUCULL_TODO)
 
 class FCullInstancesCs : public FGlobalShader
 {
@@ -53,7 +60,7 @@ public:
 
 	static bool ShouldCompilePermutation(const FGlobalShaderPermutationParameters& Parameters)
 	{
-		return IsFeatureLevelSupported(Parameters.Platform, ERHIFeatureLevel::SM5);
+		return UseGPUScene(Parameters.Platform);
 	}
 
 	/**
@@ -85,9 +92,11 @@ public:
 };
 IMPLEMENT_GLOBAL_SHADER(FCullInstancesCs, "/Engine/Private/InstanceCulling/CullInstances.usf", "CullInstancesCs", SF_Compute);
 
+#endif // GPUCULL_TODO
 
 void FInstanceCullingManager::CullInstances(FRDGBuilder& GraphBuilder, FGPUScene& GPUScene)
 {
+#if defined(GPUCULL_TODO)
 	RDG_EVENT_SCOPE(GraphBuilder, "CullInstances");
 
 	check(!CullingIntermediate.InstanceIdOutOffsetBuffer);
@@ -145,38 +154,21 @@ void FInstanceCullingManager::CullInstances(FRDGBuilder& GraphBuilder, FGPUScene
 			AddClearUAVPass(GraphBuilder, VisibleInstanceFlagsUAV, 0xFFFFFFFF);
 		}
 	}
+#endif // defined(GPUCULL_TODO)
 }
 
 
 FInstanceCullingContext* FInstanceCullingManager::CreateContext(const int32* ViewIds, int32 NumViews)
 {
-	FInstanceCullingContext* InstanceCullingContext = new(FMemStack::Get()) FInstanceCullingContext;
-	InstanceCullingContext->ViewIds.Insert(ViewIds, NumViews, 0);
-	InstanceCullingContext->InstanceCullingManager = this;
-	return InstanceCullingContext;
-}
-
-
-#else //!defined(GPUCULL_TODO)
-
-int32 FInstanceCullingManager::RegisterView(const FViewInfo& ViewInfo)
-{
-	return 0;
-}
-
-int32 FInstanceCullingManager::RegisterView(const Nanite::FPackedViewParams& Params)
-{
-	return 0;
-}
-
-void FInstanceCullingManager::CullInstances(FRDGBuilder& GraphBuilder, FGPUScene& GPUScene)
-{
-}
-
-FInstanceCullingContext* FInstanceCullingManager::CreateContext(const int32 * ViewIds, int32 NumViews)
-{
-	return nullptr;
+	if (bIsEnabled)
+	{
+		FInstanceCullingContext* InstanceCullingContext = new(FMemStack::Get()) FInstanceCullingContext;
+		InstanceCullingContext->ViewIds.Insert(ViewIds, NumViews, 0);
+		InstanceCullingContext->InstanceCullingManager = this;
+		return InstanceCullingContext;
+	}
+	else
+	{
+		return nullptr;
+	}
 }
-
-
-#endif // defined(GPUCULL_TODO)
diff --git a/Engine/Source/Runtime/Renderer/Private/InstanceCulling/InstanceCullingManager.h b/Engine/Source/Runtime/Renderer/Private/InstanceCulling/InstanceCullingManager.h
--- a/Engine/Source/Runtime/Renderer/Private/InstanceCulling/InstanceCullingManager.h
+++ b/Engine/Source/Runtime/Renderer/Private/InstanceCulling/InstanceCullingManager.h
@@ -89,8 +89,9 @@ public:
 	 */
 	void CullInstances(FRDGBuilder& GraphBuilder, FGPUScene& GPUScene);
 
-	FInstanceCullingManager(FInstanceCullingManagerResources& InResources)
-	: Resources(InResources)
+	FInstanceCullingManager(FInstanceCullingManagerResources& InResources, bool bInIsEnabled)
+	: Resources(InResources),
+		bIsEnabled(bInIsEnabled)
 	{
 	}
 
@@ -100,11 +101,16 @@ public:
 	 */
 	FInstanceCullingContext* CreateContext(const int32* ViewIds, int32 NumViews);
 
+	// Populated by CullInstances, used when performing final culling & rendering 
+	FInstanceCullingIntermediate CullingIntermediate;
+
+private:
+	FInstanceCullingManager() = delete;
+	FInstanceCullingManager(FInstanceCullingManager &) = delete;
+
 	FInstanceCullingManagerResources& Resources;
-	//TArray<TFunction<void(FInstanceCullingManager*, const FRDGBufferRef&, FRDGBuilder&, const FInstanceCullingIntermediate&, FGPUScene&)>, SceneRenderingAllocator> GPUPostCullingJobs;
 	TArray<Nanite::FPackedView> CullingViews;
 	
-	// Populated by CullInstances, used when performing final culling & rendering 
-	FInstanceCullingIntermediate CullingIntermediate;
+	bool bIsEnabled;
 };
 
diff --git a/Engine/Source/Runtime/Renderer/Private/MeshDrawCommands.cpp b/Engine/Source/Runtime/Renderer/Private/MeshDrawCommands.cpp
--- a/Engine/Source/Runtime/Renderer/Private/MeshDrawCommands.cpp
+++ b/Engine/Source/Runtime/Renderer/Private/MeshDrawCommands.cpp
@@ -1467,31 +1467,38 @@ public:
 void FParallelMeshDrawCommandPass::BuildRenderingCommands(FRDGBuilder& GraphBuilder, FGPUScene& GPUScene, FInstanceCullingDrawParams& OutInstanceCullingDrawParams)
 {
 #if defined(GPUCULL_TODO)
-	WaitForMeshPassSetupTask();
-	if (MaxNumDraws <= 0 || !TaskContext.InstanceCullingContext->HasCullingCommands())
+	if (TaskContext.InstanceCullingContext)
 	{
-		OutInstanceCullingDrawParams.DrawIndirectArgsBuffer = nullptr;
-		OutInstanceCullingDrawParams.DrawIndirectArgsBufferAccess = nullptr;
-		OutInstanceCullingDrawParams.InstanceIdOffsetBuffer = nullptr;
-		OutInstanceCullingDrawParams.InstanceIdOffsetBufferAccess = nullptr;
-		return;
+		WaitForMeshPassSetupTask();
+		if (MaxNumDraws > 0 && TaskContext.InstanceCullingContext->HasCullingCommands())
+		{
+			// 2. Run finalize culling commands pass
+			TaskContext.InstanceCullingContext->BuildRenderingCommands(GraphBuilder, GPUScene, TaskContext.InstanceCullingResult);
+			TaskContext.InstanceCullingResult.GetDrawParameters(OutInstanceCullingDrawParams);
+			return;
+		}
 	}
-	// 2. Run finalize culling commands pass
-	TaskContext.InstanceCullingContext->BuildRenderingCommands(GraphBuilder, GPUScene, TaskContext.InstanceCullingResult);
-	TaskContext.InstanceCullingResult.GetDrawParameters(OutInstanceCullingDrawParams);
+	OutInstanceCullingDrawParams.DrawIndirectArgsBuffer = nullptr;
+	OutInstanceCullingDrawParams.DrawIndirectArgsBufferAccess = nullptr;
+	OutInstanceCullingDrawParams.InstanceIdOffsetBuffer = nullptr;
+	OutInstanceCullingDrawParams.InstanceIdOffsetBufferAccess = nullptr;
+
 #endif
 }
 
 void FParallelMeshDrawCommandPass::BuildInstanceList(FRDGBuilder& GraphBuilder, FGPUScene& GPUScene, FInstanceCullingRdgParams& OutParams)
 {
 #if defined(GPUCULL_TODO)
-	WaitForMeshPassSetupTask();
-	if (MaxNumDraws <= 0)
+	if (TaskContext.InstanceCullingContext)
 	{
-		return;
+		WaitForMeshPassSetupTask();
+		if (MaxNumDraws <= 0)
+		{
+			return;
+		}
+		// Run pass to build ID lists (temporary)
+		TaskContext.InstanceCullingContext->BuildRenderingCommands(GraphBuilder, GPUScene, OutParams);
 	}
-	// Run pass to build ID lists (temporary)
-	TaskContext.InstanceCullingContext->BuildRenderingCommands(GraphBuilder, GPUScene, OutParams);
 #endif
 }
 
@@ -1513,13 +1520,16 @@ void FParallelMeshDrawCommandPass::DispatchDraw(FParallelCommandListSet* Paralle
 		DrawIndirectArgsBuffer = InstanceCullingDrawParams->DrawIndirectArgsBuffer->GetRHI();
 		InstanceIdOffsetBuffer = InstanceCullingDrawParams->InstanceIdOffsetBuffer->GetRHI();
 	}
-#endif // defined(GPUCULL_TODO)
+
+#else // !defined(GPUCULL_TODO)
 
 	FRHIBuffer* PrimitiveIdsBuffer = PrimitiveIdVertexBufferPoolEntry.BufferRHI;
 	const int32 BasePrimitiveIdsOffset = 0;
+#endif // defined(GPUCULL_TODO)
 
 	if (ParallelCommandListSet)
 	{
+#if !defined(GPUCULL_TODO)
 		if (TaskContext.bUseGPUScene)
 		{
 			// Queue a command on the RHI thread which will upload PrimitiveIdVertexBuffer after finishing FMeshDrawCommandPassSetupTask.
@@ -1548,6 +1558,7 @@ void FParallelMeshDrawCommandPass::DispatchDraw(FParallelCommandListSet* Paralle
 
 			bPrimitiveIdBufferDataOwnedByRHIThread = true;
 		}
+#endif // !defined(GPUCULL_TODO)
 
 		const ENamedThreads::Type RenderThread = ENamedThreads::GetRenderThread();
 
diff --git a/Engine/Source/Runtime/Renderer/Private/MobileShadingRenderer.cpp b/Engine/Source/Runtime/Renderer/Private/MobileShadingRenderer.cpp
--- a/Engine/Source/Runtime/Renderer/Private/MobileShadingRenderer.cpp
+++ b/Engine/Source/Runtime/Renderer/Private/MobileShadingRenderer.cpp
@@ -575,7 +575,7 @@ void FMobileSceneRenderer::Render(FRDGBuilder& GraphBuilder)
 	GSystemTextures.InitializeTextures(GraphBuilder.RHICmdList, FeatureLevel);
 	FRDGSystemTextures::Create(GraphBuilder);
 
-	FInstanceCullingManager InstanceCullingManager(GInstanceCullingManagerResources);
+	FInstanceCullingManager InstanceCullingManager(GInstanceCullingManagerResources, Scene->GPUScene.IsEnabled());
 
 	// Find the visible primitives and prepare targets and buffers for rendering
 	InitViews(GraphBuilder, SceneTexturesConfig, InstanceCullingManager);
diff --git a/Engine/Source/Runtime/Renderer/Private/SceneHitProxyRendering.cpp b/Engine/Source/Runtime/Renderer/Private/SceneHitProxyRendering.cpp
--- a/Engine/Source/Runtime/Renderer/Private/SceneHitProxyRendering.cpp
+++ b/Engine/Source/Runtime/Renderer/Private/SceneHitProxyRendering.cpp
@@ -601,7 +601,7 @@ void FMobileSceneRenderer::RenderHitProxies(FRDGBuilder& GraphBuilder)
 	FRDGTextureRef HitProxyDepthTexture = nullptr;
 	InitHitProxyRender(GraphBuilder, this, SceneTexturesConfig, HitProxyTexture, HitProxyDepthTexture);
 
-	FInstanceCullingManager InstanceCullingManager(GInstanceCullingManagerResources);
+	FInstanceCullingManager InstanceCullingManager(GInstanceCullingManagerResources, Scene->GPUScene.IsEnabled());
 
 	// Find the visible primitives.
 	InitViews(GraphBuilder, SceneTexturesConfig, InstanceCullingManager);
@@ -641,7 +641,7 @@ void FDeferredShadingSceneRenderer::RenderHitProxies(FRDGBuilder& GraphBuilder)
 
 	const FIntPoint HitProxyTextureSize = HitProxyDepthTexture->Desc.Extent;
 
-	FInstanceCullingManager InstanceCullingManager(GInstanceCullingManagerResources);
+	FInstanceCullingManager InstanceCullingManager(GInstanceCullingManagerResources, Scene->GPUScene.IsEnabled());
 
 	// Find the visible primitives.
 	{
diff --git a/Engine/Source/Runtime/Renderer/Private/ShadowDepthRendering.cpp b/Engine/Source/Runtime/Renderer/Private/ShadowDepthRendering.cpp
--- a/Engine/Source/Runtime/Renderer/Private/ShadowDepthRendering.cpp
+++ b/Engine/Source/Runtime/Renderer/Private/ShadowDepthRendering.cpp
@@ -323,9 +323,14 @@ public:
 		OutEnvironment.SetDefine(TEXT("USING_VERTEX_SHADER_LAYER"), (uint32)(ShaderMode == VertexShadowDepth_VSLayer));
 		OutEnvironment.SetDefine(TEXT("POSITION_ONLY"), (uint32)bUsePositionOnlyStream);
 		OutEnvironment.SetDefine(TEXT("IS_FOR_GEOMETRY_SHADER"), (uint32)bIsForGeometryShader);
-		OutEnvironment.SetDefine(TEXT("ENABLE_NON_NANITE_VSM"), (uint32)ENABLE_NON_NANITE_VSM);
+
+		uint32 bEnableNonNaniteVSM = (uint32)(ENABLE_NON_NANITE_VSM != 0 && UseGPUScene(Parameters.Platform));
+		OutEnvironment.SetDefine(TEXT("ENABLE_NON_NANITE_VSM"), bEnableNonNaniteVSM);
 #if ENABLE_NON_NANITE_VSM
-		FVirtualShadowMapArray::SetShaderDefines(OutEnvironment);
+		if (bEnableNonNaniteVSM != 0)
+		{
+			FVirtualShadowMapArray::SetShaderDefines(OutEnvironment);
+		}
 #endif // ENABLE_NON_NANITE_VSM
 		if (bIsForGeometryShader)
 		{
@@ -444,6 +449,9 @@ public:
 		FMeshMaterialShader::ModifyCompilationEnvironment(Parameters, OutEnvironment);
 		OutEnvironment.SetDefine(TEXT("ONEPASS_POINTLIGHT_SHADOW"), 1);
 		TShadowDepthVS<VertexShadowDepth_OnePassPointLight, false, true>::ModifyCompilationEnvironment(Parameters, OutEnvironment);
+#if defined(GPUCULL_TODO)
+		OutEnvironment.SetDefine(TEXT("ENABLE_FALLBACK_POINTLIGHT_SHADOW_GS"), UseGPUScene(Parameters.Platform) ? 1 : 0);
+#endif // defined(GPUCULL_TODO)
 	}
 
 	FOnePassPointShadowDepthGS(const ShaderMetaType::CompiledShaderInitializerType& Initializer)
