Fixed canvas triangle rendering by copying the mesh batch for each triangle instead of modifying the original one in-place

#rnx

diff --git a/Engine/Source/Runtime/Engine/Private/TriangleRendering.cpp b/Engine/Source/Runtime/Engine/Private/TriangleRendering.cpp
--- a/Engine/Source/Runtime/Engine/Private/TriangleRendering.cpp
+++ b/Engine/Source/Runtime/Engine/Private/TriangleRendering.cpp
@@ -148,7 +148,7 @@ void FCanvasTriangleRendererItem::FRenderData::RenderTriangles(
 	{
 		const FTriangleInst& Tri = Triangles[TriIdx];
 
-		FMeshBatch& MeshBatch = TriMesh.MeshBatch;
+		FMeshBatch& MeshBatch = *RenderContext.Alloc<FMeshBatch>(TriMesh.MeshBatch);
 		MeshBatch.VertexFactory = &VertexFactory;
 		MeshBatch.MaterialRenderProxy = MaterialRenderProxy;
 		MeshBatch.Elements[0].FirstIndex = 3 * TriIdx;
