Small refactor of DataLayer utility methods in Actor.

diff --git a/Engine/Source/Runtime/Engine/Private/ActorEditor.cpp b/Engine/Source/Runtime/Engine/Private/ActorEditor.cpp
--- a/Engine/Source/Runtime/Engine/Private/ActorEditor.cpp
+++ b/Engine/Source/Runtime/Engine/Private/ActorEditor.cpp
@@ -1234,34 +1234,14 @@ bool AActor::HasAllDataLayers(const TArray<const UDataLayer*>& InDataLayers) con
 
 TArray<FName> AActor::GetDataLayerNames() const
 {
-	TArray<FName> OutDataLayerNames;
-	OutDataLayerNames.Reserve(DataLayers.Num());
-
-	for (const UDataLayer* DataLayer : GetDataLayerObjects())
-	{
-		OutDataLayerNames.Add(DataLayer->GetFName());
-	}
-
-	return MoveTemp(OutDataLayerNames);
+	const AWorldDataLayers* WorldDataLayers = AWorldDataLayers::Get(GetWorld());
+	return WorldDataLayers ? WorldDataLayers->GetDataLayerNames(DataLayers) : TArray<FName>();
 }
 
 TArray<const UDataLayer*> AActor::GetDataLayerObjects() const
 {
-	TArray<const UDataLayer*> OutDataLayers;
-	OutDataLayers.Reserve(DataLayers.Num());
-	
-	if (const AWorldDataLayers* WorldDataLayers = AWorldDataLayers::Get(GetWorld()))
-	{
-		for (const FActorDataLayer& DataLayer : DataLayers)
-		{
-			if (const UDataLayer* DataLayerObject = WorldDataLayers->GetDataLayerFromName(DataLayer.Name))
-			{
-				OutDataLayers.Add(DataLayerObject);
-			}
-		}
-	}
-
-	return MoveTemp(OutDataLayers);
+	const AWorldDataLayers* WorldDataLayers = AWorldDataLayers::Get(GetWorld());
+	return WorldDataLayers ? WorldDataLayers->GetDataLayerObjects(DataLayers) : TArray<const UDataLayer*>();
 }
 
 bool AActor::HasAnyOfDataLayers(const TArray<FName>& DataLayerNames) const
