Fix skeletalmesh fix bad tangent triangle iteration error
Merging fix from //UE4/Main //UE5/Main

diff --git a/Engine/Source/Runtime/Engine/Private/SkeletalMesh.cpp b/Engine/Source/Runtime/Engine/Private/SkeletalMesh.cpp
--- a/Engine/Source/Runtime/Engine/Private/SkeletalMesh.cpp
+++ b/Engine/Source/Runtime/Engine/Private/SkeletalMesh.cpp
@@ -2416,47 +2416,77 @@ void USkeletalMesh::PostLoadVerifyAndFixBadTangent()
 		{
 			FSkelMeshSection& Section = ThisLODModel.Sections[SectionIndex];
 			const int32 NumVertices = Section.GetNumVertices();
+			const int32 SectionBaseIndex = Section.BaseIndex;
+			const int32 SectionNumTriangles = Section.NumTriangles;
+			TArray<uint32>& IndexBuffer = ThisLODModel.IndexBuffer;
 			//We inspect triangle per section so we need to reset the array when we start a new section.
-			TriangleTangents.Reset();
-			for (int32 VertexIndex = 0; VertexIndex < NumVertices; ++VertexIndex)
+			TriangleTangents.Empty(SectionNumTriangles);
+			for (int32 FaceIndex = 0; FaceIndex < SectionNumTriangles; ++FaceIndex)
 			{
-				FSoftSkinVertex& SoftSkinVertex = Section.SoftVertices[VertexIndex];
-				//Make sure we have normalized tangents
-				auto NormalizedTangent = [](FVector& Tangent)
+				int32 BaseFaceIndexBufferIndex = SectionBaseIndex + (FaceIndex * 3);
+				if (!ensure(IndexBuffer.IsValidIndex(BaseFaceIndexBufferIndex)) || !ensure(IndexBuffer.IsValidIndex(BaseFaceIndexBufferIndex + 2)))
 				{
-					if (Tangent.ContainsNaN() || Tangent.SizeSquared() < THRESH_VECTOR_NORMALIZED)
-					{
-						//This is a degenerated tangent, we will set it to zero. It will be fix by the
-						//FixTangent lambda function.
-						Tangent = FVector::ZeroVector;
-						return false;
-					}
-					else if (!Tangent.IsNormalized())
-					{
-						//This is not consider has a bad normal since the tangent vector is not near zero.
-						//We are just making sure the tangent is normalize.
-						Tangent.Normalize();
-					}
-					return true;
-				};
-
-				/** Call this lambda only if you need to fix the tangent */
-				auto FixTangent = [&VertexIndex, &Section, &TriangleTangents, &ComputeTriangleTangent](FVector& TangentA, const FVector& TangentB, const FVector& TangentC, const int32 Offset)
+					break;
+				}
+				for (int32 Corner = 0; Corner < 3; ++Corner)
 				{
-					//If the two other axis are valid, fix the tangent with a cross product and normalize the answer.
-					if (TangentB.IsNormalized() && TangentC.IsNormalized())
+					const int32 CornerIndexBufferIndex = BaseFaceIndexBufferIndex + Corner;
+					ensure(IndexBuffer.IsValidIndex(CornerIndexBufferIndex));
+					int32 VertexIndex = IndexBuffer[CornerIndexBufferIndex] - Section.BaseVertexIndex;
+					ensure(Section.SoftVertices.IsValidIndex(VertexIndex));
+					FSoftSkinVertex& SoftSkinVertex = Section.SoftVertices[VertexIndex];
+
+					bool bNeedToOrthonormalize = false;
+
+					//Make sure we have normalized tangents
+					auto NormalizedTangent = [&bNeedToOrthonormalize, &bFoundBadTangents](FVector& Tangent)
 					{
-						TangentA = FVector::CrossProduct(TangentB, TangentC);
-						TangentA.Normalize();
-					}
-					else
+						if (Tangent.ContainsNaN() || Tangent.SizeSquared() < THRESH_VECTOR_NORMALIZED)
+						{
+							//This is a degenerated tangent, we will set it to zero. It will be fix by the
+							//FixTangent lambda function.
+							Tangent = FVector::ZeroVector;
+							//If we can fix this tangents, we have to orthonormalize the result
+							bNeedToOrthonormalize = true;
+							bFoundBadTangents = true;
+							return false;
+						}
+						else if (!Tangent.IsNormalized())
+						{
+							//This is not consider has a bad normal since the tangent vector is not near zero.
+							//We are just making sure the tangent is normalize.
+							Tangent.Normalize();
+						}
+						return true;
+					};
+
+					/** Call this lambda only if you need to fix the tangent */
+					auto FixTangent = [&IndexBuffer, &Section, &TriangleTangents, &ComputeTriangleTangent, &BaseFaceIndexBufferIndex](FVector& TangentA, const FVector& TangentB, const FVector& TangentC, const int32 Offset)
 					{
+						//If the two other axis are valid, fix the tangent with a cross product and normalize the answer.
+						if (TangentB.IsNormalized() && TangentC.IsNormalized())
+						{
+							TangentA = FVector::CrossProduct(TangentB, TangentC);
+							TangentA.Normalize();
+							return true;
+						}
+
 						//We do not have any valid data to help us for fixing this normal so apply the triangle normals, this will create a faceted mesh but this is better then a black not shade mesh.
-						int32 FaceVertexIndex = FMath::Floor<int32>(VertexIndex / 3) * 3;
-						TArray<FVector>& Tangents = TriangleTangents.FindOrAdd(FaceVertexIndex);
+						TArray<FVector>& Tangents = TriangleTangents.FindOrAdd(BaseFaceIndexBufferIndex);
 						if (Tangents.Num() == 0)
 						{
-							ComputeTriangleTangent(Section.SoftVertices[FaceVertexIndex], Section.SoftVertices[FaceVertexIndex + 1], Section.SoftVertices[FaceVertexIndex + 2], Tangents);
+							const int32 VertexIndex0 = IndexBuffer[BaseFaceIndexBufferIndex] - Section.BaseVertexIndex;
+							const int32 VertexIndex1 = IndexBuffer[BaseFaceIndexBufferIndex + 1] - Section.BaseVertexIndex;
+							const int32 VertexIndex2 = IndexBuffer[BaseFaceIndexBufferIndex + 2] - Section.BaseVertexIndex;
+							if (!ensure(
+								Section.SoftVertices.IsValidIndex(VertexIndex0) &&
+								Section.SoftVertices.IsValidIndex(VertexIndex1) &&
+								Section.SoftVertices.IsValidIndex(VertexIndex2) ) )
+							{
+								//We found bad vertex indices, we cannot compute this face tangents.
+								return false;
+							}
+							ComputeTriangleTangent(Section.SoftVertices[VertexIndex0], Section.SoftVertices[VertexIndex1], Section.SoftVertices[VertexIndex2], Tangents);
 							const FVector Axis[3] = { {1.0f, 0.0f, 0.0f}, {0.0f, 1.0f, 0.0f}, {0.0f, 0.0f, 1.0f} };
 							if (!ensure(Tangents.Num() == 3))
 							{
@@ -2474,7 +2504,7 @@ void USkeletalMesh::PostLoadVerifyAndFixBadTangent()
 							{
 								//We are not able to compute the triangle tangent, this is probably a degenerated triangle
 								Tangents.Empty(3);
-								
+
 								Tangents.Add(Axis[0]);
 								Tangents.Add(Axis[1]);
 								Tangents.Add(Axis[2]);
@@ -2482,51 +2512,48 @@ void USkeletalMesh::PostLoadVerifyAndFixBadTangent()
 						}
 						//Use the offset to know which tangent type we are setting (0: Tangent X, 1: bi-normal Y, 2: Normal Z)
 						TangentA = Tangents[(Offset) % 3];
+						return TangentA.IsNormalized();
+					};
+
+					//The SoftSkinVertex TangentZ is a FVector4 so we must use a temporary FVector to be able to pass reference
+					FVector TangentZ = SoftSkinVertex.TangentZ;
+					//Make sure the tangent space is normalize before fixing bad tangent, because we want to do a cross product
+					//of 2 valid axis if possible. If not possible we will use the triangle normal which give a faceted triangle.
+					bool ValidTangentX = NormalizedTangent(SoftSkinVertex.TangentX);
+					bool ValidTangentY = NormalizedTangent(SoftSkinVertex.TangentY);
+					bool ValidTangentZ = NormalizedTangent(TangentZ);
+
+					if (!ValidTangentX)
+					{
+						ValidTangentX = FixTangent(SoftSkinVertex.TangentX, SoftSkinVertex.TangentY, TangentZ, 0);
+					}
+					if (!ValidTangentY)
+					{
+						ValidTangentY = FixTangent(SoftSkinVertex.TangentY, TangentZ, SoftSkinVertex.TangentX, 1);
+					}
+					if (!ValidTangentZ)
+					{
+						ValidTangentZ = FixTangent(TangentZ, SoftSkinVertex.TangentX, SoftSkinVertex.TangentY, 2);
 					}
-					return false;
-				};
-
-				//The SoftSkinVertex TangentZ is a FVector4 so we must use a temporary FVector to be able to pass reference
-				FVector TangentZ = SoftSkinVertex.TangentZ;
-				//Make sure the tangent space is normalize before fixing bad tangent, because we want to do a cross product
-				//of 2 valid axis if possible. If not possible we will use the triangle normal.
-				bool ValidTangentX = NormalizedTangent(SoftSkinVertex.TangentX);
-				bool ValidTangentY = NormalizedTangent(SoftSkinVertex.TangentY);
-				bool ValidTangentZ = NormalizedTangent(TangentZ);
-
-				if (!ValidTangentX)
-				{
-					bFoundBadTangents = true;
-					ValidTangentX = FixTangent(SoftSkinVertex.TangentX, SoftSkinVertex.TangentY, TangentZ, 0);
-				}
-				if (!ValidTangentY)
-				{
-					bFoundBadTangents = true;
-					ValidTangentY = FixTangent(SoftSkinVertex.TangentY, TangentZ, SoftSkinVertex.TangentX, 1);
-				}
-				if (!ValidTangentZ)
-				{
-					bFoundBadTangents = true;
-					ValidTangentZ = FixTangent(TangentZ, SoftSkinVertex.TangentX, SoftSkinVertex.TangentY, 2);
-				}
 
-				//Make sure the result tangent space is orthonormal, only if we succeed to 
-				if (ValidTangentX && ValidTangentY && ValidTangentZ)
-				{
-					FVector::CreateOrthonormalBasis(
-						SoftSkinVertex.TangentX,
-						SoftSkinVertex.TangentY,
-						TangentZ
-					);
+					//Make sure the result tangent space is orthonormal, only if we succeed to fix all tangents
+					if (bNeedToOrthonormalize && ValidTangentX && ValidTangentY && ValidTangentZ)
+					{
+						FVector::CreateOrthonormalBasis(
+							SoftSkinVertex.TangentX,
+							SoftSkinVertex.TangentY,
+							TangentZ
+						);
+					}
+					SoftSkinVertex.TangentZ = TangentZ;
 				}
-				SoftSkinVertex.TangentZ = TangentZ;
 			}
 		}
 	}
 	if (bFoundBadTangents)
 	{
 		//Notify the user that we have to fix the normals on this model.
-		UE_ASSET_LOG(LogSkeletalMesh, Warning, this, TEXT("Find some bad tangent! please re-import this skeletal mesh asset to fix the issue. The shading of the skeletal mesh will be bad and faceted."));
+		UE_ASSET_LOG(LogSkeletalMesh, Display, this, TEXT("Find and fix some bad tangent! please re-import this skeletal mesh asset to fix the issue. The shading of the skeletal mesh will be bad and faceted."));
 	}
 }
 
