Fix threading crash in ComponentRequestsCPUAccess() for editor -game builds (async mesh compilation can be running, and calling GetRenderData() flushes FPerInstanceRenderData in dangerous ways)

diff --git a/Engine/Source/Runtime/Engine/Private/InstancedStaticMesh.cpp b/Engine/Source/Runtime/Engine/Private/InstancedStaticMesh.cpp
--- a/Engine/Source/Runtime/Engine/Private/InstancedStaticMesh.cpp
+++ b/Engine/Source/Runtime/Engine/Private/InstancedStaticMesh.cpp
@@ -2905,8 +2905,6 @@ static bool ComponentRequestsCPUAccess(UInstancedStaticMeshComponent* InComponen
 	// Check mesh distance fields
 	if (StaticMesh)
 	{
-		const FStaticMeshRenderData* RenderData = StaticMesh->GetRenderData();
-
 		if (FeatureLevel > ERHIFeatureLevel::ES3_1)
 		{
 			static const auto CVar = IConsoleManager::Get().FindTConsoleVariableDataInt(TEXT("r.GenerateMeshDistanceFields"));
@@ -2917,13 +2915,18 @@ static bool ComponentRequestsCPUAccess(UInstancedStaticMeshComponent* InComponen
 		}
 
 		// Check Nanite
-		if (RenderData && FeatureLevel >= ERHIFeatureLevel::SM5)
+		if (FeatureLevel >= ERHIFeatureLevel::SM5)
 		{
 			// TODO: Call UseNanite(GetScene()->GetShaderPlatform())?
 			//static const auto CVar = IConsoleManager::Get().FindTConsoleVariableDataInt(TEXT("r.Nanite"));
 
-			// Nanite needs lightmap/shadowmap UV biases and per instance random values accessible from GPUScene
+		#if WITH_EDITOR
+			const bool bHasNaniteData = StaticMesh->NaniteSettings.bEnabled;
+		#else
+			const FStaticMeshRenderData* RenderData = StaticMesh->GetRenderData();
 			const bool bHasNaniteData = RenderData->NaniteResources.PageStreamingStates.Num() > 0;
+		#endif
+
 			bNeedsCPUAccess |= bHasNaniteData;
 		}
 	}
