Make sure to unlock VT bulk data when validation fails

diff --git a/Engine/Source/Runtime/Engine/Private/VT/VirtualTextureBuiltData.cpp b/Engine/Source/Runtime/Engine/Private/VT/VirtualTextureBuiltData.cpp
--- a/Engine/Source/Runtime/Engine/Private/VT/VirtualTextureBuiltData.cpp
+++ b/Engine/Source/Runtime/Engine/Private/VT/VirtualTextureBuiltData.cpp
@@ -183,6 +183,19 @@ void FVirtualTextureBuiltData::Serialize(FArchive& Ar, UObject* Owner, int32 Fir
 	}
 }
 
+struct FBulkDataLockedScope
+{
+	~FBulkDataLockedScope()
+	{
+		if (BulkData)
+		{
+			BulkData->Unlock();
+		}
+	}
+
+	const FUntypedBulkData* BulkData = nullptr;
+};
+
 bool FVirtualTextureBuiltData::ValidateData(FStringView const& InDDCDebugContext, bool bValidateCompression) const
 {
 	TRACE_CPUPROFILER_EVENT_SCOPE(FVirtualTextureBuiltData::ValidateCompression);
@@ -197,15 +210,15 @@ bool FVirtualTextureBuiltData::ValidateData(FStringView const& InDDCDebugContext
 	for (int32 ChunkIndex = 0; bResult && ChunkIndex < Chunks.Num(); ++ChunkIndex)
 	{
 		const FVirtualTextureDataChunk& Chunk = Chunks[ChunkIndex];
+		FBulkDataLockedScope BulkDataLockedScope;
 
 		const uint8* ChunkData = nullptr;
 		uint32 ChunkDataSize = 0u;
-		bool bNeedToUnlockBulkData = false;
 		if (Chunk.BulkData.GetBulkDataSize() > 0)
 		{
 			ChunkDataSize = (uint32)Chunk.BulkData.GetBulkDataSize();
 			ChunkData = (uint8*)Chunk.BulkData.LockReadOnly();
-			bNeedToUnlockBulkData = true;
+			BulkDataLockedScope.BulkData = &Chunk.BulkData;
 		}
 #if WITH_EDITORONLY_DATA
 		else
@@ -277,11 +290,6 @@ bool FVirtualTextureBuiltData::ValidateData(FStringView const& InDDCDebugContext
 				}
 			}
 		}
-
-		if (bNeedToUnlockBulkData)
-		{
-			Chunk.BulkData.Unlock();
-		}
 	}
 
 	return bResult;
