Fix compile error

diff --git a/Engine/Source/Runtime/Engine/Private/VT/VirtualTextureBuiltData.cpp b/Engine/Source/Runtime/Engine/Private/VT/VirtualTextureBuiltData.cpp
--- a/Engine/Source/Runtime/Engine/Private/VT/VirtualTextureBuiltData.cpp
+++ b/Engine/Source/Runtime/Engine/Private/VT/VirtualTextureBuiltData.cpp
@@ -193,7 +193,7 @@ struct FBulkDataLockedScope
 		}
 	}
 
-	const FUntypedBulkData* BulkData = nullptr;
+	const FByteBulkData* BulkData = nullptr;
 };
 
 bool FVirtualTextureBuiltData::ValidateData(FStringView const& InDDCDebugContext, bool bValidateCompression) const
