Hitting ensure when opening an Additive Animation
#fix changed ensure condition to take into account the type of additive animation (ref-pose type doesn't need a valid animation asset)

diff --git a/Engine/Source/Runtime/Engine/Private/Animation/AnimSequence.cpp b/Engine/Source/Runtime/Engine/Private/Animation/AnimSequence.cpp
--- a/Engine/Source/Runtime/Engine/Private/Animation/AnimSequence.cpp
+++ b/Engine/Source/Runtime/Engine/Private/Animation/AnimSequence.cpp
@@ -2449,7 +2449,8 @@ void UAnimSequence::BakeOutAdditiveIntoRawData(TArray<FRawAnimSequenceTrack>& Ne
 	USkeleton* MySkeleton = GetSkeleton();
 	check(MySkeleton);
 
-	ensure(RefPoseSeq && !RefPoseSeq->HasAnyFlags(EObjectFlags::RF_NeedPostLoad));
+	const bool bRequiresValidRefPoseAsset = RefPoseType == EAdditiveBasePoseType::ABPT_AnimScaled || RefPoseType == EAdditiveBasePoseType::ABPT_AnimFrame;
+	ensure(!bRequiresValidRefPoseAsset || (RefPoseSeq && !RefPoseSeq->HasAnyFlags(EObjectFlags::RF_NeedPostLoad)));
 	
 	const FSmartNameMapping* CurveMapping = MySkeleton->GetSmartNameContainer(USkeleton::AnimCurveMappingName);
 	TArray<FName> SkeletonCurveNames;
