When renaming an external actor, make sure it's still external after changing it's name and/or outer.

diff --git a/Engine/Source/Runtime/Engine/Private/Actor.cpp b/Engine/Source/Runtime/Engine/Private/Actor.cpp
--- a/Engine/Source/Runtime/Engine/Private/Actor.cpp
+++ b/Engine/Source/Runtime/Engine/Private/Actor.cpp
@@ -1118,7 +1118,7 @@ bool AActor::Rename( const TCHAR* InName, UObject* NewOuter, ERenameFlags Flags
 	const bool bSuccess = Super::Rename( InName, NewOuter, Flags );
 
 #if WITH_EDITOR
-	if (!bRenameTest && bExternalActor && NewOuter)
+	if (!bRenameTest && bExternalActor)
 	{
 		if (ULevel* MyLevel = GetLevel())
 		{
