Added bPopulating guard value when populating the subsystem collection as a result of modules being loaded.

diff --git a/Engine/Source/Runtime/Engine/Private/Subsystems/SubsystemCollection.cpp b/Engine/Source/Runtime/Engine/Private/Subsystems/SubsystemCollection.cpp
--- a/Engine/Source/Runtime/Engine/Private/Subsystems/SubsystemCollection.cpp
+++ b/Engine/Source/Runtime/Engine/Private/Subsystems/SubsystemCollection.cpp
@@ -186,6 +186,8 @@ FString FSubsystemCollectionBase::GetReferencerName() const
 
 USubsystem* FSubsystemCollectionBase::AddAndInitializeSubsystem(UClass* SubsystemClass)
 {
+	TGuardValue<bool> PopulatingGuard(bPopulating, true);
+
 	if (!SubsystemMap.Contains(SubsystemClass))
 	{
 		// Only add instances for non abstract Subsystems
