Downgrade check to an ensure and make code more resilient until we fix all cases where this happens
  - SpawnActor currently has no means to notify components after properties have been overwritten by CDO

diff --git a/Engine/Source/Runtime/Engine/Classes/Components/StaticMeshComponent.h b/Engine/Source/Runtime/Engine/Classes/Components/StaticMeshComponent.h
--- a/Engine/Source/Runtime/Engine/Classes/Components/StaticMeshComponent.h
+++ b/Engine/Source/Runtime/Engine/Classes/Components/StaticMeshComponent.h
@@ -383,7 +383,14 @@ public:
 	UStaticMesh* GetStaticMesh() const 
 	{ 
 #if WITH_EDITOR
-		checkf(KnownStaticMesh == StaticMesh, TEXT("There is a missing call to NotifyIfStaticMeshChanged after StaticMesh has been modified"));
+		if (KnownStaticMesh != StaticMesh)
+		{
+			ensureMsgf(KnownStaticMesh == StaticMesh, TEXT("There is a missing call to NotifyIfStaticMeshChanged after StaticMesh has been overwritten"));
+
+			// This is a last resort, we should have catched the property overwrite well before we reach this code
+			UStaticMeshComponent* MutableThis = const_cast<UStaticMeshComponent*>(this);
+			MutableThis->NotifyIfStaticMeshChanged();
+		}
 #endif
 		return StaticMesh; 
 	}
