Fix missing notification after SpawnActor overwrites some properties from it's CDO.

diff --git a/Engine/Source/Runtime/Engine/Classes/Components/StaticMeshComponent.h b/Engine/Source/Runtime/Engine/Classes/Components/StaticMeshComponent.h
--- a/Engine/Source/Runtime/Engine/Classes/Components/StaticMeshComponent.h
+++ b/Engine/Source/Runtime/Engine/Classes/Components/StaticMeshComponent.h
@@ -436,6 +436,7 @@ public:
 	virtual bool IsCachedCookedPlatformDataLoaded(const ITargetPlatform* TargetPlatform) override;
 	virtual void PostDuplicate(bool bDuplicateForPIE) override;
 	virtual void InitializeComponent() override;
+	virtual void UpdateBounds() override;
 #endif // WITH_EDITOR
 #if WITH_EDITORONLY_DATA
 	virtual void PreSave(const class ITargetPlatform* TargetPlatform) override;
diff --git a/Engine/Source/Runtime/Engine/Private/Components/StaticMeshComponent.cpp b/Engine/Source/Runtime/Engine/Private/Components/StaticMeshComponent.cpp
--- a/Engine/Source/Runtime/Engine/Private/Components/StaticMeshComponent.cpp
+++ b/Engine/Source/Runtime/Engine/Private/Components/StaticMeshComponent.cpp
@@ -2544,6 +2544,16 @@ void UStaticMeshComponent::GetNavigationData(FNavigationRelevantData& Data) cons
 
 #if WITH_EDITOR
 
+// UpdateBounds is currently the closest place we can get from the StaticMesh property
+// being overwritten by a CDO construction from SpawnActor because PostInitProperties is 
+// not currently being called during subobjects property copies.
+void UStaticMeshComponent::UpdateBounds()
+{
+	NotifyIfStaticMeshChanged();
+
+	Super::UpdateBounds();
+}
+
 void UStaticMeshComponent::PostStaticMeshCompilation()
 {
 	CachePaintedDataIfNecessary();
