(Follow up CL, check for nullptr) Enabling ForceInline UsoundWave path for non-metasound rendering.

diff --git a/Engine/Source/Runtime/Engine/Private/AudioDecompress.cpp b/Engine/Source/Runtime/Engine/Private/AudioDecompress.cpp
--- a/Engine/Source/Runtime/Engine/Private/AudioDecompress.cpp
+++ b/Engine/Source/Runtime/Engine/Private/AudioDecompress.cpp
@@ -753,6 +753,11 @@ bool ShouldUseBackgroundPoolFor_FAsyncRealtimeAudioTask()
 
 bool ICompressedAudioInfo::StreamCompressedInfo(USoundWave* Wave, FSoundQualityInfo* QualityInfo)
 {
+	if (!Wave)
+	{
+		return false;
+	}
+
 	// Create and cache proxy object
 	Audio::FProxyDataInitParams Params;
 	Params.NameOfFeatureRequestingProxy = "ICompressedAudioInfo";
