Splitted WorldPartitionRuntimeSpatialHash to multiple files

diff --git a/Engine/Source/Runtime/Engine/Private/WorldPartition/RuntimeSpatialHash/RuntimeSpatialHashActorCluster.cpp b/Engine/Source/Runtime/Engine/Private/WorldPartition/RuntimeSpatialHash/RuntimeSpatialHashActorCluster.cpp
new file mode 100644
--- /dev/null
+++ b/Engine/Source/Runtime/Engine/Private/WorldPartition/RuntimeSpatialHash/RuntimeSpatialHashActorCluster.cpp
@@ -0,0 +1,187 @@
+// Copyright Epic Games, Inc. All Rights Reserved.
+
+#include "WorldPartition/RuntimeSpatialHash/RuntimeSpatialHashActorCluster.h"
+
+#if WITH_EDITOR
+
+#include "Algo/Transform.h"
+
+#include "Engine/World.h"
+#include "ActorReferencesUtils.h"
+#include "Engine/LevelScriptBlueprint.h"
+
+FActorCluster::FActorCluster(const FWorldPartitionActorDesc* InActorDesc, EActorGridPlacement InGridPlacement, UWorld* InWorld)
+	: GridPlacement(InGridPlacement)
+	, RuntimeGrid(InActorDesc->GetRuntimeGrid())
+	, Bounds(InActorDesc->GetBounds())
+{
+	check(GridPlacement != EActorGridPlacement::None);
+
+	Actors.Add(InActorDesc->GetGuid());
+	if (const AWorldDataLayers* WorldDataLayers = AWorldDataLayers::Get(InWorld))
+	{
+		for (const FName& DataLayerName : InActorDesc->GetDataLayers())
+		{
+			if (const UDataLayer* DataLayer = WorldDataLayers->GetDataLayerFromName(DataLayerName))
+			{
+				if (DataLayer->IsDynamicallyLoaded())
+				{
+					DataLayers.Add(DataLayer);
+				}
+			}
+		}
+	}
+
+	DataLayersID = FDataLayersID(DataLayers);
+}
+
+void FActorCluster::Add(const FActorCluster& InActorCluster)
+{
+	// Merge Actors
+	Actors.Append(InActorCluster.Actors);
+
+	// Merge RuntimeGrid
+	RuntimeGrid = RuntimeGrid == InActorCluster.RuntimeGrid ? RuntimeGrid : NAME_None;
+
+	// Merge Bounds
+	Bounds += InActorCluster.Bounds;
+
+	// Merge GridPlacement
+	// If currently None, will always stay None
+	if (GridPlacement != EActorGridPlacement::None)
+	{
+		// If grid placement differs between the two clusters
+		if (GridPlacement != InActorCluster.GridPlacement)
+		{
+			// If one of the two cluster was always loaded, set to None
+			if (InActorCluster.GridPlacement == EActorGridPlacement::AlwaysLoaded || GridPlacement == EActorGridPlacement::AlwaysLoaded)
+			{
+				GridPlacement = EActorGridPlacement::None;
+			}
+			else
+			{
+				GridPlacement = InActorCluster.GridPlacement;
+			}
+		}
+
+		// If current placement is set to Location, that won't make sense when merging two clusters. Set to Bounds
+		if (GridPlacement == EActorGridPlacement::Location)
+		{
+			GridPlacement = EActorGridPlacement::Bounds;
+		}
+	}
+
+	// Merge DataLayers
+	if (DataLayersID != InActorCluster.DataLayersID)
+	{
+		for (const UDataLayer* DataLayer : InActorCluster.DataLayers)
+		{
+			check(DataLayer->IsDynamicallyLoaded());
+			DataLayers.AddUnique(DataLayer);
+		}
+		DataLayersID = FDataLayersID(DataLayers);
+	}
+}
+
+
+void CreateActorCluster(const FWorldPartitionActorDesc* ActorDesc, EActorGridPlacement GridPlacement, TMap<FGuid, FActorCluster*>& ActorToActorCluster, TSet<FActorCluster*>& ActorClustersSet, UWorldPartition* WorldPartition)
+{
+	UWorld* World = WorldPartition->GetWorld();
+	const FGuid& ActorGuid = ActorDesc->GetGuid();
+
+	FActorCluster* ActorCluster = ActorToActorCluster.FindRef(ActorGuid);
+	if (!ActorCluster)
+	{
+		ActorCluster = new FActorCluster(ActorDesc, GridPlacement, World);
+		ActorClustersSet.Add(ActorCluster);
+		ActorToActorCluster.Add(ActorGuid, ActorCluster);
+	}
+
+	// Don't include references from editor-only actors
+	if (!ActorDesc->GetActorIsEditorOnly())
+	{
+		for (const FGuid& ReferenceGuid : ActorDesc->GetReferences())
+		{
+			const FWorldPartitionActorDesc* ReferenceActorDesc = WorldPartition->GetActorDesc(ReferenceGuid);
+
+			// Don't include references to editor-only actors
+			if (!ReferenceActorDesc->GetActorIsEditorOnly())
+			{
+				FActorCluster* ReferenceCluster = ActorToActorCluster.FindRef(ReferenceGuid);
+				if (ReferenceCluster)
+				{
+					if (ReferenceCluster != ActorCluster)
+					{
+						// Merge reference cluster in Actor's cluster
+						ActorCluster->Add(*ReferenceCluster);
+						for (const FGuid& ReferenceClusterActorGuid : ReferenceCluster->Actors)
+						{
+							ActorToActorCluster[ReferenceClusterActorGuid] = ActorCluster;
+						}
+						ActorClustersSet.Remove(ReferenceCluster);
+						delete ReferenceCluster;
+					}
+				}
+				else
+				{
+					// Put Reference in Actor's cluster
+					ActorCluster->Add(FActorCluster(ReferenceActorDesc, GridPlacement, World));
+				}
+
+				// Map its cluster
+				ActorToActorCluster.Add(ReferenceGuid, ActorCluster);
+			}
+		}
+	}
+}
+
+TArray<FActorCluster> CreateActorClustersImpl(UWorldPartition* WorldPartition, TOptional<FFilterPredicate> InFilterPredicate)
+{
+	TMap<FGuid, FActorCluster*> ActorToActorCluster;
+	TSet<FActorCluster*> ActorClustersSet;
+
+	// Gather all references to external actors from the level script
+	TSet<AActor*> LevelScriptExternalActorReferences;
+	if (ULevelScriptBlueprint* LevelScriptBlueprint = WorldPartition->GetWorld()->PersistentLevel->GetLevelScriptBlueprint(true))
+	{
+		LevelScriptExternalActorReferences.Append(ActorsReferencesUtils::GetExternalActorReferences(LevelScriptBlueprint));
+	}
+
+	for (const auto& Pair : WorldPartition->Actors)
+	{
+		const FWorldPartitionActorDesc* ActorDesc = Pair.Value->Get();
+		EActorGridPlacement GridPlacement = ActorDesc->GetGridPlacement();
+
+		// Check if the actor is loaded (potentially referenced by the level script)
+		if (AActor* Actor = ActorDesc->GetActor())
+		{
+			if (LevelScriptExternalActorReferences.Contains(Actor))
+			{
+				GridPlacement = EActorGridPlacement::AlwaysLoaded;
+			}
+		}
+
+		if (!InFilterPredicate.IsSet() || InFilterPredicate.GetValue()(*ActorDesc))
+		{
+			CreateActorCluster(ActorDesc, GridPlacement, ActorToActorCluster, ActorClustersSet, WorldPartition);
+		}
+	}
+
+	TArray<FActorCluster> ActorClusters;
+	ActorClusters.Reserve(ActorClustersSet.Num());
+	Algo::Transform(ActorClustersSet, ActorClusters, [](FActorCluster* ActorCluster) { return MoveTemp(*ActorCluster); });
+	for (FActorCluster* ActorCluster : ActorClustersSet) { delete ActorCluster; }
+	return ActorClusters;
+}
+
+TArray<FActorCluster> CreateActorClusters(UWorldPartition* WorldPartition, const FFilterPredicate& InFilterPredicate)
+{
+	return CreateActorClustersImpl(WorldPartition, InFilterPredicate);
+}
+
+TArray<FActorCluster> CreateActorClusters(UWorldPartition* WorldPartition)
+{
+	return CreateActorClustersImpl(WorldPartition, TOptional<FFilterPredicate>());
+}
+
+#endif // #if WITH_EDITOR
\ No newline at end of file
diff --git a/Engine/Source/Runtime/Engine/Private/WorldPartition/RuntimeSpatialHash/RuntimeSpatialHashActorCluster.h b/Engine/Source/Runtime/Engine/Private/WorldPartition/RuntimeSpatialHash/RuntimeSpatialHashActorCluster.h
new file mode 100644
--- /dev/null
+++ b/Engine/Source/Runtime/Engine/Private/WorldPartition/RuntimeSpatialHash/RuntimeSpatialHashActorCluster.h
@@ -0,0 +1,27 @@
+// Copyright Epic Games, Inc. All Rights Reserved.
+#pragma once
+
+#include "WorldPartition/WorldPartitionRuntimeSpatialHash.h"
+
+#if WITH_EDITOR
+
+// Clustering
+struct FActorCluster
+{
+	TSet<FGuid>					Actors;
+	EActorGridPlacement			GridPlacement;
+	FName						RuntimeGrid;
+	FBox						Bounds;
+	TArray<const UDataLayer*>	DataLayers;
+	FDataLayersID				DataLayersID;
+
+	FActorCluster(const FWorldPartitionActorDesc* InActorDesc, EActorGridPlacement InGridPlacement, UWorld* InWorld);
+	void Add(const FActorCluster& InActorCluster);
+};
+
+typedef TFunctionRef<bool(const FWorldPartitionActorDesc&)> FFilterPredicate;
+
+TArray<FActorCluster> CreateActorClusters(UWorldPartition* WorldPartition, const FFilterPredicate& InFilterPredicate);
+TArray<FActorCluster> CreateActorClusters(UWorldPartition* WorldPartition);
+
+#endif // #if WITH_EDITOR
diff --git a/Engine/Source/Runtime/Engine/Private/WorldPartition/RuntimeSpatialHash/RuntimeSpatialHashGridHelper.cpp b/Engine/Source/Runtime/Engine/Private/WorldPartition/RuntimeSpatialHash/RuntimeSpatialHashGridHelper.cpp
new file mode 100644
--- /dev/null
+++ b/Engine/Source/Runtime/Engine/Private/WorldPartition/RuntimeSpatialHash/RuntimeSpatialHashGridHelper.cpp
@@ -0,0 +1,244 @@
+// Copyright Epic Games, Inc. All Rights Reserved.
+
+#include "WorldPartition/RuntimeSpatialHash/RuntimeSpatialHashGridHelper.h"
+
+#include "ProfilingDebugging/ScopedTimers.h"
+
+FSquare2DGridHelper::FSquare2DGridHelper(int32 InNumLevels, const FVector& InOrigin, int32 InCellSize, int32 InGridSize)
+	: Origin(InOrigin)
+	, CellSize(InCellSize)
+	, GridSize(InGridSize)
+{
+	Levels.Reserve(InNumLevels);
+
+	int32 CurrentCellSize = CellSize;
+	int32 CurrentGridSize = GridSize;
+
+	const FVector2D BaseLevelOffset = 0.5 * FVector2D(CellSize, CellSize);
+	for (int32 Level = 0; Level < InNumLevels; ++Level)
+	{
+		// Add offset on origin based on level's cell size to break pattern of perfectly aligned cell edges at multiple level.
+		// This will prevent weird artefact during actor promotion.
+		// Apply base level offset so that first level isn't offset.
+		const FVector2D GridLevelOffset = 0.5f * FVector2D(CurrentCellSize, CurrentCellSize) - BaseLevelOffset;
+		const FVector2D LevelOrigin = FVector2D(InOrigin) + GridLevelOffset;
+
+		Levels.Emplace(LevelOrigin, CurrentCellSize, CurrentGridSize);
+
+		CurrentCellSize <<= 1;
+		CurrentGridSize >>= 1;
+	}
+}
+
+void FSquare2DGridHelper::ForEachCells(TFunctionRef<void(const FIntVector&)> InOperation) const
+{
+	for (int32 Level = 0; Level < Levels.Num(); Level++)
+	{
+		Levels[Level].ForEachCells([Level, &InOperation](const FIntVector2& Coord) { InOperation(FIntVector(Coord.X, Coord.Y, Level)); });
+	}
+}
+
+int32 FSquare2DGridHelper::ForEachIntersectingCells(const FBox& InBox, TFunctionRef<void(const FIntVector&)> InOperation) const
+{
+	int32 NumCells = 0;
+
+	for (int32 Level = 0; Level < Levels.Num(); Level++)
+	{
+		NumCells += Levels[Level].ForEachIntersectingCells(InBox, [InBox, Level, InOperation](const FIntVector2& Coord) { InOperation(FIntVector(Coord.X, Coord.Y, Level)); });
+	}
+
+	return NumCells;
+}
+
+int32 FSquare2DGridHelper::ForEachIntersectingCells(const FSphere& InSphere, TFunctionRef<void(const FIntVector&)> InOperation) const
+{
+	int32 NumCells = 0;
+
+	for (int32 Level = 0; Level < Levels.Num(); Level++)
+	{
+		NumCells += Levels[Level].ForEachIntersectingCells(InSphere, [InSphere, Level, InOperation](const FIntVector2& Coord) { InOperation(FIntVector(Coord.X, Coord.Y, Level)); });
+	}
+
+	return NumCells;
+}
+
+#if WITH_EDITOR
+void FSquare2DGridHelper::ValidateSingleActorReferer()
+{
+	UE_SCOPED_TIMER(TEXT("ValidateSingleActorReferer"), LogWorldPartitionRuntimeSpatialHash);
+
+	TSet<FGuid> ActorUsage;
+	for (int32 Level = 0; Level < Levels.Num() - 1; Level++)
+	{
+		const int32 CurrentGridSize = Levels[Level].GridSize;
+		for (int32 y = 0; y < CurrentGridSize; y++)
+		{
+			for (int32 x = 0; x < CurrentGridSize; x++)
+			{
+				FGridLevel::FGridCell& ThisCell = Levels[Level].Cells[y * CurrentGridSize + x];
+				for (const FGridLevel::FGridCellDataChunk& DataChunk : ThisCell.GetDataChunks())
+				{
+					for (const FGuid& ActorGuid : DataChunk.GetActors())
+					{
+						bool bWasAlreadyInSet;
+						ActorUsage.Add(ActorGuid, &bWasAlreadyInSet);
+						check(!bWasAlreadyInSet);
+					}
+				}
+			}
+		}
+	}
+}
+#endif // #if WITH_EDITOR
+
+#if WITH_EDITOR
+
+FSquare2DGridHelper GetGridHelper(const FBox& WorldBounds, int32 GridCellSize)
+{
+	// Default grid to a minimum of 1 level and 1 cell, for always loaded actors
+	FVector GridOrigin = FVector::ZeroVector;
+	int32 GridSize = 1;
+	int32 GridLevelCount = 1;
+
+	float WorldBoundsMaxExtent = 0.f;
+	// If World bounds is valid, compute Grid's size and level count based on it
+	if (WorldBounds.IsValid)
+	{
+		const FVector2D DistMin = FMath::Abs(FVector2D(WorldBounds.Min - GridOrigin));
+		const FVector2D DistMax = FMath::Abs(FVector2D(WorldBounds.Max - GridOrigin));
+		WorldBoundsMaxExtent = FMath::Max(DistMin.GetMax(), DistMax.GetMax());
+	}
+	if (WorldBoundsMaxExtent > 0.f)
+	{
+		GridSize = 2.f * FMath::CeilToFloat(WorldBoundsMaxExtent / GridCellSize);
+		if (!FMath::IsPowerOfTwo(GridSize))
+		{
+			GridSize = FMath::Pow(2.f, FMath::CeilToFloat(FMath::Log2(static_cast<float>(GridSize))));
+		}
+		GridLevelCount = FMath::Log2(static_cast<float>(GridSize)) + 1;
+	}
+	else
+	{
+		UE_LOG(LogWorldPartitionRuntimeSpatialHash, Warning, TEXT("Invalid world bounds, grid partitioning will use a runtime grid with 1 cell."));
+	}
+
+	return FSquare2DGridHelper(GridLevelCount, GridOrigin, GridCellSize, GridSize);
+}
+
+FSquare2DGridHelper GetPartitionedActors(const UWorldPartition* WorldPartition, const FBox& WorldBounds, const FSpatialHashRuntimeGrid& Grid, const TArray<FActorCluster>& GridActors)
+{
+	UE_SCOPED_TIMER(TEXT("GetPartitionedActors"), LogWorldPartitionRuntimeSpatialHash);
+
+	//
+	// Create the hierarchical grids for the game
+	//	
+	FSquare2DGridHelper PartitionedActors = GetGridHelper(WorldBounds, Grid.CellSize);
+
+	for (const FActorCluster& ActorCluster : GridActors)
+	{
+		check(ActorCluster.Actors.Num() > 0);
+
+		EActorGridPlacement GridPlacement = ActorCluster.GridPlacement;
+		bool bAlwaysLoadedPromotedCluster = (GridPlacement == EActorGridPlacement::None);
+		bool bAlwaysLoadedPromotedOutOfGrid = false;
+
+		if (bAlwaysLoadedPromotedCluster)
+		{
+			GridPlacement = EActorGridPlacement::AlwaysLoaded;
+		}
+
+		switch (GridPlacement)
+		{
+		case EActorGridPlacement::Location:
+		{
+			FIntVector2 CellCoords;
+			check(ActorCluster.Actors.Num() == 1);
+			const FGuid& ActorGuid = *ActorCluster.Actors.CreateConstIterator();
+			const FWorldPartitionActorDesc& ActorDesc = *WorldPartition->GetActorDesc(ActorGuid);
+			if (PartitionedActors.GetLowestLevel().GetCellCoords(FVector2D(ActorDesc.GetOrigin()), CellCoords))
+			{
+				TArray<const UDataLayer*> ActorDataLayers;
+				if (const AWorldDataLayers* WorldDataLayers = AWorldDataLayers::Get(WorldPartition->GetWorld()))
+				{
+					for (const FName& DataLayerName : ActorDesc.GetDataLayers())
+					{
+						if (const UDataLayer* DataLayer = WorldDataLayers->GetDataLayerFromName(DataLayerName))
+						{
+							ActorDataLayers.Add(DataLayer);
+						}
+					}
+				}
+				PartitionedActors.GetLowestLevel().GetCell(CellCoords).AddActor(ActorGuid, ActorDataLayers);
+			}
+			else
+			{
+				GridPlacement = EActorGridPlacement::AlwaysLoaded;
+				bAlwaysLoadedPromotedOutOfGrid = true;
+			}
+			break;
+		}
+		case EActorGridPlacement::Bounds:
+		{
+			// Find grid level cell that encompasses the actor cluster and put actors in it.
+			bool bFoundCell = false;
+			for (FSquare2DGridHelper::FGridLevel& GridLevel : PartitionedActors.Levels)
+			{
+				int32 IntersectingCellCount = 0;
+				GridLevel.ForEachIntersectingCellsBreakable(ActorCluster.Bounds, [&IntersectingCellCount](const FIntVector2& Coords) { return ++IntersectingCellCount <= 1; });
+				if (IntersectingCellCount == 1)
+				{
+					GridLevel.ForEachIntersectingCells(ActorCluster.Bounds, [&GridLevel, &ActorCluster](const FIntVector2& Coords)
+						{
+							GridLevel.GetCell(Coords).AddActors(ActorCluster.Actors, ActorCluster.DataLayers);
+						});
+					bFoundCell = true;
+					break;
+				}
+			}
+			if (!bFoundCell)
+			{
+				GridPlacement = EActorGridPlacement::AlwaysLoaded;
+				bAlwaysLoadedPromotedOutOfGrid = true;
+			}
+			break;
+		}
+		case EActorGridPlacement::AlwaysLoaded:
+		{
+			PartitionedActors.GetAlwaysLoadedCell().AddActors(ActorCluster.Actors, ActorCluster.DataLayers);
+			break;
+		}
+		default:
+			check(0);
+		}
+
+		if (UE_LOG_ACTIVE(LogWorldPartitionRuntimeSpatialHash, Verbose))
+		{
+			if (ActorCluster.Actors.Num() > 1)
+			{
+				static UEnum* ActorGridPlacementEnum = FindObject<UEnum>(ANY_PACKAGE, TEXT("EActorGridPlacement"));
+
+				UE_LOG(LogWorldPartitionRuntimeSpatialHash, Verbose, TEXT("Clustered %d actors (%s%s%s), generated shared BV of [%d x %d] (meters)"),
+					ActorCluster.Actors.Num(),
+					*ActorGridPlacementEnum->GetNameStringByValue((int64)GridPlacement),
+					bAlwaysLoadedPromotedCluster ? TEXT(":PromotedCluster") : TEXT(""),
+					bAlwaysLoadedPromotedOutOfGrid ? TEXT(":PromotedOutOfGrid") : TEXT(""),
+					(int)(0.01f * ActorCluster.Bounds.GetSize().X),
+					(int)(0.01f * ActorCluster.Bounds.GetSize().Y));
+
+				for (const FGuid& ActorGuid : ActorCluster.Actors)
+				{
+					const FWorldPartitionActorDesc& Desc = *WorldPartition->GetActorDesc(ActorGuid);
+					UE_LOG(LogWorldPartitionRuntimeSpatialHash, Verbose, TEXT("   - Actor: %s (%s)"), *Desc.GetActorPath().ToString(), *ActorGuid.ToString(EGuidFormats::UniqueObjectGuid));
+					UE_LOG(LogWorldPartitionRuntimeSpatialHash, Verbose, TEXT("            %s"), *Desc.GetActorPackage().ToString());
+				}
+			}
+		}
+	}
+
+	// Perform validation
+	PartitionedActors.ValidateSingleActorReferer();
+
+	return PartitionedActors;
+}
+
+#endif // #if WITH_EDITOR
diff --git a/Engine/Source/Runtime/Engine/Private/WorldPartition/RuntimeSpatialHash/RuntimeSpatialHashGridHelper.h b/Engine/Source/Runtime/Engine/Private/WorldPartition/RuntimeSpatialHash/RuntimeSpatialHashGridHelper.h
new file mode 100644
--- /dev/null
+++ b/Engine/Source/Runtime/Engine/Private/WorldPartition/RuntimeSpatialHash/RuntimeSpatialHashGridHelper.h
@@ -0,0 +1,408 @@
+// Copyright Epic Games, Inc. All Rights Reserved.
+#pragma once
+
+#include "Algo/Find.h"
+#include "Algo/Transform.h"
+
+#include "ProfilingDebugging/ScopedTimers.h"
+
+#include "WorldPartition/WorldPartitionRuntimeSpatialHash.h"
+#include "WorldPartition/RuntimeSpatialHash/RuntimeSpatialHashActorCluster.h"
+
+/**
+  * Square 2D grid helper
+  */
+struct FSquare2DGridHelper
+{
+	struct FGrid2D
+	{
+		FVector2D Origin;
+		int32 CellSize;
+		int32 GridSize;
+
+		inline FGrid2D(const FVector2D& InOrigin, int32 InCellSize, int32 InGridSize)
+			: Origin(InOrigin)
+			, CellSize(InCellSize)
+			, GridSize(InGridSize)
+		{}
+
+		/**
+		 * Validate that the coordinates fit the grid size
+		 *
+		 * @return true if the specified coordinates are valid
+		 */
+		inline bool IsValidCoords(const FIntVector2& InCoords) const
+		{
+			return (InCoords.X >= 0) && (InCoords.X < GridSize) && (InCoords.Y >= 0) && (InCoords.Y < GridSize);
+		}
+
+		/**
+		 * Returns the cell bounds
+		 *
+		 * @return true if the specified index was valid
+		 */
+		inline bool GetCellBounds(int32 InIndex, FBox2D& OutBounds) const
+		{
+			if (InIndex >= 0 && InIndex <= (GridSize * GridSize))
+			{
+				const FIntVector2 Coords(InIndex % GridSize, InIndex / GridSize);
+				return GetCellBounds(Coords, OutBounds);
+			}
+
+			return false;
+		}
+
+		/**
+		 * Returns the cell bounds
+		 *
+		 * @return true if the specified coord was valid
+		 */
+		inline bool GetCellBounds(const FIntVector2& InCoords, FBox2D& OutBounds) const
+		{
+			if (IsValidCoords(InCoords))
+			{
+				const FVector2D Min = (FVector2D(Origin) - FVector2D(GridSize * CellSize * 0.5f, GridSize * CellSize * 0.5f)) + FVector2D(InCoords.X * CellSize, InCoords.Y * CellSize);
+				const FVector2D Max = Min + FVector2D(CellSize, CellSize);
+				OutBounds = FBox2D(Min, Max);
+				return true;
+			}
+
+			return false;
+		}
+
+		/**
+		 * Returns the cell coordinates of the provided position
+		 *
+		 * @return true if the position was inside the grid
+		 */
+		inline bool GetCellCoords(const FVector2D& InPos, FIntVector2& OutCoords) const
+		{
+			OutCoords = FIntVector2(
+				FMath::FloorToInt(((InPos.X - Origin.X) / CellSize) + GridSize * 0.5f),
+				FMath::FloorToInt(((InPos.Y - Origin.Y) / CellSize) + GridSize * 0.5f)
+			);
+
+			return IsValidCoords(OutCoords);
+		}
+
+		/**
+		 * Returns the cells coordinates of the provided box
+		 *
+		 * @return true if the bounds was intersecting with the grid
+		 */
+		inline bool GetCellCoords(const FBox2D& InBounds2D, FIntVector2& OutMinCellCoords, FIntVector2& OutMaxCellCoords) const
+		{
+			GetCellCoords(InBounds2D.Min, OutMinCellCoords);
+			if ((OutMinCellCoords.X >= GridSize) || (OutMinCellCoords.Y >= GridSize))
+			{
+				return false;
+			}
+
+			GetCellCoords(InBounds2D.Max, OutMaxCellCoords);
+			if ((OutMaxCellCoords.X < 0) || (OutMaxCellCoords.Y < 0))
+			{
+				return false;
+			}
+
+			OutMinCellCoords.X = FMath::Clamp(OutMinCellCoords.X, 0, GridSize - 1);
+			OutMinCellCoords.Y = FMath::Clamp(OutMinCellCoords.Y, 0, GridSize - 1);
+			OutMaxCellCoords.X = FMath::Clamp(OutMaxCellCoords.X, 0, GridSize - 1);
+			OutMaxCellCoords.Y = FMath::Clamp(OutMaxCellCoords.Y, 0, GridSize - 1);
+
+			return true;
+		}
+
+		/**
+		 * Returns the cell index of the provided coords
+		 *
+		 * @return true if the coords was inside the grid
+		 */
+		inline bool GetCellIndex(const FIntVector2& InCoords, uint32& OutIndex) const
+		{
+			if (IsValidCoords(InCoords))
+			{
+				OutIndex = (InCoords.Y * GridSize) + InCoords.X;
+				return true;
+			}
+
+			return false;
+		}
+
+		/**
+		 * Returns the cell index of the provided position
+		 *
+		 * @return true if the position was inside the grid
+		 */
+		inline bool GetCellIndex(const FVector& InPos, uint32& OutIndex) const
+		{
+			FIntVector2 Coords = FIntVector2(
+				FMath::FloorToInt(((InPos.X - Origin.X) / CellSize) + GridSize * 0.5f),
+				FMath::FloorToInt(((InPos.Y - Origin.Y) / CellSize) + GridSize * 0.5f)
+			);
+
+			return GetCellIndex(Coords, OutIndex);
+		}
+
+		/**
+		 * Get the number of intersecting cells of the provided box
+		 *
+		 * @return the number of intersecting cells
+		 */
+		int32 GetNumIntersectingCells(const FBox& InBox) const
+		{
+			FIntVector2 MinCellCoords;
+			FIntVector2 MaxCellCoords;
+			const FBox2D Bounds2D(FVector2D(InBox.Min), FVector2D(InBox.Max));
+
+			if (GetCellCoords(Bounds2D, MinCellCoords, MaxCellCoords))
+			{
+				return (MaxCellCoords.X - MinCellCoords.X + 1) * (MaxCellCoords.Y - MinCellCoords.Y + 1);
+			}
+
+			return 0;
+		}
+
+		// Runs a function on all cells
+		void ForEachCells(TFunctionRef<void(const FIntVector2&)> InOperation) const
+		{
+			for (int32 y = 0; y < GridSize; y++)
+			{
+				for (int32 x = 0; x < GridSize; x++)
+				{
+					InOperation(FIntVector2(x, y));
+				}
+			}
+		}
+
+		/**
+		 * Runs a function on all intersecting cells for the provided box
+		 *
+		 * @return the number of intersecting cells
+		 */
+		int32 ForEachIntersectingCellsBreakable(const FBox& InBox, TFunctionRef<bool(const FIntVector2&)> InOperation) const
+		{
+			int32 NumCells = 0;
+
+			FIntVector2 MinCellCoords;
+			FIntVector2 MaxCellCoords;
+			const FBox2D Bounds2D(FVector2D(InBox.Min), FVector2D(InBox.Max));
+
+			if (GetCellCoords(Bounds2D, MinCellCoords, MaxCellCoords))
+			{
+				for (int32 y = MinCellCoords.Y; y <= MaxCellCoords.Y; y++)
+				{
+					for (int32 x = MinCellCoords.X; x <= MaxCellCoords.X; x++)
+					{
+						if (!InOperation(FIntVector2(x, y)))
+						{
+							return NumCells;
+						}
+						++NumCells;
+					}
+				}
+			}
+
+			return NumCells;
+		}
+
+		int32 ForEachIntersectingCells(const FBox& InBox, TFunctionRef<void(const FIntVector2&)> InOperation) const
+		{
+			return ForEachIntersectingCellsBreakable(InBox, [InOperation](const FIntVector2& Vector) { InOperation(Vector); return true; });
+		}
+
+		/**
+		 * Runs a function on all intersecting cells for the provided sphere
+		 *
+		 * @return the number of intersecting cells
+		 */
+		int32 ForEachIntersectingCells(const FSphere& InSphere, TFunctionRef<void(const FIntVector2&)> InOperation) const
+		{
+			int32 NumCells = 0;
+
+			// @todo_ow: rasterize circle instead?
+			const FBox Box(InSphere.Center - FVector(InSphere.W), InSphere.Center + FVector(InSphere.W));
+
+			ForEachIntersectingCells(Box, [this, &InSphere, &InOperation, &NumCells](const FIntVector2& Coords)
+				{
+					const int32 CellIndex = Coords.Y * GridSize + Coords.X;
+
+					FBox2D CellBounds;
+					GetCellBounds(CellIndex, CellBounds);
+
+					FVector2D Delta = FVector2D(InSphere.Center) - FVector2D::Max(CellBounds.GetCenter() - CellBounds.GetExtent(), FVector2D::Min(FVector2D(InSphere.Center), CellBounds.GetCenter() + CellBounds.GetExtent()));
+					if ((Delta.X * Delta.X + Delta.Y * Delta.Y) < (InSphere.W * InSphere.W))
+					{
+						InOperation(Coords);
+						NumCells++;
+					}
+				});
+
+			return NumCells;
+		}
+	};
+
+	struct FGridLevel : public FGrid2D
+	{
+#if WITH_EDITOR
+		struct FGridCellDataChunk
+		{
+			FGridCellDataChunk(const TArray<const UDataLayer*>& InDataLayers)
+			{
+				Algo::TransformIf(InDataLayers, DataLayers, [](const UDataLayer* DataLayer) { return DataLayer->IsDynamicallyLoaded(); }, [](const UDataLayer* DataLayer) { return DataLayer; });
+				DataLayersID = FDataLayersID(DataLayers);
+			}
+
+			void AddActor(const FGuid& Actor) { Actors.Add(Actor); }
+			const TSet<FGuid>& GetActors() const { return Actors; }
+			bool HasDataLayers() const { return !DataLayers.IsEmpty(); }
+			const TArray<const UDataLayer*>& GetDataLayers() const { return DataLayers; }
+			const FDataLayersID& GetDataLayersID() const { return DataLayersID; }
+
+		private:
+			TSet<FGuid> Actors;
+			TArray<const UDataLayer*> DataLayers;
+			FDataLayersID DataLayersID;
+		};
+
+		struct FGridCell
+		{
+			void AddActor(const FGuid& InActor, const TArray<const UDataLayer*>& InDataLayers)
+			{
+				FDataLayersID DataLayersID = FDataLayersID(InDataLayers);
+				FGridCellDataChunk* ActorDataChunk = Algo::FindByPredicate(DataChunks, [&](FGridCellDataChunk& InDataChunk) { return InDataChunk.GetDataLayersID() == DataLayersID; });
+				if (!ActorDataChunk)
+				{
+					ActorDataChunk = &DataChunks.Emplace_GetRef(InDataLayers);
+				}
+				ActorDataChunk->AddActor(InActor);
+			}
+
+			void AddActors(const TSet<FGuid>& InActors, const TArray<const UDataLayer*>& InDataLayers)
+			{
+				for (const FGuid& Actor : InActors)
+				{
+					AddActor(Actor, InDataLayers);
+				}
+			}
+
+			const TArray<FGridCellDataChunk>& GetDataChunks() const { return DataChunks; }
+
+			const FGridCellDataChunk* GetNoDataLayersDataChunk() const
+			{
+				for (const FGridCellDataChunk& DataChunk : DataChunks)
+				{
+					if (!DataChunk.HasDataLayers())
+					{
+						return &DataChunk;
+					}
+				}
+				return nullptr;
+			}
+
+		private:
+			TArray<FGridCellDataChunk> DataChunks;
+		};
+#endif
+
+		inline FGridLevel(const FVector2D& InOrigin, int32 InCellSize, int32 InGridSize)
+			: FGrid2D(InOrigin, InCellSize, InGridSize)
+		{
+#if WITH_EDITOR
+			Cells.InsertDefaulted(0, GridSize * GridSize);
+#endif
+		}
+
+#if WITH_EDITOR
+		/**
+		 * Returns the cell at the specified grid coordinate
+		 *
+		 * @return the cell at the specified grid coordinate
+		 */
+		inline FGridCell& GetCell(const FIntVector2& InCoords)
+		{
+			check(IsValidCoords(InCoords));
+			uint32 CellIndex;
+			GetCellIndex(InCoords, CellIndex);
+			return Cells[CellIndex];
+		}
+
+		/**
+		 * Returns the cell at the specified grid coordinate
+		 *
+		 * @return the cell at the specified grid coordinate
+		 */
+		inline const FGridCell& GetCell(const FIntVector2& InCoords) const
+		{
+			check(IsValidCoords(InCoords));
+			uint32 CellIndex;
+			GetCellIndex(InCoords, CellIndex);
+			return Cells[CellIndex];
+		}
+
+		TArray<FGridCell> Cells;
+#endif
+	};
+
+	FSquare2DGridHelper(int32 InNumLevels, const FVector& InOrigin, int32 InCellSize, int32 InGridSize);
+
+#if WITH_EDITOR
+	// Returns the lowest grid level
+	inline FGridLevel& GetLowestLevel() { return Levels[0]; }
+
+	// Returns the always loaded (top level) cell
+	inline FGridLevel::FGridCell& GetAlwaysLoadedCell() { return Levels.Last().Cells[0]; }
+
+	// Returns the always loaded (top level) cell
+	inline const FGridLevel::FGridCell& GetAlwaysLoadedCell() const { return Levels.Last().Cells[0]; }
+
+	// Returns the cell at the given coord
+	inline const FGridLevel::FGridCell& GetCell(const FIntVector& InCoords) const { return Levels[InCoords.Z].GetCell(FIntVector2(InCoords.X, InCoords.Y)); }
+#endif
+
+	/**
+	 * Returns the cell bounds
+	 *
+	 * @return true if the specified coord was valid
+	 */
+	inline bool GetCellBounds(const FIntVector& InCoords, FBox2D& OutBounds) const
+	{
+		if (Levels.IsValidIndex(InCoords.Z))
+		{
+			return Levels[InCoords.Z].GetCellBounds(FIntVector2(InCoords.X, InCoords.Y), OutBounds);
+		}
+		return false;
+	}
+
+	// Runs a function on all cells
+	void ForEachCells(TFunctionRef<void(const FIntVector&)> InOperation) const;
+
+	/**
+	 * Runs a function on all intersecting cells for the provided box
+	 *
+	 * @return the number of intersecting cells
+	 */
+	int32 ForEachIntersectingCells(const FBox& InBox, TFunctionRef<void(const FIntVector&)> InOperation) const;
+
+	/**
+	 * Runs a function on all intersecting cells for the provided sphere
+	 *
+	 * @return the number of intersecting cells
+	 */
+	int32 ForEachIntersectingCells(const FSphere& InSphere, TFunctionRef<void(const FIntVector&)> InOperation) const;
+
+	// Validates that actor is not referenced by multiple cells
+#if WITH_EDITOR
+	void ValidateSingleActorReferer();
+#endif
+
+public:
+	FVector Origin;
+	int32 CellSize;
+	int32 GridSize;
+	TArray<FGridLevel> Levels;
+};
+
+#if WITH_EDITOR
+FSquare2DGridHelper GetGridHelper(const FBox& WorldBounds, int32 GridCellSize);
+FSquare2DGridHelper GetPartitionedActors(const UWorldPartition* WorldPartition, const FBox& WorldBounds, const FSpatialHashRuntimeGrid& Grid, const TArray<FActorCluster>& GridActors);
+#endif // #if WITH_EDITOR
diff --git a/Engine/Source/Runtime/Engine/Private/WorldPartition/RuntimeSpatialHash/RuntimeSpatialHashHLOD.cpp b/Engine/Source/Runtime/Engine/Private/WorldPartition/RuntimeSpatialHash/RuntimeSpatialHashHLOD.cpp
new file mode 100644
--- /dev/null
+++ b/Engine/Source/Runtime/Engine/Private/WorldPartition/RuntimeSpatialHash/RuntimeSpatialHashHLOD.cpp
@@ -0,0 +1,559 @@
+// Copyright Epic Games, Inc. All Rights Reserved.
+
+#include "WorldPartition/WorldPartitionRuntimeSpatialHash.h"
+
+#if WITH_EDITOR
+
+#include "WorldPartition/RuntimeSpatialHash/RuntimeSpatialHashGridHelper.h"
+#include "WorldPartition/WorldPartitionActorDescIterator.h"
+#include "WorldPartition/HLOD/HLODBuilder.h"
+#include "WorldPartition/HLOD/HLODLayer.h"
+#include "WorldPartition/HLOD/HLODActor.h"
+#include "WorldPartition/HLOD/HLODActorDesc.h"
+
+#include "HAL/PlatformFileManager.h"
+#include "Misc/PackageName.h"
+#include "Misc/ScopedSlowTask.h"
+
+#include "Algo/ForEach.h"
+
+#include "EngineUtils.h"
+
+#include "DistanceFieldAtlas.h"
+#include "MeshCardRepresentation.h"
+#include "StaticMeshCompiler.h"
+
+#include "AssetRegistryModule.h"
+#include "AssetData.h"
+
+#include "IDirectoryWatcher.h"
+#include "DirectoryWatcherModule.h"
+
+
+DEFINE_LOG_CATEGORY_STATIC(LogWorldPartitionRuntimeSpatialHashHLOD, Log, All);
+
+
+static void CheckoutAndSavePackage(UPackage* Package, ISourceControlHelper* SourceControlHelper)
+{
+	// Checkout package
+	Package->MarkAsFullyLoaded();
+
+	if (SourceControlHelper && !SourceControlHelper->Checkout(Package))
+	{
+		UE_LOG(LogWorldPartitionRuntimeSpatialHashHLOD, Error, TEXT("Error checking out package %s."), *Package->GetName());
+		check(0);
+	}
+
+	// Save package
+	FPackagePath PackagePath = FPackagePath::FromPackageNameChecked(Package->GetName());
+	const FString PackageFileName = PackagePath.GetLocalFullPath();
+	if (!UPackage::SavePackage(Package, nullptr, RF_Standalone, *PackageFileName))
+	{
+		UE_LOG(LogWorldPartitionRuntimeSpatialHashHLOD, Error, TEXT("Error saving package %s."), *Package->GetName());
+		check(0);
+	}
+
+	// Add new package to source control
+	if (SourceControlHelper && !SourceControlHelper->Add(Package))
+	{
+		UE_LOG(LogWorldPartitionRuntimeSpatialHashHLOD, Error, TEXT("Error adding package %s to source control."), *Package->GetName());
+		check(0);
+	}
+}
+
+static void DeletePackage(const FString& PackageName, ISourceControlHelper* SourceControlHelper)
+{
+	FPackagePath PackagePath = FPackagePath::FromPackageNameChecked(PackageName);
+	const FString PackageFileName = PackagePath.GetLocalFullPath();
+
+	if (SourceControlHelper)
+	{
+		SourceControlHelper->Delete(PackageFileName);
+	}
+	else
+	{
+		FPlatformFileManager::Get().GetPlatformFile().DeleteFile(*PackageFileName);
+	}
+}
+
+static void DeletePackage(UPackage* Package, ISourceControlHelper* SourceControlHelper)
+{
+	DeletePackage(Package->GetName(), SourceControlHelper);
+}
+
+static TArray<FGuid> GenerateHLODsForGrid(UWorldPartition* WorldPartition, const FSpatialHashRuntimeGrid& RuntimeGrid, uint32 HLODLevel, FHLODGenerationContext& Context, ISourceControlHelper* SourceControlHelper, const TArray<FActorCluster>& ActorClusters)
+{
+	auto HasExceededMaxMemory = []()
+	{
+		const uint64 MemoryMinFreePhysical = 1024ll * 1024 * 1024;
+		const uint64 MemoryMaxUsedPhysical = 16384ll * 1024 * 1024l;
+		const FPlatformMemoryStats MemStats = FPlatformMemory::GetStats();
+		return (MemStats.AvailablePhysical < MemoryMinFreePhysical) || (MemStats.UsedPhysical >= MemoryMaxUsedPhysical);
+	};
+
+	auto DoCollectGarbage = []()
+	{
+		if (GDistanceFieldAsyncQueue)
+		{
+			GDistanceFieldAsyncQueue->BlockUntilAllBuildsComplete();
+		}
+
+		if (GCardRepresentationAsyncQueue)
+		{
+			GCardRepresentationAsyncQueue->BlockUntilAllBuildsComplete();
+		}
+
+		const FPlatformMemoryStats MemStatsBefore = FPlatformMemory::GetStats();
+		CollectGarbage(RF_NoFlags, true);
+		const FPlatformMemoryStats MemStatsAfter = FPlatformMemory::GetStats();
+
+		UE_LOG(LogWorldPartitionRuntimeSpatialHashHLOD, Warning, TEXT("AvailablePhysical:%.2fGB AvailableVirtual %.2fGB"),
+			((int64)MemStatsAfter.AvailablePhysical - (int64)MemStatsBefore.AvailablePhysical) / (1024.0 * 1024.0 * 1024.0),
+			((int64)MemStatsAfter.AvailableVirtual - (int64)MemStatsBefore.AvailableVirtual) / (1024.0 * 1024.0 * 1024.0)
+		);
+	};
+
+	const FBox WorldBounds = WorldPartition->GetWorldBounds();
+
+	const FSquare2DGridHelper PartitionedActors = GetPartitionedActors(WorldPartition, WorldBounds, RuntimeGrid, ActorClusters);
+
+	// Quick pass to compute the number of cells we'll have to process, to provide a meaningful progress display
+	int32 NbCellsToProcess = 0;
+	PartitionedActors.ForEachCells([&PartitionedActors, &NbCellsToProcess](const FIntVector& CellCoord)
+	{
+		const FSquare2DGridHelper::FGridLevel::FGridCell& GridCell = PartitionedActors.GetCell(CellCoord);
+
+		// For now, HLOD only processes actors that are not in data layers
+		if (const FSquare2DGridHelper::FGridLevel::FGridCellDataChunk* GridCellDataChunk = GridCell.GetNoDataLayersDataChunk())
+		{
+			check(!GridCellDataChunk->HasDataLayers());
+			const bool bIsCellAlwaysLoaded = &GridCell == &PartitionedActors.GetAlwaysLoadedCell();
+
+			if (!bIsCellAlwaysLoaded && GridCellDataChunk->GetActors().Num() != 0)
+			{
+				NbCellsToProcess++;
+			}
+		}
+	});
+
+	FScopedSlowTask SlowTask(NbCellsToProcess, FText::FromString(FString::Printf(TEXT("Building HLODs for grid %s..."), *RuntimeGrid.GridName.ToString())));
+	SlowTask.MakeDialog();
+
+	FDirectoryWatcherModule& DirectoryWatcherModule = FModuleManager::LoadModuleChecked<FDirectoryWatcherModule>(TEXT("DirectoryWatcher"));
+	FAssetRegistryModule& AssetRegistryModule = FModuleManager::LoadModuleChecked<FAssetRegistryModule>(TEXT("AssetRegistry"));
+
+	TArray<FGuid> GridHLODActors;
+	PartitionedActors.ForEachCells([&PartitionedActors, &GridHLODActors, &SlowTask, RuntimeGrid, HLODLevel, WorldPartition, WorldBounds, &Context, SourceControlHelper, &DirectoryWatcherModule, &AssetRegistryModule, &HasExceededMaxMemory, &DoCollectGarbage](const FIntVector& CellCoord)
+	{
+		const FSquare2DGridHelper::FGridLevel::FGridCell& GridCell = PartitionedActors.GetCell(CellCoord);
+
+		// For now, HLOD only processes actors that are not in data layers
+		if (const FSquare2DGridHelper::FGridLevel::FGridCellDataChunk* GridCellDataChunk = GridCell.GetNoDataLayersDataChunk())
+		{
+			check(!GridCellDataChunk->HasDataLayers());
+			const bool bIsCellAlwaysLoaded = &GridCell == &PartitionedActors.GetAlwaysLoadedCell();
+
+			if (!bIsCellAlwaysLoaded && GridCellDataChunk->GetActors().Num() != 0)
+			{
+				SlowTask.EnterProgressFrame(1);
+
+				FBox2D CellBounds2D;
+				PartitionedActors.GetCellBounds(CellCoord, CellBounds2D);
+
+				FName CellName = UWorldPartitionRuntimeSpatialHash::GetCellName(WorldPartition, RuntimeGrid.GridName, CellCoord.Z, CellCoord.X, CellCoord.Y, GridCellDataChunk->GetDataLayersID());
+				FBox CellBounds = FBox(FVector(CellBounds2D.Min, WorldBounds.Min.Z), FVector(CellBounds2D.Max, WorldBounds.Max.Z));
+
+				Context.GridIndexX = CellCoord.X;
+				Context.GridIndexY = CellCoord.Y;
+				Context.GridIndexZ = CellCoord.Z;
+
+				// Gather (and potentially load) actors
+				TArray<AActor*> CellActors;
+				for (const FGuid& ActorGuid : GridCellDataChunk->GetActors())
+				{
+					FWorldPartitionReference& ActorRef = Context.ActorReferences.Emplace_GetRef(WorldPartition, ActorGuid);
+					FWorldPartitionActorDesc* ActorDesc = ActorRef.Get();
+					CellActors.Add(ActorDesc->GetActor());
+				}
+
+				FStaticMeshCompilingManager::Get().FinishAllCompilation();
+
+				if (HasExceededMaxMemory())
+				{
+					DoCollectGarbage();
+				}
+
+				TArray<AWorldPartitionHLOD*> CellHLODActors;
+				CellHLODActors = UHLODLayer::GenerateHLODForCell(WorldPartition, &Context, CellName, CellBounds, HLODLevel, CellActors);
+
+				if (CellHLODActors.Num())
+				{
+					TArray<AWorldPartitionHLOD*> NewCellHLODActors;
+
+					for (AWorldPartitionHLOD* CellHLODActor : CellHLODActors)
+					{
+						FGuid ActorGuid = CellHLODActor->GetActorGuid();
+						GridHLODActors.Add(ActorGuid);
+
+						UPackage* CellHLODActorPackage = CellHLODActor->GetPackage();
+						if (CellHLODActorPackage->HasAnyPackageFlags(PKG_NewlyCreated))
+						{
+							NewCellHLODActors.Add(CellHLODActor);
+						}
+
+						CheckoutAndSavePackage(CellHLODActorPackage, SourceControlHelper);
+					}
+
+					// Manually tick the directory watcher and the asset registry to register newly created actors
+					DirectoryWatcherModule.Get()->Tick(-1.0f);
+					AssetRegistryModule.Get().Tick(-1.0f);
+
+					// Update newly created HLOD actors
+					for (AWorldPartitionHLOD* CellHLODActor : NewCellHLODActors)
+					{
+						Context.ActorReferences.Emplace(WorldPartition, CellHLODActor->GetActorGuid());
+					}
+				}
+
+				// Unload actors
+				Context.ActorReferences.Empty();
+			}
+		}
+	});
+
+	// Need to collect garbage here since some HLOD actors have been marked pending kill when destroying them
+	// and they may be loaded when generating the nexty HLOD layer.
+	DoCollectGarbage();
+
+	return GridHLODActors;
+}
+
+// Find all HLOD grids from the HLODLayer assets we have and build a dependency graph
+static void GatherHLODGrids(TMap<FName, FSpatialHashRuntimeGrid>& OutHLODGrids, TMap<FName, TSet<FName>>& OutHLODGridsGraph)
+{
+	// We don't know the HLOD level for those grids yet, we'll have to compute a graph of dependencies between our HLOD layers first
+	// Start with level 0, we'll rename once we have the correct info
+	const uint32 InitialHLODLevel = 0;
+
+	// Gather up all HLODLayer assets
+	FAssetRegistryModule& AssetRegistryModule = FModuleManager::LoadModuleChecked<FAssetRegistryModule>(TEXT("AssetRegistry"));
+	TArray<FAssetData> HLODLayerAssets;
+	AssetRegistryModule.Get().GetAssetsByClass(UHLODLayer::StaticClass()->GetFName(), HLODLayerAssets);
+
+	// Find all used HLOD grids & build a dependency graph
+	for (const FAssetData& HLODLayerAsset : HLODLayerAssets)
+	{
+		if (const UHLODLayer* HLODLayer = Cast<UHLODLayer>(HLODLayerAsset.GetAsset()))
+		{
+			FSpatialHashRuntimeGrid HLODGrid;
+			HLODGrid.CellSize = HLODLayer->GetCellSize();
+			HLODGrid.LoadingRange = HLODLayer->GetLoadingRange();
+			HLODGrid.DebugColor = FLinearColor::Red;
+			HLODGrid.GridName = HLODLayer->GetRuntimeGrid(InitialHLODLevel);
+
+			OutHLODGrids.Emplace(HLODGrid.GridName, HLODGrid);
+
+			TSet<FName>& HLODParentGrids = OutHLODGridsGraph.FindOrAdd(HLODGrid.GridName);
+			if (const UHLODLayer* ParentHLODLayer = HLODLayer->GetParentLayer().LoadSynchronous())
+			{
+				HLODParentGrids.Add(ParentHLODLayer->GetRuntimeGrid(InitialHLODLevel));
+			}
+		}
+	}
+}
+
+// Sort HLOD grids in the order they'll need to be processed for HLOD generation
+static bool SortHLODGrids(const TMap<FName, TSet<FName>>& InHLODGridsGraph, TArray<FName>& OutSortedGrids, TMap<FName, uint32>& OutGridsDepth)
+{
+	TSet<FName>		ProcessedGridsSet;	// Processed grids
+	TSet<FName>		VisitedGridsSet;	// Visited grids
+
+	TFunctionRef<bool(const FName, uint32)> VisitGraph = [&](const FName GridName, uint32 CurrentDepth) -> bool
+	{
+		uint32& GridDepth = OutGridsDepth.FindOrAdd(GridName);
+		GridDepth = FMath::Max(GridDepth, CurrentDepth);
+
+		if (ProcessedGridsSet.Contains(GridName))
+		{
+			return true;
+		}
+
+		// Detect cyclic dependencies between HLOD grids...
+		if (VisitedGridsSet.Contains(GridName))
+		{
+			return false; // Not a DAG
+		}
+
+		VisitedGridsSet.Add(GridName);
+
+		for (const FName ParentGridName : InHLODGridsGraph.FindChecked(GridName))
+		{
+			if (!VisitGraph(ParentGridName, CurrentDepth + 1))
+			{
+				return false;
+			}
+		}
+
+		VisitedGridsSet.Remove(GridName);
+
+		ProcessedGridsSet.Add(GridName);
+		OutSortedGrids.Insert(GridName, 0);
+		return true;
+	};
+
+	bool bIsADAG = true;
+	for (const auto& HLODGridEntry : InHLODGridsGraph)
+	{
+		FName GridName = HLODGridEntry.Key;
+		if (!ProcessedGridsSet.Contains(GridName))
+		{
+			const uint32 HLODDepth = 0;
+			bIsADAG = VisitGraph(GridName, HLODDepth);
+			if (!bIsADAG)
+			{
+				break;
+			}
+		}
+	}
+
+	// Remove leafs, those grids are the standard runtime grids, we'll generate HLOD for those in a first pass
+	OutSortedGrids.RemoveAll([&InHLODGridsGraph](const FName GridName) { return !InHLODGridsGraph.Contains(GridName); });
+
+	return bIsADAG;
+}
+
+static void RenameHLODGrids(TMap<FName, FSpatialHashRuntimeGrid>& HLODGrids, TArray<FName>& SortedGrids, TMap<FName, uint32>& GridsDepth)
+{
+	// Build a mapping of old -> new names
+	TMap<FName, FName> NewNamesMapping;
+	for (const auto& HLODGrid : HLODGrids)
+	{
+		NewNamesMapping.Emplace(HLODGrid.Key, UHLODLayer::GetRuntimeGridName(GridsDepth[HLODGrid.Key], HLODGrid.Value.CellSize, HLODGrid.Value.LoadingRange));
+	}
+
+	// Remplace map entries
+	for (const auto& Names : NewNamesMapping)
+	{
+		FSpatialHashRuntimeGrid RuntimeGrid = HLODGrids.FindAndRemoveChecked(Names.Key);
+		RuntimeGrid.GridName = Names.Value;
+		HLODGrids.Emplace(Names.Value, RuntimeGrid);
+
+		int32 GridDepth = GridsDepth.FindAndRemoveChecked(Names.Key);
+		GridsDepth.Emplace(Names.Value, GridDepth);
+	}
+
+	// Replace arrays entries
+	for (FName& Name : SortedGrids)
+	{
+		Name = NewNamesMapping[Name];
+	}
+}
+
+// Create/destroy HLOD grid actors
+static void UpdateHLODGridsActors(UWorld* World, const TMap<FName, FSpatialHashRuntimeGrid>& HLODGrids, ISourceControlHelper* SourceControlHelper)
+{
+	static const FName HLODGridTag = TEXT("HLOD");
+	static const uint32 HLODGridTagLen = HLODGridTag.ToString().Len();
+
+	// Gather all existing HLOD grid actors, see if some are unused and needs to be deleted
+	TMap<FName, ASpatialHashRuntimeGridInfo*> ExistingGridActors;
+	for (TActorIterator<ASpatialHashRuntimeGridInfo> ItRuntimeGridActor(World); ItRuntimeGridActor; ++ItRuntimeGridActor)
+	{
+		ASpatialHashRuntimeGridInfo* GridActor = *ItRuntimeGridActor;
+		if (GridActor->ActorHasTag(HLODGridTag))
+		{
+			if (HLODGrids.Contains(GridActor->GridSettings.GridName))
+			{
+				ExistingGridActors.Emplace(GridActor->GridSettings.GridName, GridActor);
+			}
+			else
+			{
+				World->DestroyActor(GridActor);
+
+				DeletePackage(GridActor->GetPackage(), SourceControlHelper);
+			}
+		}
+	}
+
+	// Create missing HLOD grid actors 
+	for (const auto& HLODGridEntry : HLODGrids)
+	{
+		FName GridName = HLODGridEntry.Key;
+		if (!ExistingGridActors.Contains(GridName))
+		{
+			const FSpatialHashRuntimeGrid& GridSettings = HLODGridEntry.Value;
+
+			FActorSpawnParameters SpawnParams;
+			SpawnParams.bCreateActorPackage = true;
+			ASpatialHashRuntimeGridInfo* GridActor = World->SpawnActor<ASpatialHashRuntimeGridInfo>(SpawnParams);
+			GridActor->Tags.Add(HLODGridTag);
+			GridActor->SetActorLabel(GridSettings.GridName.ToString());
+			GridActor->GridSettings = GridSettings;
+
+			// Setup grid debug color to match HLODColorationColors
+			const uint32 LastLODColorationColorIdx = GEngine->HLODColorationColors.Num() - 1;
+			uint32 HLODLevel;
+			if (!LexTryParseString(HLODLevel, *GridSettings.GridName.ToString().Mid(HLODGridTagLen, 1)))
+			{
+				HLODLevel = LastLODColorationColorIdx;
+			}
+			HLODLevel = FMath::Clamp(HLODLevel + 2, 0U, LastLODColorationColorIdx);
+			GridActor->GridSettings.DebugColor = GEngine->HLODColorationColors[HLODLevel];
+
+			CheckoutAndSavePackage(GridActor->GetPackage(), SourceControlHelper);
+		}
+	}
+}
+
+bool UWorldPartitionRuntimeSpatialHash::GenerateHLOD(ISourceControlHelper* SourceControlHelper)
+{
+	TRACE_CPUPROFILER_EVENT_SCOPE(UWorldPartitionRuntimeSpatialHash::GenerateHLOD);
+
+	if (!Grids.Num())
+	{
+		UE_LOG(LogWorldPartitionRuntimeSpatialHashHLOD, Error, TEXT("Invalid partition grids setup"));
+		return false;
+	}
+
+	// Find all used HLOD grids & build a dependency graph
+	TMap<FName, FSpatialHashRuntimeGrid> HLODGrids;
+	TMap<FName, TSet<FName>>			 HLODGridsGraph;
+	GatherHLODGrids(HLODGrids, HLODGridsGraph);
+
+	// Sort HLOD grids in the order they'll need to be processed
+	TArray<FName>	SortedGrids;		// HLOD Grids, sorted in the order they'll need to be processed for HLOD generation
+	TMap<FName, uint32>	GridsDepth;		// Depth - Used to obtain the HLOD level
+	bool bIsADAG = SortHLODGrids(HLODGridsGraph, SortedGrids, GridsDepth);
+	if (!bIsADAG)
+	{
+		UE_LOG(LogWorldPartitionRuntimeSpatialHashHLOD, Error, TEXT("Invalid grids setup, cycles detected"));
+		return false;
+	}
+
+	// Now that we computed proper depth per grid, rename the grids
+	RenameHLODGrids(HLODGrids, SortedGrids, GridsDepth);
+
+	TMap<FName, int32> GridsMapping;
+	GridsMapping.Add(NAME_None, 0);
+	for (int32 i = 0; i < Grids.Num(); i++)
+	{
+		const FSpatialHashRuntimeGrid& Grid = Grids[i];
+		check(!GridsMapping.Contains(Grid.GridName));
+		GridsMapping.Add(Grid.GridName, i);
+	}
+
+	UWorldPartition* WorldPartition = GetOuterUWorldPartition();
+
+	// Create HLODs generation context
+	FHLODGenerationContext Context;
+	TArray<FWorldPartitionHandle> InvalidHLODActors;
+	for (TWorldPartitionActorDescIterator<AWorldPartitionHLOD, FHLODActorDesc> HLODIterator(WorldPartition); HLODIterator; ++HLODIterator)
+	{
+		FWorldPartitionHandle HLODActorHandle(WorldPartition, HLODIterator->GetGuid());
+
+		if (HLODIterator->GetCellHash())
+		{
+			Context.HLODActorDescs.Add(HLODIterator->GetCellHash(), HLODActorHandle);
+		}
+		else
+		{
+			InvalidHLODActors.Add(HLODActorHandle);
+		}
+	}
+
+	// Create actor clusters - ignore HLOD actors
+	TArray<FActorCluster> ActorClusters = CreateActorClusters(WorldPartition, [](const FWorldPartitionActorDesc& ActorDesc)
+		{
+			return !ActorDesc.GetActorClass()->IsChildOf<AWorldPartitionHLOD>();
+		});
+
+	TArray<TArray<FActorCluster>> GridsClusters;
+	GridsClusters.InsertDefaulted(0, Grids.Num());
+
+	for (FActorCluster& ActorCluster : ActorClusters)
+	{
+		int32* FoundIndex = GridsMapping.Find(ActorCluster.RuntimeGrid);
+		if (!FoundIndex)
+		{
+			UE_LOG(LogWorldPartitionRuntimeSpatialHashHLOD, Error, TEXT("Invalid partition grid '%s' referenced by actor cluster"), *ActorCluster.RuntimeGrid.ToString());
+		}
+
+		int32 GridIndex = FoundIndex ? *FoundIndex : 0;
+		GridsClusters[GridIndex].Add(MoveTemp(ActorCluster));
+	}
+
+	// Keep track of all valid HLOD actors, along with which runtime grid they live in
+	TMap<FName, TArray<FGuid>> GridsHLODActors;
+
+	auto GenerateHLODs = [&GridsHLODActors, WorldPartition, &GridsDepth, &Context, SourceControlHelper](const FSpatialHashRuntimeGrid& RuntimeGrid, uint32 HLODLevel, const TArray<FActorCluster>& ActorClusters)
+	{
+		// Generate HLODs for this grid
+		TArray<FGuid> HLODActors = GenerateHLODsForGrid(WorldPartition, RuntimeGrid, HLODLevel, Context, SourceControlHelper, ActorClusters);
+
+		for (const FGuid& HLODActorGuid : HLODActors)
+		{
+			FWorldPartitionActorDesc* HLODActorDesc = WorldPartition->GetActorDesc(HLODActorGuid);
+			check(HLODActorDesc);
+
+			GridsHLODActors.FindOrAdd(HLODActorDesc->GetRuntimeGrid()).Add(HLODActorGuid);
+		}
+	};
+
+	// Generate HLODs for the standard runtime grids (HLOD 0)
+	for (int32 GridIndex = 0; GridIndex < Grids.Num(); GridIndex++)
+	{
+		// Generate HLODs for this grid - retrieve actors from the world partition
+		GenerateHLODs(Grids[GridIndex], 0, GridsClusters[GridIndex]);
+	}
+
+	// Now, go on and create HLOD actors from HLOD grids (HLOD 1-N)
+	for (const FName HLODGridName : SortedGrids)
+	{
+		// No need to process empty grids
+		if (!GridsHLODActors.Contains(HLODGridName))
+		{
+			HLODGrids.Remove(HLODGridName); // No need to keep this grid around, we have no actors in it
+			continue;
+		}
+
+		// Create one actor cluster per HLOD actor
+		TArray<FActorCluster> HLODActorClusters;
+		HLODActorClusters.Reserve(GridsHLODActors[HLODGridName].Num());
+
+		for (const FGuid& HLODActorGuid : GridsHLODActors[HLODGridName])
+		{
+			FWorldPartitionActorDesc* HLODActorDesc = WorldPartition->GetActorDesc(HLODActorGuid);
+			check(HLODActorDesc);
+
+			HLODActorClusters.Emplace(HLODActorDesc, HLODActorDesc->GetGridPlacement(), GetWorld());
+		};
+
+		// Generate HLODs for this grid - retrieve actors from our ValidHLODActors map
+		// We can't rely on actor descs for newly created HLOD actors
+		GenerateHLODs(HLODGrids[HLODGridName], GridsDepth[HLODGridName] + 1, HLODActorClusters);
+	}
+
+	auto DeleteHLODActor = [&SourceControlHelper](FWorldPartitionHandle ActorHandle)
+	{
+		FWorldPartitionActorDesc* HLODActorDesc = ActorHandle.Get();
+		check(HLODActorDesc);
+
+		DeletePackage(HLODActorDesc->GetActorPackage().ToString(), SourceControlHelper);
+	};
+
+	// Destroy all unreferenced HLOD actors
+	for (const auto& HLODActorPair : Context.HLODActorDescs)
+	{
+		DeleteHLODActor(HLODActorPair.Value);
+	}
+
+	// Destroy all invalid HLOD actors
+	Algo::ForEach(InvalidHLODActors, DeleteHLODActor);
+
+	// Create/destroy HLOD grid actors
+	UpdateHLODGridsActors(GetWorld(), HLODGrids, SourceControlHelper);
+
+	return true;
+}
+
+#endif // #if WITH_EDITOR
\ No newline at end of file
diff --git a/Engine/Source/Runtime/Engine/Private/WorldPartition/RuntimeSpatialHash/RuntimeSpatialHashNavigation.cpp b/Engine/Source/Runtime/Engine/Private/WorldPartition/RuntimeSpatialHash/RuntimeSpatialHashNavigation.cpp
new file mode 100644
--- /dev/null
+++ b/Engine/Source/Runtime/Engine/Private/WorldPartition/RuntimeSpatialHash/RuntimeSpatialHashNavigation.cpp
@@ -0,0 +1,118 @@
+// Copyright Epic Games, Inc. All Rights Reserved.
+
+#include "WorldPartition/WorldPartitionRuntimeSpatialHash.h"
+
+#if WITH_EDITOR
+
+#include "WorldPartition/RuntimeSpatialHash/RuntimeSpatialHashGridHelper.h"
+#include "WorldPartition/NavigationData/NavigationDataChunkActor.h"
+
+#include "EngineUtils.h"
+
+#include "AI/NavigationSystemBase.h"
+
+DEFINE_LOG_CATEGORY_STATIC(LogWorldPartitionRuntimeSpatialHashNav, Log, All);
+
+bool UWorldPartitionRuntimeSpatialHash::GenerateNavigationData()
+{
+	TRACE_CPUPROFILER_EVENT_SCOPE(UWorldPartitionRuntimeSpatialHash::GenerateNavigationData);
+
+	UE_LOG(LogWorldPartitionRuntimeSpatialHashNav, Log, TEXT("%s"), ANSI_TO_TCHAR(__FUNCTION__));
+
+	UWorldPartition* WorldPartition = GetOuterUWorldPartition();
+	UWorld* World = WorldPartition->World;
+
+	// Generate navmesh
+	// Make sure navigation is added and initialized in EditorMode
+	FNavigationSystem::AddNavigationSystemToWorld(*World, FNavigationSystemRunMode::EditorMode);
+
+	// Invoke navigation data generator
+	FNavigationSystem::Build(*World);
+
+	// For each cell, gather navmesh and generate a datachunk actor
+	const FBox WorldBounds = WorldPartition->GetWorldBounds();
+
+	const int32 GridIndex = 0; //Todo At: only work with grid 0 for now.
+	const uint32 GridLevel = 3; //Todo AT, only generate for level 3 for now
+	UE_LOG(LogWorldPartitionRuntimeSpatialHashNav, Verbose, TEXT("Generate NavDataChunk actor for GridIndex %i."), GridIndex);
+
+	const FSpatialHashRuntimeGrid& RuntimeGrid = Grids[GridIndex];
+	const FSquare2DGridHelper GridHelper = GetGridHelper(WorldBounds, RuntimeGrid.CellSize);
+	int32 ActorCount = 0;
+
+	const UNavigationSystemBase* NavSystem = World->GetNavigationSystem();
+	if (NavSystem == nullptr)
+	{
+		UE_LOG(LogWorldPartitionRuntimeSpatialHashNav, Verbose, TEXT("No navigation system to generate navigation data."));
+		return false;
+	}
+
+	// Keep track of all valid navigation data chunk actors
+	TSet<ANavigationDataChunkActor*> ValidNavigationDataChunkActors;
+
+	const FSquare2DGridHelper::FGridLevel& GridLevelHelper = GridHelper.Levels[GridLevel];
+
+	GridLevelHelper.ForEachCells([GridLevel, &GridLevelHelper, RuntimeGrid, WorldPartition, &ActorCount, World, NavSystem, &ValidNavigationDataChunkActors, this](const FIntVector2& CellCoord)
+		{
+			FBox2D CellBounds;
+			GridLevelHelper.GetCellBounds(CellCoord, CellBounds);
+			if (CellBounds.GetExtent().X < 1.f || CellBounds.GetExtent().Y < 1.f)
+			{
+				// Safety, since we reduce by 1.f below.
+				UE_LOG(LogWorldPartitionRuntimeSpatialHashNav, Warning, TEXT("%s: grid cell too small."), ANSI_TO_TCHAR(__FUNCTION__));
+				return;
+			}
+
+			const float HalfHeight = HALF_WORLD_MAX;
+			const FBox QueryBounds(FVector(CellBounds.Min.X, CellBounds.Min.Y, -HalfHeight), FVector(CellBounds.Max.X, CellBounds.Max.Y, HalfHeight));
+
+			if (NavSystem->ContainsNavData(QueryBounds) == false)
+			{
+				// Skip if there is no navdata for this cell
+				return;
+			}
+
+			const FSquare2DGridHelper::FGridLevel::FGridCell& GridCell = GridLevelHelper.GetCell(CellCoord);
+
+			FActorSpawnParameters SpawnParams;
+			SpawnParams.bDeferConstruction = true;
+			SpawnParams.bCreateActorPackage = true;
+			ANavigationDataChunkActor* DataChunkActor = World->SpawnActor<ANavigationDataChunkActor>(SpawnParams);
+			ActorCount++;
+
+			const FVector2D CellCenter = CellBounds.GetCenter();
+			DataChunkActor->SetActorLocation(FVector(CellCenter.X, CellCenter.Y, 0.f));
+
+			FBox TilesBounds(EForceInit::ForceInit);
+			DataChunkActor->CollectNavData(QueryBounds, TilesBounds);
+
+			FBox ChunkActorBounds(FVector(QueryBounds.Min.X, QueryBounds.Min.Y, TilesBounds.Min.Z), FVector(QueryBounds.Max.X, QueryBounds.Max.Y, TilesBounds.Max.Z));
+			ChunkActorBounds = ChunkActorBounds.ExpandBy(FVector(-1.f, -1.f, 1.f)); //reduce XY by 1cm to avoid precision issues causing potential overflow on neighboring cell, add 1cm in Z to have a minimum of volume.
+			UE_LOG(LogWorldPartitionRuntimeSpatialHashNav, VeryVerbose, TEXT("Setting ChunkActorBounds to %s"), *ChunkActorBounds.ToString());
+			DataChunkActor->SetDataChunkActorBounds(ChunkActorBounds);
+
+			const FName CellName = GetCellName(RuntimeGrid.GridName, GridLevel, CellCoord.X, CellCoord.Y);
+			DataChunkActor->SetActorLabel(FString::Printf(TEXT("NavDataChunkActor_%s_%s"), *GetName(), *CellName.ToString()));
+
+			//@todo_ow: Properly handle data layers
+
+			// Set target grid
+			DataChunkActor->RuntimeGrid = RuntimeGrid.GridName;
+			ValidNavigationDataChunkActors.Add(DataChunkActor);
+
+			UE_LOG(LogWorldPartitionRuntimeSpatialHashNav, Verbose, TEXT("%i) %s added."), ActorCount, *DataChunkActor->GetName());
+		});
+
+	// Destroy all invalid navigation data chunk actors
+	for (TActorIterator<ANavigationDataChunkActor> It(GetWorld()); It; ++It)
+	{
+		if (!ValidNavigationDataChunkActors.Contains(*It))
+		{
+			GetWorld()->DestroyActor(*It);
+		}
+	}
+
+	return true;
+}
+
+#endif WITH_EDITOR
\ No newline at end of file
diff --git a/Engine/Source/Runtime/Engine/Private/WorldPartition/WorldPartitionRuntimeSpatialHash.cpp b/Engine/Source/Runtime/Engine/Private/WorldPartition/WorldPartitionRuntimeSpatialHash.cpp
--- a/Engine/Source/Runtime/Engine/Private/WorldPartition/WorldPartitionRuntimeSpatialHash.cpp
+++ b/Engine/Source/Runtime/Engine/Private/WorldPartition/WorldPartitionRuntimeSpatialHash.cpp
@@ -9,9 +9,16 @@
 #include "WorldPartition/WorldPartitionActorDescIterator.h"
 #include "WorldPartition/WorldPartitionStreamingPolicy.h"
 #include "WorldPartition/WorldPartitionHandle.h"
+
+#include "WorldPartition/RuntimeSpatialHash/RuntimeSpatialHashGridHelper.h"
+
+#include "WorldPartition/HLOD/HLODActor.h"
+#include "WorldPartition/HLOD/HLODActorDesc.h"
+
 #include "WorldPartition/DataLayer/DataLayerSubsystem.h"
 #include "WorldPartition/DataLayer/WorldDataLayers.h"
 #include "WorldPartition/DataLayer/DataLayer.h"
+
 #include "GameFramework/WorldSettings.h"
 #include "ProfilingDebugging/ScopedTimers.h"
 #include "Engine/World.h"
@@ -24,38 +31,17 @@
 #if WITH_EDITOR
 #include "Editor/EditorEngine.h"
 #include "Misc/Parse.h"
-#include "Misc/ScopedSlowTask.h"
+
 #include "Algo/Find.h"
 #include "Algo/ForEach.h"
-#include "EngineUtils.h"
-
-#include "DistanceFieldAtlas.h"
-#include "MeshCardRepresentation.h"
-
-#include "WorldPartition/HLOD/HLODBuilder.h"
-#include "WorldPartition/HLOD/HLODLayer.h"
-#include "WorldPartition/HLOD/HLODActor.h"
-#include "WorldPartition/HLOD/HLODActorDesc.h"
-#include "WorldPartition/NavigationData/NavigationDataChunkActor.h"
-#include "AI/NavigationSystemBase.h"
-#include "Engine/Public/ActorReferencesUtils.h"
-#include "Engine/LevelScriptBlueprint.h"
-
-#include "AssetRegistryModule.h"
-#include "AssetData.h"
-
-#include "IDirectoryWatcher.h"
-#include "DirectoryWatcherModule.h"
 
 #include "Engine/WorldComposition.h"
 #include "LevelUtils.h"
 
-#include "StaticMeshCompiler.h"
-
 extern UNREALED_API class UEditorEngine* GEditor;
 #endif //WITH_EDITOR
 
-DEFINE_LOG_CATEGORY_STATIC(LogWorldPartitionRuntimeSpatialHash, Log, All);
+DEFINE_LOG_CATEGORY(LogWorldPartitionRuntimeSpatialHash);
 
 static int32 GShowRuntimeSpatialHashGridLevel = 0;
 static FAutoConsoleVariableRef CVarShowRuntimeSpatialHashGridLevel(
@@ -106,1820 +92,383 @@ FDataLayersID::FDataLayersID(const TArray<const UDataLayer*>& InDataLayers)
 		check(Hash);
 	}
 }
+#endif
 
-// Clustering
-struct FActorCluster
+// ------------------------------------------------------------------------------------------------
+FSpatialHashStreamingGrid::FSpatialHashStreamingGrid()
+	: Origin(ForceInitToZero)
+	, CellSize(0)
+	, GridSize(0)
+	, LoadingRange(0.0f)
+#if WITH_EDITOR
+	, DebugColor(ForceInitToZero)
+	, OverrideLoadingRange(0)
+#endif
+	, GridHelper(nullptr)
 {
-	TSet<FGuid>					Actors;
-	EActorGridPlacement			GridPlacement;
-	FName						RuntimeGrid;
-	FBox						Bounds;
-	TArray<const UDataLayer*>	DataLayers;
-	FDataLayersID				DataLayersID;
-
-	FActorCluster(const FWorldPartitionActorDesc* InActorDesc, EActorGridPlacement InGridPlacement, UWorld* InWorld)
-		: GridPlacement(InGridPlacement)
-		, RuntimeGrid(InActorDesc->GetRuntimeGrid())
-		, Bounds(InActorDesc->GetBounds())
-	{
-		check(GridPlacement != EActorGridPlacement::None);
-
-		Actors.Add(InActorDesc->GetGuid());
-		if (const AWorldDataLayers* WorldDataLayers = AWorldDataLayers::Get(InWorld))
-		{
-			for (const FName& DataLayerName : InActorDesc->GetDataLayers())
-			{
-				if (const UDataLayer* DataLayer = WorldDataLayers->GetDataLayerFromName(DataLayerName))
-				{
-					if (DataLayer->IsDynamicallyLoaded())
-					{
-						DataLayers.Add(DataLayer);
-					}
-				}
-			}
-		}
-
-		DataLayersID = FDataLayersID(DataLayers);
-	}
+}
 
-	void Add(const FActorCluster& InActorCluster)
+FSpatialHashStreamingGrid::~FSpatialHashStreamingGrid()
+{
+	if (GridHelper)
 	{
-		// Merge Actors
-		Actors.Append(InActorCluster.Actors);
-
-		// Merge RuntimeGrid
-		RuntimeGrid = RuntimeGrid == InActorCluster.RuntimeGrid ? RuntimeGrid : NAME_None;
-
-		// Merge Bounds
-		Bounds += InActorCluster.Bounds;
-
-		// Merge GridPlacement
-		// If currently None, will always stay None
-		if (GridPlacement != EActorGridPlacement::None)
-		{
-			// If grid placement differs between the two clusters
-			if (GridPlacement != InActorCluster.GridPlacement)
-			{
-				// If one of the two cluster was always loaded, set to None
-				if (InActorCluster.GridPlacement == EActorGridPlacement::AlwaysLoaded || GridPlacement == EActorGridPlacement::AlwaysLoaded)
-				{
-					GridPlacement = EActorGridPlacement::None;
-				}
-				else
-				{
-					GridPlacement = InActorCluster.GridPlacement;
-				}
-			}
-
-			// If current placement is set to Location, that won't make sense when merging two clusters. Set to Bounds
-			if (GridPlacement == EActorGridPlacement::Location)
-			{
-				GridPlacement = EActorGridPlacement::Bounds;
-			}
-		}
-
-		// Merge DataLayers
-		if (DataLayersID != InActorCluster.DataLayersID)
-		{
-			for (const UDataLayer* DataLayer : InActorCluster.DataLayers)
-			{
-				check(DataLayer->IsDynamicallyLoaded());
-				DataLayers.AddUnique(DataLayer);
-			}
-			DataLayersID = FDataLayersID(DataLayers);
-		}
+		delete GridHelper;
 	}
-};
+}
 
-void CreateActorCluster(const FWorldPartitionActorDesc* ActorDesc, EActorGridPlacement GridPlacement, TMap<FGuid, FActorCluster*>& ActorToActorCluster, TSet<FActorCluster*>& ActorClustersSet, UWorldPartition* WorldPartition)
+const FSquare2DGridHelper& FSpatialHashStreamingGrid::GetGridHelper() const
 {
-	UWorld* World = WorldPartition->GetWorld();
-	const FGuid& ActorGuid = ActorDesc->GetGuid();
-
-	FActorCluster* ActorCluster = ActorToActorCluster.FindRef(ActorGuid);
-	if (!ActorCluster)
+	if (!GridHelper)
+	{
+		GridHelper = new FSquare2DGridHelper(GridLevels.Num(), Origin, CellSize, GridSize);
+	}
+	else
 	{
-		ActorCluster = new FActorCluster(ActorDesc, GridPlacement, World);
-		ActorClustersSet.Add(ActorCluster);
-		ActorToActorCluster.Add(ActorGuid, ActorCluster);
+		check(GridHelper->Levels.Num() == GridLevels.Num());
+		check(GridHelper->Origin == Origin);
+		check(GridHelper->CellSize == CellSize);
+		check(GridHelper->GridSize == GridSize);
 	}
+	return *GridHelper;
+}
 
-	// Don't include references from editor-only actors
-	if (!ActorDesc->GetActorIsEditorOnly())
+void FSpatialHashStreamingGrid::GetCells(const TArray<FWorldPartitionStreamingSource>& Sources, const UDataLayerSubsystem* DataLayerSubsystem, TSet<const UWorldPartitionRuntimeCell*>& Cells) const
+{
+	const FSquare2DGridHelper& Helper = GetGridHelper();
+	for (const FWorldPartitionStreamingSource& Source : Sources)
 	{
-		for (const FGuid& ReferenceGuid : ActorDesc->GetReferences())
+		const FSphere GridSphere(Source.Location, GetLoadingRange());
+		Helper.ForEachIntersectingCells(GridSphere, [&](const FIntVector& Coords)
 		{
-			const FWorldPartitionActorDesc* ReferenceActorDesc = WorldPartition->GetActorDesc(ReferenceGuid);
-
-			// Don't include references to editor-only actors
-			if (!ReferenceActorDesc->GetActorIsEditorOnly())
+			const FSpatialHashStreamingGridLayerCell& LayerCell = GridLevels[Coords.Z].LayerCells[Coords.Y * Helper.Levels[Coords.Z].GridSize + Coords.X];
+			for (const UWorldPartitionRuntimeCell* Cell : LayerCell.GridCells)
 			{
-				FActorCluster* ReferenceCluster = ActorToActorCluster.FindRef(ReferenceGuid);
-				if (ReferenceCluster)
-				{
-					if (ReferenceCluster != ActorCluster)
-					{
-						// Merge reference cluster in Actor's cluster
-						ActorCluster->Add(*ReferenceCluster);
-						for (const FGuid& ReferenceClusterActorGuid : ReferenceCluster->Actors)
-						{
-							ActorToActorCluster[ReferenceClusterActorGuid] = ActorCluster;
-						}
-						ActorClustersSet.Remove(ReferenceCluster);
-						delete ReferenceCluster;
-					}
-				}
-				else
+				if (!Cell->HasDataLayers() || (DataLayerSubsystem && DataLayerSubsystem->IsAnyDataLayerActive(Cell->GetDataLayers())))
 				{
-					// Put Reference in Actor's cluster
-					ActorCluster->Add(FActorCluster(ReferenceActorDesc, GridPlacement, World));
+					Cells.Add(Cell);
 				}
-
-				// Map its cluster
-				ActorToActorCluster.Add(ReferenceGuid, ActorCluster);
 			}
-		}
+		});
 	}
-}
 
-typedef TFunctionRef<bool(const FWorldPartitionActorDesc&)> FFilterPredicate;
+	GetAlwaysLoadedCells(DataLayerSubsystem, Cells);
+}
 
-TArray<FActorCluster> CreateActorClustersImpl(UWorldPartition* WorldPartition, TOptional<FFilterPredicate> InFilterPredicate)
+void FSpatialHashStreamingGrid::GetAlwaysLoadedCells(const UDataLayerSubsystem* DataLayerSubsystem, TSet<const UWorldPartitionRuntimeCell*>& Cells) const
 {
-	TMap<FGuid, FActorCluster*> ActorToActorCluster;
-	TSet<FActorCluster*> ActorClustersSet;
-
-	// Gather all references to external actors from the level script
-	TSet<AActor*> LevelScriptExternalActorReferences;
-	if (ULevelScriptBlueprint* LevelScriptBlueprint = WorldPartition->GetWorld()->PersistentLevel->GetLevelScriptBlueprint(true))
-	{
-		LevelScriptExternalActorReferences.Append(ActorsReferencesUtils::GetExternalActorReferences(LevelScriptBlueprint));
-	}
-
-	for (const auto& Pair : WorldPartition->Actors)
+	if (GridLevels.Num() > 0)
 	{
-		const FWorldPartitionActorDesc* ActorDesc = Pair.Value->Get();
-		EActorGridPlacement GridPlacement = ActorDesc->GetGridPlacement();
-
-		// Check if the actor is loaded (potentially referenced by the level script)
-		if (AActor* Actor = ActorDesc->GetActor())
+		const int32 TopLevel = GridLevels.Num() - 1;
+		for (const FSpatialHashStreamingGridLayerCell& LayerCell : GridLevels[TopLevel].LayerCells)
 		{
-			if (LevelScriptExternalActorReferences.Contains(Actor))
+			for (const UWorldPartitionRuntimeCell* Cell : LayerCell.GridCells)
 			{
-				GridPlacement = EActorGridPlacement::AlwaysLoaded;
+				if (!Cell->HasDataLayers() || (DataLayerSubsystem && DataLayerSubsystem->IsAnyDataLayerActive(Cell->GetDataLayers())))
+				{
+					check(Cell->IsAlwaysLoaded() || Cell->HasDataLayers());
+					Cells.Add(Cell);
+				}
 			}
 		}
-
-		if (!InFilterPredicate.IsSet() || InFilterPredicate.GetValue()(*ActorDesc))
-		{
-			CreateActorCluster(ActorDesc, GridPlacement, ActorToActorCluster, ActorClustersSet, WorldPartition);
-		}
 	}
-
-	TArray<FActorCluster> ActorClusters;
-	ActorClusters.Reserve(ActorClustersSet.Num());
-	Algo::Transform(ActorClustersSet, ActorClusters, [](FActorCluster* ActorCluster) { return MoveTemp(*ActorCluster); });
-	for (FActorCluster* ActorCluster : ActorClustersSet) { delete ActorCluster; }
-	return ActorClusters;
-}
-
-TArray<FActorCluster> CreateActorClusters(UWorldPartition* WorldPartition, const FFilterPredicate& InFilterPredicate)
-{
-	return CreateActorClustersImpl(WorldPartition, InFilterPredicate);
 }
 
-TArray<FActorCluster> CreateActorClusters(UWorldPartition* WorldPartition)
+void FSpatialHashStreamingGrid::Draw3D(UWorld* World, const TArray<FWorldPartitionStreamingSource>& Sources, const FTransform& Transform) const
 {
-	return CreateActorClustersImpl(WorldPartition, TOptional<FFilterPredicate>());
-}
-
-#endif
+	const FSquare2DGridHelper& Helper = GetGridHelper();
+	int32 MinGridLevel = FMath::Clamp<int32>(GShowRuntimeSpatialHashGridLevel, 0, GridLevels.Num() - 1);
+	int32 MaxGridLevel = FMath::Clamp<int32>(MinGridLevel + GShowRuntimeSpatialHashGridLevelCount - 1, 0, GridLevels.Num() - 1);
+	const float GridViewMinimumSizeInCellCount = 5.f;
+	const float GridViewLoadingRangeExtentRatio = 1.5f;
+	const float Radius = GetLoadingRange();
+	const float GridSideDistance = FMath::Max((2.f * Radius * GridViewLoadingRangeExtentRatio), CellSize * GridViewMinimumSizeInCellCount);
 
-/**
-  * Square 2D grid helper
-  */
-struct FSquare2DGridHelper
-{
-	struct FGrid2D
+	for (const FWorldPartitionStreamingSource& Source : Sources)
 	{
-		FVector2D Origin;
-		int32 CellSize;
-		int32 GridSize;
-
-		inline FGrid2D(const FVector2D& InOrigin, int32 InCellSize, int32 InGridSize)
-			: Origin(InOrigin)
-			, CellSize(InCellSize)
-			, GridSize(InGridSize)
-		{}
-
-		/**
-		 * Validate that the coordinates fit the grid size
-		 *
-		 * @return true if the specified coordinates are valid
-		 */
-		inline bool IsValidCoords(const FIntVector2& InCoords) const
-		{
-			return (InCoords.X >= 0) && (InCoords.X < GridSize) && (InCoords.Y >= 0) && (InCoords.Y < GridSize);
-		}
-
-		/**
-		 * Returns the cell bounds
-		 *
-		 * @return true if the specified index was valid
-		 */
-		inline bool GetCellBounds(int32 InIndex, FBox2D& OutBounds) const
+		FVector StartTrace = Source.Location + FVector(0.f, 0.f, 100.f);
+		FVector EndTrace = StartTrace - FVector(0.f, 0.f, 1000000.f);
+		float Z = Source.Location.Z;
+		FHitResult Hit;
+		if (World->LineTraceSingleByObjectType(Hit, StartTrace, EndTrace, FCollisionObjectQueryParams(ECC_WorldStatic), FCollisionQueryParams(SCENE_QUERY_STAT(DebugWorldPartitionTrace), true)))
 		{
-			if (InIndex >= 0 && InIndex <= (GridSize * GridSize))
-			{
-				const FIntVector2 Coords(InIndex % GridSize, InIndex / GridSize);
-				return GetCellBounds(Coords, OutBounds);
-			}
-
-			return false;
+			Z = Hit.ImpactPoint.Z;
 		}
 
-		/**
-		 * Returns the cell bounds
-		 *
-		 * @return true if the specified coord was valid
-		 */
-		inline bool GetCellBounds(const FIntVector2& InCoords, FBox2D& OutBounds) const
+		FSphere Sphere(Source.Location, GridSideDistance * 0.5f);
+		const FBox Region(Sphere.Center - Sphere.W, Sphere.Center + Sphere.W);
+		for (int32 GridLevel = MinGridLevel; GridLevel <= MaxGridLevel; ++GridLevel)
 		{
-			if (IsValidCoords(InCoords))
+			Helper.Levels[GridLevel].ForEachIntersectingCells(Region, [&](const FIntVector2& Coords)
 			{
-				const FVector2D Min = (FVector2D(Origin) - FVector2D(GridSize * CellSize * 0.5f, GridSize * CellSize * 0.5f)) + FVector2D(InCoords.X * CellSize, InCoords.Y * CellSize);
-				const FVector2D Max = Min + FVector2D(CellSize, CellSize);
-				OutBounds = FBox2D(Min, Max);
-				return true;
-			}
-
-			return false;
+				FBox2D CellWorldBounds;
+				Helper.Levels[GridLevel].GetCellBounds(FIntVector2(Coords.X, Coords.Y), CellWorldBounds);
+				
+				FVector BoundsExtent(CellWorldBounds.GetExtent(), 100.f);
+				const FSpatialHashStreamingGridLayerCell& LayerCell = GridLevels[GridLevel].LayerCells[Coords.Y * Helper.Levels[GridLevel].GridSize + Coords.X];
+				const UWorldPartitionRuntimeSpatialHashCell* Cell = nullptr;
+				for (const UWorldPartitionRuntimeSpatialHashCell* GridCell : LayerCell.GridCells)
+				{
+					Cell = GridCell;
+					if (!GridCell->HasDataLayers())
+					{
+						break;
+					}
+				}
+				FColor CellColor = Cell ? Cell->GetDebugColor().ToFColor(false).WithAlpha(16) : FColor(0, 0, 0, 16);
+				FVector BoundsOrigin(CellWorldBounds.GetCenter(), Z);
+				FBox Box = FBox::BuildAABB(BoundsOrigin, BoundsExtent);
+				DrawDebugSolidBox(World, Box, CellColor, Transform, false, -1.f, 255);
+				FVector Pos = Transform.TransformPosition(BoundsOrigin);
+				DrawDebugBox(World, Pos, BoundsExtent, Transform.GetRotation(), CellColor.WithAlpha(255), false, -1.f, 255, 10.f);
+			});
 		}
 
-		/**
-		 * Returns the cell coordinates of the provided position
-		 *
-		 * @return true if the position was inside the grid
-		 */
-		inline bool GetCellCoords(const FVector2D& InPos, FIntVector2& OutCoords) const
-		{
-			OutCoords = FIntVector2(
-				FMath::FloorToInt(((InPos.X - Origin.X) / CellSize) + GridSize * 0.5f),
-				FMath::FloorToInt(((InPos.Y - Origin.Y) / CellSize) + GridSize * 0.5f)
-			);
+		// Draw Loading Ranges
+		FVector SphereLocation(FVector2D(Source.Location), Z);
+		SphereLocation = Transform.TransformPosition(SphereLocation);
+		DrawDebugSphere(World, SphereLocation, Radius, 32, FColor::White, false, -1.f, 0, 20.f);
+	}
+}
 
-			return IsValidCoords(OutCoords);
-		}
+void FSpatialHashStreamingGrid::Draw2D(UCanvas* Canvas, const TArray<FWorldPartitionStreamingSource>& Sources, const FBox& Region, const FBox2D& GridScreenBounds, TFunctionRef<FVector2D(const FVector2D&)> WorldToScreen) const
+{
+	FCanvas* CanvasObject = Canvas->Canvas;
+	int32 MinGridLevel = FMath::Clamp<int32>(GShowRuntimeSpatialHashGridLevel, 0, GridLevels.Num() - 1);
+	int32 MaxGridLevel = FMath::Clamp<int32>(MinGridLevel + GShowRuntimeSpatialHashGridLevelCount - 1, 0, GridLevels.Num() - 1);
 
-		/**
-		 * Returns the cells coordinates of the provided box
-		 *
-		 * @return true if the bounds was intersecting with the grid
-		 */
-		inline bool GetCellCoords(const FBox2D& InBounds2D, FIntVector2& OutMinCellCoords, FIntVector2& OutMaxCellCoords) const
+	for (int32 GridLevel = MinGridLevel; GridLevel <= MaxGridLevel; ++GridLevel)
+	{
+		// Draw X/Y Axis
 		{
-			GetCellCoords(InBounds2D.Min, OutMinCellCoords);
-			if ((OutMinCellCoords.X >= GridSize) || (OutMinCellCoords.Y >= GridSize))
+			FCanvasLineItem Axis;
+			Axis.LineThickness = 3;
 			{
-				return false;
+				Axis.SetColor(FLinearColor::Green);
+				FVector2D LineStart = WorldToScreen(FVector2D(-1638400.f, 0.f));
+				FVector2D LineEnd = WorldToScreen(FVector2D(1638400.f, 0.f));
+				LineStart.X = FMath::Clamp(LineStart.X, GridScreenBounds.Min.X, GridScreenBounds.Max.X);
+				LineStart.Y = FMath::Clamp(LineStart.Y, GridScreenBounds.Min.Y, GridScreenBounds.Max.Y);
+				LineEnd.X = FMath::Clamp(LineEnd.X, GridScreenBounds.Min.X, GridScreenBounds.Max.X);
+				LineEnd.Y = FMath::Clamp(LineEnd.Y, GridScreenBounds.Min.Y, GridScreenBounds.Max.Y);
+				Axis.Draw(CanvasObject, LineStart, LineEnd);
 			}
-
-			GetCellCoords(InBounds2D.Max, OutMaxCellCoords);
-			if ((OutMaxCellCoords.X < 0) || (OutMaxCellCoords.Y < 0))
 			{
-				return false;
+				Axis.SetColor(FLinearColor::Red);
+				FVector2D LineStart = WorldToScreen(FVector2D(0.f, -1638400.f));
+				FVector2D LineEnd = WorldToScreen(FVector2D(0.f, 1638400.f));
+				LineStart.X = FMath::Clamp(LineStart.X, GridScreenBounds.Min.X, GridScreenBounds.Max.X);
+				LineStart.Y = FMath::Clamp(LineStart.Y, GridScreenBounds.Min.Y, GridScreenBounds.Max.Y);
+				LineEnd.X = FMath::Clamp(LineEnd.X, GridScreenBounds.Min.X, GridScreenBounds.Max.X);
+				LineEnd.Y = FMath::Clamp(LineEnd.Y, GridScreenBounds.Min.Y, GridScreenBounds.Max.Y);
+				Axis.Draw(CanvasObject, LineStart, LineEnd);
 			}
-
-			OutMinCellCoords.X = FMath::Clamp(OutMinCellCoords.X, 0, GridSize - 1);
-			OutMinCellCoords.Y = FMath::Clamp(OutMinCellCoords.Y, 0, GridSize - 1);
-			OutMaxCellCoords.X = FMath::Clamp(OutMaxCellCoords.X, 0, GridSize - 1);
-			OutMaxCellCoords.Y = FMath::Clamp(OutMaxCellCoords.Y, 0, GridSize - 1);
-
-			return true;
 		}
 
-		/**
-		 * Returns the cell index of the provided coords
-		 *
-		 * @return true if the coords was inside the grid
-		 */
-		inline bool GetCellIndex(const FIntVector2& InCoords, uint32& OutIndex) const
+		// Draw Grid cells at desired grid level
+		const FSquare2DGridHelper& Helper = GetGridHelper();
+		Helper.Levels[GridLevel].ForEachIntersectingCells(Region, [&](const FIntVector2& Coords)
 		{
-			if (IsValidCoords(InCoords))
+			FBox2D CellWorldBounds;
+			Helper.Levels[GridLevel].GetCellBounds(FIntVector2(Coords.X, Coords.Y), CellWorldBounds);
+			FBox2D CellScreenBounds = FBox2D(WorldToScreen(CellWorldBounds.Min), WorldToScreen(CellWorldBounds.Max));
+			// Clamp inside grid bounds
+			if (!GridScreenBounds.IsInside(CellScreenBounds))
 			{
-				OutIndex = (InCoords.Y * GridSize) + InCoords.X;
-				return true;
+				CellScreenBounds.Min.X = FMath::Clamp(CellScreenBounds.Min.X, GridScreenBounds.Min.X, GridScreenBounds.Max.X);
+				CellScreenBounds.Min.Y = FMath::Clamp(CellScreenBounds.Min.Y, GridScreenBounds.Min.Y, GridScreenBounds.Max.Y);
+				CellScreenBounds.Max.X = FMath::Clamp(CellScreenBounds.Max.X, GridScreenBounds.Min.X, GridScreenBounds.Max.X);
+				CellScreenBounds.Max.Y = FMath::Clamp(CellScreenBounds.Max.Y, GridScreenBounds.Min.Y, GridScreenBounds.Max.Y);
 			}
-
-			return false;
-		}
-
-		/**
-		 * Returns the cell index of the provided position
-		 *
-		 * @return true if the position was inside the grid
-		 */
-		inline bool GetCellIndex(const FVector& InPos, uint32& OutIndex) const
-		{
-			FIntVector2 Coords = FIntVector2(
-				FMath::FloorToInt(((InPos.X - Origin.X) / CellSize) + GridSize * 0.5f),
-				FMath::FloorToInt(((InPos.Y - Origin.Y) / CellSize) + GridSize * 0.5f)
-			);
-
-			return GetCellIndex(Coords, OutIndex);
-		}
-
-		/**
-		 * Get the number of intersecting cells of the provided box
-		 *
-		 * @return the number of intersecting cells
-		 */
-		int32 GetNumIntersectingCells(const FBox& InBox) const
-		{
-			FIntVector2 MinCellCoords;
-			FIntVector2 MaxCellCoords;
-			const FBox2D Bounds2D(FVector2D(InBox.Min), FVector2D(InBox.Max));
-
-			if (GetCellCoords(Bounds2D, MinCellCoords, MaxCellCoords))
+			else
 			{
-				return (MaxCellCoords.X - MinCellCoords.X + 1) * (MaxCellCoords.Y - MinCellCoords.Y + 1);
+				FString GridInfoText = FString::Printf(TEXT("X%02d_Y%02d"), Coords.X, Coords.Y);
+				float TextWidth, TextHeight;
+				Canvas->SetDrawColor(255, 255, 0);
+				Canvas->StrLen(GEngine->GetTinyFont(), GridInfoText, TextWidth, TextHeight);
+				FVector2D CellBoundsSize = CellScreenBounds.GetSize();
+				if (TextWidth < CellBoundsSize.X && TextHeight < CellBoundsSize.Y)
+				{
+					FVector2D GridInfoPos = CellScreenBounds.GetCenter() - FVector2D(TextWidth / 2, TextHeight / 2);
+					Canvas->DrawText(GEngine->GetTinyFont(), GridInfoText, GridInfoPos.X, GridInfoPos.Y);
+				}
 			}
 
-			return 0;
-		}
-
-		// Runs a function on all cells
-		void ForEachCells(TFunctionRef<void(const FIntVector2&)> InOperation) const
-		{
-			for (int32 y = 0; y < GridSize; y++)
+			const FSpatialHashStreamingGridLayerCell& LayerCell = GridLevels[GridLevel].LayerCells[Coords.Y * Helper.Levels[GridLevel].GridSize + Coords.X];
+			const UWorldPartitionRuntimeSpatialHashCell* Cell = nullptr;
+			for (const UWorldPartitionRuntimeSpatialHashCell* GridCell : LayerCell.GridCells)
 			{
-				for (int32 x = 0; x < GridSize; x++)
+				Cell = GridCell;
+				if (!GridCell->HasDataLayers())
 				{
-					InOperation(FIntVector2(x, y));
+					break;
 				}
 			}
-		}
-
-		/**
-		 * Runs a function on all intersecting cells for the provided box
-		 *
-		 * @return the number of intersecting cells
-		 */
-		int32 ForEachIntersectingCellsBreakable(const FBox& InBox, TFunctionRef<bool(const FIntVector2&)> InOperation) const
-		{
-			int32 NumCells = 0;
-
-			FIntVector2 MinCellCoords;
-			FIntVector2 MaxCellCoords;
-			const FBox2D Bounds2D(FVector2D(InBox.Min), FVector2D(InBox.Max));
-
-			if (GetCellCoords(Bounds2D, MinCellCoords, MaxCellCoords))
-			{
-				for (int32 y = MinCellCoords.Y; y <= MaxCellCoords.Y; y++)
-				{
-					for (int32 x = MinCellCoords.X; x <= MaxCellCoords.X; x++)
-					{
-						if (!InOperation(FIntVector2(x, y)))
-						{
-							return NumCells;
-						}
-						++NumCells;
-					}
-				}
-			}
-
-			return NumCells;
-		}
-
-		int32 ForEachIntersectingCells(const FBox& InBox, TFunctionRef<void(const FIntVector2&)> InOperation) const
-		{
-			return ForEachIntersectingCellsBreakable(InBox, [InOperation](const FIntVector2& Vector) { InOperation(Vector); return true; });
-		}
-
-		/**
-		 * Runs a function on all intersecting cells for the provided sphere
-		 *
-		 * @return the number of intersecting cells
-		 */
-		int32 ForEachIntersectingCells(const FSphere& InSphere, TFunctionRef<void(const FIntVector2&)> InOperation) const
-		{
-			int32 NumCells = 0;
-
-			// @todo_ow: rasterize circle instead?
-			const FBox Box(InSphere.Center - FVector(InSphere.W), InSphere.Center + FVector(InSphere.W));
-
-			ForEachIntersectingCells(Box, [this, &InSphere, &InOperation, &NumCells](const FIntVector2& Coords)
-				{
-					const int32 CellIndex = Coords.Y * GridSize + Coords.X;
+			FLinearColor CellColor = Cell ? Cell->GetDebugColor() : FLinearColor(0.f, 0.f, 0.f, 0.25f);
+			FCanvasTileItem Item(CellScreenBounds.Min, GWhiteTexture, CellScreenBounds.GetSize(), CellColor);
+			Item.BlendMode = SE_BLEND_Translucent;
+			Canvas->DrawItem(Item);
 
-					FBox2D CellBounds;
-					GetCellBounds(CellIndex, CellBounds);
+			FCanvasBoxItem Box(CellScreenBounds.Min, CellScreenBounds.GetSize());
+			Box.SetColor(CellColor);
+			Box.BlendMode = SE_BLEND_Translucent;
+			Canvas->DrawItem(Box);
+		});
 
-					FVector2D Delta = FVector2D(InSphere.Center) - FVector2D::Max(CellBounds.GetCenter() - CellBounds.GetExtent(), FVector2D::Min(FVector2D(InSphere.Center), CellBounds.GetCenter() + CellBounds.GetExtent()));
-					if ((Delta.X * Delta.X + Delta.Y * Delta.Y) < (InSphere.W * InSphere.W))
-					{
-						InOperation(Coords);
-						NumCells++;
-					}
-				});
+		// Draw Loading Ranges
+		float Range = GetLoadingRange();
 
-			return NumCells;
-		}
-	};
+		FCanvasLineItem LineItem;
+		LineItem.LineThickness = 2;
+		LineItem.SetColor(FLinearColor::White);
 
-	struct FGridLevel : public FGrid2D
-	{
-#if WITH_EDITOR
-		struct FGridCellDataChunk
+		for (const FWorldPartitionStreamingSource& Source : Sources)
 		{
-			FGridCellDataChunk(const TArray<const UDataLayer*>& InDataLayers)
-			{
-				Algo::TransformIf(InDataLayers, DataLayers, [](const UDataLayer* DataLayer) { return DataLayer->IsDynamicallyLoaded(); }, [](const UDataLayer* DataLayer) { return DataLayer; });
-				DataLayersID = FDataLayersID(DataLayers);
-			}
-
-			void AddActor(const FGuid& Actor) { Actors.Add(Actor); }
-			const TSet<FGuid>& GetActors() const { return Actors; }
-			bool HasDataLayers() const { return !DataLayers.IsEmpty(); }
-			const TArray<const UDataLayer*>& GetDataLayers() const { return DataLayers; }
-			const FDataLayersID& GetDataLayersID() const { return DataLayersID; }
-
-		private:
-			TSet<FGuid> Actors;
-			TArray<const UDataLayer*> DataLayers;
-			FDataLayersID DataLayersID;
-		};
+			TArray<FVector> LinePoints;
+			LinePoints.SetNum(2);
 
-		struct FGridCell
-		{
-			void AddActor(const FGuid& InActor, const TArray<const UDataLayer*>& InDataLayers)
-			{
-				FDataLayersID DataLayersID = FDataLayersID(InDataLayers);
-				FGridCellDataChunk* ActorDataChunk = Algo::FindByPredicate(DataChunks, [&](FGridCellDataChunk& InDataChunk) { return InDataChunk.GetDataLayersID() == DataLayersID; });
-				if (!ActorDataChunk)
-				{
-					ActorDataChunk = &DataChunks.Emplace_GetRef(InDataLayers);
-				}
-				ActorDataChunk->AddActor(InActor);
-			}
+			float Sin, Cos;
+			FMath::SinCos(&Sin, &Cos, (63.0f / 64.0f) * 2.0f * PI);
+			FVector2D LineStart(Sin * Range, Cos * Range);
 
-			void AddActors(const TSet<FGuid>& InActors, const TArray<const UDataLayer*>& InDataLayers)
+			for (int32 i = 0; i < 64; i++)
 			{
-				for (const FGuid& Actor : InActors)
-				{
-					AddActor(Actor, InDataLayers);
-				}
+				FMath::SinCos(&Sin, &Cos, (i / 64.0f) * 2.0f * PI);
+				FVector2D LineEnd(Sin * Range, Cos * Range);
+				LineItem.Draw(CanvasObject, WorldToScreen(FVector2D(Source.Location) + LineStart), WorldToScreen(FVector2D(Source.Location) + LineEnd));
+				LineStart = LineEnd;
 			}
 
-			const TArray<FGridCellDataChunk>& GetDataChunks() const { return DataChunks; }
-
-			const FGridCellDataChunk* GetNoDataLayersDataChunk() const
+			FVector2D SourceDir = FVector2D(Source.Rotation.Vector());
+			if (SourceDir.Size())
 			{
-				for (const FGridCellDataChunk& DataChunk : DataChunks)
-				{
-					if (!DataChunk.HasDataLayers())
-					{
-						return &DataChunk;
-					}
-				}
-				return nullptr;
+				SourceDir.Normalize();
+				FVector2D ConeCenter(FVector2D(Source.Location));
+				LineItem.Draw(CanvasObject, WorldToScreen(ConeCenter), WorldToScreen(ConeCenter + SourceDir * Range));
 			}
-
-		private:
-			TArray<FGridCellDataChunk> DataChunks;
-		};
-#endif
-
-		inline FGridLevel(const FVector2D& InOrigin, int32 InCellSize, int32 InGridSize)
-			: FGrid2D(InOrigin, InCellSize, InGridSize)
-		{
-#if WITH_EDITOR
-			Cells.InsertDefaulted(0, GridSize * GridSize);
-#endif
-		}
-
-#if WITH_EDITOR
-		/**
-		 * Returns the cell at the specified grid coordinate
-		 *
-		 * @return the cell at the specified grid coordinate
-		 */
-		inline FGridCell& GetCell(const FIntVector2& InCoords)
-		{
-			check(IsValidCoords(InCoords));
-			uint32 CellIndex;
-			GetCellIndex(InCoords, CellIndex);
-			return Cells[CellIndex];
-		}
-
-		/**
-		 * Returns the cell at the specified grid coordinate
-		 *
-		 * @return the cell at the specified grid coordinate
-		 */
-		inline const FGridCell& GetCell(const FIntVector2& InCoords) const
-		{
-			check(IsValidCoords(InCoords));
-			uint32 CellIndex;
-			GetCellIndex(InCoords, CellIndex);
-			return Cells[CellIndex];
-		}
-
-		TArray<FGridCell> Cells;
-#endif
-	};
-
-	FSquare2DGridHelper(int32 InNumLevels, const FVector& InOrigin, int32 InCellSize, int32 InGridSize)
-		: Origin(InOrigin)
-		, CellSize(InCellSize)
-		, GridSize(InGridSize)
-	{
-		Levels.Reserve(InNumLevels);
-
-		int32 CurrentCellSize = CellSize;
-		int32 CurrentGridSize = GridSize;
-
-		const FVector2D BaseLevelOffset = 0.5 * FVector2D(CellSize, CellSize);
-		for (int32 Level = 0; Level < InNumLevels; ++Level)
-		{
-			// Add offset on origin based on level's cell size to break pattern of perfectly aligned cell edges at multiple level.
-			// This will prevent weird artefact during actor promotion.
-			// Apply base level offset so that first level isn't offset.
-			const FVector2D GridLevelOffset = 0.5f * FVector2D(CurrentCellSize, CurrentCellSize) - BaseLevelOffset;
-			const FVector2D LevelOrigin = FVector2D(InOrigin) + GridLevelOffset;
-
-			Levels.Emplace(LevelOrigin, CurrentCellSize, CurrentGridSize);
-
-			CurrentCellSize <<= 1;
-			CurrentGridSize >>= 1;
-		}
-	}
-
-#if WITH_EDITOR
-	// Returns the lowest grid level
-	inline FGridLevel& GetLowestLevel() { return Levels[0]; }
-
-	// Returns the always loaded (top level) cell
-	inline FGridLevel::FGridCell& GetAlwaysLoadedCell() { return Levels.Last().Cells[0]; }
-
-	// Returns the always loaded (top level) cell
-	inline const FGridLevel::FGridCell& GetAlwaysLoadedCell() const { return Levels.Last().Cells[0]; }
-
-	// Returns the cell at the given coord
-	inline const FGridLevel::FGridCell& GetCell(const FIntVector& InCoords) const { return Levels[InCoords.Z].GetCell(FIntVector2(InCoords.X, InCoords.Y)); }
-#endif
-
-	/**
-	 * Returns the cell bounds
-	 *
-	 * @return true if the specified coord was valid
-	 */
-	inline bool GetCellBounds(const FIntVector& InCoords, FBox2D& OutBounds) const
-	{
-		if (Levels.IsValidIndex(InCoords.Z))
-		{
-			return Levels[InCoords.Z].GetCellBounds(FIntVector2(InCoords.X, InCoords.Y), OutBounds);
-		}
-		return false;
-	}
-
-	// Runs a function on all cells
-	void ForEachCells(TFunctionRef<void(const FIntVector&)> InOperation) const
-	{
-		for (int32 Level = 0; Level < Levels.Num(); Level++)
-		{
-			Levels[Level].ForEachCells([Level, &InOperation](const FIntVector2& Coord) { InOperation(FIntVector(Coord.X, Coord.Y, Level)); });
-		}
-	}
-
-	/**
-	 * Runs a function on all intersecting cells for the provided box
-	 *
-	 * @return the number of intersecting cells
-	 */
-	int32 ForEachIntersectingCells(const FBox& InBox, TFunctionRef<void(const FIntVector&)> InOperation) const
-	{
-		int32 NumCells = 0;
-
-		for (int32 Level = 0; Level < Levels.Num(); Level++)
-		{
-			NumCells += Levels[Level].ForEachIntersectingCells(InBox, [InBox, Level, InOperation](const FIntVector2& Coord) { InOperation(FIntVector(Coord.X, Coord.Y, Level)); });
-		}
-
-		return NumCells;
-	}
-
-	/**
-	 * Runs a function on all intersecting cells for the provided sphere
-	 *
-	 * @return the number of intersecting cells
-	 */
-	int32 ForEachIntersectingCells(const FSphere& InSphere, TFunctionRef<void(const FIntVector&)> InOperation) const
-	{
-		int32 NumCells = 0;
-
-		for (int32 Level = 0; Level < Levels.Num(); Level++)
-		{
-			NumCells += Levels[Level].ForEachIntersectingCells(InSphere, [InSphere, Level, InOperation](const FIntVector2& Coord) { InOperation(FIntVector(Coord.X, Coord.Y, Level)); });
 		}
 
-		return NumCells;
-	}
-
-#if WITH_EDITOR
-	// Validates that actor is not referenced by multiple cells
-	void ValidateSingleActorReferer()
-	{
-		UE_SCOPED_TIMER(TEXT("ValidateSingleActorReferer"), LogWorldPartitionRuntimeSpatialHash);
-
-		TSet<FGuid> ActorUsage;
-		for (int32 Level = 0; Level < Levels.Num() - 1; Level++)
-		{
-			const int32 CurrentGridSize = Levels[Level].GridSize;
-			for (int32 y = 0; y < CurrentGridSize; y++)
-			{
-				for (int32 x = 0; x < CurrentGridSize; x++)
-				{
-					FGridLevel::FGridCell& ThisCell = Levels[Level].Cells[y * CurrentGridSize + x];
-					for (const FGridLevel::FGridCellDataChunk& DataChunk : ThisCell.GetDataChunks())
-					{
-						for (const FGuid& ActorGuid : DataChunk.GetActors())
-						{
-							bool bWasAlreadyInSet;
-							ActorUsage.Add(ActorGuid, &bWasAlreadyInSet);
-							check(!bWasAlreadyInSet);
-						}
-					}
-				}
-			}
-		}
+		FCanvasBoxItem Box(GridScreenBounds.Min, GridScreenBounds.GetSize());
+		Box.SetColor(DebugColor);
+		Canvas->DrawItem(Box);
 	}
-#endif //WITH_EDITOR
-
-public:
-	FVector Origin;
-	int32 CellSize;
-	int32 GridSize;
-	TArray<FGridLevel> Levels;
-};
-
-// ------------------------------------------------------------------------------------------------
-FSpatialHashStreamingGrid::FSpatialHashStreamingGrid()
-	: Origin(ForceInitToZero)
-	, CellSize(0)
-	, GridSize(0)
-	, LoadingRange(0.0f)
-#if WITH_EDITOR
-	, DebugColor(ForceInitToZero)
-	, OverrideLoadingRange(0)
-#endif
-	, GridHelper(nullptr)
-{
 }
 
-FSpatialHashStreamingGrid::~FSpatialHashStreamingGrid()
-{
-	if (GridHelper)
-	{
-		delete GridHelper;
-	}
-}
+// ------------------------------------------------------------------------------------------------
 
-const FSquare2DGridHelper& FSpatialHashStreamingGrid::GetGridHelper() const
+ASpatialHashRuntimeGridInfo::ASpatialHashRuntimeGridInfo(const FObjectInitializer& ObjectInitializer)
+	: Super(ObjectInitializer)
 {
-	if (!GridHelper)
-	{
-		GridHelper = new FSquare2DGridHelper(GridLevels.Num(), Origin, CellSize, GridSize);
-	}
-	else
-	{
-		check(GridHelper->Levels.Num() == GridLevels.Num());
-		check(GridHelper->Origin == Origin);
-		check(GridHelper->CellSize == CellSize);
-		check(GridHelper->GridSize == GridSize);
-	}
-	return *GridHelper;
+#if WITH_EDITORONLY_DATA
+	bListedInSceneOutliner = false;
+#endif
 }
 
-void FSpatialHashStreamingGrid::GetCells(const TArray<FWorldPartitionStreamingSource>& Sources, const UDataLayerSubsystem* DataLayerSubsystem, TSet<const UWorldPartitionRuntimeCell*>& Cells) const
-{
-	const FSquare2DGridHelper& Helper = GetGridHelper();
-	for (const FWorldPartitionStreamingSource& Source : Sources)
-	{
-		const FSphere GridSphere(Source.Location, GetLoadingRange());
-		Helper.ForEachIntersectingCells(GridSphere, [&](const FIntVector& Coords)
-		{
-			const FSpatialHashStreamingGridLayerCell& LayerCell = GridLevels[Coords.Z].LayerCells[Coords.Y * Helper.Levels[Coords.Z].GridSize + Coords.X];
-			for (const UWorldPartitionRuntimeCell* Cell : LayerCell.GridCells)
-			{
-				if (!Cell->HasDataLayers() || (DataLayerSubsystem && DataLayerSubsystem->IsAnyDataLayerActive(Cell->GetDataLayers())))
-				{
-					Cells.Add(Cell);
-				}
-			}
-		});
-	}
-
-	GetAlwaysLoadedCells(DataLayerSubsystem, Cells);
-}
+UWorldPartitionRuntimeSpatialHash::UWorldPartitionRuntimeSpatialHash(const FObjectInitializer& ObjectInitializer)
+	: Super(ObjectInitializer)
+{}
 
-void FSpatialHashStreamingGrid::GetAlwaysLoadedCells(const UDataLayerSubsystem* DataLayerSubsystem, TSet<const UWorldPartitionRuntimeCell*>& Cells) const
+#if WITH_EDITOR
+FName UWorldPartitionRuntimeSpatialHash::GetActorRuntimeGrid(const AActor* Actor) const
 {
-	if (GridLevels.Num() > 0)
+	if (ULevel* Level = Actor ? Actor->GetLevel() : nullptr)
 	{
-		const int32 TopLevel = GridLevels.Num() - 1;
-		for (const FSpatialHashStreamingGridLayerCell& LayerCell : GridLevels[TopLevel].LayerCells)
+		if (const FName* ActorRuntimeGrid = WorldCompositionStreamingLevelToRuntimeGrid.Find(FLevelUtils::FindStreamingLevel(Level)))
 		{
-			for (const UWorldPartitionRuntimeCell* Cell : LayerCell.GridCells)
-			{
-				if (!Cell->HasDataLayers() || (DataLayerSubsystem && DataLayerSubsystem->IsAnyDataLayerActive(Cell->GetDataLayers())))
-				{
-					check(Cell->IsAlwaysLoaded() || Cell->HasDataLayers());
-					Cells.Add(Cell);
-				}
-			}
+			return *ActorRuntimeGrid;
 		}
 	}
+	return Super::GetActorRuntimeGrid(Actor);
 }
 
-void FSpatialHashStreamingGrid::Draw3D(UWorld* World, const TArray<FWorldPartitionStreamingSource>& Sources, const FTransform& Transform) const
-{
-	const FSquare2DGridHelper& Helper = GetGridHelper();
-	int32 MinGridLevel = FMath::Clamp<int32>(GShowRuntimeSpatialHashGridLevel, 0, GridLevels.Num() - 1);
-	int32 MaxGridLevel = FMath::Clamp<int32>(MinGridLevel + GShowRuntimeSpatialHashGridLevelCount - 1, 0, GridLevels.Num() - 1);
-	const float GridViewMinimumSizeInCellCount = 5.f;
-	const float GridViewLoadingRangeExtentRatio = 1.5f;
-	const float Radius = GetLoadingRange();
-	const float GridSideDistance = FMath::Max((2.f * Radius * GridViewLoadingRangeExtentRatio), CellSize * GridViewMinimumSizeInCellCount);
-
-	for (const FWorldPartitionStreamingSource& Source : Sources)
-	{
-		FVector StartTrace = Source.Location + FVector(0.f, 0.f, 100.f);
-		FVector EndTrace = StartTrace - FVector(0.f, 0.f, 1000000.f);
-		float Z = Source.Location.Z;
-		FHitResult Hit;
-		if (World->LineTraceSingleByObjectType(Hit, StartTrace, EndTrace, FCollisionObjectQueryParams(ECC_WorldStatic), FCollisionQueryParams(SCENE_QUERY_STAT(DebugWorldPartitionTrace), true)))
-		{
-			Z = Hit.ImpactPoint.Z;
-		}
-
-		FSphere Sphere(Source.Location, GridSideDistance * 0.5f);
-		const FBox Region(Sphere.Center - Sphere.W, Sphere.Center + Sphere.W);
-		for (int32 GridLevel = MinGridLevel; GridLevel <= MaxGridLevel; ++GridLevel)
-		{
-			Helper.Levels[GridLevel].ForEachIntersectingCells(Region, [&](const FIntVector2& Coords)
-			{
-				FBox2D CellWorldBounds;
-				Helper.Levels[GridLevel].GetCellBounds(FIntVector2(Coords.X, Coords.Y), CellWorldBounds);
-				
-				FVector BoundsExtent(CellWorldBounds.GetExtent(), 100.f);
-				const FSpatialHashStreamingGridLayerCell& LayerCell = GridLevels[GridLevel].LayerCells[Coords.Y * Helper.Levels[GridLevel].GridSize + Coords.X];
-				const UWorldPartitionRuntimeSpatialHashCell* Cell = nullptr;
-				for (const UWorldPartitionRuntimeSpatialHashCell* GridCell : LayerCell.GridCells)
-				{
-					Cell = GridCell;
-					if (!GridCell->HasDataLayers())
-					{
-						break;
-					}
-				}
-				FColor CellColor = Cell ? Cell->GetDebugColor().ToFColor(false).WithAlpha(16) : FColor(0, 0, 0, 16);
-				FVector BoundsOrigin(CellWorldBounds.GetCenter(), Z);
-				FBox Box = FBox::BuildAABB(BoundsOrigin, BoundsExtent);
-				DrawDebugSolidBox(World, Box, CellColor, Transform, false, -1.f, 255);
-				FVector Pos = Transform.TransformPosition(BoundsOrigin);
-				DrawDebugBox(World, Pos, BoundsExtent, Transform.GetRotation(), CellColor.WithAlpha(255), false, -1.f, 255, 10.f);
-			});
-		}
-
-		// Draw Loading Ranges
-		FVector SphereLocation(FVector2D(Source.Location), Z);
-		SphereLocation = Transform.TransformPosition(SphereLocation);
-		DrawDebugSphere(World, SphereLocation, Radius, 32, FColor::White, false, -1.f, 0, 20.f);
-	}
-}
-
-void FSpatialHashStreamingGrid::Draw2D(UCanvas* Canvas, const TArray<FWorldPartitionStreamingSource>& Sources, const FBox& Region, const FBox2D& GridScreenBounds, TFunctionRef<FVector2D(const FVector2D&)> WorldToScreen) const
-{
-	FCanvas* CanvasObject = Canvas->Canvas;
-	int32 MinGridLevel = FMath::Clamp<int32>(GShowRuntimeSpatialHashGridLevel, 0, GridLevels.Num() - 1);
-	int32 MaxGridLevel = FMath::Clamp<int32>(MinGridLevel + GShowRuntimeSpatialHashGridLevelCount - 1, 0, GridLevels.Num() - 1);
-
-	for (int32 GridLevel = MinGridLevel; GridLevel <= MaxGridLevel; ++GridLevel)
-	{
-		// Draw X/Y Axis
-		{
-			FCanvasLineItem Axis;
-			Axis.LineThickness = 3;
-			{
-				Axis.SetColor(FLinearColor::Green);
-				FVector2D LineStart = WorldToScreen(FVector2D(-1638400.f, 0.f));
-				FVector2D LineEnd = WorldToScreen(FVector2D(1638400.f, 0.f));
-				LineStart.X = FMath::Clamp(LineStart.X, GridScreenBounds.Min.X, GridScreenBounds.Max.X);
-				LineStart.Y = FMath::Clamp(LineStart.Y, GridScreenBounds.Min.Y, GridScreenBounds.Max.Y);
-				LineEnd.X = FMath::Clamp(LineEnd.X, GridScreenBounds.Min.X, GridScreenBounds.Max.X);
-				LineEnd.Y = FMath::Clamp(LineEnd.Y, GridScreenBounds.Min.Y, GridScreenBounds.Max.Y);
-				Axis.Draw(CanvasObject, LineStart, LineEnd);
-			}
-			{
-				Axis.SetColor(FLinearColor::Red);
-				FVector2D LineStart = WorldToScreen(FVector2D(0.f, -1638400.f));
-				FVector2D LineEnd = WorldToScreen(FVector2D(0.f, 1638400.f));
-				LineStart.X = FMath::Clamp(LineStart.X, GridScreenBounds.Min.X, GridScreenBounds.Max.X);
-				LineStart.Y = FMath::Clamp(LineStart.Y, GridScreenBounds.Min.Y, GridScreenBounds.Max.Y);
-				LineEnd.X = FMath::Clamp(LineEnd.X, GridScreenBounds.Min.X, GridScreenBounds.Max.X);
-				LineEnd.Y = FMath::Clamp(LineEnd.Y, GridScreenBounds.Min.Y, GridScreenBounds.Max.Y);
-				Axis.Draw(CanvasObject, LineStart, LineEnd);
-			}
-		}
-
-		// Draw Grid cells at desired grid level
-		const FSquare2DGridHelper& Helper = GetGridHelper();
-		Helper.Levels[GridLevel].ForEachIntersectingCells(Region, [&](const FIntVector2& Coords)
-		{
-			FBox2D CellWorldBounds;
-			Helper.Levels[GridLevel].GetCellBounds(FIntVector2(Coords.X, Coords.Y), CellWorldBounds);
-			FBox2D CellScreenBounds = FBox2D(WorldToScreen(CellWorldBounds.Min), WorldToScreen(CellWorldBounds.Max));
-			// Clamp inside grid bounds
-			if (!GridScreenBounds.IsInside(CellScreenBounds))
-			{
-				CellScreenBounds.Min.X = FMath::Clamp(CellScreenBounds.Min.X, GridScreenBounds.Min.X, GridScreenBounds.Max.X);
-				CellScreenBounds.Min.Y = FMath::Clamp(CellScreenBounds.Min.Y, GridScreenBounds.Min.Y, GridScreenBounds.Max.Y);
-				CellScreenBounds.Max.X = FMath::Clamp(CellScreenBounds.Max.X, GridScreenBounds.Min.X, GridScreenBounds.Max.X);
-				CellScreenBounds.Max.Y = FMath::Clamp(CellScreenBounds.Max.Y, GridScreenBounds.Min.Y, GridScreenBounds.Max.Y);
-			}
-			else
-			{
-				FString GridInfoText = FString::Printf(TEXT("X%02d_Y%02d"), Coords.X, Coords.Y);
-				float TextWidth, TextHeight;
-				Canvas->SetDrawColor(255, 255, 0);
-				Canvas->StrLen(GEngine->GetTinyFont(), GridInfoText, TextWidth, TextHeight);
-				FVector2D CellBoundsSize = CellScreenBounds.GetSize();
-				if (TextWidth < CellBoundsSize.X && TextHeight < CellBoundsSize.Y)
-				{
-					FVector2D GridInfoPos = CellScreenBounds.GetCenter() - FVector2D(TextWidth / 2, TextHeight / 2);
-					Canvas->DrawText(GEngine->GetTinyFont(), GridInfoText, GridInfoPos.X, GridInfoPos.Y);
-				}
-			}
-
-			const FSpatialHashStreamingGridLayerCell& LayerCell = GridLevels[GridLevel].LayerCells[Coords.Y * Helper.Levels[GridLevel].GridSize + Coords.X];
-			const UWorldPartitionRuntimeSpatialHashCell* Cell = nullptr;
-			for (const UWorldPartitionRuntimeSpatialHashCell* GridCell : LayerCell.GridCells)
-			{
-				Cell = GridCell;
-				if (!GridCell->HasDataLayers())
-				{
-					break;
-				}
-			}
-			FLinearColor CellColor = Cell ? Cell->GetDebugColor() : FLinearColor(0.f, 0.f, 0.f, 0.25f);
-			FCanvasTileItem Item(CellScreenBounds.Min, GWhiteTexture, CellScreenBounds.GetSize(), CellColor);
-			Item.BlendMode = SE_BLEND_Translucent;
-			Canvas->DrawItem(Item);
-
-			FCanvasBoxItem Box(CellScreenBounds.Min, CellScreenBounds.GetSize());
-			Box.SetColor(CellColor);
-			Box.BlendMode = SE_BLEND_Translucent;
-			Canvas->DrawItem(Box);
-		});
-
-		// Draw Loading Ranges
-		float Range = GetLoadingRange();
-
-		FCanvasLineItem LineItem;
-		LineItem.LineThickness = 2;
-		LineItem.SetColor(FLinearColor::White);
-
-		for (const FWorldPartitionStreamingSource& Source : Sources)
-		{
-			TArray<FVector> LinePoints;
-			LinePoints.SetNum(2);
-
-			float Sin, Cos;
-			FMath::SinCos(&Sin, &Cos, (63.0f / 64.0f) * 2.0f * PI);
-			FVector2D LineStart(Sin * Range, Cos * Range);
-
-			for (int32 i = 0; i < 64; i++)
-			{
-				FMath::SinCos(&Sin, &Cos, (i / 64.0f) * 2.0f * PI);
-				FVector2D LineEnd(Sin * Range, Cos * Range);
-				LineItem.Draw(CanvasObject, WorldToScreen(FVector2D(Source.Location) + LineStart), WorldToScreen(FVector2D(Source.Location) + LineEnd));
-				LineStart = LineEnd;
-			}
-
-			FVector2D SourceDir = FVector2D(Source.Rotation.Vector());
-			if (SourceDir.Size())
-			{
-				SourceDir.Normalize();
-				FVector2D ConeCenter(FVector2D(Source.Location));
-				LineItem.Draw(CanvasObject, WorldToScreen(ConeCenter), WorldToScreen(ConeCenter + SourceDir * Range));
-			}
-		}
-
-		FCanvasBoxItem Box(GridScreenBounds.Min, GridScreenBounds.GetSize());
-		Box.SetColor(DebugColor);
-		Canvas->DrawItem(Box);
-	}
-}
-
-// ------------------------------------------------------------------------------------------------
-
-#if WITH_EDITOR
-static FSquare2DGridHelper GetGridHelper(const FBox& WorldBounds, int32 GridCellSize)
-{
-	// Default grid to a minimum of 1 level and 1 cell, for always loaded actors
-	FVector GridOrigin = FVector::ZeroVector;
-	int32 GridSize = 1;
-	int32 GridLevelCount = 1;
-
-	float WorldBoundsMaxExtent = 0.f;
-	// If World bounds is valid, compute Grid's size and level count based on it
-	if (WorldBounds.IsValid)
-	{
-		const FVector2D DistMin = FMath::Abs(FVector2D(WorldBounds.Min - GridOrigin));
-		const FVector2D DistMax = FMath::Abs(FVector2D(WorldBounds.Max - GridOrigin));
-		WorldBoundsMaxExtent = FMath::Max(DistMin.GetMax(), DistMax.GetMax());
-	}
-	if (WorldBoundsMaxExtent > 0.f)
-	{
-		GridSize = 2.f * FMath::CeilToFloat(WorldBoundsMaxExtent / GridCellSize);
-		if (!FMath::IsPowerOfTwo(GridSize))
-		{
-			GridSize = FMath::Pow(2.f, FMath::CeilToFloat(FMath::Log2(static_cast<float>(GridSize))));
-		}
-		GridLevelCount = FMath::Log2(static_cast<float>(GridSize)) + 1;
-	}
-	else
-	{
-		UE_LOG(LogWorldPartitionRuntimeSpatialHash, Warning, TEXT("Invalid world bounds, grid partitioning will use a runtime grid with 1 cell."));
-	}
-
-	return FSquare2DGridHelper(GridLevelCount, GridOrigin, GridCellSize, GridSize);
-}
-
-static FSquare2DGridHelper GetPartitionedActors(const UWorldPartition* WorldPartition, const FBox& WorldBounds, const FSpatialHashRuntimeGrid& Grid, const TArray<FActorCluster>& GridActors)
-{
-	UE_SCOPED_TIMER(TEXT("GetPartitionedActors"), LogWorldPartitionRuntimeSpatialHash);
-
-	//
-	// Create the hierarchical grids for the game
-	//	
-	FSquare2DGridHelper PartitionedActors = GetGridHelper(WorldBounds, Grid.CellSize);
-
-	for (const FActorCluster& ActorCluster : GridActors)
-	{
-		check(ActorCluster.Actors.Num() > 0);
-
-		EActorGridPlacement GridPlacement = ActorCluster.GridPlacement;
-		bool bAlwaysLoadedPromotedCluster = (GridPlacement == EActorGridPlacement::None);
-		bool bAlwaysLoadedPromotedOutOfGrid = false;
-
-		if (bAlwaysLoadedPromotedCluster)
-		{
-			GridPlacement = EActorGridPlacement::AlwaysLoaded;
-		}
-
-		switch (GridPlacement)
-		{
-		case EActorGridPlacement::Location:
-		{
-			FIntVector2 CellCoords;
-			check(ActorCluster.Actors.Num() == 1);
-			const FGuid& ActorGuid = *ActorCluster.Actors.CreateConstIterator();
-			const FWorldPartitionActorDesc& ActorDesc = *WorldPartition->GetActorDesc(ActorGuid);
-			if (PartitionedActors.GetLowestLevel().GetCellCoords(FVector2D(ActorDesc.GetOrigin()), CellCoords))
-			{
-				TArray<const UDataLayer*> ActorDataLayers;
-				if (const AWorldDataLayers* WorldDataLayers = AWorldDataLayers::Get(WorldPartition->GetWorld()))
-				{
-					for (const FName& DataLayerName : ActorDesc.GetDataLayers())
-					{
-						if (const UDataLayer* DataLayer = WorldDataLayers->GetDataLayerFromName(DataLayerName))
-						{
-							ActorDataLayers.Add(DataLayer);
-						}
-					}
-				}
-				PartitionedActors.GetLowestLevel().GetCell(CellCoords).AddActor(ActorGuid, ActorDataLayers);
-			}
-			else
-			{
-				GridPlacement = EActorGridPlacement::AlwaysLoaded;
-				bAlwaysLoadedPromotedOutOfGrid = true;
-			}
-			break;
-		}
-		case EActorGridPlacement::Bounds:
-		{
-			// Find grid level cell that encompasses the actor cluster and put actors in it.
-			bool bFoundCell = false;
-			for (FSquare2DGridHelper::FGridLevel& GridLevel : PartitionedActors.Levels)
-			{
-				int32 IntersectingCellCount = 0;
-				GridLevel.ForEachIntersectingCellsBreakable(ActorCluster.Bounds, [&IntersectingCellCount](const FIntVector2& Coords) { return ++IntersectingCellCount <= 1; });
-				if (IntersectingCellCount == 1)
-				{
-					GridLevel.ForEachIntersectingCells(ActorCluster.Bounds, [&GridLevel,&ActorCluster](const FIntVector2& Coords)
-					{
-						GridLevel.GetCell(Coords).AddActors(ActorCluster.Actors, ActorCluster.DataLayers);
-					});
-					bFoundCell = true;
-					break;
-				}
-			}
-			if (!bFoundCell)
-			{
-				GridPlacement = EActorGridPlacement::AlwaysLoaded;
-				bAlwaysLoadedPromotedOutOfGrid = true;
-			}
-			break;
-		}
-		case EActorGridPlacement::AlwaysLoaded:
-		{
-			PartitionedActors.GetAlwaysLoadedCell().AddActors(ActorCluster.Actors, ActorCluster.DataLayers);
-			break;
-		}
-		default:	
-			check(0);
-		}
-
-		if (UE_LOG_ACTIVE(LogWorldPartitionRuntimeSpatialHash, Verbose))
-		{
-			if (ActorCluster.Actors.Num() > 1)
-			{
-				static UEnum* ActorGridPlacementEnum = FindObject<UEnum>(ANY_PACKAGE, TEXT("EActorGridPlacement"));
-
-				UE_LOG(LogWorldPartitionRuntimeSpatialHash, Verbose, TEXT("Clustered %d actors (%s%s%s), generated shared BV of [%d x %d] (meters)"), 
-					ActorCluster.Actors.Num(),
-					*ActorGridPlacementEnum->GetNameStringByValue((int64)GridPlacement),
-					bAlwaysLoadedPromotedCluster ? TEXT(":PromotedCluster") : TEXT(""),
-					bAlwaysLoadedPromotedOutOfGrid ? TEXT(":PromotedOutOfGrid") : TEXT(""),
-					(int)(0.01f * ActorCluster.Bounds.GetSize().X),
-					(int)(0.01f * ActorCluster.Bounds.GetSize().Y));
-
-				for (const FGuid& ActorGuid : ActorCluster.Actors)
-				{
-					const FWorldPartitionActorDesc& Desc = *WorldPartition->GetActorDesc(ActorGuid);
-					UE_LOG(LogWorldPartitionRuntimeSpatialHash, Verbose, TEXT("   - Actor: %s (%s)"), *Desc.GetActorPath().ToString(), *ActorGuid.ToString(EGuidFormats::UniqueObjectGuid));
-					UE_LOG(LogWorldPartitionRuntimeSpatialHash, Verbose, TEXT("            %s"), *Desc.GetActorPackage().ToString());
-				}
-			}
-		}
-	}
-
-	// Perform validation
-	PartitionedActors.ValidateSingleActorReferer();
-	
-	return PartitionedActors;
-}
-#endif //WITH_EDITOR
-
-ASpatialHashRuntimeGridInfo::ASpatialHashRuntimeGridInfo(const FObjectInitializer& ObjectInitializer)
-	: Super(ObjectInitializer)
-{
-#if WITH_EDITORONLY_DATA
-	bListedInSceneOutliner = false;
-#endif
-}
-
-UWorldPartitionRuntimeSpatialHash::UWorldPartitionRuntimeSpatialHash(const FObjectInitializer& ObjectInitializer)
-	: Super(ObjectInitializer)
-{}
-
-#if WITH_EDITOR
-FName UWorldPartitionRuntimeSpatialHash::GetActorRuntimeGrid(const AActor* Actor) const
-{
-	if (ULevel* Level = Actor ? Actor->GetLevel() : nullptr)
-	{
-		if (const FName* ActorRuntimeGrid = WorldCompositionStreamingLevelToRuntimeGrid.Find(FLevelUtils::FindStreamingLevel(Level)))
-		{
-			return *ActorRuntimeGrid;
-		}
-	}
-	return Super::GetActorRuntimeGrid(Actor);
-}
-
-void UWorldPartitionRuntimeSpatialHash::SetDefaultValues()
+void UWorldPartitionRuntimeSpatialHash::SetDefaultValues()
 {
 	FSpatialHashRuntimeGrid& MainGrid = Grids.AddDefaulted_GetRef();
 	MainGrid.GridName = TEXT("MainGrid");
 	MainGrid.CellSize = 3200;
 	MainGrid.LoadingRange = 25600;
-	MainGrid.DebugColor = FLinearColor::Gray;
-}
-
-void UWorldPartitionRuntimeSpatialHash::ImportFromWorldComposition(UWorldComposition* WorldComposition)
-{
-	check(IsRunningCommandlet());
-
-	if (WorldComposition)
-	{
-		const TArray<FWorldTileLayer> WorldCompositionTileLayers = WorldComposition->GetDistanceDependentLayers();
-		for (const FWorldTileLayer& Layer : WorldCompositionTileLayers)
-		{
-			FName GridName = FName(Layer.Name);
-			FSpatialHashRuntimeGrid* Grid = Algo::FindByPredicate(Grids, [GridName](const FSpatialHashRuntimeGrid& Grid) { return Grid.GridName == GridName; });
-			if (!Grid)
-			{
-				Grid = &Grids.AddDefaulted_GetRef();
-				Grid->GridName = GridName;
-				Grid->CellSize = 3200;
-				Grid->DebugColor = FLinearColor::MakeRandomColor();
-			}
-			// World Composition Layer Streaming Distance always wins over existing value (config file)
-			Grid->LoadingRange = Layer.StreamingDistance;
-		}
-
-		const UWorldComposition::FTilesList& Tiles = WorldComposition->GetTilesList();
-		for (int32 TileIdx = 0; TileIdx < Tiles.Num(); TileIdx++)
-		{
-			const FWorldCompositionTile& Tile = Tiles[TileIdx];
-			ULevelStreaming* StreamingLevel = WorldComposition->TilesStreaming[TileIdx];
-			if (StreamingLevel && WorldComposition->IsDistanceDependentLevel(Tile.PackageName))
-			{
-				// Map WorldComposition tiles streaming level to Runtime Grid
-				WorldCompositionStreamingLevelToRuntimeGrid.Add(StreamingLevel, FName(Tile.Info.Layer.Name));
-			}
-		}
-	}
-}
-
-bool UWorldPartitionRuntimeSpatialHash::GenerateStreaming(EWorldPartitionStreamingMode Mode, UWorldPartitionStreamingPolicy* StreamingPolicy, TArray<FString>* OutPackagesToGenerate)
-{
-	TRACE_CPUPROFILER_EVENT_SCOPE(UWorldPartitionRuntimeSpatialHash::GenerateStreaming);
-	UWorldPartition* WorldPartition = GetOuterUWorldPartition();
-
-	UE_SCOPED_TIMER(TEXT("GenerateStreaming"), LogWorldPartitionRuntimeSpatialHash);
-	
-	if (!Grids.Num())
-	{
-		UE_LOG(LogWorldPartitionRuntimeSpatialHash, Error, TEXT("Invalid partition grids setup"));
-		return false;
-	}
-
-	TArray<FSpatialHashRuntimeGrid> AllGrids;
-	AllGrids.Append(Grids);
-
-	// Append grids from runtime grid actors
-	for (AActor* Actor : GetWorld()->PersistentLevel->Actors)
-	{
-		if (ASpatialHashRuntimeGridInfo* RuntimeGridActor = Cast<ASpatialHashRuntimeGridInfo>(Actor))
-		{
-			AllGrids.Add(RuntimeGridActor->GridSettings);
-		}
-	}
-
-	check(!StreamingGrids.Num());
-
-	// Build a map of Actor GUID -> HLODActor GUID once instead of having to recompute for every streaming grid we create
-	CacheHLODParents();
-
-	TMap<FName, int32> GridsMapping;
-	GridsMapping.Add(NAME_None, 0);
-	for (int32 i = 0; i < AllGrids.Num(); i++)
-	{
-		const FSpatialHashRuntimeGrid& Grid = AllGrids[i];
-		check(!GridsMapping.Contains(Grid.GridName));
-		GridsMapping.Add(Grid.GridName, i);
-	}
-
-	// Create actor clusters
-	TArray<FActorCluster> ActorClusters = CreateActorClusters(WorldPartition);
-
-	TArray<TArray<FActorCluster>> GridActors;
-	GridActors.InsertDefaulted(0, AllGrids.Num());
-
-	for (FActorCluster& ActorCluster : ActorClusters)
-	{
-		int32* FoundIndex = GridsMapping.Find(ActorCluster.RuntimeGrid);
-		if (!FoundIndex)
-		{
-			UE_LOG(LogWorldPartitionRuntimeSpatialHash, Error, TEXT("Invalid partition grid '%s' referenced by actor cluster"), *ActorCluster.RuntimeGrid.ToString());
-		}
-
-		int32 GridIndex = FoundIndex ? *FoundIndex : 0;
-		GridActors[GridIndex].Add(MoveTemp(ActorCluster));
-	}
-	
-	const FBox WorldBounds = WorldPartition->GetWorldBounds();
-	for (int32 GridIndex=0; GridIndex < AllGrids.Num(); GridIndex++)
-	{
-		const FSpatialHashRuntimeGrid& Grid = AllGrids[GridIndex];
-		const FSquare2DGridHelper PartionedActors = GetPartitionedActors(WorldPartition, WorldBounds, Grid, GridActors[GridIndex]);
-		if (!CreateStreamingGrid(Grid, PartionedActors, Mode, StreamingPolicy, OutPackagesToGenerate))
-		{
-			return false;
-		}
-	}
-
-	return true;
-}
-
-static FName GetCellName(UWorldPartition* WorldPartition, FName InGridName, int32 InLevel, int32 InCellX, int32 InCellY, const FDataLayersID& InDataLayerID)
-{
-	const FString PackageName = FPackageName::GetShortName(WorldPartition->GetPackage());
-	const FString PackageNameNoPIEPrefix = UWorld::RemovePIEPrefix(PackageName);
-
-	return FName(*FString::Printf(TEXT("WPRT_%s_%s_Cell_L%d_X%02d_Y%02d_DL%X"), *PackageNameNoPIEPrefix, *InGridName.ToString(), InLevel, InCellX, InCellY, InDataLayerID.GetHash()));
-}
-
-FName UWorldPartitionRuntimeSpatialHash::GetCellName(FName InGridName, int32 InLevel, int32 InCellX, int32 InCellY, const FDataLayersID& InDataLayerID) const
-{
-	UWorldPartition* WorldPartition = GetOuterUWorldPartition();
-	return ::GetCellName(WorldPartition, InGridName, InLevel, InCellX, InCellY, InDataLayerID);
-}
-
-void UWorldPartitionRuntimeSpatialHash::CacheHLODParents()
-{
-	CachedHLODParents.Reset();
-
-	UWorldPartition* WorldPartition = GetOuterUWorldPartition();
-
-	for (TWorldPartitionActorDescIterator<AWorldPartitionHLOD, FHLODActorDesc> HLODIterator(WorldPartition); HLODIterator; ++HLODIterator)
-	{
-		for (const auto& SubActor : HLODIterator->GetSubActors())
-		{
-			CachedHLODParents.Emplace(SubActor, HLODIterator->GetGuid());
-		}
-	}
-}
-
-bool UWorldPartitionRuntimeSpatialHash::CreateStreamingGrid(const FSpatialHashRuntimeGrid& RuntimeGrid, const FSquare2DGridHelper& PartionedActors, EWorldPartitionStreamingMode Mode, UWorldPartitionStreamingPolicy* StreamingPolicy, TArray<FString>* OutPackagesToGenerate)
-{
-	TRACE_CPUPROFILER_EVENT_SCOPE(UWorldPartitionRuntimeSpatialHash::CreateStreamingGrid);
-
-	UWorldPartition* WorldPartition = GetOuterUWorldPartition();
-
-	check(FMath::IsPowerOfTwo(PartionedActors.GridSize));
-	FSpatialHashStreamingGrid& CurrentStreamingGrid = StreamingGrids.AddDefaulted_GetRef();
-	CurrentStreamingGrid.GridName = RuntimeGrid.GridName;
-	CurrentStreamingGrid.Origin = PartionedActors.Origin;
-	CurrentStreamingGrid.CellSize = PartionedActors.CellSize;
-	CurrentStreamingGrid.GridSize = PartionedActors.GridSize;
-	CurrentStreamingGrid.LoadingRange = RuntimeGrid.LoadingRange;
-	CurrentStreamingGrid.DebugColor = RuntimeGrid.DebugColor;
-
-	// Move actors into the final streaming grids
-	CurrentStreamingGrid.GridLevels.Reserve(PartionedActors.Levels.Num());
-
-	TArray<FGuid> FilteredActors;
-	int32 Level = INDEX_NONE;
-	for (const FSquare2DGridHelper::FGridLevel& TempLevel : PartionedActors.Levels)
-	{
-		Level++;
-
-		FSpatialHashStreamingGridLevel& GridLevel = CurrentStreamingGrid.GridLevels.AddDefaulted_GetRef();
-
-		GridLevel.LayerCells.SetNum(TempLevel.Cells.Num());
-
-		int32 CellIndex = INDEX_NONE;
-		for (const FSquare2DGridHelper::FGridLevel::FGridCell& TempCell : TempLevel.Cells)
-		{
-			CellIndex++;
-			const int32 CellCoordX = CellIndex % TempLevel.GridSize;
-			const int32 CellCoordY = CellIndex / TempLevel.GridSize;
-			TArray<UWorldPartitionRuntimeSpatialHashCell*>& GridCells = GridLevel.LayerCells[CellIndex].GridCells;
-
-			for (const FSquare2DGridHelper::FGridLevel::FGridCellDataChunk& GridCellDataChunk : TempCell.GetDataChunks())
-			{
-				FilteredActors.SetNum(0, false);
-				FilteredActors.Reset(GridCellDataChunk.GetActors().Num());
-				if (GridCellDataChunk.GetActors().Num())
-				{
-					Algo::TransformIf(GridCellDataChunk.GetActors(), FilteredActors, [&WorldPartition](const FGuid& ActorGuid)
-					{
-						const FWorldPartitionActorDesc* ActorDesc = WorldPartition->GetActorDesc(ActorGuid);
-						const bool bShouldStripActorFromStreaming = ActorDesc->GetActorIsEditorOnly();
-						UE_CLOG(bShouldStripActorFromStreaming, LogWorldPartitionRuntimeSpatialHash, Verbose, TEXT("Stripping Actor %s (%s) from streaming grid"), *(ActorDesc->GetActorPath().ToString()), *ActorGuid.ToString(EGuidFormats::UniqueObjectGuid));
-						return !bShouldStripActorFromStreaming;
-					}, [](const FGuid& ActorGuid) { return ActorGuid; });
-				}
-
-				if (!FilteredActors.Num())
-				{
-					continue;
-				}
-				
-				// Cell cannot be treated as always loaded if it has data layers
-				const bool bIsCellAlwaysLoaded = (&TempCell == &PartionedActors.GetAlwaysLoadedCell()) && !GridCellDataChunk.HasDataLayers();
-
-				FName CellName = GetCellName(CurrentStreamingGrid.GridName, Level, CellCoordX, CellCoordY, GridCellDataChunk.GetDataLayersID());
-
-				UWorldPartitionRuntimeSpatialHashCell* StreamingCell = NewObject<UWorldPartitionRuntimeSpatialHashCell>(WorldPartition, StreamingPolicy->GetRuntimeCellClass(), CellName);
-				GridCells.Add(StreamingCell);
-				StreamingCell->SetIsAlwaysLoaded(bIsCellAlwaysLoaded);
-				StreamingCell->SetDataLayers(GridCellDataChunk.GetDataLayers());
-				StreamingCell->Level = Level;
-				FBox2D Bounds;
-				verify(TempLevel.GetCellBounds(FIntVector2(CellCoordX, CellCoordY), Bounds));
-				StreamingCell->Position = FVector(Bounds.GetCenter(), 0.f);
-
-				UE_LOG(LogWorldPartitionRuntimeSpatialHash, Verbose, TEXT("Cell%s %s Actors = %d"), bIsCellAlwaysLoaded ? TEXT(" (AlwaysLoaded)") : TEXT(""), *StreamingCell->GetName(), FilteredActors.Num());
-
-				// Keep track of all AWorldPartitionHLOD actors referenced by this cell
-				TSet<FGuid> ReferencedHLODActors;
-
-				for (const FGuid& ActorGuid : FilteredActors)
-				{
-					const FWorldPartitionActorDesc* ActorDesc = WorldPartition->GetActorDesc(ActorGuid);
-					FGuid ParentHLOD = CachedHLODParents.FindRef(ActorGuid);
-					if (ParentHLOD.IsValid())
-					{
-						ReferencedHLODActors.Add(ParentHLOD);
-					}
-					StreamingCell->AddActorToCell(ActorGuid, ActorDesc->GetActorPackage(), ActorDesc->GetActorPath());
-					UE_LOG(LogWorldPartitionRuntimeSpatialHash, Verbose, TEXT("  Actor : %s (%s) Origin(%s)"), *(ActorDesc->GetActorPath().ToString()), *ActorGuid.ToString(EGuidFormats::UniqueObjectGuid), *FVector2D(ActorDesc->GetOrigin()).ToString());
-				}
-
-				if (ReferencedHLODActors.Num() > 0)
-				{
-					// Store the referenced HLOD actors as custom cell data
-					UWorldPartitionRuntimeHLODCellData* HLODCellData = NewObject<UWorldPartitionRuntimeHLODCellData>(StreamingCell);
-					HLODCellData->SetReferencedHLODActors(ReferencedHLODActors.Array());
-					StreamingCell->AddCellData(HLODCellData);
-				}
-
-				if (Mode == EWorldPartitionStreamingMode::RuntimeStreamingCells)
-				{
-					UE_LOG(LogWorldPartitionRuntimeSpatialHash, Log, TEXT("Creating runtime streaming cells %s."), *StreamingCell->GetName());
-
-					if (StreamingCell->GetActorCount())
-					{
-						// Always loaded cell actors are transfered to World's Persistent Level
-						if (StreamingCell->IsAlwaysLoaded())
-						{
-							StreamingCell->MoveAlwaysLoadedContentToPersistentLevel();
-						}
-						else
-						{
-							if (!OutPackagesToGenerate)
-							{
-								UE_LOG(LogWorldPartitionRuntimeSpatialHash, Error, TEXT("Error creating runtime streaming cells for cook, OutPackagesToGenerate is null."));
-								return false;
-							}
-
-							const FString PackageRelativePath = StreamingCell->GetPackageNameToCreate();
-							check(!PackageRelativePath.IsEmpty());
-							OutPackagesToGenerate->Add(PackageRelativePath);
-
-							// Map relative package to StreamingCell for PopulateGeneratedPackageForCook/FinalizeGeneratedPackageForCook
-							PackagesToGenerateForCook.Add(PackageRelativePath, StreamingCell);
-						}
-					}
-				}
-			}
-		}
-	}
-
-	return true;
-}
-
-bool UWorldPartitionRuntimeSpatialHash::PopulateGeneratedPackageForCook(UPackage* InPackage, const FString& InPackageRelativePath, const FString& InPackageCookName)
-{
-	if (UWorldPartitionRuntimeCell** MatchingCell = PackagesToGenerateForCook.Find(InPackageRelativePath))
-	{
-		UWorldPartitionRuntimeCell* Cell = *MatchingCell;
-		if (ensure(Cell))
-		{
-			return Cell->PopulateGeneratedPackageForCook(InPackage, InPackageCookName);
-		}
-	}
-	return false;
-}
-
-void UWorldPartitionRuntimeSpatialHash::FinalizeGeneratedPackageForCook()
-{
-	for (const auto& Package : PackagesToGenerateForCook)
-	{
-		UWorldPartitionRuntimeCell* Cell = Package.Value;
-		if (ensure(Cell))
-		{
-			Cell->FinalizeGeneratedPackageForCook();
-		}
-	}
-}
-
-void UWorldPartitionRuntimeSpatialHash::FlushStreaming()
-{
-	StreamingGrids.Empty();
-}
-
-void CheckoutAndSavePackage(UPackage* Package, ISourceControlHelper* SourceControlHelper)
-{
-	// Checkout package
-	Package->MarkAsFullyLoaded();
-
-	if (!SourceControlHelper->Checkout(Package))
-	{
-		UE_LOG(LogWorldPartitionRuntimeSpatialHash, Error, TEXT("Error checking out package %s."), *Package->GetName());
-		check(0);
-	}
-
-	// Save package
-	FString PackageFileName = SourceControlHelper->GetFilename(Package);
-	if (!UPackage::SavePackage(Package, nullptr, RF_Standalone, *PackageFileName))
-	{
-		UE_LOG(LogWorldPartitionRuntimeSpatialHash, Error, TEXT("Error saving package %s."), *Package->GetName());
-		check(0);
-	}
-
-	// Add new package to source control
-	if (!SourceControlHelper->Add(Package))
-	{
-		UE_LOG(LogWorldPartitionRuntimeSpatialHash, Error, TEXT("Error adding package %s to source control."), *Package->GetName());
-		check(0);
-	}
-}
-
-TArray<FGuid> GenerateHLODsForGrid(UWorldPartition* WorldPartition, const FSpatialHashRuntimeGrid& RuntimeGrid, uint32 HLODLevel, FHLODGenerationContext& Context, ISourceControlHelper* SourceControlHelper, const TArray<FActorCluster>& ActorClusters)
-{
-	auto HasExceededMaxMemory = []()
-	{
-		const uint64 MemoryMinFreePhysical = 1024ll * 1024 * 1024;
-		const uint64 MemoryMaxUsedPhysical = 16384ll * 1024 * 1024l;		
-		const FPlatformMemoryStats MemStats = FPlatformMemory::GetStats();
-		return (MemStats.AvailablePhysical < MemoryMinFreePhysical) || (MemStats.UsedPhysical >= MemoryMaxUsedPhysical);
-	};
-
-	auto DoCollectGarbage = []()
-	{
-		if (GDistanceFieldAsyncQueue)
-		{
-			GDistanceFieldAsyncQueue->BlockUntilAllBuildsComplete();
-		}
-
-		if (GCardRepresentationAsyncQueue)
-		{
-			GCardRepresentationAsyncQueue->BlockUntilAllBuildsComplete();
-		}
-
-		const FPlatformMemoryStats MemStatsBefore = FPlatformMemory::GetStats();
-		CollectGarbage(RF_NoFlags, true);
-		const FPlatformMemoryStats MemStatsAfter = FPlatformMemory::GetStats();
-
-		UE_LOG(LogWorldPartitionRuntimeSpatialHash, Warning, TEXT("AvailablePhysical:%.2fGB AvailableVirtual %.2fGB"),
-			((int64)MemStatsAfter.AvailablePhysical - (int64)MemStatsBefore.AvailablePhysical) / (1024.0 * 1024.0 * 1024.0),
-			((int64)MemStatsAfter.AvailableVirtual - (int64)MemStatsBefore.AvailableVirtual) / (1024.0 * 1024.0 * 1024.0)
-		);
-	};
-
-	const FBox WorldBounds = WorldPartition->GetWorldBounds();
-
-	const FSquare2DGridHelper PartitionedActors = GetPartitionedActors(WorldPartition, WorldBounds, RuntimeGrid, ActorClusters);
-
-	// Quick pass to compute the number of cells we'll have to process, to provide a meaningful progress display
-	int32 NbCellsToProcess = 0;
-	PartitionedActors.ForEachCells([&PartitionedActors, &NbCellsToProcess](const FIntVector& CellCoord)
-	{
-		const FSquare2DGridHelper::FGridLevel::FGridCell& GridCell = PartitionedActors.GetCell(CellCoord);
-		
-		// For now, HLOD only processes actors that are not in data layers
-		if (const FSquare2DGridHelper::FGridLevel::FGridCellDataChunk* GridCellDataChunk = GridCell.GetNoDataLayersDataChunk())
-		{
-			check(!GridCellDataChunk->HasDataLayers());
-			const bool bIsCellAlwaysLoaded = &GridCell == &PartitionedActors.GetAlwaysLoadedCell();
-
-			if (!bIsCellAlwaysLoaded && GridCellDataChunk->GetActors().Num() != 0)
-			{
-				NbCellsToProcess++;
-			}
-		}
-	});
-
-	FScopedSlowTask SlowTask(NbCellsToProcess, FText::FromString(FString::Printf(TEXT("Building HLODs for grid %s..."), *RuntimeGrid.GridName.ToString())));
-	SlowTask.MakeDialog();
-
-	FDirectoryWatcherModule& DirectoryWatcherModule = FModuleManager::LoadModuleChecked<FDirectoryWatcherModule>(TEXT("DirectoryWatcher"));
-	FAssetRegistryModule& AssetRegistryModule = FModuleManager::LoadModuleChecked<FAssetRegistryModule>(TEXT("AssetRegistry"));
-
-	TArray<FGuid> GridHLODActors;
-	PartitionedActors.ForEachCells([&PartitionedActors, &GridHLODActors, &SlowTask, RuntimeGrid, HLODLevel, WorldPartition, WorldBounds, &Context, SourceControlHelper, &DirectoryWatcherModule, &AssetRegistryModule, &HasExceededMaxMemory, &DoCollectGarbage](const FIntVector& CellCoord)
-	{
-		const FSquare2DGridHelper::FGridLevel::FGridCell& GridCell = PartitionedActors.GetCell(CellCoord);
-
-		// For now, HLOD only processes actors that are not in data layers
-		if (const FSquare2DGridHelper::FGridLevel::FGridCellDataChunk* GridCellDataChunk = GridCell.GetNoDataLayersDataChunk())
-		{
-			check(!GridCellDataChunk->HasDataLayers());
-			const bool bIsCellAlwaysLoaded = &GridCell == &PartitionedActors.GetAlwaysLoadedCell();
-
-			if (!bIsCellAlwaysLoaded && GridCellDataChunk->GetActors().Num() != 0)
-			{
-				SlowTask.EnterProgressFrame(1);
-
-				FBox2D CellBounds2D;
-				PartitionedActors.GetCellBounds(CellCoord, CellBounds2D);
-
-				FName CellName = GetCellName(WorldPartition, RuntimeGrid.GridName, CellCoord.Z, CellCoord.X, CellCoord.Y, GridCellDataChunk->GetDataLayersID());
-				FBox CellBounds = FBox(FVector(CellBounds2D.Min, WorldBounds.Min.Z), FVector(CellBounds2D.Max, WorldBounds.Max.Z));
-
-				Context.GridIndexX = CellCoord.X;
-				Context.GridIndexY = CellCoord.Y;
-				Context.GridIndexZ = CellCoord.Z;
-
-				// Gather (and potentially load) actors
-				TArray<AActor*> CellActors;
-				for (const FGuid& ActorGuid: GridCellDataChunk->GetActors())
-				{
-					FWorldPartitionReference& ActorRef = Context.ActorReferences.Emplace_GetRef(WorldPartition, ActorGuid);
-					FWorldPartitionActorDesc* ActorDesc  = ActorRef.Get();
-					CellActors.Add(ActorDesc->GetActor());
-				}
-
-				FStaticMeshCompilingManager::Get().FinishAllCompilation();
-
-				if (HasExceededMaxMemory())
-				{
-					DoCollectGarbage();
-				}
-
-				TArray<AWorldPartitionHLOD*> CellHLODActors;
-				CellHLODActors = UHLODLayer::GenerateHLODForCell(WorldPartition, &Context, CellName, CellBounds, HLODLevel, CellActors);
-
-				if (CellHLODActors.Num())
-				{
-					TArray<AWorldPartitionHLOD*> NewCellHLODActors;
-
-					for (AWorldPartitionHLOD* CellHLODActor: CellHLODActors)
-					{
-						FGuid ActorGuid = CellHLODActor->GetActorGuid();
-						GridHLODActors.Add(ActorGuid);
-
-						UPackage* CellHLODActorPackage = CellHLODActor->GetPackage();
-						if (CellHLODActorPackage->HasAnyPackageFlags(PKG_NewlyCreated))
-						{
-							NewCellHLODActors.Add(CellHLODActor);
-						}
-
-						CheckoutAndSavePackage(CellHLODActorPackage, SourceControlHelper);
-					}
-
-					// Manually tick the directory watcher and the asset registry to register newly created actors
-					DirectoryWatcherModule.Get()->Tick(-1.0f);
-					AssetRegistryModule.Get().Tick(-1.0f);
-
-					// Update newly created HLOD actors
-					for (AWorldPartitionHLOD* CellHLODActor: NewCellHLODActors)
-					{
-						Context.ActorReferences.Emplace(WorldPartition, CellHLODActor->GetActorGuid());
-					}
-				}
-
-				// Unload actors
-				Context.ActorReferences.Empty();
-			}
-		}
-	});
-
-	// Need to collect garbage here since some HLOD actors have been marked pending kill when destroying them
-	// and they may be loaded when generating the nexty HLOD layer.
-	DoCollectGarbage();
-
-	return GridHLODActors;
-}
-
-// Find all HLOD grids from the HLODLayer assets we have and build a dependency graph
-void GatherHLODGrids(TMap<FName, FSpatialHashRuntimeGrid>& OutHLODGrids, TMap<FName, TSet<FName>>& OutHLODGridsGraph)
-{
-	// We don't know the HLOD level for those grids yet, we'll have to compute a graph of dependencies between our HLOD layers first
-	// Start with level 0, we'll rename once we have the correct info
-	const uint32 InitialHLODLevel = 0;
-
-	// Gather up all HLODLayer assets
-	FAssetRegistryModule& AssetRegistryModule = FModuleManager::LoadModuleChecked<FAssetRegistryModule>(TEXT("AssetRegistry"));
-	TArray<FAssetData> HLODLayerAssets;
-	AssetRegistryModule.Get().GetAssetsByClass(UHLODLayer::StaticClass()->GetFName(), HLODLayerAssets);
-
-	// Find all used HLOD grids & build a dependency graph
-	for (const FAssetData& HLODLayerAsset : HLODLayerAssets)
-	{
-		if (const UHLODLayer* HLODLayer = Cast<UHLODLayer>(HLODLayerAsset.GetAsset()))
-		{
-			FSpatialHashRuntimeGrid HLODGrid;
-			HLODGrid.CellSize = HLODLayer->GetCellSize();
-			HLODGrid.LoadingRange = HLODLayer->GetLoadingRange();
-			HLODGrid.DebugColor = FLinearColor::Red;
-			HLODGrid.GridName = HLODLayer->GetRuntimeGrid(InitialHLODLevel);
-
-			OutHLODGrids.Emplace(HLODGrid.GridName, HLODGrid);
-
-			TSet<FName>& HLODParentGrids = OutHLODGridsGraph.FindOrAdd(HLODGrid.GridName);
-			if (const UHLODLayer* ParentHLODLayer = HLODLayer->GetParentLayer().LoadSynchronous())
-			{
-				HLODParentGrids.Add(ParentHLODLayer->GetRuntimeGrid(InitialHLODLevel));
-			}
-		}
-	}
-}
-
-// Sort HLOD grids in the order they'll need to be processed for HLOD generation
-bool SortHLODGrids(const TMap<FName, TSet<FName>>& InHLODGridsGraph, TArray<FName>& OutSortedGrids, TMap<FName, uint32>& OutGridsDepth)
-{
-	TSet<FName>		ProcessedGridsSet;	// Processed grids
-	TSet<FName>		VisitedGridsSet;	// Visited grids
-
-	TFunctionRef<bool(const FName, uint32)> VisitGraph = [&](const FName GridName, uint32 CurrentDepth) -> bool
-	{
-		uint32& GridDepth = OutGridsDepth.FindOrAdd(GridName);
-		GridDepth = FMath::Max(GridDepth, CurrentDepth);
-
-		if (ProcessedGridsSet.Contains(GridName))
-		{
-			return true;
-		}
-
-		// Detect cyclic dependencies between HLOD grids...
-		if (VisitedGridsSet.Contains(GridName))
-		{
-			return false; // Not a DAG
-		}
-
-		VisitedGridsSet.Add(GridName);
-
-		for (const FName ParentGridName : InHLODGridsGraph.FindChecked(GridName))
-		{
-			if (!VisitGraph(ParentGridName, CurrentDepth + 1))
-			{
-				return false;
-			}
-		}
-
-		VisitedGridsSet.Remove(GridName);
-
-		ProcessedGridsSet.Add(GridName);
-		OutSortedGrids.Insert(GridName, 0);
-		return true;
-	};
-
-	bool bIsADAG = true;
-	for (const auto& HLODGridEntry : InHLODGridsGraph)
-	{
-		FName GridName = HLODGridEntry.Key;
-		if (!ProcessedGridsSet.Contains(GridName))
-		{
-			const uint32 HLODDepth = 0;
-			bIsADAG = VisitGraph(GridName, HLODDepth);
-			if (!bIsADAG)
-			{
-				break;
-			}
-		}
-	}
-
-	// Remove leafs, those grids are the standard runtime grids, we'll generate HLOD for those in a first pass
-	OutSortedGrids.RemoveAll([&InHLODGridsGraph](const FName GridName) { return !InHLODGridsGraph.Contains(GridName); });
-
-	return bIsADAG;
-}
-
-void RenameHLODGrids(TMap<FName, FSpatialHashRuntimeGrid>& HLODGrids, TArray<FName>& SortedGrids, TMap<FName, uint32>& GridsDepth)
-{
-	// Build a mapping of old -> new names
-	TMap<FName, FName> NewNamesMapping;
-	for (const auto& HLODGrid : HLODGrids)
-	{
-		NewNamesMapping.Emplace(HLODGrid.Key, UHLODLayer::GetRuntimeGridName(GridsDepth[HLODGrid.Key], HLODGrid.Value.CellSize, HLODGrid.Value.LoadingRange));
-	}
-
-	// Remplace map entries
-	for (const auto& Names : NewNamesMapping)
-	{
-		FSpatialHashRuntimeGrid RuntimeGrid = HLODGrids.FindAndRemoveChecked(Names.Key);
-		RuntimeGrid.GridName = Names.Value;
-		HLODGrids.Emplace(Names.Value, RuntimeGrid);
-
-		int32 GridDepth = GridsDepth.FindAndRemoveChecked(Names.Key);
-		GridsDepth.Emplace(Names.Value, GridDepth);
-	}
-
-	// Replace arrays entries
-	for (FName& Name : SortedGrids)
-	{
-		Name = NewNamesMapping[Name];
-	}
+	MainGrid.DebugColor = FLinearColor::Gray;
 }
 
-// Create/destroy HLOD grid actors
-void UpdateHLODGridsActors(UWorld* World, const TMap<FName, FSpatialHashRuntimeGrid>& HLODGrids, ISourceControlHelper* SourceControlHelper)
+void UWorldPartitionRuntimeSpatialHash::ImportFromWorldComposition(UWorldComposition* WorldComposition)
 {
-	static const FName HLODGridTag = TEXT("HLOD");
-	static const uint32 HLODGridTagLen = HLODGridTag.ToString().Len();
+	check(IsRunningCommandlet());
 
-	// Gather all existing HLOD grid actors, see if some are unused and needs to be deleted
-	TMap<FName, ASpatialHashRuntimeGridInfo*> ExistingGridActors;
-	for (TActorIterator<ASpatialHashRuntimeGridInfo> ItRuntimeGridActor(World); ItRuntimeGridActor; ++ItRuntimeGridActor)
+	if (WorldComposition)
 	{
-		ASpatialHashRuntimeGridInfo* GridActor = *ItRuntimeGridActor;
-		if (GridActor->ActorHasTag(HLODGridTag))
+		const TArray<FWorldTileLayer> WorldCompositionTileLayers = WorldComposition->GetDistanceDependentLayers();
+		for (const FWorldTileLayer& Layer : WorldCompositionTileLayers)
 		{
-			if (HLODGrids.Contains(GridActor->GridSettings.GridName))
-			{
-				ExistingGridActors.Emplace(GridActor->GridSettings.GridName, GridActor);
-			}
-			else
+			FName GridName = FName(Layer.Name);
+			FSpatialHashRuntimeGrid* Grid = Algo::FindByPredicate(Grids, [GridName](const FSpatialHashRuntimeGrid& Grid) { return Grid.GridName == GridName; });
+			if (!Grid)
 			{
-				World->DestroyActor(GridActor);
-
-				const FString ActorPackageFilename = SourceControlHelper->GetFilename(GridActor->GetPackage());
-				SourceControlHelper->Delete(ActorPackageFilename);
+				Grid = &Grids.AddDefaulted_GetRef();
+				Grid->GridName = GridName;
+				Grid->CellSize = 3200;
+				Grid->DebugColor = FLinearColor::MakeRandomColor();
 			}
+			// World Composition Layer Streaming Distance always wins over existing value (config file)
+			Grid->LoadingRange = Layer.StreamingDistance;
 		}
-	}
 
-	// Create missing HLOD grid actors 
-	for (const auto& HLODGridEntry : HLODGrids)
-	{
-		FName GridName = HLODGridEntry.Key;
-		if (!ExistingGridActors.Contains(GridName))
+		const UWorldComposition::FTilesList& Tiles = WorldComposition->GetTilesList();
+		for (int32 TileIdx = 0; TileIdx < Tiles.Num(); TileIdx++)
 		{
-			const FSpatialHashRuntimeGrid& GridSettings = HLODGridEntry.Value;
-
-			FActorSpawnParameters SpawnParams;
-			SpawnParams.bCreateActorPackage = true;
-			ASpatialHashRuntimeGridInfo* GridActor = World->SpawnActor<ASpatialHashRuntimeGridInfo>(SpawnParams);
-			GridActor->Tags.Add(HLODGridTag);
-			GridActor->SetActorLabel(GridSettings.GridName.ToString());
-			GridActor->GridSettings = GridSettings;
-
-			// Setup grid debug color to match HLODColorationColors
-			const uint32 LastLODColorationColorIdx = GEngine->HLODColorationColors.Num() - 1;
-			uint32 HLODLevel;
-			if (!LexTryParseString(HLODLevel, *GridSettings.GridName.ToString().Mid(HLODGridTagLen, 1)))
+			const FWorldCompositionTile& Tile = Tiles[TileIdx];
+			ULevelStreaming* StreamingLevel = WorldComposition->TilesStreaming[TileIdx];
+			if (StreamingLevel && WorldComposition->IsDistanceDependentLevel(Tile.PackageName))
 			{
-				HLODLevel = LastLODColorationColorIdx;
+				// Map WorldComposition tiles streaming level to Runtime Grid
+				WorldCompositionStreamingLevelToRuntimeGrid.Add(StreamingLevel, FName(Tile.Info.Layer.Name));
 			}
-			HLODLevel = FMath::Clamp(HLODLevel + 2, 0U, LastLODColorationColorIdx);
-			GridActor->GridSettings.DebugColor = GEngine->HLODColorationColors[HLODLevel];
-
-			CheckoutAndSavePackage(GridActor->GetPackage(), SourceControlHelper);
 		}
 	}
 }
 
-bool UWorldPartitionRuntimeSpatialHash::GenerateHLOD(ISourceControlHelper* SourceControlHelper)
+bool UWorldPartitionRuntimeSpatialHash::GenerateStreaming(EWorldPartitionStreamingMode Mode, UWorldPartitionStreamingPolicy* StreamingPolicy, TArray<FString>* OutPackagesToGenerate)
 {
-	TRACE_CPUPROFILER_EVENT_SCOPE(UWorldPartitionRuntimeSpatialHash::GenerateHLOD);
+	TRACE_CPUPROFILER_EVENT_SCOPE(UWorldPartitionRuntimeSpatialHash::GenerateStreaming);
+	UWorldPartition* WorldPartition = GetOuterUWorldPartition();
 
+	UE_SCOPED_TIMER(TEXT("GenerateStreaming"), LogWorldPartitionRuntimeSpatialHash);
+	
 	if (!Grids.Num())
 	{
 		UE_LOG(LogWorldPartitionRuntimeSpatialHash, Error, TEXT("Invalid partition grids setup"));
 		return false;
 	}
 
-	// Find all used HLOD grids & build a dependency graph
-	TMap<FName, FSpatialHashRuntimeGrid> HLODGrids;
-	TMap<FName, TSet<FName>>			 HLODGridsGraph;
-	GatherHLODGrids(HLODGrids, HLODGridsGraph);
+	TArray<FSpatialHashRuntimeGrid> AllGrids;
+	AllGrids.Append(Grids);
 
-	// Sort HLOD grids in the order they'll need to be processed
-	TArray<FName>	SortedGrids;		// HLOD Grids, sorted in the order they'll need to be processed for HLOD generation
-	TMap<FName, uint32>	GridsDepth;		// Depth - Used to obtain the HLOD level
-	bool bIsADAG = SortHLODGrids(HLODGridsGraph, SortedGrids, GridsDepth);
-	if (!bIsADAG)
+	// Append grids from runtime grid actors
+	for (AActor* Actor : GetWorld()->PersistentLevel->Actors)
 	{
-		UE_LOG(LogWorldPartitionRuntimeSpatialHash, Error, TEXT("Invalid grids setup, cycles detected"));
-		return false;
+		if (ASpatialHashRuntimeGridInfo* RuntimeGridActor = Cast<ASpatialHashRuntimeGridInfo>(Actor))
+		{
+			AllGrids.Add(RuntimeGridActor->GridSettings);
+		}
 	}
 
-	// Now that we computed proper depth per grid, rename the grids
-	RenameHLODGrids(HLODGrids, SortedGrids, GridsDepth);
+	check(!StreamingGrids.Num());
+
+	// Build a map of Actor GUID -> HLODActor GUID once instead of having to recompute for every streaming grid we create
+	CacheHLODParents();
 
 	TMap<FName, int32> GridsMapping;
 	GridsMapping.Add(NAME_None, 0);
-	for (int32 i = 0; i < Grids.Num(); i++)
+	for (int32 i = 0; i < AllGrids.Num(); i++)
 	{
-		const FSpatialHashRuntimeGrid& Grid = Grids[i];
+		const FSpatialHashRuntimeGrid& Grid = AllGrids[i];
 		check(!GridsMapping.Contains(Grid.GridName));
 		GridsMapping.Add(Grid.GridName, i);
 	}
 
-	UWorldPartition* WorldPartition = GetOuterUWorldPartition();
-
-	// Create HLODs generation context
-	FHLODGenerationContext Context;
-	TArray<FWorldPartitionHandle> InvalidHLODActors;
-	for (TWorldPartitionActorDescIterator<AWorldPartitionHLOD, FHLODActorDesc> HLODIterator(WorldPartition); HLODIterator; ++HLODIterator)
-	{
-		FWorldPartitionHandle HLODActorHandle(WorldPartition, HLODIterator->GetGuid());
-
-		if (HLODIterator->GetCellHash())
-		{
-			Context.HLODActorDescs.Add(HLODIterator->GetCellHash(), HLODActorHandle);
-		}
-		else
-		{
-			InvalidHLODActors.Add(HLODActorHandle);
-		}
-	}
-
-	// Create actor clusters - ignore HLOD actors
-	TArray<FActorCluster> ActorClusters = CreateActorClusters(WorldPartition, [](const FWorldPartitionActorDesc& ActorDesc)
-	{
-		return !ActorDesc.GetActorClass()->IsChildOf<AWorldPartitionHLOD>();
-	});
+	// Create actor clusters
+	TArray<FActorCluster> ActorClusters = CreateActorClusters(WorldPartition);
 
-	TArray<TArray<FActorCluster>> GridsClusters;
-	GridsClusters.InsertDefaulted(0, Grids.Num());
+	TArray<TArray<FActorCluster>> GridActors;
+	GridActors.InsertDefaulted(0, AllGrids.Num());
 
 	for (FActorCluster& ActorCluster : ActorClusters)
 	{
@@ -1930,185 +479,212 @@ bool UWorldPartitionRuntimeSpatialHash::GenerateHLOD(ISourceControlHelper* Sourc
 		}
 
 		int32 GridIndex = FoundIndex ? *FoundIndex : 0;
-		GridsClusters[GridIndex].Add(MoveTemp(ActorCluster));
+		GridActors[GridIndex].Add(MoveTemp(ActorCluster));
 	}
-
-	// Keep track of all valid HLOD actors, along with which runtime grid they live in
-	TMap<FName, TArray<FGuid>> GridsHLODActors;
-
-	auto GenerateHLODs = [&GridsHLODActors, WorldPartition, &GridsDepth, &Context, SourceControlHelper](const FSpatialHashRuntimeGrid& RuntimeGrid, uint32 HLODLevel, const TArray<FActorCluster>& ActorClusters)
+	
+	const FBox WorldBounds = WorldPartition->GetWorldBounds();
+	for (int32 GridIndex=0; GridIndex < AllGrids.Num(); GridIndex++)
 	{
-		// Generate HLODs for this grid
-		TArray<FGuid> HLODActors = GenerateHLODsForGrid(WorldPartition, RuntimeGrid, HLODLevel, Context, SourceControlHelper, ActorClusters);
-
-		for (const FGuid& HLODActorGuid : HLODActors)
+		const FSpatialHashRuntimeGrid& Grid = AllGrids[GridIndex];
+		const FSquare2DGridHelper PartionedActors = GetPartitionedActors(WorldPartition, WorldBounds, Grid, GridActors[GridIndex]);
+		if (!CreateStreamingGrid(Grid, PartionedActors, Mode, StreamingPolicy, OutPackagesToGenerate))
 		{
-			FWorldPartitionActorDesc* HLODActorDesc = WorldPartition->GetActorDesc(HLODActorGuid);
-			check(HLODActorDesc);
-
-			GridsHLODActors.FindOrAdd(HLODActorDesc->GetRuntimeGrid()).Add(HLODActorGuid);
+			return false;
 		}
-	};
-
-	// Generate HLODs for the standard runtime grids (HLOD 0)
-	for (int32 GridIndex = 0; GridIndex < Grids.Num(); GridIndex++)
-	{
-		// Generate HLODs for this grid - retrieve actors from the world partition
-		GenerateHLODs(Grids[GridIndex], 0, GridsClusters[GridIndex]);
 	}
 
-	// Now, go on and create HLOD actors from HLOD grids (HLOD 1-N)
-	for (const FName HLODGridName : SortedGrids)
-	{
-		// No need to process empty grids
-		if (!GridsHLODActors.Contains(HLODGridName))
-		{
-			HLODGrids.Remove(HLODGridName); // No need to keep this grid around, we have no actors in it
-			continue;
-		}
-
-		// Create one actor cluster per HLOD actor
-		TArray<FActorCluster> HLODActorClusters;
-		HLODActorClusters.Reserve(GridsHLODActors[HLODGridName].Num());
+	return true;
+}
 
-		for (const FGuid& HLODActorGuid: GridsHLODActors[HLODGridName])
-		{
-			FWorldPartitionActorDesc* HLODActorDesc = WorldPartition->GetActorDesc(HLODActorGuid);
-			check(HLODActorDesc);
+FName UWorldPartitionRuntimeSpatialHash::GetCellName(UWorldPartition* WorldPartition, FName InGridName, int32 InLevel, int32 InCellX, int32 InCellY, const FDataLayersID& InDataLayerID)
+{
+	const FString PackageName = FPackageName::GetShortName(WorldPartition->GetPackage());
+	const FString PackageNameNoPIEPrefix = UWorld::RemovePIEPrefix(PackageName);
 
-			HLODActorClusters.Emplace(HLODActorDesc, HLODActorDesc->GetGridPlacement(), GetWorld());
-		};
+	return FName(*FString::Printf(TEXT("WPRT_%s_%s_Cell_L%d_X%02d_Y%02d_DL%X"), *PackageNameNoPIEPrefix, *InGridName.ToString(), InLevel, InCellX, InCellY, InDataLayerID.GetHash()));
+}
 
-		// Generate HLODs for this grid - retrieve actors from our ValidHLODActors map
-		// We can't rely on actor descs for newly created HLOD actors
-		GenerateHLODs(HLODGrids[HLODGridName], GridsDepth[HLODGridName] + 1, HLODActorClusters);
-	}
+FName UWorldPartitionRuntimeSpatialHash::GetCellName(FName InGridName, int32 InLevel, int32 InCellX, int32 InCellY, const FDataLayersID& InDataLayerID) const
+{
+	UWorldPartition* WorldPartition = GetOuterUWorldPartition();
+	return UWorldPartitionRuntimeSpatialHash::GetCellName(WorldPartition, InGridName, InLevel, InCellX, InCellY, InDataLayerID);
+}
 
-	auto DeleteHLODActor = [&SourceControlHelper](FWorldPartitionHandle ActorHandle)
-	{
-		FWorldPartitionActorDesc* HLODActorDesc = ActorHandle.Get();
-		check(HLODActorDesc);
+void UWorldPartitionRuntimeSpatialHash::CacheHLODParents()
+{
+	CachedHLODParents.Reset();
 
-		const FString ActorPackageFilename = SourceControlHelper->GetFilename(HLODActorDesc->GetActorPackage().ToString());
-		SourceControlHelper->Delete(ActorPackageFilename);
-	};
+	UWorldPartition* WorldPartition = GetOuterUWorldPartition();
 
-	// Destroy all unreferenced HLOD actors
-	for (const auto& HLODActorPair: Context.HLODActorDescs)
+	for (TWorldPartitionActorDescIterator<AWorldPartitionHLOD, FHLODActorDesc> HLODIterator(WorldPartition); HLODIterator; ++HLODIterator)
 	{
-		DeleteHLODActor(HLODActorPair.Value);
+		for (const auto& SubActor : HLODIterator->GetSubActors())
+		{
+			CachedHLODParents.Emplace(SubActor, HLODIterator->GetGuid());
+		}
 	}
-
-	// Destroy all invalid HLOD actors
-	Algo::ForEach(InvalidHLODActors, DeleteHLODActor);
-
-	// Create/destroy HLOD grid actors
-	UpdateHLODGridsActors(GetWorld(), HLODGrids, SourceControlHelper);
-
-	return true;
 }
 
-bool UWorldPartitionRuntimeSpatialHash::GenerateNavigationData()
+bool UWorldPartitionRuntimeSpatialHash::CreateStreamingGrid(const FSpatialHashRuntimeGrid& RuntimeGrid, const FSquare2DGridHelper& PartionedActors, EWorldPartitionStreamingMode Mode, UWorldPartitionStreamingPolicy* StreamingPolicy, TArray<FString>* OutPackagesToGenerate)
 {
-	TRACE_CPUPROFILER_EVENT_SCOPE(UWorldPartitionRuntimeSpatialHash::GenerateNavigationData);
-
-	UE_LOG(LogWorldPartitionRuntimeSpatialHash, Log, TEXT("%s"), ANSI_TO_TCHAR(__FUNCTION__));
+	TRACE_CPUPROFILER_EVENT_SCOPE(UWorldPartitionRuntimeSpatialHash::CreateStreamingGrid);
 
 	UWorldPartition* WorldPartition = GetOuterUWorldPartition();
-	UWorld* World = WorldPartition->World;
 
-	// Generate navmesh
-	// Make sure navigation is added and initialized in EditorMode
-	FNavigationSystem::AddNavigationSystemToWorld(*World, FNavigationSystemRunMode::EditorMode);
+	check(FMath::IsPowerOfTwo(PartionedActors.GridSize));
+	FSpatialHashStreamingGrid& CurrentStreamingGrid = StreamingGrids.AddDefaulted_GetRef();
+	CurrentStreamingGrid.GridName = RuntimeGrid.GridName;
+	CurrentStreamingGrid.Origin = PartionedActors.Origin;
+	CurrentStreamingGrid.CellSize = PartionedActors.CellSize;
+	CurrentStreamingGrid.GridSize = PartionedActors.GridSize;
+	CurrentStreamingGrid.LoadingRange = RuntimeGrid.LoadingRange;
+	CurrentStreamingGrid.DebugColor = RuntimeGrid.DebugColor;
 
-	// Invoke navigation data generator
-	FNavigationSystem::Build(*World);
+	// Move actors into the final streaming grids
+	CurrentStreamingGrid.GridLevels.Reserve(PartionedActors.Levels.Num());
 
-	// For each cell, gather navmesh and generate a datachunk actor
-	const FBox WorldBounds = WorldPartition->GetWorldBounds();
+	TArray<FGuid> FilteredActors;
+	int32 Level = INDEX_NONE;
+	for (const FSquare2DGridHelper::FGridLevel& TempLevel : PartionedActors.Levels)
+	{
+		Level++;
 
-	const int32 GridIndex = 0; //Todo At: only work with grid 0 for now.
-	const uint32 GridLevel = 3; //Todo AT, only generate for level 3 for now
-	UE_LOG(LogWorldPartitionRuntimeSpatialHash, Verbose, TEXT("Generate NavDataChunk actor for GridIndex %i."), GridIndex);
+		FSpatialHashStreamingGridLevel& GridLevel = CurrentStreamingGrid.GridLevels.AddDefaulted_GetRef();
 
-	const FSpatialHashRuntimeGrid& RuntimeGrid = Grids[GridIndex];
-	const FSquare2DGridHelper GridHelper = GetGridHelper(WorldBounds, RuntimeGrid.CellSize);
-	int32 ActorCount = 0;
+		GridLevel.LayerCells.SetNum(TempLevel.Cells.Num());
 
-	const UNavigationSystemBase* NavSystem = World->GetNavigationSystem();
-	if (NavSystem == nullptr)
-	{
-		UE_LOG(LogWorldPartitionRuntimeSpatialHash, Verbose, TEXT("No navigation system to generate navigation data."));
-		return false;
-	}
+		int32 CellIndex = INDEX_NONE;
+		for (const FSquare2DGridHelper::FGridLevel::FGridCell& TempCell : TempLevel.Cells)
+		{
+			CellIndex++;
+			const int32 CellCoordX = CellIndex % TempLevel.GridSize;
+			const int32 CellCoordY = CellIndex / TempLevel.GridSize;
+			TArray<UWorldPartitionRuntimeSpatialHashCell*>& GridCells = GridLevel.LayerCells[CellIndex].GridCells;
 
-	// Keep track of all valid navigation data chunk actors
-	TSet<ANavigationDataChunkActor*> ValidNavigationDataChunkActors;
+			for (const FSquare2DGridHelper::FGridLevel::FGridCellDataChunk& GridCellDataChunk : TempCell.GetDataChunks())
+			{
+				FilteredActors.SetNum(0, false);
+				FilteredActors.Reset(GridCellDataChunk.GetActors().Num());
+				if (GridCellDataChunk.GetActors().Num())
+				{
+					Algo::TransformIf(GridCellDataChunk.GetActors(), FilteredActors, [&WorldPartition](const FGuid& ActorGuid)
+					{
+						const FWorldPartitionActorDesc* ActorDesc = WorldPartition->GetActorDesc(ActorGuid);
+						const bool bShouldStripActorFromStreaming = ActorDesc->GetActorIsEditorOnly();
+						UE_CLOG(bShouldStripActorFromStreaming, LogWorldPartitionRuntimeSpatialHash, Verbose, TEXT("Stripping Actor %s (%s) from streaming grid"), *(ActorDesc->GetActorPath().ToString()), *ActorGuid.ToString(EGuidFormats::UniqueObjectGuid));
+						return !bShouldStripActorFromStreaming;
+					}, [](const FGuid& ActorGuid) { return ActorGuid; });
+				}
+
+				if (!FilteredActors.Num())
+				{
+					continue;
+				}
+				
+				// Cell cannot be treated as always loaded if it has data layers
+				const bool bIsCellAlwaysLoaded = (&TempCell == &PartionedActors.GetAlwaysLoadedCell()) && !GridCellDataChunk.HasDataLayers();
 
-	const FSquare2DGridHelper::FGridLevel& GridLevelHelper = GridHelper.Levels[GridLevel];
+				FName CellName = GetCellName(CurrentStreamingGrid.GridName, Level, CellCoordX, CellCoordY, GridCellDataChunk.GetDataLayersID());
 
-	GridLevelHelper.ForEachCells([GridLevel, &GridLevelHelper, RuntimeGrid, WorldPartition, &ActorCount, World, NavSystem, &ValidNavigationDataChunkActors, this](const FIntVector2& CellCoord)
-	{
-		FBox2D CellBounds;
-		GridLevelHelper.GetCellBounds(CellCoord, CellBounds);
-		if (CellBounds.GetExtent().X < 1.f || CellBounds.GetExtent().Y < 1.f)
-		{
-			// Safety, since we reduce by 1.f below.
-			UE_LOG(LogWorldPartitionRuntimeSpatialHash, Warning, TEXT("%s: grid cell too small."), ANSI_TO_TCHAR(__FUNCTION__));
-			return;
-		}
+				UWorldPartitionRuntimeSpatialHashCell* StreamingCell = NewObject<UWorldPartitionRuntimeSpatialHashCell>(WorldPartition, StreamingPolicy->GetRuntimeCellClass(), CellName);
+				GridCells.Add(StreamingCell);
+				StreamingCell->SetIsAlwaysLoaded(bIsCellAlwaysLoaded);
+				StreamingCell->SetDataLayers(GridCellDataChunk.GetDataLayers());
+				StreamingCell->Level = Level;
+				FBox2D Bounds;
+				verify(TempLevel.GetCellBounds(FIntVector2(CellCoordX, CellCoordY), Bounds));
+				StreamingCell->Position = FVector(Bounds.GetCenter(), 0.f);
 
-		const float HalfHeight = HALF_WORLD_MAX;
-		const FBox QueryBounds(FVector(CellBounds.Min.X, CellBounds.Min.Y, -HalfHeight), FVector(CellBounds.Max.X, CellBounds.Max.Y, HalfHeight));
+				UE_LOG(LogWorldPartitionRuntimeSpatialHash, Verbose, TEXT("Cell%s %s Actors = %d"), bIsCellAlwaysLoaded ? TEXT(" (AlwaysLoaded)") : TEXT(""), *StreamingCell->GetName(), FilteredActors.Num());
 
-		if (NavSystem->ContainsNavData(QueryBounds) == false)
-		{
-			// Skip if there is no navdata for this cell
-			return;
-		}
+				// Keep track of all AWorldPartitionHLOD actors referenced by this cell
+				TSet<FGuid> ReferencedHLODActors;
 
-		const FSquare2DGridHelper::FGridLevel::FGridCell& GridCell = GridLevelHelper.GetCell(CellCoord);
+				for (const FGuid& ActorGuid : FilteredActors)
+				{
+					const FWorldPartitionActorDesc* ActorDesc = WorldPartition->GetActorDesc(ActorGuid);
+					FGuid ParentHLOD = CachedHLODParents.FindRef(ActorGuid);
+					if (ParentHLOD.IsValid())
+					{
+						ReferencedHLODActors.Add(ParentHLOD);
+					}
+					StreamingCell->AddActorToCell(ActorGuid, ActorDesc->GetActorPackage(), ActorDesc->GetActorPath());
+					UE_LOG(LogWorldPartitionRuntimeSpatialHash, Verbose, TEXT("  Actor : %s (%s) Origin(%s)"), *(ActorDesc->GetActorPath().ToString()), *ActorGuid.ToString(EGuidFormats::UniqueObjectGuid), *FVector2D(ActorDesc->GetOrigin()).ToString());
+				}
 
-		FActorSpawnParameters SpawnParams;
-		SpawnParams.bDeferConstruction = true;
-		SpawnParams.bCreateActorPackage = true;
-		ANavigationDataChunkActor* DataChunkActor = World->SpawnActor<ANavigationDataChunkActor>(SpawnParams);
-		ActorCount++;
+				if (ReferencedHLODActors.Num() > 0)
+				{
+					// Store the referenced HLOD actors as custom cell data
+					UWorldPartitionRuntimeHLODCellData* HLODCellData = NewObject<UWorldPartitionRuntimeHLODCellData>(StreamingCell);
+					HLODCellData->SetReferencedHLODActors(ReferencedHLODActors.Array());
+					StreamingCell->AddCellData(HLODCellData);
+				}
 
-		const FVector2D CellCenter = CellBounds.GetCenter();
-		DataChunkActor->SetActorLocation(FVector(CellCenter.X, CellCenter.Y, 0.f));
+				if (Mode == EWorldPartitionStreamingMode::RuntimeStreamingCells)
+				{
+					UE_LOG(LogWorldPartitionRuntimeSpatialHash, Log, TEXT("Creating runtime streaming cells %s."), *StreamingCell->GetName());
 
-		FBox TilesBounds(EForceInit::ForceInit);
-		DataChunkActor->CollectNavData(QueryBounds, TilesBounds);
+					if (StreamingCell->GetActorCount())
+					{
+						// Always loaded cell actors are transfered to World's Persistent Level
+						if (StreamingCell->IsAlwaysLoaded())
+						{
+							StreamingCell->MoveAlwaysLoadedContentToPersistentLevel();
+						}
+						else
+						{
+							if (!OutPackagesToGenerate)
+							{
+								UE_LOG(LogWorldPartitionRuntimeSpatialHash, Error, TEXT("Error creating runtime streaming cells for cook, OutPackagesToGenerate is null."));
+								return false;
+							}
 
-		FBox ChunkActorBounds(FVector(QueryBounds.Min.X, QueryBounds.Min.Y, TilesBounds.Min.Z), FVector(QueryBounds.Max.X, QueryBounds.Max.Y, TilesBounds.Max.Z));
-		ChunkActorBounds = ChunkActorBounds.ExpandBy(FVector(-1.f, -1.f, 1.f)); //reduce XY by 1cm to avoid precision issues causing potential overflow on neighboring cell, add 1cm in Z to have a minimum of volume.
-		UE_LOG(LogWorldPartitionRuntimeSpatialHash, VeryVerbose, TEXT("Setting ChunkActorBounds to %s"), *ChunkActorBounds.ToString());
-		DataChunkActor->SetDataChunkActorBounds(ChunkActorBounds);
+							const FString PackageRelativePath = StreamingCell->GetPackageNameToCreate();
+							check(!PackageRelativePath.IsEmpty());
+							OutPackagesToGenerate->Add(PackageRelativePath);
 
-		const FName CellName = GetCellName(RuntimeGrid.GridName, GridLevel, CellCoord.X, CellCoord.Y);
-		DataChunkActor->SetActorLabel(FString::Printf(TEXT("NavDataChunkActor_%s_%s"), *GetName(), *CellName.ToString()));
-		
-		//@todo_ow: Properly handle data layers
+							// Map relative package to StreamingCell for PopulateGeneratedPackageForCook/FinalizeGeneratedPackageForCook
+							PackagesToGenerateForCook.Add(PackageRelativePath, StreamingCell);
+						}
+					}
+				}
+			}
+		}
+	}
 
-		// Set target grid
-		DataChunkActor->RuntimeGrid = RuntimeGrid.GridName;
-		ValidNavigationDataChunkActors.Add(DataChunkActor);
+	return true;
+}
 
-		UE_LOG(LogWorldPartitionRuntimeSpatialHash, Verbose, TEXT("%i) %s added."), ActorCount, *DataChunkActor->GetName());
-	});
+bool UWorldPartitionRuntimeSpatialHash::PopulateGeneratedPackageForCook(UPackage* InPackage, const FString& InPackageRelativePath, const FString& InPackageCookName)
+{
+	if (UWorldPartitionRuntimeCell** MatchingCell = PackagesToGenerateForCook.Find(InPackageRelativePath))
+	{
+		UWorldPartitionRuntimeCell* Cell = *MatchingCell;
+		if (ensure(Cell))
+		{
+			return Cell->PopulateGeneratedPackageForCook(InPackage, InPackageCookName);
+		}
+	}
+	return false;
+}
 
-	// Destroy all invalid navigation data chunk actors
-	for (TActorIterator<ANavigationDataChunkActor> It(GetWorld()); It; ++It)
+void UWorldPartitionRuntimeSpatialHash::FinalizeGeneratedPackageForCook()
+{
+	for (const auto& Package : PackagesToGenerateForCook)
 	{
-		if (!ValidNavigationDataChunkActors.Contains(*It))
+		UWorldPartitionRuntimeCell* Cell = Package.Value;
+		if (ensure(Cell))
 		{
-			GetWorld()->DestroyActor(*It);
+			Cell->FinalizeGeneratedPackageForCook();
 		}
 	}
+}
 
-	return true;
+void UWorldPartitionRuntimeSpatialHash::FlushStreaming()
+{
+	StreamingGrids.Empty();
 }
+
 #endif //WITH_EDITOR
 
 FAutoConsoleCommand UWorldPartitionRuntimeSpatialHash::OverrideLoadingRangeCommand(
diff --git a/Engine/Source/Runtime/Engine/Public/WorldPartition/WorldPartitionRuntimeSpatialHash.h b/Engine/Source/Runtime/Engine/Public/WorldPartition/WorldPartitionRuntimeSpatialHash.h
--- a/Engine/Source/Runtime/Engine/Public/WorldPartition/WorldPartitionRuntimeSpatialHash.h
+++ b/Engine/Source/Runtime/Engine/Public/WorldPartition/WorldPartitionRuntimeSpatialHash.h
@@ -12,6 +12,8 @@
 #include "WorldPartitionRuntimeSpatialHashCell.h"
 #include "WorldPartitionRuntimeSpatialHash.generated.h"
 
+DECLARE_LOG_CATEGORY_EXTERN(LogWorldPartitionRuntimeSpatialHash, Log, All);
+
 USTRUCT()
 struct FSpatialHashStreamingGridLayerCell
 {
@@ -220,6 +222,9 @@ public:
 	virtual bool GenerateHLOD(ISourceControlHelper* SourceControlHelper) override;
 	virtual bool GenerateNavigationData() override;
 	virtual FName GetActorRuntimeGrid(const AActor* Actor) const override;
+
+	FName GetCellName(FName InGridName, int32 InLevel, int32 InCellX, int32 InCellY, const FDataLayersID& InDataLayerID = FDataLayersID()) const;
+	static FName GetCellName(UWorldPartition* WorldPartition, FName InGridName, int32 InLevel, int32 InCellX, int32 InCellY, const FDataLayersID& InDataLayerID);
 #endif
 
 	// streaming interface
@@ -232,13 +237,6 @@ protected:
 	TMap<class ULevelStreaming*, FName> WorldCompositionStreamingLevelToRuntimeGrid;
 
 private:
-
-#if WITH_EDITOR
-	FName GetCellName(FName InGridName, int32 InLevel, int32 InCellX, int32 InCellY, const FDataLayersID& InDataLayerID = FDataLayersID()) const;
-
-	void CacheHLODParents();
-#endif
-
 	/** Console command used to change loading range for a given streaming grid */
 	static class FAutoConsoleCommand OverrideLoadingRangeCommand;
 
@@ -265,6 +263,7 @@ private:
 	void GetStreamingCells(const FVector& Position, const FSpatialHashStreamingGrid& StreamingGrid, TSet<const UWorldPartitionRuntimeCell*>& Cells) const;
 #if WITH_EDITOR
 	bool CreateStreamingGrid(const FSpatialHashRuntimeGrid& RuntimeGrid, const FSquare2DGridHelper& PartionedActors, EWorldPartitionStreamingMode Mode, UWorldPartitionStreamingPolicy* StreamingPolicy, TArray<FString>* OutPackagesToGenerate = nullptr);
+	void CacheHLODParents();
 #endif
 
 	friend class UWorldPartitionSubsystem;
