Fixed Editor Linux build

diff --git a/Engine/Source/Runtime/Engine/Private/WorldPartition/RuntimeSpatialHash/RuntimeSpatialHashNavigation.cpp b/Engine/Source/Runtime/Engine/Private/WorldPartition/RuntimeSpatialHash/RuntimeSpatialHashNavigation.cpp
--- a/Engine/Source/Runtime/Engine/Private/WorldPartition/RuntimeSpatialHash/RuntimeSpatialHashNavigation.cpp
+++ b/Engine/Source/Runtime/Engine/Private/WorldPartition/RuntimeSpatialHash/RuntimeSpatialHashNavigation.cpp
@@ -115,4 +115,4 @@ bool UWorldPartitionRuntimeSpatialHash::GenerateNavigationData()
 	return true;
 }
 
-#endif WITH_EDITOR
\ No newline at end of file
+#endif // #if WITH_EDITOR
\ No newline at end of file
