Temporarily disable Reverse Lookup feature during some bug investigation

diff --git a/Engine/Source/Runtime/Engine/Private/ObjectCacheContext.cpp b/Engine/Source/Runtime/Engine/Private/ObjectCacheContext.cpp
--- a/Engine/Source/Runtime/Engine/Private/ObjectCacheContext.cpp
+++ b/Engine/Source/Runtime/Engine/Private/ObjectCacheContext.cpp
@@ -23,7 +23,7 @@ DEFINE_LOG_CATEGORY(LogObjectCache);
 
 static TAutoConsoleVariable<int32> CVarObjectReverseLookupMode(
 	TEXT("Editor.ObjectReverseLookupMode"),
-	1,
+	0,
 	TEXT("0 - Reverse lookup tables are computed every time they are needed (slower behavior) \n")
 	TEXT("1 - Maintain permanent reverse lookup tables (faster behavior) \n")
 	TEXT("2 - Comparison mode (slowest to do validation between both mode) \n"),
