Add MI cached data to ARO

diff --git a/Engine/Source/Runtime/Engine/Private/Materials/MaterialInstance.cpp b/Engine/Source/Runtime/Engine/Private/Materials/MaterialInstance.cpp
--- a/Engine/Source/Runtime/Engine/Private/Materials/MaterialInstance.cpp
+++ b/Engine/Source/Runtime/Engine/Private/Materials/MaterialInstance.cpp
@@ -3737,6 +3737,11 @@ void UMaterialInstance::AddReferencedObjects(UObject* InThis, FReferenceCollecto
 {
 	UMaterialInstance* This = CastChecked<UMaterialInstance>(InThis);
 
+	if (This->CachedData)
+	{
+		This->CachedData->AddReferencedObjects(Collector);
+	}
+
 	if (This->bHasStaticPermutationResource)
 	{
 		for (FMaterialResource* CurrentResource : This->StaticPermutationMaterialResources)
