Fix potential for array to be accessed from one thread while modified from another by storing array elements as pointers rather than by value
Edigrated from CL15214427

diff --git a/Engine/Source/Runtime/Engine/Private/Materials/MaterialShader.cpp b/Engine/Source/Runtime/Engine/Private/Materials/MaterialShader.cpp
--- a/Engine/Source/Runtime/Engine/Private/Materials/MaterialShader.cpp
+++ b/Engine/Source/Runtime/Engine/Private/Materials/MaterialShader.cpp
@@ -1927,46 +1927,50 @@ public:
 	const FMaterialShaderMapLayout& AcquireLayout(EShaderPlatform Platform, EShaderPermutationFlags Flags, const FMaterialShaderParameters& MaterialParameters)
 	{
 		const uint64 ParameterHash = CityHash64WithSeed((char*)&MaterialParameters, sizeof(MaterialParameters), (uint64)Platform | ((uint64)Flags << 32));
-		
-		int32 Index = INDEX_NONE;
+
+		FMaterialShaderMapLayout* Layout = nullptr;
 		{
 			FReadScopeLock Locker(LayoutLock);
-			Index = FindLayoutIndex(ParameterHash);
+			Layout = FindLayout(ParameterHash);
 		}
 
-		if (Index == INDEX_NONE)
+		if (!Layout)
 		{
 			FWriteScopeLock Locker(LayoutLock);
 			// Need to check for existing index again once we've taken the write-lock
 			// TODO maybe better to create the layout without any lock, then discard if another thread created the same layout
-			Index = FindLayoutIndex(ParameterHash);
-			if (Index == INDEX_NONE)
+			Layout = FindLayout(ParameterHash);
+			if (!Layout)
 			{
-				Index = MaterialParameterHashes.Add(ParameterHash);
+				Layout = new FMaterialShaderMapLayout();
+				CreateLayout(*Layout, Platform, Flags, MaterialParameters);
+				Layout->Platform = Platform;
+
+				const int32 Index = ShaderMapLayouts.Add(TUniquePtr<FMaterialShaderMapLayout>(Layout));
+				MaterialParameterHashes.Add(ParameterHash);
 				MaterialShaderParameters.Add(MaterialParameters);
 				ShaderMapHashTable.Add(ParameterHash, Index);
-				FMaterialShaderMapLayout& Layout = ShaderMapLayouts.AddDefaulted_GetRef();
+
 				check(MaterialParameterHashes.Num() == ShaderMapLayouts.Num());
 				check(MaterialShaderParameters.Num() == ShaderMapLayouts.Num());
-				Layout.Platform = Platform;
-				CreateLayout(Layout, Platform, Flags, MaterialParameters);
 			}
 		}
 
-		return ShaderMapLayouts[Index];
+		check(Layout);
+		return *Layout;
 	}
 
 private:
-	int32 FindLayoutIndex(uint64 ParameterHash)
+	FMaterialShaderMapLayout* FindLayout(uint64 ParameterHash)
 	{
 		for (int32 Index = ShaderMapHashTable.First(ParameterHash); ShaderMapHashTable.IsValid(Index); Index = ShaderMapHashTable.Next(Index))
 		{
 			if (MaterialParameterHashes[Index] == ParameterHash)
 			{
-				return Index;
+				return ShaderMapLayouts[Index].Get();
 			}
 		}
-		return INDEX_NONE;
+		return nullptr;
 	}
 
 	void CreateLayout(FMaterialShaderMapLayout& Layout, EShaderPlatform Platform, EShaderPermutationFlags Flags, const FMaterialShaderParameters& MaterialParameters)
@@ -2075,7 +2079,7 @@ private:
 		Hasher.GetHash(Layout.ShaderMapHash.Hash);
 	}
 
-	TArray<FMaterialShaderMapLayout> ShaderMapLayouts;
+	TArray<TUniquePtr<FMaterialShaderMapLayout>> ShaderMapLayouts;
 	TArray<FMaterialShaderParameters> MaterialShaderParameters;
 	TArray<uint64> MaterialParameterHashes;
 	FHashTable ShaderMapHashTable;
