Fix crash caused by render command being queued during FinishDestroy

diff --git a/Engine/Source/Runtime/Engine/Private/Texture.cpp b/Engine/Source/Runtime/Engine/Private/Texture.cpp
--- a/Engine/Source/Runtime/Engine/Private/Texture.cpp
+++ b/Engine/Source/Runtime/Engine/Private/Texture.cpp
@@ -559,11 +559,11 @@ void UTexture::FinishDestroy()
 	check(!bAsyncResourceReleaseHasBeenStarted || ReleaseFence.IsFenceComplete());
 	check(TextureReference.IsInitialized_GameThread() == false);
 
-	if(Resource)
+	if(PrivateResource)
 	{
 		// Free the resource.
-		delete Resource;
-		Resource = NULL;
+		delete PrivateResource;
+		PrivateResource = NULL;
 	}
 
 	CleanupCachedRunningPlatformData();
