Non-unity / non-pch build fixes

diff --git a/Engine/Source/Runtime/Engine/Public/ActorReferencesUtils.h b/Engine/Source/Runtime/Engine/Public/ActorReferencesUtils.h
--- a/Engine/Source/Runtime/Engine/Public/ActorReferencesUtils.h
+++ b/Engine/Source/Runtime/Engine/Public/ActorReferencesUtils.h
@@ -3,6 +3,7 @@
 #pragma once
 
 #include "CoreMinimal.h"
+#include "UObject/ObjectMacros.h"
 
 class AActor;
 
diff --git a/Engine/Source/Runtime/Engine/Public/AudioStreaming.h b/Engine/Source/Runtime/Engine/Public/AudioStreaming.h
--- a/Engine/Source/Runtime/Engine/Public/AudioStreaming.h
+++ b/Engine/Source/Runtime/Engine/Public/AudioStreaming.h
@@ -15,6 +15,7 @@ AudioStreaming.h: Definitions of classes used for audio streaming.
 #include "Async/AsyncWork.h"
 #include "Async/AsyncFileHandle.h"
 #include "HAL/ThreadSafeBool.h"
+#include "UObject/ObjectKey.h"
 
 class FSoundSource;
 class FSoundWaveProxy;
diff --git a/Engine/Source/Runtime/Renderer/Private/DynamicBVH.cpp b/Engine/Source/Runtime/Renderer/Private/DynamicBVH.cpp
--- a/Engine/Source/Runtime/Renderer/Private/DynamicBVH.cpp
+++ b/Engine/Source/Runtime/Renderer/Private/DynamicBVH.cpp
@@ -4,7 +4,7 @@
 
 #include "DynamicBVH.h"
 #include "Misc/AutomationTest.h"
-
+#include "Misc/OutputDeviceRedirector.h"
 
 FMortonArray::FMortonArray( const TArray< FBounds >& InBounds )
 	: Bounds( InBounds )
