WorldPartition: Extract UActorDescContainer for futur use case

diff --git a/Engine/Source/Runtime/Engine/Classes/GameFramework/Actor.h b/Engine/Source/Runtime/Engine/Classes/GameFramework/Actor.h
--- a/Engine/Source/Runtime/Engine/Classes/GameFramework/Actor.h
+++ b/Engine/Source/Runtime/Engine/Classes/GameFramework/Actor.h
@@ -804,7 +804,7 @@ public:
 	inline const FGuid& GetActorGuid() const { return ActorGuid; }
 
 private:
-	friend class UWorldPartition;
+	friend class UActorDescContainer;
 
 	/**
 	 * Creates an uninitialized actor descriptor. Can be overriden.
diff --git a/Engine/Source/Runtime/Engine/Private/WorldPartition/ActorDescContainer.cpp b/Engine/Source/Runtime/Engine/Private/WorldPartition/ActorDescContainer.cpp
new file mode 100644
--- /dev/null
+++ b/Engine/Source/Runtime/Engine/Private/WorldPartition/ActorDescContainer.cpp
@@ -0,0 +1,224 @@
+// Copyright Epic Games, Inc. All Rights Reserved.
+
+#include "WorldPartition/ActorDescContainer.h"
+
+#if WITH_EDITOR
+#include "Editor.h"
+#include "AssetRegistryModule.h"
+#include "WorldPartition/WorldPartitionHandle.h"
+#include "Misc/Base64.h"
+#endif
+
+UActorDescContainer::UActorDescContainer(const FObjectInitializer& ObjectInitializer)
+	: Super(ObjectInitializer)
+#if WITH_EDITOR
+	, bContainerInitialized(false)
+	, bIgnoreAssetRegistryEvents(false)
+#endif
+{
+}
+
+void UActorDescContainer::Initialize(FName InPackageName, bool bRegisterDelegates)
+{
+#if WITH_EDITOR
+	check(!bContainerInitialized);
+	ContainerPackageName = InPackageName;
+	TArray<FAssetData> Assets;
+	
+	if (!ContainerPackageName.IsNone())
+	{
+		TGuardValue<bool> IgnoreAssetRegistryEvents(bIgnoreAssetRegistryEvents, true);
+
+		const FString LevelPathStr = ContainerPackageName.ToString();
+		const FString LevelExternalActorsPath = ULevel::GetExternalActorsPath(LevelPathStr);
+
+		// Do a synchronous scan of the level external actors path.			
+		IAssetRegistry& AssetRegistry = FModuleManager::LoadModuleChecked<FAssetRegistryModule>(TEXT("AssetRegistry")).Get();
+		AssetRegistry.ScanPathsSynchronous({ LevelExternalActorsPath }, /*bForceRescan*/true, /*bIgnoreBlackListScanFilters*/true);
+
+		FARFilter Filter;
+		Filter.bRecursivePaths = true;
+		Filter.bIncludeOnlyOnDiskAssets = true;
+		Filter.PackagePaths.Add(*LevelExternalActorsPath);
+
+		AssetRegistry.GetAssets(Filter, Assets);
+	}
+
+	for (const FAssetData& Asset : Assets)
+	{
+		TUniquePtr<FWorldPartitionActorDesc>* NewActorDesc = new(ActorDescList) TUniquePtr<FWorldPartitionActorDesc>(GetActorDescriptor(Asset));
+		check(NewActorDesc->IsValid());
+				
+		Actors.Add((*NewActorDesc)->GetGuid(), NewActorDesc);
+	}
+
+	if (bRegisterDelegates)
+	{
+		RegisterDelegates();
+	}
+
+	bContainerInitialized = true;
+#endif
+}
+
+void UActorDescContainer::Uninitialize()
+{
+#if WITH_EDITOR
+	if (bContainerInitialized)
+	{
+		UnregisterDelegates();
+		bContainerInitialized = false;
+	}
+#endif
+}
+
+#if WITH_EDITOR
+TUniquePtr<FWorldPartitionActorDesc> UActorDescContainer::GetActorDescriptor(const FAssetData& InAssetData)
+{
+	FString ActorClass;
+	static FName NAME_ActorMetaDataClass(TEXT("ActorMetaDataClass"));
+	if (InAssetData.GetTagValue(NAME_ActorMetaDataClass, ActorClass))
+	{
+		FString ActorMetaDataStr;
+		static FName NAME_ActorMetaData(TEXT("ActorMetaData"));
+		if (InAssetData.GetTagValue(NAME_ActorMetaData, ActorMetaDataStr))
+		{
+			FWorldPartitionActorDescInitData ActorDescInitData;
+			ActorDescInitData.NativeClass = FindObjectChecked<UClass>(ANY_PACKAGE, *ActorClass, true);
+			ActorDescInitData.PackageName = InAssetData.PackageName;
+			ActorDescInitData.ActorPath = InAssetData.ObjectPath;
+			FBase64::Decode(ActorMetaDataStr, ActorDescInitData.SerializedData);
+
+			TUniquePtr<FWorldPartitionActorDesc> NewActorDesc(AActor::CreateClassActorDesc(ActorDescInitData.NativeClass));
+			NewActorDesc->Init(this, ActorDescInitData);
+			return NewActorDesc;
+		}
+	}
+
+	return nullptr;
+}
+
+bool UActorDescContainer::ShouldHandleAssetEvent(const FAssetData& InAssetData)
+{
+	// Ignore asset event when specifically asking to
+	if (bIgnoreAssetRegistryEvents)
+	{
+		return false;
+	}
+
+	// Ignore in-memory assets until they gets saved
+	if (InAssetData.HasAnyPackageFlags(PKG_NewlyCreated))
+	{
+		return false;
+	}
+
+	// Only handle actors
+	if (!InAssetData.GetClass()->IsChildOf<AActor>())
+	{
+		return false;
+	}
+
+	// Make sure asset contains the required tags
+	static FName NAME_ActorMetaDataClass(TEXT("ActorMetaDataClass"));
+	static FName NAME_ActorMetaData(TEXT("ActorMetaData"));
+	if (!InAssetData.FindTag(NAME_ActorMetaDataClass) || !InAssetData.FindTag(NAME_ActorMetaData))
+	{
+		return false;
+	}
+
+	// Only handle assets that belongs to our level
+	auto RemoveAfterFirstDot = [](const FString& InValue)
+	{
+		int32 DotIndex;
+		if (InValue.FindChar(TEXT('.'), DotIndex))
+		{
+			return InValue.LeftChop(InValue.Len() - DotIndex);
+		}
+		return InValue;
+	};
+
+	const FString ThisLevelPath = ContainerPackageName.ToString();
+	const FString AssetLevelPath = RemoveAfterFirstDot(InAssetData.ObjectPath.ToString());
+	return (ThisLevelPath == AssetLevelPath);
+}
+
+void UActorDescContainer::OnAssetAdded(const FAssetData& InAssetData)
+{
+	if (ShouldHandleAssetEvent(InAssetData))
+	{
+		TUniquePtr<FWorldPartitionActorDesc>* NewActorDesc = new(ActorDescList) TUniquePtr<FWorldPartitionActorDesc>(GetActorDescriptor(InAssetData));
+		check(NewActorDesc->IsValid());
+
+		check(!Actors.Contains((*NewActorDesc)->GetGuid()));
+		Actors.Add((*NewActorDesc)->GetGuid(), NewActorDesc);
+
+		OnActorDescAdded(*NewActorDesc);
+	}
+}
+
+void UActorDescContainer::OnAssetRemoved(const FAssetData& InAssetData)
+{
+	if (ShouldHandleAssetEvent(InAssetData))
+	{
+		TUniquePtr<FWorldPartitionActorDesc> NewActorDesc = GetActorDescriptor(InAssetData);
+		check(NewActorDesc.IsValid());
+
+		TUniquePtr<FWorldPartitionActorDesc>* ExistingActorDesc = Actors.FindChecked(NewActorDesc->GetGuid());
+
+		OnActorDescRemoved(*ExistingActorDesc);
+
+		Actors.Remove((*ExistingActorDesc)->GetGuid());
+		ExistingActorDesc->Release();
+	}
+}
+
+void UActorDescContainer::OnAssetUpdated(const FAssetData& InAssetData)
+{
+	if (ShouldHandleAssetEvent(InAssetData))
+	{
+		TUniquePtr<FWorldPartitionActorDesc> NewActorDesc = GetActorDescriptor(InAssetData);
+		check(NewActorDesc.IsValid());
+
+		TUniquePtr<FWorldPartitionActorDesc>* ExistingActorDesc = Actors.FindChecked(NewActorDesc->GetGuid());
+
+		// Pin the actor handle on the actor to prevent unloading it when unhashing
+		FWorldPartitionHandle ExistingActorHandle(ExistingActorDesc);
+		FWorldPartitionHandlePinRefScope ExistingActorHandlePin(ExistingActorHandle);
+
+		OnActorDescUpdating(*ExistingActorDesc);
+
+		// Transfer any reference count from external sources
+		NewActorDesc->TransferRefCounts(ExistingActorDesc->Get());
+
+		*ExistingActorDesc = MoveTemp(NewActorDesc);
+
+		OnActorDescUpdated(*ExistingActorDesc);
+	}
+}
+
+void UActorDescContainer::RegisterDelegates()
+{
+	if (GEditor && !IsTemplate())
+	{
+		IAssetRegistry& AssetRegistry = FModuleManager::LoadModuleChecked<FAssetRegistryModule>(TEXT("AssetRegistry")).Get();
+		AssetRegistry.OnAssetAdded().AddUObject(this, &UActorDescContainer::OnAssetAdded);
+		AssetRegistry.OnAssetRemoved().AddUObject(this, &UActorDescContainer::OnAssetRemoved);
+		AssetRegistry.OnAssetUpdated().AddUObject(this, &UActorDescContainer::OnAssetUpdated);
+	}
+}
+
+void UActorDescContainer::UnregisterDelegates()
+{
+	if (GEditor && !IsTemplate())
+	{
+		if (FAssetRegistryModule* AssetRegistryModule = FModuleManager::GetModulePtr<FAssetRegistryModule>(TEXT("AssetRegistry")))
+		{
+			IAssetRegistry& AssetRegistry = AssetRegistryModule->Get();
+			AssetRegistry.OnAssetAdded().RemoveAll(this);
+			AssetRegistry.OnAssetRemoved().RemoveAll(this);
+			AssetRegistry.OnAssetUpdated().RemoveAll(this);
+		}
+	}
+}
+
+#endif
\ No newline at end of file
diff --git a/Engine/Source/Runtime/Engine/Private/WorldPartition/Landscape/LandscapeActorDesc.cpp b/Engine/Source/Runtime/Engine/Private/WorldPartition/Landscape/LandscapeActorDesc.cpp
--- a/Engine/Source/Runtime/Engine/Private/WorldPartition/Landscape/LandscapeActorDesc.cpp
+++ b/Engine/Source/Runtime/Engine/Private/WorldPartition/Landscape/LandscapeActorDesc.cpp
@@ -6,7 +6,7 @@
 #include "LandscapeProxy.h"
 #include "LandscapeStreamingProxy.h"
 #include "LandscapeInfo.h"
-#include "WorldPartition/WorldPartition.h"
+#include "WorldPartition/ActorDescContainer.h"
 #include "WorldPartition/WorldPartitionHandle.h"
 #include "UObject/UE5MainStreamObjectVersion.h"
 
@@ -46,9 +46,9 @@ void FLandscapeActorDesc::OnRegister()
 {
 	FPartitionActorDesc::OnRegister();
 
-	if (ULandscapeInfo* LandscapeInfo = ULandscapeInfo::Find(WorldPartition->GetWorld(), GridGuid))
+	if (ULandscapeInfo* LandscapeInfo = ULandscapeInfo::Find(Container->GetWorld(), GridGuid))
 	{
-		FWorldPartitionHandle Handle(WorldPartition, GetGuid());
+		FWorldPartitionHandle Handle(Container, GetGuid());
 		check(Handle.IsValid());
 		LandscapeInfo->ProxyHandles.Add(MoveTemp(Handle));
 	}
@@ -58,9 +58,9 @@ void FLandscapeActorDesc::OnUnregister()
 {
 	FPartitionActorDesc::OnUnregister();
 
-	if (ULandscapeInfo* LandscapeInfo = ULandscapeInfo::Find(WorldPartition->GetWorld(), GridGuid))
+	if (ULandscapeInfo* LandscapeInfo = ULandscapeInfo::Find(Container->GetWorld(), GridGuid))
 	{
-		FWorldPartitionHandle Handle(WorldPartition, GetGuid());
+		FWorldPartitionHandle Handle(Container, GetGuid());
 		check(Handle.IsValid());
 		LandscapeInfo->ProxyHandles.Remove(Handle);
 	}
diff --git a/Engine/Source/Runtime/Engine/Private/WorldPartition/Landscape/LandscapeSplineActorDesc.cpp b/Engine/Source/Runtime/Engine/Private/WorldPartition/Landscape/LandscapeSplineActorDesc.cpp
--- a/Engine/Source/Runtime/Engine/Private/WorldPartition/Landscape/LandscapeSplineActorDesc.cpp
+++ b/Engine/Source/Runtime/Engine/Private/WorldPartition/Landscape/LandscapeSplineActorDesc.cpp
@@ -6,7 +6,7 @@
 
 #include "LandscapeSplineActor.h"
 #include "LandscapeInfo.h"
-#include "WorldPartition/WorldPartition.h"
+#include "WorldPartition/ActorDescContainer.h"
 #include "WorldPartition/WorldPartitionHandle.h"
 #include "UObject/UE5MainStreamObjectVersion.h"
 
@@ -35,9 +35,9 @@ void FLandscapeSplineActorDesc::OnRegister()
 {
 	FWorldPartitionActorDesc::OnRegister();
 
-	if (ULandscapeInfo* LandscapeInfo = ULandscapeInfo::Find(WorldPartition->GetWorld(), LandscapeGuid))
+	if (ULandscapeInfo* LandscapeInfo = ULandscapeInfo::Find(Container->GetWorld(), LandscapeGuid))
 	{
-		FWorldPartitionHandle Handle(WorldPartition, GetGuid());
+		FWorldPartitionHandle Handle(Container, GetGuid());
 		check(Handle.IsValid());
 		LandscapeInfo->SplineHandles.Add(MoveTemp(Handle));
 	}
@@ -47,9 +47,9 @@ void FLandscapeSplineActorDesc::OnUnregister()
 {
 	FWorldPartitionActorDesc::OnUnregister();
 
-	if (ULandscapeInfo* LandscapeInfo = ULandscapeInfo::Find(WorldPartition->GetWorld(), LandscapeGuid))
+	if (ULandscapeInfo* LandscapeInfo = ULandscapeInfo::Find(Container->GetWorld(), LandscapeGuid))
 	{
-		FWorldPartitionHandle Handle(WorldPartition, GetGuid());
+		FWorldPartitionHandle Handle(Container, GetGuid());
 		check(Handle.IsValid());
 		LandscapeInfo->SplineHandles.Remove(Handle);
 	}
diff --git a/Engine/Source/Runtime/Engine/Private/WorldPartition/WorldPartition.cpp b/Engine/Source/Runtime/Engine/Private/WorldPartition/WorldPartition.cpp
--- a/Engine/Source/Runtime/Engine/Private/WorldPartition/WorldPartition.cpp
+++ b/Engine/Source/Runtime/Engine/Private/WorldPartition/WorldPartition.cpp
@@ -26,14 +26,13 @@
 #include "Misc/ScopeExit.h"
 #include "ScopedTransaction.h"
 #include "UnrealEdMisc.h"
-#include "AssetRegistryModule.h"
 #include "WorldPartition/IWorldPartitionEditorModule.h"
 #include "WorldPartition/WorldPartitionLevelStreamingDynamic.h"
 #include "WorldPartition/WorldPartitionEditorHash.h"
 #include "WorldPartition/WorldPartitionEditorSpatialHash.h"
 #include "WorldPartition/WorldPartitionRuntimeHash.h"
 #include "WorldPartition/WorldPartitionRuntimeSpatialHash.h"
-#include "Misc/Base64.h"
+#include "Modules/ModuleManager.h"
 #endif //WITH_EDITOR
 
 DEFINE_LOG_CATEGORY(LogWorldPartition);
@@ -106,7 +105,6 @@ UWorldPartition::UWorldPartition(const FObjectInitializer& ObjectInitializer)
 #if WITH_EDITOR
 	, EditorHash(nullptr)
 	, WorldPartitionEditor(nullptr)
-	, bIgnoreAssetRegistryEvents(false)
 	, bForceGarbageCollection(false)
 	, bForceGarbageCollectionPurge(false)
 #endif
@@ -179,11 +177,6 @@ void UWorldPartition::Initialize(UWorld* InWorld, const FTransform& InTransform)
 		}
 
 		EditorHash->Initialize();
-
-		if (IsMainWorldPartition())
-		{
-			RegisterDelegates();
-		}
 	}
 
 	if (!RuntimeHash)
@@ -197,25 +190,7 @@ void UWorldPartition::Initialize(UWorld* InWorld, const FTransform& InTransform)
 	{
 		TArray<FAssetData> Assets;
 		UPackage* LevelPackage = OuterWorld->PersistentLevel->GetOutermost();
-
-		if (!LevelPackage->GetLoadedPath().GetPackageFName().IsNone())
-		{
-			TGuardValue<bool> IgnoreAssetRegistryEvents(bIgnoreAssetRegistryEvents, true);
-
-			const FString LevelPathStr = LevelPackage->GetLoadedPath().GetPackageName();
-			const FString LevelExternalActorsPath = ULevel::GetExternalActorsPath(LevelPathStr);
-
-			// Do a synchronous scan of the level external actors path.			
-			IAssetRegistry& AssetRegistry = FModuleManager::LoadModuleChecked<FAssetRegistryModule>(TEXT("AssetRegistry")).Get();
-			AssetRegistry.ScanPathsSynchronous({LevelExternalActorsPath}, /*bForceRescan*/true, /*bIgnoreBlackListScanFilters*/true);
-
-			FARFilter Filter;
-			Filter.bRecursivePaths = true;
-			Filter.bIncludeOnlyOnDiskAssets = true;
-			Filter.PackagePaths.Add(*LevelExternalActorsPath);
-
-			AssetRegistry.GetAssets(Filter, Assets);
-		}
+		FName PackageName = LevelPackage->GetLoadedPath().GetPackageFName();
 
 		bool bIsInstanced = (bEditorOnly && !IsRunningCommandlet()) ? OuterWorld->PersistentLevel->IsInstancedLevel() : false;
 
@@ -233,14 +208,17 @@ void UWorldPartition::Initialize(UWorld* InWorld, const FTransform& InTransform)
 			ReplaceTo = DestWorldName + TEXT(".") + DestWorldName;
 		}
 
-		for (const FAssetData& Asset : Assets)
+		const bool bRegisterDelegates = bEditorOnly && IsMainWorldPartition();
+		UActorDescContainer::Initialize(PackageName, bRegisterDelegates);
+		check(bContainerInitialized);
+
+		for (TPair<FGuid, TUniquePtr<FWorldPartitionActorDesc>*>& Pair : Actors)
 		{
-			TUniquePtr<FWorldPartitionActorDesc>* NewActorDesc = new(ActorDescList) TUniquePtr<FWorldPartitionActorDesc>(GetActorDescriptor(Asset));
-			check(NewActorDesc->IsValid());
+			TUniquePtr<FWorldPartitionActorDesc>* NewActorDesc = Pair.Value;
 
 			if (bIsInstanced)
 			{
-				const FString LongActorPackageName = Asset.PackageName.ToString();
+				const FString LongActorPackageName = (*NewActorDesc)->ActorPackage.ToString();
 				const FString ActorPackageName = FPaths::GetBaseFilename(LongActorPackageName);
 				const FString InstancedName = FString::Printf(TEXT("%s_InstanceOf_%s"), *LevelPackage->GetName(), *ActorPackageName);
 
@@ -249,8 +227,6 @@ void UWorldPartition::Initialize(UWorld* InWorld, const FTransform& InTransform)
 				(*NewActorDesc)->TransformInstance(ReplaceFrom, ReplaceTo, InstanceTransform);
 			}
 
-			Actors.Add((*NewActorDesc)->GetGuid(), NewActorDesc);
-
 			if (bEditorOnly)
 			{
 				HashActorDesc(NewActorDesc->Get());
@@ -318,6 +294,8 @@ void UWorldPartition::WorldPartitionOnLevelRemovedFromWorld(ULevel* Level, UWorl
 
 void UWorldPartition::Uninitialize()
 {
+	Super::Uninitialize();
+
 	if (IsInitialized())
 	{
 		check(World);
@@ -346,7 +324,6 @@ void UWorldPartition::Uninitialize()
 					UpdateLoadingEditorCell(Cell, /*bShouldBeLoaded*/false);
 				});
 			}
-			UnregisterDelegates();
 		}
 
 		EditorHash = nullptr;
@@ -383,13 +360,9 @@ bool UWorldPartition::IsMainWorldPartition() const
 #if WITH_EDITOR
 void UWorldPartition::RegisterDelegates()
 {
+	Super::RegisterDelegates();
 	if (GEditor && !IsTemplate())
 	{
-		IAssetRegistry& AssetRegistry = FModuleManager::LoadModuleChecked<FAssetRegistryModule>(TEXT("AssetRegistry")).Get();
-		AssetRegistry.OnAssetAdded().AddUObject(this, &UWorldPartition::OnAssetAdded);
-		AssetRegistry.OnAssetRemoved().AddUObject(this, &UWorldPartition::OnAssetRemoved);
-		AssetRegistry.OnAssetUpdated().AddUObject(this, &UWorldPartition::OnAssetUpdated);
-
 		FEditorDelegates::PreBeginPIE.AddUObject(this, &UWorldPartition::OnPreBeginPIE);
 		FEditorDelegates::EndPIE.AddUObject(this, &UWorldPartition::OnEndPIE);
 	}
@@ -397,16 +370,9 @@ void UWorldPartition::RegisterDelegates()
 
 void UWorldPartition::UnregisterDelegates()
 {
+	Super::UnregisterDelegates();
 	if (GEditor && !IsTemplate())
 	{
-		if (FAssetRegistryModule* AssetRegistryModule = FModuleManager::GetModulePtr<FAssetRegistryModule>(TEXT("AssetRegistry")))
-		{
-			IAssetRegistry& AssetRegistry = AssetRegistryModule->Get();
-			AssetRegistry.OnAssetAdded().RemoveAll(this);
-			AssetRegistry.OnAssetRemoved().RemoveAll(this);
-			AssetRegistry.OnAssetUpdated().RemoveAll(this);
-		}
-
 		FEditorDelegates::PreBeginPIE.RemoveAll(this);
 		FEditorDelegates::EndPIE.RemoveAll(this);
 	}
@@ -690,146 +656,42 @@ void UWorldPartition::UpdateLoadingEditorCell(UWorldPartitionEditorCell* Cell, b
 	}
 }
 
-void UWorldPartition::OnAssetAdded(const FAssetData& InAssetData)
+void UWorldPartition::OnActorDescAdded(const TUniquePtr<FWorldPartitionActorDesc>& NewActorDesc)
 {
-	if (ShouldHandleAssetEvent(InAssetData))
-	{
-		TUniquePtr<FWorldPartitionActorDesc>* NewActorDesc = new(ActorDescList) TUniquePtr<FWorldPartitionActorDesc>(GetActorDescriptor(InAssetData));
-		check(NewActorDesc->IsValid());
-
-		check(!Actors.Contains((*NewActorDesc)->GetGuid()));
-		Actors.Add((*NewActorDesc)->GetGuid(), NewActorDesc);
-
-		HashActorDesc(NewActorDesc->Get());
-
-		(*NewActorDesc)->OnRegister();
+	HashActorDesc(NewActorDesc.Get());
 
-		if (WorldPartitionEditor)
-		{
-			WorldPartitionEditor->Refresh();
-		}
-	}
-}
+	NewActorDesc->OnRegister();
 
-void UWorldPartition::OnAssetRemoved(const FAssetData& InAssetData)
-{
-	if (ShouldHandleAssetEvent(InAssetData))
+	if (WorldPartitionEditor)
 	{
-		TUniquePtr<FWorldPartitionActorDesc> NewActorDesc = GetActorDescriptor(InAssetData);
-		check(NewActorDesc.IsValid());
-
-		TUniquePtr<FWorldPartitionActorDesc>* ExistingActorDesc = Actors.FindChecked(NewActorDesc->GetGuid());
-
-		UnhashActorDesc(ExistingActorDesc->Get());
-
-		(*ExistingActorDesc)->OnUnregister();
-
-		Actors.Remove((*ExistingActorDesc)->GetGuid());
-		ExistingActorDesc->Release();
-
-		if (WorldPartitionEditor)
-		{
-			WorldPartitionEditor->Refresh();
-		}
+		WorldPartitionEditor->Refresh();
 	}
 }
 
-void UWorldPartition::OnAssetUpdated(const FAssetData& InAssetData)
+void UWorldPartition::OnActorDescRemoved(const TUniquePtr<FWorldPartitionActorDesc>& ActorDesc)
 {
-	if (ShouldHandleAssetEvent(InAssetData))
-	{
-		TUniquePtr<FWorldPartitionActorDesc> NewActorDesc = GetActorDescriptor(InAssetData);
-		check(NewActorDesc.IsValid());
-
-		TUniquePtr<FWorldPartitionActorDesc>* ExistingActorDesc = Actors.FindChecked(NewActorDesc->GetGuid());
-
-		// Pin the actor handle on the actor to prevent unloading it when unhashing
-		FWorldPartitionHandle ExistingActorHandle(ExistingActorDesc);
-		FWorldPartitionHandlePinRefScope ExistingActorHandlePin(ExistingActorHandle);
-
-		UnhashActorDesc(ExistingActorDesc->Get());
-
-		// Transfer any reference count from external sources
-		NewActorDesc->TransferRefCounts(ExistingActorDesc->Get());
-
-		*ExistingActorDesc = MoveTemp(NewActorDesc);
+	UnhashActorDesc(ActorDesc.Get());
+	ActorDesc->OnUnregister();
 		
-		HashActorDesc(ExistingActorDesc->Get());
-
-		if (WorldPartitionEditor)
-		{
-			WorldPartitionEditor->Refresh();
-		}
+	if (WorldPartitionEditor)
+	{
+		WorldPartitionEditor->Refresh();
 	}
 }
 
-bool UWorldPartition::ShouldHandleAssetEvent(const FAssetData& InAssetData)
+void UWorldPartition::OnActorDescUpdating(const TUniquePtr<FWorldPartitionActorDesc>& ActorDesc)
 {
-	// Ignore asset event when specifically asking to
-	if (bIgnoreAssetRegistryEvents)
-	{
-		return false;
-	}
-
-	// Ignore in-memory assets until they gets saved
-	if (InAssetData.HasAnyPackageFlags(PKG_NewlyCreated))
-	{
-		return false;
-	}
-
-	// Only handle actors
-	if (!InAssetData.GetClass()->IsChildOf<AActor>())
-	{
-		return false;
-	}
-
-	// Make sure asset contains the required tags
-	static FName NAME_ActorMetaDataClass(TEXT("ActorMetaDataClass"));
-	static FName NAME_ActorMetaData(TEXT("ActorMetaData"));
-	if (!InAssetData.FindTag(NAME_ActorMetaDataClass) || !InAssetData.FindTag(NAME_ActorMetaData))
-	{
-		return false;
-	}
-
-	// Only handle assets that belongs to our level
-	auto RemoveAfterFirstDot = [](const FString& InValue)
-	{
-		int32 DotIndex;
-		if (InValue.FindChar(TEXT('.'), DotIndex))
-		{
-			return InValue.LeftChop(InValue.Len() - DotIndex);
-		}
-		return InValue;
-	};
-
-	const FString ThisLevelPath = GetPackage()->GetLoadedPath().GetPackageName();
-	const FString AssetLevelPath = RemoveAfterFirstDot(InAssetData.ObjectPath.ToString());
-	return (ThisLevelPath == AssetLevelPath);
+	UnhashActorDesc(ActorDesc.Get());
 }
 
-TUniquePtr<FWorldPartitionActorDesc> UWorldPartition::GetActorDescriptor(const FAssetData& InAssetData)
+void UWorldPartition::OnActorDescUpdated(const TUniquePtr<FWorldPartitionActorDesc>& ActorDesc)
 {
-	FString ActorClass;
-	static FName NAME_ActorMetaDataClass(TEXT("ActorMetaDataClass"));
-	if (InAssetData.GetTagValue(NAME_ActorMetaDataClass, ActorClass))
+	HashActorDesc(ActorDesc.Get());
+
+	if (WorldPartitionEditor)
 	{
-		FString ActorMetaDataStr;
-		static FName NAME_ActorMetaData(TEXT("ActorMetaData"));
-		if (InAssetData.GetTagValue(NAME_ActorMetaData, ActorMetaDataStr))
-		{
-			FWorldPartitionActorDescInitData ActorDescInitData;
-			ActorDescInitData.NativeClass = FindObjectChecked<UClass>(ANY_PACKAGE, *ActorClass, true);
-			ActorDescInitData.PackageName = InAssetData.PackageName;
-			ActorDescInitData.ActorPath = InAssetData.ObjectPath;
-			FBase64::Decode(ActorMetaDataStr, ActorDescInitData.SerializedData);
-
-			TUniquePtr<FWorldPartitionActorDesc> NewActorDesc(AActor::CreateClassActorDesc(ActorDescInitData.NativeClass));
-			NewActorDesc->Init(this, ActorDescInitData);
-			return NewActorDesc;
-		}
+		WorldPartitionEditor->Refresh();
 	}
-
-	return nullptr;
 }
 
 void UWorldPartition::HashActorDesc(FWorldPartitionActorDesc* ActorDesc)
diff --git a/Engine/Source/Runtime/Engine/Private/WorldPartition/WorldPartitionActorDesc.cpp b/Engine/Source/Runtime/Engine/Private/WorldPartition/WorldPartitionActorDesc.cpp
--- a/Engine/Source/Runtime/Engine/Private/WorldPartition/WorldPartitionActorDesc.cpp
+++ b/Engine/Source/Runtime/Engine/Private/WorldPartition/WorldPartitionActorDesc.cpp
@@ -25,7 +25,7 @@ uint32 FWorldPartitionActorDesc::GlobalTag = 0;
 FWorldPartitionActorDesc::FWorldPartitionActorDesc()
 	: SoftRefCount(0)
 	, HardRefCount(0)
-	, WorldPartition(nullptr)
+	, Container(nullptr)
 	, Tag(0)
 {}
 
@@ -79,11 +79,11 @@ void FWorldPartitionActorDesc::Init(const AActor* InActor)
 
 	ActorLabel = *InActor->GetActorLabel(false);
 
-	WorldPartition = InActor->GetLevel()->GetWorldPartition();
-	check(WorldPartition);
+	Container = InActor->GetLevel()->GetWorldPartition();
+	check(Container);
 }
 
-void FWorldPartitionActorDesc::Init(UWorldPartition* InWorldPartition, const FWorldPartitionActorDescInitData& DescData)
+void FWorldPartitionActorDesc::Init(UActorDescContainer* InContainer, const FWorldPartitionActorDescInitData& DescData)
 {
 	ActorPackage = DescData.PackageName;
 	ActorPath = DescData.ActorPath;
@@ -108,7 +108,7 @@ void FWorldPartitionActorDesc::Init(UWorldPartition* InWorldPartition, const FWo
 		GridPlacement = DefaultGridPlacement;
 	}
 
-	WorldPartition = InWorldPartition;
+	Container = InContainer;
 
 	ActorPtr = FindObject<AActor>(nullptr, *ActorPath.ToString());
 }
@@ -233,6 +233,8 @@ AActor* FWorldPartitionActorDesc::Load() const
 	// The, if the actor isn't loaded, load it
 	if (ActorPtr.IsExplicitlyNull())
 	{
+		// For now we assume that an ActorDesc that gets loaded is owned by an actual WorldPartition
+		UWorldPartition* WorldPartition = CastChecked<UWorldPartition>(Container);
 		check(WorldPartition || GIsAutomationTesting);
 
 		const FLinkerInstancingContext* InstancingContext = nullptr;
@@ -286,7 +288,7 @@ void FWorldPartitionActorDesc::RegisterActor()
 {
 	if (AActor* Actor = GetActor())
 	{
-		ApplyActorTransform(WorldPartition->InstanceTransform);
+		ApplyActorTransform(CastChecked<UWorldPartition>(Container)->InstanceTransform);
 		Actor->GetLevel()->AddLoadedActor(Actor);
 	}
 }
@@ -298,7 +300,7 @@ void FWorldPartitionActorDesc::UnregisterActor()
 		if (!Actor->IsPendingKill())
 		{
 			Actor->GetLevel()->RemoveLoadedActor(Actor);
-			ApplyActorTransform(WorldPartition->InstanceTransform.Inverse());
+			ApplyActorTransform(CastChecked<UWorldPartition>(Container)->InstanceTransform.Inverse());
 		}
 	}
 }
diff --git a/Engine/Source/Runtime/Engine/Private/WorldPartition/WorldPartitionHandle.cpp b/Engine/Source/Runtime/Engine/Private/WorldPartition/WorldPartitionHandle.cpp
--- a/Engine/Source/Runtime/Engine/Private/WorldPartition/WorldPartitionHandle.cpp
+++ b/Engine/Source/Runtime/Engine/Private/WorldPartition/WorldPartitionHandle.cpp
@@ -2,12 +2,12 @@
 
 #include "WorldPartition/WorldPartitionHandle.h"
 #include "WorldPartition/WorldPartitionActorDesc.h"
-#include "WorldPartition/WorldPartition.h"
+#include "WorldPartition/ActorDescContainer.h"
 
 #if WITH_EDITOR
-TUniquePtr<FWorldPartitionActorDesc>* FWorldPartitionHandleUtils::GetActorDesc(UWorldPartition* WorldPartition, const FGuid& ActorGuid)
+TUniquePtr<FWorldPartitionActorDesc>* FWorldPartitionHandleUtils::GetActorDesc(UActorDescContainer* Container, const FGuid& ActorGuid)
 {
-	if (TUniquePtr<FWorldPartitionActorDesc>** ActorDescPtr = WorldPartition->Actors.Find(ActorGuid))
+	if (TUniquePtr<FWorldPartitionActorDesc>** ActorDescPtr = Container->Actors.Find(ActorGuid))
 	{
 		return *ActorDescPtr;
 	}
diff --git a/Engine/Source/Runtime/Engine/Public/WorldPartition/ActorDescContainer.h b/Engine/Source/Runtime/Engine/Public/WorldPartition/ActorDescContainer.h
new file mode 100644
--- /dev/null
+++ b/Engine/Source/Runtime/Engine/Public/WorldPartition/ActorDescContainer.h
@@ -0,0 +1,50 @@
+// Copyright Epic Games, Inc. All Rights Reserved.
+#pragma once
+
+#include "CoreMinimal.h"
+#include "UObject/ObjectMacros.h"
+#include "UObject/Object.h"
+#include "GameFramework/Actor.h"
+#include "WorldPartition/WorldPartitionActorDesc.h"
+#include "ActorDescContainer.generated.h"
+
+UCLASS()
+class ENGINE_API UActorDescContainer : public UObject
+{
+	GENERATED_UCLASS_BODY()
+public:
+	void Initialize(FName InPackageName, bool bRegisterDelegates);
+	virtual void Uninitialize();
+	
+#if WITH_EDITOR
+	// Asset registry events
+	virtual void OnAssetAdded(const FAssetData& InAssetData);
+	virtual void OnAssetRemoved(const FAssetData& InAssetData);
+	virtual void OnAssetUpdated(const FAssetData& InAssetData);
+#endif
+
+protected:
+#if WITH_EDITOR
+	TUniquePtr<FWorldPartitionActorDesc> GetActorDescriptor(const FAssetData& InAssetData);
+	
+	virtual void RegisterDelegates();
+	virtual void UnregisterDelegates();
+
+	virtual void OnActorDescAdded(const TUniquePtr<FWorldPartitionActorDesc>& NewActorDesc) {}
+	virtual void OnActorDescRemoved(const TUniquePtr<FWorldPartitionActorDesc>& ActorDesc) {}
+	virtual void OnActorDescUpdating(const TUniquePtr<FWorldPartitionActorDesc>& ActorDesc) {}
+	virtual void OnActorDescUpdated(const TUniquePtr<FWorldPartitionActorDesc>& ActorDesc) {}
+
+	bool ShouldHandleAssetEvent(const FAssetData& InAssetData);
+
+	bool bContainerInitialized;
+	bool bIgnoreAssetRegistryEvents;
+
+	TChunkedArray<TUniquePtr<FWorldPartitionActorDesc>> ActorDescList;
+
+	FName ContainerPackageName;
+public:
+	TMap<FGuid, TUniquePtr<FWorldPartitionActorDesc>*> Actors;
+	
+#endif
+};
\ No newline at end of file
diff --git a/Engine/Source/Runtime/Engine/Public/WorldPartition/WorldPartition.h b/Engine/Source/Runtime/Engine/Public/WorldPartition/WorldPartition.h
--- a/Engine/Source/Runtime/Engine/Public/WorldPartition/WorldPartition.h
+++ b/Engine/Source/Runtime/Engine/Public/WorldPartition/WorldPartition.h
@@ -7,6 +7,7 @@
 #include "GameFramework/Actor.h"
 #include "Templates/SubclassOf.h"
 #include "WorldPartition/WorldPartitionActorDesc.h"
+#include "WorldPartition/ActorDescContainer.h"
 
 #if WITH_EDITOR
 #include "PackageSourceControlHelper.h"
@@ -60,7 +61,7 @@ public:
 #endif
 
 UCLASS(AutoExpandCategories=(WorldPartition))
-class ENGINE_API UWorldPartition final : public UObject
+class ENGINE_API UWorldPartition final : public UActorDescContainer
 {
 	GENERATED_UCLASS_BODY()
 
@@ -83,13 +84,13 @@ private:
 	void OnPreBeginPIE(bool bStartSimulate);
 	void OnEndPIE(bool bStartSimulate);
 
-	// Asset registry events
-	void OnAssetAdded(const FAssetData& InAssetData);
-	void OnAssetRemoved(const FAssetData& InAssetData);
-	void OnAssetUpdated(const FAssetData& InAssetData);
+	// UActorDescContainer events
+	virtual void OnActorDescAdded(const TUniquePtr<FWorldPartitionActorDesc>& NewActorDesc) override;
+	virtual void OnActorDescRemoved(const TUniquePtr<FWorldPartitionActorDesc>& ActorDesc) override;
+	virtual void OnActorDescUpdating(const TUniquePtr<FWorldPartitionActorDesc>& ActorDesc) override;
+	virtual void OnActorDescUpdated(const TUniquePtr<FWorldPartitionActorDesc>& ActorDesc) override;
 
 	bool ShouldHandleAssetEvent(const FAssetData& InAssetData);
-	TUniquePtr<FWorldPartitionActorDesc> GetActorDescriptor(const FAssetData& InAssetData);
 #endif
 
 	//~ Begin UObject Interface
@@ -138,7 +139,7 @@ public:
 public:
 	void Initialize(UWorld* World, const FTransform& InTransform);
 	bool IsInitialized() const;
-	void Uninitialize();
+	virtual void Uninitialize() override;
 
 	void CleanupWorldPartition();
 
@@ -168,10 +169,6 @@ public:
 	UWorldPartitionRuntimeHash* RuntimeHash;
 
 #if WITH_EDITOR
-	TChunkedArray<TUniquePtr<FWorldPartitionActorDesc>> ActorDescList;
-	TMap<FGuid, TUniquePtr<FWorldPartitionActorDesc>*> Actors;
-	
-	bool bIgnoreAssetRegistryEvents;
 	bool bForceGarbageCollection;
 	bool bForceGarbageCollectionPurge;
 #endif
@@ -194,12 +191,12 @@ private:
 #endif
 
 	bool IsMainWorldPartition() const;
-
+		
+#if WITH_EDITOR
 	// Delegates registration
-	void RegisterDelegates();
-	void UnregisterDelegates();
+	virtual void RegisterDelegates() override;
+	virtual void UnregisterDelegates() override;
 
-#if WITH_EDITOR
 	void UpdateLoadingEditorCell(UWorldPartitionEditorCell* Cell, bool bShouldBeLoaded);
 	void HashActorDesc(FWorldPartitionActorDesc* ActorDesc);
 	void UnhashActorDesc(FWorldPartitionActorDesc* ActorDesc);
diff --git a/Engine/Source/Runtime/Engine/Public/WorldPartition/WorldPartitionActorDesc.h b/Engine/Source/Runtime/Engine/Public/WorldPartition/WorldPartitionActorDesc.h
--- a/Engine/Source/Runtime/Engine/Public/WorldPartition/WorldPartitionActorDesc.h
+++ b/Engine/Source/Runtime/Engine/Public/WorldPartition/WorldPartitionActorDesc.h
@@ -35,6 +35,7 @@ class ENGINE_API FWorldPartitionActorDesc
 #if WITH_EDITOR
 	friend class AActor;
 	friend class UWorldPartition;
+	friend class UActorDescContainer;
 	friend struct FWorldPartitionHandleImpl;
 	friend struct FWorldPartitionReferenceImpl;
 	friend struct FWorldPartitionHandleUtils;
@@ -111,7 +112,7 @@ public:
 	void UnregisterActor();
 
 	virtual void Init(const AActor* InActor);
-	void Init(UWorldPartition* InWorldPartition, const FWorldPartitionActorDescInitData& DescData);
+	void Init(UActorDescContainer* InContainer, const FWorldPartitionActorDescInitData& DescData);
 
 	void SerializeTo(TArray<uint8>& OutData);
 
@@ -152,7 +153,7 @@ protected:
 	mutable uint32					HardRefCount;
 	UClass*							ActorClass;
 	mutable TWeakObjectPtr<AActor>	ActorPtr;
-	UWorldPartition*				WorldPartition;
+	UActorDescContainer*			Container;
 
 public:
 	// Tagging
diff --git a/Engine/Source/Runtime/Engine/Public/WorldPartition/WorldPartitionHandle.h b/Engine/Source/Runtime/Engine/Public/WorldPartition/WorldPartitionHandle.h
--- a/Engine/Source/Runtime/Engine/Public/WorldPartition/WorldPartitionHandle.h
+++ b/Engine/Source/Runtime/Engine/Public/WorldPartition/WorldPartitionHandle.h
@@ -5,12 +5,12 @@
 #include "Templates/UniquePtr.h"
 
 #if WITH_EDITOR
-class UWorldPartition;
+class UActorDescContainer;
 class FWorldPartitionActorDesc;
 
 struct ENGINE_API FWorldPartitionHandleUtils
 {
-	static TUniquePtr<FWorldPartitionActorDesc>* GetActorDesc(UWorldPartition* WorldPartition, const FGuid& ActorGuid);
+	static TUniquePtr<FWorldPartitionActorDesc>* GetActorDesc(UActorDescContainer* Container, const FGuid& ActorGuid);
 	static bool IsActorDescLoaded(FWorldPartitionActorDesc* ActorDesc);
 };
 
@@ -31,8 +31,8 @@ public:
 		}
 	}
 
-	FORCEINLINE TWorldPartitionHandle(UWorldPartition* WorldPartition, const FGuid& ActorGuid)
-		: ActorDesc(FWorldPartitionHandleUtils::GetActorDesc(WorldPartition, ActorGuid))
+	FORCEINLINE TWorldPartitionHandle(UActorDescContainer* Container, const FGuid& ActorGuid)
+		: ActorDesc(FWorldPartitionHandleUtils::GetActorDesc(Container, ActorGuid))
 	{
 		if (IsValid())
 		{
