Bump SKELETALMESH_DERIVEDDATA_VER to flush out poisoned cache

diff --git a/Engine/Source/Runtime/Engine/Private/SkeletalMeshRenderData.cpp b/Engine/Source/Runtime/Engine/Private/SkeletalMeshRenderData.cpp
--- a/Engine/Source/Runtime/Engine/Private/SkeletalMeshRenderData.cpp
+++ b/Engine/Source/Runtime/Engine/Private/SkeletalMeshRenderData.cpp
@@ -59,7 +59,7 @@ static void SerializeLODInfoForDDC(USkeletalMesh* SkeletalMesh, FString& KeySuff
 // differences, etc.) replace the version GUID below with a new one.
 // In case of merge conflicts with DDC versions, you MUST generate a new GUID
 // and set this new GUID as the version.
-#define SKELETALMESH_DERIVEDDATA_VER TEXT("8ACA4B98D9DE4ECCA7C84BC6F1A0FB89")
+#define SKELETALMESH_DERIVEDDATA_VER TEXT("B4868F69D68341E3A567D9BD079A75A8")
 
 static const FString& GetSkeletalMeshDerivedDataVersion()
 {
