Fix Linux & Mac

diff --git a/Engine/Source/Runtime/Engine/Private/WorldPartition/WorldPartitionActorDesc.cpp b/Engine/Source/Runtime/Engine/Private/WorldPartition/WorldPartitionActorDesc.cpp
--- a/Engine/Source/Runtime/Engine/Private/WorldPartition/WorldPartitionActorDesc.cpp
+++ b/Engine/Source/Runtime/Engine/Private/WorldPartition/WorldPartitionActorDesc.cpp
@@ -15,6 +15,7 @@
 #include "Editor/GroupActor.h"
 #include "UObject/UE5MainStreamObjectVersion.h"
 #include "WorldPartition/WorldPartition.h"
+#include "WorldPartition/ActorDescContainer.h"
 #include "WorldPartition/DataLayer/DataLayer.h"
 #include "Engine/Public/ActorReferencesUtils.h"
 #endif
diff --git a/Engine/Source/Runtime/Engine/Public/WorldPartition/WorldPartitionActorDesc.h b/Engine/Source/Runtime/Engine/Public/WorldPartition/WorldPartitionActorDesc.h
--- a/Engine/Source/Runtime/Engine/Public/WorldPartition/WorldPartitionActorDesc.h
+++ b/Engine/Source/Runtime/Engine/Public/WorldPartition/WorldPartitionActorDesc.h
@@ -18,6 +18,8 @@ struct FWorldPartitionActorDescInitData
 	FName ActorPath;
 	TArray<uint8> SerializedData;
 };
+
+class UActorDescContainer;
 #endif
 
 #if WITH_DEV_AUTOMATION_TESTS
