World Partition - HLODActor
* Fixed GUID being copied on duplicate
* Fixed HLOD actors not being hidden properly when loading editing cells in a certain order

diff --git a/Engine/Source/Runtime/Engine/Private/WorldPartition/HLOD/HLODActor.cpp b/Engine/Source/Runtime/Engine/Private/WorldPartition/HLOD/HLODActor.cpp
--- a/Engine/Source/Runtime/Engine/Private/WorldPartition/HLOD/HLODActor.cpp
+++ b/Engine/Source/Runtime/Engine/Private/WorldPartition/HLOD/HLODActor.cpp
@@ -52,8 +52,6 @@ void AWorldPartitionHLOD::EndPlay(const EEndPlayReason::Type EndPlayReason)
 
 void AWorldPartitionHLOD::OnSubActorLoaded(AActor& Actor)
 {
-	check(!Actor.GetRootComponent() || Actor.GetRootComponent()->IsRegistered());
-	
 	bool bIsAlreadyInSet = false;
 	LoadedSubActors.Add(&Actor, &bIsAlreadyInSet);
 
@@ -90,8 +88,8 @@ void AWorldPartitionHLOD::SetupLoadedSubActors()
 		for (const FGuid& SubActorGuid : SubActors)
 		{
 			const FWorldPartitionActorDesc* SubActorDesc = WorldPartition->GetActorDesc(SubActorGuid);
-			AActor* SubActor = SubActorDesc ? SubActorDesc->GetActor() : nullptr;
-			if (SubActor && SubActor->GetRootComponent() && SubActor->GetRootComponent()->IsRegistered())
+			AActor* SubActor = SubActorDesc->IsLoaded() ? SubActorDesc->GetActor() : nullptr;
+			if (SubActor)
 			{
 				OnSubActorLoaded(*SubActor);
 			}
@@ -227,6 +225,15 @@ void AWorldPartitionHLOD::PostActorCreated()
 	HLODGuid = GetActorGuid();
 }
 
+void AWorldPartitionHLOD::PostDuplicate(bool bDuplicateForPIE)
+{
+	Super::PostDuplicate(bDuplicateForPIE);
+	if (!bDuplicateForPIE)
+	{
+		HLODGuid = GetActorGuid();
+	}
+}
+
 #endif // WITH_EDITOR
 
 
diff --git a/Engine/Source/Runtime/Engine/Private/WorldPartition/HLOD/HLODSubsystem.cpp b/Engine/Source/Runtime/Engine/Private/WorldPartition/HLOD/HLODSubsystem.cpp
--- a/Engine/Source/Runtime/Engine/Private/WorldPartition/HLOD/HLODSubsystem.cpp
+++ b/Engine/Source/Runtime/Engine/Private/WorldPartition/HLOD/HLODSubsystem.cpp
@@ -65,14 +65,11 @@ void UHLODSubsystem::OnActorLoaded(AActor& Actor)
 		}
 	}
 
-	if(!Actor.GetRootComponent() || Actor.GetRootComponent()->IsRegistered())
+	// If HLOD for this actor is already loaded, notify the HLOD that it should be hidden
+	AWorldPartitionHLOD** HLODActorPtr = ActorsToHLOD.Find(Actor.GetActorGuid());
+	if (HLODActorPtr)
 	{
-		// If HLOD for this actor is already loaded, setup LOD parenting
-		AWorldPartitionHLOD** HLODActorPtr = ActorsToHLOD.Find(Actor.GetActorGuid());
-		if (HLODActorPtr)
-		{
-			(*HLODActorPtr)->OnSubActorLoaded(Actor);
-		}
+		(*HLODActorPtr)->OnSubActorLoaded(Actor);
 	}
 }
 
diff --git a/Engine/Source/Runtime/Engine/Public/WorldPartition/HLOD/HLODActor.h b/Engine/Source/Runtime/Engine/Public/WorldPartition/HLOD/HLODActor.h
--- a/Engine/Source/Runtime/Engine/Public/WorldPartition/HLOD/HLODActor.h
+++ b/Engine/Source/Runtime/Engine/Public/WorldPartition/HLOD/HLODActor.h
@@ -62,7 +62,9 @@ protected:
 	virtual void PostUnregisterAllComponents() override;
 	virtual EActorGridPlacement GetDefaultGridPlacement() const override;
 	virtual TUniquePtr<class FWorldPartitionActorDesc> CreateClassActorDesc() const override;
+
 	virtual void PostActorCreated() override;
+	virtual void PostDuplicate(bool bDuplicateForPIE) override;
 #endif
 	//~ End AActor Interface.
 
@@ -96,7 +98,7 @@ private:
 	TSet<TWeakObjectPtr<AActor>> LoadedSubActors;
 #endif
 
-	UPROPERTY()
+	UPROPERTY(NonPIEDuplicateTransient, TextExportTransient, NonTransactional)
 	FGuid HLODGuid;
 
 	UPROPERTY()
