Fix lookup table being invalidated whenever an object serial number changes (i.e. in-place UStaticMesh reallocation on reimport)
Reenable ObjectCache by default

diff --git a/Engine/Source/Runtime/Engine/Private/ObjectCacheContext.cpp b/Engine/Source/Runtime/Engine/Private/ObjectCacheContext.cpp
--- a/Engine/Source/Runtime/Engine/Private/ObjectCacheContext.cpp
+++ b/Engine/Source/Runtime/Engine/Private/ObjectCacheContext.cpp
@@ -23,7 +23,7 @@ DEFINE_LOG_CATEGORY(LogObjectCache);
 
 static TAutoConsoleVariable<int32> CVarObjectReverseLookupMode(
 	TEXT("Editor.ObjectReverseLookupMode"),
-	0,
+	1,
 	TEXT("0 - Reverse lookup tables are computed every time they are needed (slower behavior) \n")
 	TEXT("1 - Maintain permanent reverse lookup tables (faster behavior) \n")
 	TEXT("2 - Comparison mode (slowest to do validation between both mode) \n"),
@@ -77,12 +77,12 @@ public:
 
 		// Remove the old reverse lookup mappings as they can't be read from the
 		// objects anymore.
-		TSet<FromType*>* OldFromObjects = ToFromMapping.Find(InTo);
+		TRawSet<FromType*>* OldFromObjects = ToFromMapping.Find(InTo);
 		if (OldFromObjects)
 		{
 			for (FromType* OldFrom : *OldFromObjects)
 			{
-				TSet<ToType*>* ToObjects = FromToMapping.Find(OldFrom);
+				TRawSet<ToType*>* ToObjects = FromToMapping.Find(OldFrom);
 				if (ToObjects)
 				{
 					ToObjects->Remove(InTo);
@@ -92,7 +92,7 @@ public:
 
 		if (NewFromObjects.Num())
 		{
-			TSet<FromType*>& FromObjects = ToFromMapping.FindOrAdd(InTo);
+			TRawSet<FromType*>& FromObjects = ToFromMapping.FindOrAdd(InTo);
 			FromObjects.Reset();
 
 			for (FromType* From : NewFromObjects)
@@ -115,7 +115,7 @@ public:
 		// Until we can get some kind of thread-safe CopyOnWrite behavior
 		// we need to extract a copy to avoid race conditions.
 		TArray<ToType*> OutTo;
-		if (const TSet<ToType*>* Set = FromToMapping.Find(InFrom))
+		if (const TRawSet<ToType*>* Set = FromToMapping.Find(InFrom))
 		{
 			OutTo.Reserve(Set->Num());
 			for (ToType* To : *Set)
@@ -134,7 +134,7 @@ public:
 		// Until we can get some kind of thread-safe CopyOnWrite behavior, 
 		// we need to extract a copy to avoid race conditions.
 		TArray<FromType*> OutFrom;
-		if (const TSet<FromType*>* Set = ToFromMapping.Find(InTo))
+		if (const TRawSet<FromType*>* Set = ToFromMapping.Find(InTo))
 		{
 			OutFrom.Reserve(Set->Num());
 			for (FromType* From : *Set)
@@ -154,13 +154,13 @@ public:
 		int32 ErrorCount = 0;
 		for (const auto& Kvp : FromToMapping)
 		{
-			const TSet<ToType*>* ToObjects = Other.FromToMapping.Find(Kvp.Key);
+			const TRawSet<ToType*>* ToObjects = Other.FromToMapping.Find(Kvp.Key);
 			for (const auto& To : Kvp.Value)
 			{
 				bool bIsPresent = ToObjects ? ToObjects->Contains(To) : false;
 				if (!bIsPresent)
 				{
-					UE_LOG(LogObjectCache, Display, TEXT("Missing a direct lookup from %s to %s mapping"), *Kvp.Key.ResolveObjectPtr()->GetFullName(), *To->GetFullName());
+					UE_LOG(LogObjectCache, Display, TEXT("Missing a direct lookup from %s to %s mapping"), *Kvp.Key->GetFullName(), *To->GetFullName());
 					ErrorCount++;
 				}
 			}
@@ -168,13 +168,13 @@ public:
 
 		for (const auto& Kvp : ToFromMapping)
 		{
-			const TSet<FromType*>* FromObjects = Other.ToFromMapping.Find(Kvp.Key);
+			const TRawSet<FromType*>* FromObjects = Other.ToFromMapping.Find(Kvp.Key);
 			for (const auto& From : Kvp.Value)
 			{
 				bool bIsPresent = FromObjects ? FromObjects->Contains(From) : false;
 				if (!bIsPresent)
 				{
-					UE_LOG(LogObjectCache, Display, TEXT("Missing a reverse lookup from %s to %s mapping"), *Kvp.Key.ResolveObjectPtr()->GetFullName(), *From->GetFullName());
+					UE_LOG(LogObjectCache, Display, TEXT("Missing a reverse lookup from %s to %s mapping"), *Kvp.Key->GetFullName(), *From->GetFullName());
 					ErrorCount++;
 				}
 			}
@@ -184,8 +184,38 @@ public:
 	}
 private:
 	mutable FRWLock Lock;
-	TMap<TObjectKey<FromType>, TSet<ToType*>> FromToMapping;
-	TMap<TObjectKey<ToType>, TSet<FromType*>> ToFromMapping;
+
+	// We use raw pointer to compare UObject here both for performance reason
+	// and because we can't use TObjectKey. Some objects are reallocated in-place
+	// with a different serial number but without changing their pointer.
+	// See: StaticAllocateObject for implementation detail.
+	// This trick is used to avoid having to track down all referencers of a UObject*
+	// just to update it, and is used for instance to replace the content of a UStaticMesh
+	// during a reimport without having to update any UStaticMeshComponent.
+	// Because it is the job of this class to maintain pointer to pointer lookups, we
+	// can't use the serial number of any object in our map, and by extention, we can't
+	// use TObjectKey nor FObjectKey to store our keys.
+	// We also assume that any UObject pointer in these maps are always valid
+	// because we require to be notified when those objects are destroyed.
+	template<typename ElementType>
+	struct TRawObjectKeyFuncs : BaseKeyFuncs<ElementType, void*, false>
+	{
+		typedef typename TTypeTraits<void*>::ConstPointerType KeyInitType;
+		typedef typename TCallTraits<ElementType>::ParamType ElementInitType;
+
+		static FORCEINLINE KeyInitType GetSetKey(ElementInitType Element) {	return (void*)Element; }
+		static FORCEINLINE bool Matches(KeyInitType A, KeyInitType B) { return A == B; }
+		static FORCEINLINE uint32 GetKeyHash(KeyInitType Key) { return GetTypeHash(Key); }
+	}; 
+	
+	template <typename InElementType>
+	using TRawSet = TSet<InElementType, TRawObjectKeyFuncs<InElementType>>;
+
+	template <typename InKeyType, typename InValueType>
+	using TRawMap = TMap<InKeyType, InValueType, FDefaultSetAllocator, TDefaultMapKeyFuncs<void*, InValueType, false>>;
+
+	TRawMap<FromType*, TRawSet<ToType*>> FromToMapping;
+	TRawMap<ToType*, TRawSet<FromType*>> ToFromMapping;
 };
 
 FObjectReverseLookupCache<UTexture, UMaterialInterface>            GTextureToMaterialLookupCache;
