Fixes bug in AHUD::GetActorsInSelectionRectangle where actors partially or fully behind the viewing plane, but fully outside of the selection box, could be included in selection. Also made underlying behavior of projection functions clearer and optional.

Adaptation of PR #7350

diff --git a/Engine/Source/Runtime/Engine/Classes/Engine/Canvas.h b/Engine/Source/Runtime/Engine/Classes/Engine/Canvas.h
--- a/Engine/Source/Runtime/Engine/Classes/Engine/Canvas.h
+++ b/Engine/Source/Runtime/Engine/Classes/Engine/Canvas.h
@@ -365,10 +365,11 @@ public:
 	/**
 	 * Transforms a 3D world-space vector into 2D screen coordinates.
 	 *
-	 * @param Location The vector to transform.
+	 * @param Location			The vector to transform.
+	 * @param bClampToZeroPlane	If true, 2D screen coordinates behind the viewing plane (-Z) will have Z set to 0 (leaving X and Y alone)
 	 * @return The transformed vector.
 	 */
-	FVector Project(FVector Location) const;
+	FVector Project(FVector Location, bool bClampToZeroPlane = true) const;
 
 	/** 
 	 * Transforms 2D screen coordinates into a 3D world-space origin and direction.
diff --git a/Engine/Source/Runtime/Engine/Classes/GameFramework/HUD.h b/Engine/Source/Runtime/Engine/Classes/GameFramework/HUD.h
--- a/Engine/Source/Runtime/Engine/Classes/GameFramework/HUD.h
+++ b/Engine/Source/Runtime/Engine/Classes/GameFramework/HUD.h
@@ -323,9 +323,13 @@ public:
 	UFUNCTION(BlueprintCallable, Category = HUD)
 	void DrawMaterialTriangle(UMaterialInterface* Material, FVector2D V0_Pos, FVector2D V1_Pos, FVector2D V2_Pos, FVector2D V0_UV, FVector2D V1_UV, FVector2D V2_UV, FLinearColor V0_Color = FLinearColor::White, FLinearColor V1_Color = FLinearColor::White, FLinearColor V2_Color = FLinearColor::White);
 	
-	/** Transforms a 3D world-space vector into 2D screen coordinates */
+	/** Transforms a 3D world-space vector into 2D screen coordinates
+	 * @param Location			The world-space position to transform
+	 * @param bClampToZeroPlane	If true, 2D screen coordinates behind the viewing plane (-Z) will have Z set to 0 (leaving X and Y alone)
+	 * @return The transformed vector
+	 */
 	UFUNCTION(BlueprintCallable, Category = HUD)
-	FVector Project(FVector Location) const;
+	FVector Project(FVector Location, bool bClampToZeroPlane = true) const;
 	
 	/** Transforms a 2D screen location into a 3D location and direction */
 	UFUNCTION(BlueprintCallable, Category = HUD)
diff --git a/Engine/Source/Runtime/Engine/Private/HUD.cpp b/Engine/Source/Runtime/Engine/Private/HUD.cpp
--- a/Engine/Source/Runtime/Engine/Private/HUD.cpp
+++ b/Engine/Source/Runtime/Engine/Private/HUD.cpp
@@ -1010,11 +1010,12 @@ void AHUD::DrawMaterialTriangle(UMaterialInterface* Material, FVector2D V0_Pos,
 		Canvas->DrawItem(TriangleItem);
 	}
 }
-FVector AHUD::Project(FVector Location) const
+
+FVector AHUD::Project(FVector Location, bool bClampToZeroPlane) const
 {
 	if (IsCanvasValid_WarnIfNot())
 	{
-		return Canvas->Project(Location);
+		return Canvas->Project(Location, bClampToZeroPlane);
 	}
 	return FVector(0, 0, 0);
 }
@@ -1077,23 +1078,23 @@ void AHUD::GetActorsInSelectionRectangle(TSubclassOf<class AActor> ClassFilter,
 		for (uint8 BoundsPointItr = 0; BoundsPointItr < 8; BoundsPointItr++)
 		{
 			// Project vert into screen space.
-			const FVector ProjectedWorldLocation = Project(BoxCenter + (BoundsPointMapping[BoundsPointItr] * BoxExtents));
-			// Add to 2D bounding box
-			ActorBox2D += FVector2D(ProjectedWorldLocation.X, ProjectedWorldLocation.Y);
+			const FVector ProjectedWorldLocation = Project(BoxCenter + (BoundsPointMapping[BoundsPointItr] * BoxExtents), true);
+			// Add to 2D bounding box if point is on the front side of the camera
+			if (ProjectedWorldLocation.Z > 0.f)
+			{
+				ActorBox2D += FVector2D(ProjectedWorldLocation.X, ProjectedWorldLocation.Y);
+			}
 		}
-
-		//Selection Box must fully enclose the Projected Actor Bounds
-		if (bActorMustBeFullyEnclosed)
+		// Only consider actor boxes that have valid points inside
+		if (ActorBox2D.bIsValid)
 		{
-			if (SelectionRectangle.IsInside(ActorBox2D))
+			//Selection Box must fully enclose the Projected Actor Bounds
+			if (bActorMustBeFullyEnclosed && SelectionRectangle.IsInside(ActorBox2D))
 			{
 				OutActors.Add(EachActor);
 			}
-		}
-		//Partial Intersection with Projected Actor Bounds
-		else
-		{
-			if (SelectionRectangle.Intersect(ActorBox2D))
+			//Partial Intersection with Projected Actor Bounds
+			else if (SelectionRectangle.Intersect(ActorBox2D))
 			{
 				OutActors.Add(EachActor);
 			}
diff --git a/Engine/Source/Runtime/Engine/Private/UserInterface/Canvas.cpp b/Engine/Source/Runtime/Engine/Private/UserInterface/Canvas.cpp
--- a/Engine/Source/Runtime/Engine/Private/UserInterface/Canvas.cpp
+++ b/Engine/Source/Runtime/Engine/Private/UserInterface/Canvas.cpp
@@ -1746,7 +1746,7 @@ void UCanvas::TextSize(const UFont* InFont, const FString& InText, float& XL, fl
 	YL = YLi;
 }
 
-FVector UCanvas::Project(FVector Location) const
+FVector UCanvas::Project(FVector Location, bool bClampToZeroPlane) const 
 {
 	FPlane V(0,0,0,0);
 
@@ -1762,7 +1762,7 @@ FVector UCanvas::Project(FVector Location) const
 	resultVec.Y = (ClipY/2.f) + (resultVec.Y*(ClipY/2.f));
 
 	// if behind the screen, clamp depth to the screen
-	if (V.W <= 0.0f)
+	if (bClampToZeroPlane && V.W <= 0.0f) 
 	{
 		resultVec.Z = 0.0f;
 	}
