Temporarily comment out an ensure to avoid build break

diff --git a/Engine/Source/Runtime/Engine/Classes/Components/StaticMeshComponent.h b/Engine/Source/Runtime/Engine/Classes/Components/StaticMeshComponent.h
--- a/Engine/Source/Runtime/Engine/Classes/Components/StaticMeshComponent.h
+++ b/Engine/Source/Runtime/Engine/Classes/Components/StaticMeshComponent.h
@@ -385,7 +385,7 @@ public:
 #if WITH_EDITOR
 		if (KnownStaticMesh != StaticMesh)
 		{
-			ensureMsgf(KnownStaticMesh == StaticMesh, TEXT("There is a missing call to NotifyIfStaticMeshChanged after StaticMesh has been overwritten"));
+			//ensureMsgf(KnownStaticMesh == StaticMesh, TEXT("There is a missing call to NotifyIfStaticMeshChanged after StaticMesh has been overwritten"));
 
 			// This is a last resort, we should have catched the property overwrite well before we reach this code
 			UStaticMeshComponent* MutableThis = const_cast<UStaticMeshComponent*>(this);
