Fixed bug where Actor's Data Layers would be read-only by mistake.

diff --git a/Engine/Source/Runtime/Engine/Private/ActorEditor.cpp b/Engine/Source/Runtime/Engine/Private/ActorEditor.cpp
--- a/Engine/Source/Runtime/Engine/Private/ActorEditor.cpp
+++ b/Engine/Source/Runtime/Engine/Private/ActorEditor.cpp
@@ -57,9 +57,10 @@ bool AActor::CanEditChange(const FProperty* PropertyThatWillChange) const
 		return false;
 	}
 
+	const bool bIsDataLayersProperty = (PropertyThatWillChange->GetFName() == GET_MEMBER_NAME_CHECKED(AActor, DataLayers));
 	if ((PropertyThatWillChange->GetFName() == GET_MEMBER_NAME_CHECKED(AActor, GridPlacement)) || 
 		(PropertyThatWillChange->GetFName() == GET_MEMBER_NAME_CHECKED(AActor, RuntimeGrid)) ||
-		(PropertyThatWillChange->GetFName() == GET_MEMBER_NAME_CHECKED(AActor, DataLayers)))
+		bIsDataLayersProperty)
 	{
 		if (!IsTemplate())
 		{
@@ -70,7 +71,7 @@ bool AActor::CanEditChange(const FProperty* PropertyThatWillChange) const
 			}
 		}
 
-		if (GetDefaultGridPlacement() != EActorGridPlacement::None)
+		if (!bIsDataLayersProperty && (GetDefaultGridPlacement() != EActorGridPlacement::None))
 		{
 			return false;
 		}
