Follow up CL to force audio stream caching in UE5

diff --git a/Engine/Source/Runtime/Engine/Private/AudioCompressionSettingsUtils.cpp b/Engine/Source/Runtime/Engine/Private/AudioCompressionSettingsUtils.cpp
--- a/Engine/Source/Runtime/Engine/Private/AudioCompressionSettingsUtils.cpp
+++ b/Engine/Source/Runtime/Engine/Private/AudioCompressionSettingsUtils.cpp
@@ -348,8 +348,9 @@ const FPlatformAudioCookOverrides* FPlatformCompressionUtilities::GetCookOverrid
 
 bool FPlatformCompressionUtilities::IsCurrentPlatformUsingStreamCaching()
 {
-	const FPlatformAudioCookOverrides* Settings = GetCookOverrides();
-	return Settings && Settings->bUseStreamCaching;
+// 	const FPlatformAudioCookOverrides* Settings = GetCookOverrides();
+// 	return Settings && Settings->bUseStreamCaching;
+	return true;
 }
 
 const FAudioStreamCachingSettings& FPlatformCompressionUtilities::GetStreamCachingSettingsForCurrentPlatform()
