Simplify default setting of VT pool size.

diff --git a/Engine/Source/Runtime/Engine/Private/VT/VirtualTexturePoolConfig.cpp b/Engine/Source/Runtime/Engine/Private/VT/VirtualTexturePoolConfig.cpp
--- a/Engine/Source/Runtime/Engine/Private/VT/VirtualTexturePoolConfig.cpp
+++ b/Engine/Source/Runtime/Engine/Private/VT/VirtualTexturePoolConfig.cpp
@@ -15,7 +15,7 @@ void UVirtualTexturePoolConfig::FindPoolConfig(TEnumAsByte<EPixelFormat> const*
 	for (int32 Id = Pools.Num() - 1; Id >= 0 ; Id--)
 	{
 		const FVirtualTextureSpacePoolConfig& Config = Pools[Id];
-		if (Config.MinTileSize <= InTileSize && Config.MaxTileSize >= InTileSize && InNumLayers == Config.Formats.Num())
+		if (Config.MinTileSize <= InTileSize && (Config.MaxTileSize == 0 || Config.MaxTileSize >= InTileSize) && InNumLayers == Config.Formats.Num())
 		{
 			bool bAllFormatsMatch = true;
 			for (int Layer = 0; Layer < InNumLayers && bAllFormatsMatch; ++Layer)
