Fix newly introduced call incompatibility when Engine + ShooterGame are upgraded to wrapped pointers.

diff --git a/Engine/Source/Runtime/Engine/Private/ObjectCacheContext.cpp b/Engine/Source/Runtime/Engine/Private/ObjectCacheContext.cpp
--- a/Engine/Source/Runtime/Engine/Private/ObjectCacheContext.cpp
+++ b/Engine/Source/Runtime/Engine/Private/ObjectCacheContext.cpp
@@ -551,7 +551,7 @@ TObjectCacheIterator<USkinnedMeshComponent> FObjectCacheContext::GetSkinnedMeshC
 		TempMap.Reserve(8192);
 		for (USkinnedMeshComponent* Component : GetSkinnedMeshComponents())
 		{
-			TempMap.FindOrAdd(Component->SkeletalMesh).Add(Component);
+			TempMap.FindOrAdd(ToRawPtr(Component->SkeletalMesh)).Add(Component);
 		}
 		SkeletalMeshToComponents = MoveTemp(TempMap);
 	}
