Automated wrapped object pointer upgrade for Engine + ShooterGame + ShooterGame referenced plugins

diff --git a/Engine/Source/Runtime/AIModule/Classes/BehaviorTree/Tasks/BTTask_PlaySound.h b/Engine/Source/Runtime/AIModule/Classes/BehaviorTree/Tasks/BTTask_PlaySound.h
--- a/Engine/Source/Runtime/AIModule/Classes/BehaviorTree/Tasks/BTTask_PlaySound.h
+++ b/Engine/Source/Runtime/AIModule/Classes/BehaviorTree/Tasks/BTTask_PlaySound.h
@@ -20,7 +20,7 @@ class AIMODULE_API UBTTask_PlaySound : public UBTTaskNode
 
 	/** CUE to play */
 	UPROPERTY(Category=Node, EditAnywhere)
-	USoundCue* SoundToPlay;
+	TObjectPtr<USoundCue> SoundToPlay;
 
 	virtual EBTNodeResult::Type ExecuteTask(UBehaviorTreeComponent& OwnerComp, uint8* NodeMemory) override;
 	virtual FString GetStaticDescription() const override;
diff --git a/Engine/Source/Runtime/AIModule/Classes/BehaviorTree/Tasks/BTTask_PushPawnAction.h b/Engine/Source/Runtime/AIModule/Classes/BehaviorTree/Tasks/BTTask_PushPawnAction.h
--- a/Engine/Source/Runtime/AIModule/Classes/BehaviorTree/Tasks/BTTask_PushPawnAction.h
+++ b/Engine/Source/Runtime/AIModule/Classes/BehaviorTree/Tasks/BTTask_PushPawnAction.h
@@ -20,7 +20,7 @@ class AIMODULE_API UBTTask_PushPawnAction : public UBTTask_PawnActionBase
 
 protected:
 	UPROPERTY(EditAnywhere, Instanced, Category = Action)
-	UPawnAction* Action;
+	TObjectPtr<UPawnAction> Action;
 
 	virtual EBTNodeResult::Type ExecuteTask(UBehaviorTreeComponent& OwnerComp, uint8* NodeMemory) override;
 	virtual FString GetStaticDescription() const override;
diff --git a/Engine/Source/Runtime/AIModule/Classes/BehaviorTree/Tasks/BTTask_RunBehavior.h b/Engine/Source/Runtime/AIModule/Classes/BehaviorTree/Tasks/BTTask_RunBehavior.h
--- a/Engine/Source/Runtime/AIModule/Classes/BehaviorTree/Tasks/BTTask_RunBehavior.h
+++ b/Engine/Source/Runtime/AIModule/Classes/BehaviorTree/Tasks/BTTask_RunBehavior.h
@@ -41,7 +41,7 @@ protected:
 
 	/** behavior to run */
 	UPROPERTY(Category = Node, EditAnywhere)
-	UBehaviorTree* BehaviorAsset;
+	TObjectPtr<UBehaviorTree> BehaviorAsset;
 
 	/** called when subtree is removed from active stack */
 	virtual void OnSubtreeDeactivated(UBehaviorTreeComponent& OwnerComp, EBTNodeResult::Type NodeResult);
diff --git a/Engine/Source/Runtime/AIModule/Classes/BehaviorTree/Tasks/BTTask_RunBehaviorDynamic.h b/Engine/Source/Runtime/AIModule/Classes/BehaviorTree/Tasks/BTTask_RunBehaviorDynamic.h
--- a/Engine/Source/Runtime/AIModule/Classes/BehaviorTree/Tasks/BTTask_RunBehaviorDynamic.h
+++ b/Engine/Source/Runtime/AIModule/Classes/BehaviorTree/Tasks/BTTask_RunBehaviorDynamic.h
@@ -42,11 +42,11 @@ protected:
 
 	/** default behavior to run */
 	UPROPERTY(Category=Node, EditAnywhere)
-	UBehaviorTree* DefaultBehaviorAsset;
+	TObjectPtr<UBehaviorTree> DefaultBehaviorAsset;
 
 	/** current subtree */
 	UPROPERTY()
-	UBehaviorTree* BehaviorAsset;
+	TObjectPtr<UBehaviorTree> BehaviorAsset;
 
 	/** called when subtree is removed from active stack */
 	virtual void OnSubtreeDeactivated(UBehaviorTreeComponent& OwnerComp, EBTNodeResult::Type NodeResult);
diff --git a/Engine/Source/Runtime/AIModule/Classes/BehaviorTree/Tasks/BTTask_RunEQSQuery.h b/Engine/Source/Runtime/AIModule/Classes/BehaviorTree/Tasks/BTTask_RunEQSQuery.h
--- a/Engine/Source/Runtime/AIModule/Classes/BehaviorTree/Tasks/BTTask_RunEQSQuery.h
+++ b/Engine/Source/Runtime/AIModule/Classes/BehaviorTree/Tasks/BTTask_RunEQSQuery.h
@@ -28,7 +28,7 @@ class AIMODULE_API UBTTask_RunEQSQuery : public UBTTask_BlackboardBase
 
 	/** query to run */
 	UPROPERTY(Category = Node, VisibleAnywhere, meta = (EditCondition = "!bUseBBKey", DisplayName = "QueryTemplate_DEPRECATED"))
-	UEnvQuery* QueryTemplate;
+	TObjectPtr<UEnvQuery> QueryTemplate;
 
 	/** optional parameters for query */
 	UPROPERTY(Category = Node, VisibleAnywhere, meta = (DisplayName = "QueryParams_DEPRECATED"))
diff --git a/Engine/Source/Runtime/AIModule/Classes/BrainComponent.h b/Engine/Source/Runtime/AIModule/Classes/BrainComponent.h
--- a/Engine/Source/Runtime/AIModule/Classes/BrainComponent.h
+++ b/Engine/Source/Runtime/AIModule/Classes/BrainComponent.h
@@ -121,10 +121,10 @@ class AIMODULE_API UBrainComponent : public UActorComponent, public IAIResourceI
 protected:
 	/** blackboard component */
 	UPROPERTY(transient)
-	UBlackboardComponent* BlackboardComp;
+	TObjectPtr<UBlackboardComponent> BlackboardComp;
 
 	UPROPERTY(transient)
-	AAIController* AIOwner;
+	TObjectPtr<AAIController> AIOwner;
 
 	// @TODO this is a temp contraption to implement delayed messages delivering
 	// until proper AI messaging is implemented
diff --git a/Engine/Source/Runtime/AIModule/Classes/DataProviders/AIDataProvider.h b/Engine/Source/Runtime/AIModule/Classes/DataProviders/AIDataProvider.h
--- a/Engine/Source/Runtime/AIModule/Classes/DataProviders/AIDataProvider.h
+++ b/Engine/Source/Runtime/AIModule/Classes/DataProviders/AIDataProvider.h
@@ -33,7 +33,7 @@ private:
 public:
 	/** (optional) provider for dynamic data binding */
 	UPROPERTY(EditAnywhere, Instanced, Category = Value)
-	UAIDataProvider* DataBinding;
+	TObjectPtr<UAIDataProvider> DataBinding;
 
 	/** name of provider's value property */
 	UPROPERTY(EditAnywhere, Category = Value)
@@ -82,7 +82,7 @@ struct AIMODULE_API FAIDataProviderTypedValue : public FAIDataProviderValue
 
 	/** type of value */
 	UPROPERTY()
-	UClass* PropertyType_DEPRECATED;
+	TObjectPtr<UClass> PropertyType_DEPRECATED;
 	FFieldClass* PropertyType;
 
 	/** filter for provider's properties */
diff --git a/Engine/Source/Runtime/AIModule/Classes/EnvironmentQuery/EQSTestingPawn.h b/Engine/Source/Runtime/AIModule/Classes/EnvironmentQuery/EQSTestingPawn.h
--- a/Engine/Source/Runtime/AIModule/Classes/EnvironmentQuery/EQSTestingPawn.h
+++ b/Engine/Source/Runtime/AIModule/Classes/EnvironmentQuery/EQSTestingPawn.h
@@ -29,7 +29,7 @@ class AIMODULE_API AEQSTestingPawn : public ACharacter, public IEQSQueryResultSo
 	GENERATED_UCLASS_BODY()
 	
 	UPROPERTY(Category=EQS, EditAnywhere)
-	UEnvQuery* QueryTemplate;
+	TObjectPtr<UEnvQuery> QueryTemplate;
 
 	/** optional parameters for query */
 	UPROPERTY(Category=EQS, VisibleAnywhere, meta=(DisplayName="QueryParams_DEPRECATED"))
@@ -72,7 +72,7 @@ class AIMODULE_API AEQSTestingPawn : public ACharacter, public IEQSQueryResultSo
 private:
 	/** Editor Preview */
 	UPROPERTY(Transient)
-	UEQSRenderingComponent* EdRenderComp;
+	TObjectPtr<UEQSRenderingComponent> EdRenderComp;
 #endif // WITH_EDITORONLY_DATA
 
 protected:
diff --git a/Engine/Source/Runtime/AIModule/Classes/EnvironmentQuery/EnvQuery.h b/Engine/Source/Runtime/AIModule/Classes/EnvironmentQuery/EnvQuery.h
--- a/Engine/Source/Runtime/AIModule/Classes/EnvironmentQuery/EnvQuery.h
+++ b/Engine/Source/Runtime/AIModule/Classes/EnvironmentQuery/EnvQuery.h
@@ -23,7 +23,7 @@ class AIMODULE_API UEnvQuery : public UDataAsset
 #if WITH_EDITORONLY_DATA
 	/** Graph for query */
 	UPROPERTY()
-	UEdGraph*	EdGraph;
+	TObjectPtr<UEdGraph>	EdGraph;
 #endif
 
 protected:
@@ -33,7 +33,7 @@ protected:
 	FName QueryName;
 
 	UPROPERTY()
-	TArray<UEnvQueryOption*> Options;
+	TArray<TObjectPtr<UEnvQueryOption>> Options;
 
 public:
 	/** Gather all required named params */
diff --git a/Engine/Source/Runtime/AIModule/Classes/EnvironmentQuery/EnvQueryManager.h b/Engine/Source/Runtime/AIModule/Classes/EnvironmentQuery/EnvQueryManager.h
--- a/Engine/Source/Runtime/AIModule/Classes/EnvironmentQuery/EnvQueryManager.h
+++ b/Engine/Source/Runtime/AIModule/Classes/EnvironmentQuery/EnvQueryManager.h
@@ -59,15 +59,15 @@ protected:
 
 	/** query to run */
 	UPROPERTY()
-	const UEnvQuery* QueryTemplate;
+	TObjectPtr<const UEnvQuery> QueryTemplate;
 
 	/** querier */
 	UPROPERTY()
-	UObject* Owner;
+	TObjectPtr<UObject> Owner;
 
 	/** world */
 	UPROPERTY()
-	UWorld* World;
+	TObjectPtr<UWorld> World;
 
 	/** list of named params */
 	TMap<FName, float> NamedParams;
@@ -83,7 +83,7 @@ struct FEnvQueryInstanceCache
 
 	/** query template, duplicated in manager's world */
 	UPROPERTY()
-	UEnvQuery* Template = nullptr;
+	TObjectPtr<UEnvQuery> Template = nullptr;
 
 	/** instance to duplicate */
 	FEnvQueryInstance Instance;
@@ -268,10 +268,10 @@ protected:
 
 	/** local cache of context objects for managing BP based objects */
 	UPROPERTY(transient)
-	TArray<UEnvQueryContext*> LocalContexts;
+	TArray<TObjectPtr<UEnvQueryContext>> LocalContexts;
 
 	UPROPERTY()
-	TArray<UEnvQueryInstanceBlueprintWrapper*> GCShieldedWrappers;
+	TArray<TObjectPtr<UEnvQueryInstanceBlueprintWrapper>> GCShieldedWrappers;
 
 	/** local contexts mapped by class names */
 	TMap<FName, UEnvQueryContext*> LocalContextMap;
diff --git a/Engine/Source/Runtime/AIModule/Classes/EnvironmentQuery/EnvQueryOption.h b/Engine/Source/Runtime/AIModule/Classes/EnvironmentQuery/EnvQueryOption.h
--- a/Engine/Source/Runtime/AIModule/Classes/EnvironmentQuery/EnvQueryOption.h
+++ b/Engine/Source/Runtime/AIModule/Classes/EnvironmentQuery/EnvQueryOption.h
@@ -16,10 +16,10 @@ class AIMODULE_API UEnvQueryOption : public UObject
 	GENERATED_UCLASS_BODY()
 
 	UPROPERTY()
-	UEnvQueryGenerator* Generator;
+	TObjectPtr<UEnvQueryGenerator> Generator;
 
 	UPROPERTY()
-	TArray<UEnvQueryTest*> Tests;
+	TArray<TObjectPtr<UEnvQueryTest>> Tests;
 
 	FText GetDescriptionTitle() const;
 	FText GetDescriptionDetails() const;
diff --git a/Engine/Source/Runtime/AIModule/Classes/EnvironmentQuery/EnvQueryTypes.h b/Engine/Source/Runtime/AIModule/Classes/EnvironmentQuery/EnvQueryTypes.h
--- a/Engine/Source/Runtime/AIModule/Classes/EnvironmentQuery/EnvQueryTypes.h
+++ b/Engine/Source/Runtime/AIModule/Classes/EnvironmentQuery/EnvQueryTypes.h
@@ -1295,7 +1295,7 @@ struct AIMODULE_API FEQSParametrizedQueryExecutionRequest
 	FEQSParametrizedQueryExecutionRequest();
 
 	UPROPERTY(Category = Node, EditAnywhere, meta = (EditCondition = "!bUseBBKeyForQueryTemplate"))
-	UEnvQuery* QueryTemplate;
+	TObjectPtr<UEnvQuery> QueryTemplate;
 
 	UPROPERTY(Category = Node, EditAnywhere)
 	TArray<FAIDynamicParam> QueryConfig;
diff --git a/Engine/Source/Runtime/AIModule/Classes/EnvironmentQuery/Generators/EnvQueryGenerator_Composite.h b/Engine/Source/Runtime/AIModule/Classes/EnvironmentQuery/Generators/EnvQueryGenerator_Composite.h
--- a/Engine/Source/Runtime/AIModule/Classes/EnvironmentQuery/Generators/EnvQueryGenerator_Composite.h
+++ b/Engine/Source/Runtime/AIModule/Classes/EnvironmentQuery/Generators/EnvQueryGenerator_Composite.h
@@ -18,7 +18,7 @@ class AIMODULE_API UEnvQueryGenerator_Composite : public UEnvQueryGenerator
 	GENERATED_UCLASS_BODY()
 
 	UPROPERTY(EditDefaultsOnly, Instanced, Category = Generator)
-	TArray<UEnvQueryGenerator*> Generators;
+	TArray<TObjectPtr<UEnvQueryGenerator>> Generators;
 
 	/** allow generators with different item types, use at own risk!
 	 *
diff --git a/Engine/Source/Runtime/AIModule/Classes/Navigation/CrowdFollowingComponent.h b/Engine/Source/Runtime/AIModule/Classes/Navigation/CrowdFollowingComponent.h
--- a/Engine/Source/Runtime/AIModule/Classes/Navigation/CrowdFollowingComponent.h
+++ b/Engine/Source/Runtime/AIModule/Classes/Navigation/CrowdFollowingComponent.h
@@ -161,7 +161,7 @@ protected:
 #if WITH_EDITORONLY_DATA
 	UE_DEPRECATED(4.26, "This member is deprecated and no longer used. Use AvoidanceInterface instead.")
 	UPROPERTY(transient)
-	class UCharacterMovementComponent* CharacterMovement;
+	TObjectPtr<class UCharacterMovementComponent> CharacterMovement;
 
 	/** DEPRECATED: Group mask for this agent - use IRVOAvoidanceInterface instead */
 	UPROPERTY()
diff --git a/Engine/Source/Runtime/AIModule/Classes/Navigation/CrowdManager.h b/Engine/Source/Runtime/AIModule/Classes/Navigation/CrowdManager.h
--- a/Engine/Source/Runtime/AIModule/Classes/Navigation/CrowdManager.h
+++ b/Engine/Source/Runtime/AIModule/Classes/Navigation/CrowdManager.h
@@ -267,7 +267,7 @@ public:
 protected:
 
 	UPROPERTY(transient)
-	ANavigationData* MyNavData;
+	TObjectPtr<ANavigationData> MyNavData;
 
 	/** obstacle avoidance params */
 	UPROPERTY(config, EditAnywhere, Category = Config)
@@ -342,7 +342,7 @@ protected:
 	 *	only on dedicated-server.
 	 *	Note that this is only relevant when ENABLE_DRAW_DEBUG is enabled*/
 	UPROPERTY(transient)
-	UWorld* DebugDrawingWorld;
+	TObjectPtr<UWorld> DebugDrawingWorld;
 #endif // WITH_EDITORONLY_DATA
 
 	/** try to initialize nav data from already existing ones */
diff --git a/Engine/Source/Runtime/AIModule/Classes/Navigation/GridPathFollowingComponent.h b/Engine/Source/Runtime/AIModule/Classes/Navigation/GridPathFollowingComponent.h
--- a/Engine/Source/Runtime/AIModule/Classes/Navigation/GridPathFollowingComponent.h
+++ b/Engine/Source/Runtime/AIModule/Classes/Navigation/GridPathFollowingComponent.h
@@ -40,7 +40,7 @@ protected:
 	void UpdateActiveGrid(const FVector& CurrentLocation);
 
 	UPROPERTY(Transient)
-	UNavLocalGridManager* GridManager;
+	TObjectPtr<UNavLocalGridManager> GridManager;
 
 	/** index of current grid */
 	int32 ActiveGridIdx;
diff --git a/Engine/Source/Runtime/AIModule/Classes/Navigation/NavLinkProxy.h b/Engine/Source/Runtime/AIModule/Classes/Navigation/NavLinkProxy.h
--- a/Engine/Source/Runtime/AIModule/Classes/Navigation/NavLinkProxy.h
+++ b/Engine/Source/Runtime/AIModule/Classes/Navigation/NavLinkProxy.h
@@ -37,7 +37,7 @@ class AIMODULE_API ANavLinkProxy : public AActor, public INavLinkHostInterface,
 private:
 	/** Smart link: can affect path following */
 	UPROPERTY(VisibleAnywhere, Category=SmartLink)
-	UNavLinkCustomComponent* SmartLinkComp;
+	TObjectPtr<UNavLinkCustomComponent> SmartLinkComp;
 public:
 
 	/** Smart link: toggle relevancy */
@@ -48,10 +48,10 @@ public:
 private:
 	/** Editor Preview */
 	UPROPERTY()
-	UNavLinkRenderingComponent* EdRenderComp;
+	TObjectPtr<UNavLinkRenderingComponent> EdRenderComp;
 
 	UPROPERTY()
-	UBillboardComponent* SpriteComponent;
+	TObjectPtr<UBillboardComponent> SpriteComponent;
 public:
 #endif // WITH_EDITORONLY_DATA
 
diff --git a/Engine/Source/Runtime/AIModule/Classes/Navigation/PathFollowingComponent.h b/Engine/Source/Runtime/AIModule/Classes/Navigation/PathFollowingComponent.h
--- a/Engine/Source/Runtime/AIModule/Classes/Navigation/PathFollowingComponent.h
+++ b/Engine/Source/Runtime/AIModule/Classes/Navigation/PathFollowingComponent.h
@@ -452,14 +452,14 @@ protected:
 
 	/** associated movement component */
 	UPROPERTY(transient)
-	UNavMovementComponent* MovementComp;
+	TObjectPtr<UNavMovementComponent> MovementComp;
 
 	/** currently traversed custom nav link */
 	FWeakObjectPtr CurrentCustomLinkOb;
 
 	/** navigation data for agent described in movement component */
 	UPROPERTY(transient)
-	ANavigationData* MyNavData;
+	TObjectPtr<ANavigationData> MyNavData;
 
 	/** requested path */
 	FNavPathSharedPtr Path;
diff --git a/Engine/Source/Runtime/AIModule/Classes/Navigation/RecastGraphAStar.h b/Engine/Source/Runtime/AIModule/Classes/Navigation/RecastGraphAStar.h
--- a/Engine/Source/Runtime/AIModule/Classes/Navigation/RecastGraphAStar.h
+++ b/Engine/Source/Runtime/AIModule/Classes/Navigation/RecastGraphAStar.h
@@ -104,7 +104,7 @@ protected:
 
 private:
 	UPROPERTY(Transient)
-	const ARecastNavMesh* RecastNavMeshActor = nullptr;
+	TObjectPtr<const ARecastNavMesh> RecastNavMeshActor = nullptr;
 
 	const dtNavMesh* DetourNavMesh = nullptr;
 
diff --git a/Engine/Source/Runtime/AIModule/Classes/Perception/AIPerceptionComponent.h b/Engine/Source/Runtime/AIModule/Classes/Perception/AIPerceptionComponent.h
--- a/Engine/Source/Runtime/AIModule/Classes/Perception/AIPerceptionComponent.h
+++ b/Engine/Source/Runtime/AIModule/Classes/Perception/AIPerceptionComponent.h
@@ -186,7 +186,7 @@ struct FActorPerceptionBlueprintInfo
 	GENERATED_USTRUCT_BODY()
 
 	UPROPERTY(BlueprintReadWrite, Category = "AI|Perception")
-	AActor* Target;
+	TObjectPtr<AActor> Target;
 
 	UPROPERTY(BlueprintReadWrite, Category = "AI|Perception")
 	TArray<FAIStimulus> LastSensedStimuli;
@@ -215,7 +215,7 @@ class AIMODULE_API UAIPerceptionComponent : public UActorComponent
 
 protected:
 	UPROPERTY(EditDefaultsOnly, Instanced, Category = "AI Perception")
-	TArray<UAISenseConfig*> SensesConfig;
+	TArray<TObjectPtr<UAISenseConfig>> SensesConfig;
 
 	/** Indicated sense that takes precedence over other senses when determining sensed actor's location. 
 	 *	Should be set to one of the senses configured in SensesConfig, or None. */
@@ -225,7 +225,7 @@ protected:
 	FAISenseID DominantSenseID;
 
 	UPROPERTY(Transient)
-	AAIController* AIOwner;
+	TObjectPtr<AAIController> AIOwner;
 
 	/** @todo this field is misnamed. It's a whitelist. */
 	FPerceptionChannelWhitelist PerceptionFilter;
diff --git a/Engine/Source/Runtime/AIModule/Classes/Perception/AIPerceptionSystem.h b/Engine/Source/Runtime/AIModule/Classes/Perception/AIPerceptionSystem.h
--- a/Engine/Source/Runtime/AIModule/Classes/Perception/AIPerceptionSystem.h
+++ b/Engine/Source/Runtime/AIModule/Classes/Perception/AIPerceptionSystem.h
@@ -43,7 +43,7 @@ protected:
 	AIPerception::FListenerMap ListenerContainer;
 
 	UPROPERTY()
-	TArray<UAISense*> Senses;
+	TArray<TObjectPtr<UAISense>> Senses;
 
 	UPROPERTY(config, EditAnywhere, Category = Perception)
 	float PerceptionAgingRate;
diff --git a/Engine/Source/Runtime/AIModule/Classes/Perception/AISense.h b/Engine/Source/Runtime/AIModule/Classes/Perception/AISense.h
--- a/Engine/Source/Runtime/AIModule/Classes/Perception/AISense.h
+++ b/Engine/Source/Runtime/AIModule/Classes/Perception/AISense.h
@@ -46,7 +46,7 @@ protected:
 	
 private:
 	UPROPERTY()
-	UAIPerceptionSystem* PerceptionSystemInstance;
+	TObjectPtr<UAIPerceptionSystem> PerceptionSystemInstance;
 
 	/** then this count reaches 0 sense will be updated */
 	float TimeUntilNextUpdate;
diff --git a/Engine/Source/Runtime/AIModule/Classes/Perception/AISense_Blueprint.h b/Engine/Source/Runtime/AIModule/Classes/Perception/AISense_Blueprint.h
--- a/Engine/Source/Runtime/AIModule/Classes/Perception/AISense_Blueprint.h
+++ b/Engine/Source/Runtime/AIModule/Classes/Perception/AISense_Blueprint.h
@@ -24,10 +24,10 @@ protected:
 	TSubclassOf<UUserDefinedStruct> ListenerDataType;
 
 	UPROPERTY(BlueprintReadOnly, Category = "Sense")
-	TArray<UAIPerceptionComponent*> ListenerContainer;
+	TArray<TObjectPtr<UAIPerceptionComponent>> ListenerContainer;
 
 	UPROPERTY()
-	TArray<UAISenseEvent*> UnprocessedEvents;
+	TArray<TObjectPtr<UAISenseEvent>> UnprocessedEvents;
 
 public:
 	UAISense_Blueprint(const FObjectInitializer& ObjectInitializer = FObjectInitializer::Get());
diff --git a/Engine/Source/Runtime/AIModule/Classes/Perception/AISense_Damage.h b/Engine/Source/Runtime/AIModule/Classes/Perception/AISense_Damage.h
--- a/Engine/Source/Runtime/AIModule/Classes/Perception/AISense_Damage.h
+++ b/Engine/Source/Runtime/AIModule/Classes/Perception/AISense_Damage.h
@@ -34,11 +34,11 @@ struct AIMODULE_API FAIDamageEvent
 	
 	/** Damaged actor */
 	UPROPERTY(EditAnywhere, BlueprintReadWrite, Category = "Sense")
-	AActor* DamagedActor;
+	TObjectPtr<AActor> DamagedActor;
 
 	/** Actor that instigated damage. Can be None */
 	UPROPERTY(EditAnywhere, BlueprintReadWrite, Category = "Sense")
-	AActor* Instigator;
+	TObjectPtr<AActor> Instigator;
 	
 	FAIDamageEvent();
 	FAIDamageEvent(AActor* InDamagedActor, AActor* InInstigator, float DamageAmount, const FVector& EventLocation, const FVector& InHitLocation = FAISystem::InvalidLocation); 
diff --git a/Engine/Source/Runtime/AIModule/Classes/Perception/AISense_Hearing.h b/Engine/Source/Runtime/AIModule/Classes/Perception/AISense_Hearing.h
--- a/Engine/Source/Runtime/AIModule/Classes/Perception/AISense_Hearing.h
+++ b/Engine/Source/Runtime/AIModule/Classes/Perception/AISense_Hearing.h
@@ -43,7 +43,7 @@ struct AIMODULE_API FAINoiseEvent
 	 * Actor triggering the sound.
 	 */
 	UPROPERTY(EditAnywhere, BlueprintReadWrite, Category = "Sense")
-	AActor* Instigator;
+	TObjectPtr<AActor> Instigator;
 
 	/**
 	 * Named identifier for the noise.
diff --git a/Engine/Source/Runtime/AIModule/Classes/Perception/AISense_Prediction.h b/Engine/Source/Runtime/AIModule/Classes/Perception/AISense_Prediction.h
--- a/Engine/Source/Runtime/AIModule/Classes/Perception/AISense_Prediction.h
+++ b/Engine/Source/Runtime/AIModule/Classes/Perception/AISense_Prediction.h
@@ -19,10 +19,10 @@ struct AIMODULE_API FAIPredictionEvent
 	typedef UAISense_Prediction FSenseClass;
 	
 	UPROPERTY()
-	AActor* Requestor;
+	TObjectPtr<AActor> Requestor;
 
 	UPROPERTY()
-	AActor* PredictedActor;
+	TObjectPtr<AActor> PredictedActor;
 
 	float TimeToPredict;
 		
diff --git a/Engine/Source/Runtime/AIModule/Classes/Perception/AISense_Sight.h b/Engine/Source/Runtime/AIModule/Classes/Perception/AISense_Sight.h
--- a/Engine/Source/Runtime/AIModule/Classes/Perception/AISense_Sight.h
+++ b/Engine/Source/Runtime/AIModule/Classes/Perception/AISense_Sight.h
@@ -33,10 +33,10 @@ struct AIMODULE_API FAISightEvent
 	ESightPerceptionEventName::Type EventType;	
 
 	UPROPERTY()
-	AActor* SeenActor;
+	TObjectPtr<AActor> SeenActor;
 
 	UPROPERTY()
-	AActor* Observer;
+	TObjectPtr<AActor> Observer;
 
 	FAISightEvent() : SeenActor(nullptr), Observer(nullptr) {}
 
diff --git a/Engine/Source/Runtime/AIModule/Classes/Perception/AISense_Team.h b/Engine/Source/Runtime/AIModule/Classes/Perception/AISense_Team.h
--- a/Engine/Source/Runtime/AIModule/Classes/Perception/AISense_Team.h
+++ b/Engine/Source/Runtime/AIModule/Classes/Perception/AISense_Team.h
@@ -27,10 +27,10 @@ public:
 	float Strength;
 private:
 	UPROPERTY()
-	AActor* Broadcaster;
+	TObjectPtr<AActor> Broadcaster;
 public:
 	UPROPERTY()
-	AActor* Enemy;
+	TObjectPtr<AActor> Enemy;
 		
 	FAITeamStimulusEvent() : Broadcaster(nullptr), Enemy(nullptr) {}
 	FAITeamStimulusEvent(AActor* InBroadcaster, AActor* InEnemy, const FVector& InLastKnowLocation, float EventRange, float PassedInfoAge = 0.f, float InStrength = 1.f);
diff --git a/Engine/Source/Runtime/AIModule/Classes/Perception/AISense_Touch.h b/Engine/Source/Runtime/AIModule/Classes/Perception/AISense_Touch.h
--- a/Engine/Source/Runtime/AIModule/Classes/Perception/AISense_Touch.h
+++ b/Engine/Source/Runtime/AIModule/Classes/Perception/AISense_Touch.h
@@ -19,10 +19,10 @@ struct AIMODULE_API FAITouchEvent
 	FVector Location;
 	
 	UPROPERTY()
-	AActor* TouchReceiver;
+	TObjectPtr<AActor> TouchReceiver;
 
 	UPROPERTY()
-	AActor* OtherActor;
+	TObjectPtr<AActor> OtherActor;
 		
 	FAITouchEvent() : TouchReceiver(nullptr), OtherActor(nullptr) {}
 	
diff --git a/Engine/Source/Runtime/AIModule/Classes/Tasks/AITask.h b/Engine/Source/Runtime/AIModule/Classes/Tasks/AITask.h
--- a/Engine/Source/Runtime/AIModule/Classes/Tasks/AITask.h
+++ b/Engine/Source/Runtime/AIModule/Classes/Tasks/AITask.h
@@ -26,7 +26,7 @@ class AIMODULE_API UAITask : public UGameplayTask
 protected:
 
 	UPROPERTY(BlueprintReadOnly, Category="AI|Tasks")
-	AAIController* OwnerController;
+	TObjectPtr<AAIController> OwnerController;
 
 	virtual void Activate() override;
 
diff --git a/Engine/Source/Runtime/AnimGraphRuntime/Public/AnimNodes/AnimNode_BlendListBase.h b/Engine/Source/Runtime/AnimGraphRuntime/Public/AnimNodes/AnimNode_BlendListBase.h
--- a/Engine/Source/Runtime/AnimGraphRuntime/Public/AnimNodes/AnimNode_BlendListBase.h
+++ b/Engine/Source/Runtime/AnimGraphRuntime/Public/AnimNodes/AnimNode_BlendListBase.h
@@ -46,10 +46,10 @@ protected:
 
 public:
 	UPROPERTY(EditAnywhere, Category=BlendType)
-	UCurveFloat* CustomBlendCurve;
+	TObjectPtr<UCurveFloat> CustomBlendCurve;
 
 	UPROPERTY(EditAnywhere, Category=BlendType, meta=(UseAsBlendProfile=true))
-	UBlendProfile* BlendProfile;
+	TObjectPtr<UBlendProfile> BlendProfile;
 
 	TArray<struct FAlphaBlend> Blends;
 
diff --git a/Engine/Source/Runtime/AnimGraphRuntime/Public/AnimNodes/AnimNode_BlendSpaceGraphBase.h b/Engine/Source/Runtime/AnimGraphRuntime/Public/AnimNodes/AnimNode_BlendSpaceGraphBase.h
--- a/Engine/Source/Runtime/AnimGraphRuntime/Public/AnimNodes/AnimNode_BlendSpaceGraphBase.h
+++ b/Engine/Source/Runtime/AnimGraphRuntime/Public/AnimNodes/AnimNode_BlendSpaceGraphBase.h
@@ -52,7 +52,7 @@ protected:
 
 	// The internal blendspace asset to play
 	UPROPERTY()
-	const UBlendSpaceBase* BlendSpace = nullptr;
+	TObjectPtr<const UBlendSpaceBase> BlendSpace = nullptr;
 
 	// Pose links for each sample in the blendspace
 	UPROPERTY()
diff --git a/Engine/Source/Runtime/AnimGraphRuntime/Public/AnimNodes/AnimNode_BlendSpacePlayer.h b/Engine/Source/Runtime/AnimGraphRuntime/Public/AnimNodes/AnimNode_BlendSpacePlayer.h
--- a/Engine/Source/Runtime/AnimGraphRuntime/Public/AnimNodes/AnimNode_BlendSpacePlayer.h
+++ b/Engine/Source/Runtime/AnimGraphRuntime/Public/AnimNodes/AnimNode_BlendSpacePlayer.h
@@ -48,7 +48,7 @@ public:
 
 	// The blendspace asset to play
 	UPROPERTY(EditAnywhere, BlueprintReadWrite, Category = Settings, meta = (PinHiddenByDefault))
-	UBlendSpaceBase* BlendSpace;
+	TObjectPtr<UBlendSpaceBase> BlendSpace;
 
 protected:
 
@@ -57,7 +57,7 @@ protected:
 	TArray<FBlendSampleData> BlendSampleDataCache;
 
 	UPROPERTY(Transient)
-	UBlendSpaceBase* PreviousBlendSpace;
+	TObjectPtr<UBlendSpaceBase> PreviousBlendSpace;
 
 public:	
 	FAnimNode_BlendSpacePlayer();
diff --git a/Engine/Source/Runtime/AnimGraphRuntime/Public/AnimNodes/AnimNode_LayeredBoneBlend.h b/Engine/Source/Runtime/AnimGraphRuntime/Public/AnimNodes/AnimNode_LayeredBoneBlend.h
--- a/Engine/Source/Runtime/AnimGraphRuntime/Public/AnimNodes/AnimNode_LayeredBoneBlend.h
+++ b/Engine/Source/Runtime/AnimGraphRuntime/Public/AnimNodes/AnimNode_LayeredBoneBlend.h
@@ -40,7 +40,7 @@ public:
 	 * Blend masks are used when BlendMode is BlendMask.
 	 */
 	UPROPERTY(EditAnywhere, editfixedsize, Category=Config, meta=(UseAsBlendMask=true))
-	TArray<UBlendProfile*> BlendMasks;
+	TArray<TObjectPtr<UBlendProfile>> BlendMasks;
 
 	/** 
 	 * Configuration for the parts of the skeleton to blend for each layer. Allows
diff --git a/Engine/Source/Runtime/AnimGraphRuntime/Public/AnimNodes/AnimNode_PoseBlendNode.h b/Engine/Source/Runtime/AnimGraphRuntime/Public/AnimNodes/AnimNode_PoseBlendNode.h
--- a/Engine/Source/Runtime/AnimGraphRuntime/Public/AnimNodes/AnimNode_PoseBlendNode.h
+++ b/Engine/Source/Runtime/AnimGraphRuntime/Public/AnimNodes/AnimNode_PoseBlendNode.h
@@ -25,7 +25,7 @@ public:
 
 	/** If you're using Custom BlendOption, you can specify curve */
 	UPROPERTY(EditAnywhere, Category = "Blend")
-	UCurveFloat* CustomCurve;
+	TObjectPtr<UCurveFloat> CustomCurve;
 
 public:	
 	FAnimNode_PoseBlendNode();
diff --git a/Engine/Source/Runtime/AnimGraphRuntime/Public/AnimNodes/AnimNode_PoseHandler.h b/Engine/Source/Runtime/AnimGraphRuntime/Public/AnimNodes/AnimNode_PoseHandler.h
--- a/Engine/Source/Runtime/AnimGraphRuntime/Public/AnimNodes/AnimNode_PoseHandler.h
+++ b/Engine/Source/Runtime/AnimGraphRuntime/Public/AnimNodes/AnimNode_PoseHandler.h
@@ -19,7 +19,7 @@ struct ANIMGRAPHRUNTIME_API FAnimNode_PoseHandler : public FAnimNode_AssetPlayer
 public:
 	// The animation sequence asset to evaluate
 	UPROPERTY(EditAnywhere, BlueprintReadWrite, Category = Settings, meta = (PinHiddenByDefault))
-	UPoseAsset* PoseAsset;
+	TObjectPtr<UPoseAsset> PoseAsset;
 
 public:	
 	FAnimNode_PoseHandler()
diff --git a/Engine/Source/Runtime/AnimGraphRuntime/Public/AnimNodes/AnimNode_RandomPlayer.h b/Engine/Source/Runtime/AnimGraphRuntime/Public/AnimNodes/AnimNode_RandomPlayer.h
--- a/Engine/Source/Runtime/AnimGraphRuntime/Public/AnimNodes/AnimNode_RandomPlayer.h
+++ b/Engine/Source/Runtime/AnimGraphRuntime/Public/AnimNodes/AnimNode_RandomPlayer.h
@@ -37,7 +37,7 @@ struct FRandomPlayerSequenceEntry
 
 	/** Sequence to play when this entry is picked */
 	UPROPERTY(EditAnywhere, BlueprintReadWrite, Category = "Settings")
-	UAnimSequence* Sequence;
+	TObjectPtr<UAnimSequence> Sequence;
 
 	/** When not in shuffle mode, this is the chance this entry will play (normalized against all other sample chances) */
 	UPROPERTY(EditAnywhere, BlueprintReadWrite, Category = "Settings", meta = (UIMin = "0", ClampMin = "0"))
diff --git a/Engine/Source/Runtime/AnimGraphRuntime/Public/AnimNodes/AnimNode_SequenceEvaluator.h b/Engine/Source/Runtime/AnimGraphRuntime/Public/AnimNodes/AnimNode_SequenceEvaluator.h
--- a/Engine/Source/Runtime/AnimGraphRuntime/Public/AnimNodes/AnimNode_SequenceEvaluator.h
+++ b/Engine/Source/Runtime/AnimGraphRuntime/Public/AnimNodes/AnimNode_SequenceEvaluator.h
@@ -32,7 +32,7 @@ struct ANIMGRAPHRUNTIME_API FAnimNode_SequenceEvaluator : public FAnimNode_Asset
 public:
 	// The animation sequence asset to evaluate
 	UPROPERTY(EditAnywhere, BlueprintReadWrite, Category = Settings, meta = (PinHiddenByDefault))
-	UAnimSequenceBase* Sequence;
+	TObjectPtr<UAnimSequenceBase> Sequence;
 
 	// The time at which to evaluate the associated sequence
 	UPROPERTY(EditAnywhere, BlueprintReadWrite, Category=Settings, meta=(PinShownByDefault))
diff --git a/Engine/Source/Runtime/AnimGraphRuntime/Public/BoneControllers/AnimNode_BoneDrivenController.h b/Engine/Source/Runtime/AnimGraphRuntime/Public/BoneControllers/AnimNode_BoneDrivenController.h
--- a/Engine/Source/Runtime/AnimGraphRuntime/Public/BoneControllers/AnimNode_BoneDrivenController.h
+++ b/Engine/Source/Runtime/AnimGraphRuntime/Public/BoneControllers/AnimNode_BoneDrivenController.h
@@ -53,7 +53,7 @@ struct ANIMGRAPHRUNTIME_API FAnimNode_BoneDrivenController : public FAnimNode_Sk
 
 	/** Curve used to map from the source attribute to the driven attributes if present (otherwise the Multiplier will be used) */
 	UPROPERTY(EditAnywhere, Category=Mapping)
-	UCurveFloat* DrivingCurve;
+	TObjectPtr<UCurveFloat> DrivingCurve;
 
 	// Multiplier to apply to the input value (Note: Ignored when a curve is used)
 	UPROPERTY(EditAnywhere, Category=Mapping)
diff --git a/Engine/Source/Runtime/AnimGraphRuntime/Public/BoneControllers/AnimNode_RigidBody.h b/Engine/Source/Runtime/AnimGraphRuntime/Public/BoneControllers/AnimNode_RigidBody.h
--- a/Engine/Source/Runtime/AnimGraphRuntime/Public/BoneControllers/AnimNode_RigidBody.h
+++ b/Engine/Source/Runtime/AnimGraphRuntime/Public/BoneControllers/AnimNode_RigidBody.h
@@ -151,7 +151,7 @@ struct ANIMGRAPHRUNTIME_API FAnimNode_RigidBody : public FAnimNode_SkeletalContr
 public:
 	/** Physics asset to use. If empty use the skeletal mesh's default physics asset */
 	UPROPERTY(EditAnywhere, Category = Settings)
-	UPhysicsAsset* OverridePhysicsAsset;
+	TObjectPtr<UPhysicsAsset> OverridePhysicsAsset;
 
 private:
 	FTransform PreviousCompWorldSpaceTM;
diff --git a/Engine/Source/Runtime/AudioAnalyzer/Classes/AudioAnalyzer.h b/Engine/Source/Runtime/AudioAnalyzer/Classes/AudioAnalyzer.h
--- a/Engine/Source/Runtime/AudioAnalyzer/Classes/AudioAnalyzer.h
+++ b/Engine/Source/Runtime/AudioAnalyzer/Classes/AudioAnalyzer.h
@@ -82,7 +82,7 @@ public:
 
 	/** The UAudioBus which is analyzed in real-time. */
 	UPROPERTY(Transient)
-	UAudioBus* AudioBus;
+	TObjectPtr<UAudioBus> AudioBus;
 
 	/** Starts analyzing audio from the given audio bus. Optionally override the audio bus desired to analyze. */
 	UFUNCTION(BlueprintCallable, BlueprintCosmetic, Category = AudioAnalyzer, meta = (WorldContext = "WorldContextObject"))
@@ -136,7 +136,7 @@ private:
 
 	// Audio analysis subsystem used with this audio analyzer
 	UPROPERTY(Transient)
-	UAudioAnalyzerSubsystem* AudioAnalyzerSubsystem;
+	TObjectPtr<UAudioAnalyzerSubsystem> AudioAnalyzerSubsystem;
 
 	// Output patch for retrieving audio from audio bus for analysis
 	Audio::FPatchOutputStrongPtr PatchOutputStrongPtr;
diff --git a/Engine/Source/Runtime/AudioAnalyzer/Classes/AudioAnalyzerNRT.h b/Engine/Source/Runtime/AudioAnalyzer/Classes/AudioAnalyzerNRT.h
--- a/Engine/Source/Runtime/AudioAnalyzer/Classes/AudioAnalyzerNRT.h
+++ b/Engine/Source/Runtime/AudioAnalyzer/Classes/AudioAnalyzerNRT.h
@@ -78,7 +78,7 @@ class AUDIOANALYZER_API UAudioAnalyzerNRT : public UAudioAnalyzerAssetBase
 		 * The USoundWave which is analyzed.
 		 */
 		UPROPERTY(EditAnywhere, BlueprintReadOnly, Category=AudioAnalyzer)
-		USoundWave* Sound;
+		TObjectPtr<USoundWave> Sound;
 
 		/** The duration of the analyzed audio in seconds. */
 		UPROPERTY(BlueprintReadOnly, Category=AudioAnalyzer)
diff --git a/Engine/Source/Runtime/AudioAnalyzer/Classes/AudioAnalyzerSubsystem.h b/Engine/Source/Runtime/AudioAnalyzer/Classes/AudioAnalyzerSubsystem.h
--- a/Engine/Source/Runtime/AudioAnalyzer/Classes/AudioAnalyzerSubsystem.h
+++ b/Engine/Source/Runtime/AudioAnalyzer/Classes/AudioAnalyzerSubsystem.h
@@ -37,7 +37,7 @@ public:
 private:
 
 	UPROPERTY(Transient);
-	TArray<UAudioAnalyzer*> AudioAnalyzers;
+	TArray<TObjectPtr<UAudioAnalyzer>> AudioAnalyzers;
 };
 
 
diff --git a/Engine/Source/Runtime/AudioExtensions/Public/ISoundfieldFormat.h b/Engine/Source/Runtime/AudioExtensions/Public/ISoundfieldFormat.h
--- a/Engine/Source/Runtime/AudioExtensions/Public/ISoundfieldFormat.h
+++ b/Engine/Source/Runtime/AudioExtensions/Public/ISoundfieldFormat.h
@@ -374,7 +374,7 @@ public:
 	 * TODO: Filter classes settable on here by GetSettingsClass.
 	 */
 	UPROPERTY(EditAnywhere, Category = EffectPreset)
-	USoundfieldEffectSettingsBase* Settings;
+	TObjectPtr<USoundfieldEffectSettingsBase> Settings;
 
 protected:
 	/*
diff --git a/Engine/Source/Runtime/AudioMixer/Classes/SubmixEffects/AudioMixerSubmixEffectDynamicsProcessor.h b/Engine/Source/Runtime/AudioMixer/Classes/SubmixEffects/AudioMixerSubmixEffectDynamicsProcessor.h
--- a/Engine/Source/Runtime/AudioMixer/Classes/SubmixEffects/AudioMixerSubmixEffectDynamicsProcessor.h
+++ b/Engine/Source/Runtime/AudioMixer/Classes/SubmixEffects/AudioMixerSubmixEffectDynamicsProcessor.h
@@ -143,11 +143,11 @@ struct AUDIOMIXER_API FSubmixEffectDynamicsProcessorSettings
 
 	// If set, uses output of provided submix as modulator of input signal for dynamics processor (Uses input signal as default modulator)
 	UPROPERTY(EditAnywhere, BlueprintReadWrite, Category = Sidechain, meta = (EditCondition = "!bBypass && KeySource == ESubmixEffectDynamicsKeySource::AudioBus", EditConditionHides))
-	UAudioBus* ExternalAudioBus = nullptr;
+	TObjectPtr<UAudioBus> ExternalAudioBus = nullptr;
 
 	// If set, uses output of provided submix as modulator of input signal for dynamics processor (Uses input signal as default modulator)
 	UPROPERTY(EditAnywhere, BlueprintReadWrite, Category = Sidechain, meta = (EditCondition = "!bBypass && KeySource == ESubmixEffectDynamicsKeySource::Submix", EditConditionHides))
-	USoundSubmix* ExternalSubmix = nullptr;
+	TObjectPtr<USoundSubmix> ExternalSubmix = nullptr;
 
 	UPROPERTY()
 	uint8 bChannelLinked_DEPRECATED : 1;
diff --git a/Engine/Source/Runtime/AudioMixer/Public/Components/SynthComponent.h b/Engine/Source/Runtime/AudioMixer/Public/Components/SynthComponent.h
--- a/Engine/Source/Runtime/AudioMixer/Public/Components/SynthComponent.h
+++ b/Engine/Source/Runtime/AudioMixer/Public/Components/SynthComponent.h
@@ -56,7 +56,7 @@ class AUDIOMIXER_API USynthSound : public USoundWaveProcedural
 
 protected:
 	UPROPERTY()
-	USynthComponent* OwningSynthComponent;
+	TObjectPtr<USynthComponent> OwningSynthComponent;
 
 	TArray<float> FloatBuffer;
 	bool bAudioMixer;
@@ -148,7 +148,7 @@ public:
 
 	/** If bOverrideSettings is false, the asset to use to determine attenuation properties for sounds generated by this component */
 	UPROPERTY(EditAnywhere, BlueprintReadWrite, Category = Attenuation, meta = (EditCondition = "!bOverrideAttenuation"))
-	class USoundAttenuation* AttenuationSettings;
+	TObjectPtr<class USoundAttenuation> AttenuationSettings;
 
 	/** If bOverrideSettings is true, the attenuation properties to use for sounds generated by this component */
 	UPROPERTY(EditAnywhere, BlueprintReadWrite, Category = Attenuation, meta = (EditCondition = "bOverrideAttenuation"))
@@ -156,23 +156,23 @@ public:
 
 	/** What sound concurrency to use for sounds generated by this audio component */
 	UPROPERTY()
-	USoundConcurrency* ConcurrencySettings_DEPRECATED;
+	TObjectPtr<USoundConcurrency> ConcurrencySettings_DEPRECATED;
 
 	/** What sound concurrency to use for sounds generated by this audio component */
 	UPROPERTY(EditAnywhere, BlueprintReadWrite, Category = Concurrency)
-	TSet<USoundConcurrency*> ConcurrencySet;
+	TSet<TObjectPtr<USoundConcurrency>> ConcurrencySet;
 
 	/** Sound class this sound belongs to */
 	UPROPERTY(EditAnywhere, Category = SoundClass)
-	USoundClass* SoundClass;
+	TObjectPtr<USoundClass> SoundClass;
 
 	/** The source effect chain to use for this sound. */
 	UPROPERTY(EditAnywhere, Category = Effects)
-	USoundEffectSourcePresetChain* SourceEffectChain;
+	TObjectPtr<USoundEffectSourcePresetChain> SourceEffectChain;
 
 	/** Submix this sound belongs to */
 	UPROPERTY(EditAnywhere, Category = Effects)
-	USoundSubmixBase* SoundSubmix;
+	TObjectPtr<USoundSubmixBase> SoundSubmix;
 
 	/** An array of submix sends. Audio from this sound will send a portion of its audio to these effects.  */
 	UPROPERTY(EditAnywhere, BlueprintReadWrite, Category = Effects)
@@ -283,10 +283,10 @@ private:
 	ISoundGeneratorPtr CreateSoundGeneratorInternal(const FSoundGeneratorInitParams& InParams);
 
 	UPROPERTY(Transient)
-	USynthSound* Synth;
+	TObjectPtr<USynthSound> Synth;
 
 	UPROPERTY(Transient)
-	UAudioComponent* AudioComponent;
+	TObjectPtr<UAudioComponent> AudioComponent;
 
 	void PumpPendingMessages();
 
diff --git a/Engine/Source/Runtime/AugmentedReality/Public/ARActor.h b/Engine/Source/Runtime/AugmentedReality/Public/ARActor.h
--- a/Engine/Source/Runtime/AugmentedReality/Public/ARActor.h
+++ b/Engine/Source/Runtime/AugmentedReality/Public/ARActor.h
@@ -35,11 +35,11 @@ public:
 	FTrackedGeometryGroup(UARTrackedGeometry* InTrackedGeometry);
 	
 	UPROPERTY()
-	AARActor* ARActor = nullptr;
+	TObjectPtr<AARActor> ARActor = nullptr;
 	
 	UPROPERTY()
-	UARComponent* ARComponent = nullptr;
+	TObjectPtr<UARComponent> ARComponent = nullptr;
 	
 	UPROPERTY()
-	UARTrackedGeometry* TrackedGeometry = nullptr;
+	TObjectPtr<UARTrackedGeometry> TrackedGeometry = nullptr;
 };
diff --git a/Engine/Source/Runtime/AugmentedReality/Public/ARComponent.h b/Engine/Source/Runtime/AugmentedReality/Public/ARComponent.h
--- a/Engine/Source/Runtime/AugmentedReality/Public/ARComponent.h
+++ b/Engine/Source/Runtime/AugmentedReality/Public/ARComponent.h
@@ -34,10 +34,10 @@ struct FARSessionPayload
 	int32 ConfigFlags = 0;
 	
 	UPROPERTY(BlueprintReadOnly, Category = "AR")
-	UMaterialInterface* DefaultMeshMaterial = nullptr;
+	TObjectPtr<UMaterialInterface> DefaultMeshMaterial = nullptr;
 
 	UPROPERTY(BlueprintReadOnly, Category = "AR")
-	UMaterialInterface* DefaultWireframeMeshMaterial = nullptr;
+	TObjectPtr<UMaterialInterface> DefaultWireframeMeshMaterial = nullptr;
 
 	void SetFlag(EARSessionConfigFlags InFlag);
 	bool HasFlag(EARSessionConfigFlags InFlag) const;
@@ -107,7 +107,7 @@ struct FARImageUpdatePayload
 	FTransform WorldTransform = FTransform::Identity;
 	
 	UPROPERTY(BlueprintReadOnly, Category = "AR")
-	UARCandidateImage* DetectedImage = nullptr;
+	TObjectPtr<UARCandidateImage> DetectedImage = nullptr;
 	
 	UPROPERTY(BlueprintReadOnly, Category = "AR")
 	FVector2D EstimatedSize = FVector2D::ZeroVector;
@@ -280,21 +280,21 @@ protected:
 	 * If not set, the DefaultMeshMaterial from ARSessionConfig will be used.
 	 */
 	UPROPERTY(EditAnywhere, Category = "AR Gameplay")
-	UMaterialInterface* DefaultMeshMaterial = nullptr;
+	TObjectPtr<UMaterialInterface> DefaultMeshMaterial = nullptr;
 
 	/**
 	 * The default wireframe material to be used for the generated mesh component.
 	 * If not set, the DefaultMeshMaterial from ARSessionConfig will be used.
 	 */
 	UPROPERTY(EditAnywhere, Category = "AR Gameplay")
-	UMaterialInterface* DefaultWireframeMeshMaterial = nullptr;
+	TObjectPtr<UMaterialInterface> DefaultWireframeMeshMaterial = nullptr;
 
 	UPROPERTY()
-	UMRMeshComponent* MRMeshComponent;
+	TObjectPtr<UMRMeshComponent> MRMeshComponent;
 	
 	/** The tracked geometry used for updating this component, only set on "local" client */
 	UPROPERTY()
-	UARTrackedGeometry* MyTrackedGeometry;
+	TObjectPtr<UARTrackedGeometry> MyTrackedGeometry;
 	
 	/** If this component has ever been updated */
 	bool bFirstUpdate = true;
diff --git a/Engine/Source/Runtime/AugmentedReality/Public/ARPin.h b/Engine/Source/Runtime/AugmentedReality/Public/ARPin.h
--- a/Engine/Source/Runtime/AugmentedReality/Public/ARPin.h
+++ b/Engine/Source/Runtime/AugmentedReality/Public/ARPin.h
@@ -93,10 +93,10 @@ private:
 	static uint32 DebugPinId;
 	
 	UPROPERTY()
-	UARTrackedGeometry* TrackedGeometry;
+	TObjectPtr<UARTrackedGeometry> TrackedGeometry;
 	
 	UPROPERTY()
-	USceneComponent* PinnedComponent;
+	TObjectPtr<USceneComponent> PinnedComponent;
 	
 	UPROPERTY()
 	FTransform LocalToTrackingTransform;
diff --git a/Engine/Source/Runtime/AugmentedReality/Public/ARSessionConfig.h b/Engine/Source/Runtime/AugmentedReality/Public/ARSessionConfig.h
--- a/Engine/Source/Runtime/AugmentedReality/Public/ARSessionConfig.h
+++ b/Engine/Source/Runtime/AugmentedReality/Public/ARSessionConfig.h
@@ -411,7 +411,7 @@ protected:
 	
 	/** The list of candidate images to detect within the AR camera view */
 	UPROPERTY(EditAnywhere, Category="AR Settings | Image Tracking")
-	TArray<UARCandidateImage*> CandidateImages;
+	TArray<TObjectPtr<UARCandidateImage>> CandidateImages;
 
     /** The maximum number of images to track at the same time. Defaults to 1 */
     UPROPERTY(EditAnywhere, Category="AR Settings | Image Tracking")
@@ -427,7 +427,7 @@ protected:
 
 	/** A list of candidate objects to search for in the scene */
 	UPROPERTY(EditAnywhere, Category="AR Settings")
-	TArray<UARCandidateObject*> CandidateObjects;
+	TArray<TObjectPtr<UARCandidateObject>> CandidateObjects;
 
 	/**
 	 * The desired video format (or the default if not supported) that this session should use if the camera is enabled
@@ -497,9 +497,9 @@ protected:
 	
 	/** The default mesh material used by the generated mesh component */
 	UPROPERTY(EditAnywhere, Category = "AR Settings | World Mapping")
-	UMaterialInterface* DefaultMeshMaterial;
+	TObjectPtr<UMaterialInterface> DefaultMeshMaterial;
 
 	/** The default mesh material used by the wireframe setting of the generated mesh component.  Note: It is reccomended to ignore this 'wireframe' feature and use a wirefraem material in the DefaultMeshMaterial if you want wireframe. */
 	UPROPERTY(EditAnywhere, Category = "AR Settings | World Mapping")
-	UMaterialInterface* DefaultWireframeMeshMaterial;
+	TObjectPtr<UMaterialInterface> DefaultWireframeMeshMaterial;
 };
diff --git a/Engine/Source/Runtime/AugmentedReality/Public/ARSkyLight.h b/Engine/Source/Runtime/AugmentedReality/Public/ARSkyLight.h
--- a/Engine/Source/Runtime/AugmentedReality/Public/ARSkyLight.h
+++ b/Engine/Source/Runtime/AugmentedReality/Public/ARSkyLight.h
@@ -25,7 +25,7 @@ private:
 	virtual void Tick(float DeltaTime) override;
 
 	UPROPERTY()
-	UAREnvironmentCaptureProbe* CaptureProbe;
+	TObjectPtr<UAREnvironmentCaptureProbe> CaptureProbe;
 
 	/** The timestamp from the environment probe when we last updated the cube map */
 	float LastUpdateTimestamp;
diff --git a/Engine/Source/Runtime/AugmentedReality/Public/ARTraceResult.h b/Engine/Source/Runtime/AugmentedReality/Public/ARTraceResult.h
--- a/Engine/Source/Runtime/AugmentedReality/Public/ARTraceResult.h
+++ b/Engine/Source/Runtime/AugmentedReality/Public/ARTraceResult.h
@@ -62,7 +62,7 @@ private:
 	 * A pointer to the geometry data that was intersected by this trace, if any.
 	 */
 	UPROPERTY()
-	UARTrackedGeometry* TrackedGeometry;
+	TObjectPtr<UARTrackedGeometry> TrackedGeometry;
 	
 	/** A reference to the AR system that creates this hit test result. */
 	TSharedPtr<FARSupportInterface , ESPMode::ThreadSafe> ARSystem;
diff --git a/Engine/Source/Runtime/AugmentedReality/Public/ARTrackable.h b/Engine/Source/Runtime/AugmentedReality/Public/ARTrackable.h
--- a/Engine/Source/Runtime/AugmentedReality/Public/ARTrackable.h
+++ b/Engine/Source/Runtime/AugmentedReality/Public/ARTrackable.h
@@ -94,7 +94,7 @@ protected:
 	
 	/** For AR systems that support arbitrary mesh geometry associated with a tracked point */
 	UPROPERTY(Transient)
-	UMRMeshComponent* UnderlyingMesh;
+	TObjectPtr<UMRMeshComponent> UnderlyingMesh;
 
 	/** What the scene understanding system thinks this object is */
 	UPROPERTY()
@@ -165,7 +165,7 @@ private:
 
 	// Used by ARCore Only
 	UPROPERTY()
-	UARPlaneGeometry* SubsumedBy = nullptr;
+	TObjectPtr<UARPlaneGeometry> SubsumedBy = nullptr;
 };
 
 UCLASS(BlueprintType)
@@ -212,7 +212,7 @@ public:
 protected:
 	/** The candidate image that was detected in the scene */
 	UPROPERTY()
-	UARCandidateImage* DetectedImage;
+	TObjectPtr<UARCandidateImage> DetectedImage;
 
 	/** The estimated image size that was detected in the scene */
 	UPROPERTY()
@@ -423,7 +423,7 @@ protected:
 
 	/** The cube map of the reflected environment */
 	UPROPERTY()
-	UAREnvironmentCaptureProbeTexture* EnvironmentCaptureTexture;
+	TObjectPtr<UAREnvironmentCaptureProbeTexture> EnvironmentCaptureTexture;
 };
 
 UCLASS(BlueprintType)
@@ -445,7 +445,7 @@ public:
 private:
 	/** The candidate object that was detected in the scene */
 	UPROPERTY()
-	UARCandidateObject* DetectedObject;
+	TObjectPtr<UARCandidateObject> DetectedObject;
 };
 
 UCLASS(BlueprintType)
diff --git a/Engine/Source/Runtime/AugmentedReality/Public/ARTypes.h b/Engine/Source/Runtime/AugmentedReality/Public/ARTypes.h
--- a/Engine/Source/Runtime/AugmentedReality/Public/ARTypes.h
+++ b/Engine/Source/Runtime/AugmentedReality/Public/ARTypes.h
@@ -339,7 +339,7 @@ protected:
 	
 	/** The image to detect in scenes */
 	UPROPERTY(EditAnywhere, Category = "AR Candidate Image")
-	UTexture2D* CandidateTexture;
+	TObjectPtr<UTexture2D> CandidateTexture;
 
 	/** The friendly name to report back when the image is detected in scenes */
 	UPROPERTY(EditAnywhere, Category = "AR Candidate Image")
diff --git a/Engine/Source/Runtime/CinematicCamera/Public/CameraRig_Crane.h b/Engine/Source/Runtime/CinematicCamera/Public/CameraRig_Crane.h
--- a/Engine/Source/Runtime/CinematicCamera/Public/CameraRig_Crane.h
+++ b/Engine/Source/Runtime/CinematicCamera/Public/CameraRig_Crane.h
@@ -61,32 +61,32 @@ private:
 
 	/** Root component to give the whole actor a transform. */
 	UPROPERTY(EditDefaultsOnly, Category = "Crane Components")
-	USceneComponent* TransformComponent;
+	TObjectPtr<USceneComponent> TransformComponent;
 
 	/** Component to control Yaw. */
 	UPROPERTY(EditDefaultsOnly, Category = "Crane Components")
-	USceneComponent* CraneYawControl;
+	TObjectPtr<USceneComponent> CraneYawControl;
 
 	/** Component to control Pitch. */
 	UPROPERTY(EditDefaultsOnly, Category = "Crane Components")
-	USceneComponent* CranePitchControl;
+	TObjectPtr<USceneComponent> CranePitchControl;
 
 	/** Component to define the attach point for cameras. */
 	UPROPERTY(EditDefaultsOnly, Category = "Crane Components")
-	USceneComponent* CraneCameraMount;
+	TObjectPtr<USceneComponent> CraneCameraMount;
 
 #if WITH_EDITORONLY_DATA
 	/** Preview meshes for visualization */
 	UPROPERTY(Transient)
-	UStaticMeshComponent* PreviewMesh_CraneArm;
+	TObjectPtr<UStaticMeshComponent> PreviewMesh_CraneArm;
 
 	UPROPERTY(Transient)
-	UStaticMeshComponent* PreviewMesh_CraneBase;
+	TObjectPtr<UStaticMeshComponent> PreviewMesh_CraneBase;
 
 	UPROPERTY(Transient)
-	UStaticMeshComponent* PreviewMesh_CraneMount;
+	TObjectPtr<UStaticMeshComponent> PreviewMesh_CraneMount;
 
 	UPROPERTY(Transient)
-	UStaticMeshComponent* PreviewMesh_CraneCounterWeight;
+	TObjectPtr<UStaticMeshComponent> PreviewMesh_CraneCounterWeight;
 #endif
 };
diff --git a/Engine/Source/Runtime/CinematicCamera/Public/CameraRig_Rail.h b/Engine/Source/Runtime/CinematicCamera/Public/CameraRig_Rail.h
--- a/Engine/Source/Runtime/CinematicCamera/Public/CameraRig_Rail.h
+++ b/Engine/Source/Runtime/CinematicCamera/Public/CameraRig_Rail.h
@@ -66,28 +66,28 @@ private:
 
  	/** Root component to give the whole actor a transform. */
 	UPROPERTY(EditDefaultsOnly, Category = "Rail Components")
- 	USceneComponent* TransformComponent;
+ 	TObjectPtr<USceneComponent> TransformComponent;
 
 	/** Spline component to define the rail path. */
 	UPROPERTY(EditDefaultsOnly, Category = "Rail Components")
-	USplineComponent* RailSplineComponent;
+	TObjectPtr<USplineComponent> RailSplineComponent;
 
 	/** Component to define the attach point for cameras. Moves along the rail. */
 	UPROPERTY(EditDefaultsOnly, Category = "Rail Components")
-	USceneComponent* RailCameraMount;
+	TObjectPtr<USceneComponent> RailCameraMount;
 
 #if WITH_EDITORONLY_DATA
 	/** Preview meshes for visualization */
 	UPROPERTY(Transient)
-	USplineMeshComponent* PreviewMesh_Rail;
+	TObjectPtr<USplineMeshComponent> PreviewMesh_Rail;
 
 	UPROPERTY(Transient)
-	TArray<USplineMeshComponent*> PreviewRailMeshSegments;
+	TArray<TObjectPtr<USplineMeshComponent>> PreviewRailMeshSegments;
 
 	UPROPERTY(Transient)
-	UStaticMesh* PreviewRailStaticMesh;
+	TObjectPtr<UStaticMesh> PreviewRailStaticMesh;
 
 	UPROPERTY(Transient)
-	UStaticMeshComponent* PreviewMesh_Mount;
+	TObjectPtr<UStaticMeshComponent> PreviewMesh_Mount;
 #endif
 };
diff --git a/Engine/Source/Runtime/CinematicCamera/Public/CineCameraComponent.h b/Engine/Source/Runtime/CinematicCamera/Public/CineCameraComponent.h
--- a/Engine/Source/Runtime/CinematicCamera/Public/CineCameraComponent.h
+++ b/Engine/Source/Runtime/CinematicCamera/Public/CineCameraComponent.h
@@ -353,19 +353,19 @@ protected:
 #if WITH_EDITORONLY_DATA
 	/** Mesh used for debug focus plane visualization */
 	UPROPERTY(transient)
-	UStaticMesh* FocusPlaneVisualizationMesh;
+	TObjectPtr<UStaticMesh> FocusPlaneVisualizationMesh;
 
 	/** Material used for debug focus plane visualization */
 	UPROPERTY(transient)
-	UMaterial* FocusPlaneVisualizationMaterial;
+	TObjectPtr<UMaterial> FocusPlaneVisualizationMaterial;
 
 	/** Component for the debug focus plane visualization */
 	UPROPERTY(transient)
-	UStaticMeshComponent* DebugFocusPlaneComponent;
+	TObjectPtr<UStaticMeshComponent> DebugFocusPlaneComponent;
 
 	/** Dynamic material instance for the debug focus plane visualization */
 	UPROPERTY(transient)
-	UMaterialInstanceDynamic* DebugFocusPlaneMID;
+	TObjectPtr<UMaterialInstanceDynamic> DebugFocusPlaneMID;
 
 	virtual void PostEditChangeProperty(FPropertyChangedEvent& PropertyChangedEvent) override;
 	virtual void ResetProxyMeshTransform() override;
diff --git a/Engine/Source/Runtime/ClothingSystemRuntimeCommon/Public/ClothLODData_Legacy.h b/Engine/Source/Runtime/ClothingSystemRuntimeCommon/Public/ClothLODData_Legacy.h
--- a/Engine/Source/Runtime/ClothingSystemRuntimeCommon/Public/ClothLODData_Legacy.h
+++ b/Engine/Source/Runtime/ClothingSystemRuntimeCommon/Public/ClothLODData_Legacy.h
@@ -64,7 +64,7 @@ public:
 
 	// Deprecated, use ClothPhysicalMeshData instead
 	UPROPERTY()
-	UClothPhysicalMeshDataBase_Legacy* PhysicalMeshData_DEPRECATED;
+	TObjectPtr<UClothPhysicalMeshDataBase_Legacy> PhysicalMeshData_DEPRECATED;
 
 	// Raw phys mesh data
 	UPROPERTY()
diff --git a/Engine/Source/Runtime/ClothingSystemRuntimeCommon/Public/ClothingAsset.h b/Engine/Source/Runtime/ClothingSystemRuntimeCommon/Public/ClothingAsset.h
--- a/Engine/Source/Runtime/ClothingSystemRuntimeCommon/Public/ClothingAsset.h
+++ b/Engine/Source/Runtime/ClothingSystemRuntimeCommon/Public/ClothingAsset.h
@@ -189,34 +189,34 @@ public:
 
 	// The physics asset to extract collisions from when building a simulation.
 	UPROPERTY(EditAnywhere, Category = Config)
-	UPhysicsAsset* PhysicsAsset;
+	TObjectPtr<UPhysicsAsset> PhysicsAsset;
 
 	// Simulation specific cloth parameters. 
 	// Use GetClothConfig() to retrieve the correct parameters/config type for the desired cloth simulation system.
 	UPROPERTY(VisibleAnywhere, BlueprintReadOnly, EditFixedSize, Instanced, Category = Config)
-	TMap<FName, UClothConfigBase*> ClothConfigs;
+	TMap<FName, TObjectPtr<UClothConfigBase>> ClothConfigs;
 
 	// Shared by all cloth instances in a skeletal mesh
 	// Only supported with Chaos Cloth for now
 	// This may not be editable on unused cloth assets
 	UPROPERTY()
-	UClothConfigBase* ClothSharedSimConfig_DEPRECATED;
+	TObjectPtr<UClothConfigBase> ClothSharedSimConfig_DEPRECATED;
 
 	// Parameters for how the NVcloth behaves.
 	// These will have no effect on Chaos cloth
 	UPROPERTY()
-	UClothConfigBase* ClothSimConfig_DEPRECATED;
+	TObjectPtr<UClothConfigBase> ClothSimConfig_DEPRECATED;
 
 	// Parameters for how Chaos cloth behaves 
 	// These will not affect NVcloth
 	// For now, we have two configuration parameters so that we can switch between chaos and
 	// non chaos at will without losing the original NVcloth data
 	UPROPERTY()
-	UClothConfigBase* ChaosClothSimConfig_DEPRECATED;
+	TObjectPtr<UClothConfigBase> ChaosClothSimConfig_DEPRECATED;
 
 	// Deprecated. Use LodData instead
 	UPROPERTY()
-	TArray<UClothLODDataCommon_Legacy*> ClothLodData_DEPRECATED;
+	TArray<TObjectPtr<UClothLODDataCommon_Legacy>> ClothLodData_DEPRECATED;
 
 	// The actual asset data, listed by LOD.
 	UPROPERTY()
@@ -240,7 +240,7 @@ public:
 
 	// Custom data applied by the importer depending on where the asset was imported from.
 	UPROPERTY()
-	UClothingAssetCustomData* CustomData;
+	TObjectPtr<UClothingAssetCustomData> CustomData;
 
 	/** 
 	 * Deprecated property for transitioning the \c FClothConfig struct to the 
diff --git a/Engine/Source/Runtime/CoreUObject/Public/UObject/NoExportTypes.h b/Engine/Source/Runtime/CoreUObject/Public/UObject/NoExportTypes.h
--- a/Engine/Source/Runtime/CoreUObject/Public/UObject/NoExportTypes.h
+++ b/Engine/Source/Runtime/CoreUObject/Public/UObject/NoExportTypes.h
@@ -1636,10 +1636,10 @@ USTRUCT(noexport)
 struct FTestUninitializedScriptStructMembersTest
 {
 	UPROPERTY(Transient)
-	UObject* UninitializedObjectReference;
+	TObjectPtr<UObject> UninitializedObjectReference;
 
 	UPROPERTY(Transient)
-	UObject* InitializedObjectReference;
+	TObjectPtr<UObject> InitializedObjectReference;
 
 	UPROPERTY(Transient)
 	float UnusedValue;
diff --git a/Engine/Source/Runtime/Engine/Classes/Animation/AnimBlueprint.h b/Engine/Source/Runtime/Engine/Classes/Animation/AnimBlueprint.h
--- a/Engine/Source/Runtime/Engine/Classes/Animation/AnimBlueprint.h
+++ b/Engine/Source/Runtime/Engine/Classes/Animation/AnimBlueprint.h
@@ -38,7 +38,7 @@ struct FAnimParentNodeAssetOverride
 	GENERATED_USTRUCT_BODY()
 
 	UPROPERTY()
-	UAnimationAsset* NewAsset;
+	TObjectPtr<UAnimationAsset> NewAsset;
 	UPROPERTY()
 	FGuid ParentNodeGuid;
 
@@ -84,7 +84,7 @@ class ENGINE_API UAnimBlueprint : public UBlueprint, public IInterface_PreviewMe
 	 * to cause errors if this is modified without updating or replacing all referenced animations.
 	 */
 	UPROPERTY(AssetRegistrySearchable, EditAnywhere, AdvancedDisplay, Category=ClassOptions)
-	USkeleton* TargetSkeleton;
+	TObjectPtr<USkeleton> TargetSkeleton;
 
 	// List of animation sync groups
 	UPROPERTY()
@@ -203,7 +203,7 @@ public:
 	// Array of active pose watches (pose watch allows us to see the bone pose at a 
 	// particular point of the anim graph) 
 	UPROPERTY(transient)
-	TArray<class UPoseWatch*> PoseWatches;
+	TArray<TObjectPtr<class UPoseWatch>> PoseWatches;
 
 private:
 	/** The default skeletal mesh to use when previewing this asset - this only applies when you open Persona using this asset*/
diff --git a/Engine/Source/Runtime/Engine/Classes/Animation/AnimBlueprintGeneratedClass.h b/Engine/Source/Runtime/Engine/Classes/Animation/AnimBlueprintGeneratedClass.h
--- a/Engine/Source/Runtime/Engine/Classes/Animation/AnimBlueprintGeneratedClass.h
+++ b/Engine/Source/Runtime/Engine/Classes/Animation/AnimBlueprintGeneratedClass.h
@@ -321,7 +321,7 @@ class ENGINE_API UAnimBlueprintGeneratedClass : public UBlueprintGeneratedClass,
 
 	/** Target skeleton for this blueprint class */
 	UPROPERTY()
-	USkeleton* TargetSkeleton;
+	TObjectPtr<USkeleton> TargetSkeleton;
 
 	/** A list of anim notifies that state machines (or anything else) may reference */
 	UPROPERTY()
@@ -366,7 +366,7 @@ private:
 
 	// Any internal blendspaces we host
 	UPROPERTY()
-	TArray<UBlendSpaceBase*> BlendSpaces;
+	TArray<TObjectPtr<UBlendSpaceBase>> BlendSpaces;
 
 public:
 	// IAnimClassInterface interface
diff --git a/Engine/Source/Runtime/Engine/Classes/Animation/AnimBoneCompressionSettings.h b/Engine/Source/Runtime/Engine/Classes/Animation/AnimBoneCompressionSettings.h
--- a/Engine/Source/Runtime/Engine/Classes/Animation/AnimBoneCompressionSettings.h
+++ b/Engine/Source/Runtime/Engine/Classes/Animation/AnimBoneCompressionSettings.h
@@ -20,7 +20,7 @@ class ENGINE_API UAnimBoneCompressionSettings : public UObject
 
 	/** A list of animation bone compression codecs to try. Empty entries are ignored but the array cannot be empty. */
 	UPROPERTY(Category = Compression, Instanced, EditAnywhere, meta = (NoElementDuplicate))
-	TArray<UAnimBoneCompressionCodec*> Codecs;
+	TArray<TObjectPtr<UAnimBoneCompressionCodec>> Codecs;
 
 #if WITH_EDITORONLY_DATA
 	/** When compressing, the best codec below this error threshold will be used. */
diff --git a/Engine/Source/Runtime/Engine/Classes/Animation/AnimClassData.h b/Engine/Source/Runtime/Engine/Classes/Animation/AnimClassData.h
--- a/Engine/Source/Runtime/Engine/Classes/Animation/AnimClassData.h
+++ b/Engine/Source/Runtime/Engine/Classes/Animation/AnimClassData.h
@@ -46,7 +46,7 @@ public:
 
 	/** Target skeleton for this blueprint class */
 	UPROPERTY()
-	class USkeleton* TargetSkeleton;
+	TObjectPtr<class USkeleton> TargetSkeleton;
 
 	/** A list of anim notifies that state machines (or anything else) may reference */
 	UPROPERTY()
diff --git a/Engine/Source/Runtime/Engine/Classes/Animation/AnimComposite.h b/Engine/Source/Runtime/Engine/Classes/Animation/AnimComposite.h
--- a/Engine/Source/Runtime/Engine/Classes/Animation/AnimComposite.h
+++ b/Engine/Source/Runtime/Engine/Classes/Animation/AnimComposite.h
@@ -29,7 +29,7 @@ public:
 #if WITH_EDITORONLY_DATA
 	/** Preview Base pose for additive BlendSpace **/
 	UPROPERTY(EditAnywhere, Category=AdditiveSettings)
-	UAnimSequence* PreviewBasePose;
+	TObjectPtr<UAnimSequence> PreviewBasePose;
 #endif // WITH_EDITORONLY_DATA
 
 	//~ Begin UAnimSequenceBase Interface
diff --git a/Engine/Source/Runtime/Engine/Classes/Animation/AnimCompositeBase.h b/Engine/Source/Runtime/Engine/Classes/Animation/AnimCompositeBase.h
--- a/Engine/Source/Runtime/Engine/Classes/Animation/AnimCompositeBase.h
+++ b/Engine/Source/Runtime/Engine/Classes/Animation/AnimCompositeBase.h
@@ -28,7 +28,7 @@ struct FRootMotionExtractionStep
 
 	/** AnimSequence ref */
 	UPROPERTY()
-	UAnimSequence* AnimSequence;
+	TObjectPtr<UAnimSequence> AnimSequence;
 
 	/** Start position to extract root motion from. */
 	UPROPERTY()
@@ -61,7 +61,7 @@ struct FAnimSegment
 
 	/** Anim Reference to play - only allow AnimSequence or AnimComposite **/
 	UPROPERTY(EditAnywhere, Category=AnimSegment)
-	UAnimSequenceBase* AnimReference;
+	TObjectPtr<UAnimSequenceBase> AnimReference;
 
 	/** Start Pos within this AnimCompositeBase */
 	UPROPERTY(VisibleAnywhere, Category=AnimSegment)
diff --git a/Engine/Source/Runtime/Engine/Classes/Animation/AnimCurveCompressionSettings.h b/Engine/Source/Runtime/Engine/Classes/Animation/AnimCurveCompressionSettings.h
--- a/Engine/Source/Runtime/Engine/Classes/Animation/AnimCurveCompressionSettings.h
+++ b/Engine/Source/Runtime/Engine/Classes/Animation/AnimCurveCompressionSettings.h
@@ -20,7 +20,7 @@ class ENGINE_API UAnimCurveCompressionSettings : public UObject
 
 	/** An animation curve compression codec. */
 	UPROPERTY(Category = Compression, Export, EditAnywhere, NoClear, meta = (EditInline))
-	UAnimCurveCompressionCodec* Codec;
+	TObjectPtr<UAnimCurveCompressionCodec> Codec;
 
 	//////////////////////////////////////////////////////////////////////////
 
diff --git a/Engine/Source/Runtime/Engine/Classes/Animation/AnimData/AnimDataController.h b/Engine/Source/Runtime/Engine/Classes/Animation/AnimData/AnimDataController.h
--- a/Engine/Source/Runtime/Engine/Classes/Animation/AnimData/AnimDataController.h
+++ b/Engine/Source/Runtime/Engine/Classes/Animation/AnimData/AnimDataController.h
@@ -484,7 +484,7 @@ private:
 #if WITH_EDITORONLY_DATA
 	/** Current UAnimDataModel instance targeted by this controller */
 	UPROPERTY(transient)
-	UAnimDataModel* Model;
+	TObjectPtr<UAnimDataModel> Model;
 #endif // WITH_EDITORONLY_DATA
 
 	friend class FAnimDataControllerTestBase;
diff --git a/Engine/Source/Runtime/Engine/Classes/Animation/AnimInstance.h b/Engine/Source/Runtime/Engine/Classes/Animation/AnimInstance.h
--- a/Engine/Source/Runtime/Engine/Classes/Animation/AnimInstance.h
+++ b/Engine/Source/Runtime/Engine/Classes/Animation/AnimInstance.h
@@ -376,7 +376,7 @@ class ENGINE_API UAnimInstance : public UObject
 
 	/** This is used to extract animation. If Mesh exists, this will be overwritten by Mesh->Skeleton */
 	UPROPERTY(transient)
-	USkeleton* CurrentSkeleton;
+	TObjectPtr<USkeleton> CurrentSkeleton;
 
 	// Sets where this blueprint pulls Root Motion from
 	UPROPERTY(Category = RootMotion, EditDefaultsOnly)
diff --git a/Engine/Source/Runtime/Engine/Classes/Animation/AnimLinkableElement.h b/Engine/Source/Runtime/Engine/Classes/Animation/AnimLinkableElement.h
--- a/Engine/Source/Runtime/Engine/Classes/Animation/AnimLinkableElement.h
+++ b/Engine/Source/Runtime/Engine/Classes/Animation/AnimLinkableElement.h
@@ -142,7 +142,7 @@ protected:
 
 	/** The montage that this element is currently linked to */
 	UPROPERTY()
-	UAnimMontage* LinkedMontage;
+	TObjectPtr<UAnimMontage> LinkedMontage;
 
 	/** The slot index we are currently using within LinkedMontage */
 	UPROPERTY(EditAnywhere, Category=AnimLink)
@@ -177,7 +177,7 @@ protected:
 	 * in either length or rate; the element will correctly place itself in relation to the sequence
 	 */
 	UPROPERTY(VisibleAnywhere, AdvancedDisplay, Category=AnimLink)
-	UAnimSequenceBase* LinkedSequence;
+	TObjectPtr<UAnimSequenceBase> LinkedSequence;
 
 private:
 
diff --git a/Engine/Source/Runtime/Engine/Classes/Animation/AnimMontage.h b/Engine/Source/Runtime/Engine/Classes/Animation/AnimMontage.h
--- a/Engine/Source/Runtime/Engine/Classes/Animation/AnimMontage.h
+++ b/Engine/Source/Runtime/Engine/Classes/Animation/AnimMontage.h
@@ -53,7 +53,7 @@ struct FCompositeSection : public FAnimLinkableElement
 	 * You can query by GetMetaData function
 	 */
 	UPROPERTY(Category=Section, Instanced, EditAnywhere)
-	TArray<class UAnimMetaData*> MetaData;
+	TArray<TObjectPtr<class UAnimMetaData>> MetaData;
 
 public:
 	FCompositeSection()
@@ -302,7 +302,7 @@ struct ENGINE_API FMontageBlendSettings
 
 	/** Blend Profile to use for this blend */
 	UPROPERTY(EditAnywhere, BlueprintReadWrite, Category = "Blend", meta = (DisplayAfter = "Blend"))
-	UBlendProfile* BlendProfile;
+	TObjectPtr<UBlendProfile> BlendProfile;
 
 	/** AlphaBlend options (time, curve, etc.) */
 	UPROPERTY(EditAnywhere, BlueprintReadWrite, Category = "Blend", meta = (DisplayAfter = "BlendMode"))
@@ -322,7 +322,7 @@ struct ENGINE_API FAnimMontageInstance
 
 	// Montage reference
 	UPROPERTY()
-	class UAnimMontage* Montage;
+	TObjectPtr<class UAnimMontage> Montage;
 
 	// delegates
 	FOnMontageEnded OnMontageEnded;
@@ -693,11 +693,11 @@ class UAnimMontage : public UAnimCompositeBase
 
 	/** The blend profile to use. */
 	UPROPERTY(EditAnywhere, BlueprintReadOnly, Category = BlendOption, meta = (UseAsBlendProfile = true))
-	UBlendProfile* BlendProfileIn;
+	TObjectPtr<UBlendProfile> BlendProfileIn;
 
 	/** The blend profile to use. */
 	UPROPERTY(EditAnywhere, BlueprintReadOnly, Category = BlendOption, meta = (UseAsBlendProfile = true))
-	UBlendProfile* BlendProfileOut;
+	TObjectPtr<UBlendProfile> BlendProfileOut;
 
 	/** Root Bone will be locked to that position when extracting root motion. DEPRECATED in 4.5 root motion is controlled by anim sequences **/
 	UPROPERTY()
@@ -706,7 +706,7 @@ class UAnimMontage : public UAnimCompositeBase
 #if WITH_EDITORONLY_DATA
 	/** Preview Base pose for additive BlendSpace **/
 	UPROPERTY(EditAnywhere, Category = AdditiveSettings)
-	UAnimSequence* PreviewBasePose;
+	TObjectPtr<UAnimSequence> PreviewBasePose;
 #endif // WITH_EDITORONLY_DATA
 
 	// Add new slot track to this montage
diff --git a/Engine/Source/Runtime/Engine/Classes/Animation/AnimNodeBase.h b/Engine/Source/Runtime/Engine/Classes/Animation/AnimNodeBase.h
--- a/Engine/Source/Runtime/Engine/Classes/Animation/AnimNodeBase.h
+++ b/Engine/Source/Runtime/Engine/Classes/Animation/AnimNodeBase.h
@@ -792,7 +792,7 @@ struct ENGINE_API FExposedValueHandler
 
 	// function pointer if BoundFunction != NAME_None
 	UPROPERTY()
-	UFunction* Function;
+	TObjectPtr<UFunction> Function;
 
 	// Node property that this value handler is associated with, when the node
 	// is instantiated from this property the node's ExposedValueHandler will 
diff --git a/Engine/Source/Runtime/Engine/Classes/Animation/AnimNode_CustomProperty.h b/Engine/Source/Runtime/Engine/Classes/Animation/AnimNode_CustomProperty.h
--- a/Engine/Source/Runtime/Engine/Classes/Animation/AnimNode_CustomProperty.h
+++ b/Engine/Source/Runtime/Engine/Classes/Animation/AnimNode_CustomProperty.h
@@ -56,7 +56,7 @@ protected:
 
 	/** This is the actual instance allocated at runtime that will run. Set by child class. */
 	UPROPERTY(Transient)
-	UObject* TargetInstance;
+	TObjectPtr<UObject> TargetInstance;
 
 	/** List of properties on the calling Source Instances instance to push from  */
 	TArray<FProperty*> SourceProperties;
diff --git a/Engine/Source/Runtime/Engine/Classes/Animation/AnimNode_SequencePlayer.h b/Engine/Source/Runtime/Engine/Classes/Animation/AnimNode_SequencePlayer.h
--- a/Engine/Source/Runtime/Engine/Classes/Animation/AnimNode_SequencePlayer.h
+++ b/Engine/Source/Runtime/Engine/Classes/Animation/AnimNode_SequencePlayer.h
@@ -18,7 +18,7 @@ struct ENGINE_API FAnimNode_SequencePlayer : public FAnimNode_AssetPlayerBase
 public:
 	// The animation sequence asset to play
 	UPROPERTY(EditAnywhere, BlueprintReadWrite, Category = Settings, meta = (PinHiddenByDefault, DisallowedClasses="AnimMontage"))
-	UAnimSequenceBase* Sequence;
+	TObjectPtr<UAnimSequenceBase> Sequence;
 
 	// The Basis in which the PlayRate is expressed in. This is used to rescale PlayRate inputs.
 	// For example a Basis of 100 means that the PlayRate input will be divided by 100.
diff --git a/Engine/Source/Runtime/Engine/Classes/Animation/AnimNode_StateMachine.h b/Engine/Source/Runtime/Engine/Classes/Animation/AnimNode_StateMachine.h
--- a/Engine/Source/Runtime/Engine/Classes/Animation/AnimNode_StateMachine.h
+++ b/Engine/Source/Runtime/Engine/Classes/Animation/AnimNode_StateMachine.h
@@ -65,7 +65,7 @@ protected:
 public:
 	// Blend profile to use for this transition. Specifying this will make the transition evaluate per-bone
 	UPROPERTY()
-	UBlendProfile* BlendProfile;
+	TObjectPtr<UBlendProfile> BlendProfile;
 
 	// Type of blend to use
 	EAlphaBlendOption BlendOption;
diff --git a/Engine/Source/Runtime/Engine/Classes/Animation/AnimNotifies/AnimNotifyState_TimedParticleEffect.h b/Engine/Source/Runtime/Engine/Classes/Animation/AnimNotifies/AnimNotifyState_TimedParticleEffect.h
--- a/Engine/Source/Runtime/Engine/Classes/Animation/AnimNotifies/AnimNotifyState_TimedParticleEffect.h
+++ b/Engine/Source/Runtime/Engine/Classes/Animation/AnimNotifies/AnimNotifyState_TimedParticleEffect.h
@@ -20,7 +20,7 @@ class ENGINE_API UAnimNotifyState_TimedParticleEffect : public UAnimNotifyState
 
 	// The particle system template to use when spawning the particle component
 	UPROPERTY(EditAnywhere, Category = ParticleSystem, meta = (ToolTip = "The particle system to spawn for the notify state"))
-	UParticleSystem* PSTemplate;
+	TObjectPtr<UParticleSystem> PSTemplate;
 
 	// The socket within our mesh component to attach to when we spawn the particle component
 	UPROPERTY(EditAnywhere, Category = ParticleSystem, meta = (ToolTip = "The socket or bone to attach the system to"))
@@ -46,7 +46,7 @@ class ENGINE_API UAnimNotifyState_TimedParticleEffect : public UAnimNotifyState
 	// Because these can change at any time we need to track previous versions when we are in an
 	// editor build. Refactor when stateful data is possible, tracking our component instead.
 	UPROPERTY(transient)
-	TArray<UParticleSystem*> PreviousPSTemplates;
+	TArray<TObjectPtr<UParticleSystem>> PreviousPSTemplates;
 
 	UPROPERTY(transient)
 	TArray<FName> PreviousSocketNames;
diff --git a/Engine/Source/Runtime/Engine/Classes/Animation/AnimNotifies/AnimNotifyState_Trail.h b/Engine/Source/Runtime/Engine/Classes/Animation/AnimNotifies/AnimNotifyState_Trail.h
--- a/Engine/Source/Runtime/Engine/Classes/Animation/AnimNotifies/AnimNotifyState_Trail.h
+++ b/Engine/Source/Runtime/Engine/Classes/Animation/AnimNotifies/AnimNotifyState_Trail.h
@@ -22,7 +22,7 @@ class ENGINE_API UAnimNotifyState_Trail : public UAnimNotifyState
 	
 	/** The particle system to use for this trail. */
 	UPROPERTY(EditAnywhere, BlueprintReadWrite, Category = Trail)
-	UParticleSystem* PSTemplate;
+	TObjectPtr<UParticleSystem> PSTemplate;
 
 	UFUNCTION(BlueprintImplementableEvent)
 	UParticleSystem* OverridePSTemplate(USkeletalMeshComponent* MeshComp, class UAnimSequenceBase* Animation) const;
diff --git a/Engine/Source/Runtime/Engine/Classes/Animation/AnimNotifies/AnimNotify_PlayParticleEffect.h b/Engine/Source/Runtime/Engine/Classes/Animation/AnimNotifies/AnimNotify_PlayParticleEffect.h
--- a/Engine/Source/Runtime/Engine/Classes/Animation/AnimNotifies/AnimNotify_PlayParticleEffect.h
+++ b/Engine/Source/Runtime/Engine/Classes/Animation/AnimNotifies/AnimNotify_PlayParticleEffect.h
@@ -39,7 +39,7 @@ public:
 
 	// Particle System to Spawn
 	UPROPERTY(EditAnywhere, BlueprintReadWrite, Category="AnimNotify", meta=(DisplayName="Particle System"))
-	UParticleSystem* PSTemplate;
+	TObjectPtr<UParticleSystem> PSTemplate;
 
 	// Location offset from the socket
 	UPROPERTY(EditAnywhere, BlueprintReadWrite, Category="AnimNotify")
diff --git a/Engine/Source/Runtime/Engine/Classes/Animation/AnimNotifies/AnimNotify_PlaySound.h b/Engine/Source/Runtime/Engine/Classes/Animation/AnimNotifies/AnimNotify_PlaySound.h
--- a/Engine/Source/Runtime/Engine/Classes/Animation/AnimNotifies/AnimNotify_PlaySound.h
+++ b/Engine/Source/Runtime/Engine/Classes/Animation/AnimNotifies/AnimNotify_PlaySound.h
@@ -31,7 +31,7 @@ public:
 
 	// Sound to Play
 	UPROPERTY(EditAnywhere, BlueprintReadWrite, Category="AnimNotify", meta=(ExposeOnSpawn = true))
-	USoundBase* Sound;
+	TObjectPtr<USoundBase> Sound;
 
 	// Volume Multiplier
 	UPROPERTY(EditAnywhere, BlueprintReadWrite, Category="AnimNotify", meta=(ExposeOnSpawn = true))
diff --git a/Engine/Source/Runtime/Engine/Classes/Animation/AnimSequence.h b/Engine/Source/Runtime/Engine/Classes/Animation/AnimSequence.h
--- a/Engine/Source/Runtime/Engine/Classes/Animation/AnimSequence.h
+++ b/Engine/Source/Runtime/Engine/Classes/Animation/AnimSequence.h
@@ -282,11 +282,11 @@ public:
 
 	/** The bone compression settings used to compress bones in this sequence. */
 	UPROPERTY(Category = Compression, EditAnywhere, meta = (ForceShowEngineContent))
-	class UAnimBoneCompressionSettings* BoneCompressionSettings;
+	TObjectPtr<class UAnimBoneCompressionSettings> BoneCompressionSettings;
 
 	/** The curve compression settings used to compress curves in this sequence. */
 	UPROPERTY(Category = Compression, EditAnywhere, meta = (ForceShowEngineContent))
-	class UAnimCurveCompressionSettings* CurveCompressionSettings;
+	TObjectPtr<class UAnimCurveCompressionSettings> CurveCompressionSettings;
 
 	FCompressedAnimSequence CompressedData;
 
@@ -311,7 +311,7 @@ public:
 
 	/* Additive reference animation if it's relevant - i.e. AnimScaled or AnimFrame **/
 	UPROPERTY(EditAnywhere, Category=AdditiveSettings, meta=(DisplayName = "Base Pose Animation"))
-	class UAnimSequence* RefPoseSeq;
+	TObjectPtr<class UAnimSequence> RefPoseSeq;
 
 	/* Additve reference frame if RefPoseType == AnimFrame **/
 	UPROPERTY(EditAnywhere, Category=AdditiveSettings)
@@ -369,7 +369,7 @@ public:
 
 	/** Importing data and options used for this mesh */
 	UPROPERTY(VisibleAnywhere, Instanced, Category=ImportSettings)
-	class UAssetImportData* AssetImportData;
+	TObjectPtr<class UAssetImportData> AssetImportData;
 
 	/***  for Reimport **/
 	/** Path to the resource used to construct this skeletal mesh */
diff --git a/Engine/Source/Runtime/Engine/Classes/Animation/AnimSequenceBase.h b/Engine/Source/Runtime/Engine/Classes/Animation/AnimSequenceBase.h
--- a/Engine/Source/Runtime/Engine/Classes/Animation/AnimSequenceBase.h
+++ b/Engine/Source/Runtime/Engine/Classes/Animation/AnimSequenceBase.h
@@ -290,14 +290,14 @@ private:
 protected:
 	/** UAnimDataModel instance containing source animation data */
 	UPROPERTY(VisibleAnywhere, BlueprintReadOnly, Category = "Animation Model")
-	UAnimDataModel* DataModel;
+	TObjectPtr<UAnimDataModel> DataModel;
 
 	/** Flag set whenever the data-model is initially populated (during upgrade path) */
 	bool bPopulatingDataModel;
 
 	/** UAnimDataController instance set to operate on DataModel */
 	UPROPERTY(VisibleAnywhere, BlueprintReadOnly, Transient, DuplicateTransient, Category = "Animation Model")
-	UAnimDataController* Controller;
+	TObjectPtr<UAnimDataController> Controller;
 	
 	/** Helper object that keeps track of any controller brackets, and all unique notify types that are broadcasted during it */
 	UE::Anim::FAnimDataModelNotifyCollector NotifyCollector;
diff --git a/Engine/Source/Runtime/Engine/Classes/Animation/AnimSet.h b/Engine/Source/Runtime/Engine/Classes/Animation/AnimSet.h
--- a/Engine/Source/Runtime/Engine/Classes/Animation/AnimSet.h
+++ b/Engine/Source/Runtime/Engine/Classes/Animation/AnimSet.h
@@ -59,7 +59,7 @@ public:
 #if WITH_EDITORONLY_DATA
 	/** Actual animation sequence information. */
 	UPROPERTY()
-	TArray<class UAnimSequence*> Sequences;
+	TArray<TObjectPtr<class UAnimSequence>> Sequences;
 
 #endif // WITH_EDITORONLY_DATA
 private:
diff --git a/Engine/Source/Runtime/Engine/Classes/Animation/AnimSingleNodeInstance.h b/Engine/Source/Runtime/Engine/Classes/Animation/AnimSingleNodeInstance.h
--- a/Engine/Source/Runtime/Engine/Classes/Animation/AnimSingleNodeInstance.h
+++ b/Engine/Source/Runtime/Engine/Classes/Animation/AnimSingleNodeInstance.h
@@ -29,7 +29,7 @@ class ENGINE_API UAnimSingleNodeInstance : public UAnimInstance
 
 	/** Current Asset being played **/
 	UPROPERTY(Transient)
-	class UAnimationAsset* CurrentAsset;
+	TObjectPtr<class UAnimationAsset> CurrentAsset;
 	 
 	UPROPERTY(Transient)
 	FPostEvaluateAnimEvent PostEvaluateAnimEvent;
diff --git a/Engine/Source/Runtime/Engine/Classes/Animation/AnimStateMachineTypes.h b/Engine/Source/Runtime/Engine/Classes/Animation/AnimStateMachineTypes.h
--- a/Engine/Source/Runtime/Engine/Classes/Animation/AnimStateMachineTypes.h
+++ b/Engine/Source/Runtime/Engine/Classes/Animation/AnimStateMachineTypes.h
@@ -141,10 +141,10 @@ struct FAnimationTransitionBetweenStates : public FAnimationStateBase
 	EAlphaBlendOption BlendMode;
 
 	UPROPERTY()
-	UCurveFloat* CustomCurve;
+	TObjectPtr<UCurveFloat> CustomCurve;
 
 	UPROPERTY()
-	UBlendProfile* BlendProfile;
+	TObjectPtr<UBlendProfile> BlendProfile;
 
 	UPROPERTY()
 	TEnumAsByte<ETransitionLogicType::Type> LogicType;
diff --git a/Engine/Source/Runtime/Engine/Classes/Animation/AnimStreamable.h b/Engine/Source/Runtime/Engine/Classes/Animation/AnimStreamable.h
--- a/Engine/Source/Runtime/Engine/Classes/Animation/AnimStreamable.h
+++ b/Engine/Source/Runtime/Engine/Classes/Animation/AnimStreamable.h
@@ -107,7 +107,7 @@ public:
 
 	// Sequence the streamable was created from (used for reflecting changes to the source in editor)
 	UPROPERTY()
-	const UAnimSequence* SourceSequence;
+	TObjectPtr<const UAnimSequence> SourceSequence;
 
 	UPROPERTY()
 	FGuid RawDataGuid;
@@ -177,11 +177,11 @@ public:
 
 	/** The bone compression settings used to compress bones in this sequence. */
 	UPROPERTY(Category = Compression, EditAnywhere)
-	class UAnimBoneCompressionSettings* BoneCompressionSettings;
+	TObjectPtr<class UAnimBoneCompressionSettings> BoneCompressionSettings;
 
 	/** The curve compression settings used to compress curves in this sequence. */
 	UPROPERTY(Category = Compression, EditAnywhere)
-	class UAnimCurveCompressionSettings* CurveCompressionSettings;
+	TObjectPtr<class UAnimCurveCompressionSettings> CurveCompressionSettings;
 
 	/** If this is on, it will allow extracting of root motion **/
 	UPROPERTY(EditAnywhere, AssetRegistrySearchable, Category = RootMotion, meta = (DisplayName = "EnableRootMotion"))
diff --git a/Engine/Source/Runtime/Engine/Classes/Animation/AnimationAsset.h b/Engine/Source/Runtime/Engine/Classes/Animation/AnimationAsset.h
--- a/Engine/Source/Runtime/Engine/Classes/Animation/AnimationAsset.h
+++ b/Engine/Source/Runtime/Engine/Classes/Animation/AnimationAsset.h
@@ -77,7 +77,7 @@ struct FBlendSampleData
 	int32 SampleDataIndex;
 
 	UPROPERTY()
-	class UAnimSequence* Animation;
+	TObjectPtr<class UAnimSequence> Animation;
 
 	UPROPERTY()
 	float TotalWeight;
@@ -287,7 +287,7 @@ struct FAnimTickRecord
 	GENERATED_USTRUCT_BODY()
 
 	UPROPERTY()
-	class UAnimationAsset* SourceAsset = nullptr;
+	TObjectPtr<class UAnimationAsset> SourceAsset = nullptr;
 
 	float* TimeAccumulator = nullptr;
 	float PlayRateMultiplier = 1.0f;
@@ -857,7 +857,7 @@ class ENGINE_API UAnimationAsset : public UObject, public IInterface_AssetUserDa
 private:
 	/** Pointer to the Skeleton this asset can be played on .	*/
 	UPROPERTY(AssetRegistrySearchable, Category=Animation, VisibleAnywhere)
-	class USkeleton* Skeleton;
+	TObjectPtr<class USkeleton> Skeleton;
 
 	/** Skeleton guid. If changes, you need to remap info*/
 	FGuid SkeletonGuid;
@@ -870,7 +870,7 @@ private:
 	 * You can query by GetMetaData function
 	 */
 	UPROPERTY(Category=MetaData, instanced, EditAnywhere)
-	TArray<UAnimMetaData*> MetaData;
+	TArray<TObjectPtr<UAnimMetaData>> MetaData;
 
 public:
 	/* 
@@ -886,7 +886,7 @@ public:
 	 * 
 	 * During cooking, this data will be used to bake out to normal asset */
 	UPROPERTY(AssetRegistrySearchable, Category=Animation, VisibleAnywhere)
-	class UAnimationAsset* ParentAsset;
+	TObjectPtr<class UAnimationAsset> ParentAsset;
 
 	/** 
 	 * @todo : comment
@@ -897,7 +897,7 @@ public:
 	 * note this is transient as they're added as they're loaded
 	 */
 	UPROPERTY(transient)
-	TArray<class UAnimationAsset*> ChildrenAssets;
+	TArray<TObjectPtr<class UAnimationAsset>> ChildrenAssets;
 
 	const UAssetMappingTable* GetAssetMappingTable() const
 	{
@@ -906,13 +906,13 @@ public:
 protected:
 	/** Asset mapping table when ParentAsset is set */
 	UPROPERTY(Category=Animation, VisibleAnywhere)
-	class UAssetMappingTable* AssetMappingTable;
+	TObjectPtr<class UAssetMappingTable> AssetMappingTable;
 #endif // WITH_EDITORONLY_DATA
 
 protected:
 	/** Array of user data stored with the asset */
 	UPROPERTY(EditAnywhere, AdvancedDisplay, Instanced, Category = Animation)
-	TArray<UAssetUserData*> AssetUserData;
+	TArray<TObjectPtr<UAssetUserData>> AssetUserData;
 
 public:
 	/** Advances the asset player instance 
@@ -1030,12 +1030,12 @@ public:
 #if WITH_EDITORONLY_DATA
 	/** Information for thumbnail rendering */
 	UPROPERTY(VisibleAnywhere, Instanced, Category = Thumbnail)
-	class UThumbnailInfo* ThumbnailInfo;
+	TObjectPtr<class UThumbnailInfo> ThumbnailInfo;
 
 	/** The default skeletal mesh to use when previewing this asset - this only applies when you open Persona using this asset*/
 	// @todo: note that this doesn't retarget right now
 	UPROPERTY(duplicatetransient, EditAnywhere, Category = Animation)
-	class UPoseAsset* PreviewPoseAsset;
+	TObjectPtr<class UPoseAsset> PreviewPoseAsset;
 
 private:
 	/** The default skeletal mesh to use when previewing this asset - this only applies when you open Persona using this asset*/
diff --git a/Engine/Source/Runtime/Engine/Classes/Animation/AssetMappingTable.h b/Engine/Source/Runtime/Engine/Classes/Animation/AssetMappingTable.h
--- a/Engine/Source/Runtime/Engine/Classes/Animation/AssetMappingTable.h
+++ b/Engine/Source/Runtime/Engine/Classes/Animation/AssetMappingTable.h
@@ -24,11 +24,11 @@ struct FAssetMapping
 
 	/** source asset **/
 	UPROPERTY(EditAnywhere, Category = "FAssetMapping")
-	class UAnimationAsset*			SourceAsset;
+	TObjectPtr<class UAnimationAsset>			SourceAsset;
 
 	/** source asset **/
 	UPROPERTY(EditAnywhere, Category = "FAssetMapping")
-	class UAnimationAsset*			TargetAsset;
+	TObjectPtr<class UAnimationAsset>			TargetAsset;
 
 	FAssetMapping()
 		: SourceAsset(nullptr)
diff --git a/Engine/Source/Runtime/Engine/Classes/Animation/BlendProfile.h b/Engine/Source/Runtime/Engine/Classes/Animation/BlendProfile.h
--- a/Engine/Source/Runtime/Engine/Classes/Animation/BlendProfile.h
+++ b/Engine/Source/Runtime/Engine/Classes/Animation/BlendProfile.h
@@ -161,7 +161,7 @@ private:
 public:
 	// The skeleton that owns this profile
 	UPROPERTY()
-	USkeleton* OwningSkeleton;
+	TObjectPtr<USkeleton> OwningSkeleton;
 
 	// List of blend scale entries
 	UPROPERTY()
diff --git a/Engine/Source/Runtime/Engine/Classes/Animation/BlendSpaceBase.h b/Engine/Source/Runtime/Engine/Classes/Animation/BlendSpaceBase.h
--- a/Engine/Source/Runtime/Engine/Classes/Animation/BlendSpaceBase.h
+++ b/Engine/Source/Runtime/Engine/Classes/Animation/BlendSpaceBase.h
@@ -93,7 +93,7 @@ struct FBlendSample
 
 	// For linked animations
 	UPROPERTY(EditAnywhere, Category=BlendSample)
-	class UAnimSequence* Animation;
+	TObjectPtr<class UAnimSequence> Animation;
 
 	//blend 0->x, blend 1->y, blend 2->z
 
@@ -467,7 +467,7 @@ public:
 #if WITH_EDITORONLY_DATA
 	/** Preview Base pose for additive BlendSpace **/
 	UPROPERTY(EditAnywhere, Category = AdditiveSettings)
-	UAnimSequence* PreviewBasePose;
+	TObjectPtr<UAnimSequence> PreviewBasePose;
 #endif // WITH_EDITORONLY_DATA
 
 	/** This animation length changes based on current input (resulting in different blend time)**/
diff --git a/Engine/Source/Runtime/Engine/Classes/Animation/InputScaleBias.h b/Engine/Source/Runtime/Engine/Classes/Animation/InputScaleBias.h
--- a/Engine/Source/Runtime/Engine/Classes/Animation/InputScaleBias.h
+++ b/Engine/Source/Runtime/Engine/Classes/Animation/InputScaleBias.h
@@ -160,7 +160,7 @@ struct ENGINE_API FInputAlphaBoolBlend
 	bool bInitialized;
 
 	UPROPERTY(EditAnywhere, BlueprintReadWrite, Category = Settings)
-	UCurveFloat* CustomCurve;
+	TObjectPtr<UCurveFloat> CustomCurve;
 
 	UPROPERTY(Transient)
 	FAlphaBlend AlphaBlend;
diff --git a/Engine/Source/Runtime/Engine/Classes/Animation/MorphTarget.h b/Engine/Source/Runtime/Engine/Classes/Animation/MorphTarget.h
--- a/Engine/Source/Runtime/Engine/Classes/Animation/MorphTarget.h
+++ b/Engine/Source/Runtime/Engine/Classes/Animation/MorphTarget.h
@@ -119,7 +119,7 @@ public:
 
 	/** USkeletalMesh that this vertex animation works on. */
 	UPROPERTY(AssetRegistrySearchable)
-	class USkeletalMesh* BaseSkelMesh;
+	TObjectPtr<class USkeletalMesh> BaseSkelMesh;
 
 	/** morph mesh vertex data for each LOD */
 	TArray<FMorphTargetLODModel>	MorphLODModels;
diff --git a/Engine/Source/Runtime/Engine/Classes/Animation/PoseAsset.h b/Engine/Source/Runtime/Engine/Classes/Animation/PoseAsset.h
--- a/Engine/Source/Runtime/Engine/Classes/Animation/PoseAsset.h
+++ b/Engine/Source/Runtime/Engine/Classes/Animation/PoseAsset.h
@@ -165,7 +165,7 @@ public:
 
 #if WITH_EDITORONLY_DATA
 	UPROPERTY(Category=Source, EditAnywhere)
-	UAnimSequence* SourceAnimation;
+	TObjectPtr<UAnimSequence> SourceAnimation;
 #endif // WITH_EDITORONLY_DATA
 
 	/**
diff --git a/Engine/Source/Runtime/Engine/Classes/Animation/PreviewAssetAttachComponent.h b/Engine/Source/Runtime/Engine/Classes/Animation/PreviewAssetAttachComponent.h
--- a/Engine/Source/Runtime/Engine/Classes/Animation/PreviewAssetAttachComponent.h
+++ b/Engine/Source/Runtime/Engine/Classes/Animation/PreviewAssetAttachComponent.h
@@ -26,7 +26,7 @@ private:
 	TSoftObjectPtr<class UObject> AttachedObject;
 
 	UPROPERTY()
-	UObject* Object_DEPRECATED;
+	TObjectPtr<UObject> Object_DEPRECATED;
 
 public:
 
diff --git a/Engine/Source/Runtime/Engine/Classes/Animation/SkeletalMeshActor.h b/Engine/Source/Runtime/Engine/Classes/Animation/SkeletalMeshActor.h
--- a/Engine/Source/Runtime/Engine/Classes/Animation/SkeletalMeshActor.h
+++ b/Engine/Source/Runtime/Engine/Classes/Animation/SkeletalMeshActor.h
@@ -35,23 +35,23 @@ class ENGINE_API ASkeletalMeshActor : public AActor, public IMatineeAnimInterfac
 
 private:
 	UPROPERTY(Category = SkeletalMeshActor, VisibleAnywhere, BlueprintReadOnly, meta = (ExposeFunctionCategories = "Mesh,Components|SkeletalMesh,Animation,Physics", AllowPrivateAccess = "true"))
-	class USkeletalMeshComponent* SkeletalMeshComponent;
+	TObjectPtr<class USkeletalMeshComponent> SkeletalMeshComponent;
 public:
 
 	/** Used to replicate mesh to clients */
 	UPROPERTY(replicatedUsing=OnRep_ReplicatedMesh, transient)
-	class USkeletalMesh* ReplicatedMesh;
+	TObjectPtr<class USkeletalMesh> ReplicatedMesh;
 
 	/** Used to replicate physics asset to clients */
 	UPROPERTY(replicatedUsing=OnRep_ReplicatedPhysAsset, transient)
-	class UPhysicsAsset* ReplicatedPhysAsset;
+	TObjectPtr<class UPhysicsAsset> ReplicatedPhysAsset;
 
 	/** used to replicate the material in index 0 */
 	UPROPERTY(replicatedUsing=OnRep_ReplicatedMaterial0)
-	class UMaterialInterface* ReplicatedMaterial0;
+	TObjectPtr<class UMaterialInterface> ReplicatedMaterial0;
 
 	UPROPERTY(replicatedUsing=OnRep_ReplicatedMaterial1)
-	class UMaterialInterface* ReplicatedMaterial1;
+	TObjectPtr<class UMaterialInterface> ReplicatedMaterial1;
 
 	/** Replication Notification Callbacks */
 	UFUNCTION()
diff --git a/Engine/Source/Runtime/Engine/Classes/Animation/Skeleton.h b/Engine/Source/Runtime/Engine/Classes/Animation/Skeleton.h
--- a/Engine/Source/Runtime/Engine/Classes/Animation/Skeleton.h
+++ b/Engine/Source/Runtime/Engine/Classes/Animation/Skeleton.h
@@ -198,7 +198,7 @@ struct FRigConfiguration
 	GENERATED_USTRUCT_BODY()
 
 	UPROPERTY()
-	class URig * Rig = nullptr;
+	TObjectPtr<class URig>  Rig = nullptr;
 
 	// @todo in the future we can make this to be run-time data
 	UPROPERTY()
@@ -332,7 +332,7 @@ public:
 	 *	everything explicitly to AttachComponent in the SkeletalMeshComponent.
 	 */
 	UPROPERTY()
-	TArray<class USkeletalMeshSocket*> Sockets;
+	TArray<TObjectPtr<class USkeletalMeshSocket>> Sockets;
 
 	/** Serializable retarget sources for this skeleton **/
 	TMap< FName, FReferencePose > AnimRetargetSources;
@@ -409,7 +409,7 @@ public:
 
 	/** List of blend profiles available in this skeleton */
 	UPROPERTY(Instanced)
-	TArray<UBlendProfile*> BlendProfiles;
+	TArray<TObjectPtr<UBlendProfile>> BlendProfiles;
 
 	/** Get the specified blend profile by name */
 	ENGINE_API UBlendProfile* GetBlendProfile(const FName& InProfileName);
@@ -899,6 +899,6 @@ public:
 protected:
 	/** Array of user data stored with the asset */
 	UPROPERTY(EditAnywhere, AdvancedDisplay, Instanced, Category = Skeleton)
-	TArray<UAssetUserData*> AssetUserData;
+	TArray<TObjectPtr<UAssetUserData>> AssetUserData;
 };
 
diff --git a/Engine/Source/Runtime/Engine/Classes/Atmosphere/AtmosphericFog.h b/Engine/Source/Runtime/Engine/Classes/Atmosphere/AtmosphericFog.h
--- a/Engine/Source/Runtime/Engine/Classes/Atmosphere/AtmosphericFog.h
+++ b/Engine/Source/Runtime/Engine/Classes/Atmosphere/AtmosphericFog.h
@@ -20,12 +20,12 @@ class UE_DEPRECATED(4.26, "Please use the SkyAtmosphere actor instead.") AAtmosp
 private:
 	/** Main fog component */
 	UPROPERTY(BlueprintReadOnly, VisibleAnywhere, Category = Atmosphere, meta = (AllowPrivateAccess = "true"))
-	class UAtmosphericFogComponent* AtmosphericFogComponent;
+	TObjectPtr<class UAtmosphericFogComponent> AtmosphericFogComponent;
 
 #if WITH_EDITORONLY_DATA
 	/** Arrow component to indicate default sun rotation */
 	UPROPERTY()
-	class UArrowComponent* ArrowComponent;
+	TObjectPtr<class UArrowComponent> ArrowComponent;
 #endif
 
 public:
diff --git a/Engine/Source/Runtime/Engine/Classes/Atmosphere/AtmosphericFogComponent.h b/Engine/Source/Runtime/Engine/Classes/Atmosphere/AtmosphericFogComponent.h
--- a/Engine/Source/Runtime/Engine/Classes/Atmosphere/AtmosphericFogComponent.h
+++ b/Engine/Source/Runtime/Engine/Classes/Atmosphere/AtmosphericFogComponent.h
@@ -222,10 +222,10 @@ protected:
 
 public:
 	UPROPERTY()
-	class UTexture2D* TransmittanceTexture_DEPRECATED;
+	TObjectPtr<class UTexture2D> TransmittanceTexture_DEPRECATED;
 
 	UPROPERTY()
-	class UTexture2D* IrradianceTexture_DEPRECATED;
+	TObjectPtr<class UTexture2D> IrradianceTexture_DEPRECATED;
 
 	enum EPrecomputeState
 	{
diff --git a/Engine/Source/Runtime/Engine/Classes/Camera/CameraActor.h b/Engine/Source/Runtime/Engine/Classes/Camera/CameraActor.h
--- a/Engine/Source/Runtime/Engine/Classes/Camera/CameraActor.h
+++ b/Engine/Source/Runtime/Engine/Classes/Camera/CameraActor.h
@@ -29,10 +29,10 @@ private:
 
 	/** The camera component for this camera */
 	UPROPERTY(Category = CameraActor, VisibleAnywhere, BlueprintReadOnly, meta = (AllowPrivateAccess = "true"))
-	class UCameraComponent* CameraComponent;
+	TObjectPtr<class UCameraComponent> CameraComponent;
 
 	UPROPERTY(Category = CameraActor, VisibleAnywhere, BlueprintReadOnly, meta = (AllowPrivateAccess = "true"))
-	class USceneComponent* SceneComponent;
+	TObjectPtr<class USceneComponent> SceneComponent;
 public:
 
 	/** If this CameraActor is being used to preview a CameraAnim in the editor, this is the anim being previewed. */
diff --git a/Engine/Source/Runtime/Engine/Classes/Camera/CameraAnim.h b/Engine/Source/Runtime/Engine/Classes/Camera/CameraAnim.h
--- a/Engine/Source/Runtime/Engine/Classes/Camera/CameraAnim.h
+++ b/Engine/Source/Runtime/Engine/Classes/Camera/CameraAnim.h
@@ -21,12 +21,12 @@ class UCameraAnim : public UObject
 public:
 	/** The UInterpGroup that holds our actual interpolation data. */
 	UPROPERTY()
-	class UInterpGroup* CameraInterpGroup;
+	TObjectPtr<class UInterpGroup> CameraInterpGroup;
 
 #if WITH_EDITORONLY_DATA
 	/** This is to preview and they only exists in editor */
 	UPROPERTY(transient)
-	class UInterpGroup* PreviewInterpGroup;
+	TObjectPtr<class UInterpGroup> PreviewInterpGroup;
 #endif // WITH_EDITORONLY_DATA
 
 	/** Length, in seconds. */
diff --git a/Engine/Source/Runtime/Engine/Classes/Camera/CameraAnimInst.h b/Engine/Source/Runtime/Engine/Classes/Camera/CameraAnimInst.h
--- a/Engine/Source/Runtime/Engine/Classes/Camera/CameraAnimInst.h
+++ b/Engine/Source/Runtime/Engine/Classes/Camera/CameraAnimInst.h
@@ -21,12 +21,12 @@ class ENGINE_API UCameraAnimInst : public UObject
 
 	/** which CameraAnim this is an instance of */
 	UPROPERTY()
-	class UCameraAnim* CamAnim;
+	TObjectPtr<class UCameraAnim> CamAnim;
 
 private:
 	/** the UInterpGroupInst used to do the interpolation */
 	UPROPERTY(instanced)
-	class UInterpGroupInst* InterpGroupInst;
+	TObjectPtr<class UInterpGroupInst> InterpGroupInst;
 
 public:
 	/** Current time for the animation */
@@ -85,10 +85,10 @@ protected:
 public:
 	/** cached movement track from the currently playing anim so we don't have to go find it every frame */
 	UPROPERTY(transient)
-	class UInterpTrackMove* MoveTrack;
+	TObjectPtr<class UInterpTrackMove> MoveTrack;
 
 	UPROPERTY(transient)
-	class UInterpTrackInstMove* MoveInst;
+	TObjectPtr<class UInterpTrackInstMove> MoveInst;
 
 	UPROPERTY()
 	ECameraShakePlaySpace PlaySpace;
diff --git a/Engine/Source/Runtime/Engine/Classes/Camera/CameraComponent.h b/Engine/Source/Runtime/Engine/Classes/Camera/CameraComponent.h
--- a/Engine/Source/Runtime/Engine/Classes/Camera/CameraComponent.h
+++ b/Engine/Source/Runtime/Engine/Classes/Camera/CameraComponent.h
@@ -151,7 +151,7 @@ protected:
 	class UDrawFrustumComponent* DrawFrustum;
 
 	UPROPERTY(BlueprintReadWrite, EditAnywhere, Category = Camera)
-	class UStaticMesh* CameraMesh;
+	TObjectPtr<class UStaticMesh> CameraMesh;
 
 	// The camera mesh to show visually where the camera is placed
 	class UStaticMeshComponent* ProxyMeshComponent;
diff --git a/Engine/Source/Runtime/Engine/Classes/Camera/CameraModifier.h b/Engine/Source/Runtime/Engine/Classes/Camera/CameraModifier.h
--- a/Engine/Source/Runtime/Engine/Classes/Camera/CameraModifier.h
+++ b/Engine/Source/Runtime/Engine/Classes/Camera/CameraModifier.h
@@ -48,7 +48,7 @@ public:
 protected:
 	/** Camera this object is associated with. */
 	UPROPERTY(transient, BlueprintReadOnly, Category = CameraModifier)
-	class APlayerCameraManager* CameraOwner;
+	TObjectPtr<class APlayerCameraManager> CameraOwner;
 
 	/** When blending in, alpha proceeds from 0 to 1 over this time */
 	UPROPERTY(EditDefaultsOnly, BlueprintReadOnly, Category = CameraModifier)
diff --git a/Engine/Source/Runtime/Engine/Classes/Camera/CameraModifier_CameraShake.h b/Engine/Source/Runtime/Engine/Classes/Camera/CameraModifier_CameraShake.h
--- a/Engine/Source/Runtime/Engine/Classes/Camera/CameraModifier_CameraShake.h
+++ b/Engine/Source/Runtime/Engine/Classes/Camera/CameraModifier_CameraShake.h
@@ -22,7 +22,7 @@ struct FPooledCameraShakes
 	GENERATED_BODY()
 		
 	UPROPERTY()
-	TArray<UCameraShakeBase*> PooledShakes;
+	TArray<TObjectPtr<UCameraShakeBase>> PooledShakes;
 };
 
 USTRUCT()
@@ -33,7 +33,7 @@ struct FActiveCameraShakeInfo
 	FActiveCameraShakeInfo() : ShakeInstance(nullptr), ShakeSource(nullptr) {}
 
 	UPROPERTY()
-	UCameraShakeBase* ShakeInstance;
+	TObjectPtr<UCameraShakeBase> ShakeInstance;
 
 	UPROPERTY()
 	TWeakObjectPtr<const UCameraShakeSourceComponent> ShakeSource;
diff --git a/Engine/Source/Runtime/Engine/Classes/Camera/CameraShakeBase.h b/Engine/Source/Runtime/Engine/Classes/Camera/CameraShakeBase.h
--- a/Engine/Source/Runtime/Engine/Classes/Camera/CameraShakeBase.h
+++ b/Engine/Source/Runtime/Engine/Classes/Camera/CameraShakeBase.h
@@ -518,11 +518,11 @@ private:
 
 	/** The root pattern for this camera shake */
 	UPROPERTY(EditAnywhere, Instanced, Category=CameraShake)
-	UCameraShakePattern* RootShakePattern;
+	TObjectPtr<UCameraShakePattern> RootShakePattern;
 
 	/** The camera manager owning this camera shake. Only valid when the shake is active. */
 	UPROPERTY(transient)
-	APlayerCameraManager* CameraManager;
+	TObjectPtr<APlayerCameraManager> CameraManager;
 
 	/** What space to play the shake in before applying to the camera. Only valid when the shake is active. */
 	ECameraShakePlaySpace PlaySpace;
diff --git a/Engine/Source/Runtime/Engine/Classes/Camera/CameraShakeSourceActor.h b/Engine/Source/Runtime/Engine/Classes/Camera/CameraShakeSourceActor.h
--- a/Engine/Source/Runtime/Engine/Classes/Camera/CameraShakeSourceActor.h
+++ b/Engine/Source/Runtime/Engine/Classes/Camera/CameraShakeSourceActor.h
@@ -26,5 +26,5 @@ public:
 
 private:
 	UPROPERTY(VisibleAnywhere, BlueprintReadOnly, Category = CameraShakeSourceActor, meta = (AllowPrivateAccess = "true"))
-	UCameraShakeSourceComponent* CameraShakeSourceComponent;
+	TObjectPtr<UCameraShakeSourceComponent> CameraShakeSourceComponent;
 };
diff --git a/Engine/Source/Runtime/Engine/Classes/Camera/CameraShakeSourceComponent.h b/Engine/Source/Runtime/Engine/Classes/Camera/CameraShakeSourceComponent.h
--- a/Engine/Source/Runtime/Engine/Classes/Camera/CameraShakeSourceComponent.h
+++ b/Engine/Source/Runtime/Engine/Classes/Camera/CameraShakeSourceComponent.h
@@ -65,7 +65,7 @@ public:
 #if WITH_EDITORONLY_DATA
 	/** Sprite to display in the editor. */
 	UPROPERTY(transient)
-	UTexture2D* EditorSpriteTexture;
+	TObjectPtr<UTexture2D> EditorSpriteTexture;
 
 	/** Sprite scaling for display in the editor. */
 	UPROPERTY(transient)
diff --git a/Engine/Source/Runtime/Engine/Classes/Camera/PlayerCameraManager.h b/Engine/Source/Runtime/Engine/Classes/Camera/PlayerCameraManager.h
--- a/Engine/Source/Runtime/Engine/Classes/Camera/PlayerCameraManager.h
+++ b/Engine/Source/Runtime/Engine/Classes/Camera/PlayerCameraManager.h
@@ -73,7 +73,7 @@ public:
 
 	/** Target Actor used to compute POV */
 	UPROPERTY(EditAnywhere, BlueprintReadWrite, Category=TViewTarget)
-	class AActor* Target;
+	TObjectPtr<class AActor> Target;
 
 	/** Computed point of view */
 	UPROPERTY(EditAnywhere, BlueprintReadWrite, Category=TViewTarget)
@@ -82,7 +82,7 @@ public:
 protected:
 	/** PlayerState (used to follow same player through pawn transitions, etc., when spectating) */
 	UPROPERTY(EditAnywhere, BlueprintReadWrite, Category=TViewTarget)
-	class APlayerState* PlayerState;
+	TObjectPtr<class APlayerState> PlayerState;
 
 public:
 	class APlayerState* GetPlayerState() const { return PlayerState; }
@@ -183,12 +183,12 @@ class ENGINE_API APlayerCameraManager : public AActor
 
 	/** PlayerController that owns this Camera actor */
 	UPROPERTY(transient)
-	class APlayerController* PCOwner;
+	TObjectPtr<class APlayerController> PCOwner;
 
 private:
 	/** Dummy component we can use to attach things to the camera. */
 	UPROPERTY(Category = PlayerCameraManager, VisibleAnywhere, BlueprintReadOnly, meta = (AllowPrivateAccess = "true"))
-	class USceneComponent* TransformComponent;
+	TObjectPtr<class USceneComponent> TransformComponent;
 
 public:
 	/** Usable to define different camera behaviors. A few simple styles are implemented by default. */
@@ -281,7 +281,7 @@ private:
 protected:
 	/** List of active camera modifier instances that have a chance to update the final camera POV */
 	UPROPERTY(transient)
-	TArray<UCameraModifier*> ModifierList;
+	TArray<TObjectPtr<UCameraModifier>> ModifierList;
 
 public:
 	/** List of modifiers to create by default for this camera */
@@ -341,7 +341,7 @@ protected:
 	// "Lens" effects (e.g. blood, dirt on camera)
 	/** CameraBlood emitter attached to this camera */
 	UPROPERTY(transient)
-	TArray<class AEmitterCameraLensEffectBase*> CameraLensEffects;
+	TArray<TObjectPtr<class AEmitterCameraLensEffectBase>> CameraLensEffects;
 
 	/////////////////////
 	// Camera Modifiers
@@ -349,7 +349,7 @@ protected:
 
 	/** Cached ref to modifier for code-driven screen shakes */
 	UPROPERTY(transient)
-	class UCameraModifier_CameraShake* CachedCameraShakeMod;
+	TObjectPtr<class UCameraModifier_CameraShake> CachedCameraShakeMod;
 
 
 	////////////////////////
@@ -360,7 +360,7 @@ protected:
 
 	/** Internal pool of camera anim instance objects available for playing camera animations. Defines the max number of camera anims that can play simultaneously. */
 	UPROPERTY(transient)
-	class UCameraAnimInst* AnimInstPool[8];    /*MAX_ACTIVE_CAMERA_ANIMS @fixme constant */
+	TObjectPtr<class UCameraAnimInst> AnimInstPool[8];    /*MAX_ACTIVE_CAMERA_ANIMS @fixme constant */
 
 	/** Internal list of active post process effects. Parallel array to PostProcessBlendCacheWeights. */
 	UPROPERTY(transient)
@@ -379,7 +379,7 @@ protected:
 public:
 	/** Array of camera anim instances that are currently playing and in-use */
 	UPROPERTY(transient)
-	TArray<class UCameraAnimInst*> ActiveAnims;
+	TArray<TObjectPtr<class UCameraAnimInst>> ActiveAnims;
 
 	/** Returns active post process info. */
 	void GetCachedPostProcessBlends(TArray<struct FPostProcessSettings> const*& OutPPSettings, TArray<float> const*& OutBlendWeigthts) const;
@@ -387,11 +387,11 @@ public:
 protected:
 	/** Array of camera anim instances that are not playing and available to be used. */
 	UPROPERTY(transient)
-	TArray<class UCameraAnimInst*> FreeAnims;
+	TArray<TObjectPtr<class UCameraAnimInst>> FreeAnims;
 
 	/** Internal. Receives the output of individual camera animations. */
 	UPROPERTY(transient)
-	class ACameraActor* AnimCameraActor;
+	TObjectPtr<class ACameraActor> AnimCameraActor;
 
 
 	////////////////////////
diff --git a/Engine/Source/Runtime/Engine/Classes/Components/ActorComponent.h b/Engine/Source/Runtime/Engine/Classes/Components/ActorComponent.h
--- a/Engine/Source/Runtime/Engine/Classes/Components/ActorComponent.h
+++ b/Engine/Source/Runtime/Engine/Classes/Components/ActorComponent.h
@@ -139,7 +139,7 @@ public:
 protected:
 	/** Array of user data stored with the component */
 	UPROPERTY(EditAnywhere, AdvancedDisplay, Instanced, Category = AssetUserData)
-	TArray<UAssetUserData*> AssetUserData;
+	TArray<TObjectPtr<UAssetUserData>> AssetUserData;
 
 private:
 	/** Used for fast removal of end of frame update */
