Automated wrapped object pointer upgrade for Engine + ShooterGame + ShooterGame referenced plugins

diff --git a/Engine/Source/Runtime/Engine/Classes/Components/AudioComponent.h b/Engine/Source/Runtime/Engine/Classes/Components/AudioComponent.h
--- a/Engine/Source/Runtime/Engine/Classes/Components/AudioComponent.h
+++ b/Engine/Source/Runtime/Engine/Classes/Components/AudioComponent.h
@@ -141,7 +141,7 @@ struct FAudioComponentParam
 
 	// Value of the parameter when used as a sound wave
 	UPROPERTY(EditAnywhere, BlueprintReadWrite, Category=AudioComponentParam)
-	class USoundWave* SoundWaveParam;
+	TObjectPtr<class USoundWave> SoundWaveParam;
 
 	FAudioComponentParam(const FName& Name)
 		: ParamName(Name)
@@ -174,7 +174,7 @@ class ENGINE_API UAudioComponent : public USceneComponent
 
 	/** The sound to be played */
 	UPROPERTY(EditAnywhere, BlueprintReadWrite, Category=Sound)
-	class USoundBase* Sound;
+	TObjectPtr<class USoundBase> Sound;
 
 	/** Array of per-instance parameters for this AudioComponent. */
 	UPROPERTY(EditAnywhere, BlueprintReadWrite, Category=Sound, AdvancedDisplay)
@@ -182,7 +182,7 @@ class ENGINE_API UAudioComponent : public USceneComponent
 
 	/** Optional sound group this AudioComponent belongs to */
 	UPROPERTY(EditAnywhere, Category=Sound, AdvancedDisplay)
-	USoundClass* SoundClassOverride;
+	TObjectPtr<USoundClass> SoundClassOverride;
 
 	/** Auto destroy this component on completion */
 	UPROPERTY()
@@ -319,7 +319,7 @@ public:
 
 	/** The chain of Source Effects to apply to the sounds playing on the Audio Component */
 	UPROPERTY(EditAnywhere, BlueprintReadWrite, Category = Sound)
-	USoundEffectSourcePresetChain* SourceEffectChain;
+	TObjectPtr<USoundEffectSourcePresetChain> SourceEffectChain;
 
 #if WITH_EDITORONLY_DATA
 	UPROPERTY()
@@ -342,7 +342,7 @@ public:
 
 	/** If bOverrideSettings is false, the asset to use to determine attenuation properties for sounds generated by this component */
 	UPROPERTY(EditAnywhere, BlueprintReadWrite, Category=Attenuation, meta=(EditCondition="!bOverrideAttenuation"))
-	class USoundAttenuation* AttenuationSettings;
+	TObjectPtr<class USoundAttenuation> AttenuationSettings;
 
 	/** If bOverrideSettings is true, the attenuation properties to use for sounds generated by this component */
 	UPROPERTY(EditAnywhere, BlueprintReadWrite, Category= Attenuation, meta=(EditCondition="bOverrideAttenuation"))
@@ -350,11 +350,11 @@ public:
 
 	/** What sound concurrency to use for sounds generated by this audio component */
 	UPROPERTY()
-	USoundConcurrency* ConcurrencySettings_DEPRECATED;
+	TObjectPtr<USoundConcurrency> ConcurrencySettings_DEPRECATED;
 
 	/** What sound concurrency rules to use for sounds generated by this audio component */
 	UPROPERTY(EditAnywhere, BlueprintReadWrite, Category = Concurrency)
-	TSet<USoundConcurrency*> ConcurrencySet;
+	TSet<TObjectPtr<USoundConcurrency>> ConcurrencySet;
 
 	/** While playing, this component will check for occlusion from its closest listener every this many seconds */
 	float OcclusionCheckInterval;
diff --git a/Engine/Source/Runtime/Engine/Classes/Components/BillboardComponent.h b/Engine/Source/Runtime/Engine/Classes/Components/BillboardComponent.h
--- a/Engine/Source/Runtime/Engine/Classes/Components/BillboardComponent.h
+++ b/Engine/Source/Runtime/Engine/Classes/Components/BillboardComponent.h
@@ -21,7 +21,7 @@ class ENGINE_API UBillboardComponent : public UPrimitiveComponent
 	GENERATED_UCLASS_BODY()
 
 	UPROPERTY(EditAnywhere, BlueprintReadWrite, Category=Sprite)
-	class UTexture2D* Sprite;
+	TObjectPtr<class UTexture2D> Sprite;
 
 	UPROPERTY(EditAnywhere, BlueprintReadWrite, Category=Sprite)
 	uint32 bIsScreenSizeScaled:1;
diff --git a/Engine/Source/Runtime/Engine/Classes/Components/BoxReflectionCaptureComponent.h b/Engine/Source/Runtime/Engine/Classes/Components/BoxReflectionCaptureComponent.h
--- a/Engine/Source/Runtime/Engine/Classes/Components/BoxReflectionCaptureComponent.h
+++ b/Engine/Source/Runtime/Engine/Classes/Components/BoxReflectionCaptureComponent.h
@@ -19,10 +19,10 @@ class UBoxReflectionCaptureComponent : public UReflectionCaptureComponent
 	float BoxTransitionDistance;
 
 	UPROPERTY()
-	class UBoxComponent* PreviewInfluenceBox;
+	TObjectPtr<class UBoxComponent> PreviewInfluenceBox;
 
 	UPROPERTY()
-	class UBoxComponent* PreviewCaptureBox;
+	TObjectPtr<class UBoxComponent> PreviewCaptureBox;
 
 public:
 	virtual void UpdatePreviewShape() override;
diff --git a/Engine/Source/Runtime/Engine/Classes/Components/BrushComponent.h b/Engine/Source/Runtime/Engine/Classes/Components/BrushComponent.h
--- a/Engine/Source/Runtime/Engine/Classes/Components/BrushComponent.h
+++ b/Engine/Source/Runtime/Engine/Classes/Components/BrushComponent.h
@@ -23,11 +23,11 @@ class UBrushComponent : public UPrimitiveComponent
 	GENERATED_UCLASS_BODY()
 
 	UPROPERTY()
-	class UModel* Brush;
+	TObjectPtr<class UModel> Brush;
 
 	/** Description of collision */
 	UPROPERTY()
-	class UBodySetup* BrushBodySetup;
+	TObjectPtr<class UBodySetup> BrushBodySetup;
 
 #if WITH_EDITORONLY_DATA
 	/** Local space translation */
diff --git a/Engine/Source/Runtime/Engine/Classes/Components/ChildActorComponent.h b/Engine/Source/Runtime/Engine/Classes/Components/ChildActorComponent.h
--- a/Engine/Source/Runtime/Engine/Classes/Components/ChildActorComponent.h
+++ b/Engine/Source/Runtime/Engine/Classes/Components/ChildActorComponent.h
@@ -116,11 +116,11 @@ private:
 
 	/** The actor that we spawned and own */
 	UPROPERTY(Replicated, BlueprintReadOnly, Category=ChildActorComponent, TextExportTransient, NonPIEDuplicateTransient, meta=(AllowPrivateAccess="true"))
-	AActor*	ChildActor;
+	TObjectPtr<AActor>	ChildActor;
 
 	/** Property to point to the template child actor for details panel purposes */
 	UPROPERTY(VisibleDefaultsOnly, DuplicateTransient, Category=ChildActorComponent, meta=(ShowInnerProperties))
-	AActor* ChildActorTemplate;
+	TObjectPtr<AActor> ChildActorTemplate;
 
 	/** We try to keep the child actor's name as best we can, so we store it off here when destroying */
 	FName ChildActorName;
diff --git a/Engine/Source/Runtime/Engine/Classes/Components/DecalComponent.h b/Engine/Source/Runtime/Engine/Classes/Components/DecalComponent.h
--- a/Engine/Source/Runtime/Engine/Classes/Components/DecalComponent.h
+++ b/Engine/Source/Runtime/Engine/Classes/Components/DecalComponent.h
@@ -26,7 +26,7 @@ class ENGINE_API UDecalComponent : public USceneComponent
 protected:
 	/** Decal material. */
 	UPROPERTY(EditAnywhere, BlueprintReadOnly, Category=Decal)
-	class UMaterialInterface* DecalMaterial;
+	TObjectPtr<class UMaterialInterface> DecalMaterial;
 
 public:
 	/** 
diff --git a/Engine/Source/Runtime/Engine/Classes/Components/DrawFrustumComponent.h b/Engine/Source/Runtime/Engine/Classes/Components/DrawFrustumComponent.h
--- a/Engine/Source/Runtime/Engine/Classes/Components/DrawFrustumComponent.h
+++ b/Engine/Source/Runtime/Engine/Classes/Components/DrawFrustumComponent.h
@@ -42,7 +42,7 @@ class UDrawFrustumComponent : public UPrimitiveComponent
 
 	/** optional texture to show on the near plane */
 	UPROPERTY(EditAnywhere, BlueprintReadWrite, Category=DrawFrustumComponent)
-	class UTexture* Texture;
+	TObjectPtr<class UTexture> Texture;
 
 
 	//~ Begin UPrimitiveComponent Interface.
diff --git a/Engine/Source/Runtime/Engine/Classes/Components/ExponentialHeightFogComponent.h b/Engine/Source/Runtime/Engine/Classes/Components/ExponentialHeightFogComponent.h
--- a/Engine/Source/Runtime/Engine/Classes/Components/ExponentialHeightFogComponent.h
+++ b/Engine/Source/Runtime/Engine/Classes/Components/ExponentialHeightFogComponent.h
@@ -71,7 +71,7 @@ class ENGINE_API UExponentialHeightFogComponent : public USceneComponent
 	 * When the cubemap is specified, FogInscatteringColor is ignored and Directional inscattering is disabled. 
 	 */
 	UPROPERTY(EditAnywhere, BlueprintReadOnly, Category=InscatteringTexture)
-	class UTextureCube* InscatteringColorCubemap;
+	TObjectPtr<class UTextureCube> InscatteringColorCubemap;
 
 	/** Angle to rotate the InscatteringColorCubemap around the Z axis. */
 	UPROPERTY(EditAnywhere, BlueprintReadOnly, Category=InscatteringTexture, meta=(UIMin = "0", UIMax = "360"))
diff --git a/Engine/Source/Runtime/Engine/Classes/Components/ForceFeedbackComponent.h b/Engine/Source/Runtime/Engine/Classes/Components/ForceFeedbackComponent.h
--- a/Engine/Source/Runtime/Engine/Classes/Components/ForceFeedbackComponent.h
+++ b/Engine/Source/Runtime/Engine/Classes/Components/ForceFeedbackComponent.h
@@ -66,7 +66,7 @@ public:
 
 	/** The feedback effect to be played */
 	UPROPERTY(EditAnywhere, BlueprintReadOnly, Category=ForceFeedback)
-	UForceFeedbackEffect* ForceFeedbackEffect;
+	TObjectPtr<UForceFeedbackEffect> ForceFeedbackEffect;
 
 	/** Auto destroy this component on completion */
 	UPROPERTY()
@@ -93,7 +93,7 @@ public:
 
 	/** If bOverrideSettings is false, the asset to use to determine attenuation properties for effects generated by this component */
 	UPROPERTY(EditAnywhere, BlueprintReadWrite, Category=Attenuation, meta=(EditCondition="!bOverrideAttenuation"))
-	class UForceFeedbackAttenuation* AttenuationSettings;
+	TObjectPtr<class UForceFeedbackAttenuation> AttenuationSettings;
 
 	/** If bOverrideSettings is true, the attenuation properties to use for effects generated by this component */
 	UPROPERTY(EditAnywhere, BlueprintReadWrite, Category=Attenuation, meta=(EditCondition="bOverrideAttenuation"))
diff --git a/Engine/Source/Runtime/Engine/Classes/Components/InputComponent.h b/Engine/Source/Runtime/Engine/Classes/Components/InputComponent.h
--- a/Engine/Source/Runtime/Engine/Classes/Components/InputComponent.h
+++ b/Engine/Source/Runtime/Engine/Classes/Components/InputComponent.h
@@ -696,7 +696,7 @@ struct FCachedKeyToActionInfo
 
 	/** Which PlayerInput object this has been built for */
 	UPROPERTY()
-	UPlayerInput* PlayerInput;
+	TObjectPtr<UPlayerInput> PlayerInput;
 
 	/** What index of the player input's key mappings was the map built for. */
 	uint32 KeyMapBuiltForIndex;
diff --git a/Engine/Source/Runtime/Engine/Classes/Components/InstancedStaticMeshComponent.h b/Engine/Source/Runtime/Engine/Classes/Components/InstancedStaticMeshComponent.h
--- a/Engine/Source/Runtime/Engine/Classes/Components/InstancedStaticMeshComponent.h
+++ b/Engine/Source/Runtime/Engine/Classes/Components/InstancedStaticMeshComponent.h
@@ -493,7 +493,7 @@ public:
 public:
 	/** Mesh being used by component */
 	UPROPERTY()
-	UStaticMesh* StaticMesh = nullptr;
+	TObjectPtr<UStaticMesh> StaticMesh = nullptr;
 
 	// Static lighting info
 	UPROPERTY()
diff --git a/Engine/Source/Runtime/Engine/Classes/Components/LODSyncComponent.h b/Engine/Source/Runtime/Engine/Classes/Components/LODSyncComponent.h
--- a/Engine/Source/Runtime/Engine/Classes/Components/LODSyncComponent.h
+++ b/Engine/Source/Runtime/Engine/Classes/Components/LODSyncComponent.h
@@ -100,11 +100,11 @@ private:
 
 	// component that drives the LOD
 	UPROPERTY(transient)
-	TArray<UPrimitiveComponent*> DriveComponents;
+	TArray<TObjectPtr<UPrimitiveComponent>> DriveComponents;
 
 	// all the components that ticks
 	UPROPERTY(transient)
-	TArray<UPrimitiveComponent*> SubComponents;
+	TArray<TObjectPtr<UPrimitiveComponent>> SubComponents;
 
 	// BEGIN AActorComponent interface
 	virtual void OnRegister() override;
diff --git a/Engine/Source/Runtime/Engine/Classes/Components/LightComponent.h b/Engine/Source/Runtime/Engine/Classes/Components/LightComponent.h
--- a/Engine/Source/Runtime/Engine/Classes/Components/LightComponent.h
+++ b/Engine/Source/Runtime/Engine/Classes/Components/LightComponent.h
@@ -155,7 +155,7 @@ class ENGINE_API ULightComponent : public ULightComponentBase
 	 * Note that only non-lightmapped lights (UseDirectLightMap=False) can have a light function. 
 	 */
 	UPROPERTY(EditAnywhere, BlueprintReadOnly, Category=LightFunction)
-	class UMaterialInterface* LightFunctionMaterial;
+	TObjectPtr<class UMaterialInterface> LightFunctionMaterial;
 
 	/** Scales the light function projection.  X and Y scale in the directions perpendicular to the light's direction, Z scales along the light direction. */
 	UPROPERTY(EditAnywhere, BlueprintReadOnly, Category=LightFunction, meta=(AllowPreserveRatio = "true"))
@@ -163,7 +163,7 @@ class ENGINE_API ULightComponent : public ULightComponentBase
 
 	/** IES texture (light profiles from real world measured data) */
 	UPROPERTY(EditAnywhere, BlueprintReadOnly, Category=LightProfiles, meta=(DisplayName = "IES Texture"))
-	class UTextureLightProfile* IESTexture;
+	TObjectPtr<class UTextureLightProfile> IESTexture;
 
 	/** true: take light brightness from IES profile, false: use the light brightness - the maximum light in one direction is used to define no masking. Use with InverseSquareFalloff. Will be disabled if a valid IES profile texture is not supplied. */
 	UPROPERTY(EditAnywhere, BlueprintReadOnly, Category=LightProfiles, meta=(DisplayName = "Use IES Intensity"))
diff --git a/Engine/Source/Runtime/Engine/Classes/Components/LightComponentBase.h b/Engine/Source/Runtime/Engine/Classes/Components/LightComponentBase.h
--- a/Engine/Source/Runtime/Engine/Classes/Components/LightComponentBase.h
+++ b/Engine/Source/Runtime/Engine/Classes/Components/LightComponentBase.h
@@ -119,7 +119,7 @@ class ENGINE_API ULightComponentBase : public USceneComponent
 #if WITH_EDITORONLY_DATA
 	/** Sprite for static light in the editor. */
 	UPROPERTY(transient)
-	UTexture2D* StaticEditorTexture;
+	TObjectPtr<UTexture2D> StaticEditorTexture;
 
 	/** Sprite scaling for static light in the editor. */
 	UPROPERTY(transient)
@@ -127,7 +127,7 @@ class ENGINE_API ULightComponentBase : public USceneComponent
 
 	/** Sprite for dynamic light in the editor. */
 	UPROPERTY(transient)
-	UTexture2D* DynamicEditorTexture;
+	TObjectPtr<UTexture2D> DynamicEditorTexture;
 
 	/** Sprite scaling for dynamic light in the editor. */
 	UPROPERTY(transient)
diff --git a/Engine/Source/Runtime/Engine/Classes/Components/LightmassPortalComponent.h b/Engine/Source/Runtime/Engine/Classes/Components/LightmassPortalComponent.h
--- a/Engine/Source/Runtime/Engine/Classes/Components/LightmassPortalComponent.h
+++ b/Engine/Source/Runtime/Engine/Classes/Components/LightmassPortalComponent.h
@@ -14,7 +14,7 @@ class ULightmassPortalComponent : public USceneComponent
 	GENERATED_UCLASS_BODY()
 
 	UPROPERTY()
-	class UBoxComponent* PreviewBox;
+	TObjectPtr<class UBoxComponent> PreviewBox;
 
 public:
 	void UpdatePreviewShape();
diff --git a/Engine/Source/Runtime/Engine/Classes/Components/MaterialBillboardComponent.h b/Engine/Source/Runtime/Engine/Classes/Components/MaterialBillboardComponent.h
--- a/Engine/Source/Runtime/Engine/Classes/Components/MaterialBillboardComponent.h
+++ b/Engine/Source/Runtime/Engine/Classes/Components/MaterialBillboardComponent.h
@@ -19,11 +19,11 @@ struct FMaterialSpriteElement
 
 	/** The material that the sprite is rendered with. */
 	UPROPERTY(EditAnywhere, BlueprintReadWrite, Category=MaterialSpriteElement)
-	class UMaterialInterface* Material;
+	TObjectPtr<class UMaterialInterface> Material;
 	
 	/** A curve that maps distance on the X axis to the sprite opacity on the Y axis. */
 	UPROPERTY(EditAnywhere, BlueprintReadWrite, Category=MaterialSpriteElement)
-	UCurveFloat* DistanceToOpacityCurve;
+	TObjectPtr<UCurveFloat> DistanceToOpacityCurve;
 	
 	/** Whether the size is defined in screen-space or world-space. */
 	UPROPERTY(EditAnywhere, BlueprintReadWrite, Category=MaterialSpriteElement)
@@ -39,7 +39,7 @@ struct FMaterialSpriteElement
 
 	/** A curve that maps distance on the X axis to the sprite size on the Y axis. */
 	UPROPERTY(EditAnywhere, BlueprintReadWrite, Category=MaterialSpriteElement)
-	UCurveFloat* DistanceToSizeCurve;
+	TObjectPtr<UCurveFloat> DistanceToSizeCurve;
 
 	FMaterialSpriteElement()
 		: Material(NULL)
diff --git a/Engine/Source/Runtime/Engine/Classes/Components/MeshComponent.h b/Engine/Source/Runtime/Engine/Classes/Components/MeshComponent.h
--- a/Engine/Source/Runtime/Engine/Classes/Components/MeshComponent.h
+++ b/Engine/Source/Runtime/Engine/Classes/Components/MeshComponent.h
@@ -25,7 +25,7 @@ class ENGINE_API UMeshComponent : public UPrimitiveComponent
 public:
 	/** Per-Component material overrides.  These must NOT be set directly or a race condition can occur between GC and the rendering thread. */
 	UPROPERTY(EditAnywhere, AdvancedDisplay, Category=Rendering, Meta=(ToolTip="Material overrides."))
-	TArray<class UMaterialInterface*> OverrideMaterials;
+	TArray<TObjectPtr<class UMaterialInterface>> OverrideMaterials;
 	
 	UFUNCTION(BlueprintCallable, Category="Components|Mesh")
 	virtual TArray<class UMaterialInterface*> GetMaterials() const;
diff --git a/Engine/Source/Runtime/Engine/Classes/Components/ModelComponent.h b/Engine/Source/Runtime/Engine/Classes/Components/ModelComponent.h
--- a/Engine/Source/Runtime/Engine/Classes/Components/ModelComponent.h
+++ b/Engine/Source/Runtime/Engine/Classes/Components/ModelComponent.h
@@ -45,7 +45,7 @@ public:
 
 	/** Description of collision */
 	UPROPERTY()
-	class UBodySetup* ModelBodySetup;
+	TObjectPtr<class UBodySetup> ModelBodySetup;
 
 private:
 	/** The nodes which this component renders. */
diff --git a/Engine/Source/Runtime/Engine/Classes/Components/PlanarReflectionComponent.h b/Engine/Source/Runtime/Engine/Classes/Components/PlanarReflectionComponent.h
--- a/Engine/Source/Runtime/Engine/Classes/Components/PlanarReflectionComponent.h
+++ b/Engine/Source/Runtime/Engine/Classes/Components/PlanarReflectionComponent.h
@@ -18,7 +18,7 @@ class UPlanarReflectionComponent : public USceneCaptureComponent
 	GENERATED_UCLASS_BODY()
 
 	UPROPERTY()
-	class UBoxComponent* PreviewBox;
+	TObjectPtr<class UBoxComponent> PreviewBox;
 
 public:
 
@@ -116,7 +116,7 @@ protected:
 #if WITH_EDITORONLY_DATA
 	/** The material to use on ProxyMeshComponent */
 	UPROPERTY(transient)
-	class UMaterial* CaptureMaterial;
+	TObjectPtr<class UMaterial> CaptureMaterial;
 #endif
 
 private:
diff --git a/Engine/Source/Runtime/Engine/Classes/Components/PlaneReflectionCaptureComponent.h b/Engine/Source/Runtime/Engine/Classes/Components/PlaneReflectionCaptureComponent.h
--- a/Engine/Source/Runtime/Engine/Classes/Components/PlaneReflectionCaptureComponent.h
+++ b/Engine/Source/Runtime/Engine/Classes/Components/PlaneReflectionCaptureComponent.h
@@ -19,10 +19,10 @@ class UPlaneReflectionCaptureComponent : public UReflectionCaptureComponent
 	float InfluenceRadiusScale;
 
 	UPROPERTY()
-	class UDrawSphereComponent* PreviewInfluenceRadius;
+	TObjectPtr<class UDrawSphereComponent> PreviewInfluenceRadius;
 
 	UPROPERTY()
-	class UBoxComponent* PreviewCaptureBox;
+	TObjectPtr<class UBoxComponent> PreviewCaptureBox;
 
 public:
 	virtual void UpdatePreviewShape() override;
diff --git a/Engine/Source/Runtime/Engine/Classes/Components/PrimitiveComponent.h b/Engine/Source/Runtime/Engine/Classes/Components/PrimitiveComponent.h
--- a/Engine/Source/Runtime/Engine/Classes/Components/PrimitiveComponent.h
+++ b/Engine/Source/Runtime/Engine/Classes/Components/PrimitiveComponent.h
@@ -619,7 +619,7 @@ public:
 	 * The material also needs to be set up to output to a virtual texture. 
 	 */
 	UPROPERTY(EditAnywhere, BlueprintReadWrite, Category = VirtualTexture, meta = (DisplayName = "Draw in Virtual Textures"))
-	TArray<URuntimeVirtualTexture*> RuntimeVirtualTextures;
+	TArray<TObjectPtr<URuntimeVirtualTexture>> RuntimeVirtualTextures;
 
 	/** Bias to the LOD selected for rendering to runtime virtual textures. */
 	UPROPERTY(EditAnywhere, AdvancedDisplay, Category = VirtualTexture, meta = (DisplayName = "Virtual Texture LOD Bias", UIMin = "-7", UIMax = "8"))
@@ -739,7 +739,7 @@ public:
 	 * @see IgnoreActorWhenMoving()
 	 */
 	UPROPERTY(Transient, DuplicateTransient)
-	TArray<AActor*> MoveIgnoreActors;
+	TArray<TObjectPtr<AActor>> MoveIgnoreActors;
 
 	/**
 	 * Tells this component whether to ignore collision with all components of a specific Actor when this component is moved.
@@ -774,7 +774,7 @@ public:
 	* @see IgnoreComponentWhenMoving()
 	*/
 	UPROPERTY(Transient, DuplicateTransient)
-	TArray<UPrimitiveComponent*> MoveIgnoreComponents;
+	TArray<TObjectPtr<UPrimitiveComponent>> MoveIgnoreComponents;
 
 	/**
 	* Tells this component whether to ignore collision with another component when this component is moved.
@@ -1575,10 +1575,10 @@ public:
 private:
 	/** LOD parent primitive to draw instead of this one (multiple UPrim's will point to the same LODParent ) */
 	UPROPERTY(NonPIEDuplicateTransient)
-	class UPrimitiveComponent* LODParentPrimitive;
+	TObjectPtr<class UPrimitiveComponent> LODParentPrimitive;
 
 	UPROPERTY(Transient)
-	class UPrimitiveComponent* CachedLODParentPrimitive;
+	TObjectPtr<class UPrimitiveComponent> CachedLODParentPrimitive;
 
 public:
 	/** Set the LOD parent component */
@@ -2402,7 +2402,7 @@ private:
 	int32 VisibilityId;
 
 	UPROPERTY()
-	UPrimitiveComponent* LODParent = nullptr;
+	TObjectPtr<UPrimitiveComponent> LODParent = nullptr;
 };
 
 
diff --git a/Engine/Source/Runtime/Engine/Classes/Components/RectLightComponent.h b/Engine/Source/Runtime/Engine/Classes/Components/RectLightComponent.h
--- a/Engine/Source/Runtime/Engine/Classes/Components/RectLightComponent.h
+++ b/Engine/Source/Runtime/Engine/Classes/Components/RectLightComponent.h
@@ -49,7 +49,7 @@ class ENGINE_API URectLightComponent : public ULocalLightComponent
 
 	/** Texture mapped to the light source rectangle */
 	UPROPERTY(EditAnywhere, BlueprintReadOnly, Category = Light)
-	class UTexture* SourceTexture;
+	TObjectPtr<class UTexture> SourceTexture;
 
 	UFUNCTION(BlueprintCallable, Category = "Rendering|Lighting")
 	void SetSourceTexture(UTexture* bNewValue);
diff --git a/Engine/Source/Runtime/Engine/Classes/Components/ReflectionCaptureComponent.h b/Engine/Source/Runtime/Engine/Classes/Components/ReflectionCaptureComponent.h
--- a/Engine/Source/Runtime/Engine/Classes/Components/ReflectionCaptureComponent.h
+++ b/Engine/Source/Runtime/Engine/Classes/Components/ReflectionCaptureComponent.h
@@ -39,7 +39,7 @@ class UReflectionCaptureComponent : public USceneComponent
 	GENERATED_UCLASS_BODY()
 
 	UPROPERTY()
-	UBillboardComponent* CaptureOffsetComponent;
+	TObjectPtr<UBillboardComponent> CaptureOffsetComponent;
 
 	/** Indicates where to get the reflection source from. */
 	UPROPERTY(EditAnywhere, BlueprintReadOnly, Category=ReflectionCapture)
@@ -51,7 +51,7 @@ class UReflectionCaptureComponent : public USceneComponent
 
 	/** Cubemap to use for reflection if ReflectionSourceType is set to RS_SpecifiedCubemap. */
 	UPROPERTY(EditAnywhere, BlueprintReadOnly, Category=ReflectionCapture)
-	class UTextureCube* Cubemap;
+	TObjectPtr<class UTextureCube> Cubemap;
 
 	/** Angle to rotate the source cubemap when SourceType is set to SLS_SpecifiedCubemap. */
 	UPROPERTY(EditAnywhere, BlueprintReadOnly, Category=ReflectionCapture, meta = (UIMin = "0", UIMax = "360"))
@@ -79,7 +79,7 @@ class UReflectionCaptureComponent : public USceneComponent
 
 	/** Cached Cubemap texture from MapBuildData used for rendering with the encoded HDR values. */
 	UPROPERTY(transient)
-	UTextureCube* CachedEncodedHDRCubemap;
+	TObjectPtr<UTextureCube> CachedEncodedHDRCubemap;
 
 #if WITH_EDITOR
 	/** Check to see if MapBuildDataId was loaded - otherwise we need to display a warning on cook */
diff --git a/Engine/Source/Runtime/Engine/Classes/Components/RuntimeVirtualTextureComponent.h b/Engine/Source/Runtime/Engine/Classes/Components/RuntimeVirtualTextureComponent.h
--- a/Engine/Source/Runtime/Engine/Classes/Components/RuntimeVirtualTextureComponent.h
+++ b/Engine/Source/Runtime/Engine/Classes/Components/RuntimeVirtualTextureComponent.h
@@ -31,7 +31,7 @@ protected:
 
 	/** The virtual texture object to use. */
 	UPROPERTY(EditAnywhere, BlueprintReadOnly, NonPIEDuplicateTransient, Category = VirtualTexture)
-	URuntimeVirtualTexture* VirtualTexture = nullptr;
+	TObjectPtr<URuntimeVirtualTexture> VirtualTexture = nullptr;
 
 	/** Set to true to enable scalability settings for the virtual texture. */
 	UPROPERTY(EditAnywhere, AdvancedDisplay, Category = VirtualTexture, meta = (InlineEditConditionToggle))
@@ -47,7 +47,7 @@ protected:
 
 	/** Texture object containing streamed low mips. */
 	UPROPERTY(EditAnywhere, BlueprintReadOnly, NonPIEDuplicateTransient, Category = VirtualTextureBuild)
-	UVirtualTextureBuilder* StreamingTexture = nullptr;
+	TObjectPtr<UVirtualTextureBuilder> StreamingTexture = nullptr;
 
 	/** Number of low mips to serialize and stream for the virtual texture. This can reduce rendering update cost. */
 	UPROPERTY(EditAnywhere, Category = VirtualTextureBuild, meta = (UIMin = "0", UIMax = "12", DisplayName = "Streaming Levels"))
diff --git a/Engine/Source/Runtime/Engine/Classes/Components/SceneCaptureComponent.h b/Engine/Source/Runtime/Engine/Classes/Components/SceneCaptureComponent.h
--- a/Engine/Source/Runtime/Engine/Classes/Components/SceneCaptureComponent.h
+++ b/Engine/Source/Runtime/Engine/Classes/Components/SceneCaptureComponent.h
@@ -86,7 +86,7 @@ class ENGINE_API USceneCaptureComponent : public USceneComponent
 
 	/** The actors to hide in the scene capture. */
 	UPROPERTY(EditInstanceOnly, BlueprintReadWrite, Category=SceneCapture)
-	TArray<AActor*> HiddenActors;
+	TArray<TObjectPtr<AActor>> HiddenActors;
 
 	/** The only components to be rendered by this scene capture, if PrimitiveRenderMode is set to UseShowOnlyList. */
  	UPROPERTY()
@@ -94,7 +94,7 @@ class ENGINE_API USceneCaptureComponent : public USceneComponent
 
 	/** The only actors to be rendered by this scene capture, if PrimitiveRenderMode is set to UseShowOnlyList.*/
 	UPROPERTY(EditInstanceOnly, BlueprintReadWrite, Category=SceneCapture)
-	TArray<AActor*> ShowOnlyActors;
+	TArray<TObjectPtr<AActor>> ShowOnlyActors;
 
 	/** Scales the distance used by LOD. Set to values greater than 1 to cause the scene capture to use lower LODs than the main view to speed up the scene capture pass. */
 	UPROPERTY(EditAnywhere, Category=PlanarReflection, meta=(UIMin = ".1", UIMax = "10"), AdvancedDisplay)
@@ -213,7 +213,7 @@ protected:
 #if WITH_EDITORONLY_DATA
 	/** The mesh used by ProxyMeshComponent */
 	UPROPERTY(transient)
-	class UStaticMesh* CaptureMesh;
+	TObjectPtr<class UStaticMesh> CaptureMesh;
 
 	/** The mesh to show visually where the camera is placed */
 	class UStaticMeshComponent* ProxyMeshComponent;
diff --git a/Engine/Source/Runtime/Engine/Classes/Components/SceneCaptureComponent2D.h b/Engine/Source/Runtime/Engine/Classes/Components/SceneCaptureComponent2D.h
--- a/Engine/Source/Runtime/Engine/Classes/Components/SceneCaptureComponent2D.h
+++ b/Engine/Source/Runtime/Engine/Classes/Components/SceneCaptureComponent2D.h
@@ -36,7 +36,7 @@ public:
 
 	/** Output render target of the scene capture that can be read in materals. */
 	UPROPERTY(EditAnywhere, BlueprintReadWrite, Category=SceneCapture)
-	class UTextureRenderTarget2D* TextureTarget;
+	TObjectPtr<class UTextureRenderTarget2D> TextureTarget;
 
 	/** When enabled, the scene capture will composite into the render target instead of overwriting its contents. */
 	UPROPERTY(EditAnywhere, BlueprintReadWrite, Category=SceneCapture)
diff --git a/Engine/Source/Runtime/Engine/Classes/Components/SceneCaptureComponentCube.h b/Engine/Source/Runtime/Engine/Classes/Components/SceneCaptureComponentCube.h
--- a/Engine/Source/Runtime/Engine/Classes/Components/SceneCaptureComponentCube.h
+++ b/Engine/Source/Runtime/Engine/Classes/Components/SceneCaptureComponentCube.h
@@ -20,7 +20,7 @@ class ENGINE_API USceneCaptureComponentCube : public USceneCaptureComponent
 
 	/** Temporary render target that can be used by the editor. */
 	UPROPERTY(EditAnywhere, BlueprintReadWrite, Category = SceneCapture)
-	class UTextureRenderTargetCube* TextureTarget;
+	TObjectPtr<class UTextureRenderTargetCube> TextureTarget;
 
 	/** Preserve the rotation of the actor when updating the capture. The default behavior is to capture the cube aligned to the world axis system.  */
 	UPROPERTY(EditAnywhere, BlueprintReadWrite, Category = SceneCapture)
@@ -38,13 +38,13 @@ class ENGINE_API USceneCaptureComponentCube : public USceneCaptureComponent
 	*/
 
 	UPROPERTY(EditAnywhere, BlueprintReadWrite, Category = "SceneCapture|Omni-directional Stereo Capture")
-	class UTextureRenderTargetCube* TextureTargetLeft;
+	TObjectPtr<class UTextureRenderTargetCube> TextureTargetLeft;
 
 	UPROPERTY(EditAnywhere, BlueprintReadWrite, Category = "SceneCapture|Omni-directional Stereo Capture")
-	class UTextureRenderTargetCube* TextureTargetRight;
+	TObjectPtr<class UTextureRenderTargetCube> TextureTargetRight;
 
 	UPROPERTY(EditAnywhere, BlueprintReadWrite, Category = "SceneCapture|Omni-directional Stereo Capture")
-	class UTextureRenderTarget2D* TextureTargetODS;
+	TObjectPtr<class UTextureRenderTarget2D> TextureTargetODS;
 
 	UPROPERTY(EditAnywhere, BlueprintReadWrite, Category = "SceneCapture|Omni-directional Stereo Capture", meta = (DisplayName = "Interpupillary Distance (cm)", UIMin = "0.0", UIMax = "16.0", ClampMin = "0.0", ClampMax = "16.0"))
 	float IPD;
diff --git a/Engine/Source/Runtime/Engine/Classes/Components/SceneComponent.h b/Engine/Source/Runtime/Engine/Classes/Components/SceneComponent.h
--- a/Engine/Source/Runtime/Engine/Classes/Components/SceneComponent.h
+++ b/Engine/Source/Runtime/Engine/Classes/Components/SceneComponent.h
@@ -128,7 +128,7 @@ private:
 
 	/** What we are currently attached to. If valid, RelativeLocation etc. are used relative to this object */
 	UPROPERTY(ReplicatedUsing = OnRep_AttachParent)
-	USceneComponent* AttachParent;
+	TObjectPtr<USceneComponent> AttachParent;
 
 	/** Optional socket name on AttachParent that we are attached to. */
 	UPROPERTY(ReplicatedUsing = OnRep_AttachSocketName)
@@ -136,11 +136,11 @@ private:
 
 	/** List of child SceneComponents that are attached to us. */
 	UPROPERTY(ReplicatedUsing = OnRep_AttachChildren, Transient)
-	TArray<USceneComponent*> AttachChildren;
+	TArray<TObjectPtr<USceneComponent>> AttachChildren;
 
 	/** Set of attached SceneComponents that were attached by the client so we can fix up AttachChildren when it is replicated to us. */
 	UPROPERTY(Transient)
-	TArray<USceneComponent*> ClientAttachedChildren;
+	TArray<TObjectPtr<USceneComponent>> ClientAttachedChildren;
 
 	FName NetOldAttachSocketName;
 	USceneComponent* NetOldAttachParent;
@@ -1635,7 +1635,7 @@ struct ENGINE_API FSceneComponentInstanceData : public FActorComponentInstanceDa
 	virtual void AddReferencedObjects(FReferenceCollector& Collector) override;
 
 	UPROPERTY() 
-	TMap<USceneComponent*, FTransform> AttachedInstanceComponents;
+	TMap<TObjectPtr<USceneComponent>, FTransform> AttachedInstanceComponents;
 };
 
 
diff --git a/Engine/Source/Runtime/Engine/Classes/Components/ShapeComponent.h b/Engine/Source/Runtime/Engine/Classes/Components/ShapeComponent.h
--- a/Engine/Source/Runtime/Engine/Classes/Components/ShapeComponent.h
+++ b/Engine/Source/Runtime/Engine/Classes/Components/ShapeComponent.h
@@ -27,7 +27,7 @@ class ENGINE_API UShapeComponent : public UPrimitiveComponent
 
 	/** Description of collision */
 	UPROPERTY(transient, duplicatetransient)
-	class UBodySetup* ShapeBodySetup;
+	TObjectPtr<class UBodySetup> ShapeBodySetup;
 
 	/** Navigation area type (empty = default obstacle) */
 	UPROPERTY(EditAnywhere, Category = Navigation)
diff --git a/Engine/Source/Runtime/Engine/Classes/Components/SkeletalMeshComponent.h b/Engine/Source/Runtime/Engine/Classes/Components/SkeletalMeshComponent.h
--- a/Engine/Source/Runtime/Engine/Classes/Components/SkeletalMeshComponent.h
+++ b/Engine/Source/Runtime/Engine/Classes/Components/SkeletalMeshComponent.h
@@ -340,12 +340,12 @@ public:
 #if WITH_EDITORONLY_DATA
 	/** The blueprint for creating an AnimationScript. */
 	UPROPERTY()
-	class UAnimBlueprint* AnimationBlueprint_DEPRECATED;
+	TObjectPtr<class UAnimBlueprint> AnimationBlueprint_DEPRECATED;
 #endif
 
 	UE_DEPRECATED(4.11, "This property is deprecated. Please use AnimClass instead")
 	UPROPERTY(BlueprintReadOnly, Category = Animation, meta = (DeprecationMessage = "This property is deprecated. Please use AnimClass instead"))
-	class UAnimBlueprintGeneratedClass* AnimBlueprintGeneratedClass;
+	TObjectPtr<class UAnimBlueprintGeneratedClass> AnimBlueprintGeneratedClass;
 
 	/* The AnimBlueprint class to use. Use 'SetAnimInstanceClass' to change at runtime. */
 	UPROPERTY(EditAnywhere, BlueprintReadOnly, Category = Animation)
@@ -353,20 +353,20 @@ public:
 
 	/** The active animation graph program instance. */
 	UPROPERTY(transient, NonTransactional)
-	UAnimInstance* AnimScriptInstance;
+	TObjectPtr<UAnimInstance> AnimScriptInstance;
 
 #if WITH_EDITORONLY_DATA
 	/** Any running linked anim instances */
 	UE_DEPRECATED(4.24, "Direct access to this property is deprecated and the array is no longer used. Storage is now in LinkedInstances. Please use GetLinkedAnimInstances() instead.")
 	UPROPERTY(transient)
-	TArray<UAnimInstance*> SubInstances;
+	TArray<TObjectPtr<UAnimInstance>> SubInstances;
 #endif
 
 	/** An instance created from the PostPhysicsBlueprint property of the skeletal mesh we're using,
 	 *  Runs after (and receives pose from) the main anim instance.
 	 */
 	UPROPERTY(transient)
-	UAnimInstance* PostProcessAnimInstance;
+	TObjectPtr<UAnimInstance> PostProcessAnimInstance;
 
 public:
 
@@ -416,7 +416,7 @@ public:
 private:
 	/** Any running linked anim instances */
 	UPROPERTY(transient)
-	TArray<UAnimInstance*> LinkedInstances;
+	TArray<TObjectPtr<UAnimInstance>> LinkedInstances;
 
 	// Update Rate
 
@@ -834,7 +834,7 @@ public:
 	* Used for per poly collision. In 99% of cases you will be better off using a Physics Asset.
 	* This BodySetup is per instance because all modification of vertices is done in place */
 	UPROPERTY(transient)
-	class UBodySetup * BodySetup;
+	TObjectPtr<class UBodySetup>  BodySetup;
 
 	void CreateBodySetup();
 
@@ -1446,7 +1446,7 @@ private:
 	 * it is safe to do so the interactor will sync to the simulation context
 	 */
 	UPROPERTY(Transient)
-	UClothingSimulationInteractor* ClothingInteractor;
+	TObjectPtr<UClothingSimulationInteractor> ClothingInteractor;
 
 	/** Helper struct used to store info about a cloth collision source */
 	struct FClothCollisionSource
@@ -2358,11 +2358,11 @@ private:
 	// remove if this version goes away : VER_UE4_REMOVE_SINGLENODEINSTANCE
 	// deprecated variable to be re-save
 	UPROPERTY()
-	class UAnimSequence* SequenceToPlay_DEPRECATED;
+	TObjectPtr<class UAnimSequence> SequenceToPlay_DEPRECATED;
 
 	// The default sequence to play on this skeletal mesh
 	UPROPERTY()
-	class UAnimationAsset* AnimToPlay_DEPRECATED;
+	TObjectPtr<class UAnimationAsset> AnimToPlay_DEPRECATED;
 
 	// Default setting for looping for SequenceToPlay. This is not current state of looping.
 	UPROPERTY()
diff --git a/Engine/Source/Runtime/Engine/Classes/Components/SkinnedMeshComponent.h b/Engine/Source/Runtime/Engine/Classes/Components/SkinnedMeshComponent.h
--- a/Engine/Source/Runtime/Engine/Classes/Components/SkinnedMeshComponent.h
+++ b/Engine/Source/Runtime/Engine/Classes/Components/SkinnedMeshComponent.h
@@ -229,7 +229,7 @@ class ENGINE_API USkinnedMeshComponent : public UMeshComponent, public ILODSyncI
 
 	/** The skeletal mesh used by this component. */
 	UPROPERTY(EditAnywhere, BlueprintReadOnly, Category="Mesh")
-	class USkeletalMesh* SkeletalMesh;
+	TObjectPtr<class USkeletalMesh> SkeletalMesh;
 
 	//
 	// MasterPoseComponent.
@@ -421,7 +421,7 @@ public:
 	 *	PhysicsAsset is set in SkeletalMesh by default, but you can override with this value
 	 */
 	UPROPERTY(EditAnywhere, AdvancedDisplay, BlueprintReadOnly, Category=Physics)
-	class UPhysicsAsset* PhysicsAssetOverride;
+	TObjectPtr<class UPhysicsAsset> PhysicsAssetOverride;
 
 	//
 	// Level of detail.
diff --git a/Engine/Source/Runtime/Engine/Classes/Components/SkyAtmosphereComponent.h b/Engine/Source/Runtime/Engine/Classes/Components/SkyAtmosphereComponent.h
--- a/Engine/Source/Runtime/Engine/Classes/Components/SkyAtmosphereComponent.h
+++ b/Engine/Source/Runtime/Engine/Classes/Components/SkyAtmosphereComponent.h
@@ -267,12 +267,12 @@ class ASkyAtmosphere : public AInfo
 private:
 
 	UPROPERTY(BlueprintReadOnly, VisibleAnywhere, Category = Atmosphere, meta = (AllowPrivateAccess = "true"))
-	class USkyAtmosphereComponent* SkyAtmosphereComponent;
+	TObjectPtr<class USkyAtmosphereComponent> SkyAtmosphereComponent;
 
 #if WITH_EDITORONLY_DATA
 	/** Arrow component to indicate default sun rotation */
 	UPROPERTY()
-	class UArrowComponent* ArrowComponent;
+	TObjectPtr<class UArrowComponent> ArrowComponent;
 #endif
 
 public:
diff --git a/Engine/Source/Runtime/Engine/Classes/Components/SkyLightComponent.h b/Engine/Source/Runtime/Engine/Classes/Components/SkyLightComponent.h
--- a/Engine/Source/Runtime/Engine/Classes/Components/SkyLightComponent.h
+++ b/Engine/Source/Runtime/Engine/Classes/Components/SkyLightComponent.h
@@ -105,7 +105,7 @@ class ENGINE_API USkyLightComponent : public ULightComponentBase
 
 	/** Cubemap to use for sky lighting if SourceType is set to SLS_SpecifiedCubemap. */
 	UPROPERTY(EditAnywhere, BlueprintReadOnly, Category=Light)
-	class UTextureCube* Cubemap;
+	TObjectPtr<class UTextureCube> Cubemap;
 
 	/** Angle to rotate the source cubemap when SourceType is set to SLS_SpecifiedCubemap. */
 	UPROPERTY(EditAnywhere, BlueprintReadOnly, Category=Light, meta=(UIMin = "0", UIMax = "360"))
@@ -320,7 +320,7 @@ protected:
 	float BlendFraction;
 
 	UPROPERTY(transient)
-	class UTextureCube* BlendDestinationCubemap;
+	TObjectPtr<class UTextureCube> BlendDestinationCubemap;
 	TRefCountPtr<FSkyTextureCubeResource> BlendDestinationProcessedSkyTexture;
 	FSHVectorRGB3 BlendDestinationIrradianceEnvironmentMap;
 	float BlendDestinationAverageBrightness;
diff --git a/Engine/Source/Runtime/Engine/Classes/Components/SphereReflectionCaptureComponent.h b/Engine/Source/Runtime/Engine/Classes/Components/SphereReflectionCaptureComponent.h
--- a/Engine/Source/Runtime/Engine/Classes/Components/SphereReflectionCaptureComponent.h
+++ b/Engine/Source/Runtime/Engine/Classes/Components/SphereReflectionCaptureComponent.h
@@ -23,7 +23,7 @@ class USphereReflectionCaptureComponent : public UReflectionCaptureComponent
 	float CaptureDistanceScale;
 
 	UPROPERTY()
-	class UDrawSphereComponent* PreviewInfluenceRadius;
+	TObjectPtr<class UDrawSphereComponent> PreviewInfluenceRadius;
 
 public:
 	virtual void UpdatePreviewShape() override;
diff --git a/Engine/Source/Runtime/Engine/Classes/Components/SplineComponent.h b/Engine/Source/Runtime/Engine/Classes/Components/SplineComponent.h
--- a/Engine/Source/Runtime/Engine/Classes/Components/SplineComponent.h
+++ b/Engine/Source/Runtime/Engine/Classes/Components/SplineComponent.h
@@ -77,7 +77,7 @@ struct ENGINE_API FSplineCurves
 	FInterpCurveFloat ReparamTable;
 
 	UPROPERTY()
-	USplineMetadata* Metadata_DEPRECATED = nullptr;
+	TObjectPtr<USplineMetadata> Metadata_DEPRECATED = nullptr;
 
 	UPROPERTY(transient)
 	uint32 Version = 0xffffffff;
diff --git a/Engine/Source/Runtime/Engine/Classes/Components/SplineMeshComponent.h b/Engine/Source/Runtime/Engine/Classes/Components/SplineMeshComponent.h
--- a/Engine/Source/Runtime/Engine/Classes/Components/SplineMeshComponent.h
+++ b/Engine/Source/Runtime/Engine/Classes/Components/SplineMeshComponent.h
@@ -119,7 +119,7 @@ class ENGINE_API USplineMeshComponent : public UStaticMeshComponent, public IInt
 
 	// Physics data.
 	UPROPERTY()
-	UBodySetup* BodySetup;
+	TObjectPtr<UBodySetup> BodySetup;
 
 	/** Maximum coordinate along the spline forward axis which corresponds to end of spline. If set to 0.0, will use bounding box to determine bounds */
 	UPROPERTY(EditAnywhere, Category = SplineMesh, AdvancedDisplay)
diff --git a/Engine/Source/Runtime/Engine/Classes/Components/StaticMeshComponent.h b/Engine/Source/Runtime/Engine/Classes/Components/StaticMeshComponent.h
--- a/Engine/Source/Runtime/Engine/Classes/Components/StaticMeshComponent.h
+++ b/Engine/Source/Runtime/Engine/Classes/Components/StaticMeshComponent.h
@@ -193,7 +193,7 @@ class ENGINE_API UStaticMeshComponent : public UMeshComponent
 	/** The static mesh that this component uses to render */
 private:
 	UPROPERTY(EditAnywhere, BlueprintReadOnly, Category=StaticMesh, ReplicatedUsing=OnRep_StaticMesh, meta=(AllowPrivateAccess="true"))
-	class UStaticMesh* StaticMesh;
+	TObjectPtr<class UStaticMesh> StaticMesh;
 
 	void SetStaticMeshInternal(UStaticMesh* StaticMesh);
 
@@ -775,7 +775,7 @@ public:
 
 	/** Mesh being used by component */
 	UPROPERTY()
-	class UStaticMesh* StaticMesh = nullptr;
+	TObjectPtr<class UStaticMesh> StaticMesh = nullptr;
 
 	/** Array of cached vertex colors for each LOD */
 	UPROPERTY()
diff --git a/Engine/Source/Runtime/Engine/Classes/Components/StereoLayerComponent.h b/Engine/Source/Runtime/Engine/Classes/Components/StereoLayerComponent.h
--- a/Engine/Source/Runtime/Engine/Classes/Components/StereoLayerComponent.h
+++ b/Engine/Source/Runtime/Engine/Classes/Components/StereoLayerComponent.h
@@ -345,11 +345,11 @@ public:
 protected:
 	/** Texture displayed on the stereo layer (is stereocopic textures are supported on the platfrom and more than one texture is provided, this will be the right eye) **/
 	UPROPERTY(EditAnywhere, BlueprintReadOnly, Category= "StereoLayer")
-	class UTexture* Texture;
+	TObjectPtr<class UTexture> Texture;
 
 	/** Texture displayed on the stereo layer for left eye, if stereoscopic textures are supported on the platform and by the layer shape **/
 	UPROPERTY(EditAnywhere, BlueprintReadOnly, Category = "Stereoscopic Properties")
-	class UTexture* LeftTexture;
+	TObjectPtr<class UTexture> LeftTexture;
 
 public:
 	/** True if the quad should internally set it's Y value based on the set texture's dimensions */
@@ -391,7 +391,7 @@ protected:
 
 	/** Specifies which shape of layer it is.  Note that some shapes will be supported only on certain platforms! **/
 	UPROPERTY(EditAnywhere, BlueprintReadOnly, NoClear, Instanced, Category = "StereoLayer", DisplayName="Stereo Layer Shape")
-	UStereoLayerShape* Shape;
+	TObjectPtr<UStereoLayerShape> Shape;
 
 
 	/** Render priority among all stereo layers, higher priority render on top of lower priority **/
diff --git a/Engine/Source/Runtime/Engine/Classes/Components/TextRenderComponent.h b/Engine/Source/Runtime/Engine/Classes/Components/TextRenderComponent.h
--- a/Engine/Source/Runtime/Engine/Classes/Components/TextRenderComponent.h
+++ b/Engine/Source/Runtime/Engine/Classes/Components/TextRenderComponent.h
@@ -51,11 +51,11 @@ class ENGINE_API UTextRenderComponent : public UPrimitiveComponent
 
 	/** Text material */
 	UPROPERTY(EditAnywhere, BlueprintReadOnly, Category=Text)
-	class UMaterialInterface* TextMaterial;
+	TObjectPtr<class UMaterialInterface> TextMaterial;
 	
 	/** Text font */
 	UPROPERTY(EditAnywhere, BlueprintReadOnly, Category=Text)
-	class UFont* Font;
+	TObjectPtr<class UFont> Font;
 
 	/** Horizontal text alignment */
 	UPROPERTY(EditAnywhere, BlueprintReadOnly, Category=Text, meta=(DisplayName = "Horizontal Alignment"))
diff --git a/Engine/Source/Runtime/Engine/Classes/Components/TimelineComponent.h b/Engine/Source/Runtime/Engine/Classes/Components/TimelineComponent.h
--- a/Engine/Source/Runtime/Engine/Classes/Components/TimelineComponent.h
+++ b/Engine/Source/Runtime/Engine/Classes/Components/TimelineComponent.h
@@ -78,7 +78,7 @@ struct FTimelineVectorTrack
 
 	/** Vector curve to be evaluated */
 	UPROPERTY()
-	class UCurveVector* VectorCurve;
+	TObjectPtr<class UCurveVector> VectorCurve;
 
 	/** Function that the output from ValueCurve will be passed to */
 	UPROPERTY()
@@ -115,7 +115,7 @@ struct FTimelineFloatTrack
 
 	/** Float curve to be evaluated */
 	UPROPERTY()
-	class UCurveFloat* FloatCurve;
+	TObjectPtr<class UCurveFloat> FloatCurve;
 
 	/** Function that the output from ValueCurve will be passed to */
 	UPROPERTY()
@@ -153,7 +153,7 @@ struct FTimelineLinearColorTrack
 
 	/** Float curve to be evaluated */
 	UPROPERTY()
-	class UCurveLinearColor* LinearColorCurve;
+	TObjectPtr<class UCurveLinearColor> LinearColorCurve;
 
 	/** Function that the output from ValueCurve will be passed to */
 	UPROPERTY()
diff --git a/Engine/Source/Runtime/Engine/Classes/Components/VectorFieldComponent.h b/Engine/Source/Runtime/Engine/Classes/Components/VectorFieldComponent.h
--- a/Engine/Source/Runtime/Engine/Classes/Components/VectorFieldComponent.h
+++ b/Engine/Source/Runtime/Engine/Classes/Components/VectorFieldComponent.h
@@ -21,7 +21,7 @@ class UVectorFieldComponent : public UPrimitiveComponent
 
 	/** The vector field asset. */
 	UPROPERTY(EditAnywhere, BlueprintReadWrite, Category=VectorFieldComponent)
-	class UVectorField* VectorField;
+	TObjectPtr<class UVectorField> VectorField;
 
 	/** The intensity at which the vector field is applied. */
 	UPROPERTY(interp, Category=VectorFieldComponent)
diff --git a/Engine/Source/Runtime/Engine/Classes/Components/VolumetricCloudComponent.h b/Engine/Source/Runtime/Engine/Classes/Components/VolumetricCloudComponent.h
--- a/Engine/Source/Runtime/Engine/Classes/Components/VolumetricCloudComponent.h
+++ b/Engine/Source/Runtime/Engine/Classes/Components/VolumetricCloudComponent.h
@@ -54,7 +54,7 @@ class UVolumetricCloudComponent : public USceneComponent
 
 	/** The material describing the cloud volume. It must be a Volume domain material. */
 	UPROPERTY(EditAnywhere, BlueprintReadOnly, Category = "Cloud Material")
-	UMaterialInterface* Material;
+	TObjectPtr<UMaterialInterface> Material;
 
 	/** Whether to apply atmosphere transmittance per sample, instead of using the light global transmittance. */
 	UPROPERTY(EditAnywhere, BlueprintReadOnly, Category = "Cloud Tracing")
@@ -178,7 +178,7 @@ class AVolumetricCloud : public AInfo
 private:
 
 	UPROPERTY(BlueprintReadOnly, VisibleAnywhere, Category = Atmosphere, meta = (AllowPrivateAccess = "true"))
-	class UVolumetricCloudComponent* VolumetricCloudComponent;
+	TObjectPtr<class UVolumetricCloudComponent> VolumetricCloudComponent;
 
 public:
 
diff --git a/Engine/Source/Runtime/Engine/Classes/Curves/CurveBase.h b/Engine/Source/Runtime/Engine/Classes/Curves/CurveBase.h
--- a/Engine/Source/Runtime/Engine/Classes/Curves/CurveBase.h
+++ b/Engine/Source/Runtime/Engine/Classes/Curves/CurveBase.h
@@ -82,7 +82,7 @@ public:
 	virtual void PostEditChangeProperty(struct FPropertyChangedEvent& PropertyChangedEvent) override;
 
 	UPROPERTY(VisibleAnywhere, Instanced, Category=ImportSettings)
-	class UAssetImportData* AssetImportData;
+	TObjectPtr<class UAssetImportData> AssetImportData;
 
 public:
 
diff --git a/Engine/Source/Runtime/Engine/Classes/Curves/CurveFloat.h b/Engine/Source/Runtime/Engine/Classes/Curves/CurveFloat.h
--- a/Engine/Source/Runtime/Engine/Classes/Curves/CurveFloat.h
+++ b/Engine/Source/Runtime/Engine/Classes/Curves/CurveFloat.h
@@ -17,7 +17,7 @@ struct ENGINE_API FRuntimeFloatCurve
 	FRichCurve EditorCurveData;
 
 	UPROPERTY(EditAnywhere,Category=RuntimeFloatCurve)
-	class UCurveFloat* ExternalCurve;
+	TObjectPtr<class UCurveFloat> ExternalCurve;
 
 	FRuntimeFloatCurve();
 
diff --git a/Engine/Source/Runtime/Engine/Classes/Curves/CurveLinearColor.h b/Engine/Source/Runtime/Engine/Classes/Curves/CurveLinearColor.h
--- a/Engine/Source/Runtime/Engine/Classes/Curves/CurveLinearColor.h
+++ b/Engine/Source/Runtime/Engine/Classes/Curves/CurveLinearColor.h
@@ -21,7 +21,7 @@ public:
 	FRichCurve ColorCurves[4];
 
 	UPROPERTY(EditAnywhere, Category = RuntimeFloatCurve)
-	class UCurveLinearColor* ExternalCurve = nullptr;
+	TObjectPtr<class UCurveLinearColor> ExternalCurve = nullptr;
 
 	FLinearColor GetLinearColorValue(float InTime) const;
 };
diff --git a/Engine/Source/Runtime/Engine/Classes/Curves/CurveLinearColorAtlas.h b/Engine/Source/Runtime/Engine/Classes/Curves/CurveLinearColorAtlas.h
--- a/Engine/Source/Runtime/Engine/Classes/Curves/CurveLinearColorAtlas.h
+++ b/Engine/Source/Runtime/Engine/Classes/Curves/CurveLinearColorAtlas.h
@@ -118,7 +118,7 @@ class ENGINE_API UCurveLinearColorAtlas : public UTexture2D
 	uint32	TextureHeight;						// Height of the lookup textures
 
 	UPROPERTY(EditAnywhere, Category = "Curves")
-	TArray<UCurveLinearColor*> GradientCurves;
+	TArray<TObjectPtr<UCurveLinearColor>> GradientCurves;
 
 #if WITH_EDITORONLY_DATA
 	/** Disable all color adjustments to preserve negative values in curves. Color adjustments clamp to 0 when enabled. */
diff --git a/Engine/Source/Runtime/Engine/Classes/DeviceProfiles/DeviceProfile.h b/Engine/Source/Runtime/Engine/Classes/DeviceProfiles/DeviceProfile.h
--- a/Engine/Source/Runtime/Engine/Classes/DeviceProfiles/DeviceProfile.h
+++ b/Engine/Source/Runtime/Engine/Classes/DeviceProfiles/DeviceProfile.h
@@ -31,7 +31,7 @@ class ENGINE_API UDeviceProfile : public UTextureLODSettings
 
 	/** The parent object of this profile, it is the object matching this DeviceType with the BaseProfileName */
 	UPROPERTY()
-	UObject* Parent;
+	TObjectPtr<UObject> Parent;
 
 	/** Flag used in the editor to determine whether the profile is visible in the property matrix */
 	bool bVisible;
diff --git a/Engine/Source/Runtime/Engine/Classes/DeviceProfiles/DeviceProfileManager.h b/Engine/Source/Runtime/Engine/Classes/DeviceProfiles/DeviceProfileManager.h
--- a/Engine/Source/Runtime/Engine/Classes/DeviceProfiles/DeviceProfileManager.h
+++ b/Engine/Source/Runtime/Engine/Classes/DeviceProfiles/DeviceProfileManager.h
@@ -202,7 +202,7 @@ public:
 
 	// Holds the collection of managed profiles.
 	UPROPERTY( EditAnywhere, Category=Properties )
-	TArray< UObject* > Profiles;
+	TArray< TObjectPtr<UObject> > Profiles;
 
 private:
 
diff --git a/Engine/Source/Runtime/Engine/Classes/Distributions/DistributionFloat.h b/Engine/Source/Runtime/Engine/Classes/Distributions/DistributionFloat.h
--- a/Engine/Source/Runtime/Engine/Classes/Distributions/DistributionFloat.h
+++ b/Engine/Source/Runtime/Engine/Classes/Distributions/DistributionFloat.h
@@ -37,7 +37,7 @@ private:
 
 public:
 	UPROPERTY(EditAnywhere, export, noclear, Category=RawDistributionFloat)
-	class UDistributionFloat* Distribution;
+	TObjectPtr<class UDistributionFloat> Distribution;
 
 
 	FRawDistributionFloat()
diff --git a/Engine/Source/Runtime/Engine/Classes/Distributions/DistributionVector.h b/Engine/Source/Runtime/Engine/Classes/Distributions/DistributionVector.h
--- a/Engine/Source/Runtime/Engine/Classes/Distributions/DistributionVector.h
+++ b/Engine/Source/Runtime/Engine/Classes/Distributions/DistributionVector.h
@@ -74,7 +74,7 @@ private:
 
 public:
 	UPROPERTY(EditAnywhere, export, noclear, Category=RawDistributionVector)
-	class UDistributionVector* Distribution;
+	TObjectPtr<class UDistributionVector> Distribution;
 
 	/** Whether the distribution data has been cooked or the object itself is available */
 	bool IsCreated();
diff --git a/Engine/Source/Runtime/Engine/Classes/EdGraph/EdGraph.h b/Engine/Source/Runtime/Engine/Classes/EdGraph/EdGraph.h
--- a/Engine/Source/Runtime/Engine/Classes/EdGraph/EdGraph.h
+++ b/Engine/Source/Runtime/Engine/Classes/EdGraph/EdGraph.h
@@ -24,11 +24,11 @@ struct ENGINE_API FGraphReference
 protected:
 	// Reference to the actual graph
 	UPROPERTY()
-	mutable class UEdGraph* MacroGraph;
+	mutable TObjectPtr<class UEdGraph> MacroGraph;
 
 	// The blueprint the graph is contained within
 	UPROPERTY()
-	class UBlueprint* GraphBlueprint;
+	TObjectPtr<class UBlueprint> GraphBlueprint;
 
 	// The graph GUID so we can refind it if it has been renamed
 	UPROPERTY()
@@ -76,7 +76,7 @@ public:
 
 	/** Set of all nodes in this graph */
 	UPROPERTY()
-	TArray<class UEdGraphNode*> Nodes;
+	TArray<TObjectPtr<class UEdGraphNode>> Nodes;
 
 	/** If true, graph can be edited by the user */
 	UPROPERTY()
@@ -97,7 +97,7 @@ public:
 #if WITH_EDITORONLY_DATA
 	/** Child graphs that are a part of this graph; the separation is purely visual */
 	UPROPERTY()
-	TArray<class UEdGraph*> SubGraphs;
+	TArray<TObjectPtr<class UEdGraph>> SubGraphs;
 
 	/** Guid for this graph */
 	UPROPERTY()
diff --git a/Engine/Source/Runtime/Engine/Classes/EdGraph/EdGraphNode.h b/Engine/Source/Runtime/Engine/Classes/EdGraph/EdGraphNode.h
--- a/Engine/Source/Runtime/Engine/Classes/EdGraph/EdGraphNode.h
+++ b/Engine/Source/Runtime/Engine/Classes/EdGraph/EdGraphNode.h
@@ -207,15 +207,15 @@ public:
 
 	/** The blueprint associated with this context; may be NULL for non-Kismet related graphs. */
 	UPROPERTY()
-	const UBlueprint* Blueprint;
+	TObjectPtr<const UBlueprint> Blueprint;
 
 	/** The graph associated with this context. */
 	UPROPERTY()
-	const UEdGraph* Graph;
+	TObjectPtr<const UEdGraph> Graph;
 
 	/** The node associated with this context. */
 	UPROPERTY()
-	const UEdGraphNode* Node;
+	TObjectPtr<const UEdGraphNode> Node;
 
 	/** The pin associated with this context; may be NULL when over a node. */
 	const UEdGraphPin* Pin;
@@ -265,7 +265,7 @@ public:
 
 	/** List of connector pins */
 	UPROPERTY()
-	TArray<class UEdGraphPin_Deprecated*> DeprecatedPins;
+	TArray<TObjectPtr<class UEdGraphPin_Deprecated>> DeprecatedPins;
 
 	/** X position of node in the editor */
 	UPROPERTY()
diff --git a/Engine/Source/Runtime/Engine/Classes/EdGraph/EdGraphPin.h b/Engine/Source/Runtime/Engine/Classes/EdGraph/EdGraphPin.h
--- a/Engine/Source/Runtime/Engine/Classes/EdGraph/EdGraphPin.h
+++ b/Engine/Source/Runtime/Engine/Classes/EdGraph/EdGraphPin.h
@@ -31,7 +31,7 @@ struct FSimpleMemberReference
 	 * if it is a native delegate signature function (declared globally).
 	 */
 	UPROPERTY()
-	UObject* MemberParent;
+	TObjectPtr<UObject> MemberParent;
 
 	/** Name of the member */
 	UPROPERTY()
@@ -669,7 +669,7 @@ class UEdGraphPin_Deprecated : public UObject
 
 	/** If the default value for this pin should be an object, we store a pointer to it */
 	UPROPERTY()
-	class UObject* DefaultObject;
+	TObjectPtr<class UObject> DefaultObject;
 
 	/** If the default value for this pin should be an FText, it is stored here. */
 	UPROPERTY()
@@ -677,19 +677,19 @@ class UEdGraphPin_Deprecated : public UObject
 
 	/** Set of pins that we are linked to */
 	UPROPERTY()
-	TArray<class UEdGraphPin_Deprecated*> LinkedTo;
+	TArray<TObjectPtr<class UEdGraphPin_Deprecated>> LinkedTo;
 
 	/** The pins created when a pin is split and hidden */ 
 	UPROPERTY()
-	TArray<class UEdGraphPin_Deprecated*> SubPins;
+	TArray<TObjectPtr<class UEdGraphPin_Deprecated>> SubPins;
 
 	/** The pin that was split and generated this pin */
 	UPROPERTY()
-	UEdGraphPin_Deprecated* ParentPin;
+	TObjectPtr<UEdGraphPin_Deprecated> ParentPin;
 
 	/** Pin that this pin uses for passing through reference connection */
 	UPROPERTY()
-	UEdGraphPin_Deprecated* ReferencePassThroughConnection;
+	TObjectPtr<UEdGraphPin_Deprecated> ReferencePassThroughConnection;
 
 #if WITH_EDITORONLY_DATA
 	/** If true, this connector is currently hidden. */
diff --git a/Engine/Source/Runtime/Engine/Classes/EdGraph/EdGraphSchema.h b/Engine/Source/Runtime/Engine/Classes/EdGraph/EdGraphSchema.h
--- a/Engine/Source/Runtime/Engine/Classes/EdGraph/EdGraphSchema.h
+++ b/Engine/Source/Runtime/Engine/Classes/EdGraph/EdGraphSchema.h
@@ -304,7 +304,7 @@ struct ENGINE_API FEdGraphSchemaAction_NewNode : public FEdGraphSchemaAction
 
 	/** Template of node we want to create */
 	UPROPERTY()
-	class UEdGraphNode* NodeTemplate;
+	TObjectPtr<class UEdGraphNode> NodeTemplate;
 
 
 	FEdGraphSchemaAction_NewNode() 
diff --git a/Engine/Source/Runtime/Engine/Classes/Engine/ActorChannel.h b/Engine/Source/Runtime/Engine/Classes/Engine/ActorChannel.h
--- a/Engine/Source/Runtime/Engine/Classes/Engine/ActorChannel.h
+++ b/Engine/Source/Runtime/Engine/Classes/Engine/ActorChannel.h
@@ -65,7 +65,7 @@ public:
 
 	// Variables.
 	UPROPERTY()
-	AActor* Actor;					// Actor this corresponds to.
+	TObjectPtr<AActor> Actor;					// Actor this corresponds to.
 
 	FNetworkGUID	ActorNetGUID;		// Actor GUID (useful when we don't have the actor resolved yet). Currently only valid on clients.
 	float			CustomTimeDilation;
@@ -109,7 +109,7 @@ public:
 	TSet<FNetworkGUID> PendingGuidResolves;	// These guids are waiting for their resolves, we need to queue up bunches until these are resolved
 
 	UPROPERTY()
-	TArray< UObject* >					CreateSubObjects;		// Any sub-object we created on this channel
+	TArray< TObjectPtr<UObject> >					CreateSubObjects;		// Any sub-object we created on this channel
 
 	TArray< FNetworkGUID >				QueuedMustBeMappedGuidsInLastBunch;	// Array of guids that will async load on client. This list is used for queued RPC's.
 	TArray< class FOutBunch * >			QueuedExportBunches;				// Bunches that need to be appended to the export list on the next SendBunch call. This list is used for queued RPC's.
diff --git a/Engine/Source/Runtime/Engine/Classes/Engine/AssetManager.h b/Engine/Source/Runtime/Engine/Classes/Engine/AssetManager.h
--- a/Engine/Source/Runtime/Engine/Classes/Engine/AssetManager.h
+++ b/Engine/Source/Runtime/Engine/Classes/Engine/AssetManager.h
@@ -712,7 +712,7 @@ protected:
 
 	/** List of UObjects that are being kept from being GCd, derived from the asset type map. Arrays are currently more efficient than Sets */
 	UPROPERTY()
-	TArray<UObject*> ObjectReferenceList;
+	TArray<TObjectPtr<UObject>> ObjectReferenceList;
 
 	/** True if we are running a build that is already scanning assets globally so we can perhaps avoid scanning paths synchronously */
 	UPROPERTY()
diff --git a/Engine/Source/Runtime/Engine/Classes/Engine/AssetManagerTypes.h b/Engine/Source/Runtime/Engine/Classes/Engine/AssetManagerTypes.h
--- a/Engine/Source/Runtime/Engine/Classes/Engine/AssetManagerTypes.h
+++ b/Engine/Source/Runtime/Engine/Classes/Engine/AssetManagerTypes.h
@@ -114,7 +114,7 @@ public:
 
 	/** Base Class of all assets of this type */
 	UPROPERTY(Transient)
-	UClass* AssetBaseClassLoaded;
+	TObjectPtr<UClass> AssetBaseClassLoaded;
 
 	/** True if the assets loaded are blueprints classes, false if they are normal UObjects */
 	UPROPERTY(EditAnywhere, Category = AssetType)
@@ -206,7 +206,7 @@ struct FAssetManagerSearchRules
 
 	/** Assets must inherit from this class, for blueprints this should be the instance base class */
 	UPROPERTY(EditAnywhere, BlueprintReadWrite, Category = Default)
-	UClass* AssetBaseClass = nullptr;
+	TObjectPtr<UClass> AssetBaseClass = nullptr;
 
 	/** True if scanning for blueprints, false for all other assets */
 	UPROPERTY(EditAnywhere, BlueprintReadWrite, Category = Default)
diff --git a/Engine/Source/Runtime/Engine/Classes/Engine/AutoDestroySubsystem.h b/Engine/Source/Runtime/Engine/Classes/Engine/AutoDestroySubsystem.h
--- a/Engine/Source/Runtime/Engine/Classes/Engine/AutoDestroySubsystem.h
+++ b/Engine/Source/Runtime/Engine/Classes/Engine/AutoDestroySubsystem.h
@@ -77,5 +77,5 @@ private:
 
 	/** Actors to check if they should auto destroy or not */
 	UPROPERTY()
-	TArray<AActor*> ActorsToPoll;
+	TArray<TObjectPtr<AActor>> ActorsToPoll;
 };
\ No newline at end of file
diff --git a/Engine/Source/Runtime/Engine/Classes/Engine/Blueprint.h b/Engine/Source/Runtime/Engine/Classes/Engine/Blueprint.h
--- a/Engine/Source/Runtime/Engine/Classes/Engine/Blueprint.h
+++ b/Engine/Source/Runtime/Engine/Classes/Engine/Blueprint.h
@@ -292,7 +292,7 @@ struct FBPInterfaceDescription
 
 	/** References to the graphs associated with the required functions for this interface */
 	UPROPERTY()
-	TArray<UEdGraph*> Graphs;
+	TArray<TObjectPtr<UEdGraph>> Graphs;
 
 
 	FBPInterfaceDescription()
@@ -356,7 +356,7 @@ struct FEditedDocumentInfo
 private:
 	// Legacy hard reference is now serialized as a soft reference (see above).
 	UPROPERTY()
-	UObject* EditedObject_DEPRECATED;
+	TObjectPtr<UObject> EditedObject_DEPRECATED;
 };
 
 template<>
@@ -550,45 +550,45 @@ PRAGMA_ENABLE_DEPRECATION_WARNINGS
 
 	/** 'Simple' construction script - graph of components to instance */
 	UPROPERTY()
-	class USimpleConstructionScript* SimpleConstructionScript;
+	TObjectPtr<class USimpleConstructionScript> SimpleConstructionScript;
 
 #if WITH_EDITORONLY_DATA
 	/** Set of pages that combine into a single uber-graph */
 	UPROPERTY()
-	TArray<UEdGraph*> UbergraphPages;
+	TArray<TObjectPtr<UEdGraph>> UbergraphPages;
 
 	/** Set of functions implemented for this class graphically */
 	UPROPERTY()
-	TArray<UEdGraph*> FunctionGraphs;
+	TArray<TObjectPtr<UEdGraph>> FunctionGraphs;
 
 	/** Graphs of signatures for delegates */
 	UPROPERTY()
-	TArray<UEdGraph*> DelegateSignatureGraphs;
+	TArray<TObjectPtr<UEdGraph>> DelegateSignatureGraphs;
 
 	/** Set of macros implemented for this class */
 	UPROPERTY()
-	TArray<UEdGraph*> MacroGraphs;
+	TArray<TObjectPtr<UEdGraph>> MacroGraphs;
 
 	/** Set of functions actually compiled for this class */
 	UPROPERTY(transient, duplicatetransient)
-	TArray<UEdGraph*> IntermediateGeneratedGraphs;
+	TArray<TObjectPtr<UEdGraph>> IntermediateGeneratedGraphs;
 
 	/** Set of functions actually compiled for this class */
 	UPROPERTY(transient, duplicatetransient)
-	TArray<UEdGraph*> EventGraphs;
+	TArray<TObjectPtr<UEdGraph>> EventGraphs;
 
 	/** Cached cosmetic information about macro graphs, use GetCosmeticInfoForMacro() to access */
 	UPROPERTY(Transient)
-	TMap<UEdGraph*, FBlueprintMacroCosmeticInfo> PRIVATE_CachedMacroInfo;
+	TMap<TObjectPtr<UEdGraph>, FBlueprintMacroCosmeticInfo> PRIVATE_CachedMacroInfo;
 #endif // WITH_EDITORONLY_DATA
 
 	/** Array of component template objects, used by AddComponent function */
 	UPROPERTY()
-	TArray<class UActorComponent*> ComponentTemplates;
+	TArray<TObjectPtr<class UActorComponent>> ComponentTemplates;
 
 	/** Array of templates for timelines that should be created */
 	UPROPERTY()
-	TArray<class UTimelineTemplate*> Timelines;
+	TArray<TObjectPtr<class UTimelineTemplate>> Timelines;
 
 	/** Array of blueprint overrides of component classes in parent classes */
 	UPROPERTY()
@@ -596,7 +596,7 @@ PRAGMA_ENABLE_DEPRECATION_WARNINGS
 
 	/** Stores data to override (in children classes) components (created by SCS) from parent classes */
 	UPROPERTY()
-	class UInheritableComponentHandler* InheritableComponentHandler;
+	TObjectPtr<class UInheritableComponentHandler> InheritableComponentHandler;
 
 #if WITH_EDITORONLY_DATA
 	/** Array of new variables to be added to generated class */
@@ -625,13 +625,13 @@ PRAGMA_ENABLE_DEPRECATION_WARNINGS
 
 	/** Persistent debugging options */
 	UPROPERTY()
-	TArray<class UBreakpoint*> Breakpoints;
+	TArray<TObjectPtr<class UBreakpoint>> Breakpoints;
 
 	UPROPERTY()
 	TArray<FEdGraphPinReference> WatchedPins;
 
 	UPROPERTY()
-	TArray<class UEdGraphPin_Deprecated*> DeprecatedPinWatches;
+	TArray<TObjectPtr<class UEdGraphPin_Deprecated>> DeprecatedPinWatches;
 
 	/** Index map for component template names */
 	UPROPERTY()
@@ -643,7 +643,7 @@ PRAGMA_ENABLE_DEPRECATION_WARNINGS
 
 	/** Array of extensions for this blueprint */
 	UPROPERTY()
-	TArray<UBlueprintExtension*> Extensions;
+	TArray<TObjectPtr<UBlueprintExtension>> Extensions;
 
 #endif // WITH_EDITORONLY_DATA
 
@@ -688,7 +688,7 @@ public:
 
 	/** Information for thumbnail rendering */
 	UPROPERTY(VisibleAnywhere, Instanced, Category=Thumbnail)
-	class UThumbnailInfo* ThumbnailInfo;
+	TObjectPtr<class UThumbnailInfo> ThumbnailInfo;
 
 	/** CRC for CDO calculated right after the latest compilation used by Reinstancer to check if default values were changed */
 	UPROPERTY(transient, duplicatetransient)
@@ -726,7 +726,7 @@ public:
 
 	// If this BP is just a duplicate created for a specific compilation, the reference to original GeneratedClass is needed
 	UPROPERTY(transient, duplicatetransient)
-	UClass* OriginalClass;
+	TObjectPtr<UClass> OriginalClass;
 
 	bool IsUpToDate() const
 	{
diff --git a/Engine/Source/Runtime/Engine/Classes/Engine/BlueprintGeneratedClass.h b/Engine/Source/Runtime/Engine/Classes/Engine/BlueprintGeneratedClass.h
--- a/Engine/Source/Runtime/Engine/Classes/Engine/BlueprintGeneratedClass.h
+++ b/Engine/Source/Runtime/Engine/Classes/Engine/BlueprintGeneratedClass.h
@@ -505,7 +505,7 @@ struct ENGINE_API FEventGraphFastCallPair
 	GENERATED_USTRUCT_BODY()
 
 	UPROPERTY()
-	UFunction* FunctionToPatch = nullptr;
+	TObjectPtr<UFunction> FunctionToPatch = nullptr;
 
 	UPROPERTY()
 	int32 EventGraphCallOffset = 0;
@@ -527,7 +527,7 @@ struct ENGINE_API FBlueprintComponentChangedPropertyInfo
 
 	/** The parent struct (owner) of the changed property. */
 	UPROPERTY()
-	UStruct* PropertyScope;
+	TObjectPtr<UStruct> PropertyScope;
 
 	/** Default constructor. */
 	FBlueprintComponentChangedPropertyInfo()
@@ -607,7 +607,7 @@ struct FBPComponentClassOverride
 
 	/** The class to use when constructing the component. */
 	UPROPERTY()
-	UClass* ComponentClass;
+	TObjectPtr<UClass> ComponentClass;
 
 	FBPComponentClassOverride()
 		: ComponentClass(nullptr)
@@ -657,15 +657,15 @@ private:
 public:
 	/** Array of objects containing information for dynamically binding delegates to functions in this blueprint */
 	UPROPERTY()
-	TArray<class UDynamicBlueprintBinding*> DynamicBindingObjects;
+	TArray<TObjectPtr<class UDynamicBlueprintBinding>> DynamicBindingObjects;
 
 	/** Array of component template objects, used by AddComponent function */
 	UPROPERTY()
-	TArray<class UActorComponent*> ComponentTemplates;
+	TArray<TObjectPtr<class UActorComponent>> ComponentTemplates;
 
 	/** Array of templates for timelines that should be created */
 	UPROPERTY()
-	TArray<class UTimelineTemplate*> Timelines;
+	TArray<TObjectPtr<class UTimelineTemplate>> Timelines;
 
 	/** Array of blueprint overrides of component classes in parent classes */
 	UPROPERTY()
@@ -673,19 +673,19 @@ public:
 
 	/** 'Simple' construction script - graph of components to instance */
 	UPROPERTY()
-	class USimpleConstructionScript* SimpleConstructionScript;
+	TObjectPtr<class USimpleConstructionScript> SimpleConstructionScript;
 
 	/** Stores data to override (in children classes) components (created by SCS) from parent classes */
 	UPROPERTY()
-	class UInheritableComponentHandler* InheritableComponentHandler;
+	TObjectPtr<class UInheritableComponentHandler> InheritableComponentHandler;
 
 	UPROPERTY()
-	class UStructProperty* UberGraphFramePointerProperty_DEPRECATED;
+	TObjectPtr<class UStructProperty> UberGraphFramePointerProperty_DEPRECATED;
 	
 	FStructProperty* UberGraphFramePointerProperty;
 
 	UPROPERTY()
-	UFunction* UberGraphFunction;
+	TObjectPtr<UFunction> UberGraphFunction;
 
 #if VALIDATE_UBER_GRAPH_PERSISTENT_FRAME
 	uint32 UberGraphFunctionKey;
@@ -700,14 +700,14 @@ public:
 
 #if WITH_EDITORONLY_DATA
 	UPROPERTY(Transient)
-	UObject* OverridenArchetypeForCDO;
+	TObjectPtr<UObject> OverridenArchetypeForCDO;
 
 	/** Property guid map */
 	UPROPERTY()
 	TMap<FName,FGuid> PropertyGuids;
 
 	UPROPERTY(Transient)
-	TArray<UFunction*> CalledFunctions;
+	TArray<TObjectPtr<UFunction>> CalledFunctions;
 #endif //WITH_EDITORONLY_DATA
 
 	// Mapping of changed properties & data to apply when instancing components in a cooked build (one entry per named AddComponent node template for fast lookup at runtime).
diff --git a/Engine/Source/Runtime/Engine/Classes/Engine/Breakpoint.h b/Engine/Source/Runtime/Engine/Classes/Engine/Breakpoint.h
--- a/Engine/Source/Runtime/Engine/Classes/Engine/Breakpoint.h
+++ b/Engine/Source/Runtime/Engine/Classes/Engine/Breakpoint.h
@@ -22,7 +22,7 @@ private:
 
 	// Node that the breakpoint is placed on
 	UPROPERTY()
-	class UEdGraphNode* Node;
+	TObjectPtr<class UEdGraphNode> Node;
 
 	// Is this breakpoint auto-generated, and should be removed when next hit?
 	UPROPERTY()
diff --git a/Engine/Source/Runtime/Engine/Classes/Engine/Brush.h b/Engine/Source/Runtime/Engine/Classes/Engine/Brush.h
--- a/Engine/Source/Runtime/Engine/Classes/Engine/Brush.h
+++ b/Engine/Source/Runtime/Engine/Classes/Engine/Brush.h
@@ -103,16 +103,16 @@ class ENGINE_API ABrush
 	uint32 bNotForClientOrServer:1;
 
 	UPROPERTY(Instanced)
-	class UModel* Brush;
+	TObjectPtr<class UModel> Brush;
 
 private:
 	UPROPERTY(Category = Collision, VisibleAnywhere, BlueprintReadOnly, meta = (AllowPrivateAccess = "true"))
-	class UBrushComponent* BrushComponent;
+	TObjectPtr<class UBrushComponent> BrushComponent;
 public:
 
 #if WITH_EDITORONLY_DATA
 	UPROPERTY(VisibleAnywhere, Instanced, Category=BrushBuilder)
-	class UBrushBuilder* BrushBuilder;
+	TObjectPtr<class UBrushBuilder> BrushBuilder;
 #endif
 
 	/** Flag set when we are in a manipulation (scaling, translation, brush builder param change etc.) */
diff --git a/Engine/Source/Runtime/Engine/Classes/Engine/Canvas.h b/Engine/Source/Runtime/Engine/Classes/Engine/Canvas.h
--- a/Engine/Source/Runtime/Engine/Classes/Engine/Canvas.h
+++ b/Engine/Source/Runtime/Engine/Classes/Engine/Canvas.h
@@ -24,7 +24,7 @@ struct FCanvasIcon
 
 	/** Source texture */
 	UPROPERTY(EditAnywhere, BlueprintReadWrite, Category=CanvasIcon)
-	class UTexture* Texture;
+	TObjectPtr<class UTexture> Texture;
 
 	/** UV coords */
 	UPROPERTY(EditAnywhere, BlueprintReadWrite, Category=CanvasIcon)
@@ -192,14 +192,14 @@ class ENGINE_API UCanvas
 	FPlane ColorModulate; 
 
 	UPROPERTY()
-	class UTexture2D* DefaultTexture; //Default texture to use 
+	TObjectPtr<class UTexture2D> DefaultTexture; //Default texture to use 
 
 	UPROPERTY()
-	class UTexture2D* GradientTexture0; //Default texture to use 
+	TObjectPtr<class UTexture2D> GradientTexture0; //Default texture to use 
 
 	/** Helper class to render 2d graphs on canvas */
 	UPROPERTY()
-	class UReporterGraph* ReporterGraph;
+	TObjectPtr<class UReporterGraph> ReporterGraph;
 
 	int32 UnsafeSizeX;   // Canvas size before safe frame adjustment
 	int32 UnsafeSizeY;	// Canvas size before safe frame adjustment
diff --git a/Engine/Source/Runtime/Engine/Classes/Engine/Channel.h b/Engine/Source/Runtime/Engine/Classes/Engine/Channel.h
--- a/Engine/Source/Runtime/Engine/Classes/Engine/Channel.h
+++ b/Engine/Source/Runtime/Engine/Classes/Engine/Channel.h
@@ -56,7 +56,7 @@ class ENGINE_API UChannel
 	GENERATED_UCLASS_BODY()
 	
 	UPROPERTY()
-	class UNetConnection*	Connection;		// Owner connection.
+	TObjectPtr<class UNetConnection>	Connection;		// Owner connection.
 
 	// Variables.
 	uint32				OpenAcked:1;		// If OpenedLocally is true, this means we have acknowledged the packet we sent the bOpen bunch on. Otherwise, it means we have received the bOpen bunch from the server.
diff --git a/Engine/Source/Runtime/Engine/Classes/Engine/ChildConnection.h b/Engine/Source/Runtime/Engine/Classes/Engine/ChildConnection.h
--- a/Engine/Source/Runtime/Engine/Classes/Engine/ChildConnection.h
+++ b/Engine/Source/Runtime/Engine/Classes/Engine/ChildConnection.h
@@ -17,7 +17,7 @@ class UChildConnection
 	GENERATED_UCLASS_BODY()
 
 	UPROPERTY(transient)
-	class UNetConnection* Parent;
+	TObjectPtr<class UNetConnection> Parent;
 
 public:
 
diff --git a/Engine/Source/Runtime/Engine/Classes/Engine/CompositeCurveTable.h b/Engine/Source/Runtime/Engine/Classes/Engine/CompositeCurveTable.h
--- a/Engine/Source/Runtime/Engine/Classes/Engine/CompositeCurveTable.h
+++ b/Engine/Source/Runtime/Engine/Classes/Engine/CompositeCurveTable.h
@@ -45,11 +45,11 @@ protected:
 	// Parent tables
 	// Tables with higher indices override data in tables with lower indices
 	UPROPERTY(EditAnywhere, BlueprintReadOnly, Category = Tables)
-	TArray<UCurveTable*> ParentTables;
+	TArray<TObjectPtr<UCurveTable>> ParentTables;
 
 	// temporary copy used to detect changes so we can update delegates correctly on removal
 	UPROPERTY(transient)
-	TArray<UCurveTable*> OldParentTables;
+	TArray<TObjectPtr<UCurveTable>> OldParentTables;
 
 	// true if this asset is currently being loaded; false otherwise
 	uint8 bIsLoading : 1;
diff --git a/Engine/Source/Runtime/Engine/Classes/Engine/CompositeDataTable.h b/Engine/Source/Runtime/Engine/Classes/Engine/CompositeDataTable.h
--- a/Engine/Source/Runtime/Engine/Classes/Engine/CompositeDataTable.h
+++ b/Engine/Source/Runtime/Engine/Classes/Engine/CompositeDataTable.h
@@ -74,11 +74,11 @@ protected:
 	// Parent tables
 	// Tables with higher indices override data in tables with lower indices
 	UPROPERTY(EditAnywhere, BlueprintReadOnly, Category = Tables)
-	TArray<UDataTable*> ParentTables;
+	TArray<TObjectPtr<UDataTable>> ParentTables;
 
 	// temporary copy used to detect changes so we can update delegates correctly on removal
 	UPROPERTY(transient)
-	TArray<UDataTable*> OldParentTables;
+	TArray<TObjectPtr<UDataTable>> OldParentTables;
 
 #if WITH_EDITORONLY_DATA
 	TMap<FName, ERowState> RowSourceMap;
diff --git a/Engine/Source/Runtime/Engine/Classes/Engine/Console.h b/Engine/Source/Runtime/Engine/Classes/Engine/Console.h
--- a/Engine/Source/Runtime/Engine/Classes/Engine/Console.h
+++ b/Engine/Source/Runtime/Engine/Classes/Engine/Console.h
@@ -67,13 +67,13 @@ class ENGINE_API UConsole
 
 	/** The player which the next console command should be executed in the context of.  If nullptr, execute in the viewport. */
 	UPROPERTY()
-	class ULocalPlayer* ConsoleTargetPlayer;
+	TObjectPtr<class ULocalPlayer> ConsoleTargetPlayer;
 
 	UPROPERTY()
-	class UTexture2D* DefaultTexture_Black;
+	TObjectPtr<class UTexture2D> DefaultTexture_Black;
 
 	UPROPERTY()
-	class UTexture2D* DefaultTexture_White;
+	TObjectPtr<class UTexture2D> DefaultTexture_White;
 
 	/** Holds the scrollback buffer */
 	TArray<FString> Scrollback;
diff --git a/Engine/Source/Runtime/Engine/Classes/Engine/CurveTable.h b/Engine/Source/Runtime/Engine/Classes/Engine/CurveTable.h
--- a/Engine/Source/Runtime/Engine/Classes/Engine/CurveTable.h
+++ b/Engine/Source/Runtime/Engine/Classes/Engine/CurveTable.h
@@ -77,7 +77,7 @@ public:
 	virtual void PostLoad() override;
 
 	UPROPERTY(VisibleAnywhere, Instanced, Category=ImportSettings)
-	class UAssetImportData* AssetImportData;
+	TObjectPtr<class UAssetImportData> AssetImportData;
 
 	/** The filename imported to create this object. Relative to this object's package, BaseDir() or absolute */
 	UPROPERTY()
@@ -252,7 +252,7 @@ struct ENGINE_API FCurveTableRowHandle
 
 	/** Pointer to table we want a row from */
 	UPROPERTY(EditAnywhere, BlueprintReadWrite, Category=CurveTableRowHandle, meta=(DisplayThumbnail="false"))
-	const class UCurveTable*	CurveTable;
+	TObjectPtr<const class UCurveTable>	CurveTable;
 
 	/** Name of row in the table that we want */
 	UPROPERTY(EditAnywhere, BlueprintReadWrite, Category=CurveTableRowHandle)
diff --git a/Engine/Source/Runtime/Engine/Classes/Engine/DataTable.h b/Engine/Source/Runtime/Engine/Classes/Engine/DataTable.h
--- a/Engine/Source/Runtime/Engine/Classes/Engine/DataTable.h
+++ b/Engine/Source/Runtime/Engine/Classes/Engine/DataTable.h
@@ -78,7 +78,7 @@ class UDataTable
 
 	/** Structure to use for each row of the table, must inherit from FTableRowBase */
 	UPROPERTY(VisibleAnywhere, Category=DataTable, meta=(DisplayThumbnail="false"))
-	UScriptStruct*			RowStruct;
+	TObjectPtr<UScriptStruct>			RowStruct;
 
 protected:
 	/** Map of name of row to row data structure. */
@@ -136,7 +136,7 @@ public:
 
 	/** The file this data table was imported from, may be empty */
 	UPROPERTY(VisibleAnywhere, Instanced, Category=ImportSource)
-	class UAssetImportData* AssetImportData;
+	TObjectPtr<class UAssetImportData> AssetImportData;
 
 	/** The filename imported to create this object. Relative to this object's package, BaseDir() or absolute */
 	UPROPERTY()
@@ -152,7 +152,7 @@ protected:
 	TArray<uint8> RowsSerializedWithTags;
 
 	UPROPERTY(Transient)
-	TSet<UObject*> TemporarilyReferencedObjects;
+	TSet<TObjectPtr<UObject>> TemporarilyReferencedObjects;
 #endif	// WITH_EDITORONLY_DATA
 
 private:
@@ -410,7 +410,7 @@ struct ENGINE_API FDataTableRowHandle
 
 	/** Pointer to table we want a row from */
 	UPROPERTY(EditAnywhere, BlueprintReadWrite, Category=DataTableRowHandle)
-	const UDataTable*	DataTable;
+	TObjectPtr<const UDataTable>	DataTable;
 
 	/** Name of row in the table that we want */
 	UPROPERTY(EditAnywhere, BlueprintReadWrite, Category=DataTableRowHandle)
@@ -476,7 +476,7 @@ struct ENGINE_API FDataTableCategoryHandle
 
 	/** Pointer to table we want a row from */
 	UPROPERTY(EditAnywhere, BlueprintReadWrite, Category=DataTableCategoryHandle)
-	const class UDataTable*	DataTable = nullptr;
+	TObjectPtr<const class UDataTable>	DataTable = nullptr;
 
 	/** Name of column in the table that we want */
 	UPROPERTY(EditAnywhere, BlueprintReadWrite, Category=DataTableCategoryHandle)
diff --git a/Engine/Source/Runtime/Engine/Classes/Engine/DebugCameraController.h b/Engine/Source/Runtime/Engine/Classes/Engine/DebugCameraController.h
--- a/Engine/Source/Runtime/Engine/Classes/Engine/DebugCameraController.h
+++ b/Engine/Source/Runtime/Engine/Classes/Engine/DebugCameraController.h
@@ -59,7 +59,7 @@ class ENGINE_API ADebugCameraController
 
 	/** Visualizes the frustum of the camera */
 	UPROPERTY()
-	class UDrawFrustumComponent* DrawFrustum;
+	TObjectPtr<class UDrawFrustumComponent> DrawFrustum;
 	
 	/** Sets whether to show information about the selected actor on the debug camera HUD.t */
 	UFUNCTION(exec)
@@ -177,7 +177,7 @@ public:
 
 	/** Currently selected actor, may be invalid */
 	UPROPERTY()
-	class AActor* SelectedActor;
+	TObjectPtr<class AActor> SelectedActor;
 
 	/** Returns the currently selected actor, or null if it is invalid or not set */
 	UFUNCTION(BlueprintCallable, Category="Debug Camera")
@@ -185,7 +185,7 @@ public:
 	
 	/** Currently selected component, may be invalid */
 	UPROPERTY()
-	class UPrimitiveComponent* SelectedComponent;
+	TObjectPtr<class UPrimitiveComponent> SelectedComponent;
 
 	/** Selected hit point */
 	UPROPERTY()
@@ -193,11 +193,11 @@ public:
 
 	/** Controller that was active before this was spawned */
 	UPROPERTY()	
-	class APlayerController* OriginalControllerRef;
+	TObjectPtr<class APlayerController> OriginalControllerRef;
 
 	/** Player object that was active before this was spawned */
 	UPROPERTY()	
-	class UPlayer* OriginalPlayer;
+	TObjectPtr<class UPlayer> OriginalPlayer;
 
 	/** Allows control over the speed of the spectator pawn. This scales the speed based on the InitialMaxSpeed. Use Set Pawn Movement Speed Scale during runtime */
 	UPROPERTY(EditDefaultsOnly, BlueprintReadOnly, Category="Debug Camera")
diff --git a/Engine/Source/Runtime/Engine/Classes/Engine/DebugDisplayProperty.h b/Engine/Source/Runtime/Engine/Classes/Engine/DebugDisplayProperty.h
--- a/Engine/Source/Runtime/Engine/Classes/Engine/DebugDisplayProperty.h
+++ b/Engine/Source/Runtime/Engine/Classes/Engine/DebugDisplayProperty.h
@@ -22,7 +22,7 @@ struct FDebugDisplayProperty
 
 	/** the object whose property to display. If this is a class, all objects of that class are drawn. */
 	UPROPERTY()
-	class UObject* Obj;
+	TObjectPtr<class UObject> Obj;
 
 	/** if Obj is a class and WithinClass is not nullptr, further limit the display to objects that have an Outer of WithinClass */
 	UPROPERTY()
diff --git a/Engine/Source/Runtime/Engine/Classes/Engine/DecalActor.h b/Engine/Source/Runtime/Engine/Classes/Engine/DecalActor.h
--- a/Engine/Source/Runtime/Engine/Classes/Engine/DecalActor.h
+++ b/Engine/Source/Runtime/Engine/Classes/Engine/DecalActor.h
@@ -27,20 +27,20 @@ class ENGINE_API ADecalActor
 private:
 	/** The decal component for this decal actor */
 	UPROPERTY(Category = Decal, VisibleAnywhere, BlueprintReadOnly, meta = (ExposeFunctionCategories = "Decal,Rendering|Components|Decal", AllowPrivateAccess = "true"))
-	UDecalComponent* Decal;
+	TObjectPtr<UDecalComponent> Decal;
 
 #if WITH_EDITORONLY_DATA
 	/* Reference to the editor only arrow visualization component */
 	UPROPERTY()
-	UArrowComponent* ArrowComponent;
+	TObjectPtr<UArrowComponent> ArrowComponent;
 
 	/* Reference to the billboard component */
 	UPROPERTY()
-	UBillboardComponent* SpriteComponent;
+	TObjectPtr<UBillboardComponent> SpriteComponent;
 
 	// formerly we used this component to draw a box, now we use the DecalComponentVisualizer
 	UPROPERTY()
-	UBoxComponent* BoxComponent_DEPRECATED;
+	TObjectPtr<UBoxComponent> BoxComponent_DEPRECATED;
 #endif
 
 public:
diff --git a/Engine/Source/Runtime/Engine/Classes/Engine/DemoNetDriver.h b/Engine/Source/Runtime/Engine/Classes/Engine/DemoNetDriver.h
--- a/Engine/Source/Runtime/Engine/Classes/Engine/DemoNetDriver.h
+++ b/Engine/Source/Runtime/Engine/Classes/Engine/DemoNetDriver.h
@@ -64,18 +64,18 @@ struct FRollbackNetStartupActorInfo
 
 	FName		Name;
 	UPROPERTY()
-	UObject*	Archetype = nullptr;
+	TObjectPtr<UObject>	Archetype = nullptr;
 	FVector		Location;
 	FRotator	Rotation;
 	FVector		Scale3D;
 	UPROPERTY()
-	ULevel*		Level = nullptr;
+	TObjectPtr<ULevel>		Level = nullptr;
 
 	TSharedPtr<FRepState> RepState;
 	TMap<FString, TSharedPtr<FRepState>> SubObjRepState;
 
 	UPROPERTY()
-	TArray<UObject*> ObjReferences;
+	TArray<TObjectPtr<UObject>> ObjReferences;
 
 	void CountBytes(FArchive& Ar) const
 	{
@@ -592,7 +592,7 @@ private:
 
 	/** An array of all the spectator controllers (the main one and all splitscreen ones) that currently exist */
 	UPROPERTY(transient)
-	TArray<APlayerController*> SpectatorControllers;
+	TArray<TObjectPtr<APlayerController>> SpectatorControllers;
 
 	/**
 	 * Removes all child connections for splitscreen viewers.
diff --git a/Engine/Source/Runtime/Engine/Classes/Engine/DirectionalLight.h b/Engine/Source/Runtime/Engine/Classes/Engine/DirectionalLight.h
--- a/Engine/Source/Runtime/Engine/Classes/Engine/DirectionalLight.h
+++ b/Engine/Source/Runtime/Engine/Classes/Engine/DirectionalLight.h
@@ -23,11 +23,11 @@ class ADirectionalLight
 	// Reference to editor visualization arrow
 private:
 	UPROPERTY()
-	UArrowComponent* ArrowComponent;
+	TObjectPtr<UArrowComponent> ArrowComponent;
 
 	/* EditorOnly reference to the light component to allow it to be displayed in the details panel correctly */
 	UPROPERTY(VisibleAnywhere, Category="Light")
-	UDirectionalLightComponent* DirectionalLightComponent;
+	TObjectPtr<UDirectionalLightComponent> DirectionalLightComponent;
 #endif
 
 public:
diff --git a/Engine/Source/Runtime/Engine/Classes/Engine/DocumentationActor.h b/Engine/Source/Runtime/Engine/Classes/Engine/DocumentationActor.h
--- a/Engine/Source/Runtime/Engine/Classes/Engine/DocumentationActor.h
+++ b/Engine/Source/Runtime/Engine/Classes/Engine/DocumentationActor.h
@@ -50,7 +50,7 @@ class ENGINE_API ADocumentationActor
 	
 private:
 	UPROPERTY(Category = Sprite, VisibleAnywhere, BlueprintReadOnly, meta = (AllowPrivateAccess = "true"))
-	class UMaterialBillboardComponent* Billboard;
+	TObjectPtr<class UMaterialBillboardComponent> Billboard;
 #endif
 
 public:
