Automated wrapped object pointer upgrade for Engine + ShooterGame + ShooterGame referenced plugins

diff --git a/Engine/Source/Runtime/Engine/Classes/Engine/Engine.h b/Engine/Source/Runtime/Engine/Classes/Engine/Engine.h
--- a/Engine/Source/Runtime/Engine/Classes/Engine/Engine.h
+++ b/Engine/Source/Runtime/Engine/Classes/Engine/Engine.h
@@ -157,7 +157,7 @@ struct FFullyLoadedPackagesInfo
 
 	/** List of objects that were loaded, for faster cleanup */
 	UPROPERTY()
-	TArray<class UObject*> LoadedObjects;
+	TArray<TObjectPtr<class UObject>> LoadedObjects;
 
 
 	FFullyLoadedPackagesInfo()
@@ -242,7 +242,7 @@ struct FNamedNetDriver
 
 	/** Instantiation of named net driver */
 	UPROPERTY(transient)
-	class UNetDriver* NetDriver;
+	TObjectPtr<class UNetDriver> NetDriver;
 
 	/** Definition associated with this net driver */
 	FNetDriverDefinition* NetDriverDef;
@@ -311,7 +311,7 @@ struct FWorldContext
 	struct FURL LastRemoteURL;
 
 	UPROPERTY()
-	UPendingNetGame * PendingNetGame;
+	TObjectPtr<UPendingNetGame>  PendingNetGame;
 
 	/** A list of tag/array pairs that is used at LoadMap time to fully load packages that may be needed for the map/game with DLC, but we can't use DynamicLoadObject to load from the packages */
 	UPROPERTY()
@@ -325,7 +325,7 @@ struct FWorldContext
 
 	/** Array of already loaded levels. The ordering is arbitrary and depends on what is already loaded and such.	*/
 	UPROPERTY()
-	TArray<class ULevel*> LoadedLevelsForPendingMapChange;
+	TArray<TObjectPtr<class ULevel>> LoadedLevelsForPendingMapChange;
 
 	/** Human readable error string for any failure during a map change request. Empty if there were no failures.	*/
 	FString PendingMapChangeFailureDescription;
@@ -335,16 +335,16 @@ struct FWorldContext
 
 	/** Handles to object references; used by the engine to e.g. the prevent objects from being garbage collected.	*/
 	UPROPERTY()
-	TArray<class UObjectReferencer*> ObjectReferencers;
+	TArray<TObjectPtr<class UObjectReferencer>> ObjectReferencers;
 
 	UPROPERTY()
 	TArray<struct FLevelStreamingStatus> PendingLevelStreamingStatusUpdates;
 
 	UPROPERTY()
-	class UGameViewportClient* GameViewport;
+	TObjectPtr<class UGameViewportClient> GameViewport;
 
 	UPROPERTY()
-	class UGameInstance* OwningGameInstance;
+	TObjectPtr<class UGameInstance> OwningGameInstance;
 
 	/** A list of active net drivers */
 	UPROPERTY(transient)
@@ -643,7 +643,7 @@ class ENGINE_API UEngine
 
 private:
 	UPROPERTY()
-	class UFont* TinyFont;
+	TObjectPtr<class UFont> TinyFont;
 
 public:
 	/** Sets the font used for the smallest engine text */
@@ -652,7 +652,7 @@ public:
 
 private:
 	UPROPERTY()
-	class UFont* SmallFont;
+	TObjectPtr<class UFont> SmallFont;
 
 public:
 	/** Sets the font used for small engine text, used for most debug displays */
@@ -661,7 +661,7 @@ public:
 
 private:
 	UPROPERTY()
-	class UFont* MediumFont;
+	TObjectPtr<class UFont> MediumFont;
 
 public:
 	/** Sets the font used for medium engine text */
@@ -670,7 +670,7 @@ public:
 
 private:
 	UPROPERTY()
-	class UFont* LargeFont;
+	TObjectPtr<class UFont> LargeFont;
 
 public:
 	/** Sets the font used for large engine text */
@@ -679,7 +679,7 @@ public:
 
 private:
 	UPROPERTY()
-	class UFont* SubtitleFont;
+	TObjectPtr<class UFont> SubtitleFont;
 
 public:
 	/** Sets the font used by the default Subtitle Manager */
@@ -688,7 +688,7 @@ public:
 
 private:
 	UPROPERTY()
-	TArray<class UFont*> AdditionalFonts;
+	TArray<TObjectPtr<class UFont>> AdditionalFonts;
 
 public:
 	/** Sets additional fonts that will be loaded at startup and available using GetAdditionalFont. */
@@ -764,7 +764,7 @@ public:
 
 	/** Global instance of the user game settings */
 	UPROPERTY()
-	class UGameUserSettings* GameUserSettings;
+	TObjectPtr<class UGameUserSettings> GameUserSettings;
 
 	UPROPERTY()
 	TSubclassOf<class ALevelScriptActor>  LevelScriptActorClass;
@@ -783,7 +783,7 @@ public:
 
 	/** A UObject spawned at initialization time to handle game-specific data */
 	UPROPERTY()
-	UObject *GameSingleton;
+	TObjectPtr<UObject> GameSingleton;
 
 	/** Sets the class to spawn as the global AssetManager, configurable per game. If empty, it will not spawn one */
 	UPROPERTY(globalconfig, noclear, EditAnywhere, Category=DefaultClasses, meta=(MetaClass="Object", DisplayName="Asset Manager Class"), AdvancedDisplay)
@@ -791,11 +791,11 @@ public:
 
 	/** A UObject spawned at initialization time to handle runtime asset loading and management */
 	UPROPERTY()
-	class UAssetManager *AssetManager;
+	TObjectPtr<class UAssetManager> AssetManager;
 
 	/** A global default texture. */
 	UPROPERTY()
-	class UTexture2D* DefaultTexture;
+	TObjectPtr<class UTexture2D> DefaultTexture;
 
 	/** Path of the global default texture that is used when no texture is specified. */
 	UPROPERTY(globalconfig)
@@ -803,7 +803,7 @@ public:
 
 	/** A global default diffuse texture.*/
 	UPROPERTY()
-	class UTexture* DefaultDiffuseTexture;
+	TObjectPtr<class UTexture> DefaultDiffuseTexture;
 
 	/** Path of the global default diffuse texture.*/
 	UPROPERTY(globalconfig)
@@ -811,7 +811,7 @@ public:
 
 	/** Texture used to render a vertex in the editor */
 	UPROPERTY()
-	class UTexture2D* DefaultBSPVertexTexture;
+	TObjectPtr<class UTexture2D> DefaultBSPVertexTexture;
 
 	/** Path of the texture used to render a vertex in the editor */
 	UPROPERTY(globalconfig)
@@ -819,7 +819,7 @@ public:
 
 	/** Texture used to get random image grain values for post processing */
 	UPROPERTY()
-	class UTexture2D* HighFrequencyNoiseTexture;
+	TObjectPtr<class UTexture2D> HighFrequencyNoiseTexture;
 
 	/** Path of the texture used to get random image grain values for post processing */
 	UPROPERTY(globalconfig)
@@ -827,7 +827,7 @@ public:
 
 	/** Texture used to blur out of focus content, mimics the Bokeh shape of actual cameras */
 	UPROPERTY()
-	class UTexture2D* DefaultBokehTexture;
+	TObjectPtr<class UTexture2D> DefaultBokehTexture;
 
 	/** Path of the texture used to blur out of focus content, mimics the Bokeh shape of actual cameras */
 	UPROPERTY(globalconfig)
@@ -835,7 +835,7 @@ public:
 
 	/** Texture used to bloom when using FFT, mimics characteristic bloom produced in a camera from a signle bright source */
 	UPROPERTY()
-	class UTexture2D* DefaultBloomKernelTexture;
+	TObjectPtr<class UTexture2D> DefaultBloomKernelTexture;
 
 	/** Path of the texture used to bloom when using FFT, mimics characteristic bloom produced in a camera from a signle bright source */
 	UPROPERTY(globalconfig)
@@ -843,7 +843,7 @@ public:
 
 	/** The material used to render wireframe meshes. */
 	UPROPERTY()
-	class UMaterial* WireframeMaterial;
+	TObjectPtr<class UMaterial> WireframeMaterial;
 	
 	/** Path of the material used to render wireframe meshes in the editor and debug tools. */
 	UPROPERTY(globalconfig)
@@ -852,7 +852,7 @@ public:
 #if WITH_EDITORONLY_DATA
 	/** A translucent material used to render things in geometry mode. */
 	UPROPERTY()
-	class UMaterial* GeomMaterial;
+	TObjectPtr<class UMaterial> GeomMaterial;
 
 	/** Path of the translucent material used to render things in geometry mode. */
 	UPROPERTY(globalconfig)
@@ -861,7 +861,7 @@ public:
 
 	/** A material used to render debug meshes. */
 	UPROPERTY()
-	class UMaterial* DebugMeshMaterial;
+	TObjectPtr<class UMaterial> DebugMeshMaterial;
 
 	/** Path of the default material for debug mesh */
 	UPROPERTY(globalconfig)
@@ -869,7 +869,7 @@ public:
 
 	/** A material used to render emissive meshes (e.g. light source surface). */
 	UPROPERTY()
-	class UMaterial* EmissiveMeshMaterial;
+	TObjectPtr<class UMaterial> EmissiveMeshMaterial;
 
 	/** Path of the default material for emissive mesh */
 	UPROPERTY(globalconfig)
@@ -877,7 +877,7 @@ public:
 
 	/** Material used for visualizing level membership in lit view port modes. */
 	UPROPERTY()
-	class UMaterial* LevelColorationLitMaterial;
+	TObjectPtr<class UMaterial> LevelColorationLitMaterial;
 
 	/** Path of the material used for visualizing level membership in lit view port modes. */
 	UPROPERTY(globalconfig)
@@ -885,7 +885,7 @@ public:
 
 	/** Material used for visualizing level membership in unlit view port modes. */
 	UPROPERTY()
-	class UMaterial* LevelColorationUnlitMaterial;
+	TObjectPtr<class UMaterial> LevelColorationUnlitMaterial;
 
 	/** Path of the material used for visualizing level membership in unlit view port modes. */
 	UPROPERTY(globalconfig)
@@ -893,7 +893,7 @@ public:
 
 	/** Material used for visualizing lighting only w/ lightmap texel density. */
 	UPROPERTY()
-	class UMaterial* LightingTexelDensityMaterial;
+	TObjectPtr<class UMaterial> LightingTexelDensityMaterial;
 
 	/** Path of the material used for visualizing lighting only w/ lightmap texel density. */
 	UPROPERTY(globalconfig)
@@ -901,7 +901,7 @@ public:
 
 	/** Material used for visualizing level membership in lit view port modes. Uses shading to show axis directions. */
 	UPROPERTY()
-	class UMaterial* ShadedLevelColorationLitMaterial;
+	TObjectPtr<class UMaterial> ShadedLevelColorationLitMaterial;
 
 	/** Path of the material used for visualizing level membership in lit view port modes. Uses shading to show axis directions. */
 	UPROPERTY(globalconfig)
@@ -909,7 +909,7 @@ public:
 
 	/** Material used for visualizing level membership in unlit view port modes.  Uses shading to show axis directions. */
 	UPROPERTY()
-	class UMaterial* ShadedLevelColorationUnlitMaterial;
+	TObjectPtr<class UMaterial> ShadedLevelColorationUnlitMaterial;
 
 	/** Path of the material used for visualizing level membership in unlit view port modes.  Uses shading to show axis directions. */
 	UPROPERTY(globalconfig)
@@ -917,7 +917,7 @@ public:
 
 	/** Material used to indicate that the associated BSP surface should be removed. */
 	UPROPERTY()
-	class UMaterial* RemoveSurfaceMaterial;
+	TObjectPtr<class UMaterial> RemoveSurfaceMaterial;
 
 	/** Path of the material used to indicate that the associated BSP surface should be removed. */
 	UPROPERTY(globalconfig)
@@ -925,7 +925,7 @@ public:
 
 	/** Material used to visualize vertex colors as emissive */
 	UPROPERTY()
-	class UMaterial* VertexColorMaterial;
+	TObjectPtr<class UMaterial> VertexColorMaterial;
 
 	/** Path of the material used to visualize vertex colors as emissive */
 	UPROPERTY(globalconfig)
@@ -933,7 +933,7 @@ public:
 
 	/** Material for visualizing vertex colors on meshes in the scene (color only, no alpha) */
 	UPROPERTY()
-	class UMaterial* VertexColorViewModeMaterial_ColorOnly;
+	TObjectPtr<class UMaterial> VertexColorViewModeMaterial_ColorOnly;
 
 	/** Path of the material for visualizing vertex colors on meshes in the scene (color only, no alpha) */
 	UPROPERTY(globalconfig)
@@ -941,7 +941,7 @@ public:
 
 	/** Material for visualizing vertex colors on meshes in the scene (alpha channel as color) */
 	UPROPERTY()
-	class UMaterial* VertexColorViewModeMaterial_AlphaAsColor;
+	TObjectPtr<class UMaterial> VertexColorViewModeMaterial_AlphaAsColor;
 
 	/** Path of the material for visualizing vertex colors on meshes in the scene (alpha channel as color) */
 	UPROPERTY(globalconfig)
@@ -949,7 +949,7 @@ public:
 
 	/** Material for visualizing vertex colors on meshes in the scene (red only) */
 	UPROPERTY()
-	class UMaterial* VertexColorViewModeMaterial_RedOnly;
+	TObjectPtr<class UMaterial> VertexColorViewModeMaterial_RedOnly;
 
 	/** Path of the material for visualizing vertex colors on meshes in the scene (red only) */
 	UPROPERTY(globalconfig)
@@ -957,7 +957,7 @@ public:
 
 	/** Material for visualizing vertex colors on meshes in the scene (green only) */
 	UPROPERTY()
-	class UMaterial* VertexColorViewModeMaterial_GreenOnly;
+	TObjectPtr<class UMaterial> VertexColorViewModeMaterial_GreenOnly;
 
 	/** Path of the material for visualizing vertex colors on meshes in the scene (green only) */
 	UPROPERTY(globalconfig)
@@ -965,7 +965,7 @@ public:
 
 	/** Material for visualizing vertex colors on meshes in the scene (blue only) */
 	UPROPERTY()
-	class UMaterial* VertexColorViewModeMaterial_BlueOnly;
+	TObjectPtr<class UMaterial> VertexColorViewModeMaterial_BlueOnly;
 
 	/** Path of the material for visualizing vertex colors on meshes in the scene (blue only) */
 	UPROPERTY(globalconfig)
@@ -974,7 +974,7 @@ public:
 #if WITH_EDITORONLY_DATA
 	/** Material used to render bone weights on skeletal meshes */
 	UPROPERTY()
-	class UMaterial* BoneWeightMaterial;
+	TObjectPtr<class UMaterial> BoneWeightMaterial;
 
 	/** Path of the material used to render bone weights on skeletal meshes */
 	UPROPERTY(globalconfig)
@@ -982,13 +982,13 @@ public:
 
 	/** Materials used to render cloth properties on skeletal meshes */
 	UPROPERTY()
-	class UMaterial* ClothPaintMaterial;
+	TObjectPtr<class UMaterial> ClothPaintMaterial;
 	UPROPERTY()
-	class UMaterial* ClothPaintMaterialWireframe;
+	TObjectPtr<class UMaterial> ClothPaintMaterialWireframe;
 	UPROPERTY()
-	class UMaterialInstanceDynamic* ClothPaintMaterialInstance;
+	TObjectPtr<class UMaterialInstanceDynamic> ClothPaintMaterialInstance;
 	UPROPERTY()
-	class UMaterialInstanceDynamic* ClothPaintMaterialWireframeInstance;
+	TObjectPtr<class UMaterialInstanceDynamic> ClothPaintMaterialWireframeInstance;
 
 	/** Name of the material used to render cloth in the clothing tools */
 	UPROPERTY(globalconfig)
@@ -1000,7 +1000,7 @@ public:
 
 	/** A material used to render physical material mask on mesh. */
 	UPROPERTY()
-	class UMaterial* PhysicalMaterialMaskMaterial;
+	TObjectPtr<class UMaterial> PhysicalMaterialMaskMaterial;
 
 	/** A material used to render physical material mask on mesh. */
 	UPROPERTY(globalconfig)
@@ -1008,7 +1008,7 @@ public:
 
 	/** A material used to render debug meshes. */
 	UPROPERTY()
-	class UMaterial* DebugEditorMaterial;
+	TObjectPtr<class UMaterial> DebugEditorMaterial;
 #endif
 
 	/** A material used to render debug opaque material. Used in various animation editor viewport features. */
@@ -1017,30 +1017,30 @@ public:
 
 	/** Material used to render constraint limits */
 	UPROPERTY()
-	class UMaterial* ConstraintLimitMaterial;
+	TObjectPtr<class UMaterial> ConstraintLimitMaterial;
 
 	UPROPERTY()
-	class UMaterialInstanceDynamic* ConstraintLimitMaterialX;
+	TObjectPtr<class UMaterialInstanceDynamic> ConstraintLimitMaterialX;
 
 	UPROPERTY()
-	class UMaterialInstanceDynamic* ConstraintLimitMaterialXAxis;
+	TObjectPtr<class UMaterialInstanceDynamic> ConstraintLimitMaterialXAxis;
 
 	UPROPERTY()
-	class UMaterialInstanceDynamic* ConstraintLimitMaterialY;
+	TObjectPtr<class UMaterialInstanceDynamic> ConstraintLimitMaterialY;
 	UPROPERTY()
-	class UMaterialInstanceDynamic* ConstraintLimitMaterialYAxis;
+	TObjectPtr<class UMaterialInstanceDynamic> ConstraintLimitMaterialYAxis;
 
 	UPROPERTY()
-	class UMaterialInstanceDynamic* ConstraintLimitMaterialZ;
+	TObjectPtr<class UMaterialInstanceDynamic> ConstraintLimitMaterialZ;
 	UPROPERTY()
-	class UMaterialInstanceDynamic* ConstraintLimitMaterialZAxis;
+	TObjectPtr<class UMaterialInstanceDynamic> ConstraintLimitMaterialZAxis;
 
 	UPROPERTY()
-	class UMaterialInstanceDynamic* ConstraintLimitMaterialPrismatic;
+	TObjectPtr<class UMaterialInstanceDynamic> ConstraintLimitMaterialPrismatic;
 
 	/** Material that renders a message about lightmap settings being invalid. */
 	UPROPERTY()
-	class UMaterial* InvalidLightmapSettingsMaterial;
+	TObjectPtr<class UMaterial> InvalidLightmapSettingsMaterial;
 
 	/** Path of the material that renders a message about lightmap settings being invalid. */
 	UPROPERTY(globalconfig)
@@ -1048,7 +1048,7 @@ public:
 
 	/** Material that renders a message about preview shadows being used. */
 	UPROPERTY()
-	class UMaterial* PreviewShadowsIndicatorMaterial;
+	TObjectPtr<class UMaterial> PreviewShadowsIndicatorMaterial;
 
 	/** Path of the material that renders a message about preview shadows being used. */
 	UPROPERTY(globalconfig, EditAnywhere, Category=DefaultMaterials, meta=(AllowedClasses="Material", DisplayName="Preview Shadows Indicator Material"))
@@ -1056,11 +1056,11 @@ public:
 
 	/** Material that 'fakes' lighting, used for arrows, widgets. */
 	UPROPERTY()
-	class UMaterial* ArrowMaterial;
+	TObjectPtr<class UMaterial> ArrowMaterial;
 
 	/** Arrow material instance with yellow color. */
 	UPROPERTY()
-	class UMaterialInstanceDynamic* ArrowMaterialYellow;
+	TObjectPtr<class UMaterialInstanceDynamic> ArrowMaterialYellow;
 
 	/** Path of the material that 'fakes' lighting, used for arrows, widgets. */
 	UPROPERTY(globalconfig)
@@ -1147,7 +1147,7 @@ public:
 #if WITH_EDITORONLY_DATA
 	/** A material used to render the sides of the builder brush/volumes/etc. */
 	UPROPERTY()
-	class UMaterial* EditorBrushMaterial;
+	TObjectPtr<class UMaterial> EditorBrushMaterial;
 
 	/** Path of the material used to render the sides of the builder brush/volumes/etc. */
 	UPROPERTY(globalconfig)
@@ -1156,7 +1156,7 @@ public:
 
 	/** PhysicalMaterial to use if none is defined for a particular object. */
 	UPROPERTY()
-	class UPhysicalMaterial* DefaultPhysMaterial;
+	TObjectPtr<class UPhysicalMaterial> DefaultPhysMaterial;
 
 	/** Path of the PhysicalMaterial to use if none is defined for a particular object. */
 	UPROPERTY(globalconfig)
@@ -1177,7 +1177,7 @@ public:
 
 	/** Texture used for pre-integrated skin shading */
 	UPROPERTY()
-	class UTexture2D* PreIntegratedSkinBRDFTexture;
+	TObjectPtr<class UTexture2D> PreIntegratedSkinBRDFTexture;
 
 	/** Path of the texture used for pre-integrated skin shading */
 	UPROPERTY(globalconfig)
@@ -1185,7 +1185,7 @@ public:
 
 	/** Tiled blue-noise texture */
 	UPROPERTY()
-	class UTexture2D* BlueNoiseTexture;
+	TObjectPtr<class UTexture2D> BlueNoiseTexture;
 
 	/** Path of the tiled blue-noise texture */
 	UPROPERTY(globalconfig)
@@ -1193,7 +1193,7 @@ public:
 
 	/** Texture used to do font rendering in shaders */
 	UPROPERTY()
-	class UTexture2D* MiniFontTexture;
+	TObjectPtr<class UTexture2D> MiniFontTexture;
 
 	/** Path of the texture used to do font rendering in shaders */
 	UPROPERTY(globalconfig)
@@ -1201,7 +1201,7 @@ public:
 
 	/** Texture used as a placeholder for terrain weight-maps to give the material the correct texture format. */
 	UPROPERTY()
-	class UTexture* WeightMapPlaceholderTexture;
+	TObjectPtr<class UTexture> WeightMapPlaceholderTexture;
 
 	/** Path of the texture used as a placeholder for terrain weight-maps to give the material the correct texture format. */
 	UPROPERTY(globalconfig)
@@ -1209,7 +1209,7 @@ public:
 
 	/** Texture used to display LightMapDensity */
 	UPROPERTY()
-	class UTexture2D* LightMapDensityTexture;
+	TObjectPtr<class UTexture2D> LightMapDensityTexture;
 
 	/** Path of the texture used to display LightMapDensity */
 	UPROPERTY(globalconfig)
@@ -1222,7 +1222,7 @@ public:
 
 	/** The view port representing the current game instance. Can be 0 so don't use without checking. */
 	UPROPERTY()
-	class UGameViewportClient* GameViewport;
+	TObjectPtr<class UGameViewportClient> GameViewport;
 
 	/** Array of deferred command strings/ execs that get executed at the end of the frame */
 	UPROPERTY()
@@ -1286,7 +1286,7 @@ public:
 private:
 	/** Controls how the Engine process the Framerate/Timestep */
 	UPROPERTY(transient)
-	UEngineCustomTimeStep* CustomTimeStep;
+	TObjectPtr<UEngineCustomTimeStep> CustomTimeStep;
 
 	/** Broadcasts whenever the custom time step changed. */
 	FSimpleMulticastDelegate CustomTimeStepChangedEvent;
@@ -1306,7 +1306,7 @@ public:
 private:
 	/** Controls the Engine's timecode. */
 	UPROPERTY(transient)
-	UTimecodeProvider* TimecodeProvider;
+	TObjectPtr<UTimecodeProvider> TimecodeProvider;
 
 	/** Broadcasts whenever the timecode provider changed. */
 	FSimpleMulticastDelegate TimecodeProviderChangedEvent;
diff --git a/Engine/Source/Runtime/Engine/Classes/Engine/EngineTypes.h b/Engine/Source/Runtime/Engine/Classes/Engine/EngineTypes.h
--- a/Engine/Source/Runtime/Engine/Classes/Engine/EngineTypes.h
+++ b/Engine/Source/Runtime/Engine/Classes/Engine/EngineTypes.h
@@ -583,7 +583,7 @@ private:
 	uint8 ShadingModelField = 0;
 
 	UPROPERTY()
-	TArray<USubsurfaceProfile*> SubsurfaceProfiles;
+	TArray<TObjectPtr<USubsurfaceProfile>> SubsurfaceProfiles;
 };
 
 /** This is used by the drawing passes to determine tessellation policy, so changes here need to be supported in native code. */
@@ -1423,7 +1423,7 @@ struct ENGINE_API FRigidBodyContactInfo
 
 	/** The physical material of the two shapes involved in a contact */
 	UPROPERTY()
-	class UPhysicalMaterial* PhysMaterial[2];
+	TObjectPtr<class UPhysicalMaterial> PhysMaterial[2];
 
 
 	FRigidBodyContactInfo()
@@ -1496,11 +1496,11 @@ struct FFractureEffect
 
 	/** Particle system effect to play at fracture location. */
 	UPROPERTY(EditAnywhere, BlueprintReadWrite, Category=FractureEffect)
-	class UParticleSystem* ParticleSystem;
+	TObjectPtr<class UParticleSystem> ParticleSystem;
 
 	/** Sound cue to play at fracture location. */
 	UPROPERTY(EditAnywhere, BlueprintReadWrite, Category=FractureEffect)
-	class USoundBase* Sound;
+	TObjectPtr<class USoundBase> Sound;
 
 	FFractureEffect()
 		: ParticleSystem(nullptr)
@@ -1516,7 +1516,7 @@ struct ENGINE_API FBasedPosition
 
 	/** Actor that is the base */
 	UPROPERTY(EditAnywhere, BlueprintReadWrite, Category=BasedPosition)
-	class AActor* Base;
+	TObjectPtr<class AActor> Base;
 
 	/** Position relative to the base actor */
 	UPROPERTY(EditAnywhere, BlueprintReadWrite, Category=BasedPosition)
@@ -1947,11 +1947,11 @@ struct FPrimitiveMaterialRef
 
 	/** Material is on a primitive component */
 	UPROPERTY()
-	class UPrimitiveComponent* Primitive;
+	TObjectPtr<class UPrimitiveComponent> Primitive;
 
 	/** Material is on a decal component */
 	UPROPERTY()
-	class UDecalComponent* Decal;
+	TObjectPtr<class UDecalComponent> Decal;
 
 	/** Index into the material on the components data */
 	UPROPERTY()
@@ -2762,7 +2762,7 @@ struct FMeshBuildSettings
 #endif
 
 	UPROPERTY(EditAnywhere, BlueprintReadWrite, Category=BuildSettings)
-	class UStaticMesh* DistanceFieldReplacementMesh;
+	TObjectPtr<class UStaticMesh> DistanceFieldReplacementMesh;
 
 	/** Default settings. */
 	FMeshBuildSettings()
@@ -3489,7 +3489,7 @@ struct FRepAttachment
 
 	/** Actor we are attached to, movement replication will not happen while AttachParent is non-nullptr */
 	UPROPERTY()
-	class AActor* AttachParent;
+	TObjectPtr<class AActor> AttachParent;
 
 	/** Location offset from attach parent */
 	UPROPERTY()
@@ -3509,7 +3509,7 @@ struct FRepAttachment
 
 	/** Specific component we are attached to */
 	UPROPERTY()
-	class USceneComponent* AttachComponent;
+	TObjectPtr<class USceneComponent> AttachComponent;
 
 	FRepAttachment()
 		: AttachParent(nullptr)
@@ -3745,7 +3745,7 @@ struct ENGINE_API FComponentReference
 	 * If this is not provided the reference refers to a component on this / the same actor.
 	 */
 	UPROPERTY(EditInstanceOnly, BlueprintReadWrite, Category=Component, meta = (DisplayName = "Referenced Actor"))
-	AActor* OtherActor;
+	TObjectPtr<AActor> OtherActor;
 
 	/** Name of component to use. If this is not specified the reference refers to the root component. */
 	UPROPERTY(EditAnywhere, BlueprintReadWrite, Category=Component, meta = (DisplayName = "Component Name"))
diff --git a/Engine/Source/Runtime/Engine/Classes/Engine/ExponentialHeightFog.h b/Engine/Source/Runtime/Engine/Classes/Engine/ExponentialHeightFog.h
--- a/Engine/Source/Runtime/Engine/Classes/Engine/ExponentialHeightFog.h
+++ b/Engine/Source/Runtime/Engine/Classes/Engine/ExponentialHeightFog.h
@@ -21,7 +21,7 @@ class AExponentialHeightFog
 private:
 	/** @todo document */
 	UPROPERTY(Category = ExponentialHeightFog, VisibleAnywhere, BlueprintReadOnly, meta = (AllowPrivateAccess = "true"))
-	class UExponentialHeightFogComponent* Component;
+	TObjectPtr<class UExponentialHeightFogComponent> Component;
 
 public:
 
diff --git a/Engine/Source/Runtime/Engine/Classes/Engine/Font.h b/Engine/Source/Runtime/Engine/Classes/Engine/Font.h
--- a/Engine/Source/Runtime/Engine/Classes/Engine/Font.h
+++ b/Engine/Source/Runtime/Engine/Classes/Engine/Font.h
@@ -113,7 +113,7 @@ class UFont : public UObject, public IFontProviderInterface
 	/** Textures that store this font's glyph image data */
 	//NOTE: Do not expose this to the editor as it has nasty crash potential
 	UPROPERTY()
-	TArray<class UTexture2D*> Textures;
+	TArray<TObjectPtr<class UTexture2D>> Textures;
 
 	/** True if font is 'remapped'.  That is, the character array is not a direct mapping to unicode values.  Instead,
 		all characters are indexed indirectly through the CharRemap array */
diff --git a/Engine/Source/Runtime/Engine/Classes/Engine/GameEngine.h b/Engine/Source/Runtime/Engine/Classes/Engine/GameEngine.h
--- a/Engine/Source/Runtime/Engine/Classes/Engine/GameEngine.h
+++ b/Engine/Source/Runtime/Engine/Classes/Engine/GameEngine.h
@@ -36,7 +36,7 @@ class ENGINE_API UGameEngine
 	float ServerFlushLogInterval;
 
 	UPROPERTY(transient)
-	UGameInstance* GameInstance;
+	TObjectPtr<UGameInstance> GameInstance;
 
 public:
 
diff --git a/Engine/Source/Runtime/Engine/Classes/Engine/GameInstance.h b/Engine/Source/Runtime/Engine/Classes/Engine/GameInstance.h
--- a/Engine/Source/Runtime/Engine/Classes/Engine/GameInstance.h
+++ b/Engine/Source/Runtime/Engine/Classes/Engine/GameInstance.h
@@ -169,15 +169,15 @@ protected:
 
 	/** List of locally participating players in this game instance */
 	UPROPERTY()
-	TArray<ULocalPlayer*> LocalPlayers;
+	TArray<TObjectPtr<ULocalPlayer>> LocalPlayers;
 	
 	/** Class to manage online services */
 	UPROPERTY()
-	class UOnlineSession* OnlineSession;
+	TObjectPtr<class UOnlineSession> OnlineSession;
 
 	/** List of objects that are being kept alive by this game instance. Stored as array for fast iteration, should not be modified every frame */
 	UPROPERTY()
-	TArray<UObject*> ReferencedObjects;
+	TArray<TObjectPtr<UObject>> ReferencedObjects;
 
 	/** Listeners to PreClientTravel call */
 	FOnPreClientTravel NotifyPreClientTravelDelegates;
diff --git a/Engine/Source/Runtime/Engine/Classes/Engine/GameViewportClient.h b/Engine/Source/Runtime/Engine/Classes/Engine/GameViewportClient.h
--- a/Engine/Source/Runtime/Engine/Classes/Engine/GameViewportClient.h
+++ b/Engine/Source/Runtime/Engine/Classes/Engine/GameViewportClient.h
@@ -66,7 +66,7 @@ public:
 
 	/** The viewport's console.   Might be null on consoles */
 	UPROPERTY()
-	class UConsole* ViewportConsole;
+	TObjectPtr<class UConsole> ViewportConsole;
 
 	/** Debug properties that have been added via one of the "displayall" commands */
 	UPROPERTY()
@@ -99,10 +99,10 @@ protected:
 
 	/* The relative world context for this viewport */
 	UPROPERTY()
-	UWorld* World;
+	TObjectPtr<UWorld> World;
 
 	UPROPERTY()
-	UGameInstance* GameInstance;
+	TObjectPtr<UGameInstance> GameInstance;
 
 	/** If true will suppress the blue transition text messages. */
 	bool bSuppressTransitionMessage;
diff --git a/Engine/Source/Runtime/Engine/Classes/Engine/HLODProxy.h b/Engine/Source/Runtime/Engine/Classes/Engine/HLODProxy.h
--- a/Engine/Source/Runtime/Engine/Classes/Engine/HLODProxy.h
+++ b/Engine/Source/Runtime/Engine/Classes/Engine/HLODProxy.h
@@ -103,5 +103,5 @@ private:
 	TArray<FHLODProxyMesh> ProxyMeshes;
 
 	UPROPERTY(VisibleAnywhere, Category = "Proxy Mesh")
-	TMap<UHLODProxyDesc*, FHLODProxyMesh> HLODActors;
+	TMap<TObjectPtr<UHLODProxyDesc>, FHLODProxyMesh> HLODActors;
 };
\ No newline at end of file
diff --git a/Engine/Source/Runtime/Engine/Classes/Engine/InheritableComponentHandler.h b/Engine/Source/Runtime/Engine/Classes/Engine/InheritableComponentHandler.h
--- a/Engine/Source/Runtime/Engine/Classes/Engine/InheritableComponentHandler.h
+++ b/Engine/Source/Runtime/Engine/Classes/Engine/InheritableComponentHandler.h
@@ -56,7 +56,7 @@ struct ENGINE_API FComponentKey
 
 private: 
 	UPROPERTY()
-	UClass* OwnerClass;
+	TObjectPtr<UClass> OwnerClass;
 
 	UPROPERTY()
 	FName SCSVariableName;
@@ -71,10 +71,10 @@ struct FComponentOverrideRecord
 	GENERATED_USTRUCT_BODY()
 
 	UPROPERTY()
-	UClass* ComponentClass;
+	TObjectPtr<UClass> ComponentClass;
 
 	UPROPERTY()
-	UActorComponent* ComponentTemplate;
+	TObjectPtr<UActorComponent> ComponentTemplate;
 
 	UPROPERTY()
 	FComponentKey ComponentKey;
@@ -168,5 +168,5 @@ private:
 
 	/** List of components that were marked unnecessary, need to keep these around so it doesn't regenerate them when a child asks for one */
 	UPROPERTY(Transient)
-	TArray<UActorComponent*> UnnecessaryComponents;
+	TArray<TObjectPtr<UActorComponent>> UnnecessaryComponents;
 };
diff --git a/Engine/Source/Runtime/Engine/Classes/Engine/InterpCurveEdSetup.h b/Engine/Source/Runtime/Engine/Classes/Engine/InterpCurveEdSetup.h
--- a/Engine/Source/Runtime/Engine/Classes/Engine/InterpCurveEdSetup.h
+++ b/Engine/Source/Runtime/Engine/Classes/Engine/InterpCurveEdSetup.h
@@ -16,7 +16,7 @@ struct FCurveEdEntry
 	GENERATED_USTRUCT_BODY()
 
 	UPROPERTY()
-	class UObject* CurveObject;
+	TObjectPtr<class UObject> CurveObject;
 
 	UPROPERTY()
 	FColor CurveColor;
diff --git a/Engine/Source/Runtime/Engine/Classes/Engine/LODActor.h b/Engine/Source/Runtime/Engine/Classes/Engine/LODActor.h
--- a/Engine/Source/Runtime/Engine/Classes/Engine/LODActor.h
+++ b/Engine/Source/Runtime/Engine/Classes/Engine/LODActor.h
@@ -26,10 +26,10 @@ struct FHLODInstancingKey
 	GENERATED_BODY()
 
 	UPROPERTY()
-	const UStaticMesh* StaticMesh;
+	TObjectPtr<const UStaticMesh> StaticMesh;
 
 	UPROPERTY()
-	const UMaterialInterface* Material;
+	TObjectPtr<const UMaterialInterface> Material;
 
 	FHLODInstancingKey()
 		: StaticMesh()
@@ -76,15 +76,15 @@ class ENGINE_API ALODActor : public AActor
 private:
 	// disable display of this component
 	UPROPERTY(Category=LODActor, VisibleAnywhere)
-	UStaticMeshComponent* StaticMeshComponent;
+	TObjectPtr<UStaticMeshComponent> StaticMeshComponent;
 
 	/** Imposters to be rendered as instanced static meshes */
 	UPROPERTY(Transient)
-	TMap<FHLODInstancingKey, UInstancedStaticMeshComponent*> InstancedStaticMeshComponents;
+	TMap<FHLODInstancingKey, TObjectPtr<UInstancedStaticMeshComponent>> InstancedStaticMeshComponents;
 
 	/** The mesh proxy used to display this LOD */
 	UPROPERTY(Category=LODActor, VisibleAnywhere)
-	UHLODProxy* Proxy;
+	TObjectPtr<UHLODProxy> Proxy;
 
 	/** The key used to validate this actor against the proxy */
 	UPROPERTY(Category=LODActor, VisibleAnywhere)
@@ -100,7 +100,7 @@ public:
 	int32 LODLevel;
 
 	UPROPERTY(Category=LODActor, VisibleAnywhere)
-	TArray<AActor*> SubActors;
+	TArray<TObjectPtr<AActor>> SubActors;
 	
 	UPROPERTY()
 	uint8 CachedNumHLODLevels;
@@ -397,6 +397,6 @@ private:
 
 	/** The hlod proxy desc used to build this LODActor */
 	UPROPERTY(Transient)
-	UHLODProxyDesc* ProxyDesc;
+	TObjectPtr<UHLODProxyDesc> ProxyDesc;
 #endif // WITH_EDITORONLY_DATA
 }; 
diff --git a/Engine/Source/Runtime/Engine/Classes/Engine/LatentActionManager.h b/Engine/Source/Runtime/Engine/Classes/Engine/LatentActionManager.h
--- a/Engine/Source/Runtime/Engine/Classes/Engine/LatentActionManager.h
+++ b/Engine/Source/Runtime/Engine/Classes/Engine/LatentActionManager.h
@@ -29,7 +29,7 @@ struct ENGINE_API FLatentActionInfo
 
 	/** Object to execute the function on. */ 
 	UPROPERTY(meta=(LatentCallbackTarget = true))
-	UObject* CallbackTarget;
+	TObjectPtr<UObject> CallbackTarget;
 
 	FLatentActionInfo()
 		: Linkage(INDEX_NONE)
diff --git a/Engine/Source/Runtime/Engine/Classes/Engine/Level.h b/Engine/Source/Runtime/Engine/Classes/Engine/Level.h
--- a/Engine/Source/Runtime/Engine/Classes/Engine/Level.h
+++ b/Engine/Source/Runtime/Engine/Classes/Engine/Level.h
@@ -118,7 +118,7 @@ struct ENGINE_API FDynamicTextureInstance : public FStreamableTextureInstance
 
 	/** Texture that is used by a dynamic UPrimitiveComponent. */
 	UPROPERTY()
-	UTexture2D*					Texture = nullptr;
+	TObjectPtr<UTexture2D>					Texture = nullptr;
 
 	/** Whether the primitive that uses this texture is attached to the scene or not. */
 	UPROPERTY()
@@ -393,7 +393,7 @@ struct FReplicatedStaticActorDestructionInfo
 	FVector	DestroyedPosition;
 	TWeakObjectPtr<UObject> ObjOuter;
 	UPROPERTY()
-	UClass* ObjClass = nullptr;
+	TObjectPtr<UClass> ObjClass = nullptr;
 };
 
 #if WITH_EDITORONLY_DATA
@@ -454,23 +454,23 @@ public:
 	 * It should not be accessed during BeginDestroy(), just like any other UObject references, since GC may occur in any order.
 	 */
 	UPROPERTY(Transient)
-	UWorld* OwningWorld;
+	TObjectPtr<UWorld> OwningWorld;
 
 	/** BSP UModel. */
 	UPROPERTY()
-	class UModel* Model;
+	TObjectPtr<class UModel> Model;
 
 	/** BSP Model components used for rendering. */
 	UPROPERTY()
-	TArray<class UModelComponent*> ModelComponents;
+	TArray<TObjectPtr<class UModelComponent>> ModelComponents;
 
 	UPROPERTY(Transient, DuplicateTransient, NonTransactional)
-	ULevelActorContainer* ActorCluster;
+	TObjectPtr<ULevelActorContainer> ActorCluster;
 
 #if WITH_EDITORONLY_DATA
 	/** Reference to the blueprint for level scripting */
 	UPROPERTY(NonTransactional)
-	class ULevelScriptBlueprint* LevelScriptBlueprint;
+	TObjectPtr<class ULevelScriptBlueprint> LevelScriptBlueprint;
 
 	/** The Guid list of all materials and meshes Guid used in the last texture streaming build. Used to know if the streaming data needs rebuild. Only used for the persistent level. */
 	UPROPERTY(NonTransactional)
@@ -487,20 +487,20 @@ public:
 
 	/** The level scripting actor, created by instantiating the class from LevelScriptBlueprint.  This handles all level scripting */
 	UPROPERTY(NonTransactional)
-	class ALevelScriptActor* LevelScriptActor;
+	TObjectPtr<class ALevelScriptActor> LevelScriptActor;
 
 	/**
 	 * Start and end of the navigation list for this level, used for quickly fixing up
 	 * when streaming this level in/out. @TODO DEPRECATED - DELETE
 	 */
 	UPROPERTY()
-	class ANavigationObjectBase *NavListStart;
+	TObjectPtr<class ANavigationObjectBase> NavListStart;
 	UPROPERTY()
-	class ANavigationObjectBase	*NavListEnd;
+	TObjectPtr<class ANavigationObjectBase>	NavListEnd;
 	
 	/** Navigation related data that can be stored per level */
 	UPROPERTY()
-	TArray<UNavigationDataChunk*> NavDataChunks;
+	TArray<TObjectPtr<UNavigationDataChunk>> NavDataChunks;
 	
 	/** Total number of KB used for lightmap textures in the level. */
 	UPROPERTY(VisibleAnywhere, Category=Level)
@@ -548,7 +548,7 @@ public:
 	 * ReleaseRenderingResources must be called before changing what is referenced, to update the rendering thread state.
 	 */
 	UPROPERTY(NonPIEDuplicateTransient)
-	UMapBuildDataRegistry* MapBuildData;
+	TObjectPtr<UMapBuildDataRegistry> MapBuildData;
 
 	/** Level offset at time when lighting was built */
 	UPROPERTY()
@@ -714,7 +714,7 @@ private:
 	FLevelBoundsActorUpdatedEvent LevelBoundsActorUpdatedEvent; 
 
 	UPROPERTY()
-	AWorldSettings* WorldSettings;
+	TObjectPtr<AWorldSettings> WorldSettings;
 
 	/** Cached level collection that this level is contained in, for faster access than looping through the collections in the world. */
 	FLevelCollection* CachedLevelCollection;
@@ -723,7 +723,7 @@ protected:
 
 	/** Array of user data stored with the asset */
 	UPROPERTY()
-	TArray<UAssetUserData*> AssetUserData;
+	TArray<TObjectPtr<UAssetUserData>> AssetUserData;
 
 private:
 	// Actors awaiting input to be enabled once the appropriate PlayerController has been created
diff --git a/Engine/Source/Runtime/Engine/Classes/Engine/LevelActorContainer.h b/Engine/Source/Runtime/Engine/Classes/Engine/LevelActorContainer.h
--- a/Engine/Source/Runtime/Engine/Classes/Engine/LevelActorContainer.h
+++ b/Engine/Source/Runtime/Engine/Classes/Engine/LevelActorContainer.h
@@ -24,7 +24,7 @@ public:
 
 	/** Array of actors in a level */
 	UPROPERTY(transient)
-	TArray<AActor*> Actors;
+	TArray<TObjectPtr<AActor>> Actors;
 
 	virtual void CreateCluster() override;
 	virtual void OnClusterMarkedAsPendingKill() override;
diff --git a/Engine/Source/Runtime/Engine/Classes/Engine/LevelBounds.h b/Engine/Source/Runtime/Engine/Classes/Engine/LevelBounds.h
--- a/Engine/Source/Runtime/Engine/Classes/Engine/LevelBounds.h
+++ b/Engine/Source/Runtime/Engine/Classes/Engine/LevelBounds.h
@@ -33,7 +33,7 @@ class ENGINE_API ALevelBounds
 		
 	/** Bounding box for the level bounds. */
 	UPROPERTY(EditAnywhere, Category = LevelBounds)
-	class UBoxComponent* BoxComponent;
+	TObjectPtr<class UBoxComponent> BoxComponent;
 
 	/** Whether to automatically update actor bounds based on all relevant actors bounds belonging to the same level */
 	UPROPERTY(EditAnywhere, Category=LevelBounds)
diff --git a/Engine/Source/Runtime/Engine/Classes/Engine/LevelStreaming.h b/Engine/Source/Runtime/Engine/Classes/Engine/LevelStreaming.h
--- a/Engine/Source/Runtime/Engine/Classes/Engine/LevelStreaming.h
+++ b/Engine/Source/Runtime/Engine/Classes/Engine/LevelStreaming.h
@@ -244,7 +244,7 @@ public:
 
 	/** The level streaming volumes bound to this level. */
 	UPROPERTY(EditAnywhere, Category=LevelStreaming, meta=(DisplayName = "Streaming Volumes", NoElementDuplicate))
-	TArray<ALevelStreamingVolume*> EditorStreamingVolumes;
+	TArray<TObjectPtr<ALevelStreamingVolume>> EditorStreamingVolumes;
 
 	/** Cooldown time in seconds between volume-based unload requests.  Used in preventing spurious unload requests. */
 	UPROPERTY(EditAnywhere, AdvancedDisplay, Category=LevelStreaming, meta=(ClampMin = "0", UIMin = "0", UIMax = "10"))
@@ -535,11 +535,11 @@ protected:
 
 	/** Pointer to Level object if currently loaded/ streamed in.																*/
 	UPROPERTY(transient)
-	class ULevel* LoadedLevel;
+	TObjectPtr<class ULevel> LoadedLevel;
 
 	/** Pointer to a Level object that was previously active and was replaced with a new LoadedLevel (for LOD switching) */
 	UPROPERTY(transient)
-	class ULevel* PendingUnloadLevel;
+	TObjectPtr<class ULevel> PendingUnloadLevel;
 
 private:
 	/** @return Name of the LOD level package used for loading.																		*/
diff --git a/Engine/Source/Runtime/Engine/Classes/Engine/Light.h b/Engine/Source/Runtime/Engine/Classes/Engine/Light.h
--- a/Engine/Source/Runtime/Engine/Classes/Engine/Light.h
+++ b/Engine/Source/Runtime/Engine/Classes/Engine/Light.h
@@ -17,7 +17,7 @@ class ENGINE_API ALight : public AActor
 private:
 	/** @todo document */
 	UPROPERTY(Category = Light, BlueprintReadOnly, meta = (ExposeFunctionCategories = "Light,Rendering,Rendering|Components|Light", AllowPrivateAccess = "true"))
-	class ULightComponent* LightComponent;
+	TObjectPtr<class ULightComponent> LightComponent;
 public:
 
 	virtual void GetLifetimeReplicatedProps(TArray< FLifetimeProperty > & OutLifetimeProps) const override;
diff --git a/Engine/Source/Runtime/Engine/Classes/Engine/LocalPlayer.h b/Engine/Source/Runtime/Engine/Classes/Engine/LocalPlayer.h
--- a/Engine/Source/Runtime/Engine/Classes/Engine/LocalPlayer.h
+++ b/Engine/Source/Runtime/Engine/Classes/Engine/LocalPlayer.h
@@ -178,7 +178,7 @@ public:
 
 	/** The master viewport containing this player's view. */
 	UPROPERTY()
-	class UGameViewportClient* ViewportClient;
+	TObjectPtr<class UGameViewportClient> ViewportClient;
 
 	/** The coordinates for the upper left corner of the master viewport subregion allocated to this player. 0-1 */
 	FVector2D Origin;
diff --git a/Engine/Source/Runtime/Engine/Classes/Engine/MemberReference.h b/Engine/Source/Runtime/Engine/Classes/Engine/MemberReference.h
--- a/Engine/Source/Runtime/Engine/Classes/Engine/MemberReference.h
+++ b/Engine/Source/Runtime/Engine/Classes/Engine/MemberReference.h
@@ -71,7 +71,7 @@ protected:
 	 * be NULL if bSelfContext is true.  
 	 */
 	UPROPERTY()
-	mutable UObject* MemberParent;
+	mutable TObjectPtr<UObject> MemberParent;
 
 	/**  */
 	UPROPERTY()
diff --git a/Engine/Source/Runtime/Engine/Classes/Engine/NavigationObjectBase.h b/Engine/Source/Runtime/Engine/Classes/Engine/NavigationObjectBase.h
--- a/Engine/Source/Runtime/Engine/Classes/Engine/NavigationObjectBase.h
+++ b/Engine/Source/Runtime/Engine/Classes/Engine/NavigationObjectBase.h
@@ -16,15 +16,15 @@ class ENGINE_API ANavigationObjectBase : public AActor, public INavAgentInterfac
 
 private:
 	UPROPERTY()
-	class UCapsuleComponent* CapsuleComponent;
+	TObjectPtr<class UCapsuleComponent> CapsuleComponent;
 
 	/** Normal editor sprite. */
 	UPROPERTY()
-	class UBillboardComponent* GoodSprite;
+	TObjectPtr<class UBillboardComponent> GoodSprite;
 
 	/** Used to draw bad collision intersection in editor. */
 	UPROPERTY()
-	class UBillboardComponent* BadSprite;
+	TObjectPtr<class UBillboardComponent> BadSprite;
 public:
 
 	/** True if this nav point was spawned to be a PIE player start. */
diff --git a/Engine/Source/Runtime/Engine/Classes/Engine/NetConnection.h b/Engine/Source/Runtime/Engine/Classes/Engine/NetConnection.h
--- a/Engine/Source/Runtime/Engine/Classes/Engine/NetConnection.h
+++ b/Engine/Source/Runtime/Engine/Classes/Engine/NetConnection.h
@@ -274,11 +274,11 @@ class UNetConnection : public UPlayer
 public:
 	/** child connections for secondary viewports */
 	UPROPERTY(transient)
-	TArray<class UChildConnection*> Children;
+	TArray<TObjectPtr<class UChildConnection>> Children;
 
 	/** Owning net driver */
 	UPROPERTY()
-	class UNetDriver* Driver;	
+	TObjectPtr<class UNetDriver> Driver;	
 
 	/** The class name for the PackageMap to be loaded */
 	UPROPERTY()
@@ -286,23 +286,23 @@ public:
 
 	UPROPERTY()
 	/** Package map between local and remote. (negotiates net serialization) */
-	class UPackageMap* PackageMap;
+	TObjectPtr<class UPackageMap> PackageMap;
 
 	/** @todo document */
 	UPROPERTY()
-	TArray<class UChannel*> OpenChannels;
+	TArray<TObjectPtr<class UChannel>> OpenChannels;
 	 
 	/** This actor is bNetTemporary, which means it should never be replicated after it's initial packet is complete */
 	UPROPERTY()
-	TArray<class AActor*> SentTemporaries;
+	TArray<TObjectPtr<class AActor>> SentTemporaries;
 
 	/** The actor that is currently being viewed/controlled by the owning controller */
 	UPROPERTY()
-	class AActor* ViewTarget;
+	TObjectPtr<class AActor> ViewTarget;
 
 	/** Reference to controlling actor (usually PlayerController) */
 	UPROPERTY()
-	class AActor* OwningActor;
+	TObjectPtr<class AActor> OwningActor;
 
 	UPROPERTY()
 	int32	MaxPacket;						// Maximum packet size.
@@ -1347,7 +1347,7 @@ private:
 	 * OpenChannels every frame.
 	 */
 	UPROPERTY()
-	TArray<UChannel*> ChannelsToTick;
+	TArray<TObjectPtr<UChannel>> ChannelsToTick;
 
 	/** Histogram of the received packet time */
 	FHistogram NetConnectionHistogram;
diff --git a/Engine/Source/Runtime/Engine/Classes/Engine/NetDriver.h b/Engine/Source/Runtime/Engine/Classes/Engine/NetDriver.h
--- a/Engine/Source/Runtime/Engine/Classes/Engine/NetDriver.h
+++ b/Engine/Source/Runtime/Engine/Classes/Engine/NetDriver.h
@@ -602,7 +602,7 @@ struct ENGINE_API FChannelDefinition
 	FName ClassName;			// UClass name used to create the UChannel
 
 	UPROPERTY()
-	UClass* ChannelClass;		// UClass used to create the UChannel
+	TObjectPtr<UClass> ChannelClass;		// UClass used to create the UChannel
 
 	UPROPERTY()
 	int32 StaticChannelIndex;	// Channel always uses this index, INDEX_NONE if dynamically chosen
@@ -776,11 +776,11 @@ public:
 
 	/** Connection to the server (this net driver is a client) */
 	UPROPERTY()
-	class UNetConnection* ServerConnection;
+	TObjectPtr<class UNetConnection> ServerConnection;
 
 	/** Array of connections to clients (this net driver is a host) - unsorted, and ordering changes depending on actor replication */
 	UPROPERTY()
-	TArray<UNetConnection*> ClientConnections;
+	TArray<TObjectPtr<UNetConnection>> ClientConnections;
 
 	/**
 	 * Map of IP's to NetConnection's - for fast lookup, particularly under DDoS.
@@ -810,10 +810,10 @@ public:
 
 	/** World this net driver is associated with */
 	UPROPERTY()
-	class UWorld* World;
+	TObjectPtr<class UWorld> World;
 
 	UPROPERTY()
-	class UPackage* WorldPackage;
+	TObjectPtr<class UPackage> WorldPackage;
 
 	/** @todo document */
 	TSharedPtr< class FNetGUIDCache > GuidCache;
@@ -822,10 +822,10 @@ public:
 
 	/** The loaded UClass of the net connection type to use */
 	UPROPERTY()
-	UClass* NetConnectionClass;
+	TObjectPtr<UClass> NetConnectionClass;
 
 	UPROPERTY()
-	UClass* ReplicationDriverClass;
+	TObjectPtr<UClass> ReplicationDriverClass;
 
 	/** @todo document */
 	FProperty* RoleProperty;
@@ -861,7 +861,7 @@ private:
 
 	/** List of channels that were previously used and can be used again */
 	UPROPERTY()
-	TArray<UChannel*> ActorChannelPool;
+	TArray<TObjectPtr<UChannel>> ActorChannelPool;
 
 public:
 
@@ -1725,7 +1725,7 @@ private:
 	void LoadChannelDefinitions();
 
 	UPROPERTY(transient)
-	UReplicationDriver* ReplicationDriver;
+	TObjectPtr<UReplicationDriver> ReplicationDriver;
 
 	/** Stores the list of objects to replicate into the replay stream. This should be a TUniquePtr, but it appears the generated.cpp file needs the full definition of the pointed-to type. */
 	TSharedPtr<FNetworkObjectList> NetworkObjects;
diff --git a/Engine/Source/Runtime/Engine/Classes/Engine/Note.h b/Engine/Source/Runtime/Engine/Classes/Engine/Note.h
--- a/Engine/Source/Runtime/Engine/Classes/Engine/Note.h
+++ b/Engine/Source/Runtime/Engine/Classes/Engine/Note.h
@@ -24,11 +24,11 @@ class ANote : public AActor
 	// Reference to sprite visualization component
 private:
 	UPROPERTY()
-	class UBillboardComponent* SpriteComponent;
+	TObjectPtr<class UBillboardComponent> SpriteComponent;
 
 	// Reference to arrow visualization component
 	UPROPERTY()
-	class UArrowComponent* ArrowComponent;
+	TObjectPtr<class UArrowComponent> ArrowComponent;
 public:
 
 #endif // WITH_EDITORONLY_DATA
diff --git a/Engine/Source/Runtime/Engine/Classes/Engine/ObjectLibrary.h b/Engine/Source/Runtime/Engine/Classes/Engine/ObjectLibrary.h
--- a/Engine/Source/Runtime/Engine/Classes/Engine/ObjectLibrary.h
+++ b/Engine/Source/Runtime/Engine/Classes/Engine/ObjectLibrary.h
@@ -18,7 +18,7 @@ class UObjectLibrary : public UObject
 
 	/** Class that Objects must be of. If ContainsBlueprints is true, this is the native class that the blueprints are instances of and not UClass  */
 	UPROPERTY(EditAnywhere, Category=ObjectLibrary, meta=(AllowAbstract = ""))
-	UClass*				ObjectBaseClass;
+	TObjectPtr<UClass>				ObjectBaseClass;
 
 	/** True if this library holds blueprint classes, false if it holds other objects */
 	UPROPERTY(EditAnywhere, Category=ObjectLibrary)
@@ -28,7 +28,7 @@ protected:
 
 	/** List of Objects in library */
 	UPROPERTY(EditAnywhere, Category=ObjectLibrary)
-	TArray<UObject*>	Objects;
+	TArray<TObjectPtr<UObject>>	Objects;
 
 	/** Weak pointers to objects */
 	UPROPERTY()
diff --git a/Engine/Source/Runtime/Engine/Classes/Engine/ObjectReferencer.h b/Engine/Source/Runtime/Engine/Classes/Engine/ObjectReferencer.h
--- a/Engine/Source/Runtime/Engine/Classes/Engine/ObjectReferencer.h
+++ b/Engine/Source/Runtime/Engine/Classes/Engine/ObjectReferencer.h
@@ -19,7 +19,7 @@ class UObjectReferencer : public UObject
 
 	/** Array of objects being referenced. */
 	UPROPERTY(EditAnywhere, Category=ObjectReferencer)
-	TArray<class UObject*> ReferencedObjects;
+	TArray<TObjectPtr<class UObject>> ReferencedObjects;
 
 };
 
diff --git a/Engine/Source/Runtime/Engine/Classes/Engine/PendingNetGame.h b/Engine/Source/Runtime/Engine/Classes/Engine/PendingNetGame.h
--- a/Engine/Source/Runtime/Engine/Classes/Engine/PendingNetGame.h
+++ b/Engine/Source/Runtime/Engine/Classes/Engine/PendingNetGame.h
@@ -29,7 +29,7 @@ public:
 	 * Transferred to world on successful connection
 	 */
 	UPROPERTY()
-	class UNetDriver*		NetDriver;
+	TObjectPtr<class UNetDriver>		NetDriver;
 
 	/** 
 	 * Demo Net driver created for loading demos, but we need to go through pending net game
@@ -37,7 +37,7 @@ public:
 	 */
 	UE_DEPRECATED(4.26, "DemoNetDriver will be made private in a future release.  Please use GetDemoNetDriver/SetDemoNetDriver instead.")
 	UPROPERTY()
-	class UDemoNetDriver*	DemoNetDriver;
+	TObjectPtr<class UDemoNetDriver>	DemoNetDriver;
 
 	PRAGMA_DISABLE_DEPRECATION_WARNINGS
 	/** Gets the demo net driver for this pending world. */
diff --git a/Engine/Source/Runtime/Engine/Classes/Engine/PlanarReflection.h b/Engine/Source/Runtime/Engine/Classes/Engine/PlanarReflection.h
--- a/Engine/Source/Runtime/Engine/Classes/Engine/PlanarReflection.h
+++ b/Engine/Source/Runtime/Engine/Classes/Engine/PlanarReflection.h
@@ -21,11 +21,11 @@ class APlanarReflection : public ASceneCapture
 private:
 	/** Planar reflection component. */
 	UPROPERTY(Category = SceneCapture, VisibleAnywhere, BlueprintReadOnly, meta = (AllowPrivateAccess = "true"))
-	class UPlanarReflectionComponent* PlanarReflectionComponent;
+	TObjectPtr<class UPlanarReflectionComponent> PlanarReflectionComponent;
 
 #if WITH_EDITORONLY_DATA
 	UPROPERTY()
-	UBillboardComponent* SpriteComponent;
+	TObjectPtr<UBillboardComponent> SpriteComponent;
 #endif
 
 public:
diff --git a/Engine/Source/Runtime/Engine/Classes/Engine/PlatformInterfaceBase.h b/Engine/Source/Runtime/Engine/Classes/Engine/PlatformInterfaceBase.h
--- a/Engine/Source/Runtime/Engine/Classes/Engine/PlatformInterfaceBase.h
+++ b/Engine/Source/Runtime/Engine/Classes/Engine/PlatformInterfaceBase.h
@@ -60,7 +60,7 @@ struct FPlatformInterfaceData
 	FString StringValue;
 
 	UPROPERTY()
-	class UObject* ObjectValue;
+	TObjectPtr<class UObject> ObjectValue;
 
 
 	FPlatformInterfaceData()
diff --git a/Engine/Source/Runtime/Engine/Classes/Engine/Player.h b/Engine/Source/Runtime/Engine/Classes/Engine/Player.h
--- a/Engine/Source/Runtime/Engine/Classes/Engine/Player.h
+++ b/Engine/Source/Runtime/Engine/Classes/Engine/Player.h
@@ -20,7 +20,7 @@ class UPlayer : public UObject, public FExec
 
 	/** The actor this player controls. */
 	UPROPERTY(transient)
-	class APlayerController* PlayerController;
+	TObjectPtr<class APlayerController> PlayerController;
 
 	// Net variables.
 	/** the current speed of the connection */
diff --git a/Engine/Source/Runtime/Engine/Classes/Engine/PointLight.h b/Engine/Source/Runtime/Engine/Classes/Engine/PointLight.h
--- a/Engine/Source/Runtime/Engine/Classes/Engine/PointLight.h
+++ b/Engine/Source/Runtime/Engine/Classes/Engine/PointLight.h
@@ -13,7 +13,7 @@ class APointLight : public ALight
 	GENERATED_UCLASS_BODY()
 
 	UPROPERTY(BlueprintReadOnly, VisibleAnywhere, Category="Light", meta=(ExposeFunctionCategories="PointLight,Rendering|Lighting"))
-	class UPointLightComponent* PointLightComponent;
+	TObjectPtr<class UPointLightComponent> PointLightComponent;
 
 	// BEGIN DEPRECATED (use component functions now in level script)
 	UFUNCTION(BlueprintCallable, Category="Rendering|Lighting", meta=(DeprecatedFunction))
diff --git a/Engine/Source/Runtime/Engine/Classes/Engine/PoseWatch.h b/Engine/Source/Runtime/Engine/Classes/Engine/PoseWatch.h
--- a/Engine/Source/Runtime/Engine/Classes/Engine/PoseWatch.h
+++ b/Engine/Source/Runtime/Engine/Classes/Engine/PoseWatch.h
@@ -27,7 +27,7 @@ class ENGINE_API UPoseWatch
 public:
 	// Node that we are trying to watch
 	UPROPERTY()
-	class UEdGraphNode* Node;
+	TObjectPtr<class UEdGraphNode> Node;
 
 	UPROPERTY()
 	FColor PoseWatchColour;
diff --git a/Engine/Source/Runtime/Engine/Classes/Engine/PreviewMeshCollection.h b/Engine/Source/Runtime/Engine/Classes/Engine/PreviewMeshCollection.h
--- a/Engine/Source/Runtime/Engine/Classes/Engine/PreviewMeshCollection.h
+++ b/Engine/Source/Runtime/Engine/Classes/Engine/PreviewMeshCollection.h
@@ -42,7 +42,7 @@ public:
 	GENERATED_BODY()
 
 	UPROPERTY(VisibleAnywhere, Category = "Skeleton", AssetRegistrySearchable)
-	USkeleton* Skeleton;
+	TObjectPtr<USkeleton> Skeleton;
 
 	/** The skeletal meshes that this collection contains */
 	UPROPERTY(EditAnywhere, Category = "Skeletal Meshes")
diff --git a/Engine/Source/Runtime/Engine/Classes/Engine/RectLight.h b/Engine/Source/Runtime/Engine/Classes/Engine/RectLight.h
--- a/Engine/Source/Runtime/Engine/Classes/Engine/RectLight.h
+++ b/Engine/Source/Runtime/Engine/Classes/Engine/RectLight.h
@@ -13,7 +13,7 @@ class ARectLight : public ALight
 	GENERATED_UCLASS_BODY()
 
 	UPROPERTY(BlueprintReadOnly, VisibleAnywhere, Category="Light", meta=(ExposeFunctionCategories="RectLight,Rendering|Lighting"))
-	class URectLightComponent* RectLightComponent;
+	TObjectPtr<class URectLightComponent> RectLightComponent;
 
 #if WITH_EDITOR
 	//~ Begin AActor Interface.
diff --git a/Engine/Source/Runtime/Engine/Classes/Engine/ReflectionCapture.h b/Engine/Source/Runtime/Engine/Classes/Engine/ReflectionCapture.h
--- a/Engine/Source/Runtime/Engine/Classes/Engine/ReflectionCapture.h
+++ b/Engine/Source/Runtime/Engine/Classes/Engine/ReflectionCapture.h
@@ -17,14 +17,14 @@ class AReflectionCapture : public AActor
 private:
 	/** Reflection capture component. */
 	UPROPERTY(Category = DecalActor, VisibleAnywhere, BlueprintReadOnly, meta = (AllowPrivateAccess = "true"))
-	class UReflectionCaptureComponent* CaptureComponent;
+	TObjectPtr<class UReflectionCaptureComponent> CaptureComponent;
 
 #if WITH_EDITORONLY_DATA
 	UPROPERTY()
-	UBillboardComponent* SpriteComponent;
+	TObjectPtr<UBillboardComponent> SpriteComponent;
 
 	UPROPERTY()
-	UBillboardComponent* CaptureOffsetComponent;
+	TObjectPtr<UBillboardComponent> CaptureOffsetComponent;
 #endif // WITH_EDITORONLY_DATA
 
 public:	
diff --git a/Engine/Source/Runtime/Engine/Classes/Engine/SCS_Node.h b/Engine/Source/Runtime/Engine/Classes/Engine/SCS_Node.h
--- a/Engine/Source/Runtime/Engine/Classes/Engine/SCS_Node.h
+++ b/Engine/Source/Runtime/Engine/Classes/Engine/SCS_Node.h
@@ -23,11 +23,11 @@ class USCS_Node : public UObject
 
 	/** Component class */
 	UPROPERTY()
-	UClass* ComponentClass;
+	TObjectPtr<UClass> ComponentClass;
 
 	/** Template for the component to create */
 	UPROPERTY()
-	class UActorComponent* ComponentTemplate;
+	TObjectPtr<class UActorComponent> ComponentTemplate;
 
 	/** Cached data for faster runtime instancing (only used in cooked builds) */
 	UPROPERTY()
@@ -58,7 +58,7 @@ class USCS_Node : public UObject
 
 	/** Set of child nodes */
 	UPROPERTY()
-	TArray<class USCS_Node*> ChildNodes;
+	TArray<TObjectPtr<class USCS_Node>> ChildNodes;
 
 	/** Metadata information for this Node */
 	UPROPERTY(EditAnywhere, Category=BPVariableDescription)
diff --git a/Engine/Source/Runtime/Engine/Classes/Engine/Scene.h b/Engine/Source/Runtime/Engine/Classes/Engine/Scene.h
--- a/Engine/Source/Runtime/Engine/Classes/Engine/Scene.h
+++ b/Engine/Source/Runtime/Engine/Classes/Engine/Scene.h
@@ -394,7 +394,7 @@ struct FConvolutionBloomSettings
 
 	/** Texture to replace default convolution bloom kernel */
 	UPROPERTY(EditAnywhere, BlueprintReadWrite, Category = "Lens|Bloom", meta = (DisplayName = "Convolution Kernel"))
-	class UTexture2D* Texture;
+	TObjectPtr<class UTexture2D> Texture;
 
 	/** Relative size of the convolution kernel image compared to the minor axis of the viewport  */
 	UPROPERTY(Interp, BlueprintReadWrite, Category = "Lens|Bloom", AdvancedDisplay, meta = (ClampMin = "0.0", UIMax = "1.0", DisplayName = "Convolution Scale"))
@@ -472,7 +472,7 @@ struct FLensImperfectionSettings
 	 * Texture that defines the dirt on the camera lens where the light of very bright objects is scattered.
 	 */
 	UPROPERTY(EditAnywhere, BlueprintReadWrite, Category="Lens|Dirt Mask", meta=(DisplayName = "Dirt Mask Texture"))
-	class UTexture* DirtMask;	
+	TObjectPtr<class UTexture> DirtMask;	
 	
 	/** BloomDirtMask intensity */
 	UPROPERTY(Interp, BlueprintReadWrite, Category="Lens|Dirt Mask", meta=(ClampMin = "0.0", UIMax = "8.0", DisplayName = "Dirt Mask Intensity"))
@@ -590,14 +590,14 @@ struct FCameraExposureSettings
 	 * 0: no adjustment, -1:2x darker, -2:4x darker, 1:2x brighter, 2:4x brighter, ...
 	 */
 	UPROPERTY(EditAnywhere, BlueprintReadWrite, Category = "Exposure", meta = (DisplayName = "Exposure Bias Curve"))
-	class UCurveFloat* BiasCurve = nullptr;
+	TObjectPtr<class UCurveFloat> BiasCurve = nullptr;
 
 	/**
 	 * Exposure metering mask. Bright spots on the mask will have high influence on auto-exposure metering
 	 * and dark spots will have low influence.
 	 */
 	UPROPERTY(EditAnywhere, BlueprintReadWrite, Category="Exposure", meta=(DisplayName = "Exposure Metering Mask"))
-	class UTexture* MeterMask = nullptr;	
+	TObjectPtr<class UTexture> MeterMask = nullptr;	
 
 	/** temporary exposed until we found good values, -8: 1/256, -10: 1/1024 */
 	UPROPERTY(Interp, BlueprintReadWrite, Category="Exposure", AdvancedDisplay, meta=(UIMin = "-16", UIMax = "0.0"))
@@ -632,7 +632,7 @@ struct FWeightedBlendable
 
 	/** should be of the IBlendableInterface* type but UProperties cannot express that */
 	UPROPERTY(EditAnywhere, BlueprintReadWrite, Category=FWeightedBlendable, meta=( AllowedClasses="BlendableInterface", Keywords="PostProcess" ))
-	UObject* Object;
+	TObjectPtr<UObject> Object;
 
 	// default constructor
 	FWeightedBlendable()
@@ -1424,7 +1424,7 @@ struct FPostProcessSettings
 
 	/** Texture to replace default convolution bloom kernel */
 	UPROPERTY(EditAnywhere, BlueprintReadWrite, Category = "Lens|Bloom", meta = (editcondition = "bOverride_BloomConvolutionTexture", DisplayName = "Convolution Kernel"))
-	class UTexture2D* BloomConvolutionTexture;
+	TObjectPtr<class UTexture2D> BloomConvolutionTexture;
 
 	/** The UV location of the center of the kernel.  Should be very close to (.5,.5) */
 	UPROPERTY(interp, BlueprintReadWrite, Category = "Lens|Bloom", AdvancedDisplay, meta = (editcondition = "bOverride_BloomConvolutionCenterUV", DisplayName = "Convolution Center"))
@@ -1455,7 +1455,7 @@ struct FPostProcessSettings
 	 * Texture that defines the dirt on the camera lens where the light of very bright objects is scattered.
 	 */
 	UPROPERTY(EditAnywhere, BlueprintReadWrite, Category="Lens|Dirt Mask", meta=(editcondition = "bOverride_BloomDirtMask", DisplayName = "Dirt Mask Texture"))
-	class UTexture* BloomDirtMask;	
+	TObjectPtr<class UTexture> BloomDirtMask;	
 	
 	/** BloomDirtMask intensity */
 	UPROPERTY(interp, BlueprintReadWrite, Category="Lens|Dirt Mask", meta=(ClampMin = "0.0", UIMax = "8.0", editcondition = "bOverride_BloomDirtMaskIntensity", DisplayName = "Dirt Mask Intensity"))
@@ -1478,7 +1478,7 @@ struct FPostProcessSettings
 
 	/** The Ambient cubemap (Affects diffuse and specular shading), blends additively which if different from all other settings here */
 	UPROPERTY(EditAnywhere, BlueprintReadWrite, Category="Rendering Features|Ambient Cubemap", meta=(DisplayName = "Cubemap Texture"))
-	class UTextureCube* AmbientCubemap;
+	TObjectPtr<class UTextureCube> AmbientCubemap;
 
 	/** The camera shutter in seconds.*/
 	UPROPERTY(EditAnywhere, BlueprintReadWrite, Category="Lens|Camera", meta=(ClampMin = "1.0", ClampMax = "2000.0", editcondition = "bOverride_CameraShutterSpeed", DisplayName = "Shutter Speed (1/s)"))
@@ -1531,14 +1531,14 @@ struct FPostProcessSettings
 	 * 0: no adjustment, -1:2x darker, -2:4x darker, 1:2x brighter, 2:4x brighter, ...
 	 */
 	UPROPERTY(EditAnywhere, BlueprintReadWrite, Category = "Lens|Exposure", meta = (editcondition = "bOverride_AutoExposureBiasCurve", DisplayName = "Exposure Compensation Curve"))
-	class UCurveFloat* AutoExposureBiasCurve = nullptr;
+	TObjectPtr<class UCurveFloat> AutoExposureBiasCurve = nullptr;
 
 	/**
 	 * Exposure metering mask. Bright spots on the mask will have high influence on auto-exposure metering
 	 * and dark spots will have low influence.
 	 */
 	UPROPERTY(EditAnywhere, BlueprintReadWrite, Category="Lens|Exposure", meta=(editcondition = "bOverride_AutoExposureMeterMask", DisplayName = "Exposure Metering Mask"))
-	class UTexture* AutoExposureMeterMask = nullptr;	
+	TObjectPtr<class UTexture> AutoExposureMeterMask = nullptr;	
 
 	/**
 	 * The eye adaptation will adapt to a value extracted from the luminance histogram of the scene color.
@@ -1614,7 +1614,7 @@ struct FPostProcessSettings
 
 	/** Defines the shape of the Bokeh when the image base lens flares are blurred, cannot be blended */
 	UPROPERTY(EditAnywhere, BlueprintReadWrite, Category="Lens|Lens Flares", meta=(editcondition = "bOverride_LensFlareBokehShape", DisplayName = "BokehShape"))
-	class UTexture* LensFlareBokehShape;
+	TObjectPtr<class UTexture> LensFlareBokehShape;
 
 	/** RGB defines the lens flare color, A it's position. This is a temporary solution. */
 	UPROPERTY(EditAnywhere, Category="Lens|Lens Flares", meta=(editcondition = "bOverride_LensFlareTints", DisplayName = "Tints"))
@@ -1735,7 +1735,7 @@ struct FPostProcessSettings
 
 	/** Look up table texture to use or none of not used*/
 	UPROPERTY(EditAnywhere, BlueprintReadWrite, Category="Color Grading|Misc", meta=(editcondition = "bOverride_ColorGradingLUT", DisplayName = "Color Grading LUT"))
-	class UTexture* ColorGradingLUT;
+	TObjectPtr<class UTexture> ColorGradingLUT;
 
 	/** Width of the camera sensor to assume, in mm. */
 	UPROPERTY(BlueprintReadWrite, Category="Lens|Depth of Field", meta=(ForceUnits=mm, ClampMin = "0.1", UIMin="0.1", UIMax= "1000.0", editcondition = "bOverride_DepthOfFieldSensorWidth", DisplayName = "Sensor Width (mm)"))
@@ -1957,7 +1957,7 @@ struct FPostProcessSettings
 #if WITH_EDITORONLY_DATA
 	// for backwards compatibility
 	UPROPERTY()
-	TArray<UObject*> Blendables_DEPRECATED;
+	TArray<TObjectPtr<UObject>> Blendables_DEPRECATED;
 
 	// for backwards compatibility
 	void OnAfterLoad()
diff --git a/Engine/Source/Runtime/Engine/Classes/Engine/SceneCapture.h b/Engine/Source/Runtime/Engine/Classes/Engine/SceneCapture.h
--- a/Engine/Source/Runtime/Engine/Classes/Engine/SceneCapture.h
+++ b/Engine/Source/Runtime/Engine/Classes/Engine/SceneCapture.h
@@ -21,10 +21,10 @@ private:
 	/** To display the 3d camera in the editor. */
 	UE_DEPRECATED(4.22, "SceneCapture's mesh and frustum components should now be accessed through the SceneCaptureComponent instead of the Actor")
 	UPROPERTY()
-	class UStaticMeshComponent* MeshComp_DEPRECATED;
+	TObjectPtr<class UStaticMeshComponent> MeshComp_DEPRECATED;
 
 	UPROPERTY(Category = Components, VisibleAnywhere, BlueprintReadOnly, meta = (AllowPrivateAccess = "true"))
-	class USceneComponent* SceneComponent;
+	TObjectPtr<class USceneComponent> SceneComponent;
 
 public:
 	ENGINE_API virtual void Serialize(FArchive& Ar) override;
diff --git a/Engine/Source/Runtime/Engine/Classes/Engine/SceneCapture2D.h b/Engine/Source/Runtime/Engine/Classes/Engine/SceneCapture2D.h
--- a/Engine/Source/Runtime/Engine/Classes/Engine/SceneCapture2D.h
+++ b/Engine/Source/Runtime/Engine/Classes/Engine/SceneCapture2D.h
@@ -20,7 +20,7 @@ class ASceneCapture2D : public ASceneCapture
 private:
 	/** Scene capture component. */
 	UPROPERTY(Category = DecalActor, VisibleAnywhere, BlueprintReadOnly, meta=(AllowPrivateAccess = "true"))
-	class USceneCaptureComponent2D* CaptureComponent2D;
+	TObjectPtr<class USceneCaptureComponent2D> CaptureComponent2D;
 
 public:
 	UFUNCTION(BlueprintCallable, Category="Rendering")
diff --git a/Engine/Source/Runtime/Engine/Classes/Engine/SceneCaptureCube.h b/Engine/Source/Runtime/Engine/Classes/Engine/SceneCaptureCube.h
--- a/Engine/Source/Runtime/Engine/Classes/Engine/SceneCaptureCube.h
+++ b/Engine/Source/Runtime/Engine/Classes/Engine/SceneCaptureCube.h
@@ -20,7 +20,7 @@ class ENGINE_API ASceneCaptureCube : public ASceneCapture
 private:
 	/** Scene capture component. */
 	UPROPERTY(Category = DecalActor, VisibleAnywhere, BlueprintReadOnly, meta=(AllowPrivateAccess = "true"))
-	class USceneCaptureComponentCube* CaptureComponentCube;
+	TObjectPtr<class USceneCaptureComponentCube> CaptureComponentCube;
 
 public:
 	UFUNCTION(BlueprintCallable, Category = "Rendering")
diff --git a/Engine/Source/Runtime/Engine/Classes/Engine/SimpleConstructionScript.h b/Engine/Source/Runtime/Engine/Classes/Engine/SimpleConstructionScript.h
--- a/Engine/Source/Runtime/Engine/Classes/Engine/SimpleConstructionScript.h
+++ b/Engine/Source/Runtime/Engine/Classes/Engine/SimpleConstructionScript.h
@@ -128,24 +128,24 @@ class USimpleConstructionScript : public UObject
 private:
 	/** Root nodes of the construction script */
 	UPROPERTY()
-	TArray<class USCS_Node*> RootNodes;
+	TArray<TObjectPtr<class USCS_Node>> RootNodes;
 
 	/** All nodes that exist in the hierarchy of this SimpleConstructionScript */
 	UPROPERTY()
-	TArray<class USCS_Node*> AllNodes;
+	TArray<TObjectPtr<class USCS_Node>> AllNodes;
 
 	/** Default scene root node; used when no other nodes are available to use as the root */
 	UPROPERTY()
-	class USCS_Node* DefaultSceneRootNode;
+	TObjectPtr<class USCS_Node> DefaultSceneRootNode;
 
 #if WITH_EDITORONLY_DATA
 	/** (DEPRECATED) Root node of the construction script */
 	UPROPERTY()
-	class USCS_Node* RootNode_DEPRECATED;
+	TObjectPtr<class USCS_Node> RootNode_DEPRECATED;
 
 	/** (DEPRECATED) Actor Component based nodes are stored here.  They cannot be in the tree hierarchy */
 	UPROPERTY()
-	TArray<USCS_Node*> ActorComponentNodes_DEPRECATED;
+	TArray<TObjectPtr<USCS_Node>> ActorComponentNodes_DEPRECATED;
 #endif
 
 	/** Helper function to find an appropriate child node index that can be promoted to the parent's level */
diff --git a/Engine/Source/Runtime/Engine/Classes/Engine/SkeletalMesh.h b/Engine/Source/Runtime/Engine/Classes/Engine/SkeletalMesh.h
--- a/Engine/Source/Runtime/Engine/Classes/Engine/SkeletalMesh.h
+++ b/Engine/Source/Runtime/Engine/Classes/Engine/SkeletalMesh.h
@@ -230,11 +230,11 @@ struct FSkeletalMeshLODInfo
 
 	/** Pose which should be used to reskin vertex influences for which the bones will be removed in this LOD level, uses ref-pose by default */
 	UPROPERTY(EditAnywhere, Category = ReductionSettings)
-	class UAnimSequence* BakePose;
+	TObjectPtr<class UAnimSequence> BakePose;
 
 	/** This is used when you are sharing the LOD settings, but you'd like to override the BasePose. This precedes prior to BakePose*/
 	UPROPERTY(EditAnywhere, Category = ReductionSettings)
-	class UAnimSequence* BakePoseOverride;
+	TObjectPtr<class UAnimSequence> BakePoseOverride;
 
 	/** The filename of the file tha was used to import this LOD if it was not auto generated. */
 	UPROPERTY(VisibleAnywhere, Category= SkeletalMeshLODInfo, AdvancedDisplay)
@@ -448,7 +448,7 @@ struct FSkeletalMaterial
 	ENGINE_API friend bool operator==( const UMaterialInterface& LHS, const FSkeletalMaterial& RHS );
 
 	UPROPERTY(EditAnywhere, BlueprintReadWrite, Category=SkeletalMesh)
-	class UMaterialInterface *	MaterialInterface;
+	TObjectPtr<class UMaterialInterface> 	MaterialInterface;
 	
 	/*This name should be use by the gameplay to avoid error if the skeletal mesh Materials array topology change*/
 	UPROPERTY(EditAnywhere, BlueprintReadWrite, Category = SkeletalMesh)
@@ -535,7 +535,7 @@ public:
 	 */
 	UE_DEPRECATED(4.27, "Please do not access this member directly; use the public ImportData API.")
 	UPROPERTY()
-	mutable USkeletalMeshEditorData* MeshEditorDataObject;
+	mutable TObjectPtr<USkeletalMeshEditorData> MeshEditorDataObject;
 
 private:
 	/*
@@ -626,7 +626,7 @@ public:
 	/** Skeleton of this skeletal mesh **/
 	UE_DEPRECATED(4.27, "Please do not access this member directly; use USkeletalMesh::GetSkeleton() or USkeletalMesh::SetSkeleton().")
 	UPROPERTY(Category=Mesh, AssetRegistrySearchable, VisibleAnywhere, BlueprintGetter = GetSkeleton)
-	USkeleton* Skeleton;
+	TObjectPtr<USkeleton> Skeleton;
 
 	static FName GetSkeletonMemberName()
 	{
@@ -974,7 +974,7 @@ public:
 
 	UE_DEPRECATED(4.27, "Please do not access this member directly; use USkeletalMesh::GetLODSettings() or USkeletalMesh::SetLODSettings().")
 	UPROPERTY(EditAnywhere, AssetRegistrySearchable, BlueprintGetter = GetLODSettings, BlueprintSetter = SetLODSettings, Category = LODSettings)
-	USkeletalMeshLODSettings* LODSettings;
+	TObjectPtr<USkeletalMeshLODSettings> LODSettings;
 
 	static FName GetLODSettingsMemberName()
 	{
@@ -1228,7 +1228,7 @@ public:
 	// Physics data for the per poly collision case. In 99% of cases you will not need this and are better off using simple ragdoll collision (physics asset)
 	UE_DEPRECATED(4.27, "Please do not access this member directly; use USkeletalMesh::GetBodySetup() or USkeletalMesh::SetBodySetup().")
 	UPROPERTY(transient)
-	class UBodySetup* BodySetup;
+	TObjectPtr<class UBodySetup> BodySetup;
 	static FName GetBodySetupMemberName()
 	{
 		PRAGMA_DISABLE_DEPRECATION_WARNINGS
@@ -1265,7 +1265,7 @@ public:
 	 */
 	UE_DEPRECATED(4.27, "Please do not access this member directly; use USkeletalMesh::GetPhysicsAsset() or USkeletalMesh::SetPhysicsAsset().")
 	UPROPERTY(EditAnywhere, AssetRegistrySearchable, BlueprintGetter = GetPhysicsAsset, Category=Physics)
-	class UPhysicsAsset* PhysicsAsset;
+	TObjectPtr<class UPhysicsAsset> PhysicsAsset;
 	
 	static FName GetPhysicsAssetMemberName()
 	{
@@ -1295,7 +1295,7 @@ public:
 	 */
 	UE_DEPRECATED(4.27, "Please do not access this member directly; use USkeletalMesh::GetShadowPhysicsAsset() or USkeletalMesh::SetShadowPhysicsAsset().")
 	UPROPERTY(EditAnywhere, AssetRegistrySearchable, BlueprintGetter = GetShadowPhysicsAsset, Category=Lighting)
-	class UPhysicsAsset* ShadowPhysicsAsset;
+	TObjectPtr<class UPhysicsAsset> ShadowPhysicsAsset;
 
 	static FName GetShadowPhysicsAssetMemberName()
 	{
@@ -1322,7 +1322,7 @@ public:
 	/** Mapping data that is saved */
 	UE_DEPRECATED(4.27, "Please do not access this member directly; use USkeletalMesh::GetNodeMappingData() or USkeletalMesh::SetNodeMappingData().")
 	UPROPERTY(EditAnywhere, editfixedsize, BlueprintGetter = GetNodeMappingData, Category=Animation)
-	TArray<class UNodeMappingContainer*> NodeMappingData;
+	TArray<TObjectPtr<class UNodeMappingContainer>> NodeMappingData;
 	static FName GetNodeMappingDataMemberName()
 	{
 		PRAGMA_DISABLE_DEPRECATION_WARNINGS
@@ -1360,7 +1360,7 @@ public:
 	/** Importing data and options used for this mesh */
 	UE_DEPRECATED(4.27, "Please do not access this member directly; use USkeletalMesh::GetAssetImportData() or USkeletalMesh::SetAssetImportData().")
 	UPROPERTY(EditAnywhere, Instanced, Category=ImportSettings)
-	class UAssetImportData* AssetImportData;
+	TObjectPtr<class UAssetImportData> AssetImportData;
 
 	static FName GetAssetImportDataMemberName()
 	{
@@ -1396,7 +1396,7 @@ public:
 	/** Information for thumbnail rendering */
 	UE_DEPRECATED(4.27, "Please do not access this member directly; use USkeletalMesh::GetThumbnailInfo() or USkeletalMesh::SetThumbnailInfo().")
 	UPROPERTY(VisibleAnywhere, Instanced, AdvancedDisplay, Category = Thumbnail)
-	class UThumbnailInfo* ThumbnailInfo;
+	TObjectPtr<class UThumbnailInfo> ThumbnailInfo;
 	
 	static FName GetThumbnailInfoMemberName()
 	{
@@ -1629,7 +1629,7 @@ public:
 	UE_DEPRECATED(4.27, "Please do not access this member directly; use USkeletalMesh::GetMorphTargets() or USkeletalMesh::SetMorphTargets().")
 	UPROPERTY(BlueprintGetter = GetMorphTargets, BlueprintSetter = SetMorphTargets, Category = Mesh)
 
-	TArray<UMorphTarget*> MorphTargets;
+	TArray<TObjectPtr<UMorphTarget>> MorphTargets;
 
 	static FName GetMorphTargetsMemberName()
 	{
@@ -1899,7 +1899,7 @@ public:
 	 */
 	UE_DEPRECATED(4.27, "Please do not access this member directly; use USkeletalMesh::GetMeshClothingAssets() or USkeletalMesh::SetMeshClothingAssets().")
 	UPROPERTY(EditAnywhere, editfixedsize, BlueprintGetter = GetMeshClothingAssets, BlueprintSetter = SetMeshClothingAssets, Category = Clothing)
-	TArray<UClothingAssetBase*> MeshClothingAssets;
+	TArray<TObjectPtr<UClothingAssetBase>> MeshClothingAssets;
 
 	static FName GetMeshClothingAssetsMemberName()
 	{
@@ -1978,7 +1978,7 @@ protected:
 
 	/** Array of user data stored with the asset */
 	UPROPERTY(EditAnywhere, AdvancedDisplay, Instanced, Category=SkeletalMesh)
-	TArray<UAssetUserData*> AssetUserData;
+	TArray<TObjectPtr<UAssetUserData>> AssetUserData;
 
 #if WITH_EDITOR
 	FOnMeshChanged OnMeshChanged;
@@ -1993,7 +1993,7 @@ private:
 	 *	everything explicitly to AttachComponent in the SkeletalMeshComponent. 
 	 */
 	UPROPERTY()
-	TArray<class USkeletalMeshSocket*> Sockets;
+	TArray<TObjectPtr<class USkeletalMeshSocket>> Sockets;
 
 	/** Cached matrices from GetComposedRefPoseMatrix */
 	TArray<FMatrix> CachedComposedRefPoseMatrices;
diff --git a/Engine/Source/Runtime/Engine/Classes/Engine/SkeletalMeshLODSettings.h b/Engine/Source/Runtime/Engine/Classes/Engine/SkeletalMeshLODSettings.h
--- a/Engine/Source/Runtime/Engine/Classes/Engine/SkeletalMeshLODSettings.h
+++ b/Engine/Source/Runtime/Engine/Classes/Engine/SkeletalMeshLODSettings.h
@@ -103,7 +103,7 @@ struct FSkeletalMeshLODGroupSettings
 
 	/** Pose which should be used to reskin vertex influences for which the bones will be removed in this LOD level, uses ref-pose by default */
 	UPROPERTY(EditAnywhere, Category = Reduction)
-	class UAnimSequence* BakePose;
+	TObjectPtr<class UAnimSequence> BakePose;
 
 	/** The optimization settings to use for the respective LOD level */
 	UPROPERTY(EditAnywhere, Category = Reduction)
diff --git a/Engine/Source/Runtime/Engine/Classes/Engine/SkyLight.h b/Engine/Source/Runtime/Engine/Classes/Engine/SkyLight.h
--- a/Engine/Source/Runtime/Engine/Classes/Engine/SkyLight.h
+++ b/Engine/Source/Runtime/Engine/Classes/Engine/SkyLight.h
@@ -17,7 +17,7 @@ class ENGINE_API ASkyLight : public AInfo
 private:
 	/** @todo document */
 	UPROPERTY(Category = Light, VisibleAnywhere, BlueprintReadOnly, meta = (ExposeFunctionCategories = "Light,Rendering,Rendering|Components|SkyLight", AllowPrivateAccess = "true"))
-	class USkyLightComponent* LightComponent;
+	TObjectPtr<class USkyLightComponent> LightComponent;
 public:
 
 	/** replicated copy of LightComponent's bEnabled property */
diff --git a/Engine/Source/Runtime/Engine/Classes/Engine/SphereReflectionCapture.h b/Engine/Source/Runtime/Engine/Classes/Engine/SphereReflectionCapture.h
--- a/Engine/Source/Runtime/Engine/Classes/Engine/SphereReflectionCapture.h
+++ b/Engine/Source/Runtime/Engine/Classes/Engine/SphereReflectionCapture.h
@@ -21,7 +21,7 @@ class ASphereReflectionCapture : public AReflectionCapture
 private:
 	/** Sphere component used to visualize the capture radius */
 	UPROPERTY()
-	UDrawSphereComponent* DrawCaptureRadius;
+	TObjectPtr<UDrawSphereComponent> DrawCaptureRadius;
 
 public:
 
diff --git a/Engine/Source/Runtime/Engine/Classes/Engine/SplineMeshActor.h b/Engine/Source/Runtime/Engine/Classes/Engine/SplineMeshActor.h
--- a/Engine/Source/Runtime/Engine/Classes/Engine/SplineMeshActor.h
+++ b/Engine/Source/Runtime/Engine/Classes/Engine/SplineMeshActor.h
@@ -19,7 +19,7 @@ class ENGINE_API ASplineMeshActor : public AActor
 
 private:
 	UPROPERTY(Category = SplineMeshActor, VisibleAnywhere, BlueprintReadOnly, meta = (ExposeFunctionCategories = "Mesh,Rendering,Physics,Components|StaticMesh,Components|SplineMesh", AllowPrivateAccess = "true"))
-	class USplineMeshComponent* SplineMeshComponent;
+	TObjectPtr<class USplineMeshComponent> SplineMeshComponent;
 
 public:
 
diff --git a/Engine/Source/Runtime/Engine/Classes/Engine/SpotLight.h b/Engine/Source/Runtime/Engine/Classes/Engine/SpotLight.h
--- a/Engine/Source/Runtime/Engine/Classes/Engine/SpotLight.h
+++ b/Engine/Source/Runtime/Engine/Classes/Engine/SpotLight.h
@@ -13,13 +13,13 @@ class ASpotLight : public ALight
 	GENERATED_UCLASS_BODY()
 
 	UPROPERTY(BlueprintReadOnly, VisibleAnywhere, Category="Light", meta=(ExposeFunctionCategories="SpotLight,Rendering|Lighting"))
-	class USpotLightComponent* SpotLightComponent;
+	TObjectPtr<class USpotLightComponent> SpotLightComponent;
 
 #if WITH_EDITORONLY_DATA
 	// Reference to editor arrow component visualization 
 private:
 	UPROPERTY()
-	class UArrowComponent* ArrowComponent;
+	TObjectPtr<class UArrowComponent> ArrowComponent;
 public:
 #endif
 
diff --git a/Engine/Source/Runtime/Engine/Classes/Engine/StaticMesh.h b/Engine/Source/Runtime/Engine/Classes/Engine/StaticMesh.h
--- a/Engine/Source/Runtime/Engine/Classes/Engine/StaticMesh.h
+++ b/Engine/Source/Runtime/Engine/Classes/Engine/StaticMesh.h
@@ -613,7 +613,7 @@ struct FStaticMaterial
 	ENGINE_API friend bool operator==(const UMaterialInterface& LHS, const FStaticMaterial& RHS);
 
 	UPROPERTY(EditAnywhere, BlueprintReadWrite, Category = StaticMesh)
-	class UMaterialInterface* MaterialInterface;
+	TObjectPtr<class UMaterialInterface> MaterialInterface;
 
 	/*This name should be use by the gameplay to avoid error if the skeletal mesh Materials array topology change*/
 	UPROPERTY(EditAnywhere, BlueprintReadWrite, Category = StaticMesh)
@@ -817,7 +817,7 @@ public:
 
 	/** Materials used by this static mesh. Individual sections index in to this array. */
 	UPROPERTY()
-	TArray<UMaterialInterface*> Materials_DEPRECATED;
+	TArray<TObjectPtr<UMaterialInterface>> Materials_DEPRECATED;
 
 	/** Settings related to building Nanite data. */
 	UPROPERTY()
@@ -989,7 +989,7 @@ private:
 	// Physics data.
 	UE_DEPRECATED(5.00, "This must be protected for async build, always use the accessors even internally.")
 	UPROPERTY(EditAnywhere, transient, duplicatetransient, Instanced, Category = StaticMesh)
-	class UBodySetup* BodySetup;
+	TObjectPtr<class UBodySetup> BodySetup;
 public:
 	UBodySetup* GetBodySetup() const
 	{
@@ -1111,7 +1111,7 @@ public:
 #if WITH_EDITORONLY_DATA
 	/** Importing data and options used for this mesh */
 	UPROPERTY(EditAnywhere, Instanced, Category=ImportSettings)
-	class UAssetImportData* AssetImportData;
+	TObjectPtr<class UAssetImportData> AssetImportData;
 
 	/** Path to the resource used to construct this static mesh */
 	UPROPERTY()
@@ -1123,7 +1123,7 @@ public:
 
 	/** Information for thumbnail rendering */
 	UPROPERTY(VisibleAnywhere, Instanced, AdvancedDisplay, Category=StaticMesh)
-	class UThumbnailInfo* ThumbnailInfo;
+	TObjectPtr<class UThumbnailInfo> ThumbnailInfo;
 
 	/** The stored camera position to use as a default for the static mesh editor */
 	UPROPERTY()
@@ -1176,7 +1176,7 @@ public:
 	 *	everything explicitly to AttachComponent in the StaticMeshComponent.
 	 */
 	UPROPERTY()
-	TArray<class UStaticMeshSocket*> Sockets;
+	TArray<TObjectPtr<class UStaticMeshSocket>> Sockets;
 
 	/** Data that is only available if this static mesh is an imported SpeedTree */
 	TSharedPtr<class FSpeedTreeWind> SpeedTreeWind;
@@ -1281,7 +1281,7 @@ protected:
 
 	/** Array of user data stored with the asset */
 	UPROPERTY(EditAnywhere, AdvancedDisplay, Instanced, Category = StaticMesh)
-	TArray<UAssetUserData*> AssetUserData;
+	TArray<TObjectPtr<UAssetUserData>> AssetUserData;
 
 	friend class FStaticMeshCompilingManager;
 	friend class FStaticMeshAsyncBuildWorker;
@@ -1291,10 +1291,10 @@ protected:
 public:
 #if WITH_EDITORONLY_DATA
 	UPROPERTY(Instanced)
-	class UObject* EditableMesh_DEPRECATED;
+	TObjectPtr<class UObject> EditableMesh_DEPRECATED;
 
 	UPROPERTY(EditAnywhere, Category = Collision)
-	class UStaticMesh* ComplexCollisionMesh;
+	TObjectPtr<class UStaticMesh> ComplexCollisionMesh;
 #endif
 
 	/**
@@ -1451,7 +1451,7 @@ public:
 private:
 	UE_DEPRECATED(5.00, "This must be protected for async build, always use the accessors even internally.")
 	UPROPERTY(VisibleAnywhere, transient, duplicatetransient, Instanced, Category = Navigation)
-	UNavCollisionBase* NavCollision;
+	TObjectPtr<UNavCollisionBase> NavCollision;
 
 public:
 	ENGINE_API void SetNavCollision(UNavCollisionBase*);
diff --git a/Engine/Source/Runtime/Engine/Classes/Engine/StaticMeshActor.h b/Engine/Source/Runtime/Engine/Classes/Engine/StaticMeshActor.h
--- a/Engine/Source/Runtime/Engine/Classes/Engine/StaticMeshActor.h
+++ b/Engine/Source/Runtime/Engine/Classes/Engine/StaticMeshActor.h
@@ -23,7 +23,7 @@ class ENGINE_API AStaticMeshActor : public AActor
 
 private:
 	UPROPERTY(Category = StaticMeshActor, VisibleAnywhere, BlueprintReadOnly, meta = (ExposeFunctionCategories = "Mesh,Rendering,Physics,Components|StaticMesh", AllowPrivateAccess = "true"))
-	class UStaticMeshComponent* StaticMeshComponent;
+	TObjectPtr<class UStaticMeshComponent> StaticMeshComponent;
 
 protected:
 	virtual void BeginPlay() override;
diff --git a/Engine/Source/Runtime/Engine/Classes/Engine/StaticMeshSocket.h b/Engine/Source/Runtime/Engine/Classes/Engine/StaticMeshSocket.h
--- a/Engine/Source/Runtime/Engine/Classes/Engine/StaticMeshSocket.h
+++ b/Engine/Source/Runtime/Engine/Classes/Engine/StaticMeshSocket.h
@@ -39,7 +39,7 @@ class UStaticMeshSocket : public UObject
 
 #if WITH_EDITORONLY_DATA
 	UPROPERTY(EditAnywhere, AdvancedDisplay, Category=StaticMeshSocket)
-	class UStaticMesh* PreviewStaticMesh;
+	TObjectPtr<class UStaticMesh> PreviewStaticMesh;
 
 	/* Weather the socket was imported with the asset or create in the editor. Importer will remove/modify only imported socket and will not touch any editor created socket. */
 	UPROPERTY()
diff --git a/Engine/Source/Runtime/Engine/Classes/Engine/TargetPoint.h b/Engine/Source/Runtime/Engine/Classes/Engine/TargetPoint.h
--- a/Engine/Source/Runtime/Engine/Classes/Engine/TargetPoint.h
+++ b/Engine/Source/Runtime/Engine/Classes/Engine/TargetPoint.h
@@ -19,10 +19,10 @@ class ATargetPoint : public AActor
 #if WITH_EDITORONLY_DATA
 private:
 	UPROPERTY(VisibleAnywhere, BlueprintReadOnly, Category = Display, meta = (AllowPrivateAccess = "true"))
-	class UBillboardComponent* SpriteComponent;
+	TObjectPtr<class UBillboardComponent> SpriteComponent;
 
 	UPROPERTY(VisibleAnywhere, BlueprintReadOnly, Category = Display, meta = (AllowPrivateAccess = "true"))
-	class UArrowComponent* ArrowComponent;
+	TObjectPtr<class UArrowComponent> ArrowComponent;
 
 public:
 	/** Returns SpriteComponent subobject **/
diff --git a/Engine/Source/Runtime/Engine/Classes/Engine/TextRenderActor.h b/Engine/Source/Runtime/Engine/Classes/Engine/TextRenderActor.h
--- a/Engine/Source/Runtime/Engine/Classes/Engine/TextRenderActor.h
+++ b/Engine/Source/Runtime/Engine/Classes/Engine/TextRenderActor.h
@@ -23,12 +23,12 @@ class ATextRenderActor : public AActor
 private:
 	/** Component to render a text in 3d with a font */
 	UPROPERTY(Category = TextRenderActor, VisibleAnywhere, BlueprintReadOnly, meta = (ExposeFunctionCategories = "Rendering|Components|TextRender", AllowPrivateAccess = "true"))
-	class UTextRenderComponent* TextRender;
+	TObjectPtr<class UTextRenderComponent> TextRender;
 
 #if WITH_EDITORONLY_DATA
 	// Reference to the billboard component
 	UPROPERTY()
-	UBillboardComponent* SpriteComponent;
+	TObjectPtr<UBillboardComponent> SpriteComponent;
 #endif
 
 public:
diff --git a/Engine/Source/Runtime/Engine/Classes/Engine/Texture.h b/Engine/Source/Runtime/Engine/Classes/Engine/Texture.h
--- a/Engine/Source/Runtime/Engine/Classes/Engine/Texture.h
+++ b/Engine/Source/Runtime/Engine/Classes/Engine/Texture.h
@@ -660,7 +660,7 @@ public:
 	FString SourceFilePath_DEPRECATED;
 
 	UPROPERTY(VisibleAnywhere, Instanced, Category=ImportSettings)
-	class UAssetImportData* AssetImportData;
+	TObjectPtr<class UAssetImportData> AssetImportData;
 
 public:
 
@@ -769,7 +769,7 @@ public:
 	 * Make sure the normal map has at least as many mips as this texture.
 	 */
 	UPROPERTY(EditAnywhere, BlueprintReadWrite, Category=Compositing)
-	class UTexture* CompositeTexture;
+	TObjectPtr<class UTexture> CompositeTexture;
 
 	/* defines how the CompositeTexture is applied, e.g. CTM_RoughnessFromNormalAlpha */
 	UPROPERTY(EditAnywhere, BlueprintReadWrite, Category=Compositing, AdvancedDisplay)
@@ -875,7 +875,7 @@ private:
 protected:
 	/** Array of user data stored with the asset */
 	UPROPERTY(EditAnywhere, AdvancedDisplay, Instanced, Category = Texture)
-	TArray<UAssetUserData*> AssetUserData;
+	TArray<TObjectPtr<UAssetUserData>> AssetUserData;
 
 private:
 	/** The texture's resource, can be NULL */
diff --git a/Engine/Source/Runtime/Engine/Classes/Engine/Texture2DArray.h b/Engine/Source/Runtime/Engine/Classes/Engine/Texture2DArray.h
--- a/Engine/Source/Runtime/Engine/Classes/Engine/Texture2DArray.h
+++ b/Engine/Source/Runtime/Engine/Classes/Engine/Texture2DArray.h
@@ -85,7 +85,7 @@ public:
 #if WITH_EDITORONLY_DATA
 	/** Add Textures*/
 	UPROPERTY(EditAnywhere, Category = Source2D, meta = (DisplayName = "Source Textures"))
-	TArray<UTexture2D*> SourceTextures;
+	TArray<TObjectPtr<UTexture2D>> SourceTextures;
 #endif
 
 	/**
diff --git a/Engine/Source/Runtime/Engine/Classes/Engine/TextureStreamingTypes.h b/Engine/Source/Runtime/Engine/Classes/Engine/TextureStreamingTypes.h
--- a/Engine/Source/Runtime/Engine/Classes/Engine/TextureStreamingTypes.h
+++ b/Engine/Source/Runtime/Engine/Classes/Engine/TextureStreamingTypes.h
@@ -38,7 +38,7 @@ struct FStreamingRenderAssetPrimitiveInfo
 	GENERATED_USTRUCT_BODY()
 
 	UPROPERTY()
-	UStreamableRenderAsset* RenderAsset;
+	TObjectPtr<UStreamableRenderAsset> RenderAsset;
 
 	/** 
 	 * The streaming bounds of the texture/mesh, usually the component material bounds. 
diff --git a/Engine/Source/Runtime/Engine/Classes/Engine/TimelineTemplate.h b/Engine/Source/Runtime/Engine/Classes/Engine/TimelineTemplate.h
--- a/Engine/Source/Runtime/Engine/Classes/Engine/TimelineTemplate.h
+++ b/Engine/Source/Runtime/Engine/Classes/Engine/TimelineTemplate.h
@@ -95,7 +95,7 @@ private:
 public:
 	/** Curve object used to store keys */
 	UPROPERTY()
-	class UCurveFloat* CurveKeys;
+	TObjectPtr<class UCurveFloat> CurveKeys;
 
 	/** Determine if Tracks are the same */
 	ENGINE_API bool operator == (const FTTEventTrack& T2) const;
@@ -130,7 +130,7 @@ struct FTTFloatTrack : public FTTPropertyTrack
 
 	/** Curve object used to define float value over time */
 	UPROPERTY()
-	class UCurveFloat* CurveFloat;
+	TObjectPtr<class UCurveFloat> CurveFloat;
 
 	/** Determine if Tracks are the same */
 	ENGINE_API bool operator == (const FTTFloatTrack& T2) const;
@@ -150,7 +150,7 @@ struct FTTVectorTrack : public FTTPropertyTrack
 
 	/** Curve object used to define vector value over time */
 	UPROPERTY()
-	class UCurveVector* CurveVector;
+	TObjectPtr<class UCurveVector> CurveVector;
 
 	/** Determine if Tracks are the same */
 	ENGINE_API bool operator == (const FTTVectorTrack& T2) const;
@@ -170,7 +170,7 @@ struct FTTLinearColorTrack : public FTTPropertyTrack
 
 	/** Curve object used to define color value over time */
 	UPROPERTY()
-	class UCurveLinearColor* CurveLinearColor;
+	TObjectPtr<class UCurveLinearColor> CurveLinearColor;
 
 	/** Determine if Tracks are the same */
 	ENGINE_API bool operator == (const FTTLinearColorTrack& T2) const;
diff --git a/Engine/Source/Runtime/Engine/Classes/Engine/TriggerBase.h b/Engine/Source/Runtime/Engine/Classes/Engine/TriggerBase.h
--- a/Engine/Source/Runtime/Engine/Classes/Engine/TriggerBase.h
+++ b/Engine/Source/Runtime/Engine/Classes/Engine/TriggerBase.h
@@ -19,12 +19,12 @@ class ATriggerBase : public AActor
 private:
 	/** Shape component used for collision */
 	UPROPERTY(Category = TriggerBase, VisibleAnywhere, BlueprintReadOnly, meta = (AllowPrivateAccess = "true"))
-	UShapeComponent* CollisionComponent;
+	TObjectPtr<UShapeComponent> CollisionComponent;
 
 #if WITH_EDITORONLY_DATA
 	/** Billboard used to see the trigger in the editor */
 	UPROPERTY(Category = TriggerBase, BlueprintReadOnly, meta = (AllowPrivateAccess = "true"))
-	UBillboardComponent* SpriteComponent;
+	TObjectPtr<UBillboardComponent> SpriteComponent;
 #endif
 
 public:
diff --git a/Engine/Source/Runtime/Engine/Classes/Engine/UserDefinedStruct.h b/Engine/Source/Runtime/Engine/Classes/Engine/UserDefinedStruct.h
--- a/Engine/Source/Runtime/Engine/Classes/Engine/UserDefinedStruct.h
+++ b/Engine/Source/Runtime/Engine/Classes/Engine/UserDefinedStruct.h
@@ -62,7 +62,7 @@ public:
 	FString ErrorMessage;
 
 	UPROPERTY()
-	UObject* EditorData;
+	TObjectPtr<UObject> EditorData;
 #endif // WITH_EDITORONLY_DATA
 
 	/** Status of this struct, outside of the editor it is assumed to always be UpToDate */
diff --git a/Engine/Source/Runtime/Engine/Classes/Engine/UserInterfaceSettings.h b/Engine/Source/Runtime/Engine/Classes/Engine/UserInterfaceSettings.h
--- a/Engine/Source/Runtime/Engine/Classes/Engine/UserInterfaceSettings.h
+++ b/Engine/Source/Runtime/Engine/Classes/Engine/UserInterfaceSettings.h
@@ -211,13 +211,13 @@ private:
 	float CalculateScale(FIntPoint Size, bool& bError) const;
 
 	UPROPERTY(Transient)
-	TArray<UObject*> CursorClasses;
+	TArray<TObjectPtr<UObject>> CursorClasses;
 
 	UPROPERTY(Transient)
-	mutable UClass* CustomScalingRuleClassInstance;
+	mutable TObjectPtr<UClass> CustomScalingRuleClassInstance;
 
 	UPROPERTY(Transient)
-	mutable UDPICustomScalingRule* CustomScalingRule;
+	mutable TObjectPtr<UDPICustomScalingRule> CustomScalingRule;
 
 	mutable TOptional<FIntPoint> LastViewportSize;
 	mutable float CalculatedScale;
diff --git a/Engine/Source/Runtime/Engine/Classes/Engine/VolumeTexture.h b/Engine/Source/Runtime/Engine/Classes/Engine/VolumeTexture.h
--- a/Engine/Source/Runtime/Engine/Classes/Engine/VolumeTexture.h
+++ b/Engine/Source/Runtime/Engine/Classes/Engine/VolumeTexture.h
@@ -25,7 +25,7 @@ public:
 #if WITH_EDITORONLY_DATA
 	/** A (optional) reference texture from which the volume texture was built */
 	UPROPERTY(EditAnywhere, Category=Source2D, meta=(DisplayName="Source Texture"))
-	UTexture2D* Source2DTexture;
+	TObjectPtr<UTexture2D> Source2DTexture;
 	/** The lighting Guid of the source 2D texture, used to trigger rebuild when the source changes. */
 	UPROPERTY()
 	FGuid SourceLightingGuid;
diff --git a/Engine/Source/Runtime/Engine/Classes/Engine/WindDirectionalSource.h b/Engine/Source/Runtime/Engine/Classes/Engine/WindDirectionalSource.h
--- a/Engine/Source/Runtime/Engine/Classes/Engine/WindDirectionalSource.h
+++ b/Engine/Source/Runtime/Engine/Classes/Engine/WindDirectionalSource.h
@@ -14,11 +14,11 @@ class ENGINE_API AWindDirectionalSource : public AInfo
 
 private:
 	UPROPERTY(Category = WindDirectionalSource, VisibleAnywhere, BlueprintReadOnly, meta = (AllowPrivateAccess = "true"))
-	class UWindDirectionalSourceComponent* Component;
+	TObjectPtr<class UWindDirectionalSourceComponent> Component;
 
 #if WITH_EDITORONLY_DATA
 	UPROPERTY()
-	class UArrowComponent* ArrowComponent;
+	TObjectPtr<class UArrowComponent> ArrowComponent;
 #endif
 
 public:
