Automated wrapped object pointer upgrade for Engine + ShooterGame + ShooterGame referenced plugins

diff --git a/Engine/Source/Runtime/Engine/Classes/Engine/World.h b/Engine/Source/Runtime/Engine/Classes/Engine/World.h
--- a/Engine/Source/Runtime/Engine/Classes/Engine/World.h
+++ b/Engine/Source/Runtime/Engine/Classes/Engine/World.h
@@ -766,32 +766,32 @@ private:
 	 * since the source collection and the duplicated collection will have their own instances.
 	 */
 	UPROPERTY()
-	class AGameStateBase* GameState;
+	TObjectPtr<class AGameStateBase> GameState;
 
 	/**
 	 * The network driver associated with this collection.
 	 * The source collection and the duplicated collection will have their own instances.
 	 */
 	UPROPERTY()
-	class UNetDriver* NetDriver;
+	TObjectPtr<class UNetDriver> NetDriver;
 
 	/**
 	 * The demo network driver associated with this collection.
 	 * The source collection and the duplicated collection will have their own instances.
 	 */
 	UPROPERTY()
-	class UDemoNetDriver* DemoNetDriver;
+	TObjectPtr<class UDemoNetDriver> DemoNetDriver;
 
 	/**
 	 * The persistent level associated with this collection.
 	 * The source collection and the duplicated collection will have their own instances.
 	 */
 	UPROPERTY()
-	class ULevel* PersistentLevel;
+	TObjectPtr<class ULevel> PersistentLevel;
 
 	/** All the levels in this collection. */
 	UPROPERTY()
-	TSet<ULevel*> Levels;
+	TSet<TObjectPtr<ULevel>> Levels;
 };
 
 template<>
@@ -852,7 +852,7 @@ private:
 
 	/** Priority sorted array of streaming levels actively being considered. */
 	UPROPERTY()
-	TArray<ULevelStreaming*> StreamingLevels;
+	TArray<TObjectPtr<ULevelStreaming>> StreamingLevels;
 
 	enum class EProcessReason : uint8
 	{
@@ -921,48 +921,48 @@ class ENGINE_API UWorld final : public UObject, public FNetworkNotify
 #if WITH_EDITORONLY_DATA
 	/** List of all the layers referenced by the world's actors */
 	UPROPERTY(Transient)
-	TArray< class ULayer* > Layers; 
+	TArray< TObjectPtr<class ULayer> > Layers; 
 
 	// Group actors currently "active"
 	UPROPERTY(Transient)
-	TArray<AActor*> ActiveGroupActors;
+	TArray<TObjectPtr<AActor>> ActiveGroupActors;
 
 	/** Information for thumbnail rendering */
 	UPROPERTY(VisibleAnywhere, Instanced, Category=Thumbnail)
-	class UThumbnailInfo* ThumbnailInfo;
+	TObjectPtr<class UThumbnailInfo> ThumbnailInfo;
 #endif // WITH_EDITORONLY_DATA
 
 	/** Persistent level containing the world info, default brush and actors spawned during gameplay among other things			*/
 	UPROPERTY(Transient)
-	class ULevel*								PersistentLevel;
+	TObjectPtr<class ULevel>								PersistentLevel;
 
 	/** The NAME_GameNetDriver game connection(s) for client/server communication */
 	UPROPERTY(Transient)
-	class UNetDriver*							NetDriver;
+	TObjectPtr<class UNetDriver>							NetDriver;
 
 	/** Line Batchers. All lines to be drawn in the world. */
 	UPROPERTY(Transient)
-	class ULineBatchComponent*					LineBatcher;
+	TObjectPtr<class ULineBatchComponent>					LineBatcher;
 
 	/** Persistent Line Batchers. They don't get flushed every frame.  */
 	UPROPERTY(Transient)
-	class ULineBatchComponent*					PersistentLineBatcher;
+	TObjectPtr<class ULineBatchComponent>					PersistentLineBatcher;
 
 	/** Foreground Line Batchers. This can't be Persistent.  */
 	UPROPERTY(Transient)
-	class ULineBatchComponent*					ForegroundLineBatcher;
+	TObjectPtr<class ULineBatchComponent>					ForegroundLineBatcher;
 
 	/** Instance of this world's game-specific networking management */
 	UPROPERTY(Transient)
-	class AGameNetworkManager*					NetworkManager;
+	TObjectPtr<class AGameNetworkManager>					NetworkManager;
 
 	/** Instance of this world's game-specific physics collision handler */
 	UPROPERTY(Transient)
-	class UPhysicsCollisionHandler*				PhysicsCollisionHandler;
+	TObjectPtr<class UPhysicsCollisionHandler>				PhysicsCollisionHandler;
 
 	/** Array of any additional objects that need to be referenced by this world, to make sure they aren't GC'd */
 	UPROPERTY(Transient)
-	TArray<UObject*>							ExtraReferencedObjects;
+	TArray<TObjectPtr<UObject>>							ExtraReferencedObjects;
 
 	/**
 	 * External modules can have additional data associated with this UWorld.
@@ -970,12 +970,12 @@ class ENGINE_API UWorld final : public UObject, public FNetworkNotify
 	 * loaded/saved by default.
 	 */
 	UPROPERTY(Transient)
-	TArray<UObject*>							PerModuleDataObjects;
+	TArray<TObjectPtr<UObject>>							PerModuleDataObjects;
 
 private:
 	/** Level collection. ULevels are referenced by FName (Package name) to avoid serialized references. Also contains offsets in world units */
 	UPROPERTY(Transient)
-	TArray<ULevelStreaming*> StreamingLevels;
+	TArray<TObjectPtr<ULevelStreaming>> StreamingLevels;
 
 	/** This is the list of streaming levels that are actively being considered for what their state should be. It will be a subset of StreamingLevels */
 	UPROPERTY(Transient, DuplicateTransient)
@@ -1049,21 +1049,21 @@ public:
 private:
 	/** Pointer to the current level in the queue to be made visible, NULL if none are pending. */
 	UPROPERTY(Transient)
-	class ULevel*								CurrentLevelPendingVisibility;
+	TObjectPtr<class ULevel>								CurrentLevelPendingVisibility;
 
 	/** Pointer to the next best candidate level to become the CurrentLevelPendingVisibility. */
 	UPROPERTY(Transient)
-	class ULevel*								NextPreferredLevelPendingVisibility;
+	TObjectPtr<class ULevel>								NextPreferredLevelPendingVisibility;
 
 	/** Pointer to the current level in the queue to be made invisible, NULL if none are pending. */
 	UPROPERTY(Transient)
-	class ULevel*								CurrentLevelPendingInvisibility;
+	TObjectPtr<class ULevel>								CurrentLevelPendingInvisibility;
 
 public:
 	/** NetDriver for capturing network traffic to record demos */
 	UE_DEPRECATED(4.26, "DemoNetDriver will be made private in a future release.  Please use GetDemoNetDriver/SetDemoNetDriver instead.")
 	UPROPERTY()
-	class UDemoNetDriver*						DemoNetDriver;
+	TObjectPtr<class UDemoNetDriver>						DemoNetDriver;
 
 	PRAGMA_DISABLE_DEPRECATION_WARNINGS
 	/** Gets the demo net driver for this world. */
@@ -1075,12 +1075,12 @@ public:
 
 	/** Particle event manager **/
 	UPROPERTY()
-	class AParticleEventManager*				MyParticleEventManager;
+	TObjectPtr<class AParticleEventManager>				MyParticleEventManager;
 
 private:
 	/** DefaultPhysicsVolume used for whole game **/
 	UPROPERTY(Transient)
-	APhysicsVolume*								DefaultPhysicsVolume;
+	TObjectPtr<APhysicsVolume>								DefaultPhysicsVolume;
 
 public:
 
@@ -1234,27 +1234,27 @@ private:
 
 	/** The world's navigation data manager */
 	UPROPERTY(Transient)
-	class UNavigationSystemBase*				NavigationSystem;
+	TObjectPtr<class UNavigationSystemBase>				NavigationSystem;
 
 	/** The current GameMode, valid only on the server */
 	UPROPERTY(Transient)
-	class AGameModeBase*						AuthorityGameMode;
+	TObjectPtr<class AGameModeBase>						AuthorityGameMode;
 
 	/** The replicated actor which contains game state information that can be accessible to clients. Direct access is not allowed, use GetGameState<>() */
 	UPROPERTY(Transient)
-	class AGameStateBase*						GameState;
+	TObjectPtr<class AGameStateBase>						GameState;
 
 	/** The AI System handles generating pathing information and AI behavior */
 	UPROPERTY(Transient)
-	class UAISystemBase*						AISystem;
+	TObjectPtr<class UAISystemBase>						AISystem;
 	
 	/** RVO avoidance manager used by game */
 	UPROPERTY(Transient)
-	class UAvoidanceManager*					AvoidanceManager;
+	TObjectPtr<class UAvoidanceManager>					AvoidanceManager;
 
 	/** Array of levels currently in this world. Not serialized to disk to avoid hard references. */
 	UPROPERTY(Transient)
-	TArray<class ULevel*>						Levels;
+	TArray<TObjectPtr<class ULevel>>						Levels;
 
 	/** Array of level collections currently in this world. */
 	UPROPERTY(Transient, NonTransactional)
@@ -1292,25 +1292,25 @@ private:
 #if WITH_EDITORONLY_DATA
 	/** Pointer to the current level being edited. Level has to be in the Levels array and == PersistentLevel in the game. */
 	UPROPERTY(Transient)
-	class ULevel*								CurrentLevel;
+	TObjectPtr<class ULevel>								CurrentLevel;
 #endif
 
 	UPROPERTY(Transient)
-	class UGameInstance*						OwningGameInstance;
+	TObjectPtr<class UGameInstance>						OwningGameInstance;
 
 	/** Parameter collection instances that hold parameter overrides for this world. */
 	UPROPERTY(Transient)
-	TArray<class UMaterialParameterCollectionInstance*> ParameterCollectionInstances;
+	TArray<TObjectPtr<class UMaterialParameterCollectionInstance>> ParameterCollectionInstances;
 
 	/** 
 	 * Canvas object used for drawing to render targets from blueprint functions eg DrawMaterialToRenderTarget.
 	 * This is cached as UCanvas creation takes >100ms.
 	 */
 	UPROPERTY(Transient)
-	UCanvas* CanvasForRenderingToTarget;
+	TObjectPtr<UCanvas> CanvasForRenderingToTarget;
 
 	UPROPERTY(Transient)
-	UCanvas* CanvasForDrawMaterialToRenderTarget;
+	TObjectPtr<UCanvas> CanvasForDrawMaterialToRenderTarget;
 
 public:
 	/** Set the pointer to the Navigation System instance. */
@@ -1404,21 +1404,21 @@ public:
 
 	/** Physics Field component. */
 	UPROPERTY(Transient)
-	class UPhysicsFieldComponent* PhysicsField;
+	TObjectPtr<class UPhysicsFieldComponent> PhysicsField;
 
 private:
 
 	/** Array of components that need to wait on tasks before end of frame updates */
 	UPROPERTY(Transient, NonTransactional)
-	TSet<UActorComponent*> ComponentsThatNeedPreEndOfFrameSync;
+	TSet<TObjectPtr<UActorComponent>> ComponentsThatNeedPreEndOfFrameSync;
 
 	/** Array of components that need updates at the end of the frame */
 	UPROPERTY(Transient, NonTransactional)
-	TArray<UActorComponent*> ComponentsThatNeedEndOfFrameUpdate;
+	TArray<TObjectPtr<UActorComponent>> ComponentsThatNeedEndOfFrameUpdate;
 
 	/** Array of components that need game thread updates at the end of the frame */
 	UPROPERTY(Transient, NonTransactional)
-	TArray<UActorComponent*> ComponentsThatNeedEndOfFrameUpdate_OnGameThread;
+	TArray<TObjectPtr<UActorComponent>> ComponentsThatNeedEndOfFrameUpdate_OnGameThread;
 
 	/** The state of async tracing - abstracted into its own object for easier reference */
 	FWorldAsyncTraceState AsyncTraceState;
@@ -1525,7 +1525,7 @@ private:
 
 	/** Array of selected levels currently in this world. Not serialized to disk to avoid hard references.	*/
 	UPROPERTY(Transient)
-	TArray<class ULevel*>						SelectedLevels;
+	TArray<TObjectPtr<class ULevel>>						SelectedLevels;
 
 	/** Disables the broadcasting of level selection change. Internal use only. */
 	uint32 bBroadcastSelectionChange:1;
@@ -1631,7 +1631,7 @@ public:
 
 	/** All levels information from which our world is composed */
 	UPROPERTY()
-	class UWorldComposition* WorldComposition;
+	TObjectPtr<class UWorldComposition> WorldComposition;
 	
 	/** Whether we flushing level streaming state */ 
 	EFlushLevelStreamingType FlushLevelStreamingType;
diff --git a/Engine/Source/Runtime/Engine/Classes/Engine/WorldComposition.h b/Engine/Source/Runtime/Engine/Classes/Engine/WorldComposition.h
--- a/Engine/Source/Runtime/Engine/Classes/Engine/WorldComposition.h
+++ b/Engine/Source/Runtime/Engine/Classes/Engine/WorldComposition.h
@@ -223,7 +223,7 @@ private:
 public:
 	// Streaming level objects for each tile
 	UPROPERTY(transient)
-	TArray<ULevelStreaming*>	TilesStreaming;
+	TArray<TObjectPtr<ULevelStreaming>>	TilesStreaming;
 
 	// Time threshold between tile streaming state changes
 	UPROPERTY(config)
diff --git a/Engine/Source/Runtime/Engine/Classes/Exporters/Exporter.h b/Engine/Source/Runtime/Engine/Classes/Exporters/Exporter.h
--- a/Engine/Source/Runtime/Engine/Classes/Exporters/Exporter.h
+++ b/Engine/Source/Runtime/Engine/Classes/Exporters/Exporter.h
@@ -32,7 +32,7 @@ class UExporter : public UObject
 	 * Objects being exported that are contained within ExportRootScope will use just their name instead of a full path
 	 */
 	UPROPERTY(transient)
-	class UObject* ExportRootScope;    // transient
+	TObjectPtr<class UObject> ExportRootScope;    // transient
 
 	/** File extension to use for this exporter */
 	UPROPERTY(BlueprintReadWrite, Category=Misc)
@@ -63,7 +63,7 @@ class UExporter : public UObject
 	uint32 bForceFileOperations:1;
 
 	UPROPERTY(BlueprintReadWrite, Category=Misc)
-	class UAssetExportTask* ExportTask;
+	TObjectPtr<class UAssetExportTask> ExportTask;
 
 	ENGINE_API static		FString	CurrentFilename;
 	/** (debugging purposes only) */
diff --git a/Engine/Source/Runtime/Engine/Classes/GameFramework/Actor.h b/Engine/Source/Runtime/Engine/Classes/GameFramework/Actor.h
--- a/Engine/Source/Runtime/Engine/Classes/GameFramework/Actor.h
+++ b/Engine/Source/Runtime/Engine/Classes/GameFramework/Actor.h
@@ -573,7 +573,7 @@ private:
 	 * @see SetOwner(), GetOwner()
 	 */
 	UPROPERTY(ReplicatedUsing=OnRep_Owner)
-	AActor* Owner;
+	TObjectPtr<AActor> Owner;
 
 protected:
 	/** Used to specify the net driver to replicate on (NAME_None || NAME_GameNetDriver is the default net driver) */
@@ -618,7 +618,7 @@ public:
 
 	/** Component that handles input for this actor, if input is enabled. */
 	UPROPERTY(DuplicateTransient)
-	class UInputComponent* InputComponent;
+	TObjectPtr<class UInputComponent> InputComponent;
 
 	/** Square of the max distance from the client's viewpoint that this actor is relevant and will be replicated. */
 	UPROPERTY(BlueprintReadOnly, EditDefaultsOnly, Category=Replication)
@@ -693,7 +693,7 @@ public:
 private:
 	/** Pawn responsible for damage and other gameplay events caused by this actor. */
 	UPROPERTY(BlueprintReadWrite, ReplicatedUsing=OnRep_Instigator, meta=(ExposeOnSpawn=true, AllowPrivateAccess=true), Category=Actor)
-	class APawn* Instigator;
+	TObjectPtr<class APawn> Instigator;
 
 public:
 	/** Called on clients when Instigator is replicated. */
@@ -702,12 +702,12 @@ public:
 
 	/** Array of all Actors whose Owner is this actor, these are not necessarily spawned by UChildActorComponent */
 	UPROPERTY(Transient)
-	TArray<AActor*> Children;
+	TArray<TObjectPtr<AActor>> Children;
 
 protected:
 	/** The component that defines the transform (location, rotation, scale) of this Actor in the world, all other components must be attached to this one somehow */
 	UPROPERTY(BlueprintGetter=K2_GetRootComponent, Category="Utilities|Transformation")
-	USceneComponent* RootComponent;
+	TObjectPtr<USceneComponent> RootComponent;
 
 #if WITH_EDITORONLY_DATA
 	/** Local space pivot offset for the actor, only used in the editor */
@@ -717,7 +717,7 @@ protected:
 
 	/** The matinee actors that control this actor. */
 	UPROPERTY(Transient)
-	TArray<class AMatineeActor*> ControllingMatineeActors;
+	TArray<TObjectPtr<class AMatineeActor>> ControllingMatineeActors;
 
 	/** Handle for efficient management of LifeSpanExpired timer */
 	FTimerHandle TimerHandle_LifeSpanExpired;
@@ -735,7 +735,7 @@ private:
 #if WITH_EDITORONLY_DATA
 	/** The UHLODLayer in which this actor should be included. */
 	UPROPERTY(EditAnywhere, Category = HLOD, meta = (DisplayName = "HLOD Layer"))
-	class UHLODLayer* HLODLayer;
+	TObjectPtr<class UHLODLayer> HLODLayer;
 #endif
 
 public:
@@ -775,7 +775,7 @@ protected:
 public:
 	/** The editor-only group this actor is a part of. */
 	UPROPERTY(Transient)
-	AActor* GroupActor;
+	TObjectPtr<AActor> GroupActor;
 
 	/** The scale to apply to any billboard components in editor builds (happens in any WITH_EDITOR build, including non-cooked games). */
 	UPROPERTY(EditAnywhere, BlueprintReadWrite, Category=Rendering, meta=(DisplayName="Editor Billboard Scale"))
@@ -3342,12 +3342,12 @@ private:
 
 	/** Array of ActorComponents that have been added by the user on a per-instance basis. */
 	UPROPERTY(Instanced)
-	TArray<UActorComponent*> InstanceComponents;
+	TArray<TObjectPtr<UActorComponent>> InstanceComponents;
 
 public:
 	/** Array of ActorComponents that are created by blueprints and serialized per-instance. */
 	UPROPERTY(TextExportTransient, NonTransactional)
-	TArray<UActorComponent*> BlueprintCreatedComponents;
+	TArray<TObjectPtr<UActorComponent>> BlueprintCreatedComponents;
 
 	/** Adds a component to the instance components array */
 	void AddInstanceComponent(UActorComponent* Component);
diff --git a/Engine/Source/Runtime/Engine/Classes/GameFramework/AsyncActionHandleSaveGame.h b/Engine/Source/Runtime/Engine/Classes/GameFramework/AsyncActionHandleSaveGame.h
--- a/Engine/Source/Runtime/Engine/Classes/GameFramework/AsyncActionHandleSaveGame.h
+++ b/Engine/Source/Runtime/Engine/Classes/GameFramework/AsyncActionHandleSaveGame.h
@@ -65,7 +65,7 @@ protected:
 
 	/** The object that was either saved or loaded */
 	UPROPERTY()
-	USaveGame* SaveGameObject;
+	TObjectPtr<USaveGame> SaveGameObject;
 	
 	/** Function callbacks for load/save */
 	virtual void HandleAsyncSave(const FString& SlotName, const int32 UserIndex, bool bSuccess);
diff --git a/Engine/Source/Runtime/Engine/Classes/GameFramework/Character.h b/Engine/Source/Runtime/Engine/Classes/GameFramework/Character.h
--- a/Engine/Source/Runtime/Engine/Classes/GameFramework/Character.h
+++ b/Engine/Source/Runtime/Engine/Classes/GameFramework/Character.h
@@ -45,7 +45,7 @@ struct FRepRootMotionMontage
 
 	/** AnimMontage providing Root Motion */
 	UPROPERTY()
-	UAnimMontage* AnimMontage = nullptr;
+	TObjectPtr<UAnimMontage> AnimMontage = nullptr;
 
 	/** Track position of Montage */
 	UPROPERTY()
@@ -61,7 +61,7 @@ struct FRepRootMotionMontage
 
 	/** Movement Relative to Base */
 	UPROPERTY()
-	UPrimitiveComponent* MovementBase = nullptr;
+	TObjectPtr<UPrimitiveComponent> MovementBase = nullptr;
 
 	/** Bone on the MovementBase, if a skeletal mesh. */
 	UPROPERTY()
@@ -156,7 +156,7 @@ struct FBasedMovementInfo
 
 	/** Component we are based on */
 	UPROPERTY()
-	UPrimitiveComponent* MovementBase = nullptr;
+	TObjectPtr<UPrimitiveComponent> MovementBase = nullptr;
 
 	/** Bone name on component, for skeletal meshes. NAME_None if not a skeletal mesh or if bone is invalid. */
 	UPROPERTY()
@@ -223,20 +223,20 @@ public:
 private:
 	/** The main skeletal mesh associated with this Character (optional sub-object). */
 	UPROPERTY(Category=Character, VisibleAnywhere, BlueprintReadOnly, meta=(AllowPrivateAccess = "true"))
-	USkeletalMeshComponent* Mesh;
+	TObjectPtr<USkeletalMeshComponent> Mesh;
 
 	/** Movement component used for movement logic in various movement modes (walking, falling, etc), containing relevant settings and functions to control movement. */
 	UPROPERTY(Category=Character, VisibleAnywhere, BlueprintReadOnly, meta=(AllowPrivateAccess = "true"))
-	UCharacterMovementComponent* CharacterMovement;
+	TObjectPtr<UCharacterMovementComponent> CharacterMovement;
 
 	/** The CapsuleComponent being used for movement collision (by CharacterMovement). Always treated as being vertically aligned in simple collision check functions. */
 	UPROPERTY(Category=Character, VisibleAnywhere, BlueprintReadOnly, meta=(AllowPrivateAccess = "true"))
-	UCapsuleComponent* CapsuleComponent;
+	TObjectPtr<UCapsuleComponent> CapsuleComponent;
 
 #if WITH_EDITORONLY_DATA
 	/** Component shown in the editor only to indicate character facing */
 	UPROPERTY()
-	UArrowComponent* ArrowComponent;
+	TObjectPtr<UArrowComponent> ArrowComponent;
 #endif
 
 public:
diff --git a/Engine/Source/Runtime/Engine/Classes/GameFramework/CharacterMovementComponent.h b/Engine/Source/Runtime/Engine/Classes/GameFramework/CharacterMovementComponent.h
--- a/Engine/Source/Runtime/Engine/Classes/GameFramework/CharacterMovementComponent.h
+++ b/Engine/Source/Runtime/Engine/Classes/GameFramework/CharacterMovementComponent.h
@@ -172,7 +172,7 @@ protected:
 
 	/** Character movement component belongs to */
 	UPROPERTY(Transient, DuplicateTransient)
-	ACharacter* CharacterOwner;
+	TObjectPtr<ACharacter> CharacterOwner;
 
 public:
 
@@ -547,7 +547,7 @@ public:
 
 	/** What to update CharacterOwner and UpdatedComponent after movement ends */
 	UPROPERTY()
-	USceneComponent* DeferredUpdatedMoveComponent;
+	TObjectPtr<USceneComponent> DeferredUpdatedMoveComponent;
 
 	/** Maximum step height for getting out of water */
 	UPROPERTY(Category="Character Movement: Swimming", EditAnywhere, BlueprintReadWrite, AdvancedDisplay, meta=(ClampMin="0", UIMin="0"))
diff --git a/Engine/Source/Runtime/Engine/Classes/GameFramework/CheatManager.h b/Engine/Source/Runtime/Engine/Classes/GameFramework/CheatManager.h
--- a/Engine/Source/Runtime/Engine/Classes/GameFramework/CheatManager.h
+++ b/Engine/Source/Runtime/Engine/Classes/GameFramework/CheatManager.h
@@ -91,7 +91,7 @@ class ENGINE_API UCheatManager : public UObject
 
 	/** Debug camera - used to have independent camera without stopping gameplay */
 	UPROPERTY()
-	class ADebugCameraController* DebugCameraControllerRef;
+	TObjectPtr<class ADebugCameraController> DebugCameraControllerRef;
 
 	/** Debug camera - used to have independent camera without stopping gameplay */
 	UPROPERTY(EditDefaultsOnly, BlueprintReadOnly, Category="Debug Camera")
@@ -446,7 +446,7 @@ public:
 protected:
 	/** List of registered cheat manager extensions */
 	UPROPERTY(Transient)
-	TArray<UCheatManagerExtension*> CheatManagerExtensions;
+	TArray<TObjectPtr<UCheatManagerExtension>> CheatManagerExtensions;
 
 	/** Delegate called when the asset manager singleton is created */
 	static FOnCheatManagerCreated OnCheatManagerCreatedDelegate;
diff --git a/Engine/Source/Runtime/Engine/Classes/GameFramework/Controller.h b/Engine/Source/Runtime/Engine/Classes/GameFramework/Controller.h
--- a/Engine/Source/Runtime/Engine/Classes/GameFramework/Controller.h
+++ b/Engine/Source/Runtime/Engine/Classes/GameFramework/Controller.h
@@ -47,7 +47,7 @@ public:
 
 	/** PlayerState containing replicated information about the player using this controller (only exists for players, not NPCs). */
 	UPROPERTY(replicatedUsing = OnRep_PlayerState, BlueprintReadOnly, Category=Controller)
-	APlayerState* PlayerState;
+	TObjectPtr<APlayerState> PlayerState;
 
 	/** Actor marking where this controller spawned in. */
 	TWeakObjectPtr<class AActor> StartSpot;
@@ -63,7 +63,7 @@ public:
 private:
 	/** Pawn currently being controlled by this controller.  Use Pawn.Possess() to take control of a pawn */
 	UPROPERTY(replicatedUsing=OnRep_Pawn)
-	APawn* Pawn;
+	TObjectPtr<APawn> Pawn;
 
 	/**
 	 * Used to track when pawn changes during OnRep_Pawn. 
@@ -73,11 +73,11 @@ private:
 
 	/** Character currently being controlled by this controller.  Value is same as Pawn if the controlled pawn is a character, otherwise nullptr */
 	UPROPERTY()
-	ACharacter* Character;
+	TObjectPtr<ACharacter> Character;
 
 	/** Component to give controllers a transform and enable attachment if desired. */
 	UPROPERTY()
-	USceneComponent* TransformComponent;
+	TObjectPtr<USceneComponent> TransformComponent;
 
 protected:
 	/** Delegate broadcasted when possessing a new pawn or unpossessing one */
diff --git a/Engine/Source/Runtime/Engine/Classes/GameFramework/DebugTextInfo.h b/Engine/Source/Runtime/Engine/Classes/GameFramework/DebugTextInfo.h
--- a/Engine/Source/Runtime/Engine/Classes/GameFramework/DebugTextInfo.h
+++ b/Engine/Source/Runtime/Engine/Classes/GameFramework/DebugTextInfo.h
@@ -22,7 +22,7 @@ struct FDebugTextInfo
 
 	/**  AActor related to text item */
 	UPROPERTY()
-	AActor* SrcActor;
+	TObjectPtr<AActor> SrcActor;
 
 	/** Offset from SrcActor.Location to apply */
 	UPROPERTY()
@@ -66,7 +66,7 @@ struct FDebugTextInfo
 
 	/** The Font which to display this as.  Will Default to GetSmallFont()**/
 	UPROPERTY()
-	UFont* Font;
+	TObjectPtr<UFont> Font;
 
 	/** Scale to apply to font when rendering */
 	UPROPERTY()
diff --git a/Engine/Source/Runtime/Engine/Classes/GameFramework/DefaultPawn.h b/Engine/Source/Runtime/Engine/Classes/GameFramework/DefaultPawn.h
--- a/Engine/Source/Runtime/Engine/Classes/GameFramework/DefaultPawn.h
+++ b/Engine/Source/Runtime/Engine/Classes/GameFramework/DefaultPawn.h
@@ -85,7 +85,7 @@ public:
 protected:
 	/** DefaultPawn movement component */
 	UPROPERTY(Category = Pawn, VisibleAnywhere, BlueprintReadOnly, meta = (AllowPrivateAccess = "true"))
-	UPawnMovementComponent* MovementComponent;
+	TObjectPtr<UPawnMovementComponent> MovementComponent;
 
 public:
 	/** Name of the CollisionComponent. */
@@ -94,7 +94,7 @@ public:
 private:
 	/** DefaultPawn collision component */
 	UPROPERTY(Category = Pawn, VisibleAnywhere, BlueprintReadOnly, meta = (AllowPrivateAccess = "true"))
-	USphereComponent* CollisionComponent;
+	TObjectPtr<USphereComponent> CollisionComponent;
 public:
 
 	/** Name of the MeshComponent. Use this name if you want to prevent creation of the component (with ObjectInitializer.DoNotCreateDefaultSubobject). */
@@ -103,7 +103,7 @@ public:
 private:
 	/** The mesh associated with this Pawn. */
 	UPROPERTY(Category = Pawn, VisibleAnywhere, BlueprintReadOnly, meta = (AllowPrivateAccess = "true"))
-	UStaticMeshComponent* MeshComponent;
+	TObjectPtr<UStaticMeshComponent> MeshComponent;
 public:
 
 	/** If true, adds default input bindings for movement and camera look. */
diff --git a/Engine/Source/Runtime/Engine/Classes/GameFramework/ForceFeedbackEffect.h b/Engine/Source/Runtime/Engine/Classes/GameFramework/ForceFeedbackEffect.h
--- a/Engine/Source/Runtime/Engine/Classes/GameFramework/ForceFeedbackEffect.h
+++ b/Engine/Source/Runtime/Engine/Classes/GameFramework/ForceFeedbackEffect.h
@@ -70,7 +70,7 @@ struct ENGINE_API FActiveForceFeedbackEffect
 	GENERATED_USTRUCT_BODY()
 
 	UPROPERTY()
-	class UForceFeedbackEffect* ForceFeedbackEffect;
+	TObjectPtr<class UForceFeedbackEffect> ForceFeedbackEffect;
 
 	FForceFeedbackParameters Parameters;
 	float PlayTime;
diff --git a/Engine/Source/Runtime/Engine/Classes/GameFramework/GameMode.h b/Engine/Source/Runtime/Engine/Classes/GameFramework/GameMode.h
--- a/Engine/Source/Runtime/Engine/Classes/GameFramework/GameMode.h
+++ b/Engine/Source/Runtime/Engine/Classes/GameFramework/GameMode.h
@@ -135,7 +135,7 @@ public:
 
 	/** PlayerStates of players who have disconnected from the server (saved in case they reconnect) */
 	UPROPERTY()
-	TArray<class APlayerState*> InactivePlayerArray;    
+	TArray<TObjectPtr<class APlayerState>> InactivePlayerArray;    
 
 protected:
 
diff --git a/Engine/Source/Runtime/Engine/Classes/GameFramework/GameModeBase.h b/Engine/Source/Runtime/Engine/Classes/GameFramework/GameModeBase.h
--- a/Engine/Source/Runtime/Engine/Classes/GameFramework/GameModeBase.h
+++ b/Engine/Source/Runtime/Engine/Classes/GameFramework/GameModeBase.h
@@ -122,11 +122,11 @@ public:
 
 	/** Game Session handles login approval, arbitration, online game interface */
 	UPROPERTY(Transient)
-	AGameSession* GameSession;
+	TObjectPtr<AGameSession> GameSession;
 
 	/** GameState is used to replicate game state relevant properties to all clients. */
 	UPROPERTY(Transient)
-	AGameStateBase* GameState;
+	TObjectPtr<AGameStateBase> GameState;
 
 	/** Helper template to returns the current GameState casted to the desired type. */
 	template< class T >
@@ -136,7 +136,7 @@ public:
 	}
 
 	UPROPERTY(Transient)
-	AServerStatReplicator* ServerStatReplicator;
+	TObjectPtr<AServerStatReplicator> ServerStatReplicator;
 
 	/** Returns number of active human players, excluding spectators */
 	UFUNCTION(BlueprintCallable, Category=Game)
diff --git a/Engine/Source/Runtime/Engine/Classes/GameFramework/GameStateBase.h b/Engine/Source/Runtime/Engine/Classes/GameFramework/GameStateBase.h
--- a/Engine/Source/Runtime/Engine/Classes/GameFramework/GameStateBase.h
+++ b/Engine/Source/Runtime/Engine/Classes/GameFramework/GameStateBase.h
@@ -44,7 +44,7 @@ public:
 
 	/** Instance of the current game mode, exists only on the server. For non-authority clients, this will be NULL. */
 	UPROPERTY(Transient, BlueprintReadOnly, Category=GameState)
-	AGameModeBase* AuthorityGameMode;
+	TObjectPtr<AGameModeBase> AuthorityGameMode;
 
 	/** Class used by spectators, assigned by GameModeBase. */
 	UPROPERTY(Transient, BlueprintReadOnly, Category=GameState, ReplicatedUsing=OnRep_SpectatorClass)
@@ -52,7 +52,7 @@ public:
 
 	/** Array of all PlayerStates, maintained on both server and clients (PlayerStates are always relevant) */
 	UPROPERTY(Transient, BlueprintReadOnly, Category=GameState)
-	TArray<APlayerState*> PlayerArray;
+	TArray<TObjectPtr<APlayerState>> PlayerArray;
 
 	/** Allow game states to react to asset packages being loaded asynchronously */
 	virtual void AsyncPackageLoaded(UObject* Package) {}
diff --git a/Engine/Source/Runtime/Engine/Classes/GameFramework/HUD.h b/Engine/Source/Runtime/Engine/Classes/GameFramework/HUD.h
--- a/Engine/Source/Runtime/Engine/Classes/GameFramework/HUD.h
+++ b/Engine/Source/Runtime/Engine/Classes/GameFramework/HUD.h
@@ -39,7 +39,7 @@ class ENGINE_API AHUD : public AActor
 
 	/** PlayerController which owns this HUD. */
 	UPROPERTY(BlueprintReadOnly, Category=HUD)
-	APlayerController* PlayerOwner;    
+	TObjectPtr<APlayerController> PlayerOwner;    
 
 	/** Tells whether the game was paused due to lost focus */
 	UPROPERTY(BlueprintReadOnly, Category=HUD)
@@ -76,7 +76,7 @@ private:
 public:
 	/** Holds a list of Actors that need PostRender() calls. */
 	UPROPERTY()
-	TArray<AActor*> PostRenderedActors;
+	TArray<TObjectPtr<AActor>> PostRenderedActors;
 
 	/** Used to calculate delta time between HUD rendering. */
 	float LastHUDRenderTime;
@@ -95,11 +95,11 @@ public:
 protected:
 	/** Canvas to Draw HUD on.  Only valid during PostRender() event.  */
 	UPROPERTY()
-	UCanvas* Canvas;
+	TObjectPtr<UCanvas> Canvas;
 
 	/** 'Foreground' debug canvas, will draw in front of Slate UI. */
 	UPROPERTY()
-	UCanvas* DebugCanvas;
+	TObjectPtr<UCanvas> DebugCanvas;
 
 	/** List of debug strings attached to actors, sorted by actor first, then by order of addition */
 	UPROPERTY()
@@ -137,7 +137,7 @@ protected:
 
 	/** Show Debug Actor used if 'bShowDebugForReticleTarget' is true, only updated if trace from reticle hit a new Actor of class 'ShowDebugTargetDesiredClass'*/
 	UPROPERTY()
-	AActor* ShowDebugTargetActor;
+	TObjectPtr<AActor> ShowDebugTargetActor;
 
 public:
 	/**
diff --git a/Engine/Source/Runtime/Engine/Classes/GameFramework/Info.h b/Engine/Source/Runtime/Engine/Classes/GameFramework/Info.h
--- a/Engine/Source/Runtime/Engine/Classes/GameFramework/Info.h
+++ b/Engine/Source/Runtime/Engine/Classes/GameFramework/Info.h
@@ -25,7 +25,7 @@ class AInfo : public AActor
 private:
 	/** Billboard Component displayed in editor */
 	UPROPERTY()
-	class UBillboardComponent* SpriteComponent;
+	TObjectPtr<class UBillboardComponent> SpriteComponent;
 public:
 #endif
 
diff --git a/Engine/Source/Runtime/Engine/Classes/GameFramework/LightWeightInstanceStaticMeshManager.h b/Engine/Source/Runtime/Engine/Classes/GameFramework/LightWeightInstanceStaticMeshManager.h
--- a/Engine/Source/Runtime/Engine/Classes/GameFramework/LightWeightInstanceStaticMeshManager.h
+++ b/Engine/Source/Runtime/Engine/Classes/GameFramework/LightWeightInstanceStaticMeshManager.h
@@ -66,7 +66,7 @@ protected:
 	void OnRep_StaticMesh();
 
 	UPROPERTY(EditAnywhere, BlueprintReadOnly, Category = Debug, AdvancedDisplay, meta = (BlueprintProtected = "true", AllowPrivateAccess = "true"))
-	class UHierarchicalInstancedStaticMeshComponent* InstancedStaticMeshComponent;
+	TObjectPtr<class UHierarchicalInstancedStaticMeshComponent> InstancedStaticMeshComponent;
 
 	//
 	// Bookkeeping info
diff --git a/Engine/Source/Runtime/Engine/Classes/GameFramework/LocalMessage.h b/Engine/Source/Runtime/Engine/Classes/GameFramework/LocalMessage.h
--- a/Engine/Source/Runtime/Engine/Classes/GameFramework/LocalMessage.h
+++ b/Engine/Source/Runtime/Engine/Classes/GameFramework/LocalMessage.h
@@ -30,7 +30,7 @@ struct ENGINE_API FClientReceiveData
 	GENERATED_USTRUCT_BODY()
 
 	UPROPERTY()
-	APlayerController* LocalPC;
+	TObjectPtr<APlayerController> LocalPC;
 
 	UPROPERTY()
 	FName MessageType;
@@ -42,13 +42,13 @@ struct ENGINE_API FClientReceiveData
 	FString MessageString;
 
 	UPROPERTY()
-	APlayerState* RelatedPlayerState_1;
+	TObjectPtr<APlayerState> RelatedPlayerState_1;
 
 	UPROPERTY()
-	APlayerState* RelatedPlayerState_2;
+	TObjectPtr<APlayerState> RelatedPlayerState_2;
 
 	UPROPERTY()
-	UObject* OptionalObject;
+	TObjectPtr<UObject> OptionalObject;
 
 	FClientReceiveData();
 };
diff --git a/Engine/Source/Runtime/Engine/Classes/GameFramework/MovementComponent.h b/Engine/Source/Runtime/Engine/Classes/GameFramework/MovementComponent.h
--- a/Engine/Source/Runtime/Engine/Classes/GameFramework/MovementComponent.h
+++ b/Engine/Source/Runtime/Engine/Classes/GameFramework/MovementComponent.h
@@ -63,13 +63,13 @@ class ENGINE_API UMovementComponent : public UActorComponent
 	 * @see bAutoRegisterUpdatedComponent, SetUpdatedComponent(), UpdatedPrimitive
 	 */
 	UPROPERTY(BlueprintReadOnly, Transient, DuplicateTransient, Category=MovementComponent)
-	USceneComponent* UpdatedComponent;
+	TObjectPtr<USceneComponent> UpdatedComponent;
 
 	/**
 	 * UpdatedComponent, cast as a UPrimitiveComponent. May be invalid if UpdatedComponent was null or not a UPrimitiveComponent.
 	 */
 	UPROPERTY(BlueprintReadOnly, Transient, DuplicateTransient, Category=MovementComponent)
-	UPrimitiveComponent* UpdatedPrimitive;
+	TObjectPtr<UPrimitiveComponent> UpdatedPrimitive;
 
 	/**
 	 * Flags that control the behavior of calls to MoveComponent() on our UpdatedComponent.
diff --git a/Engine/Source/Runtime/Engine/Classes/GameFramework/NavMovementComponent.h b/Engine/Source/Runtime/Engine/Classes/GameFramework/NavMovementComponent.h
--- a/Engine/Source/Runtime/Engine/Classes/GameFramework/NavMovementComponent.h
+++ b/Engine/Source/Runtime/Engine/Classes/GameFramework/NavMovementComponent.h
@@ -59,7 +59,7 @@ private:
 	/** object implementing IPathFollowingAgentInterface. Is private to control access to it.
 	 *	@see SetPathFollowingAgent, GetPathFollowingAgent */
 	UPROPERTY()
-	UObject* PathFollowingComp;
+	TObjectPtr<UObject> PathFollowingComp;
 
 public:
 	/** Stops applying further movement (usually zeros acceleration). */
diff --git a/Engine/Source/Runtime/Engine/Classes/GameFramework/PainCausingVolume.h b/Engine/Source/Runtime/Engine/Classes/GameFramework/PainCausingVolume.h
--- a/Engine/Source/Runtime/Engine/Classes/GameFramework/PainCausingVolume.h
+++ b/Engine/Source/Runtime/Engine/Classes/GameFramework/PainCausingVolume.h
@@ -49,7 +49,7 @@ class ENGINE_API APainCausingVolume : public APhysicsVolume
 
 	/** Controller that gets credit for any damage caused by this volume */
 	UPROPERTY()
-	class AController* DamageInstigator;
+	TObjectPtr<class AController> DamageInstigator;
 
 	/** Damage overlapping actors if pain-causing. */
 	virtual void PainTimer();
diff --git a/Engine/Source/Runtime/Engine/Classes/GameFramework/Pawn.h b/Engine/Source/Runtime/Engine/Classes/GameFramework/Pawn.h
--- a/Engine/Source/Runtime/Engine/Classes/GameFramework/Pawn.h
+++ b/Engine/Source/Runtime/Engine/Classes/GameFramework/Pawn.h
@@ -132,7 +132,7 @@ public:
 private:
 	/** If Pawn is possessed by a player, points to its Player State.  Needed for network play as controllers are not replicated to clients. */
 	UPROPERTY(replicatedUsing=OnRep_PlayerState, BlueprintReadOnly, Category=Pawn, meta=(AllowPrivateAccess="true"))
-	APlayerState* PlayerState;
+	TObjectPtr<APlayerState> PlayerState;
 
 public:
 
@@ -155,11 +155,11 @@ public:
 
 	/** Controller of the last Actor that caused us damage. */
 	UPROPERTY(BlueprintReadOnly, transient, Category="Pawn")
-	AController* LastHitBy;
+	TObjectPtr<AController> LastHitBy;
 
 	/** Controller currently possessing this Actor */
 	UPROPERTY(replicatedUsing=OnRep_Controller)
-	AController* Controller;
+	TObjectPtr<AController> Controller;
 
 	/** Max difference between pawn's Rotation.Yaw and GetDesiredRotation().Yaw for pawn to be considered as having reached its desired rotation */
 	float AllowedYawError;
diff --git a/Engine/Source/Runtime/Engine/Classes/GameFramework/PawnMovementComponent.h b/Engine/Source/Runtime/Engine/Classes/GameFramework/PawnMovementComponent.h
--- a/Engine/Source/Runtime/Engine/Classes/GameFramework/PawnMovementComponent.h
+++ b/Engine/Source/Runtime/Engine/Classes/GameFramework/PawnMovementComponent.h
@@ -86,7 +86,7 @@ protected:
 
 	/** Pawn that owns this component. */
 	UPROPERTY(Transient, DuplicateTransient)
-	class APawn* PawnOwner;
+	TObjectPtr<class APawn> PawnOwner;
 
 	/**
 	 * Attempts to mark the PlayerCameraManager as dirty.
diff --git a/Engine/Source/Runtime/Engine/Classes/GameFramework/PlayerController.h b/Engine/Source/Runtime/Engine/Classes/GameFramework/PlayerController.h
--- a/Engine/Source/Runtime/Engine/Classes/GameFramework/PlayerController.h
+++ b/Engine/Source/Runtime/Engine/Classes/GameFramework/PlayerController.h
@@ -230,26 +230,26 @@ public:
 
 	/** UPlayer associated with this PlayerController.  Could be a local player or a net connection. */
 	UPROPERTY()
-	UPlayer* Player;
+	TObjectPtr<UPlayer> Player;
 
 	/** Used in net games so client can acknowledge it possessed a specific pawn. */
 	UPROPERTY()
-	APawn* AcknowledgedPawn;
+	TObjectPtr<APawn> AcknowledgedPawn;
 
 	/** Director track that's currently possessing this player controller, or none if not possessed. */
 	UPROPERTY(Transient)
-	UInterpTrackInstDirector* ControllingDirTrackInst;
+	TObjectPtr<UInterpTrackInstDirector> ControllingDirTrackInst;
 
 	/** Heads up display associated with this PlayerController. */
 	UPROPERTY()
-	AHUD* MyHUD;
+	TObjectPtr<AHUD> MyHUD;
 
 	// ******************************************************************************
 	// Camera/view related variables
 
 	/** Camera manager associated with this Player Controller. */
 	UPROPERTY(BlueprintReadOnly, Category=PlayerController)
-	APlayerCameraManager* PlayerCameraManager;
+	TObjectPtr<APlayerCameraManager> PlayerCameraManager;
 
 	/** PlayerCamera class should be set for each game, otherwise Engine.PlayerCameraManager is used */
 	UPROPERTY(EditAnywhere, BlueprintReadOnly, Category=PlayerController)
@@ -282,7 +282,7 @@ public:
 
 	/** The actors which the camera shouldn't see - e.g. used to hide actors which the camera penetrates */
 	UPROPERTY()
-	TArray<class AActor*> HiddenActors;
+	TArray<TObjectPtr<class AActor>> HiddenActors;
 
 	/** Explicit components the camera shouldn't see (helpful for external systems to hide a component from a single player) */
 	UPROPERTY()
@@ -317,7 +317,7 @@ public:
 	 * This behavior can be changed either by overriding APlayerController::EnableCheats or AGameModeBase::AllowCheats.
 	 */
 	UPROPERTY(Transient, BlueprintReadOnly, Category="Cheat Manager")
-	UCheatManager* CheatManager;
+	TObjectPtr<UCheatManager> CheatManager;
 	
 	/**
 	 * Class of my CheatManager.
@@ -328,7 +328,7 @@ public:
 
 	/** Object that manages player input. */
 	UPROPERTY(Transient)
-	class UPlayerInput* PlayerInput;    
+	TObjectPtr<class UPlayerInput> PlayerInput;    
 	
 	UPROPERTY(Transient)
 	TArray<FActiveForceFeedbackEffect> ActiveForceFeedbackEffects;
@@ -428,11 +428,11 @@ public:
 	 * @see GameModeBase::SwapPlayerControllers()
 	 */
 	UPROPERTY(DuplicateTransient)
-	UNetConnection* PendingSwapConnection;
+	TObjectPtr<UNetConnection> PendingSwapConnection;
 
 	/** The net connection this controller is communicating on, nullptr for local players on server */
 	UPROPERTY(DuplicateTransient)
-	UNetConnection* NetConnection;
+	TObjectPtr<UNetConnection> NetConnection;
 
 	/** Input axes values, accumulated each tick. */
 	FRotator RotationInput;
@@ -1480,7 +1480,7 @@ protected:
 	
 	/** InputComponent we use when player is in Inactive state. */
 	UPROPERTY()
-	UInputComponent* InactiveStateInputComponent;
+	TObjectPtr<UInputComponent> InactiveStateInputComponent;
 
 	/** Sets up input bindings for the input component pushed on the stack in the inactive state. */
 	virtual void SetupInactiveStateInputComponent(UInputComponent* InComponent);
@@ -1510,7 +1510,7 @@ protected:
 
 	/** The currently set touch interface */
 	UPROPERTY()
-	class UTouchInterface* CurrentTouchInterface;
+	TObjectPtr<class UTouchInterface> CurrentTouchInterface;
 
 	/** Handle for efficient management of UnFreeze timer */
 	FTimerHandle TimerHandle_UnFreeze;
@@ -1931,7 +1931,7 @@ protected:
 private:
 	/** The pawn used when spectating (nullptr if not spectating). */
 	UPROPERTY()
-	ASpectatorPawn* SpectatorPawn;
+	TObjectPtr<ASpectatorPawn> SpectatorPawn;
 
 	/** Used to delay calling ClientRestart() again when it hasn't been appropriately acknowledged. */
 	float		LastRetryPlayerTime;
diff --git a/Engine/Source/Runtime/Engine/Classes/GameFramework/PlayerStart.h b/Engine/Source/Runtime/Engine/Classes/GameFramework/PlayerStart.h
--- a/Engine/Source/Runtime/Engine/Classes/GameFramework/PlayerStart.h
+++ b/Engine/Source/Runtime/Engine/Classes/GameFramework/PlayerStart.h
@@ -30,7 +30,7 @@ public:
 #if WITH_EDITORONLY_DATA
 private:
 	UPROPERTY()
-	class UArrowComponent* ArrowComponent;
+	TObjectPtr<class UArrowComponent> ArrowComponent;
 public:
 #endif
 
diff --git a/Engine/Source/Runtime/Engine/Classes/GameFramework/PlayerState.h b/Engine/Source/Runtime/Engine/Classes/GameFramework/PlayerState.h
--- a/Engine/Source/Runtime/Engine/Classes/GameFramework/PlayerState.h
+++ b/Engine/Source/Runtime/Engine/Classes/GameFramework/PlayerState.h
@@ -165,7 +165,7 @@ private:
 
 	/** The pawn that is controlled by by this player state. */
 	UPROPERTY(BlueprintReadOnly, Category=PlayerState, meta=(AllowPrivateAccess="true"))
-	APawn* PawnPrivate;
+	TObjectPtr<APawn> PawnPrivate;
 
 	/**
 	 * Stores the last 4 seconds worth of ping data (one second per 'bucket').
diff --git a/Engine/Source/Runtime/Engine/Classes/GameFramework/RootMotionSource.h b/Engine/Source/Runtime/Engine/Classes/GameFramework/RootMotionSource.h
--- a/Engine/Source/Runtime/Engine/Classes/GameFramework/RootMotionSource.h
+++ b/Engine/Source/Runtime/Engine/Classes/GameFramework/RootMotionSource.h
@@ -418,7 +418,7 @@ struct ENGINE_API FRootMotionSource_ConstantForce : public FRootMotionSource
 	FVector Force;
 
 	UPROPERTY()
-	UCurveFloat* StrengthOverTime;
+	TObjectPtr<UCurveFloat> StrengthOverTime;
 
 	virtual FRootMotionSource* Clone() const override;
 
@@ -469,7 +469,7 @@ struct ENGINE_API FRootMotionSource_RadialForce : public FRootMotionSource
 	FVector Location;
 
 	UPROPERTY()
-	AActor* LocationActor;
+	TObjectPtr<AActor> LocationActor;
 
 	UPROPERTY()
 	float Radius;
@@ -484,10 +484,10 @@ struct ENGINE_API FRootMotionSource_RadialForce : public FRootMotionSource
 	bool bNoZForce;
 
 	UPROPERTY()
-	UCurveFloat* StrengthDistanceFalloff;
+	TObjectPtr<UCurveFloat> StrengthDistanceFalloff;
 
 	UPROPERTY()
-	UCurveFloat* StrengthOverTime;
+	TObjectPtr<UCurveFloat> StrengthOverTime;
 
 	UPROPERTY()
 	bool bUseFixedWorldDirection;
@@ -550,7 +550,7 @@ struct ENGINE_API FRootMotionSource_MoveToForce : public FRootMotionSource
 	bool bRestrictSpeedToExpected;
 
 	UPROPERTY()
-	UCurveVector* PathOffsetCurve;
+	TObjectPtr<UCurveVector> PathOffsetCurve;
 
 	FVector GetPathOffsetInWorldSpace(float MoveFraction) const;
 
@@ -617,10 +617,10 @@ struct ENGINE_API FRootMotionSource_MoveToDynamicForce : public FRootMotionSourc
 	bool bRestrictSpeedToExpected;
 
 	UPROPERTY()
-	UCurveVector* PathOffsetCurve;
+	TObjectPtr<UCurveVector> PathOffsetCurve;
 
 	UPROPERTY()
-	UCurveFloat* TimeMappingCurve;
+	TObjectPtr<UCurveFloat> TimeMappingCurve;
 
 	void SetTargetLocation(FVector NewTargetLocation);
 
@@ -686,10 +686,10 @@ struct ENGINE_API FRootMotionSource_JumpForce : public FRootMotionSource
 	bool bDisableTimeout;
 
 	UPROPERTY()
-	UCurveVector* PathOffsetCurve;
+	TObjectPtr<UCurveVector> PathOffsetCurve;
 
 	UPROPERTY()
-	UCurveFloat* TimeMappingCurve;
+	TObjectPtr<UCurveFloat> TimeMappingCurve;
 
 	FVector SavedHalfwayLocation;
 
diff --git a/Engine/Source/Runtime/Engine/Classes/GameFramework/TouchInterface.h b/Engine/Source/Runtime/Engine/Classes/GameFramework/TouchInterface.h
--- a/Engine/Source/Runtime/Engine/Classes/GameFramework/TouchInterface.h
+++ b/Engine/Source/Runtime/Engine/Classes/GameFramework/TouchInterface.h
@@ -18,9 +18,9 @@ struct FTouchInputControl
 
 	// basically mirroring SVirtualJoystick::FControlInfo but as an editable class
 	UPROPERTY(EditAnywhere, Category="Control", meta=(ToolTip="For sticks, this is the Thumb"))
-	UTexture2D* Image1;
+	TObjectPtr<UTexture2D> Image1;
 	UPROPERTY(EditAnywhere, Category="Control", meta=(ToolTip="For sticks, this is the Background"))
-	UTexture2D* Image2;
+	TObjectPtr<UTexture2D> Image2;
 	UPROPERTY(EditAnywhere, Category="Control", meta=(ToolTip="The center point of the control (if <= 1.0, it's relative to screen, > 1.0 is absolute)"))
 	FVector2D Center;
 	UPROPERTY(EditAnywhere, Category="Control", meta=(ToolTip="The size of the control (if <= 1.0, it's relative to screen, > 1.0 is absolute)"))
diff --git a/Engine/Source/Runtime/Engine/Classes/GameFramework/WorldSettings.h b/Engine/Source/Runtime/Engine/Classes/GameFramework/WorldSettings.h
--- a/Engine/Source/Runtime/Engine/Classes/GameFramework/WorldSettings.h
+++ b/Engine/Source/Runtime/Engine/Classes/GameFramework/WorldSettings.h
@@ -234,15 +234,15 @@ struct ENGINE_API FNetViewer
 	GENERATED_USTRUCT_BODY()
 
 	UPROPERTY()
-	UNetConnection* Connection;
+	TObjectPtr<UNetConnection> Connection;
 
 	/** The "controlling net object" associated with this view (typically player controller) */
 	UPROPERTY()
-	class AActor* InViewer;
+	TObjectPtr<class AActor> InViewer;
 
 	/** The actor that is being directly viewed, usually a pawn.  Could also be the net actor of consequence */
 	UPROPERTY()
-	class AActor* ViewTarget;
+	TObjectPtr<class AActor> ViewTarget;
 
 	/** Where the viewer is looking from */
 	UPROPERTY()
@@ -539,14 +539,14 @@ protected:
 	 *	in NavigationSystem instance not being created for this world. Note that
 	 *	if set NavigationSystemConfigOverride will be used instead. */
 	UPROPERTY(EditAnywhere, BlueprintReadOnly, Category=World, AdvancedDisplay, Instanced, NoClear, meta=(NoResetToDefault))
-	UNavigationSystemConfig* NavigationSystemConfig;
+	TObjectPtr<UNavigationSystemConfig> NavigationSystemConfig;
 
 	/** Overrides NavigationSystemConfig. */
 	UPROPERTY(Transient)
-	UNavigationSystemConfig* NavigationSystemConfigOverride;
+	TObjectPtr<UNavigationSystemConfig> NavigationSystemConfigOverride;
 
 	UPROPERTY(EditAnywhere, Category = WorldPartition, NoClear, meta = (NoResetToDefault), Instanced)
-	UWorldPartition* WorldPartition;
+	TObjectPtr<UWorldPartition> WorldPartition;
 
 public:
 	/** Size of the grid for instanced foliage actors*/
@@ -647,7 +647,7 @@ public:
 
 	/** Default Base SoundMix.																			*/
 	UPROPERTY(EditAnywhere, Category=Audio)
-	class USoundMix* DefaultBaseSoundMix;
+	TObjectPtr<class USoundMix> DefaultBaseSoundMix;
 
 #if WITH_EDITORONLY_DATA
 	/** If set overrides the level settings and global project settings */
@@ -681,7 +681,7 @@ public:
 	/** Level Bookmarks: 10 should be MAX_BOOKMARK_NUMBER @fixmeconst */
 	UE_DEPRECATED(4.21, "This member will be removed. Please use the Bookmark accessor functions instead.")
 	UPROPERTY()
-	class UBookMark* BookMarks[10];
+	TObjectPtr<class UBookMark> BookMarks[10];
 #endif
 
 	/************************************/
@@ -724,7 +724,7 @@ public:
 	// If paused, FName of person pausing the game.
 	UE_DEPRECATED(4.23, "This property is deprecated. Please use Get/SetPauserPlayerState().")
 	UPROPERTY(transient)
-	class APlayerState* Pauser;
+	TObjectPtr<class APlayerState> Pauser;
 
 	/** valid only during replication - information about the player(s) being replicated to
 	 * (there could be more than one in the case of a splitscreen client)
@@ -742,11 +742,11 @@ public:
 
 	/** Array of user data stored with the asset */
 	UPROPERTY()
-	TArray<UAssetUserData*> AssetUserData;
+	TArray<TObjectPtr<UAssetUserData>> AssetUserData;
 
 	// If paused, PlayerState of person pausing the game.
 	UPROPERTY(transient, replicated)
-	class APlayerState* PauserPlayerState;
+	TObjectPtr<class APlayerState> PauserPlayerState;
 
 public:
 	//~ Begin UObject Interface.
@@ -874,7 +874,7 @@ private:
 	TSubclassOf<class UBookmarkBase> DefaultBookmarkClass;
 
 	UPROPERTY()
-	TArray<class UBookmarkBase*> BookmarkArray;
+	TArray<TObjectPtr<class UBookmarkBase>> BookmarkArray;
 
 	// Tracked so we can detect changes from Config
 	UPROPERTY()
diff --git a/Engine/Source/Runtime/Engine/Classes/Haptics/HapticFeedbackEffect_Base.h b/Engine/Source/Runtime/Engine/Classes/Haptics/HapticFeedbackEffect_Base.h
--- a/Engine/Source/Runtime/Engine/Classes/Haptics/HapticFeedbackEffect_Base.h
+++ b/Engine/Source/Runtime/Engine/Classes/Haptics/HapticFeedbackEffect_Base.h
@@ -28,7 +28,7 @@ struct ENGINE_API FActiveHapticFeedbackEffect
 	GENERATED_USTRUCT_BODY()
 
 	UPROPERTY()
-	class UHapticFeedbackEffect_Base* HapticEffect;
+	TObjectPtr<class UHapticFeedbackEffect_Base> HapticEffect;
 
 	FActiveHapticFeedbackEffect()
 		: HapticEffect(nullptr)
diff --git a/Engine/Source/Runtime/Engine/Classes/Haptics/HapticFeedbackEffect_SoundWave.h b/Engine/Source/Runtime/Engine/Classes/Haptics/HapticFeedbackEffect_SoundWave.h
--- a/Engine/Source/Runtime/Engine/Classes/Haptics/HapticFeedbackEffect_SoundWave.h
+++ b/Engine/Source/Runtime/Engine/Classes/Haptics/HapticFeedbackEffect_SoundWave.h
@@ -16,7 +16,7 @@ class UHapticFeedbackEffect_SoundWave : public UHapticFeedbackEffect_Base
 	GENERATED_UCLASS_BODY()
 
 	UPROPERTY(EditAnywhere, Category = "HapticFeedbackEffect_SoundWave")
-	USoundWave *SoundWave;
+	TObjectPtr<USoundWave> SoundWave;
 
 	~UHapticFeedbackEffect_SoundWave();
 
diff --git a/Engine/Source/Runtime/Engine/Classes/Kismet/GameplayStaticsTypes.h b/Engine/Source/Runtime/Engine/Classes/Kismet/GameplayStaticsTypes.h
--- a/Engine/Source/Runtime/Engine/Classes/Kismet/GameplayStaticsTypes.h
+++ b/Engine/Source/Runtime/Engine/Classes/Kismet/GameplayStaticsTypes.h
@@ -56,7 +56,7 @@ struct FPredictProjectilePathParams
 
 	// Actors to ignore when tracing with collision.
 	UPROPERTY(EditAnywhere, BlueprintReadWrite, Category=PredictProjectilePathParams, AdvancedDisplay)
-	TArray<AActor*> ActorsToIgnore;
+	TArray<TObjectPtr<AActor>> ActorsToIgnore;
 
 	// Determines size of each sub-step in the simulation (chopping up MaxSimTime). Recommended between 10 to 30 depending on desired quality versus performance.
 	UPROPERTY(EditAnywhere, BlueprintReadWrite, Category=PredictProjectilePathParams, AdvancedDisplay)
diff --git a/Engine/Source/Runtime/Engine/Classes/Kismet/KismetRenderingLibrary.h b/Engine/Source/Runtime/Engine/Classes/Kismet/KismetRenderingLibrary.h
--- a/Engine/Source/Runtime/Engine/Classes/Kismet/KismetRenderingLibrary.h
+++ b/Engine/Source/Runtime/Engine/Classes/Kismet/KismetRenderingLibrary.h
@@ -27,7 +27,7 @@ struct FDrawToRenderTargetContext
 	{}
 
 	UPROPERTY()
-	UTextureRenderTarget2D* RenderTarget;
+	TObjectPtr<UTextureRenderTarget2D> RenderTarget;
 
 	FDrawEvent* DrawEvent;
 };
diff --git a/Engine/Source/Runtime/Engine/Classes/Layers/Layer.h b/Engine/Source/Runtime/Engine/Classes/Layers/Layer.h
--- a/Engine/Source/Runtime/Engine/Classes/Layers/Layer.h
+++ b/Engine/Source/Runtime/Engine/Classes/Layers/Layer.h
@@ -14,7 +14,7 @@ struct FLayerActorStats
 
 	/** A Type of Actor currently associated with the Layer */
 	UPROPERTY()
-	UClass* Type;
+	TObjectPtr<UClass> Type;
 
 	/** The total number of Actors of Type assigned to the Layer */
 	UPROPERTY()
diff --git a/Engine/Source/Runtime/Engine/Classes/Lightmass/LightmappedSurfaceCollection.h b/Engine/Source/Runtime/Engine/Classes/Lightmass/LightmappedSurfaceCollection.h
--- a/Engine/Source/Runtime/Engine/Classes/Lightmass/LightmappedSurfaceCollection.h
+++ b/Engine/Source/Runtime/Engine/Classes/Lightmass/LightmappedSurfaceCollection.h
@@ -19,7 +19,7 @@ class ULightmappedSurfaceCollection : public UObject
 
 	/** The UModel these surfaces come from. */
 	UPROPERTY(EditAnywhere, Category=LightmappedSurfaceCollection)
-	class UModel* SourceModel;
+	TObjectPtr<class UModel> SourceModel;
 
 	/** An array of the surface indices grouped into a single static lighting mapping. */
 	UPROPERTY(EditAnywhere, Category=LightmappedSurfaceCollection)
diff --git a/Engine/Source/Runtime/Engine/Classes/Lightmass/LightmassPortal.h b/Engine/Source/Runtime/Engine/Classes/Lightmass/LightmassPortal.h
--- a/Engine/Source/Runtime/Engine/Classes/Lightmass/LightmassPortal.h
+++ b/Engine/Source/Runtime/Engine/Classes/Lightmass/LightmassPortal.h
@@ -16,11 +16,11 @@ class ALightmassPortal : public AActor
 
 private:
 	UPROPERTY(Category = Portal, VisibleAnywhere, BlueprintReadOnly, meta = (AllowPrivateAccess = "true"))
-	class ULightmassPortalComponent* PortalComponent;
+	TObjectPtr<class ULightmassPortalComponent> PortalComponent;
 
 #if WITH_EDITORONLY_DATA
 	UPROPERTY()
-	UBillboardComponent* SpriteComponent;
+	TObjectPtr<UBillboardComponent> SpriteComponent;
 #endif // WITH_EDITORONLY_DATA
 
 public:	
diff --git a/Engine/Source/Runtime/Engine/Classes/Lightmass/PrecomputedVisibilityOverrideVolume.h b/Engine/Source/Runtime/Engine/Classes/Lightmass/PrecomputedVisibilityOverrideVolume.h
--- a/Engine/Source/Runtime/Engine/Classes/Lightmass/PrecomputedVisibilityOverrideVolume.h
+++ b/Engine/Source/Runtime/Engine/Classes/Lightmass/PrecomputedVisibilityOverrideVolume.h
@@ -18,11 +18,11 @@ class APrecomputedVisibilityOverrideVolume : public AVolume
 
 	/** Array of actors that will always be considered visible by Precomputed Visibility when viewed from inside this volume. */
 	UPROPERTY(EditAnywhere, BlueprintReadWrite, Category=PrecomputedVisibilityOverrideVolume)
-	TArray<class AActor*> OverrideVisibleActors;
+	TArray<TObjectPtr<class AActor>> OverrideVisibleActors;
 
 	/** Array of actors that will always be considered invisible by Precomputed Visibility when viewed from inside this volume. */
 	UPROPERTY(EditAnywhere, BlueprintReadWrite, Category=PrecomputedVisibilityOverrideVolume)
-	TArray<class AActor*> OverrideInvisibleActors;
+	TArray<TObjectPtr<class AActor>> OverrideInvisibleActors;
 
 	/** Array of level names whose actors will always be considered invisible by Precomputed Visibility when viewed from inside this volume. */
 	UPROPERTY(EditAnywhere, BlueprintReadWrite, Category=PrecomputedVisibilityOverrideVolume)
diff --git a/Engine/Source/Runtime/Engine/Classes/Materials/Material.h b/Engine/Source/Runtime/Engine/Classes/Materials/Material.h
--- a/Engine/Source/Runtime/Engine/Classes/Materials/Material.h
+++ b/Engine/Source/Runtime/Engine/Classes/Materials/Material.h
@@ -159,7 +159,7 @@ struct FMaterialInput
 #if WITH_EDITORONLY_DATA
 	/** Material expression that this input is connected to, or NULL if not connected. */
 	UPROPERTY()
-	class UMaterialExpression* Expression;
+	TObjectPtr<class UMaterialExpression> Expression;
 #endif
 
 	/** Index into Expression's outputs array that this input is connected to. */
@@ -312,15 +312,15 @@ class UMaterial : public UMaterialInterface
 	
 	/** Physical material to use for this graphics material. Used for sounds, effects etc.*/
 	UPROPERTY(EditAnywhere, Category=PhysicalMaterial)
-	class UPhysicalMaterial* PhysMaterial;
+	TObjectPtr<class UPhysicalMaterial> PhysMaterial;
 
 	/** Physical material mask to use for this graphics material. Used for sounds, effects etc.*/
 	UPROPERTY(EditAnywhere, Category = PhysicalMaterial)
-	class UPhysicalMaterialMask* PhysMaterialMask;
+	TObjectPtr<class UPhysicalMaterialMask> PhysMaterialMask;
 
 	/** Physical material mask map to use for this graphics material. Used for sounds, effects etc.*/
 	UPROPERTY(EditAnywhere, Category = PhysicalMaterialMask)
-	class UPhysicalMaterial* PhysicalMaterialMap[EPhysicalMaterialMaskColor::MAX];
+	TObjectPtr<class UPhysicalMaterial> PhysicalMaterialMap[EPhysicalMaterialMaskColor::MAX];
 
 	// Reflection.
 #if WITH_EDITORONLY_DATA
@@ -827,11 +827,11 @@ public:
 
 	/** Array of material expressions, excluding Comments.  Used by the material editor. */
 	UPROPERTY()
-	TArray<class UMaterialExpression*> Expressions;
+	TArray<TObjectPtr<class UMaterialExpression>> Expressions;
 
 	/** Array of comments associated with this material; viewed in the material editor. */
 	UPROPERTY()
-	TArray<class UMaterialExpressionComment*> EditorComments;
+	TArray<TObjectPtr<class UMaterialExpressionComment>> EditorComments;
 
 	/** Controls where this parameter group is displayed in a material instance parameter list.  The lower the number the higher up in the parameter list. */
 	UPROPERTY(EditAnywhere, EditFixedSize, Category = "Group Sorting")
diff --git a/Engine/Source/Runtime/Engine/Classes/Materials/MaterialExpression.h b/Engine/Source/Runtime/Engine/Classes/Materials/MaterialExpression.h
--- a/Engine/Source/Runtime/Engine/Classes/Materials/MaterialExpression.h
+++ b/Engine/Source/Runtime/Engine/Classes/Materials/MaterialExpression.h
@@ -26,7 +26,7 @@ struct FExpressionInput
 #if WITH_EDITORONLY_DATA
 	/** UMaterial expression that this input is connected to, or NULL if not connected. */
 	UPROPERTY()
-	class UMaterialExpression* Expression;
+	TObjectPtr<class UMaterialExpression> Expression;
 #endif
 
 	/** Index into Expression's outputs array that this input is connected to. */
@@ -113,11 +113,11 @@ class ENGINE_API UMaterialExpression : public UObject
 
 	/** Expression's Graph representation */
 	UPROPERTY(transient)
-	UEdGraphNode*	GraphNode;
+	TObjectPtr<UEdGraphNode>	GraphNode;
 
 	/** If exists, expresssion containing this expression within its subgraph. */
 	UPROPERTY()
-	UMaterialExpression* SubgraphExpression;
+	TObjectPtr<UMaterialExpression> SubgraphExpression;
 
 	/** Text of last error for this expression */
 	FString LastErrorText;
@@ -132,14 +132,14 @@ class ENGINE_API UMaterialExpression : public UObject
 	 * This is not necessarily the object which owns this expression, for example a preview material compiling a material function's expressions.
 	 */
 	UPROPERTY()
-	class UMaterial* Material;
+	TObjectPtr<class UMaterial> Material;
 
 	/** 
 	 * The material function that this expression is being used with, if any.
 	 * This will be NULL if the expression belongs to a function that is currently being edited, 
 	 */
 	UPROPERTY()
-	class UMaterialFunction* Function;
+	TObjectPtr<class UMaterialFunction> Function;
 
 #if WITH_EDITORONLY_DATA
 	/** A description that level designers can add (shows in the material editor UI). */
diff --git a/Engine/Source/Runtime/Engine/Classes/Materials/MaterialExpressionCollectionParameter.h b/Engine/Source/Runtime/Engine/Classes/Materials/MaterialExpressionCollectionParameter.h
--- a/Engine/Source/Runtime/Engine/Classes/Materials/MaterialExpressionCollectionParameter.h
+++ b/Engine/Source/Runtime/Engine/Classes/Materials/MaterialExpressionCollectionParameter.h
@@ -21,7 +21,7 @@ class UMaterialExpressionCollectionParameter : public UMaterialExpression
 
 	/** The Parameter Collection to use. */
 	UPROPERTY(EditAnywhere, Category=MaterialExpressionCollectionParameter)
-	class UMaterialParameterCollection* Collection;
+	TObjectPtr<class UMaterialParameterCollection> Collection;
 
 	/** Name of the parameter being referenced. */
 	UPROPERTY(EditAnywhere, Category=MaterialExpressionCollectionParameter)
diff --git a/Engine/Source/Runtime/Engine/Classes/Materials/MaterialExpressionComposite.h b/Engine/Source/Runtime/Engine/Classes/Materials/MaterialExpressionComposite.h
--- a/Engine/Source/Runtime/Engine/Classes/Materials/MaterialExpressionComposite.h
+++ b/Engine/Source/Runtime/Engine/Classes/Materials/MaterialExpressionComposite.h
@@ -54,10 +54,10 @@ class UMaterialExpressionComposite : public UMaterialExpression
 	FString SubgraphName;
 
 	UPROPERTY()
-	UMaterialExpressionPinBase* InputExpressions;
+	TObjectPtr<UMaterialExpressionPinBase> InputExpressions;
 
 	UPROPERTY()
-	UMaterialExpressionPinBase* OutputExpressions;
+	TObjectPtr<UMaterialExpressionPinBase> OutputExpressions;
 
 #if WITH_EDITOR
 	/** Get all reroute expressions used by this composite & its pin bases */
diff --git a/Engine/Source/Runtime/Engine/Classes/Materials/MaterialExpressionCurveAtlasRowParameter.h b/Engine/Source/Runtime/Engine/Classes/Materials/MaterialExpressionCurveAtlasRowParameter.h
--- a/Engine/Source/Runtime/Engine/Classes/Materials/MaterialExpressionCurveAtlasRowParameter.h
+++ b/Engine/Source/Runtime/Engine/Classes/Materials/MaterialExpressionCurveAtlasRowParameter.h
@@ -15,10 +15,10 @@ class UMaterialExpressionCurveAtlasRowParameter : public UMaterialExpressionScal
 	GENERATED_UCLASS_BODY()
 
 	UPROPERTY(EditAnywhere, Category=MaterialExpressionCurveAtlasRowParameter)
-	class UCurveLinearColor* Curve;
+	TObjectPtr<class UCurveLinearColor> Curve;
 
 	UPROPERTY(EditAnywhere, Category = MaterialExpressionCurveAtlasRowParameter)
-	class UCurveLinearColorAtlas* Atlas;
+	TObjectPtr<class UCurveLinearColorAtlas> Atlas;
 
 	UPROPERTY()
 	FExpressionInput InputTime;
diff --git a/Engine/Source/Runtime/Engine/Classes/Materials/MaterialExpressionFontSample.h b/Engine/Source/Runtime/Engine/Classes/Materials/MaterialExpressionFontSample.h
--- a/Engine/Source/Runtime/Engine/Classes/Materials/MaterialExpressionFontSample.h
+++ b/Engine/Source/Runtime/Engine/Classes/Materials/MaterialExpressionFontSample.h
@@ -17,7 +17,7 @@ class UMaterialExpressionFontSample : public UMaterialExpression
 
 	/** font resource that will be sampled */
 	UPROPERTY(EditAnywhere, Category=MaterialExpressionFontSample)
-	class UFont* Font;
+	TObjectPtr<class UFont> Font;
 
 	/** allow access to the various font pages */
 	UPROPERTY(EditAnywhere, Category=MaterialExpressionFontSample)
diff --git a/Engine/Source/Runtime/Engine/Classes/Materials/MaterialExpressionMaterialAttributeLayers.h b/Engine/Source/Runtime/Engine/Classes/Materials/MaterialExpressionMaterialAttributeLayers.h
--- a/Engine/Source/Runtime/Engine/Classes/Materials/MaterialExpressionMaterialAttributeLayers.h
+++ b/Engine/Source/Runtime/Engine/Classes/Materials/MaterialExpressionMaterialAttributeLayers.h
@@ -71,13 +71,13 @@ class UMaterialExpressionMaterialAttributeLayers : public UMaterialExpression
 	}
 
 	UPROPERTY(Transient)
-	TArray<UMaterialExpressionMaterialFunctionCall*> LayerCallers;
+	TArray<TObjectPtr<UMaterialExpressionMaterialFunctionCall>> LayerCallers;
 
 	UPROPERTY(Transient)
 	int32 NumActiveLayerCallers;
 
 	UPROPERTY(Transient)
-	TArray<UMaterialExpressionMaterialFunctionCall*> BlendCallers;
+	TArray<TObjectPtr<UMaterialExpressionMaterialFunctionCall>> BlendCallers;
 
 	UPROPERTY(Transient)
 	int32 NumActiveBlendCallers;
diff --git a/Engine/Source/Runtime/Engine/Classes/Materials/MaterialExpressionMaterialFunctionCall.h b/Engine/Source/Runtime/Engine/Classes/Materials/MaterialExpressionMaterialFunctionCall.h
--- a/Engine/Source/Runtime/Engine/Classes/Materials/MaterialExpressionMaterialFunctionCall.h
+++ b/Engine/Source/Runtime/Engine/Classes/Materials/MaterialExpressionMaterialFunctionCall.h
@@ -27,7 +27,7 @@ struct FFunctionExpressionInput
 	 * This is a reference to a private object so it can't be saved, and must be generated by UpdateFromFunctionResource or SetMaterialFunction. 
 	 */
 	UPROPERTY(transient)
-	class UMaterialExpressionFunctionInput* ExpressionInput;
+	TObjectPtr<class UMaterialExpressionFunctionInput> ExpressionInput;
 
 	/** Id of the FunctionInput, used to link ExpressionInput. */
 	UPROPERTY()
@@ -56,7 +56,7 @@ struct FFunctionExpressionOutput
 	 * This is a reference to a private object so it can't be saved, and must be generated by UpdateFromFunctionResource or SetMaterialFunction. 
 	 */
 	UPROPERTY(transient)
-	class UMaterialExpressionFunctionOutput* ExpressionOutput;
+	TObjectPtr<class UMaterialExpressionFunctionOutput> ExpressionOutput;
 
 	/** Id of the FunctionOutput, used to link ExpressionOutput. */
 	UPROPERTY()
@@ -81,7 +81,7 @@ class UMaterialExpressionMaterialFunctionCall : public UMaterialExpression
 
 	/** The function to call. */
 	UPROPERTY(EditAnywhere, Category=MaterialExpressionMaterialFunctionCall, meta=(DisallowedClasses="MaterialFunctionInstance"))
-	class UMaterialFunctionInterface* MaterialFunction;
+	TObjectPtr<class UMaterialFunctionInterface> MaterialFunction;
 
 #if WITH_EDITORONLY_DATA
 	/** Array of all the function inputs that this function exposes. */
diff --git a/Engine/Source/Runtime/Engine/Classes/Materials/MaterialExpressionNamedReroute.h b/Engine/Source/Runtime/Engine/Classes/Materials/MaterialExpressionNamedReroute.h
--- a/Engine/Source/Runtime/Engine/Classes/Materials/MaterialExpressionNamedReroute.h
+++ b/Engine/Source/Runtime/Engine/Classes/Materials/MaterialExpressionNamedReroute.h
@@ -99,7 +99,7 @@ class ENGINE_API UMaterialExpressionNamedRerouteUsage : public UMaterialExpressi
 public:
 	// The declaration this node is linked to
 	UPROPERTY()
-	UMaterialExpressionNamedRerouteDeclaration* Declaration;
+	TObjectPtr<UMaterialExpressionNamedRerouteDeclaration> Declaration;
 	
 	// The variable GUID, to support copy across graphs
 	UPROPERTY()
diff --git a/Engine/Source/Runtime/Engine/Classes/Materials/MaterialExpressionPinBase.h b/Engine/Source/Runtime/Engine/Classes/Materials/MaterialExpressionPinBase.h
--- a/Engine/Source/Runtime/Engine/Classes/Materials/MaterialExpressionPinBase.h
+++ b/Engine/Source/Runtime/Engine/Classes/Materials/MaterialExpressionPinBase.h
@@ -38,7 +38,7 @@ struct FCompositeReroute
 	FName Name;
 
 	UPROPERTY(meta = (NoResetToDefault))
-	UMaterialExpressionReroute* Expression;
+	TObjectPtr<UMaterialExpressionReroute> Expression;
 
 	FCompositeReroute()
 		: Name(NAME_None)
diff --git a/Engine/Source/Runtime/Engine/Classes/Materials/MaterialExpressionRuntimeVirtualTextureSample.h b/Engine/Source/Runtime/Engine/Classes/Materials/MaterialExpressionRuntimeVirtualTextureSample.h
--- a/Engine/Source/Runtime/Engine/Classes/Materials/MaterialExpressionRuntimeVirtualTextureSample.h
+++ b/Engine/Source/Runtime/Engine/Classes/Materials/MaterialExpressionRuntimeVirtualTextureSample.h
@@ -62,7 +62,7 @@ class ENGINE_API UMaterialExpressionRuntimeVirtualTextureSample : public UMateri
 
 	/** The virtual texture object to sample. */
 	UPROPERTY(EditAnywhere, Category = VirtualTexture)
-	class URuntimeVirtualTexture* VirtualTexture;
+	TObjectPtr<class URuntimeVirtualTexture> VirtualTexture;
 
 	/** How to interpret the virtual texture contents. Note that the bound Virtual Texture should have the same setting for sampling to work correctly. */
 	UPROPERTY(EditAnywhere, Category = VirtualTexture, meta = (DisplayName = "Virtual texture content"))
diff --git a/Engine/Source/Runtime/Engine/Classes/Materials/MaterialExpressionTextureBase.h b/Engine/Source/Runtime/Engine/Classes/Materials/MaterialExpressionTextureBase.h
--- a/Engine/Source/Runtime/Engine/Classes/Materials/MaterialExpressionTextureBase.h
+++ b/Engine/Source/Runtime/Engine/Classes/Materials/MaterialExpressionTextureBase.h
@@ -21,7 +21,7 @@ class ENGINE_API UMaterialExpressionTextureBase : public UMaterialExpression
 	GENERATED_UCLASS_BODY()
 
 	UPROPERTY(EditAnywhere, BlueprintReadWrite, Category=MaterialExpressionTextureBase)
-	class UTexture* Texture;
+	TObjectPtr<class UTexture> Texture;
 
 #if WITH_EDITORONLY_DATA
 	UPROPERTY(EditAnywhere, BlueprintReadWrite, Category=MaterialExpressionTextureBase)
diff --git a/Engine/Source/Runtime/Engine/Classes/Materials/MaterialFunction.h b/Engine/Source/Runtime/Engine/Classes/Materials/MaterialFunction.h
--- a/Engine/Source/Runtime/Engine/Classes/Materials/MaterialFunction.h
+++ b/Engine/Source/Runtime/Engine/Classes/Materials/MaterialFunction.h
@@ -27,7 +27,7 @@ class UMaterialFunction : public UMaterialFunctionInterface
 #if WITH_EDITORONLY_DATA
 	/** Used in the material editor, points to the function asset being edited, which this function is just a preview for. */
 	UPROPERTY(transient)
-	class UMaterialFunction* ParentFunction;
+	TObjectPtr<class UMaterialFunction> ParentFunction;
 
 #endif // WITH_EDITORONLY_DATA
 	/** Description of the function which will be displayed as a tooltip wherever the function is used. */
@@ -37,7 +37,7 @@ class UMaterialFunction : public UMaterialFunctionInterface
 #if WITH_EDITORONLY_DATA
 	/** Array of material expressions, excluding Comments.  Used by the material editor. */
 	UPROPERTY()
-	TArray<UMaterialExpression*> FunctionExpressions;
+	TArray<TObjectPtr<UMaterialExpression>> FunctionExpressions;
 #endif // WITH_EDITORONLY_DATA
 
 	/** Whether to list this function in the material function library, which is a window in the material editor that lists categorized functions. */
@@ -67,13 +67,13 @@ class UMaterialFunction : public UMaterialFunctionInterface
 #if WITH_EDITORONLY_DATA
 	/** Array of comments associated with this material; viewed in the material editor. */
 	UPROPERTY()
-	TArray<class UMaterialExpressionComment*> FunctionEditorComments;
+	TArray<TObjectPtr<class UMaterialExpressionComment>> FunctionEditorComments;
 
 	UPROPERTY(transient)
-	UMaterial* PreviewMaterial;
+	TObjectPtr<UMaterial> PreviewMaterial;
 
 	UPROPERTY()
-	TArray<class UMaterialExpressionMaterialFunctionCall*> DependentFunctionExpressionCandidates;
+	TArray<TObjectPtr<class UMaterialExpressionMaterialFunctionCall>> DependentFunctionExpressionCandidates;
 
 private:
 	/** Transient flag used to track re-entrance in recursive functions like IsDependent. */
diff --git a/Engine/Source/Runtime/Engine/Classes/Materials/MaterialFunctionInstance.h b/Engine/Source/Runtime/Engine/Classes/Materials/MaterialFunctionInstance.h
--- a/Engine/Source/Runtime/Engine/Classes/Materials/MaterialFunctionInstance.h
+++ b/Engine/Source/Runtime/Engine/Classes/Materials/MaterialFunctionInstance.h
@@ -31,11 +31,11 @@ class UMaterialFunctionInstance : public UMaterialFunctionInterface
 
 	/** Parent function. */
 	UPROPERTY(EditAnywhere, BlueprintReadOnly, Category=MaterialFunctionInstance, AssetRegistrySearchable)
-	UMaterialFunctionInterface* Parent;
+	TObjectPtr<UMaterialFunctionInterface> Parent;
 
 	/** Base function. */
 	UPROPERTY(AssetRegistrySearchable)
-	UMaterialFunctionInterface* Base;
+	TObjectPtr<UMaterialFunctionInterface> Base;
 
 	/** Scalar parameters. */
 	UPROPERTY(EditAnywhere, Category=MaterialFunctionInstance)
@@ -180,6 +180,6 @@ public:
 protected:
 #if WITH_EDITORONLY_DATA
 	UPROPERTY(transient)
-	class UMaterialInstanceConstant* PreviewMaterial;
+	TObjectPtr<class UMaterialInstanceConstant> PreviewMaterial;
 #endif // WITH_EDITORONLY_DATA
 };
diff --git a/Engine/Source/Runtime/Engine/Classes/Materials/MaterialFunctionInterface.h b/Engine/Source/Runtime/Engine/Classes/Materials/MaterialFunctionInterface.h
--- a/Engine/Source/Runtime/Engine/Classes/Materials/MaterialFunctionInterface.h
+++ b/Engine/Source/Runtime/Engine/Classes/Materials/MaterialFunctionInterface.h
@@ -114,7 +114,7 @@ public:
 
 	/** Information for thumbnail rendering */
 	UPROPERTY(VisibleAnywhere, Instanced, Category = Thumbnail)
-	class UThumbnailInfo* ThumbnailInfo;
+	TObjectPtr<class UThumbnailInfo> ThumbnailInfo;
 #endif
 	
 	virtual UMaterialFunctionInterface* GetBaseFunction()
diff --git a/Engine/Source/Runtime/Engine/Classes/Materials/MaterialInstance.h b/Engine/Source/Runtime/Engine/Classes/Materials/MaterialInstance.h
--- a/Engine/Source/Runtime/Engine/Classes/Materials/MaterialInstance.h
+++ b/Engine/Source/Runtime/Engine/Classes/Materials/MaterialInstance.h
@@ -165,7 +165,7 @@ struct FTextureParameterValue
 	FMaterialParameterInfo ParameterInfo;
 
 	UPROPERTY(EditAnywhere, BlueprintReadWrite, Category=TextureParameterValue)
-	class UTexture* ParameterValue;
+	TObjectPtr<class UTexture> ParameterValue;
 
 	UPROPERTY()
 	FGuid ExpressionGUID;
@@ -201,7 +201,7 @@ struct FRuntimeVirtualTextureParameterValue
 	FMaterialParameterInfo ParameterInfo;
 
 	UPROPERTY(EditAnywhere, BlueprintReadWrite, Category = RuntimeVirtualTextureParameterValue)
-	class URuntimeVirtualTexture* ParameterValue;
+	TObjectPtr<class URuntimeVirtualTexture> ParameterValue;
 
 	UPROPERTY()
 	FGuid ExpressionGUID;
@@ -242,7 +242,7 @@ struct FFontParameterValue
 	FMaterialParameterInfo ParameterInfo;
 
 	UPROPERTY(EditAnywhere, BlueprintReadWrite, Category=FontParameterValue)
-	class UFont* FontValue;
+	TObjectPtr<class UFont> FontValue;
 
 	UPROPERTY(EditAnywhere, BlueprintReadWrite, Category=FontParameterValue)
 	int32 FontPage;
@@ -299,15 +299,15 @@ class UMaterialInstance : public UMaterialInterface
 
 	/** Physical material to use for this graphics material. Used for sounds, effects etc.*/
 	UPROPERTY(EditAnywhere, BlueprintReadWrite, Category=MaterialInstance)
-	class UPhysicalMaterial* PhysMaterial;
+	TObjectPtr<class UPhysicalMaterial> PhysMaterial;
 
 	/** Physical material map used with physical material mask, when it exists.*/
 	UPROPERTY(EditAnywhere, Category = PhysicalMaterialMask)
-	class UPhysicalMaterial* PhysicalMaterialMap[EPhysicalMaterialMaskColor::MAX];
+	TObjectPtr<class UPhysicalMaterial> PhysicalMaterialMap[EPhysicalMaterialMaskColor::MAX];
 
 	/** Parent material. */
 	UPROPERTY(EditAnywhere, BlueprintReadOnly, Category=MaterialInstance, AssetRegistrySearchable)
-	class UMaterialInterface* Parent;
+	TObjectPtr<class UMaterialInterface> Parent;
 
 	/**
 	 * Delegate for custom static parameters getter.
diff --git a/Engine/Source/Runtime/Engine/Classes/Materials/MaterialInstanceActor.h b/Engine/Source/Runtime/Engine/Classes/Materials/MaterialInstanceActor.h
--- a/Engine/Source/Runtime/Engine/Classes/Materials/MaterialInstanceActor.h
+++ b/Engine/Source/Runtime/Engine/Classes/Materials/MaterialInstanceActor.h
@@ -21,7 +21,7 @@ class AMaterialInstanceActor : public AActor
 
 	/** Pointer to actors that we want to control paramters of using Matinee. */
 	UPROPERTY(EditAnywhere, BlueprintReadWrite, Category=MaterialInstanceActor)
-	TArray<class AActor*> TargetActors;
+	TArray<TObjectPtr<class AActor>> TargetActors;
 
 #if WITH_EDITORONLY_DATA
 private:
diff --git a/Engine/Source/Runtime/Engine/Classes/Materials/MaterialInstanceConstant.h b/Engine/Source/Runtime/Engine/Classes/Materials/MaterialInstanceConstant.h
--- a/Engine/Source/Runtime/Engine/Classes/Materials/MaterialInstanceConstant.h
+++ b/Engine/Source/Runtime/Engine/Classes/Materials/MaterialInstanceConstant.h
@@ -39,7 +39,7 @@ class UMaterialInstanceConstant : public UMaterialInstance
 
 	/** Physical material mask to use for this graphics material. Used for sounds, effects etc.*/
 	UPROPERTY(EditAnywhere, Category = PhysicalMaterial)
-	class UPhysicalMaterialMask* PhysMaterialMask;
+	TObjectPtr<class UPhysicalMaterialMask> PhysMaterialMask;
 
 	// Begin UMaterialInterface interface.
 	ENGINE_API virtual UPhysicalMaterialMask* GetPhysicalMaterialMask() const override;
diff --git a/Engine/Source/Runtime/Engine/Classes/Materials/MaterialInterface.h b/Engine/Source/Runtime/Engine/Classes/Materials/MaterialInterface.h
--- a/Engine/Source/Runtime/Engine/Classes/Materials/MaterialInterface.h
+++ b/Engine/Source/Runtime/Engine/Classes/Materials/MaterialInterface.h
@@ -209,7 +209,7 @@ class UMaterialInterface : public UObject, public IBlendableInterface, public II
 
 	/** SubsurfaceProfile, for Screen Space Subsurface Scattering */
 	UPROPERTY(EditAnywhere, BlueprintReadOnly, Category = Material, meta = (DisplayName = "Subsurface Profile"))
-	class USubsurfaceProfile* SubsurfaceProfile;
+	TObjectPtr<class USubsurfaceProfile> SubsurfaceProfile;
 
 	/* -------------------------- */
 
@@ -236,7 +236,7 @@ protected:
 
 	/** Array of user data stored with the asset */
 	UPROPERTY(EditAnywhere, AdvancedDisplay, Instanced, Category = Material)
-	TArray<UAssetUserData*> AssetUserData;
+	TArray<TObjectPtr<UAssetUserData>> AssetUserData;
 
 private:
 	/** Feature levels to force to compile. */
@@ -261,7 +261,7 @@ public:
 
 	/** Information for thumbnail rendering */
 	UPROPERTY(VisibleAnywhere, Instanced, Category = Thumbnail)
-	class UThumbnailInfo* ThumbnailInfo;
+	TObjectPtr<class UThumbnailInfo> ThumbnailInfo;
 
 	UPROPERTY()
 	TMap<FString, bool> LayerParameterExpansion;
@@ -271,7 +271,7 @@ public:
 
 	/** Importing data and options used for this material */
 	UPROPERTY(EditAnywhere, Instanced, Category = ImportSettings)
-	class UAssetImportData* AssetImportData;
+	TObjectPtr<class UAssetImportData> AssetImportData;
 
 private:
 	/** Unique ID for this material, used for caching during distributed lighting */
diff --git a/Engine/Source/Runtime/Engine/Classes/Materials/MaterialLayersFunctions.h b/Engine/Source/Runtime/Engine/Classes/Materials/MaterialLayersFunctions.h
--- a/Engine/Source/Runtime/Engine/Classes/Materials/MaterialLayersFunctions.h
+++ b/Engine/Source/Runtime/Engine/Classes/Materials/MaterialLayersFunctions.h
@@ -231,10 +231,10 @@ struct ENGINE_API FMaterialLayersFunctions
 	}
 
 	UPROPERTY(EditAnywhere, Category=MaterialLayers)
-	TArray<class UMaterialFunctionInterface*> Layers;
+	TArray<TObjectPtr<class UMaterialFunctionInterface>> Layers;
 
 	UPROPERTY(EditAnywhere, Category=MaterialLayers)
-	TArray<class UMaterialFunctionInterface*> Blends;
+	TArray<TObjectPtr<class UMaterialFunctionInterface>> Blends;
 
 	UPROPERTY(EditAnywhere, Category = MaterialLayers)
 	TArray<bool> LayerStates;
diff --git a/Engine/Source/Runtime/Engine/Classes/Materials/MaterialParameterCollectionInstance.h b/Engine/Source/Runtime/Engine/Classes/Materials/MaterialParameterCollectionInstance.h
--- a/Engine/Source/Runtime/Engine/Classes/Materials/MaterialParameterCollectionInstance.h
+++ b/Engine/Source/Runtime/Engine/Classes/Materials/MaterialParameterCollectionInstance.h
@@ -80,7 +80,7 @@ protected:
 
 	/** Collection resource this instance is based off of. */
 	UPROPERTY()
-	UMaterialParameterCollection* Collection;
+	TObjectPtr<UMaterialParameterCollection> Collection;
 
 	/** World that owns this instance. */
 	TWeakObjectPtr<UWorld> World;
diff --git a/Engine/Source/Runtime/Engine/Classes/Matinee/InterpData.h b/Engine/Source/Runtime/Engine/Classes/Matinee/InterpData.h
--- a/Engine/Source/Runtime/Engine/Classes/Matinee/InterpData.h
+++ b/Engine/Source/Runtime/Engine/Classes/Matinee/InterpData.h
@@ -28,24 +28,24 @@ class UInterpData : public UObject
 
 	/** Actual interpolation data. Groups of InterpTracks. */
 	UPROPERTY(export, BlueprintReadOnly, Category=InterpData)
-	TArray<class UInterpGroup*> InterpGroups;
+	TArray<TObjectPtr<class UInterpGroup>> InterpGroups;
 
 	/** Used for curve editor to remember curve-editing setup. Only loaded in editor. */
 	UPROPERTY(export)
-	class UInterpCurveEdSetup* CurveEdSetup;
+	TObjectPtr<class UInterpCurveEdSetup> CurveEdSetup;
 
 #if WITH_EDITORONLY_DATA
 	/** Used for filtering which tracks are currently visible. */
 	UPROPERTY()
-	TArray<class UInterpFilter*> InterpFilters;
+	TArray<TObjectPtr<class UInterpFilter>> InterpFilters;
 
 	/** The currently selected filter. */
 	UPROPERTY()
-	class UInterpFilter* SelectedFilter;
+	TObjectPtr<class UInterpFilter> SelectedFilter;
 
 	/** Array of default filters. */
 	UPROPERTY(transient)
-	TArray<class UInterpFilter*> DefaultFilters;
+	TArray<TObjectPtr<class UInterpFilter>> DefaultFilters;
 
 #endif // WITH_EDITORONLY_DATA
 	/** Used in editor for defining sections to loop, stretch etc. */
@@ -62,7 +62,7 @@ class UInterpData : public UObject
 
 	/** Cached version of the director group, if any, for easy access while in game */
 	UPROPERTY(transient)
-	class UInterpGroupDirector* CachedDirectorGroup;
+	TObjectPtr<class UInterpGroupDirector> CachedDirectorGroup;
 
 	/** Unique names of all events contained across all event tracks */
 	UPROPERTY()
diff --git a/Engine/Source/Runtime/Engine/Classes/Matinee/InterpFilter_Custom.h b/Engine/Source/Runtime/Engine/Classes/Matinee/InterpFilter_Custom.h
--- a/Engine/Source/Runtime/Engine/Classes/Matinee/InterpFilter_Custom.h
+++ b/Engine/Source/Runtime/Engine/Classes/Matinee/InterpFilter_Custom.h
@@ -21,7 +21,7 @@ class UInterpFilter_Custom : public UInterpFilter
 #if WITH_EDITORONLY_DATA
 	/** Which groups are included in this filter. */
 	UPROPERTY()
-	TArray<class UInterpGroup*> GroupsToInclude;
+	TArray<TObjectPtr<class UInterpGroup>> GroupsToInclude;
 
 #endif // WITH_EDITORONLY_DATA
 
diff --git a/Engine/Source/Runtime/Engine/Classes/Matinee/InterpGroup.h b/Engine/Source/Runtime/Engine/Classes/Matinee/InterpGroup.h
--- a/Engine/Source/Runtime/Engine/Classes/Matinee/InterpGroup.h
+++ b/Engine/Source/Runtime/Engine/Classes/Matinee/InterpGroup.h
@@ -23,10 +23,10 @@ struct FInterpEdSelKey
 	GENERATED_USTRUCT_BODY()
 
 	UPROPERTY()
-	class UInterpGroup* Group;
+	TObjectPtr<class UInterpGroup> Group;
 
 	UPROPERTY()
-	class UInterpTrack* Track;
+	TObjectPtr<class UInterpTrack> Track;
 
 	UPROPERTY()
 	int32 KeyIndex;
@@ -87,7 +87,7 @@ class UInterpGroup : public UObject, public FInterpEdInputInterface
 	GENERATED_UCLASS_BODY()
 
 	UPROPERTY(export, BlueprintReadOnly, Category=InterpGroup)
-	TArray<class UInterpTrack*> InterpTracks;
+	TArray<TObjectPtr<class UInterpTrack>> InterpTracks;
 
 	/** 
 	 *	Within an InterpData, all GroupNames must be unique. 
diff --git a/Engine/Source/Runtime/Engine/Classes/Matinee/InterpGroupCamera.h b/Engine/Source/Runtime/Engine/Classes/Matinee/InterpGroupCamera.h
--- a/Engine/Source/Runtime/Engine/Classes/Matinee/InterpGroupCamera.h
+++ b/Engine/Source/Runtime/Engine/Classes/Matinee/InterpGroupCamera.h
@@ -25,7 +25,7 @@ struct FCameraPreviewInfo
 	TSubclassOf<class APawn>  PawnClass;
 
 	UPROPERTY(EditAnywhere, Category=CameraPreviewInfo)
-	class UAnimSequence* AnimSeq;
+	TObjectPtr<class UAnimSequence> AnimSeq;
 
 	/* for now this is read-only. It has maintenance issue to be resolved if I enable this.*/
 	UPROPERTY(VisibleAnywhere, Category=CameraPreviewInfo)
@@ -36,7 +36,7 @@ struct FCameraPreviewInfo
 
 	/** APawn Inst - CameraAnimInst doesn't really exist in editor **/
 	UPROPERTY(transient)
-	class APawn* PawnInst;
+	TObjectPtr<class APawn> PawnInst;
 
 
 
@@ -57,7 +57,7 @@ class UInterpGroupCamera : public UInterpGroup
 	GENERATED_UCLASS_BODY()
 
 	UPROPERTY(transient)
-	class UCameraAnim* CameraAnimInst;
+	TObjectPtr<class UCameraAnim> CameraAnimInst;
 
 #if WITH_EDITORONLY_DATA
 	// this is interaction property info for CameraAnim
diff --git a/Engine/Source/Runtime/Engine/Classes/Matinee/InterpGroupInst.h b/Engine/Source/Runtime/Engine/Classes/Matinee/InterpGroupInst.h
--- a/Engine/Source/Runtime/Engine/Classes/Matinee/InterpGroupInst.h
+++ b/Engine/Source/Runtime/Engine/Classes/Matinee/InterpGroupInst.h
@@ -25,18 +25,18 @@ class UInterpGroupInst : public UObject
 
 	/** UInterpGroup within the InterpData that this is an instance of. */
 	UPROPERTY()
-	class UInterpGroup* Group;
+	TObjectPtr<class UInterpGroup> Group;
 
 	/** 
 	 *	Actor that this Group instance is acting upon.
 	 *	NB: that this may be set to NULL at any time as a result of the  AActor  being destroyed.
 	 */
 	UPROPERTY()
-	class AActor* GroupActor;
+	TObjectPtr<class AActor> GroupActor;
 
 	/** Array if InterpTrack instances. TrackInst.Num() == UInterpGroup.InterpTrack.Num() must be true. */
 	UPROPERTY()
-	TArray<class UInterpTrackInst*> TrackInst;
+	TArray<TObjectPtr<class UInterpTrackInst>> TrackInst;
 
 
 	/** 
diff --git a/Engine/Source/Runtime/Engine/Classes/Matinee/InterpTrack.h b/Engine/Source/Runtime/Engine/Classes/Matinee/InterpTrack.h
--- a/Engine/Source/Runtime/Engine/Classes/Matinee/InterpTrack.h
+++ b/Engine/Source/Runtime/Engine/Classes/Matinee/InterpTrack.h
@@ -106,7 +106,7 @@ class UInterpTrack : public UObject, public FCurveEdInterface, public FInterpEdI
 
 	/** A list of subtracks that belong to this track */
 	UPROPERTY(BlueprintReadOnly, Category=InterpTrack)
-	TArray<class UInterpTrack*> SubTracks;
+	TArray<TObjectPtr<class UInterpTrack>> SubTracks;
 
 #if WITH_EDITORONLY_DATA
 	/** A list of subtrack groups (for editor UI organization only) */
@@ -149,7 +149,7 @@ private:
 public:
 #if WITH_EDITORONLY_DATA
 	UPROPERTY(transient)
-	UTexture2D* TrackIcon;
+	TObjectPtr<UTexture2D> TrackIcon;
 #endif // WITH_EDITORONLY_DATA
 
 	/** If true, the  AActor  this track is working on will have BeginAnimControl/FinishAnimControl called on it. */
diff --git a/Engine/Source/Runtime/Engine/Classes/Matinee/InterpTrackAnimControl.h b/Engine/Source/Runtime/Engine/Classes/Matinee/InterpTrackAnimControl.h
--- a/Engine/Source/Runtime/Engine/Classes/Matinee/InterpTrackAnimControl.h
+++ b/Engine/Source/Runtime/Engine/Classes/Matinee/InterpTrackAnimControl.h
@@ -24,7 +24,7 @@ struct FAnimControlTrackKey
 
 	/** Animation Sequence to play */
 	UPROPERTY()
-	class UAnimSequence* AnimSeq;
+	TObjectPtr<class UAnimSequence> AnimSeq;
 
 	/** Time to start playing AnimSequence at. */
 	UPROPERTY()
diff --git a/Engine/Source/Runtime/Engine/Classes/Matinee/InterpTrackDirector.h b/Engine/Source/Runtime/Engine/Classes/Matinee/InterpTrackDirector.h
--- a/Engine/Source/Runtime/Engine/Classes/Matinee/InterpTrackDirector.h
+++ b/Engine/Source/Runtime/Engine/Classes/Matinee/InterpTrackDirector.h
@@ -67,7 +67,7 @@ private:
 #if WITH_EDITORONLY_DATA
 	/** The camera actor which the track is currently focused on. Only valid if this track or it's group is selected */
 	UPROPERTY(Transient)
-	class ACameraActor* PreviewCamera;
+	TObjectPtr<class ACameraActor> PreviewCamera;
 #endif // WITH_EDITORONLY_DATA
 
 public:
diff --git a/Engine/Source/Runtime/Engine/Classes/Matinee/InterpTrackFloatAnimBPParam.h b/Engine/Source/Runtime/Engine/Classes/Matinee/InterpTrackFloatAnimBPParam.h
--- a/Engine/Source/Runtime/Engine/Classes/Matinee/InterpTrackFloatAnimBPParam.h
+++ b/Engine/Source/Runtime/Engine/Classes/Matinee/InterpTrackFloatAnimBPParam.h
@@ -19,7 +19,7 @@ class UInterpTrackFloatAnimBPParam : public UInterpTrackFloatBase
 	
 	UE_DEPRECATED(4.11, "This property is deprecated. Please use AnimClass instead")
 	UPROPERTY(EditAnywhere, Category = InterpTrackFloatAnimBPParam)
-	class UAnimBlueprintGeneratedClass* AnimBlueprintClass;
+	TObjectPtr<class UAnimBlueprintGeneratedClass> AnimBlueprintClass;
 
 	/** Materials whose parameters we want to change and the references to those materials. */
 	UPROPERTY(EditAnywhere, Category=InterpTrackFloatAnimBPParam)
diff --git a/Engine/Source/Runtime/Engine/Classes/Matinee/InterpTrackFloatMaterialParam.h b/Engine/Source/Runtime/Engine/Classes/Matinee/InterpTrackFloatMaterialParam.h
--- a/Engine/Source/Runtime/Engine/Classes/Matinee/InterpTrackFloatMaterialParam.h
+++ b/Engine/Source/Runtime/Engine/Classes/Matinee/InterpTrackFloatMaterialParam.h
@@ -17,7 +17,7 @@ class UInterpTrackFloatMaterialParam : public UInterpTrackFloatBase
 	
 	/** Materials whose parameters we want to change and the references to those materials. */
 	UPROPERTY(EditAnywhere, Category=InterpTrackFloatMaterialParam)
-	TArray<class UMaterialInterface*> TargetMaterials;
+	TArray<TObjectPtr<class UMaterialInterface>> TargetMaterials;
 
 	/** Name of parameter in the MaterialInstance which this track will modify over time. */
 	UPROPERTY(EditAnywhere, Category=InterpTrackFloatMaterialParam)
diff --git a/Engine/Source/Runtime/Engine/Classes/Matinee/InterpTrackInstDirector.h b/Engine/Source/Runtime/Engine/Classes/Matinee/InterpTrackInstDirector.h
--- a/Engine/Source/Runtime/Engine/Classes/Matinee/InterpTrackInstDirector.h
+++ b/Engine/Source/Runtime/Engine/Classes/Matinee/InterpTrackInstDirector.h
@@ -14,7 +14,7 @@ class UInterpTrackInstDirector : public UInterpTrackInst
 	GENERATED_UCLASS_BODY()
 
 	UPROPERTY()
-	class AActor* OldViewTarget;
+	TObjectPtr<class AActor> OldViewTarget;
 
 
 	// Begin UInterpTrackInst Instance
diff --git a/Engine/Source/Runtime/Engine/Classes/Matinee/InterpTrackInstFloatAnimBPParam.h b/Engine/Source/Runtime/Engine/Classes/Matinee/InterpTrackInstFloatAnimBPParam.h
--- a/Engine/Source/Runtime/Engine/Classes/Matinee/InterpTrackInstFloatAnimBPParam.h
+++ b/Engine/Source/Runtime/Engine/Classes/Matinee/InterpTrackInstFloatAnimBPParam.h
@@ -17,7 +17,7 @@ class UInterpTrackInstFloatAnimBPParam : public UInterpTrackInst
 
 	/** MIDs we're using to set the desired parameter. */
 	UPROPERTY(transient)
-	class UAnimInstance* AnimScriptInstance;
+	TObjectPtr<class UAnimInstance> AnimScriptInstance;
 
 	/** Saved values for restoring state when exiting Matinee. */
 	UPROPERTY(transient)
diff --git a/Engine/Source/Runtime/Engine/Classes/Matinee/InterpTrackInstFloatMaterialParam.h b/Engine/Source/Runtime/Engine/Classes/Matinee/InterpTrackInstFloatMaterialParam.h
--- a/Engine/Source/Runtime/Engine/Classes/Matinee/InterpTrackInstFloatMaterialParam.h
+++ b/Engine/Source/Runtime/Engine/Classes/Matinee/InterpTrackInstFloatMaterialParam.h
@@ -18,7 +18,7 @@ class UInterpTrackInstFloatMaterialParam : public UInterpTrackInst
 
 	/** MIDs we're using to set the desired parameter. */
 	UPROPERTY()
-	TArray<class UMaterialInstanceDynamic*> MaterialInstances;
+	TArray<TObjectPtr<class UMaterialInstanceDynamic>> MaterialInstances;
 
 	/** Saved values for restoring state when exiting Matinee. */
 	UPROPERTY()
@@ -30,7 +30,7 @@ class UInterpTrackInstFloatMaterialParam : public UInterpTrackInst
 
 	/** track we are an instance of - used in the editor to propagate changes to the track's Materials array immediately */
 	UPROPERTY()
-	class UInterpTrackFloatMaterialParam* InstancedTrack;
+	TObjectPtr<class UInterpTrackFloatMaterialParam> InstancedTrack;
 
 
 	// Begin UInterpTrackInst Instance
diff --git a/Engine/Source/Runtime/Engine/Classes/Matinee/InterpTrackInstProperty.h b/Engine/Source/Runtime/Engine/Classes/Matinee/InterpTrackInstProperty.h
--- a/Engine/Source/Runtime/Engine/Classes/Matinee/InterpTrackInstProperty.h
+++ b/Engine/Source/Runtime/Engine/Classes/Matinee/InterpTrackInstProperty.h
@@ -24,7 +24,7 @@ class UInterpTrackInstProperty : public UInterpTrackInst
 
 	/** Pointer to the UObject instance that is the outer of the color property we are interpolating on, this is used to process the property update callback. */
 	UPROPERTY()
-	class UObject* PropertyOuterObjectInst;
+	TObjectPtr<class UObject> PropertyOuterObjectInst;
 
 
 	/**
diff --git a/Engine/Source/Runtime/Engine/Classes/Matinee/InterpTrackInstSound.h b/Engine/Source/Runtime/Engine/Classes/Matinee/InterpTrackInstSound.h
--- a/Engine/Source/Runtime/Engine/Classes/Matinee/InterpTrackInstSound.h
+++ b/Engine/Source/Runtime/Engine/Classes/Matinee/InterpTrackInstSound.h
@@ -19,7 +19,7 @@ class UInterpTrackInstSound : public UInterpTrackInst
 	float LastUpdatePosition;
 
 	UPROPERTY(transient)
-	class UAudioComponent* PlayAudioComp;
+	TObjectPtr<class UAudioComponent> PlayAudioComp;
 
 
 	// Begin UInterpTrackInst Instance
diff --git a/Engine/Source/Runtime/Engine/Classes/Matinee/InterpTrackInstVectorMaterialParam.h b/Engine/Source/Runtime/Engine/Classes/Matinee/InterpTrackInstVectorMaterialParam.h
--- a/Engine/Source/Runtime/Engine/Classes/Matinee/InterpTrackInstVectorMaterialParam.h
+++ b/Engine/Source/Runtime/Engine/Classes/Matinee/InterpTrackInstVectorMaterialParam.h
@@ -18,7 +18,7 @@ class UInterpTrackInstVectorMaterialParam : public UInterpTrackInst
 
 	/** MIDs we're using to set the desired parameter. */
 	UPROPERTY()
-	TArray<class UMaterialInstanceDynamic*> MaterialInstances;
+	TArray<TObjectPtr<class UMaterialInstanceDynamic>> MaterialInstances;
 
 	/** Saved values for restoring state when exiting Matinee. */
 	UPROPERTY()
@@ -30,7 +30,7 @@ class UInterpTrackInstVectorMaterialParam : public UInterpTrackInst
 
 	/** Track we are an instance of - used in the editor to propagate changes to the track's Materials array immediately. */
 	UPROPERTY()
-	class UInterpTrackVectorMaterialParam* InstancedTrack;
+	TObjectPtr<class UInterpTrackVectorMaterialParam> InstancedTrack;
 
 
 	// Begin UInterpTrackInst Instance
diff --git a/Engine/Source/Runtime/Engine/Classes/Matinee/InterpTrackSound.h b/Engine/Source/Runtime/Engine/Classes/Matinee/InterpTrackSound.h
--- a/Engine/Source/Runtime/Engine/Classes/Matinee/InterpTrackSound.h
+++ b/Engine/Source/Runtime/Engine/Classes/Matinee/InterpTrackSound.h
@@ -32,7 +32,7 @@ struct FSoundTrackKey
 	float Pitch;
 
 	UPROPERTY(EditAnywhere, Category=SoundTrackKey)
-	class USoundBase* Sound;
+	TObjectPtr<class USoundBase> Sound;
 
 
 
diff --git a/Engine/Source/Runtime/Engine/Classes/Matinee/InterpTrackVectorMaterialParam.h b/Engine/Source/Runtime/Engine/Classes/Matinee/InterpTrackVectorMaterialParam.h
--- a/Engine/Source/Runtime/Engine/Classes/Matinee/InterpTrackVectorMaterialParam.h
+++ b/Engine/Source/Runtime/Engine/Classes/Matinee/InterpTrackVectorMaterialParam.h
@@ -17,7 +17,7 @@ class UInterpTrackVectorMaterialParam : public UInterpTrackVectorBase
 
 	/** Materials whose parameters we want to change and the references to those materials. */
 	UPROPERTY(EditAnywhere, Category=InterpTrackVectorMaterialParam)
-	TArray<class UMaterialInterface*> TargetMaterials;
+	TArray<TObjectPtr<class UMaterialInterface>> TargetMaterials;
 
 	/** Name of parameter in the MaterialInstance which this track will modify over time. */
 	UPROPERTY(EditAnywhere, Category=InterpTrackVectorMaterialParam)
diff --git a/Engine/Source/Runtime/Engine/Classes/Matinee/MatineeActor.h b/Engine/Source/Runtime/Engine/Classes/Matinee/MatineeActor.h
--- a/Engine/Source/Runtime/Engine/Classes/Matinee/MatineeActor.h
+++ b/Engine/Source/Runtime/Engine/Classes/Matinee/MatineeActor.h
@@ -51,7 +51,7 @@ struct FInterpGroupActorInfo
 	FName ObjectName;
 
 	UPROPERTY(EditAnywhere, Category=InterpGroupActorInfo)
-	TArray<class AActor*> Actors;
+	TArray<TObjectPtr<class AActor>> Actors;
 
 };
 
@@ -64,7 +64,7 @@ class ENGINE_API AMatineeActor : public AActor
 
 	/** The matinee data used by this actor*/
 	UPROPERTY(EditAnywhere, BlueprintReadWrite, Category=MatineeActor, replicated, meta=(ForceRebuildProperty = "GroupActorInfos"))
-	class UInterpData* MatineeData;
+	TObjectPtr<class UInterpData> MatineeData;
 
 	/** Name of controller node in level script, used to know what function to try and find for events */
 	UPROPERTY()
@@ -162,7 +162,7 @@ class ENGINE_API AMatineeActor : public AActor
 
 	/** Instance data for interp groups. One for each variable/group combination. */
 	UPROPERTY(transient)
-	TArray<class UInterpGroupInst*> GroupInst;
+	TArray<TObjectPtr<class UInterpGroupInst>> GroupInst;
 
 	/** Contains the camera world-position for each camera cut in the cinematic. */
 	UPROPERTY(transient)
@@ -173,7 +173,7 @@ class ENGINE_API AMatineeActor : public AActor
 	// Reference to the actor sprite
 private:
 	UPROPERTY()
-	class UBillboardComponent* SpriteComponent;
+	TObjectPtr<class UBillboardComponent> SpriteComponent;
 public:
 
 	UPROPERTY(transient)
diff --git a/Engine/Source/Runtime/Engine/Classes/Matinee/MatineeActorCameraAnim.h b/Engine/Source/Runtime/Engine/Classes/Matinee/MatineeActorCameraAnim.h
--- a/Engine/Source/Runtime/Engine/Classes/Matinee/MatineeActorCameraAnim.h
+++ b/Engine/Source/Runtime/Engine/Classes/Matinee/MatineeActorCameraAnim.h
@@ -19,5 +19,5 @@ class AMatineeActorCameraAnim : public AMatineeActor
 
 	/** The camera anim we are editing */
 	UPROPERTY(Transient)
-	UCameraAnim* CameraAnim;
+	TObjectPtr<UCameraAnim> CameraAnim;
 };
diff --git a/Engine/Source/Runtime/Engine/Classes/Particles/Acceleration/ParticleModuleAccelerationDrag.h b/Engine/Source/Runtime/Engine/Classes/Particles/Acceleration/ParticleModuleAccelerationDrag.h
--- a/Engine/Source/Runtime/Engine/Classes/Particles/Acceleration/ParticleModuleAccelerationDrag.h
+++ b/Engine/Source/Runtime/Engine/Classes/Particles/Acceleration/ParticleModuleAccelerationDrag.h
@@ -22,7 +22,7 @@ class UParticleModuleAccelerationDrag : public UParticleModuleAccelerationBase
 
 	/** Per-particle drag coefficient. Evaluted using emitter time. */
 	UPROPERTY()
-	class UDistributionFloat* DragCoefficient_DEPRECATED;
+	TObjectPtr<class UDistributionFloat> DragCoefficient_DEPRECATED;
 
 	/** Per-particle drag coefficient. Evaluted using emitter time. */
 	UPROPERTY(EditAnywhere, Category = Drag)
diff --git a/Engine/Source/Runtime/Engine/Classes/Particles/Acceleration/ParticleModuleAccelerationDragScaleOverLife.h b/Engine/Source/Runtime/Engine/Classes/Particles/Acceleration/ParticleModuleAccelerationDragScaleOverLife.h
--- a/Engine/Source/Runtime/Engine/Classes/Particles/Acceleration/ParticleModuleAccelerationDragScaleOverLife.h
+++ b/Engine/Source/Runtime/Engine/Classes/Particles/Acceleration/ParticleModuleAccelerationDragScaleOverLife.h
@@ -22,7 +22,7 @@ class UParticleModuleAccelerationDragScaleOverLife : public UParticleModuleAccel
 
 	/** Per-particle drag scale. Evaluted using particle relative time. */
 	UPROPERTY()
-	class UDistributionFloat* DragScale_DEPRECATED;
+	TObjectPtr<class UDistributionFloat> DragScale_DEPRECATED;
 
 	/** Per-particle drag scale. Evaluted using particle relative time. */
 	UPROPERTY(EditAnywhere, Category = Drag)
diff --git a/Engine/Source/Runtime/Engine/Classes/Particles/Attractor/ParticleModuleAttractorPointGravity.h b/Engine/Source/Runtime/Engine/Classes/Particles/Attractor/ParticleModuleAttractorPointGravity.h
--- a/Engine/Source/Runtime/Engine/Classes/Particles/Attractor/ParticleModuleAttractorPointGravity.h
+++ b/Engine/Source/Runtime/Engine/Classes/Particles/Attractor/ParticleModuleAttractorPointGravity.h
@@ -30,7 +30,7 @@ class UParticleModuleAttractorPointGravity : public UParticleModuleAttractorBase
 
 	/** The strength of the point source. */
 	UPROPERTY()
-	class UDistributionFloat* Strength_DEPRECATED;
+	TObjectPtr<class UDistributionFloat> Strength_DEPRECATED;
 
 	/** The strength of the point source. */
 	UPROPERTY(EditAnywhere, noclear, Category=PointGravitySource)
diff --git a/Engine/Source/Runtime/Engine/Classes/Particles/Emitter.h b/Engine/Source/Runtime/Engine/Classes/Particles/Emitter.h
--- a/Engine/Source/Runtime/Engine/Classes/Particles/Emitter.h
+++ b/Engine/Source/Runtime/Engine/Classes/Particles/Emitter.h
@@ -60,7 +60,7 @@ class ENGINE_API AEmitter : public AActor
 
 private:
 	UPROPERTY(Category = Emitter, VisibleAnywhere, BlueprintReadOnly, meta = (ExposeFunctionCategories = "Particles|Beam,Particles|Parameters,Particles,Effects|Components|ParticleSystem,Rendering,Activation,Components|Activation", AllowPrivateAccess = "true"))
-	class UParticleSystemComponent* ParticleSystemComponent;
+	TObjectPtr<class UParticleSystemComponent> ParticleSystemComponent;
 public:
 
 	UPROPERTY()
@@ -89,10 +89,10 @@ public:
 
 private:
 	UPROPERTY()
-	class UBillboardComponent* SpriteComponent;
+	TObjectPtr<class UBillboardComponent> SpriteComponent;
 
 	UPROPERTY()
-	class UArrowComponent* ArrowComponent;
+	TObjectPtr<class UArrowComponent> ArrowComponent;
 public:
 
 #endif
diff --git a/Engine/Source/Runtime/Engine/Classes/Particles/EmitterCameraLensEffectBase.h b/Engine/Source/Runtime/Engine/Classes/Particles/EmitterCameraLensEffectBase.h
--- a/Engine/Source/Runtime/Engine/Classes/Particles/EmitterCameraLensEffectBase.h
+++ b/Engine/Source/Runtime/Engine/Classes/Particles/EmitterCameraLensEffectBase.h
@@ -23,15 +23,15 @@ class ENGINE_API AEmitterCameraLensEffectBase : public AEmitter
 protected:
 	/** Particle System to use */
 	UPROPERTY(EditDefaultsOnly, Category = EmitterCameraLensEffectBase)
-	class UParticleSystem* PS_CameraEffect;
+	TObjectPtr<class UParticleSystem> PS_CameraEffect;
 
 	/** The effect to use for non extreme content */
 	UPROPERTY()
-	class UParticleSystem* PS_CameraEffectNonExtremeContent_DEPRECATED;
+	TObjectPtr<class UParticleSystem> PS_CameraEffectNonExtremeContent_DEPRECATED;
 
 	/** Camera this emitter is attached to, will be notified when emitter is destroyed */
 	UPROPERTY(transient)
-	class APlayerCameraManager* BaseCamera;
+	TObjectPtr<class APlayerCameraManager> BaseCamera;
 
 	/** 
 	 * Effect-to-camera transform to allow arbitrary placement of the particle system .
