Automated wrapped object pointer upgrade for Engine + ShooterGame + ShooterGame referenced plugins

diff --git a/Engine/Source/Runtime/Engine/Classes/Particles/Event/ParticleModuleEventGenerator.h b/Engine/Source/Runtime/Engine/Classes/Particles/Event/ParticleModuleEventGenerator.h
--- a/Engine/Source/Runtime/Engine/Classes/Particles/Event/ParticleModuleEventGenerator.h
+++ b/Engine/Source/Runtime/Engine/Classes/Particles/Event/ParticleModuleEventGenerator.h
@@ -53,7 +53,7 @@ struct FParticleEvent_GenerateInfo
 
 	/** The events we want to fire off when this event has been generated */
 	UPROPERTY(EditAnywhere, Instanced, Category = ParticleEvent_GenerateInfo)
-	TArray<class UParticleModuleEventSendToGame*> ParticleModuleEventsToSendToGame;
+	TArray<TObjectPtr<class UParticleModuleEventSendToGame>> ParticleModuleEventsToSendToGame;
 
 	FParticleEvent_GenerateInfo()
 	: Type(0)
diff --git a/Engine/Source/Runtime/Engine/Classes/Particles/Event/ParticleModuleEventReceiverSpawn.h b/Engine/Source/Runtime/Engine/Classes/Particles/Event/ParticleModuleEventReceiverSpawn.h
--- a/Engine/Source/Runtime/Engine/Classes/Particles/Event/ParticleModuleEventReceiverSpawn.h
+++ b/Engine/Source/Runtime/Engine/Classes/Particles/Event/ParticleModuleEventReceiverSpawn.h
@@ -55,7 +55,7 @@ class UParticleModuleEventReceiverSpawn : public UParticleModuleEventReceiverBas
 	*	Array of physical materials that can be used to allow or ban a specific set of materials when receiving collision events.
 	*/
 	UPROPERTY(EditAnywhere, Category = Collision)
-	TArray<UPhysicalMaterial*> PhysicalMaterials;
+	TArray<TObjectPtr<UPhysicalMaterial>> PhysicalMaterials;
 
 	/**
 	*	When true, the PhysicalMaterials list is used to ban specified materials for collision events but allow all others.
diff --git a/Engine/Source/Runtime/Engine/Classes/Particles/Location/ParticleModuleLocationBoneSocket.h b/Engine/Source/Runtime/Engine/Classes/Particles/Location/ParticleModuleLocationBoneSocket.h
--- a/Engine/Source/Runtime/Engine/Classes/Particles/Location/ParticleModuleLocationBoneSocket.h
+++ b/Engine/Source/Runtime/Engine/Classes/Particles/Location/ParticleModuleLocationBoneSocket.h
@@ -123,7 +123,7 @@ class ENGINE_API UParticleModuleLocationBoneSocket : public UParticleModuleLocat
 #if WITH_EDITORONLY_DATA
 	/** The name of the skeletal mesh to use in the editor */
 	UPROPERTY(EditAnywhere, Category=BoneSocket)
-	class USkeletalMesh* EditorSkelMesh;
+	TObjectPtr<class USkeletalMesh> EditorSkelMesh;
 
 #endif // WITH_EDITORONLY_DATA
 
diff --git a/Engine/Source/Runtime/Engine/Classes/Particles/Location/ParticleModuleLocationSkelVertSurface.h b/Engine/Source/Runtime/Engine/Classes/Particles/Location/ParticleModuleLocationSkelVertSurface.h
--- a/Engine/Source/Runtime/Engine/Classes/Particles/Location/ParticleModuleLocationSkelVertSurface.h
+++ b/Engine/Source/Runtime/Engine/Classes/Particles/Location/ParticleModuleLocationSkelVertSurface.h
@@ -114,7 +114,7 @@ class ENGINE_API UParticleModuleLocationSkelVertSurface : public UParticleModule
 #if WITH_EDITORONLY_DATA
 	/** The name of the skeletal mesh to use in the editor */
 	UPROPERTY(EditAnywhere, Category=VertSurface)
-	class USkeletalMesh* EditorSkelMesh;
+	TObjectPtr<class USkeletalMesh> EditorSkelMesh;
 
 #endif // WITH_EDITORONLY_DATA
 	/** This module will only spawn from verts or surfaces associated with the bones in this list */
diff --git a/Engine/Source/Runtime/Engine/Classes/Particles/Material/ParticleModuleMeshMaterial.h b/Engine/Source/Runtime/Engine/Classes/Particles/Material/ParticleModuleMeshMaterial.h
--- a/Engine/Source/Runtime/Engine/Classes/Particles/Material/ParticleModuleMeshMaterial.h
+++ b/Engine/Source/Runtime/Engine/Classes/Particles/Material/ParticleModuleMeshMaterial.h
@@ -17,7 +17,7 @@ class UParticleModuleMeshMaterial : public UParticleModuleMaterialBase
 
 	/** The array of materials to apply to the mesh particles. */
 	UPROPERTY(EditAnywhere, Category=MeshMaterials)
-	TArray<class UMaterialInterface*> MeshMaterials;
+	TArray<TObjectPtr<class UMaterialInterface>> MeshMaterials;
 
 
 	//Begin UParticleModule Interface
diff --git a/Engine/Source/Runtime/Engine/Classes/Particles/ParticleEmitter.h b/Engine/Source/Runtime/Engine/Classes/Particles/ParticleEmitter.h
--- a/Engine/Source/Runtime/Engine/Classes/Particles/ParticleEmitter.h
+++ b/Engine/Source/Runtime/Engine/Classes/Particles/ParticleEmitter.h
@@ -184,7 +184,7 @@ class UParticleEmitter : public UObject
 	//	'Private' data - not required by the editor
 	//~=============================================================================
 	UPROPERTY(instanced)
-	TArray<class UParticleLODLevel*> LODLevels;
+	TArray<TObjectPtr<class UParticleLODLevel>> LODLevels;
 
 	UPROPERTY()
 	int32 PeakActiveParticles;
diff --git a/Engine/Source/Runtime/Engine/Classes/Particles/ParticleLODLevel.h b/Engine/Source/Runtime/Engine/Classes/Particles/ParticleLODLevel.h
--- a/Engine/Source/Runtime/Engine/Classes/Particles/ParticleLODLevel.h
+++ b/Engine/Source/Runtime/Engine/Classes/Particles/ParticleLODLevel.h
@@ -31,45 +31,45 @@ class UParticleLODLevel : public UObject
 
 	/** The required module for this LOD level											*/
 	UPROPERTY(instanced)
-	class UParticleModuleRequired* RequiredModule;
+	TObjectPtr<class UParticleModuleRequired> RequiredModule;
 
 	/** An array of particle modules that contain the adjusted data for the LOD level	*/
 	UPROPERTY(instanced)
-	TArray<class UParticleModule*> Modules;
+	TArray<TObjectPtr<class UParticleModule>> Modules;
 
 	// Module<SINGULAR> used for emitter type "extension".
 	UPROPERTY(export)
-	class UParticleModuleTypeDataBase* TypeDataModule;
+	TObjectPtr<class UParticleModuleTypeDataBase> TypeDataModule;
 
 	/** The SpawnRate/Burst module - required by all emitters. */
 	UPROPERTY(export)
-	class UParticleModuleSpawn* SpawnModule;
+	TObjectPtr<class UParticleModuleSpawn> SpawnModule;
 
 	/** The optional EventGenerator module. */
 	UPROPERTY(export)
-	class UParticleModuleEventGenerator* EventGenerator;
+	TObjectPtr<class UParticleModuleEventGenerator> EventGenerator;
 
 	/** SpawningModules - These are called to determine how many particles to spawn.	*/
 	UPROPERTY(transient, duplicatetransient)
-	TArray<class UParticleModuleSpawnBase*> SpawningModules;
+	TArray<TObjectPtr<class UParticleModuleSpawnBase>> SpawningModules;
 
 	/** SpawnModules - These are called when particles are spawned.						*/
 	UPROPERTY(transient, duplicatetransient)
-	TArray<class UParticleModule*> SpawnModules;
+	TArray<TObjectPtr<class UParticleModule>> SpawnModules;
 
 	/** UpdateModules - These are called when particles are updated.					*/
 	UPROPERTY(transient, duplicatetransient)
-	TArray<class UParticleModule*> UpdateModules;
+	TArray<TObjectPtr<class UParticleModule>> UpdateModules;
 
 	/** OrbitModules 
 	 *	These are used to do offsets of the sprite from the particle location.
 	 */
 	UPROPERTY(transient, duplicatetransient)
-	TArray<class UParticleModuleOrbit*> OrbitModules;
+	TArray<TObjectPtr<class UParticleModuleOrbit>> OrbitModules;
 
 	/** Event receiver modules only! */
 	UPROPERTY(transient, duplicatetransient)
-	TArray<class UParticleModuleEventReceiverBase*> EventReceiverModules;
+	TArray<TObjectPtr<class UParticleModuleEventReceiverBase>> EventReceiverModules;
 
 	UPROPERTY()
 	uint32 ConvertedModules:1;
diff --git a/Engine/Source/Runtime/Engine/Classes/Particles/ParticleModule.h b/Engine/Source/Runtime/Engine/Classes/Particles/ParticleModule.h
--- a/Engine/Source/Runtime/Engine/Classes/Particles/ParticleModule.h
+++ b/Engine/Source/Runtime/Engine/Classes/Particles/ParticleModule.h
@@ -73,7 +73,7 @@ struct FParticleCurvePair
 	FString CurveName;
 
 	UPROPERTY()
-	class UObject* CurveObject;
+	TObjectPtr<class UObject> CurveObject;
 
 
 	FParticleCurvePair()
diff --git a/Engine/Source/Runtime/Engine/Classes/Particles/ParticleModuleRequired.h b/Engine/Source/Runtime/Engine/Classes/Particles/ParticleModuleRequired.h
--- a/Engine/Source/Runtime/Engine/Classes/Particles/ParticleModuleRequired.h
+++ b/Engine/Source/Runtime/Engine/Classes/Particles/ParticleModuleRequired.h
@@ -128,7 +128,7 @@ class UParticleModuleRequired : public UParticleModule
 	
 	/** The material to utilize for the emitter at this LOD level.						*/
 	UPROPERTY(EditAnywhere, Category=Emitter)
-	class UMaterialInterface* Material;
+	TObjectPtr<class UMaterialInterface> Material;
 
 	/** The distance at which PSA_FacingCameraDistanceBlend	is fully PSA_Square */
 	UPROPERTY(EditAnywhere, Category = Emitter, meta = (UIMin = "0", DisplayAfter="ScreenAlignment"))
@@ -378,7 +378,7 @@ class UParticleModuleRequired : public UParticleModule
 	* Texture to generate bounding geometry from.
 	*/
 	UPROPERTY(EditAnywhere, Category = ParticleCutout)
-	UTexture2D* CutoutTexture;
+	TObjectPtr<UTexture2D> CutoutTexture;
 
 	/**
 	 *	The maximum number of particles to DRAW for this emitter.
diff --git a/Engine/Source/Runtime/Engine/Classes/Particles/ParticleSystem.h b/Engine/Source/Runtime/Engine/Classes/Particles/ParticleSystem.h
--- a/Engine/Source/Runtime/Engine/Classes/Particles/ParticleSystem.h
+++ b/Engine/Source/Runtime/Engine/Classes/Particles/ParticleSystem.h
@@ -105,7 +105,7 @@ struct FNamedEmitterMaterial
 	FName Name;
 
 	UPROPERTY(EditAnywhere, Category = NamedMaterial)
-	UMaterialInterface* Material;
+	TObjectPtr<UMaterialInterface> Material;
 };
 
 UCLASS(Abstract, MinimalAPI, BlueprintType)
@@ -173,11 +173,11 @@ class UParticleSystem : public UFXSystemAsset
 
 	/** Emitters	- internal - the array of emitters in the system				*/
 	UPROPERTY(instanced)
-	TArray<UParticleEmitter*> Emitters;
+	TArray<TObjectPtr<UParticleEmitter>> Emitters;
 
 	/** The component used to preview the particle system in Cascade				*/
 	UPROPERTY(transient)
-	UParticleSystemComponent* PreviewComponent;
+	TObjectPtr<UParticleSystemComponent> PreviewComponent;
 
 #if WITH_EDITORONLY_DATA
 	/** The angle to use when rendering the thumbnail image							*/
@@ -195,7 +195,7 @@ class UParticleSystem : public UFXSystemAsset
 #endif // WITH_EDITORONLY_DATA
 	/** Used for curve editor to remember curve-editing setup.						*/
 	UPROPERTY(export)
-	UInterpCurveEdSetup* CurveEdSetup;
+	TObjectPtr<UInterpCurveEdSetup> CurveEdSetup;
 
 	//
 	//	LOD
@@ -315,7 +315,7 @@ public:
 #if WITH_EDITORONLY_DATA
 	/** Internal: The PSys thumbnail image									*/
 	UPROPERTY()
-	class UTexture2D* ThumbnailImage;
+	TObjectPtr<class UTexture2D> ThumbnailImage;
 
 #endif // WITH_EDITORONLY_DATA
 
diff --git a/Engine/Source/Runtime/Engine/Classes/Particles/ParticleSystemComponent.h b/Engine/Source/Runtime/Engine/Classes/Particles/ParticleSystemComponent.h
--- a/Engine/Source/Runtime/Engine/Classes/Particles/ParticleSystemComponent.h
+++ b/Engine/Source/Runtime/Engine/Classes/Particles/ParticleSystemComponent.h
@@ -138,10 +138,10 @@ struct FParticleSysParam
 	FColor Color;
 
 	UPROPERTY(EditAnywhere, BlueprintReadWrite, Category=ParticleSysParam)
-	class AActor* Actor;
+	TObjectPtr<class AActor> Actor;
 
 	UPROPERTY(EditAnywhere, BlueprintReadWrite, Category=ParticleSysParam)
-	class UMaterialInterface* Material;
+	TObjectPtr<class UMaterialInterface> Material;
 
 	FParticleSysParam()
 		: ParamType(0)
@@ -433,17 +433,17 @@ class ENGINE_API UParticleSystemComponent : public UFXSystemComponent
 	GENERATED_UCLASS_BODY()
 
 	UPROPERTY(EditAnywhere, BlueprintReadOnly, Category=Particles)
-	class UParticleSystem* Template;
+	TObjectPtr<class UParticleSystem> Template;
 
 	UPROPERTY(transient, duplicatetransient)
-	TArray<class UMaterialInterface*> EmitterMaterials;
+	TArray<TObjectPtr<class UMaterialInterface>> EmitterMaterials;
 
 	/**
 	 *	The skeletal mesh components used with the socket location module.
 	 *	This is to prevent them from being garbage collected.
 	 */
 	UPROPERTY(transient, duplicatetransient)
-	TArray<class USkeletalMeshComponent*> SkelMeshComponents;
+	TArray<TObjectPtr<class USkeletalMeshComponent>> SkelMeshComponents;
 
 	uint8 bWasCompleted:1;
 
@@ -693,7 +693,7 @@ public:
 
 	/** Array of replay clips for this particle system component.  These are serialized to disk.  You really should never add anything to this in the editor.  It's exposed so that you can delete clips if you need to, but be careful when doing so! */
 	UPROPERTY()
-	TArray<class UParticleSystemReplay*> ReplayClips;
+	TArray<TObjectPtr<class UParticleSystemReplay>> ReplayClips;
 
 	/** Clip ID number we're either playing back or capturing to, depending on the value of ReplayState. */
 	int32 ReplayClipIDNumber;
diff --git a/Engine/Source/Runtime/Engine/Classes/Particles/SubUV/ParticleModuleSubUV.h b/Engine/Source/Runtime/Engine/Classes/Particles/SubUV/ParticleModuleSubUV.h
--- a/Engine/Source/Runtime/Engine/Classes/Particles/SubUV/ParticleModuleSubUV.h
+++ b/Engine/Source/Runtime/Engine/Classes/Particles/SubUV/ParticleModuleSubUV.h
@@ -25,7 +25,7 @@ class UParticleModuleSubUV : public UParticleModuleSubUVBase
 	 * When specified, SubImages_Horizontal and SubImages_Vertical will come from the asset instead of the Required Module.
 	 */
 	UPROPERTY()
-	USubUVAnimation* Animation;
+	TObjectPtr<USubUVAnimation> Animation;
 
 	/**
 	 *	The index of the sub-image that should be used for the particle.
diff --git a/Engine/Source/Runtime/Engine/Classes/Particles/SubUVAnimation.h b/Engine/Source/Runtime/Engine/Classes/Particles/SubUVAnimation.h
--- a/Engine/Source/Runtime/Engine/Classes/Particles/SubUVAnimation.h
+++ b/Engine/Source/Runtime/Engine/Classes/Particles/SubUVAnimation.h
@@ -117,7 +117,7 @@ class USubUVAnimation : public UObject
 	 * Texture to generate bounding geometry from.
 	 */
 	UPROPERTY(EditAnywhere, Category=SubUV)
-	UTexture2D* SubUVTexture;
+	TObjectPtr<UTexture2D> SubUVTexture;
 
 	/** The number of sub-images horizontally in the texture							*/
 	UPROPERTY(EditAnywhere, Category=SubUV)
diff --git a/Engine/Source/Runtime/Engine/Classes/Particles/TypeData/ParticleModuleTypeDataGpu.h b/Engine/Source/Runtime/Engine/Classes/Particles/TypeData/ParticleModuleTypeDataGpu.h
--- a/Engine/Source/Runtime/Engine/Classes/Particles/TypeData/ParticleModuleTypeDataGpu.h
+++ b/Engine/Source/Runtime/Engine/Classes/Particles/TypeData/ParticleModuleTypeDataGpu.h
@@ -27,7 +27,7 @@ struct FGPUSpriteLocalVectorFieldInfo
 
 	/** Local vector field to apply to this emitter. */
 	UPROPERTY()
-	class UVectorField* Field;
+	TObjectPtr<class UVectorField> Field;
 
 	/** Local vector field transform. */
 	UPROPERTY()
@@ -97,19 +97,19 @@ struct FGPUSpriteEmitterInfo
 
 	/** The required module. Needed for now, but should be divorced from the runtime. */
 	UPROPERTY()
-	class UParticleModuleRequired* RequiredModule;
+	TObjectPtr<class UParticleModuleRequired> RequiredModule;
 
 	/** The spawn module. Needed for now, but should be divorced from the runtime. */
 	UPROPERTY()
-	class UParticleModuleSpawn* SpawnModule;
+	TObjectPtr<class UParticleModuleSpawn> SpawnModule;
 
 	/** The spawn-per-unit module. */
 	UPROPERTY()
-	class UParticleModuleSpawnPerUnit* SpawnPerUnitModule;
+	TObjectPtr<class UParticleModuleSpawnPerUnit> SpawnPerUnitModule;
 
 	/** List of spawn modules that must be evaluated at runtime. */
 	UPROPERTY()
-	TArray<class UParticleModule*> SpawnModules;
+	TArray<TObjectPtr<class UParticleModule>> SpawnModules;
 
 	/** Local vector field info. */
 	UPROPERTY()
diff --git a/Engine/Source/Runtime/Engine/Classes/Particles/TypeData/ParticleModuleTypeDataMesh.h b/Engine/Source/Runtime/Engine/Classes/Particles/TypeData/ParticleModuleTypeDataMesh.h
--- a/Engine/Source/Runtime/Engine/Classes/Particles/TypeData/ParticleModuleTypeDataMesh.h
+++ b/Engine/Source/Runtime/Engine/Classes/Particles/TypeData/ParticleModuleTypeDataMesh.h
@@ -65,7 +65,7 @@ class UParticleModuleTypeDataMesh : public UParticleModuleTypeDataBase
 
 	/** The static mesh to render at the particle positions */
 	UPROPERTY(EditAnywhere, Category=Mesh)
-	UStaticMesh* Mesh;
+	TObjectPtr<UStaticMesh> Mesh;
 
 	/** Random stream for the initial rotation distribution */
 	FRandomStream RandomStream;
diff --git a/Engine/Source/Runtime/Engine/Classes/Particles/VectorField/ParticleModuleVectorFieldLocal.h b/Engine/Source/Runtime/Engine/Classes/Particles/VectorField/ParticleModuleVectorFieldLocal.h
--- a/Engine/Source/Runtime/Engine/Classes/Particles/VectorField/ParticleModuleVectorFieldLocal.h
+++ b/Engine/Source/Runtime/Engine/Classes/Particles/VectorField/ParticleModuleVectorFieldLocal.h
@@ -18,7 +18,7 @@ class UParticleModuleVectorFieldLocal : public UParticleModuleVectorFieldBase
 
 	/** Vector field asset to use. */
 	UPROPERTY(EditAnywhere, Category=VectorField)
-	class UVectorField* VectorField;
+	TObjectPtr<class UVectorField> VectorField;
 
 	/** Translation of the vector field relative to the emitter. */
 	UPROPERTY(EditAnywhere, Category=VectorField)
diff --git a/Engine/Source/Runtime/Engine/Classes/Particles/VectorField/ParticleModuleVectorFieldScale.h b/Engine/Source/Runtime/Engine/Classes/Particles/VectorField/ParticleModuleVectorFieldScale.h
--- a/Engine/Source/Runtime/Engine/Classes/Particles/VectorField/ParticleModuleVectorFieldScale.h
+++ b/Engine/Source/Runtime/Engine/Classes/Particles/VectorField/ParticleModuleVectorFieldScale.h
@@ -21,7 +21,7 @@ class UParticleModuleVectorFieldScale : public UParticleModuleVectorFieldBase
 
 	/** Per-particle vector field scale. Evaluated using emitter time. */
 	UPROPERTY()
-	class UDistributionFloat* VectorFieldScale_DEPRECATED;
+	TObjectPtr<class UDistributionFloat> VectorFieldScale_DEPRECATED;
 
 	/** Per-particle vector field scale. Evaluated using emitter time. */
 	UPROPERTY(EditAnywhere, Category=VectorField)
diff --git a/Engine/Source/Runtime/Engine/Classes/Particles/VectorField/ParticleModuleVectorFieldScaleOverLife.h b/Engine/Source/Runtime/Engine/Classes/Particles/VectorField/ParticleModuleVectorFieldScaleOverLife.h
--- a/Engine/Source/Runtime/Engine/Classes/Particles/VectorField/ParticleModuleVectorFieldScaleOverLife.h
+++ b/Engine/Source/Runtime/Engine/Classes/Particles/VectorField/ParticleModuleVectorFieldScaleOverLife.h
@@ -21,7 +21,7 @@ class UParticleModuleVectorFieldScaleOverLife : public UParticleModuleVectorFiel
 
 	/** Per-particle vector field scale. Evaluated using particle relative time. */
 	UPROPERTY()
-	class UDistributionFloat* VectorFieldScaleOverLife_DEPRECATED;
+	TObjectPtr<class UDistributionFloat> VectorFieldScaleOverLife_DEPRECATED;
 
 	/** Per-particle vector field scale. Evaluated using particle relative time. */
 	UPROPERTY(EditAnywhere, Category=VectorField)
diff --git a/Engine/Source/Runtime/Engine/Classes/Particles/WorldPSCPool.h b/Engine/Source/Runtime/Engine/Classes/Particles/WorldPSCPool.h
--- a/Engine/Source/Runtime/Engine/Classes/Particles/WorldPSCPool.h
+++ b/Engine/Source/Runtime/Engine/Classes/Particles/WorldPSCPool.h
@@ -51,7 +51,7 @@ struct FPSCPoolElem
 	GENERATED_BODY()
 
 	UPROPERTY(transient)
-	UParticleSystemComponent* PSC;
+	TObjectPtr<UParticleSystemComponent> PSC;
 
 	float LastUsedTime;
 
@@ -111,7 +111,7 @@ struct ENGINE_API FWorldPSCPool
 
 private:
 	UPROPERTY()
-	TMap<UParticleSystem*, FPSCPool> WorldParticleSystemPools;
+	TMap<TObjectPtr<UParticleSystem>, FPSCPool> WorldParticleSystemPools;
 
 	float LastParticleSytemPoolCleanTime;
 
diff --git a/Engine/Source/Runtime/Engine/Classes/PhysicalMaterials/PhysicalMaterialMask.h b/Engine/Source/Runtime/Engine/Classes/PhysicalMaterials/PhysicalMaterialMask.h
--- a/Engine/Source/Runtime/Engine/Classes/PhysicalMaterials/PhysicalMaterialMask.h
+++ b/Engine/Source/Runtime/Engine/Classes/PhysicalMaterials/PhysicalMaterialMask.h
@@ -30,11 +30,11 @@ class ENGINE_API UPhysicalMaterialMask : public UObject
 #if WITH_EDITORONLY_DATA
 
 	UPROPERTY(VisibleAnywhere, Instanced, Category = ImportSettings)
-	class UAssetImportData* AssetImportData;
+	TObjectPtr<class UAssetImportData> AssetImportData;
 
 	/** Mask input texture, square aspect ratio recommended. Recognized mask colors include: white, black, red, green, yellow, cyan, turquoise, and magenta. */
 	UPROPERTY(VisibleAnywhere, BlueprintReadOnly, Category = TextureSource)
-	UTexture* MaskTexture;
+	TObjectPtr<UTexture> MaskTexture;
 
 #endif
 
diff --git a/Engine/Source/Runtime/Engine/Classes/PhysicsEngine/BodyInstance.h b/Engine/Source/Runtime/Engine/Classes/PhysicsEngine/BodyInstance.h
--- a/Engine/Source/Runtime/Engine/Classes/PhysicsEngine/BodyInstance.h
+++ b/Engine/Source/Runtime/Engine/Classes/PhysicsEngine/BodyInstance.h
@@ -509,7 +509,7 @@ protected:
 
 	/**	Allows you to override the PhysicalMaterial to use for simple collision on this body. */
 	UPROPERTY(EditAnywhere, BlueprintReadOnly, Category=Collision)
-	class UPhysicalMaterial* PhysMaterialOverride;
+	TObjectPtr<class UPhysicalMaterial> PhysMaterialOverride;
 
 public:
 	/** The maximum angular velocity for this instance [degrees/s]*/
diff --git a/Engine/Source/Runtime/Engine/Classes/PhysicsEngine/BodySetup.h b/Engine/Source/Runtime/Engine/Classes/PhysicsEngine/BodySetup.h
--- a/Engine/Source/Runtime/Engine/Classes/PhysicsEngine/BodySetup.h
+++ b/Engine/Source/Runtime/Engine/Classes/PhysicsEngine/BodySetup.h
@@ -224,7 +224,7 @@ class UBodySetup : public UBodySetupCore
 	
 	/** Physical material to use for simple collision on this body. Encodes information about density, friction etc. */
 	UPROPERTY(EditAnywhere, Category=Physics, meta=(DisplayName="Simple Collision Physical Material"))
-	class UPhysicalMaterial* PhysMaterial;
+	TObjectPtr<class UPhysicalMaterial> PhysMaterial;
 
 	/** Custom walkable slope setting for this body. */
 	UPROPERTY(EditAnywhere, AdvancedDisplay, Category=Physics)
diff --git a/Engine/Source/Runtime/Engine/Classes/PhysicsEngine/PhysicalAnimationComponent.h b/Engine/Source/Runtime/Engine/Classes/PhysicsEngine/PhysicalAnimationComponent.h
--- a/Engine/Source/Runtime/Engine/Classes/PhysicsEngine/PhysicalAnimationComponent.h
+++ b/Engine/Source/Runtime/Engine/Classes/PhysicsEngine/PhysicalAnimationComponent.h
@@ -124,7 +124,7 @@ public:
 
 private:
 	UPROPERTY()
-	USkeletalMeshComponent* SkeletalMeshComponent;
+	TObjectPtr<USkeletalMeshComponent> SkeletalMeshComponent;
 
 	struct FPhysicalAnimationInstanceData
 	{
diff --git a/Engine/Source/Runtime/Engine/Classes/PhysicsEngine/PhysicsAsset.h b/Engine/Source/Runtime/Engine/Classes/PhysicsEngine/PhysicsAsset.h
--- a/Engine/Source/Runtime/Engine/Classes/PhysicsEngine/PhysicsAsset.h
+++ b/Engine/Source/Runtime/Engine/Classes/PhysicsEngine/PhysicsAsset.h
@@ -98,7 +98,7 @@ class UPhysicsAsset : public UObject, public IInterface_PreviewMeshProvider
 	 *	Is the one that was used as the basis for creating this Asset.
 	 */
 	UPROPERTY()
-	class USkeletalMesh * DefaultSkelMesh_DEPRECATED;
+	TObjectPtr<class USkeletalMesh>  DefaultSkelMesh_DEPRECATED;
 
 	UPROPERTY(AssetRegistrySearchable)
 	TSoftObjectPtr<class USkeletalMesh> PreviewSkeletalMesh;
@@ -126,14 +126,14 @@ class UPhysicsAsset : public UObject, public IInterface_PreviewMeshProvider
 	*	Does not include body position - those are taken from mesh.
 	*/
 	UPROPERTY(instanced)
-	TArray<USkeletalBodySetup*> SkeletalBodySetups;
+	TArray<TObjectPtr<USkeletalBodySetup>> SkeletalBodySetups;
 
 	/** 
 	 *	Array of RB_ConstraintSetup objects. 
 	 *	Stores information about a joint between two bodies, such as position relative to each body, joint limits etc.
 	 */
 	UPROPERTY(instanced)
-	TArray<class UPhysicsConstraintTemplate*> ConstraintSetup;
+	TArray<TObjectPtr<class UPhysicsConstraintTemplate>> ConstraintSetup;
 
 public:
 
@@ -156,7 +156,7 @@ public:
 
 	/** Information for thumbnail rendering */
 	UPROPERTY(VisibleAnywhere, Instanced, Category = Thumbnail)
-	class UThumbnailInfo* ThumbnailInfo;
+	TObjectPtr<class UThumbnailInfo> ThumbnailInfo;
 
 	//~ Begin UObject Interface
 	virtual void Serialize(FArchive& Ar) override;
@@ -280,7 +280,7 @@ private:
 
 
 	UPROPERTY(instanced)
-	TArray<class UBodySetup*> BodySetup_DEPRECATED;
+	TArray<TObjectPtr<class UBodySetup>> BodySetup_DEPRECATED;
 };
 
 USTRUCT()
diff --git a/Engine/Source/Runtime/Engine/Classes/PhysicsEngine/PhysicsCollisionHandler.h b/Engine/Source/Runtime/Engine/Classes/PhysicsEngine/PhysicsCollisionHandler.h
--- a/Engine/Source/Runtime/Engine/Classes/PhysicsEngine/PhysicsCollisionHandler.h
+++ b/Engine/Source/Runtime/Engine/Classes/PhysicsEngine/PhysicsCollisionHandler.h
@@ -28,7 +28,7 @@ class ENGINE_API UPhysicsCollisionHandler : public UObject
 
 	/** Sound to play  */
 	UPROPERTY(EditAnywhere, BlueprintReadWrite, Category=Impact)
-	class USoundBase* DefaultImpactSound;
+	TObjectPtr<class USoundBase> DefaultImpactSound;
 
 	/** Time since last impact sound */
 	UPROPERTY()
diff --git a/Engine/Source/Runtime/Engine/Classes/PhysicsEngine/PhysicsConstraintActor.h b/Engine/Source/Runtime/Engine/Classes/PhysicsEngine/PhysicsConstraintActor.h
--- a/Engine/Source/Runtime/Engine/Classes/PhysicsEngine/PhysicsConstraintActor.h
+++ b/Engine/Source/Runtime/Engine/Classes/PhysicsEngine/PhysicsConstraintActor.h
@@ -18,13 +18,13 @@ class APhysicsConstraintActor : public ARigidBodyBase
 	// Cached reference to constraint component
 private:
 	UPROPERTY(Category = ConstraintActor, VisibleAnywhere, BlueprintReadOnly, meta = (ExposeFunctionCategories = "JointDrive,Physics|Components|PhysicsConstraint", AllowPrivateAccess = "true"))
-	class UPhysicsConstraintComponent* ConstraintComp;
+	TObjectPtr<class UPhysicsConstraintComponent> ConstraintComp;
 public:
 	
 	UPROPERTY()
-	class AActor* ConstraintActor1_DEPRECATED;
+	TObjectPtr<class AActor> ConstraintActor1_DEPRECATED;
 	UPROPERTY()
-	class AActor* ConstraintActor2_DEPRECATED;
+	TObjectPtr<class AActor> ConstraintActor2_DEPRECATED;
 	UPROPERTY()
 	uint32 bDisableCollision_DEPRECATED:1;
 
diff --git a/Engine/Source/Runtime/Engine/Classes/PhysicsEngine/PhysicsConstraintComponent.h b/Engine/Source/Runtime/Engine/Classes/PhysicsEngine/PhysicsConstraintComponent.h
--- a/Engine/Source/Runtime/Engine/Classes/PhysicsEngine/PhysicsConstraintComponent.h
+++ b/Engine/Source/Runtime/Engine/Classes/PhysicsEngine/PhysicsConstraintComponent.h
@@ -22,7 +22,7 @@ class ENGINE_API UPhysicsConstraintComponent : public USceneComponent
 
 	/** Pointer to first Actor to constrain.  */
 	UPROPERTY(EditInstanceOnly, Category=Constraint)
-	AActor* ConstraintActor1;
+	TObjectPtr<AActor> ConstraintActor1;
 
 	/** 
 	 *	Name of first component property to constrain. If Actor1 is NULL, will look within Owner.
@@ -34,7 +34,7 @@ class ENGINE_API UPhysicsConstraintComponent : public USceneComponent
 
 	/** Pointer to second Actor to constrain. */
 	UPROPERTY(EditInstanceOnly, Category=Constraint)
-	AActor* ConstraintActor2;
+	TObjectPtr<AActor> ConstraintActor2;
 
 	/** 
 	 *	Name of second component property to constrain. If Actor2 is NULL, will look within Owner. 
@@ -52,7 +52,7 @@ class ENGINE_API UPhysicsConstraintComponent : public USceneComponent
 
 
 	UPROPERTY(instanced)
-	class UPhysicsConstraintTemplate* ConstraintSetup_DEPRECATED;
+	TObjectPtr<class UPhysicsConstraintTemplate> ConstraintSetup_DEPRECATED;
 
 	/** Notification when constraint is broken. */
 	UPROPERTY(BlueprintAssignable)
diff --git a/Engine/Source/Runtime/Engine/Classes/PhysicsEngine/PhysicsHandleComponent.h b/Engine/Source/Runtime/Engine/Classes/PhysicsEngine/PhysicsHandleComponent.h
--- a/Engine/Source/Runtime/Engine/Classes/PhysicsEngine/PhysicsHandleComponent.h
+++ b/Engine/Source/Runtime/Engine/Classes/PhysicsEngine/PhysicsHandleComponent.h
@@ -27,7 +27,7 @@ class ENGINE_API UPhysicsHandleComponent : public UActorComponent
 
 	/** Component we are currently holding */
 	UPROPERTY()
-	class UPrimitiveComponent* GrabbedComponent;
+	TObjectPtr<class UPrimitiveComponent> GrabbedComponent;
 
 	/** Name of bone, if we are grabbing a skeletal component */
 	FName GrabbedBoneName;
diff --git a/Engine/Source/Runtime/Engine/Classes/PhysicsEngine/PhysicsThruster.h b/Engine/Source/Runtime/Engine/Classes/PhysicsEngine/PhysicsThruster.h
--- a/Engine/Source/Runtime/Engine/Classes/PhysicsEngine/PhysicsThruster.h
+++ b/Engine/Source/Runtime/Engine/Classes/PhysicsEngine/PhysicsThruster.h
@@ -21,11 +21,11 @@ class APhysicsThruster : public ARigidBodyBase
 private:
 	/** Thruster component */
 	UPROPERTY(Category = Physics, VisibleAnywhere, BlueprintReadOnly, meta = (ExposeFunctionCategories = "Activation,Components|Activation", AllowPrivateAccess = "true"))
-	class UPhysicsThrusterComponent* ThrusterComponent;
+	TObjectPtr<class UPhysicsThrusterComponent> ThrusterComponent;
 
 #if WITH_EDITORONLY_DATA
 	UPROPERTY()
-	class UArrowComponent* ArrowComponent;
+	TObjectPtr<class UArrowComponent> ArrowComponent;
 #endif
 
 public:
diff --git a/Engine/Source/Runtime/Engine/Classes/PhysicsEngine/RadialForceActor.h b/Engine/Source/Runtime/Engine/Classes/PhysicsEngine/RadialForceActor.h
--- a/Engine/Source/Runtime/Engine/Classes/PhysicsEngine/RadialForceActor.h
+++ b/Engine/Source/Runtime/Engine/Classes/PhysicsEngine/RadialForceActor.h
@@ -18,11 +18,11 @@ class ARadialForceActor : public ARigidBodyBase
 private:
 	/** Force component */
 	UPROPERTY(Category = RadialForceActor, VisibleAnywhere, BlueprintReadOnly, meta = (ExposeFunctionCategories = "Activation,Components|Activation,Physics,Physics|Components|RadialForce", AllowPrivateAccess = "true"))
-	class URadialForceComponent* ForceComponent;
+	TObjectPtr<class URadialForceComponent> ForceComponent;
 
 #if WITH_EDITORONLY_DATA
 	UPROPERTY()
-	UBillboardComponent* SpriteComponent;
+	TObjectPtr<UBillboardComponent> SpriteComponent;
 #endif
 public:
 
diff --git a/Engine/Source/Runtime/Engine/Classes/Sound/AmbientSound.h b/Engine/Source/Runtime/Engine/Classes/Sound/AmbientSound.h
--- a/Engine/Source/Runtime/Engine/Classes/Sound/AmbientSound.h
+++ b/Engine/Source/Runtime/Engine/Classes/Sound/AmbientSound.h
@@ -16,7 +16,7 @@ class ENGINE_API AAmbientSound : public AActor
 private:
 	/** Audio component that handles sound playing */
 	UPROPERTY(Category = Sound, VisibleAnywhere, BlueprintReadOnly, meta = (ExposeFunctionCategories = "Sound,Audio,Audio|Components|Audio", AllowPrivateAccess = "true"))
-	class UAudioComponent* AudioComponent;
+	TObjectPtr<class UAudioComponent> AudioComponent;
 public:
 	
 	FString GetInternalSoundCueName();
diff --git a/Engine/Source/Runtime/Engine/Classes/Sound/AudioVolume.h b/Engine/Source/Runtime/Engine/Classes/Sound/AudioVolume.h
--- a/Engine/Source/Runtime/Engine/Classes/Sound/AudioVolume.h
+++ b/Engine/Source/Runtime/Engine/Classes/Sound/AudioVolume.h
@@ -58,11 +58,11 @@ struct FAudioVolumeSubmixOverrideSettings
 
 	// The submix to override the effect chain of
 	UPROPERTY(EditAnywhere, BlueprintReadWrite, Category = AudioVolumeSubmixSends)
-	USoundSubmix* Submix = nullptr;
+	TObjectPtr<USoundSubmix> Submix = nullptr;
 
 	// The submix effect chain to overrideac
 	UPROPERTY(EditAnywhere, BlueprintReadOnly, Category = SoundSubmix)
-	TArray<USoundEffectSubmixPreset*> SubmixEffectChain;
+	TArray<TObjectPtr<USoundEffectSubmixPreset>> SubmixEffectChain;
 
 	// The amount of time to crossfade to the override for the submix chain
 	UPROPERTY(EditAnywhere, BlueprintReadOnly, Category = SoundSubmix)
diff --git a/Engine/Source/Runtime/Engine/Classes/Sound/DialogueTypes.h b/Engine/Source/Runtime/Engine/Classes/Sound/DialogueTypes.h
--- a/Engine/Source/Runtime/Engine/Classes/Sound/DialogueTypes.h
+++ b/Engine/Source/Runtime/Engine/Classes/Sound/DialogueTypes.h
@@ -48,11 +48,11 @@ struct ENGINE_API FDialogueContext
 
 	/** The person speaking the dialogue. */
 	UPROPERTY(EditAnywhere, BlueprintReadWrite, Category=DialogueContext )
-	UDialogueVoice* Speaker;
+	TObjectPtr<UDialogueVoice> Speaker;
 
 	/** The people being spoken to. */
 	UPROPERTY(EditAnywhere, BlueprintReadWrite, Category=DialogueContext )
-	TArray<UDialogueVoice*> Targets;
+	TArray<TObjectPtr<UDialogueVoice>> Targets;
 
 	/** Gets a generated hash created from the source and targets. */
 	FString GetContextHash() const;
@@ -70,7 +70,7 @@ struct ENGINE_API FDialogueWaveParameter
 
 	/** The dialogue wave to play. */
 	UPROPERTY(EditAnywhere, Category=DialogueWaveParameter )
-	UDialogueWave* DialogueWave;
+	TObjectPtr<UDialogueWave> DialogueWave;
 
 	/** The context to use for the dialogue wave. */
 	UPROPERTY(EditAnywhere, Category=DialogueWaveParameter )
diff --git a/Engine/Source/Runtime/Engine/Classes/Sound/DialogueWave.h b/Engine/Source/Runtime/Engine/Classes/Sound/DialogueWave.h
--- a/Engine/Source/Runtime/Engine/Classes/Sound/DialogueWave.h
+++ b/Engine/Source/Runtime/Engine/Classes/Sound/DialogueWave.h
@@ -53,7 +53,7 @@ struct ENGINE_API FDialogueContextMapping
 
 	/** The soundwave to play for this dialogue. */
 	UPROPERTY(EditAnywhere, Category=DialogueContextMapping )
-	USoundWave* SoundWave;
+	TObjectPtr<USoundWave> SoundWave;
 
 	/**
 	 * The format string to use when generating the localization key for this context. This must be unique within the owner dialogue wave.
@@ -65,7 +65,7 @@ struct ENGINE_API FDialogueContextMapping
 
 	/** Cached object for playing the soundwave with subtitle information included. */
 	UPROPERTY(Transient)
-	UDialogueSoundWaveProxy* Proxy;
+	TObjectPtr<UDialogueSoundWaveProxy> Proxy;
 
 	/** Gets the localization key to use for this context mapping */
 	FString GetLocalizationKey() const;
diff --git a/Engine/Source/Runtime/Engine/Classes/Sound/ReverbSettings.h b/Engine/Source/Runtime/Engine/Classes/Sound/ReverbSettings.h
--- a/Engine/Source/Runtime/Engine/Classes/Sound/ReverbSettings.h
+++ b/Engine/Source/Runtime/Engine/Classes/Sound/ReverbSettings.h
@@ -61,11 +61,11 @@ struct FReverbSettings
 
 	/** The reverb asset to employ. */
 	UPROPERTY(EditAnywhere, BlueprintReadWrite, Category=ReverbSettings)
-	UReverbEffect* ReverbEffect;
+	TObjectPtr<UReverbEffect> ReverbEffect;
 
 	/** This is used to apply plugin-specific settings when a Reverb Plugin is being used. */
 	UPROPERTY(EditAnywhere, BlueprintReadWrite, Category = ReverbSettings)
-	USoundEffectSubmixPreset* ReverbPluginEffect;
+	TObjectPtr<USoundEffectSubmixPreset> ReverbPluginEffect;
 
 	/** Volume level of the reverb affect. */
 	UPROPERTY(EditAnywhere, BlueprintReadWrite, Category = ReverbSettings)
diff --git a/Engine/Source/Runtime/Engine/Classes/Sound/SoundAttenuation.h b/Engine/Source/Runtime/Engine/Classes/Sound/SoundAttenuation.h
--- a/Engine/Source/Runtime/Engine/Classes/Sound/SoundAttenuation.h
+++ b/Engine/Source/Runtime/Engine/Classes/Sound/SoundAttenuation.h
@@ -88,15 +88,15 @@ struct ENGINE_API FSoundAttenuationPluginSettings
 
 	/** Settings to use with spatialization audio plugin. These are defined by the plugin creator. Not all audio plugins utilize this feature. This is an array so multiple plugins can have settings. */
 	UPROPERTY(EditAnywhere, BlueprintReadWrite, Category = AttenuationSpatialization, meta = (DisplayName = "Spatialization Plugin Settings"))
-	TArray<USpatializationPluginSourceSettingsBase*> SpatializationPluginSettingsArray;
+	TArray<TObjectPtr<USpatializationPluginSourceSettingsBase>> SpatializationPluginSettingsArray;
 
 	/** Settings to use with occlusion audio plugin. These are defined by the plugin creator. Not all audio plugins utilize this feature. This  is an array so multiple plugins can have settings. */
 	UPROPERTY(EditAnywhere, BlueprintReadWrite, Category = AttenuationOcclusion, meta = (DisplayName = "Occlusion Plugin Settings"))
-	TArray<UOcclusionPluginSourceSettingsBase*> OcclusionPluginSettingsArray;
+	TArray<TObjectPtr<UOcclusionPluginSourceSettingsBase>> OcclusionPluginSettingsArray;
 
 	/** Settings to use with reverb audio plugin. These are defined by the plugin creator. Not all audio plugins utilize this feature. This  is an array so multiple plugins can have settings. */
 	UPROPERTY(EditAnywhere, BlueprintReadWrite, Category = AttenuationReverbSend, meta = (DisplayName = "Reverb Plugin Settings"))
-	TArray<UReverbPluginSourceSettingsBase*> ReverbPluginSettingsArray;
+	TArray<TObjectPtr<UReverbPluginSourceSettingsBase>> ReverbPluginSettingsArray;
 };
 
 USTRUCT(BlueprintType)
@@ -106,7 +106,7 @@ struct ENGINE_API FAttenuationSubmixSendSettings
 
 	/** Submix to send audio to based on distance. */
 	UPROPERTY(EditAnywhere, BlueprintReadWrite, Category = AttenuationSubmixSend)
-	USoundSubmixBase* Submix = nullptr;
+	TObjectPtr<USoundSubmixBase> Submix = nullptr;
 
 	/** What method to use to use for submix sends. */
 	UPROPERTY(EditAnywhere, BlueprintReadWrite, Category = AttenuationReverbSend)
@@ -232,7 +232,7 @@ struct ENGINE_API FSoundAttenuationSettings : public FBaseAttenuationSettings
 
 #if WITH_EDITORONLY_DATA
 	UPROPERTY()
-	USpatializationPluginSourceSettingsBase* SpatializationPluginSettings_DEPRECATED;
+	TObjectPtr<USpatializationPluginSourceSettingsBase> SpatializationPluginSettings_DEPRECATED;
 
 	UPROPERTY()
 	float RadiusMin_DEPRECATED;
@@ -327,10 +327,10 @@ struct ENGINE_API FSoundAttenuationSettings : public FBaseAttenuationSettings
 
 #if WITH_EDITORONLY_DATA
 	UPROPERTY()
-	UOcclusionPluginSourceSettingsBase* OcclusionPluginSettings_DEPRECATED;
+	TObjectPtr<UOcclusionPluginSourceSettingsBase> OcclusionPluginSettings_DEPRECATED;
 
 	UPROPERTY()
-	UReverbPluginSourceSettingsBase* ReverbPluginSettings_DEPRECATED;
+	TObjectPtr<UReverbPluginSourceSettingsBase> ReverbPluginSettings_DEPRECATED;
 #endif
 
 	/** The amount to send to master reverb when sound is located at a distance equal to value specified in the reverb min send distance. */
diff --git a/Engine/Source/Runtime/Engine/Classes/Sound/SoundBase.h b/Engine/Source/Runtime/Engine/Classes/Sound/SoundBase.h
--- a/Engine/Source/Runtime/Engine/Classes/Sound/SoundBase.h
+++ b/Engine/Source/Runtime/Engine/Classes/Sound/SoundBase.h
@@ -63,7 +63,7 @@ public:
 
 	/** Sound class this sound belongs to */
 	UPROPERTY(EditAnywhere, Category = Sound, meta = (DisplayName = "Class"), AssetRegistrySearchable)
-	USoundClass* SoundClassObject;
+	TObjectPtr<USoundClass> SoundClassObject;
 
 	/** When "au.debug.Sounds -debug" has been specified, draw this sound's attenuation shape when the sound is audible. For debugging purpose only. */
 	UPROPERTY(EditAnywhere, Category = Developer)
@@ -109,12 +109,12 @@ public:
 #if WITH_EDITORONLY_DATA
 	/** If Override Concurrency is false, the sound concurrency settings to use for this sound. */
 	UPROPERTY()
-	USoundConcurrency* SoundConcurrencySettings_DEPRECATED;
+	TObjectPtr<USoundConcurrency> SoundConcurrencySettings_DEPRECATED;
 #endif // WITH_EDITORONLY_DATA
 
 	/** Set of concurrency settings to observe (if override is set to false).  Sound must pass all concurrency settings to play. */
 	UPROPERTY(EditAnywhere, BlueprintReadWrite, Category = "Voice Management|Concurrency", meta = (EditCondition = "!bOverrideConcurrency"))
-	TSet<USoundConcurrency*> ConcurrencySet;
+	TSet<TObjectPtr<USoundConcurrency>> ConcurrencySet;
 
 	/** If Override Concurrency is true, concurrency settings to use. */
 	UPROPERTY(EditAnywhere, BlueprintReadWrite, Category = "Voice Management|Concurrency", meta = (EditCondition = "bOverrideConcurrency"))
@@ -147,7 +147,7 @@ public:
 
 	/** Attenuation settings package for the sound */
 	UPROPERTY(EditAnywhere, Category = Attenuation)
-	USoundAttenuation* AttenuationSettings;
+	TObjectPtr<USoundAttenuation> AttenuationSettings;
 
 	/** Modulation Settings */
 	UPROPERTY(EditAnywhere, BlueprintReadWrite, Category = "Modulation")
@@ -156,7 +156,7 @@ public:
 	/** Submix to route sound output to. If unset, falls back to referenced SoundClass submix.
 	  * If SoundClass submix is unset, sends to the 'Master Submix' as set in the 'Audio' category of Project Settings'. */
 	UPROPERTY(EditAnywhere, BlueprintReadWrite, Category = "Effects|Submix", meta = (DisplayName = "Submix"))
-	USoundSubmixBase* SoundSubmixObject;
+	TObjectPtr<USoundSubmixBase> SoundSubmixObject;
 
 	/** Array of submix sends to which a prescribed amount (see 'Send Level') of this sound is sent. */
 	UPROPERTY(EditAnywhere, BlueprintReadWrite, Category = "Effects|Submix", meta = (DisplayName = "Submix Sends"))
@@ -164,7 +164,7 @@ public:
 
 	/** The source effect chain to use for this sound. */
 	UPROPERTY(EditAnywhere, BlueprintReadWrite, Category = "Effects|Source")
-	USoundEffectSourcePresetChain* SourceEffectChain;
+	TObjectPtr<USoundEffectSourcePresetChain> SourceEffectChain;
 
 	/** This sound will send its audio output to this list of buses if there are bus instances playing after source effects are processed. */
 	UPROPERTY(EditAnywhere, BlueprintReadWrite, Category = "Effects|Source", meta = (DisplayName = "Post-Effect Bus Sends"))
@@ -176,7 +176,7 @@ public:
 
 	/** Array of user data stored with the asset */
 	UPROPERTY(EditAnywhere, AdvancedDisplay, Instanced, Category = Sound)
-	TArray<UAssetUserData*> AssetUserData;
+	TArray<TObjectPtr<UAssetUserData>> AssetUserData;
 
 public:
 
diff --git a/Engine/Source/Runtime/Engine/Classes/Sound/SoundClass.h b/Engine/Source/Runtime/Engine/Classes/Sound/SoundClass.h
--- a/Engine/Source/Runtime/Engine/Classes/Sound/SoundClass.h
+++ b/Engine/Source/Runtime/Engine/Classes/Sound/SoundClass.h
@@ -140,7 +140,7 @@ struct FSoundClassProperties
 	/** Default output submix of referencing sounds. If unset, falls back to the 'Master Submix' as set in the 'Audio' category of Project Settings. 
 	  * (Unavailable if legacy 'Output to Master EQ Submix' is set) */
 	UPROPERTY(EditAnywhere, BlueprintReadOnly, Category = Submix, meta = (EditCondition = "!bApplyEffects"))
-	USoundSubmix* DefaultSubmix;
+	TObjectPtr<USoundSubmix> DefaultSubmix;
 
 	FSoundClassProperties();
 };
@@ -167,7 +167,7 @@ struct FPassiveSoundMixModifier
 
 	/** The SoundMix to activate */
 	UPROPERTY(EditAnywhere, BlueprintReadOnly, Category=PassiveSoundMixModifier)
-	class USoundMix* SoundMix;
+	TObjectPtr<class USoundMix> SoundMix;
 
 	/** Minimum volume level required to activate SoundMix. Below this value the SoundMix will not be active. */
 	UPROPERTY(EditAnywhere, BlueprintReadOnly, Category=PassiveSoundMixModifier)
@@ -212,14 +212,14 @@ public:
 	FSoundClassProperties Properties;
 
 	UPROPERTY(EditAnywhere, BlueprintReadOnly, Category = General)
-	TArray<USoundClass*> ChildClasses;
+	TArray<TObjectPtr<USoundClass>> ChildClasses;
 
 	/** SoundMix Modifiers to activate automatically when a sound of this class is playing. */
 	UPROPERTY(EditAnywhere, BlueprintReadOnly, Category = General)
 	TArray<FPassiveSoundMixModifier> PassiveSoundMixModifiers;
 
 	UPROPERTY(BlueprintReadOnly, Category = General)
-	USoundClass* ParentClass;
+	TObjectPtr<USoundClass> ParentClass;
 
 #if WITH_EDITORONLY_DATA
 	/** EdGraph based representation of the SoundClass */
diff --git a/Engine/Source/Runtime/Engine/Classes/Sound/SoundCue.h b/Engine/Source/Runtime/Engine/Classes/Sound/SoundCue.h
--- a/Engine/Source/Runtime/Engine/Classes/Sound/SoundCue.h
+++ b/Engine/Source/Runtime/Engine/Classes/Sound/SoundCue.h
@@ -94,7 +94,7 @@ class ENGINE_API USoundCue : public USoundBase
 	uint32 bPrimeOnLoad : 1;
 
 	UPROPERTY()
-	USoundNode* FirstNode;
+	TObjectPtr<USoundNode> FirstNode;
 
 	/* Base volume multiplier */
 	UPROPERTY(EditAnywhere, BlueprintReadOnly, Category = Sound, AssetRegistrySearchable)
@@ -110,10 +110,10 @@ class ENGINE_API USoundCue : public USoundBase
 
 #if WITH_EDITORONLY_DATA
 	UPROPERTY()
-	TArray<USoundNode*> AllNodes;
+	TArray<TObjectPtr<USoundNode>> AllNodes;
 
 	UPROPERTY()
-	UEdGraph* SoundCueGraph;
+	TObjectPtr<UEdGraph> SoundCueGraph;
 #endif
 
 protected:
diff --git a/Engine/Source/Runtime/Engine/Classes/Sound/SoundEffectSource.h b/Engine/Source/Runtime/Engine/Classes/Sound/SoundEffectSource.h
--- a/Engine/Source/Runtime/Engine/Classes/Sound/SoundEffectSource.h
+++ b/Engine/Source/Runtime/Engine/Classes/Sound/SoundEffectSource.h
@@ -27,7 +27,7 @@ struct ENGINE_API FSourceEffectChainEntry
 	GENERATED_USTRUCT_BODY()
 
 	UPROPERTY(EditAnywhere, BlueprintReadWrite, Category = "SourceEffect")
-	USoundEffectSourcePreset* Preset;
+	TObjectPtr<USoundEffectSourcePreset> Preset;
 
 	UPROPERTY(EditAnywhere, BlueprintReadWrite, Category = "SourceEffect")
 	uint32 bBypass:1;
diff --git a/Engine/Source/Runtime/Engine/Classes/Sound/SoundMix.h b/Engine/Source/Runtime/Engine/Classes/Sound/SoundMix.h
--- a/Engine/Source/Runtime/Engine/Classes/Sound/SoundMix.h
+++ b/Engine/Source/Runtime/Engine/Classes/Sound/SoundMix.h
@@ -134,7 +134,7 @@ struct FSoundClassAdjuster
 
 	/* The sound class this adjuster affects. */
 	UPROPERTY(EditAnywhere, BlueprintReadOnly, Category=SoundClassAdjuster, DisplayName = "Sound Class" )
-	USoundClass* SoundClassObject;
+	TObjectPtr<USoundClass> SoundClassObject;
 
 	/* A multiplier applied to the volume. */
 	UPROPERTY(EditAnywhere, BlueprintReadOnly, Category=SoundClassAdjuster, meta = (ClampMin = "0.0", UIMin = "0.0", UIMax = "4.0"))
diff --git a/Engine/Source/Runtime/Engine/Classes/Sound/SoundModulationDestination.h b/Engine/Source/Runtime/Engine/Classes/Sound/SoundModulationDestination.h
--- a/Engine/Source/Runtime/Engine/Classes/Sound/SoundModulationDestination.h
+++ b/Engine/Source/Runtime/Engine/Classes/Sound/SoundModulationDestination.h
@@ -45,7 +45,7 @@ struct ENGINE_API FSoundModulationDestinationSettings
 
 	/** Modulation source, which provides value to mix with base value. */
 	UPROPERTY(EditAnywhere, BlueprintReadWrite, Category = Modulation)
-	USoundModulatorBase* Modulator = nullptr;
+	TObjectPtr<USoundModulatorBase> Modulator = nullptr;
 };
 
 /** Default parameter destination settings for source audio object. */
diff --git a/Engine/Source/Runtime/Engine/Classes/Sound/SoundNode.h b/Engine/Source/Runtime/Engine/Classes/Sound/SoundNode.h
--- a/Engine/Source/Runtime/Engine/Classes/Sound/SoundNode.h
+++ b/Engine/Source/Runtime/Engine/Classes/Sound/SoundNode.h
@@ -61,12 +61,12 @@ class ENGINE_API USoundNode : public UObject
 	static const int32 MAX_ALLOWED_CHILD_NODES = 32;
 
 	UPROPERTY()
-	TArray<class USoundNode*> ChildNodes;
+	TArray<TObjectPtr<class USoundNode>> ChildNodes;
 
 #if WITH_EDITORONLY_DATA
 	/** Node's Graph representation, used to get position. */
 	UPROPERTY()
-	UEdGraphNode*	GraphNode;
+	TObjectPtr<UEdGraphNode>	GraphNode;
 
 	class UEdGraphNode* GetGraphNode() const;
 #endif
diff --git a/Engine/Source/Runtime/Engine/Classes/Sound/SoundNodeAttenuation.h b/Engine/Source/Runtime/Engine/Classes/Sound/SoundNodeAttenuation.h
--- a/Engine/Source/Runtime/Engine/Classes/Sound/SoundNodeAttenuation.h
+++ b/Engine/Source/Runtime/Engine/Classes/Sound/SoundNodeAttenuation.h
@@ -22,7 +22,7 @@ class USoundNodeAttenuation : public USoundNode
 	GENERATED_UCLASS_BODY()
 
 	UPROPERTY(EditAnywhere, Category=Attenuation, meta=(EditCondition="!bOverrideAttenuation"))
-	USoundAttenuation* AttenuationSettings;
+	TObjectPtr<USoundAttenuation> AttenuationSettings;
 
 	UPROPERTY(EditAnywhere, Category=Attenuation, meta=(EditCondition="bOverrideAttenuation"))
 	FSoundAttenuationSettings AttenuationOverrides;
diff --git a/Engine/Source/Runtime/Engine/Classes/Sound/SoundNodeEnveloper.h b/Engine/Source/Runtime/Engine/Classes/Sound/SoundNodeEnveloper.h
--- a/Engine/Source/Runtime/Engine/Classes/Sound/SoundNodeEnveloper.h
+++ b/Engine/Source/Runtime/Engine/Classes/Sound/SoundNodeEnveloper.h
@@ -46,10 +46,10 @@ class USoundNodeEnveloper : public USoundNode
 	uint32 bLoop:1;
 
 	UPROPERTY()
-	class UDistributionFloatConstantCurve* VolumeInterpCurve_DEPRECATED;
+	TObjectPtr<class UDistributionFloatConstantCurve> VolumeInterpCurve_DEPRECATED;
 
 	UPROPERTY()
-	class UDistributionFloatConstantCurve* PitchInterpCurve_DEPRECATED;
+	TObjectPtr<class UDistributionFloatConstantCurve> PitchInterpCurve_DEPRECATED;
 
 	// The distribution defining the volume envelope.
 	UPROPERTY(EditAnywhere, Category=Envelope)
diff --git a/Engine/Source/Runtime/Engine/Classes/Sound/SoundNodeSoundClass.h b/Engine/Source/Runtime/Engine/Classes/Sound/SoundNodeSoundClass.h
--- a/Engine/Source/Runtime/Engine/Classes/Sound/SoundNodeSoundClass.h
+++ b/Engine/Source/Runtime/Engine/Classes/Sound/SoundNodeSoundClass.h
@@ -22,7 +22,7 @@ class USoundNodeSoundClass : public USoundNode
 	GENERATED_UCLASS_BODY()
 
 	UPROPERTY(EditAnywhere, Category=SoundClass)
-	USoundClass* SoundClassOverride;
+	TObjectPtr<USoundClass> SoundClassOverride;
 
 public:
 	//~ Begin USoundNode Interface. 
diff --git a/Engine/Source/Runtime/Engine/Classes/Sound/SoundNodeWavePlayer.h b/Engine/Source/Runtime/Engine/Classes/Sound/SoundNodeWavePlayer.h
--- a/Engine/Source/Runtime/Engine/Classes/Sound/SoundNodeWavePlayer.h
+++ b/Engine/Source/Runtime/Engine/Classes/Sound/SoundNodeWavePlayer.h
@@ -29,7 +29,7 @@ private:
 	TSoftObjectPtr<USoundWave> SoundWaveAssetPtr;
 
 	UPROPERTY(transient)
-	USoundWave* SoundWave;
+	TObjectPtr<USoundWave> SoundWave;
 
 	void OnSoundWaveLoaded(const FName& PackageName, UPackage * Package, EAsyncLoadingResult::Type Result, bool bAddToRoot);
 
diff --git a/Engine/Source/Runtime/Engine/Classes/Sound/SoundSourceBus.h b/Engine/Source/Runtime/Engine/Classes/Sound/SoundSourceBus.h
--- a/Engine/Source/Runtime/Engine/Classes/Sound/SoundSourceBus.h
+++ b/Engine/Source/Runtime/Engine/Classes/Sound/SoundSourceBus.h
@@ -39,7 +39,7 @@ public:
 
 	/** Audio bus to use to use as audio for this source bus. This source bus will sonify the audio from the audio bus. */
 	UPROPERTY(EditAnywhere, Category = BusProperties)
-	UAudioBus* AudioBus;
+	TObjectPtr<UAudioBus> AudioBus;
 
 	/** Stop the source bus when the volume goes to zero. */
 	UPROPERTY(meta = (DeprecatedProperty))
diff --git a/Engine/Source/Runtime/Engine/Classes/Sound/SoundSourceBusSend.h b/Engine/Source/Runtime/Engine/Classes/Sound/SoundSourceBusSend.h
--- a/Engine/Source/Runtime/Engine/Classes/Sound/SoundSourceBusSend.h
+++ b/Engine/Source/Runtime/Engine/Classes/Sound/SoundSourceBusSend.h
@@ -37,11 +37,11 @@ struct ENGINE_API FSoundSourceBusSendInfo
 
 	// A source Bus to send the audio to. Source buses sonify (make audible) the audio sent to it and are themselves sounds which take up a voice slot in the audio engine.
 	UPROPERTY(EditAnywhere, Category = BusSend)
-	USoundSourceBus* SoundSourceBus;
+	TObjectPtr<USoundSourceBus> SoundSourceBus;
 
 	// An audio bus to send the audio to. Audio buses can be used to route audio to DSP effects or other purposes. E.g. side-chaining, analysis, etc. Audio buses are not audible unless hooked up to a source bus.
 	UPROPERTY(EditAnywhere, Category = BusSend)
-	UAudioBus* AudioBus;
+	TObjectPtr<UAudioBus> AudioBus;
 
 	// The amount of audio to send to the bus.
 	UPROPERTY(EditAnywhere, Category = BusSend)
diff --git a/Engine/Source/Runtime/Engine/Classes/Sound/SoundSubmix.h b/Engine/Source/Runtime/Engine/Classes/Sound/SoundSubmix.h
--- a/Engine/Source/Runtime/Engine/Classes/Sound/SoundSubmix.h
+++ b/Engine/Source/Runtime/Engine/Classes/Sound/SoundSubmix.h
@@ -170,7 +170,7 @@ public:
 
 	// Child submixes to this sound mix
 	UPROPERTY(VisibleAnywhere, BlueprintReadOnly, Category = SoundSubmix)
-	TArray<USoundSubmixBase*> ChildSubmixes;
+	TArray<TObjectPtr<USoundSubmixBase>> ChildSubmixes;
 
 protected:
 	//~ Begin UObject Interface.
@@ -219,7 +219,7 @@ class ENGINE_API USoundSubmixWithParentBase : public USoundSubmixBase
 public:
 
 	UPROPERTY(VisibleAnywhere, BlueprintReadOnly, Category = SoundSubmix)
-	USoundSubmixBase* ParentSubmix;
+	TObjectPtr<USoundSubmixBase> ParentSubmix;
 
 	/**
 	* Set the parent submix of this SoundSubmix, removing it as a child from its previous owner
@@ -259,11 +259,11 @@ public:
 	uint8 bMuteWhenBackgrounded : 1;
 
 	UPROPERTY(EditAnywhere, BlueprintReadOnly, Category = SoundSubmix)
-	TArray<USoundEffectSubmixPreset*> SubmixEffectChain;
+	TArray<TObjectPtr<USoundEffectSubmixPreset>> SubmixEffectChain;
 
 	/** Optional settings used by plugins which support ambisonics file playback. */
 	UPROPERTY(EditAnywhere, BlueprintReadWrite, Category = SoundSubmix)
-	USoundfieldEncodingSettingsBase* AmbisonicsPluginSettings;
+	TObjectPtr<USoundfieldEncodingSettingsBase> AmbisonicsPluginSettings;
 
 	/** The attack time in milliseconds for the envelope follower. Delegate callbacks can be registered to get the envelope value of sounds played with this submix. */
 	UPROPERTY(EditAnywhere, BlueprintReadWrite, Category = EnvelopeFollower, meta = (ClampMin = "0", UIMin = "0"))
@@ -413,12 +413,12 @@ public:
 	//TODO: Make this editable only if SoundfieldEncodingFormat is non-default,
 	// and filter classes based on ISoundfieldFactory::GetCustomSettingsClass().
 	UPROPERTY(EditAnywhere, Category = Soundfield)
-	USoundfieldEncodingSettingsBase* EncodingSettings;
+	TObjectPtr<USoundfieldEncodingSettingsBase> EncodingSettings;
 
 	// TODO: make this editable only if SoundfieldEncodingFormat is non-default
 	// and filter classes based on USoundfieldProcessorBase::SupportsFormat.
 	UPROPERTY(EditAnywhere, Category = Soundfield)
-	TArray<USoundfieldEffectBase*> SoundfieldEffectChain;
+	TArray<TObjectPtr<USoundfieldEffectBase>> SoundfieldEffectChain;
 
 	// Traverses parent submixes until we find a submix that doesn't inherit it's soundfield format.
 	FName GetSubmixFormat() const;
@@ -462,7 +462,7 @@ public:
 	//TODO: Make this editable only if EndpointType is non-default,
 	// and filter classes based on ISoundfieldFactory::GetCustomSettingsClass().
 	UPROPERTY(EditAnywhere, Category = Endpoint)
-	UAudioEndpointSettingsBase* EndpointSettings;
+	TObjectPtr<UAudioEndpointSettingsBase> EndpointSettings;
 
 };
 
@@ -498,16 +498,16 @@ public:
 	//TODO: Make this editable only if EndpointType is non-default,
 	// and filter classes based on ISoundfieldFactory::GetCustomSettingsClass().
 	UPROPERTY(EditAnywhere, Category = Endpoint)
-	USoundfieldEndpointSettingsBase* EndpointSettings;
+	TObjectPtr<USoundfieldEndpointSettingsBase> EndpointSettings;
 
 	UPROPERTY()
 	TSubclassOf<USoundfieldEncodingSettingsBase> EncodingSettingsClass;
 
 	UPROPERTY(EditAnywhere, Category = Soundfield)
-	USoundfieldEncodingSettingsBase* EncodingSettings;
+	TObjectPtr<USoundfieldEncodingSettingsBase> EncodingSettings;
 
 	UPROPERTY(EditAnywhere, Category = Soundfield)
-	TArray<USoundfieldEffectBase*> SoundfieldEffectChain;
+	TArray<TObjectPtr<USoundfieldEffectBase>> SoundfieldEffectChain;
 
 protected:
 #if WITH_EDITOR
diff --git a/Engine/Source/Runtime/Engine/Classes/Sound/SoundSubmixSend.h b/Engine/Source/Runtime/Engine/Classes/Sound/SoundSubmixSend.h
--- a/Engine/Source/Runtime/Engine/Classes/Sound/SoundSubmixSend.h
+++ b/Engine/Source/Runtime/Engine/Classes/Sound/SoundSubmixSend.h
@@ -110,7 +110,7 @@ struct ENGINE_API FSoundSubmixSendInfo
 
 	// The submix to send the audio to
 	UPROPERTY(EditAnywhere, BlueprintReadWrite, Category = SubmixSend)
-	USoundSubmixBase* SoundSubmix;
+	TObjectPtr<USoundSubmixBase> SoundSubmix;
 
 	// The amount of audio to send
 	UPROPERTY(EditAnywhere, BlueprintReadWrite, Category = SubmixSend)
diff --git a/Engine/Source/Runtime/Engine/Classes/Sound/SoundWave.h b/Engine/Source/Runtime/Engine/Classes/Sound/SoundWave.h
--- a/Engine/Source/Runtime/Engine/Classes/Sound/SoundWave.h
+++ b/Engine/Source/Runtime/Engine/Classes/Sound/SoundWave.h
@@ -183,7 +183,7 @@ struct FSoundWaveSpectralDataPerSound
 
 	// The sound wave this spectral data is associated with
 	UPROPERTY(EditAnywhere, BlueprintReadWrite, Category = "SpectralData")
-	USoundWave* SoundWave = nullptr;
+	TObjectPtr<USoundWave> SoundWave = nullptr;
 };
 
 USTRUCT(BlueprintType)
@@ -201,7 +201,7 @@ struct FSoundWaveEnvelopeDataPerSound
 
 	// The sound wave this envelope data is associated with
 	UPROPERTY(EditAnywhere, BlueprintReadWrite, Category = "EnvelopeData")
-	USoundWave* SoundWave = nullptr;
+	TObjectPtr<USoundWave> SoundWave = nullptr;
 };
 
 // Sort predicate for sorting spectral data by frequency (lowest first)
@@ -480,7 +480,7 @@ public:
 #if WITH_EDITORONLY_DATA
 	/** Specify a sound to use for the baked analysis. Will default to this USoundWave if not set. */
 	UPROPERTY(EditAnywhere, Category = "Analysis")
-		USoundWave* OverrideSoundToUseForAnalysis;
+		TObjectPtr<USoundWave> OverrideSoundToUseForAnalysis;
 
 	/**
 		Whether or not we should treat the sound wave used for analysis (this or the override) as looping while performing analysis.
@@ -666,7 +666,7 @@ public:
 		FString SourceFileTimestamp_DEPRECATED;
 
 	UPROPERTY(VisibleAnywhere, Instanced, Category = ImportSettings)
-		class UAssetImportData* AssetImportData;
+		TObjectPtr<class UAssetImportData> AssetImportData;
 
 #endif // WITH_EDITORONLY_DATA
 
@@ -674,11 +674,11 @@ protected:
 
 	/** Curves associated with this sound wave */
 	UPROPERTY(EditAnywhere, Category = SoundWave, AdvancedDisplay)
-		class UCurveTable* Curves;
+		TObjectPtr<class UCurveTable> Curves;
 
 	/** Hold a reference to our internal curve so we can switch back to it if we want to */
 	UPROPERTY()
-		class UCurveTable* InternalCurves;
+		TObjectPtr<class UCurveTable> InternalCurves;
 
 	/** Potential strong handle to the first chunk of audio data. Can be released via ReleaseCompressedAudioData. */
 	FAudioChunkHandle FirstChunk;
diff --git a/Engine/Source/Runtime/Engine/Classes/VT/RuntimeVirtualTexture.h b/Engine/Source/Runtime/Engine/Classes/VT/RuntimeVirtualTexture.h
--- a/Engine/Source/Runtime/Engine/Classes/VT/RuntimeVirtualTexture.h
+++ b/Engine/Source/Runtime/Engine/Classes/VT/RuntimeVirtualTexture.h
@@ -74,7 +74,7 @@ protected:
 
 	/** Deprecated texture object containing streamed low mips. */
 	UPROPERTY()
-	class URuntimeVirtualTextureStreamingProxy* StreamingTexture_DEPRECATED = nullptr;
+	TObjectPtr<class URuntimeVirtualTextureStreamingProxy> StreamingTexture_DEPRECATED = nullptr;
 
 public:
 	/** Get the material set that this virtual texture stores. */
diff --git a/Engine/Source/Runtime/Engine/Classes/VT/RuntimeVirtualTextureVolume.h b/Engine/Source/Runtime/Engine/Classes/VT/RuntimeVirtualTextureVolume.h
--- a/Engine/Source/Runtime/Engine/Classes/VT/RuntimeVirtualTextureVolume.h
+++ b/Engine/Source/Runtime/Engine/Classes/VT/RuntimeVirtualTextureVolume.h
@@ -15,12 +15,12 @@ class ARuntimeVirtualTextureVolume : public AActor
 public:
 	/** Component that owns the runtime virtual texture. */
 	UPROPERTY(VisibleAnywhere, BlueprintReadOnly, Category = VirtualTexture, meta = (AllowPrivateAccess = "true"))
-	class URuntimeVirtualTextureComponent* VirtualTextureComponent;
+	TObjectPtr<class URuntimeVirtualTextureComponent> VirtualTextureComponent;
 
 #if WITH_EDITORONLY_DATA
 	/** Box for visualizing virtual texture extents. */
 	UPROPERTY(Transient)
-	class UBoxComponent* Box = nullptr;
+	TObjectPtr<class UBoxComponent> Box = nullptr;
 #endif // WITH_EDITORONLY_DATA
 
 #if WITH_EDITOR
diff --git a/Engine/Source/Runtime/Engine/Classes/VT/VirtualTextureBuilder.h b/Engine/Source/Runtime/Engine/Classes/VT/VirtualTextureBuilder.h
--- a/Engine/Source/Runtime/Engine/Classes/VT/VirtualTextureBuilder.h
+++ b/Engine/Source/Runtime/Engine/Classes/VT/VirtualTextureBuilder.h
@@ -49,7 +49,7 @@ public:
 
 	/** The UTexture object. */
 	UPROPERTY(VisibleAnywhere, BlueprintReadOnly, Category = Texture)
-	class UVirtualTexture2D* Texture;
+	TObjectPtr<class UVirtualTexture2D> Texture;
 
 	/** Some client defined hash of that defines how the Texture was built. */
 	UPROPERTY()
diff --git a/Engine/Source/Runtime/Engine/Classes/VectorField/VectorFieldAnimated.h b/Engine/Source/Runtime/Engine/Classes/VectorField/VectorFieldAnimated.h
--- a/Engine/Source/Runtime/Engine/Classes/VectorField/VectorFieldAnimated.h
+++ b/Engine/Source/Runtime/Engine/Classes/VectorField/VectorFieldAnimated.h
@@ -31,7 +31,7 @@ class UVectorFieldAnimated : public UVectorField
 
 	/** The texture from which to create the vector field. */
 	UPROPERTY(EditAnywhere, Category=Reconstruction)
-	class UTexture2D* Texture;
+	TObjectPtr<class UTexture2D> Texture;
 
 	/** The operation used to construct the vector field. */
 	UPROPERTY(EditAnywhere, Category=Reconstruction)
@@ -71,7 +71,7 @@ class UVectorFieldAnimated : public UVectorField
 
 	/** A static vector field used to add noise. */
 	UPROPERTY(EditAnywhere, Category=Noise)
-	class UVectorFieldStatic* NoiseField;
+	TObjectPtr<class UVectorFieldStatic> NoiseField;
 
 	/** Scale to apply to vectors in the noise field. */
 	UPROPERTY(EditAnywhere, Category=Noise)
diff --git a/Engine/Source/Runtime/Engine/Classes/VectorField/VectorFieldStatic.h b/Engine/Source/Runtime/Engine/Classes/VectorField/VectorFieldStatic.h
--- a/Engine/Source/Runtime/Engine/Classes/VectorField/VectorFieldStatic.h
+++ b/Engine/Source/Runtime/Engine/Classes/VectorField/VectorFieldStatic.h
@@ -54,7 +54,7 @@ public:
 	FString SourceFilePath_DEPRECATED;
 
 	UPROPERTY(VisibleAnywhere, Instanced, Category=ImportSettings)
-	class UAssetImportData* AssetImportData;
+	TObjectPtr<class UAssetImportData> AssetImportData;
 #endif // WITH_EDITORONLY_DATA
 
 	//~ Begin UObject Interface.
diff --git a/Engine/Source/Runtime/Engine/Classes/VectorField/VectorFieldVolume.h b/Engine/Source/Runtime/Engine/Classes/VectorField/VectorFieldVolume.h
--- a/Engine/Source/Runtime/Engine/Classes/VectorField/VectorFieldVolume.h
+++ b/Engine/Source/Runtime/Engine/Classes/VectorField/VectorFieldVolume.h
@@ -20,11 +20,11 @@ class AVectorFieldVolume : public AActor
 
 private:
 	UPROPERTY(VisibleAnywhere, BlueprintReadOnly, Category = VectorFieldVolume, meta = (AllowPrivateAccess = "true"))
-	class UVectorFieldComponent* VectorFieldComponent;
+	TObjectPtr<class UVectorFieldComponent> VectorFieldComponent;
 
 #if WITH_EDITORONLY_DATA
 	UPROPERTY()
-	UBillboardComponent* SpriteComponent;
+	TObjectPtr<UBillboardComponent> SpriteComponent;
 #endif
 
 public:
diff --git a/Engine/Source/Runtime/Engine/Public/AlphaBlend.h b/Engine/Source/Runtime/Engine/Public/AlphaBlend.h
--- a/Engine/Source/Runtime/Engine/Public/AlphaBlend.h
+++ b/Engine/Source/Runtime/Engine/Public/AlphaBlend.h
@@ -56,7 +56,7 @@ struct ENGINE_API FAlphaBlendArgs
 
 	/** If you're using Custom BlendOption, you can specify curve */
 	UPROPERTY(EditAnywhere, BlueprintReadWrite, Category = "Blend", meta=(DisplayAfter="BlendOption"))
-	UCurveFloat* CustomCurve;
+	TObjectPtr<UCurveFloat> CustomCurve;
 
 	/** Blend Time */
 	UPROPERTY(EditAnywhere, BlueprintReadWrite, Category = "Blend")
@@ -82,7 +82,7 @@ private:
 
 	/** If you're using Custom BlendOption, you can specify curve */
 	UPROPERTY(EditAnywhere, Category = "Blend", meta=(DisplayAfter="BlendOption"))
-	UCurveFloat* CustomCurve;
+	TObjectPtr<UCurveFloat> CustomCurve;
 
 public:
 	/* Constructor */
diff --git a/Engine/Source/Runtime/Engine/Public/Animation/AnimNotifyQueue.h b/Engine/Source/Runtime/Engine/Public/Animation/AnimNotifyQueue.h
--- a/Engine/Source/Runtime/Engine/Public/Animation/AnimNotifyQueue.h
+++ b/Engine/Source/Runtime/Engine/Public/Animation/AnimNotifyQueue.h
@@ -50,7 +50,7 @@ private:
 	const FAnimNotifyEvent* Notify;
 
 	UPROPERTY(transient)
-	const UObject* NotifySource;
+	TObjectPtr<const UObject> NotifySource;
 };
 
 USTRUCT()
diff --git a/Engine/Source/Runtime/Engine/Public/Animation/AnimTypes.h b/Engine/Source/Runtime/Engine/Public/Animation/AnimTypes.h
--- a/Engine/Source/Runtime/Engine/Public/Animation/AnimTypes.h
+++ b/Engine/Source/Runtime/Engine/Public/Animation/AnimTypes.h
@@ -305,10 +305,10 @@ struct FAnimNotifyEvent : public FAnimLinkableElement
 	FName NotifyName;
 
 	UPROPERTY(EditAnywhere, Instanced, BlueprintReadWrite, Category=AnimNotifyEvent)
-	class UAnimNotify * Notify;
+	TObjectPtr<class UAnimNotify>  Notify;
 
 	UPROPERTY(EditAnywhere, Instanced, BlueprintReadWrite, Category=AnimNotifyEvent)
-	class UAnimNotifyState * NotifyStateClass;
+	TObjectPtr<class UAnimNotifyState>  NotifyStateClass;
 
 	UPROPERTY()
 	float Duration;
diff --git a/Engine/Source/Runtime/Engine/Public/AssetExportTask.h b/Engine/Source/Runtime/Engine/Public/AssetExportTask.h
--- a/Engine/Source/Runtime/Engine/Public/AssetExportTask.h
+++ b/Engine/Source/Runtime/Engine/Public/AssetExportTask.h
@@ -18,11 +18,11 @@ class ENGINE_API UAssetExportTask : public UObject
 public:
 	/** Asset to export */
 	UPROPERTY(BlueprintReadWrite, Category = Miscellaneous)
-	UObject* Object;
+	TObjectPtr<UObject> Object;
 
 	/** Optional exporter, otherwise it will be determined automatically */
 	UPROPERTY(BlueprintReadWrite, Category = Miscellaneous)
-	class UExporter* Exporter;
+	TObjectPtr<class UExporter> Exporter;
 
 	/** File to export as */
 	UPROPERTY(BlueprintReadWrite, Category = Miscellaneous)
@@ -54,11 +54,11 @@ public:
 
 	/** Array of objects to ignore exporting */
 	UPROPERTY(BlueprintReadWrite, Category = Miscellaneous)
-	TArray<UObject*> IgnoreObjectList;
+	TArray<TObjectPtr<UObject>> IgnoreObjectList;
 
 	/** Exporter specific options */
 	UPROPERTY(BlueprintReadWrite, Category = Miscellaneous)
-	UObject* Options;
+	TObjectPtr<UObject> Options;
 
 	/** Array of error messages encountered during exporter */
 	UPROPERTY(BlueprintReadWrite, Category = Miscellaneous)
diff --git a/Engine/Source/Runtime/Engine/Public/CanvasTypes.h b/Engine/Source/Runtime/Engine/Public/CanvasTypes.h
--- a/Engine/Source/Runtime/Engine/Public/CanvasTypes.h
+++ b/Engine/Source/Runtime/Engine/Public/CanvasTypes.h
@@ -52,7 +52,7 @@ struct FTextSizingParameters
 
 	/** the font to use for sizing/wrapping the string */
 	UPROPERTY()
-	const UFont* DrawFont;
+	TObjectPtr<const UFont> DrawFont;
 
 	/** Horizontal spacing adjustment between characters and vertical spacing adjustment between wrapped lines */
 	UPROPERTY()
diff --git a/Engine/Source/Runtime/Engine/Public/ComponentInstanceDataCache.h b/Engine/Source/Runtime/Engine/Public/ComponentInstanceDataCache.h
--- a/Engine/Source/Runtime/Engine/Public/ComponentInstanceDataCache.h
+++ b/Engine/Source/Runtime/Engine/Public/ComponentInstanceDataCache.h
@@ -103,7 +103,7 @@ protected:
 
 	/** The template used to create the source component */
 	UPROPERTY()
-	const UObject* SourceComponentTemplate;
+	TObjectPtr<const UObject> SourceComponentTemplate;
 
 	/** The method that was used to create the source component */
 	UPROPERTY() 
@@ -130,7 +130,7 @@ protected:
 
 	// Referenced objects in component instance saved properties
 	UPROPERTY()
-	TArray<UObject*> ReferencedObjects;
+	TArray<TObjectPtr<UObject>> ReferencedObjects;
 
 	// Referenced names in component instance saved properties
 	UPROPERTY()
diff --git a/Engine/Source/Runtime/Engine/Public/ComputeFramework/ComputeGraph.h b/Engine/Source/Runtime/Engine/Public/ComputeFramework/ComputeGraph.h
--- a/Engine/Source/Runtime/Engine/Public/ComputeFramework/ComputeGraph.h
+++ b/Engine/Source/Runtime/Engine/Public/ComputeFramework/ComputeGraph.h
@@ -73,7 +73,7 @@ struct FComputeKernelInvocation
 	static constexpr uint32 MAX_NAME_LENGTH = 32;
 
 	UPROPERTY()
-	UComputeKernel* ComputeKernel = nullptr;
+	TObjectPtr<UComputeKernel> ComputeKernel = nullptr;
 	
 	UPROPERTY()
 	uint16 GenerationIdx = 0xFFFF;
diff --git a/Engine/Source/Runtime/Engine/Public/ComputeFramework/ComputeGraphComponent.h b/Engine/Source/Runtime/Engine/Public/ComputeFramework/ComputeGraphComponent.h
--- a/Engine/Source/Runtime/Engine/Public/ComputeFramework/ComputeGraphComponent.h
+++ b/Engine/Source/Runtime/Engine/Public/ComputeFramework/ComputeGraphComponent.h
@@ -21,7 +21,7 @@ public:
 	UComputeGraphComponent();
 
 	UPROPERTY(EditAnywhere, Category = "Compute")
-	UComputeGraph* ComputeGraph = nullptr;
+	TObjectPtr<UComputeGraph> ComputeGraph = nullptr;
 
 	UFUNCTION(BlueprintCallable, Category = "Compute")
 	void QueueExecute();
diff --git a/Engine/Source/Runtime/Engine/Public/ComputeFramework/ComputeKernel.h b/Engine/Source/Runtime/Engine/Public/ComputeFramework/ComputeKernel.h
--- a/Engine/Source/Runtime/Engine/Public/ComputeFramework/ComputeKernel.h
+++ b/Engine/Source/Runtime/Engine/Public/ComputeFramework/ComputeKernel.h
@@ -78,7 +78,7 @@ public:
 	 * This abstracts the source and compilation process.
 	 */
 	UPROPERTY(EditAnywhere, AssetRegistrySearchable, meta = (ShowOnlyInnerProperties), Category = "Kernel")
-	UComputeKernelSource* KernelSource = nullptr;
+	TObjectPtr<UComputeKernelSource> KernelSource = nullptr;
 
 	/* Specifying certain memory access flags allows for optimizations such as kernel fusing. */
 	UPROPERTY(VisibleAnywhere, BlueprintReadOnly, meta = (Bitmask, BitmaskEnum = "EComputeKernelFlags"), Category = "Kernel")
diff --git a/Engine/Source/Runtime/Engine/Public/HLOD/HLODProxyDesc.h b/Engine/Source/Runtime/Engine/Public/HLOD/HLODProxyDesc.h
--- a/Engine/Source/Runtime/Engine/Public/HLOD/HLODProxyDesc.h
+++ b/Engine/Source/Runtime/Engine/Public/HLOD/HLODProxyDesc.h
@@ -25,10 +25,10 @@ public:
 
 public:
 	UPROPERTY()
-	UStaticMesh* StaticMesh = nullptr;
+	TObjectPtr<UStaticMesh> StaticMesh = nullptr;
 	
 	UPROPERTY()
-	const UMaterialInterface* Material = nullptr;
+	TObjectPtr<const UMaterialInterface> Material = nullptr;
 
 	UPROPERTY()
 	TArray<FTransform> Instances;
@@ -63,7 +63,7 @@ private:
 	TArray<FName> SubActors;
 
 	UPROPERTY()
-	UStaticMesh* StaticMesh;
+	TObjectPtr<UStaticMesh> StaticMesh;
 
 	UPROPERTY()
 	TArray<FHLODISMComponentDesc> ISMComponentsDesc;
diff --git a/Engine/Source/Runtime/Engine/Public/HLOD/HLODProxyMesh.h b/Engine/Source/Runtime/Engine/Public/HLOD/HLODProxyMesh.h
--- a/Engine/Source/Runtime/Engine/Public/HLOD/HLODProxyMesh.h
+++ b/Engine/Source/Runtime/Engine/Public/HLOD/HLODProxyMesh.h
@@ -41,7 +41,7 @@ private:
 
 	/** The mesh used to display this proxy */
 	UPROPERTY(VisibleAnywhere, Category = "Proxy Mesh")
-	UStaticMesh* StaticMesh;
+	TObjectPtr<UStaticMesh> StaticMesh;
 
 	/** The key generated from an ALODActor. If this differs from that generated from the ALODActor, then the mesh needs regenerating. */
 	UPROPERTY(VisibleAnywhere, Category = "Proxy Mesh")
diff --git a/Engine/Source/Runtime/Engine/Public/LevelInstance/LevelInstanceTypes.h b/Engine/Source/Runtime/Engine/Public/LevelInstance/LevelInstanceTypes.h
--- a/Engine/Source/Runtime/Engine/Public/LevelInstance/LevelInstanceTypes.h
+++ b/Engine/Source/Runtime/Engine/Public/LevelInstance/LevelInstanceTypes.h
@@ -42,8 +42,8 @@ struct FNewLevelInstanceParams
 	ELevelInstancePivotType PivotType = ELevelInstancePivotType::CenterMinZ;
 
 	UPROPERTY(EditAnywhere, Category = Pivot)
-	AActor* PivotActor = nullptr;
+	TObjectPtr<AActor> PivotActor = nullptr;
 
 	UPROPERTY()
-	UWorld* TemplateWorld = nullptr;
+	TObjectPtr<UWorld> TemplateWorld = nullptr;
 };
\ No newline at end of file
diff --git a/Engine/Source/Runtime/Engine/Public/MaterialCachedData.h b/Engine/Source/Runtime/Engine/Public/MaterialCachedData.h
--- a/Engine/Source/Runtime/Engine/Public/MaterialCachedData.h
+++ b/Engine/Source/Runtime/Engine/Public/MaterialCachedData.h
@@ -30,7 +30,7 @@ struct FMaterialFunctionInfo
 
 	/** The function which this material has a dependency on. */
 	UPROPERTY()
-	UMaterialFunctionInterface* Function = nullptr;
+	TObjectPtr<UMaterialFunctionInterface> Function = nullptr;
 };
 
 /** Stores information about a parameter collection that this material references, used to know when the material needs to be recompiled. */
@@ -45,7 +45,7 @@ struct FMaterialParameterCollectionInfo
 
 	/** The collection which this material has a dependency on. */
 	UPROPERTY()
-	class UMaterialParameterCollection* ParameterCollection = nullptr;
+	TObjectPtr<class UMaterialParameterCollection> ParameterCollection = nullptr;
 
 	bool operator==(const FMaterialParameterCollectionInfo& Other) const
 	{
@@ -166,16 +166,16 @@ struct FMaterialCachedParameters
 	TArray<FLinearColor> VectorValues;
 
 	UPROPERTY()
-	TArray<UTexture*> TextureValues;
+	TArray<TObjectPtr<UTexture>> TextureValues;
 
 	UPROPERTY()
-	TArray<UFont*> FontValues;
+	TArray<TObjectPtr<UFont>> FontValues;
 
 	UPROPERTY()
 	TArray<int32> FontPageValues;
 
 	UPROPERTY()
-	TArray<URuntimeVirtualTexture*> RuntimeVirtualTextureValues;
+	TArray<TObjectPtr<URuntimeVirtualTexture>> RuntimeVirtualTextureValues;
 
 #if WITH_EDITORONLY_DATA
 	UPROPERTY()
@@ -191,10 +191,10 @@ struct FMaterialCachedParameters
 	TArray<FVector2D> ScalarMinMaxValues;
 
 	UPROPERTY()
-	TArray<UCurveLinearColor*> ScalarCurveValues;
+	TArray<TObjectPtr<UCurveLinearColor>> ScalarCurveValues;
 
 	UPROPERTY()
-	TArray<UCurveLinearColorAtlas*> ScalarCurveAtlasValues;
+	TArray<TObjectPtr<UCurveLinearColorAtlas>> ScalarCurveAtlasValues;
 
 	UPROPERTY()
 	TArray<FParameterChannelNames> VectorChannelNameValues;
@@ -238,7 +238,7 @@ struct FMaterialCachedExpressionData
 
 	/** Array of all texture referenced by this material */
 	UPROPERTY()
-	TArray<UObject*> ReferencedTextures;
+	TArray<TObjectPtr<UObject>> ReferencedTextures;
 
 	/** Array of all functions this material depends on. */
 	UPROPERTY()
@@ -249,13 +249,13 @@ struct FMaterialCachedExpressionData
 	TArray<FMaterialParameterCollectionInfo> ParameterCollectionInfos;
 
 	UPROPERTY()
-	TArray<UMaterialFunctionInterface*> DefaultLayers;
+	TArray<TObjectPtr<UMaterialFunctionInterface>> DefaultLayers;
 
 	UPROPERTY()
-	TArray<UMaterialFunctionInterface*> DefaultLayerBlends;
+	TArray<TObjectPtr<UMaterialFunctionInterface>> DefaultLayerBlends;
 
 	UPROPERTY()
-	TArray<ULandscapeGrassType*> GrassTypes;
+	TArray<TObjectPtr<ULandscapeGrassType>> GrassTypes;
 
 	UPROPERTY()
 	TArray<FName> DynamicParameterNames;
@@ -282,5 +282,5 @@ struct FMaterialInstanceCachedData
 	FMaterialCachedParameters Parameters;
 
 	UPROPERTY()
-	TArray<UObject*> ReferencedTextures;
+	TArray<TObjectPtr<UObject>> ReferencedTextures;
 };
diff --git a/Engine/Source/Runtime/Engine/Public/Net/VoiceConfig.h b/Engine/Source/Runtime/Engine/Public/Net/VoiceConfig.h
--- a/Engine/Source/Runtime/Engine/Public/Net/VoiceConfig.h
+++ b/Engine/Source/Runtime/Engine/Public/Net/VoiceConfig.h
@@ -51,16 +51,16 @@ struct FVoiceSettings
 	// that the voice is coming from that scene component.
 	// If this is not set, the voice will not be spatialized.
 	UPROPERTY(EditAnywhere, BlueprintReadWrite, Category = "Audio|Voice|Spatialization")
-	USceneComponent* ComponentToAttachTo;
+	TObjectPtr<USceneComponent> ComponentToAttachTo;
 
 	// Optional attenuation settings to attach to this player's voice.
 	// This should only be used when ComponentToAttachTo is set.
 	UPROPERTY(EditAnywhere, BlueprintReadWrite, Category = "Audio|Voice|Spatialization")
-	USoundAttenuation* AttenuationSettings;
+	TObjectPtr<USoundAttenuation> AttenuationSettings;
 	
 	// Optional audio effects to apply to this player's voice.
 	UPROPERTY(EditAnywhere, BlueprintReadWrite, Category = "Audio|Voice|Effects")
-	USoundEffectSourcePresetChain* SourceEffectChain;
+	TObjectPtr<USoundEffectSourcePresetChain> SourceEffectChain;
 
 	FVoiceSettings()
 		: ComponentToAttachTo(nullptr)
diff --git a/Engine/Source/Runtime/Engine/Public/SingleAnimationPlayData.h b/Engine/Source/Runtime/Engine/Public/SingleAnimationPlayData.h
--- a/Engine/Source/Runtime/Engine/Public/SingleAnimationPlayData.h
+++ b/Engine/Source/Runtime/Engine/Public/SingleAnimationPlayData.h
@@ -17,7 +17,7 @@ struct FSingleAnimationPlayData
 	// and have detail customization to display different things
 	// The default sequence to play on this skeletal mesh
 	UPROPERTY(EditAnywhere, BlueprintReadWrite, Category = Animation)
-	class UAnimationAsset* AnimToPlay;
+	TObjectPtr<class UAnimationAsset> AnimToPlay;
 
 	/** Default setting for looping for SequenceToPlay. This is not current state of looping. */
 	UPROPERTY(EditAnywhere, BlueprintReadWrite, Category = Animation, meta = (DisplayName = "Looping"))
diff --git a/Engine/Source/Runtime/Engine/Public/SkeletalMeshReductionSettings.h b/Engine/Source/Runtime/Engine/Public/SkeletalMeshReductionSettings.h
--- a/Engine/Source/Runtime/Engine/Public/SkeletalMeshReductionSettings.h
+++ b/Engine/Source/Runtime/Engine/Public/SkeletalMeshReductionSettings.h
@@ -161,7 +161,7 @@ struct FSkeletalMeshOptimizationSettings
 	TArray<FBoneReference> BonesToRemove_DEPRECATED;
 
 	UPROPERTY()
-	class UAnimSequence* BakePose_DEPRECATED;
+	TObjectPtr<class UAnimSequence> BakePose_DEPRECATED;
 
 	//Transient mutable delegate. If the delegate is bound, the reduction will call it instead of deleting the replaced LODModel. It will be then the delegate owner responsible of the LODModel memory
 	mutable FOnDeleteLODModelOverride OnDeleteLODModelDelegate;
diff --git a/Engine/Source/Runtime/Engine/Public/WorldPartition/DataLayer/WorldDataLayers.h b/Engine/Source/Runtime/Engine/Public/WorldPartition/DataLayer/WorldDataLayers.h
--- a/Engine/Source/Runtime/Engine/Public/WorldPartition/DataLayer/WorldDataLayers.h
+++ b/Engine/Source/Runtime/Engine/Public/WorldPartition/DataLayer/WorldDataLayers.h
@@ -51,5 +51,5 @@ private:
 #endif
 
 	UPROPERTY()
-	TSet<UDataLayer*> WorldDataLayers;
+	TSet<TObjectPtr<UDataLayer>> WorldDataLayers;
 };
\ No newline at end of file
diff --git a/Engine/Source/Runtime/Engine/Public/WorldPartition/HLOD/HLODActor.h b/Engine/Source/Runtime/Engine/Public/WorldPartition/HLOD/HLODActor.h
--- a/Engine/Source/Runtime/Engine/Public/WorldPartition/HLOD/HLODActor.h
+++ b/Engine/Source/Runtime/Engine/Public/WorldPartition/HLOD/HLODActor.h
@@ -84,7 +84,7 @@ private:
 	TArray<FGuid> SubActors;
 
 	UPROPERTY()
-	const UHLODLayer* SubActorsHLODLayer;
+	TObjectPtr<const UHLODLayer> SubActorsHLODLayer;
 
 	UPROPERTY()
 	int64 GridIndexX;
diff --git a/Engine/Source/Runtime/Engine/Public/WorldPartition/NavigationData/NavigationDataChunkActor.h b/Engine/Source/Runtime/Engine/Public/WorldPartition/NavigationData/NavigationDataChunkActor.h
--- a/Engine/Source/Runtime/Engine/Public/WorldPartition/NavigationData/NavigationDataChunkActor.h
+++ b/Engine/Source/Runtime/Engine/Public/WorldPartition/NavigationData/NavigationDataChunkActor.h
@@ -47,7 +47,7 @@ protected:
 	void Log(const TCHAR* FunctionName) const;
 
 	UPROPERTY()
-	TArray<UNavigationDataChunk*> NavDataChunks;
+	TArray<TObjectPtr<UNavigationDataChunk>> NavDataChunks;
 
 	UPROPERTY()
 	FBox DataChunkActorBounds;
diff --git a/Engine/Source/Runtime/Engine/Public/WorldPartition/WorldPartition.h b/Engine/Source/Runtime/Engine/Public/WorldPartition/WorldPartition.h
--- a/Engine/Source/Runtime/Engine/Public/WorldPartition/WorldPartition.h
+++ b/Engine/Source/Runtime/Engine/Public/WorldPartition/WorldPartition.h
@@ -156,17 +156,17 @@ public:
 	virtual UWorld* GetWorld() const override;
 
 	UPROPERTY(Transient)
-	UWorld* World;
+	TObjectPtr<UWorld> World;
 
 #if WITH_EDITORONLY_DATA
 	UPROPERTY(DuplicateTransient)
-	UWorldPartitionEditorHash* EditorHash;
+	TObjectPtr<UWorldPartitionEditorHash> EditorHash;
 
 	IWorldPartitionEditor* WorldPartitionEditor;
 #endif
 
 	UPROPERTY(EditAnywhere, Category=RuntimeHash, NoClear, meta = (NoResetToDefault), Instanced)
-	UWorldPartitionRuntimeHash* RuntimeHash;
+	TObjectPtr<UWorldPartitionRuntimeHash> RuntimeHash;
 
 #if WITH_EDITOR
 	bool bForceGarbageCollection;
@@ -176,7 +176,7 @@ public:
 #if WITH_EDITORONLY_DATA
 	// Default HLOD layer
 	UPROPERTY(EditAnywhere, Category=HLOD)
-	class UHLODLayer* DefaultHLODLayer;
+	TObjectPtr<class UHLODLayer> DefaultHLODLayer;
 #endif
 
 private:
@@ -184,7 +184,7 @@ private:
 	FTransform InstanceTransform;
 
 	UPROPERTY(Transient, DuplicateTransient)
-	mutable UWorldPartitionStreamingPolicy* StreamingPolicy;
+	mutable TObjectPtr<UWorldPartitionStreamingPolicy> StreamingPolicy;
 
 #if WITH_EDITORONLY_DATA
 	FLinkerInstancingContext InstancingContext;
diff --git a/Engine/Source/Runtime/Engine/Public/WorldPartition/WorldPartitionEditorSpatialHash.h b/Engine/Source/Runtime/Engine/Public/WorldPartition/WorldPartitionEditorSpatialHash.h
--- a/Engine/Source/Runtime/Engine/Public/WorldPartition/WorldPartitionEditorSpatialHash.h
+++ b/Engine/Source/Runtime/Engine/Public/WorldPartition/WorldPartitionEditorSpatialHash.h
@@ -241,13 +241,13 @@ private:
 	TMap<FCellCoord, UWorldPartitionEditorCell*> HashCells;
 
 	UPROPERTY(Transient)
-	TSet<UWorldPartitionEditorCell*> Cells;
+	TSet<TObjectPtr<UWorldPartitionEditorCell>> Cells;
 	
 	FBox Bounds;
 	bool bBoundsDirty;
 	
 	UPROPERTY(Transient)
-	UWorldPartitionEditorCell* AlwaysLoadedCell;
+	TObjectPtr<UWorldPartitionEditorCell> AlwaysLoadedCell;
 
 public:
 	UPROPERTY(Config, meta = (AllowedClasses = "Texture2D, MaterialInterface"))
diff --git a/Engine/Source/Runtime/Engine/Public/WorldPartition/WorldPartitionLevelStreamingDynamic.h b/Engine/Source/Runtime/Engine/Public/WorldPartition/WorldPartitionLevelStreamingDynamic.h
--- a/Engine/Source/Runtime/Engine/Public/WorldPartition/WorldPartitionLevelStreamingDynamic.h
+++ b/Engine/Source/Runtime/Engine/Public/WorldPartition/WorldPartitionLevelStreamingDynamic.h
@@ -48,7 +48,7 @@ private:
 
 #if WITH_EDITORONLY_DATA
 	UPROPERTY()
-	ULevel* RuntimeLevel;
+	TObjectPtr<ULevel> RuntimeLevel;
 #endif
 
 	FDelegateHandle OnCleanupLevelDelegateHandle;
diff --git a/Engine/Source/Runtime/Engine/Public/WorldPartition/WorldPartitionMiniMap.h b/Engine/Source/Runtime/Engine/Public/WorldPartition/WorldPartitionMiniMap.h
--- a/Engine/Source/Runtime/Engine/Public/WorldPartition/WorldPartitionMiniMap.h
+++ b/Engine/Source/Runtime/Engine/Public/WorldPartition/WorldPartitionMiniMap.h
@@ -28,7 +28,7 @@ public:
 
 	/*MiniMap Texture for displaying on world partition window*/
 	UPROPERTY(VisibleAnywhere, Category = WorldPartitionMiniMap)
-	UTexture2D* MiniMapTexture;
+	TObjectPtr<UTexture2D> MiniMapTexture;
 
 	/*MiniMap Size*/
 	UPROPERTY(EditAnywhere, Category = WorldPartitionMiniMap, meta=(UIMin = "256", UIMax = "8192"))
diff --git a/Engine/Source/Runtime/Engine/Public/WorldPartition/WorldPartitionRuntimeCell.h b/Engine/Source/Runtime/Engine/Public/WorldPartition/WorldPartitionRuntimeCell.h
--- a/Engine/Source/Runtime/Engine/Public/WorldPartition/WorldPartitionRuntimeCell.h
+++ b/Engine/Source/Runtime/Engine/Public/WorldPartition/WorldPartitionRuntimeCell.h
@@ -79,7 +79,7 @@ protected:
 
 private:
 	UPROPERTY()
-	TMap<const TSubclassOf<UWorldPartitionRuntimeCellData>, const UWorldPartitionRuntimeCellData*> CellDataMap;
+	TMap<const TSubclassOf<UWorldPartitionRuntimeCellData>, TObjectPtr<const UWorldPartitionRuntimeCellData>> CellDataMap;
 
 	UPROPERTY()
 	TArray<FName> DataLayers;
diff --git a/Engine/Source/Runtime/Engine/Public/WorldPartition/WorldPartitionRuntimeLevelStreamingCell.h b/Engine/Source/Runtime/Engine/Public/WorldPartition/WorldPartitionRuntimeLevelStreamingCell.h
--- a/Engine/Source/Runtime/Engine/Public/WorldPartition/WorldPartitionRuntimeLevelStreamingCell.h
+++ b/Engine/Source/Runtime/Engine/Public/WorldPartition/WorldPartitionRuntimeLevelStreamingCell.h
@@ -52,5 +52,5 @@ private:
 #endif
 
 	UPROPERTY()
-	mutable class UWorldPartitionLevelStreamingDynamic* LevelStreaming;
+	mutable TObjectPtr<class UWorldPartitionLevelStreamingDynamic> LevelStreaming;
 };
\ No newline at end of file
diff --git a/Engine/Source/Runtime/Engine/Public/WorldPartition/WorldPartitionRuntimeSpatialHash.h b/Engine/Source/Runtime/Engine/Public/WorldPartition/WorldPartitionRuntimeSpatialHash.h
--- a/Engine/Source/Runtime/Engine/Public/WorldPartition/WorldPartitionRuntimeSpatialHash.h
+++ b/Engine/Source/Runtime/Engine/Public/WorldPartition/WorldPartitionRuntimeSpatialHash.h
@@ -30,7 +30,7 @@ struct FSpatialHashStreamingGridLayerCell
 	FSpatialHashStreamingGridLayerCell& operator=(FSpatialHashStreamingGridLayerCell&& Other) = default;
 
 	UPROPERTY()
-	TArray<UWorldPartitionRuntimeSpatialHashCell*> GridCells;
+	TArray<TObjectPtr<UWorldPartitionRuntimeSpatialHashCell>> GridCells;
 };
 
 // Disable copy of this USTRUCT by the autogenerated code.
diff --git a/Engine/Source/Runtime/Engine/Public/WorldPartition/WorldPartitionSubsystem.h b/Engine/Source/Runtime/Engine/Public/WorldPartition/WorldPartitionSubsystem.h
--- a/Engine/Source/Runtime/Engine/Public/WorldPartition/WorldPartitionSubsystem.h
+++ b/Engine/Source/Runtime/Engine/Public/WorldPartition/WorldPartitionSubsystem.h
@@ -56,7 +56,7 @@ private:
 	friend class UWorldPartition;
 
 	UPROPERTY()
-	TArray<UWorldPartition*> RegisteredWorldPartitions;
+	TArray<TObjectPtr<UWorldPartition>> RegisteredWorldPartitions;
 
 	FDelegateHandle	DrawRuntimeHash2DHandle;
 };
diff --git a/Engine/Source/Runtime/Experimental/ChaosSolverEngine/Public/Chaos/ChaosGameplayEventDispatcher.h b/Engine/Source/Runtime/Experimental/ChaosSolverEngine/Public/Chaos/ChaosGameplayEventDispatcher.h
--- a/Engine/Source/Runtime/Experimental/ChaosSolverEngine/Public/Chaos/ChaosGameplayEventDispatcher.h
+++ b/Engine/Source/Runtime/Experimental/ChaosSolverEngine/Public/Chaos/ChaosGameplayEventDispatcher.h
@@ -23,7 +23,7 @@ struct CHAOSSOLVERENGINE_API FChaosBreakEvent
 
 public:
 	UPROPERTY(BlueprintReadOnly, Category = "Break Event")
-	UPrimitiveComponent* Component = nullptr;
+	TObjectPtr<UPrimitiveComponent> Component = nullptr;
 
 	UPROPERTY(BlueprintReadOnly, Category = "Break Event")
 	FVector Location;
@@ -61,7 +61,7 @@ struct FChaosHandlerSet
 		
 	/** These should be IChaosNotifyHandlerInterface refs, but we can't store those here */
 	UPROPERTY()
-	TSet<UObject*> ChaosHandlers;
+	TSet<TObjectPtr<UObject>> ChaosHandlers;
 };
 
 struct FChaosPendingCollisionNotify
@@ -133,10 +133,10 @@ public:
 private:
 
  	UPROPERTY()
- 	TMap<UPrimitiveComponent*, FChaosHandlerSet> CollisionEventRegistrations;
+ 	TMap<TObjectPtr<UPrimitiveComponent>, FChaosHandlerSet> CollisionEventRegistrations;
 
 	UPROPERTY()
-	TMap<UPrimitiveComponent*, FBreakEventCallbackWrapper> BreakEventRegistrations;
+	TMap<TObjectPtr<UPrimitiveComponent>, FBreakEventCallbackWrapper> BreakEventRegistrations;
 
 	float LastCollisionDataTime = -1.f;
 	float LastBreakingDataTime = -1.f;
diff --git a/Engine/Source/Runtime/Experimental/ChaosSolverEngine/Public/Chaos/ChaosNotifyHandlerInterface.h b/Engine/Source/Runtime/Experimental/ChaosSolverEngine/Public/Chaos/ChaosNotifyHandlerInterface.h
--- a/Engine/Source/Runtime/Experimental/ChaosSolverEngine/Public/Chaos/ChaosNotifyHandlerInterface.h
+++ b/Engine/Source/Runtime/Experimental/ChaosSolverEngine/Public/Chaos/ChaosNotifyHandlerInterface.h
@@ -13,10 +13,10 @@ struct CHAOSSOLVERENGINE_API FChaosPhysicsCollisionInfo
 	GENERATED_BODY()
 public:
 	UPROPERTY(EditAnywhere, BlueprintReadWrite, Category="Chaos")
-	UPrimitiveComponent* Component = nullptr;
+	TObjectPtr<UPrimitiveComponent> Component = nullptr;
 
 	UPROPERTY(EditAnywhere, BlueprintReadWrite, Category = "Chaos")
-	UPrimitiveComponent* OtherComponent = nullptr;
+	TObjectPtr<UPrimitiveComponent> OtherComponent = nullptr;
 
 	/** Location of the impact */
 	UPROPERTY(EditAnywhere, BlueprintReadWrite, Category = "Chaos")
diff --git a/Engine/Source/Runtime/Experimental/ChaosSolverEngine/Public/Chaos/ChaosSolverActor.h b/Engine/Source/Runtime/Experimental/ChaosSolverEngine/Public/Chaos/ChaosSolverActor.h
--- a/Engine/Source/Runtime/Experimental/ChaosSolverEngine/Public/Chaos/ChaosSolverActor.h
+++ b/Engine/Source/Runtime/Experimental/ChaosSolverEngine/Public/Chaos/ChaosSolverActor.h
@@ -134,7 +134,7 @@ public:
 	* Display icon in the editor
 	*/
 	UPROPERTY()
-	UBillboardComponent* SpriteComponent;
+	TObjectPtr<UBillboardComponent> SpriteComponent;
 
 	UChaosGameplayEventDispatcher* GetGameplayEventDispatcher() const { return GameplayEventDispatcherComponent; };
 
@@ -166,7 +166,7 @@ private:
 
 	/** Component responsible for harvesting and triggering physics-related gameplay events (hits, breaks, etc) */
 	UPROPERTY()
-	UChaosGameplayEventDispatcher* GameplayEventDispatcherComponent;
+	TObjectPtr<UChaosGameplayEventDispatcher> GameplayEventDispatcherComponent;
 
 	/** If floor is enabled - this will point to the solver particle for it */
 	TUniquePtr<Chaos::TGeometryParticle<Chaos::FReal, 3>> FloorParticle;
diff --git a/Engine/Source/Runtime/Experimental/FieldSystem/Source/FieldSystemEngine/Public/Field/FieldSystemActor.h b/Engine/Source/Runtime/Experimental/FieldSystem/Source/FieldSystemEngine/Public/Field/FieldSystemActor.h
--- a/Engine/Source/Runtime/Experimental/FieldSystem/Source/FieldSystemEngine/Public/Field/FieldSystemActor.h
+++ b/Engine/Source/Runtime/Experimental/FieldSystem/Source/FieldSystemEngine/Public/Field/FieldSystemActor.h
@@ -23,7 +23,7 @@ public:
 
 	/* FieldSystemComponent */
 	UPROPERTY(VisibleAnywhere, BlueprintReadOnly, Category = Field, meta = (ExposeFunctionCategories = "Components|FieldSystem", AllowPrivateAccess = "true"))
-	UFieldSystemComponent* FieldSystemComponent;
+	TObjectPtr<UFieldSystemComponent> FieldSystemComponent;
 	UFieldSystemComponent* GetFieldSystemComponent() const { return FieldSystemComponent; }
 	virtual void OnConstruction(const FTransform& Transform) override;
 };
diff --git a/Engine/Source/Runtime/Experimental/FieldSystem/Source/FieldSystemEngine/Public/Field/FieldSystemComponent.h b/Engine/Source/Runtime/Experimental/FieldSystem/Source/FieldSystemEngine/Public/Field/FieldSystemComponent.h
--- a/Engine/Source/Runtime/Experimental/FieldSystem/Source/FieldSystemEngine/Public/Field/FieldSystemComponent.h
+++ b/Engine/Source/Runtime/Experimental/FieldSystem/Source/FieldSystemEngine/Public/Field/FieldSystemComponent.h
@@ -40,7 +40,7 @@ public:
 
 	/** Field system asset to be used to store the construction fields */
 	UPROPERTY(EditAnywhere, BlueprintReadOnly, Category = "Field", meta = (ToolTip = "Field system asset to be used to store the construction fields"))
-	UFieldSystem* FieldSystem;
+	TObjectPtr<UFieldSystem> FieldSystem;
 
 	/** If enabled the field will be pushed to the world fields and will be available to materials and niagara */
 	UPROPERTY(EditAnywhere, Category = "Field", meta = (ToolTip = "If enabled the field will be pushed to the world fields and will be available to materials and niagara"))
diff --git a/Engine/Source/Runtime/Experimental/FieldSystem/Source/FieldSystemEngine/Public/Field/FieldSystemObjects.h b/Engine/Source/Runtime/Experimental/FieldSystem/Source/FieldSystemEngine/Public/Field/FieldSystemObjects.h
--- a/Engine/Source/Runtime/Experimental/FieldSystem/Source/FieldSystemEngine/Public/Field/FieldSystemObjects.h
+++ b/Engine/Source/Runtime/Experimental/FieldSystem/Source/FieldSystemEngine/Public/Field/FieldSystemObjects.h
@@ -762,11 +762,11 @@ public:
 
 	/** Right field to be processed */
 	UPROPERTY(BlueprintReadWrite, EditAnywhere, Category = "Field", DisplayName = "Input Field A")
-	const UFieldNodeBase* RightField;
+	TObjectPtr<const UFieldNodeBase> RightField;
 
 	/** Left field to be processed */
 	UPROPERTY(BlueprintReadWrite, EditAnywhere, Category = "Field", DisplayName = "Input Field B")
-	const UFieldNodeBase* LeftField;
+	TObjectPtr<const UFieldNodeBase> LeftField;
 
 	/** Type of operation you want to perform between the 2 fields */
 	UPROPERTY(BlueprintReadWrite, EditAnywhere, Category = "Field", DisplayName = "Field Operation")
@@ -800,7 +800,7 @@ public:
 
 	/** Scalar field to be converted to an an integer one */
 	UPROPERTY(BlueprintReadWrite, EditAnywhere, Category = "Field", DisplayName = "Scalar Field")
-	const UFieldNodeFloat* FloatField;
+	TObjectPtr<const UFieldNodeFloat> FloatField;
 };
 
 /**
@@ -830,7 +830,7 @@ public:
 
 	/** Integer field to be converted to an a scalar one */
 	UPROPERTY(BlueprintReadWrite, EditAnywhere, Category = "Field", DisplayName = "Integer Field")
-	const UFieldNodeInt* IntField;
+	TObjectPtr<const UFieldNodeInt> IntField;
 };
 
 /**
@@ -868,11 +868,11 @@ public:
 
 	/** Culling field to be used */
 	UPROPERTY(BlueprintReadWrite, EditAnywhere, Category = "Field", DisplayName = "Culling Field")
-	const UFieldNodeBase* Culling;
+	TObjectPtr<const UFieldNodeBase> Culling;
 
 	/** Input field that will be evaluated according to the culling field result */
 	UPROPERTY(BlueprintReadWrite, EditAnywhere, Category = "Field", DisplayName = "Input Field")
-	const UFieldNodeBase* Field;
+	TObjectPtr<const UFieldNodeBase> Field;
 
 	/** Evaluate the input field if the result of the culling field is equal to 0 (Inside) or different from 0 (Outside) */
 	UPROPERTY(BlueprintReadWrite, EditAnywhere, Category = "Field", DisplayName = "Culling Operation")
diff --git a/Engine/Source/Runtime/Experimental/GeometryCollectionEngine/Public/ChaosBlueprint.h b/Engine/Source/Runtime/Experimental/GeometryCollectionEngine/Public/ChaosBlueprint.h
--- a/Engine/Source/Runtime/Experimental/GeometryCollectionEngine/Public/ChaosBlueprint.h
+++ b/Engine/Source/Runtime/Experimental/GeometryCollectionEngine/Public/ChaosBlueprint.h
@@ -77,11 +77,11 @@ public:
 
 	// Which chaos solver actors we're using. If empty, this listener will fallback to the "world" solver.
 	UPROPERTY(EditAnywhere, BlueprintReadOnly, Category = "Solvers")
-	TSet<AChaosSolverActor*> ChaosSolverActors;
+	TSet<TObjectPtr<AChaosSolverActor>> ChaosSolverActors;
 
 	// Which chaos solver actors we're using. If empty, this listener will fallback to the "world" solver.
 	UPROPERTY(EditAnywhere, BlueprintReadOnly, Category = "GeometryCollections")
-	TSet<AGeometryCollectionActor*> GeometryCollectionActors; // Using TSet automatically blocks if user tries to add same actor twice
+	TSet<TObjectPtr<AGeometryCollectionActor>> GeometryCollectionActors; // Using TSet automatically blocks if user tries to add same actor twice
 
 	// Dynamically adds a chaos solver to the listener
 	UFUNCTION(BlueprintCallable, Category = "Destruction Listener", meta = (WorldContext = "WorldContextObject"))
diff --git a/Engine/Source/Runtime/Experimental/GeometryCollectionEngine/Public/GeometryCollection/GeometryCollectionActor.h b/Engine/Source/Runtime/Experimental/GeometryCollectionEngine/Public/GeometryCollection/GeometryCollectionActor.h
--- a/Engine/Source/Runtime/Experimental/GeometryCollectionEngine/Public/GeometryCollection/GeometryCollectionActor.h
+++ b/Engine/Source/Runtime/Experimental/GeometryCollectionEngine/Public/GeometryCollection/GeometryCollectionActor.h
@@ -31,11 +31,11 @@ public:
 
 	/* GeometryCollectionComponent */
 	UPROPERTY(VisibleAnywhere, BlueprintReadOnly, Category = Destruction, meta = (ExposeFunctionCategories = "Components|GeometryCollection", AllowPrivateAccess = "true"))
-	UGeometryCollectionComponent* GeometryCollectionComponent;
+	TObjectPtr<UGeometryCollectionComponent> GeometryCollectionComponent;
 	UGeometryCollectionComponent* GetGeometryCollectionComponent() const { return GeometryCollectionComponent; }
 
 	UPROPERTY(VisibleAnywhere, Category = Destruction, meta = (ExposeFunctionCategories = "Components|GeometryCollection", AllowPrivateAccess = "true"))
-	UGeometryCollectionDebugDrawComponent* GeometryCollectionDebugDrawComponent;
+	TObjectPtr<UGeometryCollectionDebugDrawComponent> GeometryCollectionDebugDrawComponent;
 	UGeometryCollectionDebugDrawComponent* GetGeometryCollectionDebugDrawComponent() const { return GeometryCollectionDebugDrawComponent; }
 
 	UFUNCTION(BlueprintCallable, Category = "Physics")
diff --git a/Engine/Source/Runtime/Experimental/GeometryCollectionEngine/Public/GeometryCollection/GeometryCollectionCache.h b/Engine/Source/Runtime/Experimental/GeometryCollectionEngine/Public/GeometryCollection/GeometryCollectionCache.h
--- a/Engine/Source/Runtime/Experimental/GeometryCollectionEngine/Public/GeometryCollection/GeometryCollectionCache.h
+++ b/Engine/Source/Runtime/Experimental/GeometryCollectionEngine/Public/GeometryCollection/GeometryCollectionCache.h
@@ -62,7 +62,7 @@ private:
 
 	/** The collection that we recorded the data from */
 	UPROPERTY()
-	const UGeometryCollection* SupportedCollection;
+	TObjectPtr<const UGeometryCollection> SupportedCollection;
 
 	/** Guid pulled from the collection when the recording was last saved */
 	UPROPERTY()
diff --git a/Engine/Source/Runtime/Experimental/GeometryCollectionEngine/Public/GeometryCollection/GeometryCollectionComponent.h b/Engine/Source/Runtime/Experimental/GeometryCollectionEngine/Public/GeometryCollection/GeometryCollectionComponent.h
--- a/Engine/Source/Runtime/Experimental/GeometryCollectionEngine/Public/GeometryCollection/GeometryCollectionComponent.h
+++ b/Engine/Source/Runtime/Experimental/GeometryCollectionEngine/Public/GeometryCollection/GeometryCollectionComponent.h
@@ -61,7 +61,7 @@ struct FGeomComponentCacheParameters
 
 	// The cache to target when recording or playing
 	UPROPERTY(EditAnywhere, Category = Cache)
-	UGeometryCollectionCache* TargetCache;
+	TObjectPtr<UGeometryCollectionCache> TargetCache;
 
 	// Cache mode, whether disabled, playing or recording
 	UPROPERTY(EditAnywhere, Category = Cache)
@@ -366,7 +366,7 @@ public:
 
 	/** Chaos RBD Solver override. Will use the world's default solver actor if null. */
 	UPROPERTY(EditAnywhere, Category = "ChaosPhysics", meta = (DisplayName = "Chaos Solver"))
-	AChaosSolverActor* ChaosSolverActor;
+	TObjectPtr<AChaosSolverActor> ChaosSolverActor;
 
 	/** RestCollection */
 	void SetRestCollection(const UGeometryCollection * RestCollectionIn);
@@ -423,10 +423,10 @@ public:
 
 
 	UPROPERTY(EditAnywhere, NoClear, BlueprintReadOnly, Category = "ChaosPhysics")
-	const UGeometryCollection* RestCollection;
+	TObjectPtr<const UGeometryCollection> RestCollection;
 
 	UPROPERTY(EditAnywhere, NoClear, BlueprintReadOnly, Category = "ChaosPhysics")
-	TArray<const AFieldSystemActor*> InitializationFields;
+	TArray<TObjectPtr<const AFieldSystemActor>> InitializationFields;
 
 	/**
 	* When Simulating is enabled the Component will initialize its rigid bodies within the solver.
@@ -477,7 +477,7 @@ public:
 
 	/** Physical Properties */
 	UPROPERTY(meta=(DeprecatedProperty, DeprecationMessage="Physical material now derived from render materials, for instance overrides use PhysicalMaterialOverride."))
-	const UChaosPhysicalMaterial* PhysicalMaterial_DEPRECATED;
+	TObjectPtr<const UChaosPhysicalMaterial> PhysicalMaterial_DEPRECATED;
 
 	/** */
 	UPROPERTY(EditAnywhere, BlueprintReadWrite, Category = "ChaosPhysics|Initial Velocity")
@@ -492,7 +492,7 @@ public:
 	FVector InitialAngularVelocity;
 
 	UPROPERTY(EditAnywhere, BlueprintReadOnly, Category="ChaosPhysics|Collisions")
-	UPhysicalMaterial* PhysicalMaterialOverride;
+	TObjectPtr<UPhysicalMaterial> PhysicalMaterialOverride;
 
 	UPROPERTY()
 	FGeomComponentCacheParameters CacheParameters;
@@ -715,13 +715,13 @@ private:
 
 	// Temporary storage for body setup in order to initialise a dummy body instance
 	UPROPERTY(Transient)
-	UBodySetup* DummyBodySetup;
+	TObjectPtr<UBodySetup> DummyBodySetup;
 
 #if WITH_EDITORONLY_DATA
 	// Tracked editor actor that owns the original component so we can write back recorded caches
 	// from PIE.
 	UPROPERTY(Transient)
-	AActor* EditorActor;
+	TObjectPtr<AActor> EditorActor;
 #endif
 	void SwitchRenderModels(const AActor* Actor);
 
diff --git a/Engine/Source/Runtime/Experimental/GeometryCollectionEngine/Public/GeometryCollection/GeometryCollectionDebugDrawActor.h b/Engine/Source/Runtime/Experimental/GeometryCollectionEngine/Public/GeometryCollection/GeometryCollectionDebugDrawActor.h
--- a/Engine/Source/Runtime/Experimental/GeometryCollectionEngine/Public/GeometryCollection/GeometryCollectionDebugDrawActor.h
+++ b/Engine/Source/Runtime/Experimental/GeometryCollectionEngine/Public/GeometryCollection/GeometryCollectionDebugDrawActor.h
@@ -65,11 +65,11 @@ struct FGeometryCollectionDebugDrawActorSelectedRigidBody
 
 	/** Chaos RBD Solver. Will use the world's default solver actor if null. */
 	UPROPERTY(EditAnywhere, Category = "Selected Rigid Body")
-	AChaosSolverActor* Solver;
+	TObjectPtr<AChaosSolverActor> Solver;
 
 	/** Currently selected geometry collection. */
 	UPROPERTY(VisibleAnywhere, Category = "Selected Rigid Body")
-	AGeometryCollectionActor* GeometryCollection;
+	TObjectPtr<AGeometryCollectionActor> GeometryCollection;
 
 	/** Return the name of selected solver, or "None" if none is selected. */
 	FString GetSolverName() const;
@@ -331,7 +331,7 @@ public:
 
 	/** Display icon in the editor. */
 	UPROPERTY()
-	UBillboardComponent* SpriteComponent;
+	TObjectPtr<UBillboardComponent> SpriteComponent;
 
 	/** Game tick callback. This tick function is required to clean up the persistent debug lines. */
 	static AGeometryCollectionDebugDrawActor* FindOrCreate(UWorld* World);
