Fix World Partition unittests crash following 15216898.

diff --git a/Engine/Source/Runtime/Engine/Private/WorldPartition/WorldPartitionActorDesc.cpp b/Engine/Source/Runtime/Engine/Private/WorldPartition/WorldPartitionActorDesc.cpp
--- a/Engine/Source/Runtime/Engine/Private/WorldPartition/WorldPartitionActorDesc.cpp
+++ b/Engine/Source/Runtime/Engine/Private/WorldPartition/WorldPartitionActorDesc.cpp
@@ -235,7 +235,7 @@ AActor* FWorldPartitionActorDesc::Load() const
 	if (ActorPtr.IsExplicitlyNull())
 	{
 		// For now we assume that an ActorDesc that gets loaded is owned by an actual WorldPartition
-		UWorldPartition* WorldPartition = CastChecked<UWorldPartition>(Container);
+		UWorldPartition* WorldPartition = Container ? CastChecked<UWorldPartition>(Container) : nullptr;
 		check(WorldPartition || GIsAutomationTesting);
 
 		const FLinkerInstancingContext* InstancingContext = nullptr;
