HLOD support for data layers

diff --git a/Engine/Source/Runtime/Engine/Private/WorldPartition/DataLayer/DataLayersID.cpp b/Engine/Source/Runtime/Engine/Private/WorldPartition/DataLayer/DataLayersID.cpp
new file mode 100644
--- /dev/null
+++ b/Engine/Source/Runtime/Engine/Private/WorldPartition/DataLayer/DataLayersID.cpp
@@ -0,0 +1,31 @@
+// Copyright Epic Games, Inc. All Rights Reserved.
+
+#include "WorldPartition/DataLayer/DataLayersID.h"
+
+#include "Algo/Transform.h"
+
+#if WITH_EDITOR
+FDataLayersID::FDataLayersID()
+	: Hash(0)
+{}
+
+FDataLayersID::FDataLayersID(const TArray<const UDataLayer*>& InDataLayers)
+	: Hash(0)
+{
+	TArray<FName> DataLayers;
+	Algo::TransformIf(InDataLayers, DataLayers, [](const UDataLayer* Item) { return Item->IsDynamicallyLoaded(); }, [](const UDataLayer* Item) { return Item->GetFName(); });
+
+	if (DataLayers.Num())
+	{
+		TArray<FName> SortedDataLayers = DataLayers;
+		SortedDataLayers.Sort([](const FName& A, const FName& B) { return A.ToString() < B.ToString(); });
+
+		for (FName LayerName : SortedDataLayers)
+		{
+			Hash = FCrc::StrCrc32(*LayerName.ToString(), Hash);
+		}
+
+		check(Hash);
+	}
+}
+#endif
diff --git a/Engine/Source/Runtime/Engine/Private/WorldPartition/HLOD/HLODActorDesc.cpp b/Engine/Source/Runtime/Engine/Private/WorldPartition/HLOD/HLODActorDesc.cpp
--- a/Engine/Source/Runtime/Engine/Private/WorldPartition/HLOD/HLODActorDesc.cpp
+++ b/Engine/Source/Runtime/Engine/Private/WorldPartition/HLOD/HLODActorDesc.cpp
@@ -26,7 +26,9 @@ void FHLODActorDesc::Init(const AActor* InActor)
 		uint64 GridIndexZ;
 		HLODActor->GetGridIndices(GridIndexX, GridIndexY, GridIndexZ);
 
-		CellHash = ComputeCellHash(SubActorsHLODLayer->GetName(), GridIndexX, GridIndexY, GridIndexZ);
+		FDataLayersID DataLayersID(HLODActor->GetDataLayerObjects());
+
+		CellHash = ComputeCellHash(SubActorsHLODLayer->GetName(), GridIndexX, GridIndexY, GridIndexZ, DataLayersID);
 	}
 }
 
@@ -50,11 +52,12 @@ void FHLODActorDesc::Serialize(FArchive& Ar)
 	}
 }
 
-uint64 FHLODActorDesc::ComputeCellHash(const FString HLODLayerName, uint64 GridIndexX, uint64 GridIndexY, uint64 GridIndexZ)
+uint64 FHLODActorDesc::ComputeCellHash(const FString HLODLayerName, uint64 GridIndexX, uint64 GridIndexY, uint64 GridIndexZ, FDataLayersID DataLayersID)
 {
 	uint64 CellHash = FCrc::StrCrc32(*HLODLayerName);
 	CellHash = AppendCityHash(GridIndexX, CellHash);
 	CellHash = AppendCityHash(GridIndexY, CellHash);
-	return AppendCityHash(GridIndexZ, CellHash);
+	CellHash = AppendCityHash(GridIndexZ, CellHash);
+	return HashCombine(DataLayersID.GetHash(), CellHash);
 }
 #endif
diff --git a/Engine/Source/Runtime/Engine/Private/WorldPartition/HLOD/HLODBuilder.cpp b/Engine/Source/Runtime/Engine/Private/WorldPartition/HLOD/HLODBuilder.cpp
--- a/Engine/Source/Runtime/Engine/Private/WorldPartition/HLOD/HLODBuilder.cpp
+++ b/Engine/Source/Runtime/Engine/Private/WorldPartition/HLOD/HLODBuilder.cpp
@@ -13,6 +13,7 @@
 #include "Algo/ForEach.h"
 #include "Algo/RemoveIf.h"
 #include "Algo/Transform.h"
+#include "Misc/HashBuilder.h"
 #include "Templates/UniquePtr.h"
 
 #include "IMeshMergeUtilities.h"
@@ -38,7 +39,7 @@ public:
 		AWorldPartitionHLOD* HLODActor = nullptr;
 
 		// Compute HLODActor hash
-		uint64 CellHash = FHLODActorDesc::ComputeCellHash(HLODLayer->GetName(), Context->GridIndexX, Context->GridIndexY, Context->GridIndexZ);
+		uint64 CellHash = FHLODActorDesc::ComputeCellHash(HLODLayer->GetName(), Context->GridIndexX, Context->GridIndexY, Context->GridIndexZ, Context->DataLayersID);
 
 		int32 HLODActorRefIndex = INDEX_NONE;
 		FWorldPartitionHandle HLODActorHandle;
diff --git a/Engine/Source/Runtime/Engine/Private/WorldPartition/HLOD/HLODBuilder.h b/Engine/Source/Runtime/Engine/Private/WorldPartition/HLOD/HLODBuilder.h
--- a/Engine/Source/Runtime/Engine/Private/WorldPartition/HLOD/HLODBuilder.h
+++ b/Engine/Source/Runtime/Engine/Private/WorldPartition/HLOD/HLODBuilder.h
@@ -7,6 +7,7 @@
 #include "CoreMinimal.h"
 #include "Containers/Array.h"
 #include "WorldPartition/WorldPartitionHandle.h"
+#include "WorldPartition/DataLayer/DataLayersID.h"
 
 class AActor;
 class UWorldPartition;
@@ -22,6 +23,7 @@ struct ENGINE_API FHLODGenerationContext
 	int64 GridIndexX;
 	int64 GridIndexY;
 	int64 GridIndexZ;
+	FDataLayersID DataLayersID;
 	FName HLODLayerName;
 };
 
diff --git a/Engine/Source/Runtime/Engine/Private/WorldPartition/RuntimeSpatialHash/RuntimeSpatialHashGridHelper.h b/Engine/Source/Runtime/Engine/Private/WorldPartition/RuntimeSpatialHash/RuntimeSpatialHashGridHelper.h
--- a/Engine/Source/Runtime/Engine/Private/WorldPartition/RuntimeSpatialHash/RuntimeSpatialHashGridHelper.h
+++ b/Engine/Source/Runtime/Engine/Private/WorldPartition/RuntimeSpatialHash/RuntimeSpatialHashGridHelper.h
@@ -8,6 +8,7 @@
 
 #include "WorldPartition/WorldPartitionRuntimeSpatialHash.h"
 #include "WorldPartition/RuntimeSpatialHash/RuntimeSpatialHashActorCluster.h"
+#include "WorldPartition/DataLayer/DataLayersID.h"
 
 /**
   * Square 2D grid helper
diff --git a/Engine/Source/Runtime/Engine/Private/WorldPartition/RuntimeSpatialHash/RuntimeSpatialHashHLOD.cpp b/Engine/Source/Runtime/Engine/Private/WorldPartition/RuntimeSpatialHash/RuntimeSpatialHashHLOD.cpp
--- a/Engine/Source/Runtime/Engine/Private/WorldPartition/RuntimeSpatialHash/RuntimeSpatialHashHLOD.cpp
+++ b/Engine/Source/Runtime/Engine/Private/WorldPartition/RuntimeSpatialHash/RuntimeSpatialHashHLOD.cpp
@@ -123,13 +123,11 @@ static TArray<FGuid> GenerateHLODsForGrid(UWorldPartition* WorldPartition, const
 	{
 		const FSquare2DGridHelper::FGridLevel::FGridCell& GridCell = PartitionedActors.GetCell(CellCoord);
 
-		// For now, HLOD only processes actors that are not in data layers
-		if (const FSquare2DGridHelper::FGridLevel::FGridCellDataChunk* GridCellDataChunk = GridCell.GetNoDataLayersDataChunk())
+		for (const FSquare2DGridHelper::FGridLevel::FGridCellDataChunk& GridCellDataChunk : GridCell.GetDataChunks())
 		{
-			check(!GridCellDataChunk->HasDataLayers());
 			const bool bIsCellAlwaysLoaded = &GridCell == &PartitionedActors.GetAlwaysLoadedCell();
 
-			if (!bIsCellAlwaysLoaded && GridCellDataChunk->GetActors().Num() != 0)
+			if (!bIsCellAlwaysLoaded && GridCellDataChunk.GetActors().Num() != 0)
 			{
 				NbCellsToProcess++;
 			}
@@ -147,29 +145,28 @@ static TArray<FGuid> GenerateHLODsForGrid(UWorldPartition* WorldPartition, const
 	{
 		const FSquare2DGridHelper::FGridLevel::FGridCell& GridCell = PartitionedActors.GetCell(CellCoord);
 
-		// For now, HLOD only processes actors that are not in data layers
-		if (const FSquare2DGridHelper::FGridLevel::FGridCellDataChunk* GridCellDataChunk = GridCell.GetNoDataLayersDataChunk())
+		for (const FSquare2DGridHelper::FGridLevel::FGridCellDataChunk& GridCellDataChunk : GridCell.GetDataChunks())
 		{
-			check(!GridCellDataChunk->HasDataLayers());
 			const bool bIsCellAlwaysLoaded = &GridCell == &PartitionedActors.GetAlwaysLoadedCell();
 
-			if (!bIsCellAlwaysLoaded && GridCellDataChunk->GetActors().Num() != 0)
+			if (!bIsCellAlwaysLoaded && GridCellDataChunk.GetActors().Num() != 0)
 			{
 				SlowTask.EnterProgressFrame(1);
 
 				FBox2D CellBounds2D;
 				PartitionedActors.GetCellBounds(CellCoord, CellBounds2D);
 
-				FName CellName = UWorldPartitionRuntimeSpatialHash::GetCellName(WorldPartition, RuntimeGrid.GridName, CellCoord.Z, CellCoord.X, CellCoord.Y, GridCellDataChunk->GetDataLayersID());
+				FName CellName = UWorldPartitionRuntimeSpatialHash::GetCellName(WorldPartition, RuntimeGrid.GridName, CellCoord.Z, CellCoord.X, CellCoord.Y, GridCellDataChunk.GetDataLayersID());
 				FBox CellBounds = FBox(FVector(CellBounds2D.Min, WorldBounds.Min.Z), FVector(CellBounds2D.Max, WorldBounds.Max.Z));
 
 				Context.GridIndexX = CellCoord.X;
 				Context.GridIndexY = CellCoord.Y;
 				Context.GridIndexZ = CellCoord.Z;
+				Context.DataLayersID = GridCellDataChunk.GetDataLayersID();
 
 				// Gather (and potentially load) actors
 				TArray<AActor*> CellActors;
-				for (const FGuid& ActorGuid : GridCellDataChunk->GetActors())
+				for (const FGuid& ActorGuid : GridCellDataChunk.GetActors())
 				{
 					FWorldPartitionReference& ActorRef = Context.ActorReferences.Emplace_GetRef(WorldPartition, ActorGuid);
 					FWorldPartitionActorDesc* ActorDesc = ActorRef.Get();
@@ -192,6 +189,12 @@ static TArray<FGuid> GenerateHLODsForGrid(UWorldPartition* WorldPartition, const
 
 					for (AWorldPartitionHLOD* CellHLODActor : CellHLODActors)
 					{
+						// Make sure the generated HLOD actor has the same data layers as the source actors
+						for(const UDataLayer* DataLayer : GridCellDataChunk.GetDataLayers())
+						{
+							CellHLODActor->AddDataLayer(DataLayer);
+						}
+
 						FGuid ActorGuid = CellHLODActor->GetActorGuid();
 						GridHLODActors.Add(ActorGuid);
 
diff --git a/Engine/Source/Runtime/Engine/Private/WorldPartition/RuntimeSpatialHash/RuntimeSpatialHashNavigation.cpp b/Engine/Source/Runtime/Engine/Private/WorldPartition/RuntimeSpatialHash/RuntimeSpatialHashNavigation.cpp
--- a/Engine/Source/Runtime/Engine/Private/WorldPartition/RuntimeSpatialHash/RuntimeSpatialHashNavigation.cpp
+++ b/Engine/Source/Runtime/Engine/Private/WorldPartition/RuntimeSpatialHash/RuntimeSpatialHashNavigation.cpp
@@ -74,33 +74,35 @@ bool UWorldPartitionRuntimeSpatialHash::GenerateNavigationData()
 
 			const FSquare2DGridHelper::FGridLevel::FGridCell& GridCell = GridLevelHelper.GetCell(CellCoord);
 
-			FActorSpawnParameters SpawnParams;
-			SpawnParams.bDeferConstruction = true;
-			SpawnParams.bCreateActorPackage = true;
-			ANavigationDataChunkActor* DataChunkActor = World->SpawnActor<ANavigationDataChunkActor>(SpawnParams);
-			ActorCount++;
+			//@todo_ow: Properly handle data layers
+			if (const FSquare2DGridHelper::FGridLevel::FGridCellDataChunk* GridCellDataChunk = GridCell.GetNoDataLayersDataChunk())
+			{
+				FActorSpawnParameters SpawnParams;
+				SpawnParams.bDeferConstruction = true;
+				SpawnParams.bCreateActorPackage = true;
+				ANavigationDataChunkActor* DataChunkActor = World->SpawnActor<ANavigationDataChunkActor>(SpawnParams);
+				ActorCount++;
 
-			const FVector2D CellCenter = CellBounds.GetCenter();
-			DataChunkActor->SetActorLocation(FVector(CellCenter.X, CellCenter.Y, 0.f));
+				const FVector2D CellCenter = CellBounds.GetCenter();
+				DataChunkActor->SetActorLocation(FVector(CellCenter.X, CellCenter.Y, 0.f));
 
-			FBox TilesBounds(EForceInit::ForceInit);
-			DataChunkActor->CollectNavData(QueryBounds, TilesBounds);
+				FBox TilesBounds(EForceInit::ForceInit);
+				DataChunkActor->CollectNavData(QueryBounds, TilesBounds);
 
-			FBox ChunkActorBounds(FVector(QueryBounds.Min.X, QueryBounds.Min.Y, TilesBounds.Min.Z), FVector(QueryBounds.Max.X, QueryBounds.Max.Y, TilesBounds.Max.Z));
-			ChunkActorBounds = ChunkActorBounds.ExpandBy(FVector(-1.f, -1.f, 1.f)); //reduce XY by 1cm to avoid precision issues causing potential overflow on neighboring cell, add 1cm in Z to have a minimum of volume.
-			UE_LOG(LogWorldPartitionRuntimeSpatialHashNav, VeryVerbose, TEXT("Setting ChunkActorBounds to %s"), *ChunkActorBounds.ToString());
-			DataChunkActor->SetDataChunkActorBounds(ChunkActorBounds);
+				FBox ChunkActorBounds(FVector(QueryBounds.Min.X, QueryBounds.Min.Y, TilesBounds.Min.Z), FVector(QueryBounds.Max.X, QueryBounds.Max.Y, TilesBounds.Max.Z));
+				ChunkActorBounds = ChunkActorBounds.ExpandBy(FVector(-1.f, -1.f, 1.f)); //reduce XY by 1cm to avoid precision issues causing potential overflow on neighboring cell, add 1cm in Z to have a minimum of volume.
+				UE_LOG(LogWorldPartitionRuntimeSpatialHashNav, VeryVerbose, TEXT("Setting ChunkActorBounds to %s"), *ChunkActorBounds.ToString());
+				DataChunkActor->SetDataChunkActorBounds(ChunkActorBounds);
 
-			const FName CellName = GetCellName(RuntimeGrid.GridName, GridLevel, CellCoord.X, CellCoord.Y);
-			DataChunkActor->SetActorLabel(FString::Printf(TEXT("NavDataChunkActor_%s_%s"), *GetName(), *CellName.ToString()));
+				const FName CellName = GetCellName(RuntimeGrid.GridName, GridLevel, CellCoord.X, CellCoord.Y, GridCellDataChunk->GetDataLayersID());
+				DataChunkActor->SetActorLabel(FString::Printf(TEXT("NavDataChunkActor_%s_%s"), *GetName(), *CellName.ToString()));
 
-			//@todo_ow: Properly handle data layers
+				// Set target grid
+				DataChunkActor->RuntimeGrid = RuntimeGrid.GridName;
+				ValidNavigationDataChunkActors.Add(DataChunkActor);
 
-			// Set target grid
-			DataChunkActor->RuntimeGrid = RuntimeGrid.GridName;
-			ValidNavigationDataChunkActors.Add(DataChunkActor);
-
-			UE_LOG(LogWorldPartitionRuntimeSpatialHashNav, Verbose, TEXT("%i) %s added."), ActorCount, *DataChunkActor->GetName());
+				UE_LOG(LogWorldPartitionRuntimeSpatialHashNav, Verbose, TEXT("%i) %s added."), ActorCount, *DataChunkActor->GetName());
+			}
 		});
 
 	// Destroy all invalid navigation data chunk actors
diff --git a/Engine/Source/Runtime/Engine/Private/WorldPartition/WorldPartitionRuntimeSpatialHash.cpp b/Engine/Source/Runtime/Engine/Private/WorldPartition/WorldPartitionRuntimeSpatialHash.cpp
--- a/Engine/Source/Runtime/Engine/Private/WorldPartition/WorldPartitionRuntimeSpatialHash.cpp
+++ b/Engine/Source/Runtime/Engine/Private/WorldPartition/WorldPartitionRuntimeSpatialHash.cpp
@@ -67,33 +67,6 @@ static FAutoConsoleVariableRef CVarRuntimeSpatialHashCellToSourceAngleContributi
 	GRuntimeSpatialHashCellToSourceAngleContributionToCellImportance,
 	TEXT("Value between 0 and 1 that modulates the contribution of the angle between streaming source-to-cell vector and source-forward vector to the cell importance. The closest to 0, the less the angle will contribute to the cell importance."));
 
-// ------------------------------------------------------------------------------------------------
-
-#if WITH_EDITOR
-FDataLayersID::FDataLayersID()
-	: Hash(0)
-{}
-
-FDataLayersID::FDataLayersID(const TArray<const UDataLayer*>& InDataLayers)
-	: Hash(0)
-{
-	if (InDataLayers.Num())
-	{
-		Algo::TransformIf(InDataLayers, DataLayers, [](const UDataLayer* Item) { return Item->IsDynamicallyLoaded(); }, [](const UDataLayer* Item) { return Item->GetFName(); });
-
-		TArray<FName> SortedDataLayers = DataLayers;
-		SortedDataLayers.Sort([](const FName& A, const FName& B) { return A.ToString() < B.ToString(); });
-
-		for (FName LayerName: SortedDataLayers)
-		{
-			Hash = FCrc::StrCrc32(*LayerName.ToString(), Hash);
-		}
-
-		check(Hash);
-	}
-}
-#endif
-
 // ------------------------------------------------------------------------------------------------
 FSpatialHashStreamingGrid::FSpatialHashStreamingGrid()
 	: Origin(ForceInitToZero)
diff --git a/Engine/Source/Runtime/Engine/Public/WorldPartition/DataLayer/DataLayersID.h b/Engine/Source/Runtime/Engine/Public/WorldPartition/DataLayer/DataLayersID.h
new file mode 100644
--- /dev/null
+++ b/Engine/Source/Runtime/Engine/Public/WorldPartition/DataLayer/DataLayersID.h
@@ -0,0 +1,29 @@
+// Copyright Epic Games, Inc. All Rights Reserved.
+
+#pragma once
+
+#include "WorldPartition/DataLayer/DataLayer.h"
+
+#if WITH_EDITOR
+class FDataLayersID
+{
+public:
+	FDataLayersID();
+	FDataLayersID(const TArray<const UDataLayer*>& InDataLayers);
+
+	bool operator==(const FDataLayersID& Other) const
+	{
+		return Hash == Other.Hash;
+	}
+
+	bool operator!=(const FDataLayersID& Other) const
+	{
+		return !(*this == Other);
+	}
+
+	uint32 GetHash() const { return Hash; }
+
+private:
+	uint32 Hash;
+};
+#endif
diff --git a/Engine/Source/Runtime/Engine/Public/WorldPartition/HLOD/HLODActorDesc.h b/Engine/Source/Runtime/Engine/Public/WorldPartition/HLOD/HLODActorDesc.h
--- a/Engine/Source/Runtime/Engine/Public/WorldPartition/HLOD/HLODActorDesc.h
+++ b/Engine/Source/Runtime/Engine/Public/WorldPartition/HLOD/HLODActorDesc.h
@@ -5,7 +5,9 @@
 #include "CoreMinimal.h"
 #include "UObject/ObjectMacros.h"
 #include "Containers/Map.h"
+
 #include "WorldPartition/WorldPartitionActorDesc.h"
+#include "WorldPartition/DataLayer/DataLayersID.h"
 
 class UHLODLayer;
 
@@ -21,7 +23,7 @@ public:
 	inline const TArray<FGuid>& GetSubActors() const { return SubActors; }
 	inline uint64 GetCellHash() const { return CellHash; }
 
-	static uint64 ComputeCellHash(const FString HLODLayerName, uint64 GridIndexX, uint64 GridIndexY, uint64 GridIndexZ);
+	static uint64 ComputeCellHash(const FString HLODLayerName, uint64 GridIndexX, uint64 GridIndexY, uint64 GridIndexZ, FDataLayersID DataLayersID);
 
 protected:
 	virtual void Init(const AActor* InActor) override;
diff --git a/Engine/Source/Runtime/Engine/Public/WorldPartition/WorldPartitionRuntimeSpatialHash.h b/Engine/Source/Runtime/Engine/Public/WorldPartition/WorldPartitionRuntimeSpatialHash.h
--- a/Engine/Source/Runtime/Engine/Public/WorldPartition/WorldPartitionRuntimeSpatialHash.h
+++ b/Engine/Source/Runtime/Engine/Public/WorldPartition/WorldPartitionRuntimeSpatialHash.h
@@ -8,6 +8,7 @@
 #include "GameFramework/Info.h"
 
 #include "WorldPartition.h"
+#include "WorldPartition/DataLayer/DataLayersID.h"
 #include "WorldPartitionRuntimeHash.h"
 #include "WorldPartitionRuntimeSpatialHashCell.h"
 #include "WorldPartitionRuntimeSpatialHash.generated.h"
@@ -181,31 +182,6 @@ public:
 	FSpatialHashRuntimeGrid	GridSettings;
 };
 
-#if WITH_EDITOR
-class FDataLayersID
-{
-public:
-	FDataLayersID();
-	FDataLayersID(const TArray<const UDataLayer*>& InDataLayers);
-
-	bool operator==(const FDataLayersID& Other) const
-	{
-		return (Hash == Other.Hash) && (DataLayers == Other.DataLayers);
-	}
-
-	bool operator!=(const FDataLayersID& Other) const
-	{
-		return !(*this == Other);
-	}
-
-	uint32 GetHash() const { return Hash; }
-
-private:
-	TArray<FName> DataLayers;
-	uint32 Hash;
-};
-#endif
-
 UCLASS()
 class ENGINE_API UWorldPartitionRuntimeSpatialHash : public UWorldPartitionRuntimeHash
 {
@@ -223,7 +199,7 @@ public:
 	virtual bool GenerateNavigationData() override;
 	virtual FName GetActorRuntimeGrid(const AActor* Actor) const override;
 
-	FName GetCellName(FName InGridName, int32 InLevel, int32 InCellX, int32 InCellY, const FDataLayersID& InDataLayerID = FDataLayersID()) const;
+	FName GetCellName(FName InGridName, int32 InLevel, int32 InCellX, int32 InCellY, const FDataLayersID& InDataLayerID) const;
 	static FName GetCellName(UWorldPartition* WorldPartition, FName InGridName, int32 InLevel, int32 InCellX, int32 InCellY, const FDataLayersID& InDataLayerID);
 #endif
 
