Ignore runtime data layer activation of non-dynamically loaded data layers

diff --git a/Engine/Source/Runtime/Engine/Private/WorldPartition/DataLayer/DataLayerSubsystem.cpp b/Engine/Source/Runtime/Engine/Private/WorldPartition/DataLayer/DataLayerSubsystem.cpp
--- a/Engine/Source/Runtime/Engine/Private/WorldPartition/DataLayer/DataLayerSubsystem.cpp
+++ b/Engine/Source/Runtime/Engine/Private/WorldPartition/DataLayer/DataLayerSubsystem.cpp
@@ -89,7 +89,7 @@ void UDataLayerSubsystem::ActivateDataLayer(const FActorDataLayer& InDataLayer,
 
 void UDataLayerSubsystem::ActivateDataLayer(const UDataLayer* InDataLayer, bool bInActivate)
 {
-	if (!InDataLayer)
+	if (!InDataLayer || !InDataLayer->IsDynamicallyLoaded())
 	{
 		return;
 	}
