LevelInstance: Fix BreakLevelInstance not working on PackedLevelInstances.

diff --git a/Engine/Source/Runtime/Engine/Private/LevelInstance/LevelInstanceSubsystem.cpp b/Engine/Source/Runtime/Engine/Private/LevelInstance/LevelInstanceSubsystem.cpp
--- a/Engine/Source/Runtime/Engine/Private/LevelInstance/LevelInstanceSubsystem.cpp
+++ b/Engine/Source/Runtime/Engine/Private/LevelInstance/LevelInstanceSubsystem.cpp
@@ -783,6 +783,12 @@ bool ULevelInstanceSubsystem::BreakLevelInstance(ALevelInstance* LevelInstanceAc
 		// Can only break the top level LevelInstance
 		check(LevelInstanceActor->GetLevel() == GetWorld()->GetCurrentLevel());
 
+		// Actors in a packed level instance will not be streamed in unless they are editing. Must force this before moving.
+		if (LevelInstanceActor->IsA<APackedLevelInstance>())
+		{
+			BlockLoadLevelInstance(LevelInstanceActor);
+		}
+
 		// need to ensure that LevelInstanceActor has been streamed in fully
 		GEngine->BlockTillLevelStreamingCompleted(LevelInstanceActor->GetWorld());
 
@@ -815,6 +821,11 @@ bool ULevelInstanceSubsystem::BreakLevelInstance(ALevelInstance* LevelInstanceAc
 			return false;
 		}
 
+		if (LevelInstanceActor->IsA<APackedLevelInstance>())
+		{
+			BlockUnloadLevelInstance(LevelInstanceActor);
+		}
+
 		// Destroy the old LevelInstance instance actor
 		GetWorld()->DestroyActor(LevelInstanceActor);
 	
@@ -1513,6 +1524,22 @@ void ULevelInstanceSubsystem::BlockLoadLevelInstance(ALevelInstance* LevelInstan
 		UpdateStreamingState();
 	}
 }
+
+void ULevelInstanceSubsystem::BlockUnloadLevelInstance(ALevelInstance* LevelInstanceActor)
+{
+	check(!LevelInstanceActor->IsEditing());
+	RequestUnloadLevelInstance(LevelInstanceActor);
+
+	// Make sure blocking unloads can happen and are not part of transaction
+	TGuardValue<ITransaction*> TransactionGuard(GUndo, nullptr);
+
+	// Blocking until LevelInstance is loaded and all its child LevelInstances
+	while (LevelInstancesToUnload.Num())
+	{
+		UpdateStreamingState();
+	}
+}
+
 #endif
 
 #undef LOCTEXT_NAMESPACE
diff --git a/Engine/Source/Runtime/Engine/Public/LevelInstance/LevelInstanceSubsystem.h b/Engine/Source/Runtime/Engine/Public/LevelInstance/LevelInstanceSubsystem.h
--- a/Engine/Source/Runtime/Engine/Public/LevelInstance/LevelInstanceSubsystem.h
+++ b/Engine/Source/Runtime/Engine/Public/LevelInstance/LevelInstanceSubsystem.h
@@ -85,6 +85,7 @@ public:
 
 	ALevelInstance* GetParentLevelInstance(const AActor* Actor) const;
 	void BlockLoadLevelInstance(ALevelInstance* LevelInstanceActor);
+	void BlockUnloadLevelInstance(ALevelInstance* LevelInstanceActor);
 #endif
 
 private:
