World Partition HLODs - Fixed instancing issue were material overrides weren't taken into account

diff --git a/Engine/Source/Runtime/Engine/Private/WorldPartition/HLOD/HLODBuilder.cpp b/Engine/Source/Runtime/Engine/Private/WorldPartition/HLOD/HLODBuilder.cpp
--- a/Engine/Source/Runtime/Engine/Private/WorldPartition/HLOD/HLODBuilder.cpp
+++ b/Engine/Source/Runtime/Engine/Private/WorldPartition/HLOD/HLODBuilder.cpp
@@ -130,6 +130,59 @@ public:
  */
 class FHLODBuilder_Instancing : public FHLODBuilder
 {
+	// We want to merge all SMC that are using the same static mesh
+	// However, we must also take material overiddes into account.
+	struct FInstancingKey
+	{
+		FInstancingKey(const UStaticMeshComponent* SMC)
+		{
+			FHashBuilder HashBuilder;
+
+			StaticMesh = SMC->GetStaticMesh();
+			HashBuilder << StaticMesh;
+
+			const int32 NumMaterials = SMC->GetNumMaterials();
+			Materials.Reserve(NumMaterials);
+
+			for (int32 MaterialIndex = 0; MaterialIndex < NumMaterials; ++MaterialIndex)
+			{
+				UMaterialInterface* Material = SMC->GetMaterial(MaterialIndex);
+
+				Materials.Add(Material);
+				HashBuilder << Material;
+			}
+
+			Hash = HashBuilder.GetHash();
+		}
+
+		void ApplyTo(UStaticMeshComponent* SMC) const
+		{
+			// Set static mesh
+			SMC->SetStaticMesh(StaticMesh);
+			
+			// Set material overrides
+			for (int32 MaterialIndex = 0; MaterialIndex < Materials.Num(); ++MaterialIndex)
+			{
+				SMC->SetMaterial(MaterialIndex, Materials[MaterialIndex]);
+			}
+		}
+
+		friend uint32 GetTypeHash(const FInstancingKey& Key)
+		{
+			return Key.Hash;
+		}
+
+		bool operator==(const FInstancingKey& Other) const
+		{
+			return Hash == Other.Hash && StaticMesh == Other.StaticMesh && Materials == Other.Materials;
+		}
+
+	private:
+		UStaticMesh*				StaticMesh;
+		TArray<UMaterialInterface*>	Materials;
+		uint32						Hash;
+	};
+
 	virtual void Build(const TArray<UPrimitiveComponent*>& InSubComponents) override
 	{
 		TRACE_CPUPROFILER_EVENT_SCOPE(FHLODBuilder_Instancing::BuildHLOD);
@@ -139,24 +192,24 @@ class FHLODBuilder_Instancing : public FHLODBuilder
 			TArray<UPrimitiveComponent*> Components;
 
 			// Gather all meshes to instantiate along with their transforms
-			TMap<UStaticMesh*, TArray<UPrimitiveComponent*>> Instances;
+			TMap<FInstancingKey, TArray<UPrimitiveComponent*>> Instances;
 			for (UPrimitiveComponent* Primitive : InSubComponents)
 			{
 				if (UStaticMeshComponent* SMC = Cast<UStaticMeshComponent>(Primitive))
 				{
-					Instances.FindOrAdd(SMC->GetStaticMesh()).Add(SMC);
+					Instances.FindOrAdd(FInstancingKey(SMC)).Add(SMC);
 				}
 			}
 
 			// Create an ISMC for each SM asset we found
 			for (const auto& Entry : Instances)
 			{
-				UStaticMesh* EntryStaticMesh = Entry.Key;
+				const FInstancingKey EntryInstancingKey = Entry.Key;
 				const TArray<UPrimitiveComponent*>& EntryComponents = Entry.Value;
 			
 				UInstancedStaticMeshComponent* Component = NewObject<UInstancedStaticMeshComponent>(HLODActor);
-				Component->SetStaticMesh(EntryStaticMesh);
-				Component->SetForcedLodModel(EntryStaticMesh->GetNumLODs());
+				EntryInstancingKey.ApplyTo(Component);
+				Component->SetForcedLodModel(Component->GetStaticMesh()->GetNumLODs());
 
 				// Add all instances
 				for (UPrimitiveComponent* SMC : EntryComponents)
