Changes to handling of source model data for StaticMeshes.
- The MeshDescription for each source model LOD is now wrapped in a UStaticMeshDescription, meaning that it can be transacted separately to the parent static mesh. This improves memory and time performance in static mesh transactions, e.g. when assigning materials.
- FStaticMeshSourceModel now has its own API for managing the mesh description and its bulk data.
- A new optional Hi-res SourceModel member has been added to StaticMesh; this allows the original source data to be stored here, while a reduced LOD0 can be stored as SourceModel LOD0, to allow geometry modelling tools to work efficiently.
- MeshDescription is now serialized using a proxy archive object which ensures that FNames can always be archived as strings.
- UStaticMeshDescription is now serialized.
- UStaticMeshDescription has been reverted to only holding its own mesh description, instead of being able to reference any external mesh description. This previous change was never released.
- RawMesh is now immediately converted to MeshDescription in PostLoad, so that the MeshDescription[BulkData] pointers are always valid, and can be tag serialised automatically during transactions. RawMesh should now not be used in the editor at all, and will soon be deprecated.

This change results in far quicker and memory-efficient transactions in StaticMeshes. e.g. assign material to a static mesh section for a 20000 vertex mesh now takes less than a second (discounting time to rebuild the render proxy) instead of 7+ seconds.

diff --git a/Engine/Plugins/Enterprise/DatasmithImporter/Source/DatasmithImporter/Private/DatasmithBlueprintLibrary.cpp b/Engine/Plugins/Enterprise/DatasmithImporter/Source/DatasmithImporter/Private/DatasmithBlueprintLibrary.cpp
--- a/Engine/Plugins/Enterprise/DatasmithImporter/Source/DatasmithImporter/Private/DatasmithBlueprintLibrary.cpp
+++ b/Engine/Plugins/Enterprise/DatasmithImporter/Source/DatasmithImporter/Private/DatasmithBlueprintLibrary.cpp
@@ -893,7 +893,7 @@ void UDatasmithStaticMeshBlueprintLibrary::SetupStaticLighting(const TArray< UOb
 			{
 				int32 MaxBiggestUVChannel = Lightmass::MAX_TEXCOORDS;
 
-				if (const FMeshDescription* MeshDescription = SourceModel.MeshDescription.Get())
+				if (FMeshDescription* MeshDescription = SourceModel.GetCachedMeshDescription())
 				{
 					FStaticMeshConstAttributes Attributes(*MeshDescription);
 
diff --git a/Engine/Plugins/Runtime/HairStrands/Source/HairStrandsCore/Private/GroomAssetCards.cpp b/Engine/Plugins/Runtime/HairStrands/Source/HairStrandsCore/Private/GroomAssetCards.cpp
--- a/Engine/Plugins/Runtime/HairStrands/Source/HairStrandsCore/Private/GroomAssetCards.cpp
+++ b/Engine/Plugins/Runtime/HairStrands/Source/HairStrandsCore/Private/GroomAssetCards.cpp
@@ -123,9 +123,9 @@ FString FHairGroupsCardsSourceDescription::GetMeshKey() const
 	{
 		ImportedMesh->ConditionalPostLoad();
 		FStaticMeshSourceModel& SourceModel = ImportedMesh->GetSourceModel(0);
-		if (SourceModel.MeshDescriptionBulkData)
+		if (SourceModel.GetMeshDescriptionBulkData())
 		{
-			return SourceModel.MeshDescriptionBulkData->GetIdString();
+			return SourceModel.GetMeshDescriptionBulkData()->GetIdString();
 		}
 
 	}
@@ -133,9 +133,9 @@ FString FHairGroupsCardsSourceDescription::GetMeshKey() const
 	{
 		ProceduralMesh->ConditionalPostLoad();
 		FStaticMeshSourceModel& SourceModel = ProceduralMesh->GetSourceModel(0);
-		if (SourceModel.MeshDescriptionBulkData)
+		if (SourceModel.GetMeshDescriptionBulkData())
 		{
-			return SourceModel.MeshDescriptionBulkData->GetIdString();
+			return SourceModel.GetMeshDescriptionBulkData()->GetIdString();
 		}
 	}
 #endif
diff --git a/Engine/Plugins/Runtime/HairStrands/Source/HairStrandsCore/Private/GroomAssetMeshes.cpp b/Engine/Plugins/Runtime/HairStrands/Source/HairStrandsCore/Private/GroomAssetMeshes.cpp
--- a/Engine/Plugins/Runtime/HairStrands/Source/HairStrandsCore/Private/GroomAssetMeshes.cpp
+++ b/Engine/Plugins/Runtime/HairStrands/Source/HairStrandsCore/Private/GroomAssetMeshes.cpp
@@ -28,9 +28,9 @@ FString FHairGroupsMeshesSourceDescription::GetMeshKey() const
 	{
 		ImportedMesh->ConditionalPostLoad();
 		FStaticMeshSourceModel& SourceModel = ImportedMesh->GetSourceModel(0);
-		if (SourceModel.MeshDescriptionBulkData)
+		if (SourceModel.GetMeshDescriptionBulkData())
 		{
-			return SourceModel.MeshDescriptionBulkData->GetIdString();
+			return SourceModel.GetMeshDescriptionBulkData()->GetIdString();
 		}
 	}
 #endif
diff --git a/Engine/Source/Runtime/Core/Public/UObject/UE5MainStreamObjectVersion.h b/Engine/Source/Runtime/Core/Public/UObject/UE5MainStreamObjectVersion.h
--- a/Engine/Source/Runtime/Core/Public/UObject/UE5MainStreamObjectVersion.h
+++ b/Engine/Source/Runtime/Core/Public/UObject/UE5MainStreamObjectVersion.h
@@ -85,6 +85,9 @@ struct CORE_API FUE5MainStreamObjectVersion
 		// Fix potentially duplicated actors when using ForceExternalActorLevelReference
 		FixForceExternalActorLevelReferenceDuplicates,
 
+		// Make UMeshDescriptionBase serializable
+		SerializeMeshDescriptionBase,
+
 		// -----<new versions can be added above this line>-------------------------------------------------
 		VersionPlusOne,
 		LatestVersion = VersionPlusOne - 1
diff --git a/Engine/Source/Runtime/Engine/Classes/Engine/StaticMesh.h b/Engine/Source/Runtime/Engine/Classes/Engine/StaticMesh.h
--- a/Engine/Source/Runtime/Engine/Classes/Engine/StaticMesh.h
+++ b/Engine/Source/Runtime/Engine/Classes/Engine/StaticMesh.h
@@ -31,10 +31,11 @@ class UMaterialInterface;
 class UNavCollisionBase;
 class UStaticMeshComponent;
 class UStaticMeshDescription;
+class UMeshDescriptionBulkDataWrapper;
 class FStaticMeshUpdate;
 struct FMeshDescription;
-struct FMeshDescriptionBulkData;
 struct FStaticMeshLODResources;
+struct FMeshDescriptionBulkData;
 
 /*-----------------------------------------------------------------------------
 	Async Static Mesh Compilation
@@ -332,6 +333,8 @@ struct FStaticMeshSourceModel
 {
 	GENERATED_USTRUCT_BODY()
 
+	friend class UStaticMesh;
+
 #if WITH_EDITOR
 	/**
 	 * Imported raw mesh data. Optional for all but the first LOD.
@@ -346,9 +349,9 @@ struct FStaticMeshSourceModel
 	 * The staticmesh owner of this source model. We need the SM to be able to convert between MeshDesription and RawMesh.
 	 * RawMesh use int32 material index and MeshDescription use FName material slot name.
 	 * This memeber is fill in the PostLoad of the static mesh.
-	 * TODO: Remove this member when FRawMesh will be remove.
 	 */
 	class UStaticMesh* StaticMeshOwner;
+
 	/*
 	 * Accessor to Load and save the raw mesh or the mesh description depending on the editor settings.
 	 * Temporary until we deprecate the RawMesh.
@@ -366,7 +369,8 @@ struct FStaticMeshSourceModel
 	 * If this is valid, this means the mesh description has either been unpacked from the bulk data stored in the asset,
 	 * or one has been generated by the build tools (or converted from legacy RawMesh).
 	 */
-	TUniquePtr<FMeshDescription> MeshDescription;
+	UPROPERTY(Transient)
+	UStaticMeshDescription* StaticMeshDescription;
 
 	/**
 	 * Bulk data containing mesh description. LOD0 must be valid, but autogenerated lower LODs may be invalid.
@@ -375,9 +379,70 @@ struct FStaticMeshSourceModel
 	 * If this is invalid, either the LOD is autogenerated (for LOD1+), or the asset is a legacy asset whose
 	 * data is in the RawMeshBulkData.
 	 */
-	TUniquePtr<FMeshDescriptionBulkData> MeshDescriptionBulkData;
+	UPROPERTY()
+	UMeshDescriptionBulkDataWrapper* StaticMeshDescriptionBulkData;
 #endif
 
+#if WITH_EDITOR
+	/**
+	 * Create a new MeshDescription object
+	 */
+	ENGINE_API FMeshDescription* CreateMeshDescription();
+
+	/**
+	 * Tries to initialize a MeshDescription into the given FMeshDescription object, returning whether it was successful.
+	 * It will try the following steps in order to obtain it:
+	 * - if there's valid MeshDescriptionBulkData, unpack it from that
+	 * - if there's a valid DDC entry from a previously converted RawMesh, load it and unpack it
+	 * - if there's a valid RawMesh, convert it to MeshDescription and return that.
+	 * If all the above failed, then return failure.
+	 */
+	ENGINE_API bool LoadMeshDescription(FMeshDescription& OutMeshDescription) const;
+
+	/**
+	 * If there's a valid cached MeshDescription, return a copy of it.
+	 * Otherwise unpack the mesh description from bulk data and return the result (without caching it).
+	 */
+	ENGINE_API bool CloneMeshDescription(FMeshDescription& OutMeshDescription) const;
+
+	/**
+	 * Gets the cached mesh description if there was one, otherwise loads, caches and returns it.
+	 */
+	ENGINE_API FMeshDescription* GetOrCacheMeshDescription();
+
+	/**
+	 * Gets the cached mesh description, or nullptr if there wasn't one
+	 */
+	ENGINE_API FMeshDescription* GetCachedMeshDescription() const;
+
+	/**
+	 * Gets the cached StaticMeshDescription if there is one, or nullptr if not.
+	 */
+	UStaticMeshDescription* GetCachedStaticMeshDescription() const { return StaticMeshDescription; }
+
+	/**
+	 * Gets the mesh description bulk data, or nullptr if not valid 
+	 */
+	ENGINE_API FMeshDescriptionBulkData* GetMeshDescriptionBulkData() const;
+
+	/**
+	 * Returns whether the cached mesh description is valid or not
+	 */
+	ENGINE_API bool IsMeshDescriptionValid() const;
+
+	/**
+	 * Commits the cached mesh description to bulk data
+	 */
+	ENGINE_API void CommitMeshDescription(bool bUseHashAsGuid);
+
+	/**
+	 * Clears the cached mesh description
+	 */
+	ENGINE_API void ClearMeshDescription();
+
+#endif
+
+public:
 	/** Settings applied when building the mesh. */
 	UPROPERTY(EditAnywhere, Category=BuildSettings)
 	FMeshBuildSettings BuildSettings;
@@ -417,8 +482,6 @@ struct FStaticMeshSourceModel
 	/** Serializes bulk data. */
 	void SerializeBulkData(FArchive& Ar, UObject* Owner);
 
-	/** Create a new MeshDescription object */
-	FMeshDescription* CreateMeshDescription();
 #endif
 };
 
@@ -733,6 +796,10 @@ private:
 	UPROPERTY()
 	TArray<FStaticMeshSourceModel> SourceModels;
 
+	/** Optional hi-res source data */
+	UPROPERTY()
+	FStaticMeshSourceModel HiResSourceModel;
+
 	void SetLightmapUVVersion(int32 InLightmapUVVersion)
 	{
 		WaitUntilAsyncPropertyReleased(EStaticMeshAsyncProperties::LightmapUVVersion);
@@ -1331,18 +1398,18 @@ public:
 		{}
 
 		/**
-		 * If set to false, the caller can be from any thread but will have the
-		 * responsability to call MarkPackageDirty() from the main thread.
-		 */
+		* If set to false, the caller can be from any thread but will have the
+		* responsability to call MarkPackageDirty() from the main thread.
+		*/
 		bool bMarkPackageDirty;
 
 		/**
-		 * Uses a hash as the GUID, useful to prevent recomputing content already in cache.
-		 */
+		* Uses a hash as the GUID, useful to prevent recomputing content already in cache.
+		*/
 		bool bUseHashAsGuid;
 	};
 
-	/*
+	/**
 	 * Serialize the mesh description into its more optimized form.
 	 *
 	 * @param	LodIndex	Index of the StaticMesh LOD.
@@ -1350,9 +1417,41 @@ public:
 	 */
 	ENGINE_API void CommitMeshDescription(int32 LodIndex, const FCommitMeshDescriptionParams& Params = FCommitMeshDescriptionParams());
 
+	/**
+	 * Clears the cached mesh description for the given LOD.
+	 * Note that this does not empty the bulk data.
+	 */
 	ENGINE_API void ClearMeshDescription(int32 LodIndex);
+
+	/**
+	 * Clears cached mesh descriptions for all LODs.
+	 */
 	ENGINE_API void ClearMeshDescriptions();
 
+	ENGINE_API bool LoadHiResMeshDescription(FMeshDescription& OutMeshDescription) const;
+	ENGINE_API bool CloneHiResMeshDescription(FMeshDescription& OutMeshDescription) const;
+	ENGINE_API FMeshDescription* CreateHiResMeshDescription();
+	ENGINE_API FMeshDescription* CreateHiResMeshDescription(FMeshDescription MeshDescription);
+	ENGINE_API FMeshDescription* GetHiResMeshDescription() const;
+	ENGINE_API bool IsHiResMeshDescriptionValid() const;
+	ENGINE_API void CommitHiResMeshDescription(const FCommitMeshDescriptionParams& Params = FCommitMeshDescriptionParams());
+	ENGINE_API void ClearHiResMeshDescription();
+
+	/**
+	 * Performs a Modify on the StaticMeshDescription object pertaining to the given LODIndex
+	 */
+	ENGINE_API void ModifyMeshDescription(int32 LodIndex, bool bAlwaysMarkDirty = true);
+
+	/**
+	 * Performs a Modify on StaticMeshDescription objects for all LODs
+	 */
+	ENGINE_API void ModifyAllMeshDescriptions(bool bAlwaysMarkDirty = true);
+
+	/**
+	 * Performs a Modify on the hi-res StaticMeshDescription
+	 */
+	ENGINE_API void ModifyHiResMeshDescription(bool bAlwaysMarkDirty = true);
+
 	/**
 	 * Adds an empty UV channel at the end of the existing channels on the given LOD of a StaticMesh.
 	 * @param	LODIndex			Index of the StaticMesh LOD.
@@ -1490,6 +1589,8 @@ public:
 	ENGINE_API const FStaticMeshSourceModel& GetSourceModel(int32 Index) const;
 	ENGINE_API int32 GetNumSourceModels() const;
 	ENGINE_API bool IsSourceModelValid(int32 Index) const;
+	ENGINE_API FStaticMeshSourceModel& GetHiResSourceModel();
+	ENGINE_API const FStaticMeshSourceModel& GetHiResSourceModel() const;
 
 	ENGINE_API FMeshSectionInfoMap& GetSectionInfoMap();
 	ENGINE_API const FMeshSectionInfoMap& GetSectionInfoMap() const;
@@ -1821,11 +1922,6 @@ private:
 	 */
 	void FixupZeroTriangleSections();
 
-	/**
-	* Return mesh data key. The key is the ddc filename for the mesh data
-	*/
-	bool GetMeshDataKey(int32 LodIndex, FString& OutKey) const;
-
 	/**
 	* Caches mesh data.
 	*/
diff --git a/Engine/Source/Runtime/Engine/Private/StaticMesh.cpp b/Engine/Source/Runtime/Engine/Private/StaticMesh.cpp
--- a/Engine/Source/Runtime/Engine/Private/StaticMesh.cpp
+++ b/Engine/Source/Runtime/Engine/Private/StaticMesh.cpp
@@ -24,6 +24,7 @@
 #include "UObject/Package.h"
 #include "UObject/PackageResourceManager.h"
 #include "UObject/RenderingObjectVersion.h"
+#include "UObject/UE5MainStreamObjectVersion.h"
 #include "UObject/UObjectAnnotation.h"
 #include "EngineUtils.h"
 #include "Engine/AssetUserData.h"
@@ -2426,10 +2427,10 @@ static FString BuildStaticMeshDerivedDataKeySuffix(const ITargetPlatform* Target
 	{
 		FStaticMeshSourceModel& SrcModel = Mesh->GetSourceModel(LODIndex);
 		
-		if (SrcModel.MeshDescriptionBulkData.IsValid())
+		if (SrcModel.GetMeshDescriptionBulkData() != nullptr)
 		{
 			KeySuffix += "MD";
-			KeySuffix += SrcModel.MeshDescriptionBulkData->GetIdString();
+			KeySuffix += SrcModel.GetMeshDescriptionBulkData()->GetIdString();
 		}
 		else if (!SrcModel.RawMeshBulkData->IsEmpty())
 		{
@@ -3442,6 +3443,8 @@ void UStaticMesh::PostEditChangeProperty(FPropertyChangedEvent& PropertyChangedE
 		GetSourceModel(Index).StaticMeshOwner = this;
 	}
 
+	GetHiResSourceModel().StaticMeshOwner = this;
+
 	Build(/*bSilent=*/ true);
 
 	if (PropertyName == GET_MEMBER_NAME_CHECKED(UStaticMesh, bHasNavigationData)
@@ -3486,6 +3489,8 @@ void UStaticMesh::PostEditUndo()
 		GetSourceModel(Index).StaticMeshOwner = this;
 	}
 
+	GetHiResSourceModel().StaticMeshOwner = this;
+
 	// The super will cause a Build() via PostEditChangeProperty().
 	Super::PostEditUndo();
 }
@@ -3737,6 +3742,17 @@ void UStaticMesh::RemoveSourceModel(const int32 Index)
 	GetSourceModels().RemoveAt(Index);
 }
 
+const FStaticMeshSourceModel& UStaticMesh::GetHiResSourceModel() const
+{
+	return HiResSourceModel;
+}
+
+FStaticMeshSourceModel& UStaticMesh::GetHiResSourceModel()
+{
+	return HiResSourceModel;
+}
+
+
 bool UStaticMesh::FixLODRequiresAdjacencyInformation(const int32 LODIndex, const bool bPreviewMode, bool bPromptUser, bool* OutUserCancel)
 {
 	if (OutUserCancel != nullptr)
@@ -3996,6 +4012,7 @@ bool UStaticMesh::IsCachedCookedPlatformDataLoaded(const ITargetPlatform* Target
 }
 #endif
 
+
 /*------------------------------------------------------------------------------
 	FStaticMeshSourceModel
 ------------------------------------------------------------------------------*/
@@ -4011,6 +4028,8 @@ FStaticMeshSourceModel::FStaticMeshSourceModel()
 	SourceImportFilename = FString();
 #if WITH_EDITORONLY_DATA
 	bImportWithBaseMesh = false;
+	StaticMeshDescription = nullptr;
+	StaticMeshDescriptionBulkData = nullptr;
 #endif
 }
 
@@ -4024,12 +4043,13 @@ FStaticMeshSourceModel::~FStaticMeshSourceModel()
 #endif // #if WITH_EDITOR
 }
 
+
 #if WITH_EDITOR
 bool FStaticMeshSourceModel::IsRawMeshEmpty() const
 {
 	// Even if the RawMeshBulkData is empty, it may be because it's a new asset serialized as MeshDescription.
 	// Hence MeshDescriptionBulkData must also be invalid (and, by consequence, also MeshDescription).
-	return RawMeshBulkData->IsEmpty() && !MeshDescriptionBulkData.IsValid() && !MeshDescription.IsValid();
+	return RawMeshBulkData->IsEmpty() && StaticMeshDescriptionBulkData == nullptr && StaticMeshDescription == nullptr;
 }
 
 void FStaticMeshSourceModel::LoadRawMesh(FRawMesh& OutRawMesh) const
@@ -4051,7 +4071,7 @@ void FStaticMeshSourceModel::LoadRawMesh(FRawMesh& OutRawMesh) const
 			{
 				MaterialMap.Add(ConstStaticMeshOwner->GetStaticMaterials()[MaterialIndex].ImportedMaterialSlotName, MaterialIndex);
 			}
-			FStaticMeshOperations::ConvertToRawMesh(*MeshDescription, OutRawMesh, MaterialMap);
+			FStaticMeshOperations::ConvertToRawMesh(*CachedMeshDescription, OutRawMesh, MaterialMap);
 		}
 	}
 	else
@@ -4072,7 +4092,7 @@ void FStaticMeshSourceModel::SaveRawMesh(FRawMesh& InRawMesh, bool /* unused */)
 	//Save both format
 	RawMeshBulkData->SaveRawMesh(InRawMesh);
 
-	CreateMeshDescription();
+	FMeshDescription* MeshDescription = CreateMeshDescription();
 
 	TMap<int32, FName> MaterialMap;
 	check(StaticMeshOwner != nullptr);
@@ -4081,101 +4101,215 @@ void FStaticMeshSourceModel::SaveRawMesh(FRawMesh& InRawMesh, bool /* unused */)
 	FStaticMeshOperations::ConvertFromRawMesh(InRawMesh, *MeshDescription, MaterialMap);
 	
 	// Package up mesh description into bulk data
-	if (!MeshDescriptionBulkData.IsValid())
+	if (StaticMeshDescriptionBulkData == nullptr)
 	{
-		MeshDescriptionBulkData = MakeUnique<FMeshDescriptionBulkData>();
+		StaticMeshDescriptionBulkData = NewObject<UMeshDescriptionBulkDataWrapper>(StaticMeshOwner, NAME_None, RF_Transactional);
 	}
 
-	MeshDescriptionBulkData->SaveMeshDescription(*MeshDescription);
+	StaticMeshDescriptionBulkData->GetBulkData().SaveMeshDescription(*MeshDescription);
+	StaticMeshDescriptionBulkData->AtomicallyClearInternalFlags(EInternalObjectFlags::Async);
+
+	// Converted to MeshDescription - can already lose the RawMeshBulkData
+	RawMeshBulkData->Empty();
 }
 
-FMeshDescription* FStaticMeshSourceModel::CreateMeshDescription()
+
+bool FStaticMeshSourceModel::LoadMeshDescription(FMeshDescription& OutMeshDescription) const
 {
-	if (!MeshDescription.IsValid())
+	TRACE_CPUPROFILER_EVENT_SCOPE(FStaticMeshSourceModel::LoadMeshDescription);
+	check(StaticMeshOwner != nullptr);
+
+	// Ensure default MeshDescription result is empty, with no attributes registered
+	OutMeshDescription = FMeshDescription();
+
+	// If we have valid MeshDescriptionBulkData, unpack it and return it...
+	if (StaticMeshDescriptionBulkData != nullptr)
 	{
-		// If this is the first time a MeshDescription is being created, create it and register its attributes
-		MeshDescription = MakeUnique<FMeshDescription>();
+		// Unpack MeshDescription from the bulk data which was deserialized
+		StaticMeshDescriptionBulkData->GetBulkData().LoadMeshDescription(OutMeshDescription);
+		return true;
 	}
-	else
+
+	// RawMeshBulkData should always be empty now (soon to be deprecated)
+	check(RawMeshBulkData->IsEmpty());
+
+	// This LOD doesn't have a MeshDescriptionBulkData, or a RawMesh, so we presume that it's a generated LOD (and return false)
+	return false;
+}
+
+bool FStaticMeshSourceModel::CloneMeshDescription(FMeshDescription& OutMeshDescription) const
+{
+	TRACE_CPUPROFILER_EVENT_SCOPE(FStaticMeshSourceModel::CloneMeshDescription);
+
+	if (StaticMeshDescription != nullptr)
 	{
-		// Otherwise, empty it completely
-		*MeshDescription = FMeshDescription();
+		OutMeshDescription = StaticMeshDescription->GetMeshDescription();
+		return true;
 	}
 
-	// Register static mesh attributes on it
-	FStaticMeshAttributes StaticMeshAttributes(*MeshDescription);
-	StaticMeshAttributes.Register();
+	return LoadMeshDescription(OutMeshDescription);
+}
 
-	return MeshDescription.Get();
+FMeshDescription* FStaticMeshSourceModel::GetOrCacheMeshDescription()
+{
+	TRACE_CPUPROFILER_EVENT_SCOPE(FStaticMeshSourceModel::GetMeshDescription);
+	check(StaticMeshOwner != nullptr);
+
+	if (StaticMeshDescription == nullptr)
+	{
+		FMeshDescription MeshDescription;
+		if (LoadMeshDescription(MeshDescription))
+		{
+			// This will sometimes be run on a different thread.
+			// However once we have assigned the new UObject to the StaticMeshDescription UPROPERTY, we are free to clear any Async flag which has been assigned
+			// as GC will work as intended.
+			StaticMeshDescription = NewObject<UStaticMeshDescription>(StaticMeshOwner, NAME_None, RF_Transient | RF_Transactional);
+			StaticMeshDescription->SetMeshDescription(MoveTemp(MeshDescription));
+			StaticMeshDescription->AtomicallyClearInternalFlags(EInternalObjectFlags::Async);
+		}
+	}
+
+	if (StaticMeshDescription != nullptr)
+	{
+		return &StaticMeshDescription->GetMeshDescription();
+	}
+	else
+	{
+		return nullptr;
+	}
 }
 
-void FStaticMeshSourceModel::SerializeBulkData(FArchive& Ar, UObject* Owner)
+
+FMeshDescription* FStaticMeshSourceModel::GetCachedMeshDescription() const
 {
-	const bool bIsLoadingLegacyArchive = Ar.IsLoading() && Ar.CustomVer(FEditorObjectVersion::GUID) < FEditorObjectVersion::StaticMeshDeprecatedRawMesh;
+	if (StaticMeshDescription)
+	{
+		return &StaticMeshDescription->GetMeshDescription();
+	}
 
-	// Serialize RawMesh when loading legacy archives
-	if (bIsLoadingLegacyArchive)
+	return nullptr;
+}
+
+
+FMeshDescriptionBulkData* FStaticMeshSourceModel::GetMeshDescriptionBulkData() const
+{
+	if (StaticMeshDescriptionBulkData)
 	{
-		check(RawMeshBulkData != NULL);
-		RawMeshBulkData->Serialize(Ar, Owner);
+		return &StaticMeshDescriptionBulkData->GetBulkData();
 	}
 
-	// Initialize the StaticMeshOwner
-	if (Ar.IsLoading())
+	return nullptr;
+}
+
+
+bool FStaticMeshSourceModel::IsMeshDescriptionValid() const
+{
+	// Determine whether a mesh description is valid without requiring it to be loaded first.
+	// If there is a valid MeshDescriptionBulkData, we know this implies a valid mesh description.
+	// If not, then it's a legacy asset which will have a mesh description built from the RawMeshBulkData, if non-empty.
+	return StaticMeshDescription != nullptr ||
+		StaticMeshDescriptionBulkData != nullptr ||
+		!RawMeshBulkData->IsEmpty();
+}
+
+
+FMeshDescription* FStaticMeshSourceModel::CreateMeshDescription()
+{
+	if (StaticMeshDescription == nullptr)
 	{
-		StaticMeshOwner = Cast<UStaticMesh>(Owner);
+		// Assign the new UObject directly to a UPROPERTY.
+		// If we were run from a different thread, we can immediately clear the Async internal flag, now that GC refs are set up correctly.
+		check(StaticMeshOwner != nullptr);
+		StaticMeshDescription = NewObject<UStaticMeshDescription>(StaticMeshOwner, NAME_None, RF_Transient | RF_Transactional);
+		StaticMeshDescription->AtomicallyClearInternalFlags(EInternalObjectFlags::Async);
 	}
 
-	// Always serialize the MeshDescription bulk data when transacting (so undo/redo works correctly).
-	// Now RawMesh is deprecated, always serialize unless we're loading an old archive.
-	if (Ar.IsTransacting() || !bIsLoadingLegacyArchive)
+	// Reset the mesh description and register its attributes
+	// Do this instead of always creating a new UStaticMeshDescription object, to save memory and improve performance
+	StaticMeshDescription->Reset();
+
+	return &StaticMeshDescription->GetMeshDescription();
+}
+
+
+void FStaticMeshSourceModel::CommitMeshDescription(bool bUseHashAsGuid)
+{
+	TRACE_CPUPROFILER_EVENT_SCOPE(FStaticMeshSourceModel::CommitMeshDescription);
+
+	check(RawMeshBulkData->IsEmpty());
+
+	// This part should remain thread-safe so it can be called from any thread
+	// as long as no more than one thread is calling it for the same UStaticMesh.
+
+	if (StaticMeshDescription != nullptr)
 	{
-		if (Ar.IsSaving() && !MeshDescriptionBulkData.IsValid() && !RawMeshBulkData->IsEmpty())
+		// Package up mesh description into bulk data
+		if (StaticMeshDescriptionBulkData == nullptr)
 		{
-			// If saving a legacy asset which hasn't yet been committed as MeshDescription, perform the conversion now
-			// so it can be loaded successfully as MeshDescription next time.
-			// Note: even if there's a more recent cached MeshDescription, if it hasn't been committed, it will not be saved.
-			FRawMesh RawMesh;
-			LoadRawMesh(RawMesh);
-			SaveRawMesh(RawMesh);
+			check(StaticMeshOwner != nullptr);
+			StaticMeshDescriptionBulkData = NewObject<UMeshDescriptionBulkDataWrapper>(StaticMeshOwner, NAME_None, RF_Transactional);
 		}
 
-		bool bIsValid = MeshDescriptionBulkData.IsValid();
-		Ar << bIsValid;
-
-		if (bIsValid)
+		StaticMeshDescriptionBulkData->GetBulkData().SaveMeshDescription(StaticMeshDescription->GetMeshDescription());
+		if (bUseHashAsGuid)
 		{
-			if (Ar.IsLoading())
-			{
-				MeshDescriptionBulkData = MakeUnique<FMeshDescriptionBulkData>();
-			}
+			StaticMeshDescriptionBulkData->GetBulkData().UseHashAsGuid();
+		}
 
-			MeshDescriptionBulkData->Serialize(Ar, Owner);
+		// Clear the async flag when the object is fully initialized and assigned to a UPROPERTY
+		StaticMeshDescriptionBulkData->AtomicallyClearInternalFlags(EInternalObjectFlags::Async);
+	}
+	else
+	{
+		FGCScopeGuard GCScopeGuard;
+		StaticMeshDescriptionBulkData->MarkPendingKill();
+		StaticMeshDescriptionBulkData = nullptr;
+	}
+}
 
-			// As we are loading a new mesh description bulkdata, if there's a cached existing unpacked MeshDescription,
-			// unpack the new one
-			if (Ar.IsLoading() && MeshDescription.IsValid())
-			{
-				CreateMeshDescription();
-				MeshDescriptionBulkData->LoadMeshDescription(*MeshDescription);
-			}
-		}
+void FStaticMeshSourceModel::ClearMeshDescription()
+{
+	TRACE_CPUPROFILER_EVENT_SCOPE(UStaticMesh::ClearMeshDescription);
+
+	if (StaticMeshDescription != nullptr)
+	{
+		FGCScopeGuard GCScopeGuard;
+		StaticMeshDescription->MarkPendingKill();
+		StaticMeshDescription = nullptr;
+	}
+}
+
+
+void FStaticMeshSourceModel::SerializeBulkData(FArchive& Ar, UObject* Owner)
+{
+	// Initialize the StaticMeshOwner
+	if (Ar.IsLoading())
+	{
+		StaticMeshOwner = Cast<UStaticMesh>(Owner);
+	}
 
-		// For transactions only, serialize the unpacked mesh description here too.
-		// This is so we can preserve any transient attributes which have been set on it when undoing.
-		if (Ar.IsTransacting())
+	if (Ar.IsLoading() && Ar.CustomVer(FEditorObjectVersion::GUID) < FEditorObjectVersion::StaticMeshDeprecatedRawMesh)
+	{
+		// If loading a legacy asset with RawMesh bulk data, serialize it here
+		check(RawMeshBulkData != NULL);
+		RawMeshBulkData->Serialize(Ar, Owner);
+	}
+	else
+	{
+		if (Ar.IsLoading() && Ar.CustomVer(FUE5MainStreamObjectVersion::GUID) < FUE5MainStreamObjectVersion::SerializeMeshDescriptionBase)
 		{
-			bool bIsMeshDescriptionValid = MeshDescription.IsValid();
-			Ar << bIsMeshDescriptionValid;
+			// If there's an inline mesh description bulk data (legacy version), serialize it here, and copy it into the new separate inner object
+			bool bIsValid;
+			Ar << bIsValid;
 
-			if (bIsMeshDescriptionValid)
+			if (bIsValid)
 			{
-				if (Ar.IsLoading())
-				{
-					CreateMeshDescription();
-				}
+				StaticMeshDescriptionBulkData = NewObject<UMeshDescriptionBulkDataWrapper>(StaticMeshOwner, NAME_None, RF_Transactional);
+				StaticMeshDescriptionBulkData->GetBulkData().Serialize(Ar, StaticMeshOwner);
+				StaticMeshDescriptionBulkData->AtomicallyClearInternalFlags(EInternalObjectFlags::Async);
 
-				Ar << (*MeshDescription);
+				// Any cached mesh description is now invalid, so clear it out
+				ClearMeshDescription();
 			}
 		}
 	}
@@ -4349,72 +4483,16 @@ static FStaticMeshRenderData& GetPlatformStaticMeshRenderData(UStaticMesh* Mesh,
 
 bool UStaticMesh::LoadMeshDescription(int32 LodIndex, FMeshDescription& OutMeshDescription) const
 {
-	TRACE_CPUPROFILER_EVENT_SCOPE(UStaticMesh::LoadMeshDescription);
-
-	// Ensure MeshDescription is empty, with no attributes registered
-	OutMeshDescription = FMeshDescription();
-
-	const FStaticMeshSourceModel& SourceModel = GetSourceModel(LodIndex);
-
-	// If we don't have a valid MeshDescription, try and get one...
-	if (SourceModel.MeshDescriptionBulkData.IsValid())
-	{
-		// Unpack MeshDescription from the bulk data which was deserialized
-		SourceModel.MeshDescriptionBulkData->LoadMeshDescription(OutMeshDescription);
-		return true;
-	}
-	
-	// If BulkData isn't valid, this means either:
-	// a) This LOD doesn't have a MeshDescription (because it's been generated), or;
-	// b) This is a legacy asset which still uses RawMesh, in which case we'll look in the DDC for it.
-	FString MeshDataKey;
-	if (GetMeshDataKey(LodIndex, MeshDataKey))
-	{
-		COOK_STAT(auto Timer = StaticMeshCookStats::UsageStats.TimeSyncWork());
-		TArray<uint8> DerivedData;
-		if (GetDerivedDataCacheRef().GetSynchronous(*MeshDataKey, DerivedData, GetPathName()))
-		{
-			COOK_STAT(Timer.AddHit(DerivedData.Num()));
-			// If there was valid DDC data, we assume this is because the asset is an old one with valid RawMeshBulkData
-			check(!SourceModel.RawMeshBulkData->IsEmpty());
-
-			// Load from the DDC
-			const bool bIsPersistent = true;
-			FMemoryReader Ar(DerivedData, bIsPersistent);
-
-			// Create a bulk data object which will be immediately thrown away (as it is not in an archive)
-			FMeshDescriptionBulkData MeshDescriptionBulkData;
-			MeshDescriptionBulkData.Serialize(Ar, const_cast<UStaticMesh*>(this));
-
-			// Unpack MeshDescription from the bulk data
-			MeshDescriptionBulkData.LoadMeshDescription(OutMeshDescription);
-			return true;
-		}
-		else
-		{
-			COOK_STAT(Timer.TrackCyclesOnly());
-		}
-	}
-
-	// If after all this we *still* don't have a valid MeshDescription, but there's a valid RawMesh, convert that to a MeshDescription.
-	if (!SourceModel.RawMeshBulkData->IsEmpty())
+	if (!IsSourceModelValid(LodIndex))
 	{
-		FRawMesh LodRawMesh;
-		SourceModel.LoadRawMesh(LodRawMesh);
-		TMap<int32, FName> MaterialMap;
-		FillMaterialName(GetStaticMaterials(), MaterialMap);
-
-		// Register static mesh attributes on the mesh description
-		FStaticMeshAttributes StaticMeshAttributes(OutMeshDescription);
-		StaticMeshAttributes.Register();
-
-		FStaticMeshOperations::ConvertFromRawMesh(LodRawMesh, OutMeshDescription, MaterialMap);
-		return true;
+		return false;
 	}
 
-	return false;
+	const FStaticMeshSourceModel& SourceModel = GetSourceModel(LodIndex);
+	return SourceModel.LoadMeshDescription(OutMeshDescription);
 }
 
+
 bool UStaticMesh::CloneMeshDescription(int32 LodIndex, FMeshDescription& OutMeshDescription) const
 {
 	if (!IsSourceModelValid(LodIndex))
@@ -4422,19 +4500,11 @@ bool UStaticMesh::CloneMeshDescription(int32 LodIndex, FMeshDescription& OutMesh
 		return false;
 	}
 
-	TRACE_CPUPROFILER_EVENT_SCOPE(UStaticMesh::CloneMeshDescription);
-	
 	const FStaticMeshSourceModel& SourceModel = GetSourceModel(LodIndex);
-
-	if (SourceModel.MeshDescription.IsValid())
-	{
-		OutMeshDescription = *SourceModel.MeshDescription.Get();
-		return true;
-	}
-
-	return LoadMeshDescription(LodIndex, OutMeshDescription);
+	return SourceModel.CloneMeshDescription(OutMeshDescription);
 }
 
+
 FMeshDescription* UStaticMesh::GetMeshDescription(int32 LodIndex) const
 {
 	if (!IsSourceModelValid(LodIndex))
@@ -4442,24 +4512,12 @@ FMeshDescription* UStaticMesh::GetMeshDescription(int32 LodIndex) const
 		return nullptr;
 	}
 
-	TRACE_CPUPROFILER_EVENT_SCOPE(UStaticMesh::GetMeshDescription);
-
 	// Require a const_cast here, because GetMeshDescription should ostensibly have const semantics,
 	// but the lazy initialization (from the BulkData or the DDC) is a one-off event which breaks constness.
 	UStaticMesh* MutableThis = const_cast<UStaticMesh*>(this);
-
 	FStaticMeshSourceModel& SourceModel = MutableThis->GetSourceModel(LodIndex);
 
-	if (!SourceModel.MeshDescription.IsValid())
-	{
-		FMeshDescription MeshDescription;
-		if (LoadMeshDescription(LodIndex, MeshDescription))
-		{
-			SourceModel.MeshDescription = MakeUnique<FMeshDescription>(MoveTemp(MeshDescription));
-		}
-	}
-
-	return SourceModel.MeshDescription.Get();
+	return SourceModel.GetOrCacheMeshDescription();
 }
 
 
@@ -4471,13 +4529,7 @@ bool UStaticMesh::IsMeshDescriptionValid(int32 LodIndex) const
 	}
 
 	const FStaticMeshSourceModel& SourceModel = GetSourceModel(LodIndex);
-
-	// Determine whether a mesh description is valid without requiring it to be loaded first.
-	// If there is a valid MeshDescriptionBulkData, we know this implies a valid mesh description.
-	// If not, then it's a legacy asset which will have a mesh description built from the RawMeshBulkData, if non-empty.
-	return SourceModel.MeshDescription.IsValid() ||
-		   SourceModel.MeshDescriptionBulkData.IsValid() ||
-		   !SourceModel.RawMeshBulkData->IsEmpty();
+	return SourceModel.IsMeshDescriptionValid();
 }
 
 
@@ -4506,43 +4558,18 @@ FMeshDescription* UStaticMesh::CreateMeshDescription(int32 LodIndex, FMeshDescri
 
 void UStaticMesh::CommitMeshDescription(int32 LodIndex, const FCommitMeshDescriptionParams& Params)
 {
-	TRACE_CPUPROFILER_EVENT_SCOPE(UStaticMesh::CommitMeshDescription);
-
-	// This part should remain thread-safe so it can be called from any thread
-	// as long as no more than one thread is calling it for the same UStaticMesh.
-
 	// The source model must be created before calling this function
 	check(IsSourceModelValid(LodIndex));
 
 	FStaticMeshSourceModel& SourceModel = GetSourceModel(LodIndex);
-	if (SourceModel.MeshDescription.IsValid())
-	{
-		if (LodIndex == 0)
-		{
-			CachedMeshDescriptionBounds = SourceModel.MeshDescription->GetBounds();
-		}
+	SourceModel.CommitMeshDescription(Params.bUseHashAsGuid);
 
-		// Package up mesh description into bulk data
-		if (!SourceModel.MeshDescriptionBulkData.IsValid())
-		{
-			SourceModel.MeshDescriptionBulkData = MakeUnique<FMeshDescriptionBulkData>();
-		}
-
-		SourceModel.MeshDescriptionBulkData->SaveMeshDescription(*SourceModel.MeshDescription);
-		if (Params.bUseHashAsGuid)
-		{
-			SourceModel.MeshDescriptionBulkData->UseHashAsGuid();
-		}
-	}
-	else
+	// For LOD0, cache the bounds of the mesh description in the static mesh for quick access
+	if (LodIndex == 0 && SourceModel.GetCachedMeshDescription())
 	{
-		SourceModel.MeshDescriptionBulkData.Reset();
+		CachedMeshDescriptionBounds = SourceModel.GetCachedMeshDescription()->GetBounds();
 	}
 
-	// Clear RawMeshBulkData and mark as invalid.
-	// If any legacy tool needs the RawMesh at this point, it will do a conversion from MD at that moment.
-	SourceModel.RawMeshBulkData->Empty();
-
 	// This part is not thread-safe, so we give the caller the option of calling it manually from the mainthread
 	if (Params.bMarkPackageDirty)
 	{
@@ -4557,7 +4584,7 @@ void UStaticMesh::ClearMeshDescription(int32 LodIndex)
 		TRACE_CPUPROFILER_EVENT_SCOPE(UStaticMesh::ClearMeshDescription);
 
 		FStaticMeshSourceModel& SourceModel = GetSourceModel(LodIndex);
-		SourceModel.MeshDescription.Reset();
+		SourceModel.ClearMeshDescription();
 	}
 }
 
@@ -4570,41 +4597,126 @@ void UStaticMesh::ClearMeshDescriptions()
 	}
 }
 
-// If static mesh derived data needs to be rebuilt (new format, serialization
-// differences, etc.) replace the version GUID below with a new one.
-// In case of merge conflicts with DDC versions, you MUST generate a new GUID
-// and set this new GUID as the version.
-#if UE_USE_VIRTUALBULKDATA
-	#define MESHDATAKEY_STATICMESH_DERIVEDDATA_VER TEXT("1F22FDD3B8C54996AAD54400037B85DF")
-#else
-#define MESHDATAKEY_STATICMESH_DERIVEDDATA_VER TEXT("BEDE1E3CC35D4D8FA944458825A18AFB")
-#endif
 
-static const FString& GetMeshDataKeyStaticMeshDerivedDataVersion()
+bool UStaticMesh::LoadHiResMeshDescription(FMeshDescription& OutMeshDescription) const
 {
-	static FString CachedVersionString(MESHDATAKEY_STATICMESH_DERIVEDDATA_VER);
-	return CachedVersionString;
+	const FStaticMeshSourceModel& SourceModel = GetHiResSourceModel();
+	return SourceModel.LoadMeshDescription(OutMeshDescription);
 }
 
-bool UStaticMesh::GetMeshDataKey(int32 LodIndex, FString& OutKey) const
+
+bool UStaticMesh::CloneHiResMeshDescription(FMeshDescription& OutMeshDescription) const
 {
-	OutKey.Empty();
-	if (LodIndex >= GetNumSourceModels())
+	const FStaticMeshSourceModel& SourceModel = GetHiResSourceModel();
+	return SourceModel.CloneMeshDescription(OutMeshDescription);
+}
+
+
+FMeshDescription* UStaticMesh::GetHiResMeshDescription() const
+{
+	// Require a const_cast here, because GetMeshDescription should ostensibly have const semantics,
+	// but the lazy initialization (from the BulkData or the DDC) is a one-off event which breaks constness.
+	UStaticMesh* MutableThis = const_cast<UStaticMesh*>(this);
+	FStaticMeshSourceModel& SourceModel = MutableThis->GetHiResSourceModel();
+
+	return SourceModel.GetOrCacheMeshDescription();
+}
+
+
+bool UStaticMesh::IsHiResMeshDescriptionValid() const
+{
+	const FStaticMeshSourceModel& SourceModel = GetHiResSourceModel();
+	return SourceModel.IsMeshDescriptionValid();
+}
+
+
+FMeshDescription* UStaticMesh::CreateHiResMeshDescription()
+{
+	FStaticMeshSourceModel& SourceModel = GetHiResSourceModel();
+	return SourceModel.CreateMeshDescription();
+}
+
+
+FMeshDescription* UStaticMesh::CreateHiResMeshDescription(FMeshDescription InMeshDescription)
+{
+	FMeshDescription* NewMeshDescription = CreateHiResMeshDescription();
+	if (NewMeshDescription != nullptr)
 	{
-		return false;
+		*NewMeshDescription = MoveTemp(InMeshDescription);
 	}
 
-	FSHA1 Sha;
-	FString LodIndexString = FString::Printf(TEXT("%d_"), LodIndex);
-	const FStaticMeshSourceModel& SourceModel = GetSourceModel(LodIndex);
-	if(!SourceModel.RawMeshBulkData->IsEmpty())
+	return NewMeshDescription;
+}
+
+
+void UStaticMesh::CommitHiResMeshDescription(const FCommitMeshDescriptionParams& Params)
+{
+	FStaticMeshSourceModel& SourceModel = GetHiResSourceModel();
+	SourceModel.CommitMeshDescription(Params.bUseHashAsGuid);
+
+	// This part is not thread-safe, so we give the caller the option of calling it manually from the mainthread
+	if (Params.bMarkPackageDirty)
 	{
-		LodIndexString += SourceModel.RawMeshBulkData->GetIdString();
+		MarkPackageDirty();
 	}
-	else
+}
+
+
+void UStaticMesh::ClearHiResMeshDescription()
+{
+	FStaticMeshSourceModel& SourceModel = GetHiResSourceModel();
+	SourceModel.ClearMeshDescription();
+}
+
+
+void UStaticMesh::ModifyMeshDescription(int32 LodIndex, bool bAlwaysMarkDirty)
+{
+	UStaticMeshDescription* StaticMeshDescription = GetSourceModel(LodIndex).StaticMeshDescription;
+	if (StaticMeshDescription)
+	{
+		StaticMeshDescription->Modify(bAlwaysMarkDirty);
+	}
+}
+
+
+void UStaticMesh::ModifyAllMeshDescriptions(bool bAlwaysMarkDirty)
+{
+	for (int LODIndex = 0; LODIndex < GetNumSourceModels(); LODIndex++)
+	{
+		ModifyMeshDescription(LODIndex);
+	}
+}
+
+
+void UStaticMesh::ModifyHiResMeshDescription(bool bAlwaysMarkDirty)
+{
+	UStaticMeshDescription* StaticMeshDescription = GetHiResSourceModel().StaticMeshDescription;
+	if (StaticMeshDescription)
 	{
-		LodIndexString += TEXT("REDUCELOD");
+		StaticMeshDescription->Modify(bAlwaysMarkDirty);
 	}
+}
+
+
+// This is the key which is used for DDC data when legacy RawMesh is converted to MeshDescription.
+// If static mesh derived data needs to be rebuilt (new format, serialization differences, etc.) replace the version GUID below with a new one.
+// In case of merge conflicts with DDC versions, you MUST generate a new GUID and set this new GUID as the version.
+#if UE_USE_VIRTUALBULKDATA
+#define MESHDATAKEY_STATICMESH_DERIVEDDATA_VER TEXT("7F1A02EF80374F0F86F9A30BAFF3F850")
+#else
+#define MESHDATAKEY_STATICMESH_DERIVEDDATA_VER TEXT("E09E8DD68C864D82B511A97080B04837")
+#endif
+
+static bool GetMeshDataKey(FString& OutKey, const FRawMeshBulkData* RawMeshBulkData, int LodIndex)
+{
+	OutKey.Empty();
+
+	FSHA1 Sha;
+
+	check(!RawMeshBulkData->IsEmpty());
+	FString LodIndexString = FString::Printf(TEXT("%d_"), LodIndex);
+	LodIndexString += RawMeshBulkData->GetIdString();
+
 	const TArray<TCHAR>& LodIndexArray = LodIndexString.GetCharArray();
 	Sha.Update((uint8*)LodIndexArray.GetData(), LodIndexArray.Num() * LodIndexArray.GetTypeSize());
 	Sha.Final();
@@ -4617,7 +4729,7 @@ bool UStaticMesh::GetMeshDataKey(int32 LodIndex, FString& OutKey) const
 
 	OutKey = FDerivedDataCacheInterface::BuildCacheKey(
 		TEXT("MESHDATAKEY_STATICMESH"),
-		*GetMeshDataKeyStaticMeshDerivedDataVersion(),
+		MESHDATAKEY_STATICMESH_DERIVEDDATA_VER,
 		*MeshLodData
 	);
 	return true;
@@ -4632,23 +4744,35 @@ void UStaticMesh::CacheMeshData()
 	for (int32 LodIndex = 0; LodIndex < GetNumSourceModels(); ++LodIndex)
 	{
 		FStaticMeshSourceModel& SourceModel = GetSourceModel(LodIndex);
-		if (!SourceModel.MeshDescriptionBulkData.IsValid())
+		if (SourceModel.GetMeshDescriptionBulkData() == nullptr)
 		{
 			// Legacy assets used to store their source data in the RawMeshBulkData
 			// Migrate it to the new description if present
 			if (!SourceModel.RawMeshBulkData->IsEmpty())
 			{
 				FString MeshDataKey;
-				if (GetMeshDataKey(LodIndex, MeshDataKey))
+				if (GetMeshDataKey(MeshDataKey, SourceModel.RawMeshBulkData, LodIndex))
 				{
 					COOK_STAT(auto Timer = StaticMeshCookStats::UsageStats.TimeSyncWork());
-					// If the DDC key doesn't exist, convert the data and save it to DDC
-					if (GetDerivedDataCacheRef().CachedDataProbablyExists(*MeshDataKey))
+
+					TArray<uint8> DerivedData;
+
+					if (GetDerivedDataCacheRef().GetSynchronous(*MeshDataKey, DerivedData, GetPathName()))
 					{
-						COOK_STAT(Timer.TrackCyclesOnly());
+						COOK_STAT(Timer.AddHit(DerivedData.Num()));
+
+						SourceModel.StaticMeshDescriptionBulkData = NewObject<UMeshDescriptionBulkDataWrapper>(this, NAME_None, RF_Transactional);
+
+						// Load from the DDC
+						const bool bIsPersistent = true;
+						FMemoryReader Ar(DerivedData, bIsPersistent);
+						SourceModel.StaticMeshDescriptionBulkData->GetBulkData().Serialize(Ar, const_cast<UStaticMesh*>(this));
+
+						check(SourceModel.GetCachedMeshDescription() == nullptr);
 					}
 					else
 					{
+						// If the DDC key doesn't exist, convert the data and save it to DDC
 						// Get the RawMesh for this LOD
 						FRawMesh TempRawMesh;
 						SourceModel.RawMeshBulkData->LoadRawMesh(TempRawMesh);
@@ -4661,21 +4785,23 @@ void UStaticMesh::CacheMeshData()
 						FillMaterialName(GetStaticMaterials(), MaterialMap);
 						FStaticMeshOperations::ConvertFromRawMesh(TempRawMesh, *MeshDescription, MaterialMap);
 
-						// Pack MeshDescription into temporary bulk data, ready to write out to DDC.
-						// This will be reloaded from the DDC when needed if a MeshDescription is requested from the static mesh.
-						FMeshDescriptionBulkData MeshDescriptionBulkData;
-						MeshDescriptionBulkData.SaveMeshDescription(*MeshDescription);
+						// Pack MeshDescription into bulk data
+						SourceModel.StaticMeshDescriptionBulkData = NewObject<UMeshDescriptionBulkDataWrapper>(this, NAME_None, RF_Transactional);
+						SourceModel.StaticMeshDescriptionBulkData->GetBulkData().SaveMeshDescription(*MeshDescription);
+						SourceModel.StaticMeshDescriptionBulkData->AtomicallyClearInternalFlags(EInternalObjectFlags::Async);
 
 						// Write the DDC cache
-						TArray<uint8> DerivedData;
 						const bool bIsPersistent = true;
 						FMemoryWriter Ar(DerivedData, bIsPersistent);
-						MeshDescriptionBulkData.Serialize(Ar, this);
+						SourceModel.StaticMeshDescriptionBulkData->GetBulkData().Serialize(Ar, this);
 						GetDerivedDataCacheRef().Put(*MeshDataKey, DerivedData, GetPathName());
 						COOK_STAT(Timer.AddMiss(DerivedData.Num()));
 					}
 				}
 			}
+
+			// We now have a MeshDescription instead of a RawMesh, so get rid of the RawMesh completely
+			SourceModel.RawMeshBulkData->Empty();
 		}
 	}
 }
@@ -4975,6 +5101,7 @@ void UStaticMesh::Serialize(FArchive& Ar)
 	Ar.UsingCustomVersion(FEditorObjectVersion::GUID);
 	Ar.UsingCustomVersion(FRenderingObjectVersion::GUID);
 	Ar.UsingCustomVersion(FReleaseObjectVersion::GUID);
+	Ar.UsingCustomVersion(FUE5MainStreamObjectVersion::GUID);
 
 	FStripDataFlags StripFlags( Ar );
 
@@ -5338,7 +5465,6 @@ void UStaticMesh::BeginPostLoadInternal(FStaticMeshPostLoadContext& Context)
 			DistanceFieldReplacementMesh->ConditionalPostLoad();
 		}
 
-		//TODO remove this code when FRawMesh will be removed
 		//Fill the static mesh owner
 		int32 NumLODs = GetNumSourceModels();
 		for (int32 LODIndex = 0; LODIndex < NumLODs; ++LODIndex)
@@ -5346,6 +5472,8 @@ void UStaticMesh::BeginPostLoadInternal(FStaticMeshPostLoadContext& Context)
 			FStaticMeshSourceModel& SrcModel = GetSourceModel(LODIndex);
 			SrcModel.StaticMeshOwner = this;
 		}
+
+		GetHiResSourceModel().StaticMeshOwner = this;
 	}
 
 	Context.bIsCookedForEditor = GetOutermost()->bIsCookedForEditor;
@@ -5931,9 +6059,7 @@ UStaticMeshDescription* UStaticMesh::GetStaticMeshDescription(int32 LODIndex)
 #if WITH_EDITOR
 	if (LODIndex < GetNumSourceModels())
 	{
-		UStaticMeshDescription* StaticMeshDescription = NewObject<UStaticMeshDescription>(this, NAME_None, RF_Transient);
-		StaticMeshDescription->SetMeshDescriptionRef(GetMeshDescription(LODIndex));
-		return StaticMeshDescription;
+		return GetSourceModel(LODIndex).StaticMeshDescription;
 	}
 #endif
 	return nullptr;
@@ -7077,6 +7203,12 @@ void UStaticMesh::SetMaterial(int32 MaterialIndex, UMaterialInterface* NewMateri
 {
 	if (GetStaticMaterials().IsValidIndex(MaterialIndex))
 	{
+		// Ensure mesh descriptions are loaded before starting the transaction
+		for (int32 LODIndex = 0; LODIndex < GetNumSourceModels(); ++LODIndex)
+		{
+			GetMeshDescription(LODIndex);
+		}
+
 		FScopedTransaction ScopeTransaction(LOCTEXT("StaticMeshMaterialChanged", "StaticMesh: Material changed"));
 
 		// flag the property (Materials) we're modifying so that not all of the object is rebuilt.
diff --git a/Engine/Source/Runtime/Engine/Private/StaticMeshBuild.cpp b/Engine/Source/Runtime/Engine/Private/StaticMeshBuild.cpp
--- a/Engine/Source/Runtime/Engine/Private/StaticMeshBuild.cpp
+++ b/Engine/Source/Runtime/Engine/Private/StaticMeshBuild.cpp
@@ -24,7 +24,7 @@
 #include "Interfaces/ITargetPlatformManagerModule.h"
 #include "StaticMeshCompiler.h"
 #include "MeshUtilities.h"
-#include "MeshDescription.h"
+#include "StaticMeshDescription.h"
 #include "MeshUtilitiesCommon.h"
 #include "Misc/FeedbackContext.h"
 #include "Misc/ScopedSlowTask.h"
@@ -400,9 +400,9 @@ void UStaticMesh::BeginBuildInternal(FStaticMeshBuildContext* Context)
 		// Extended bounds are super important to avoid many stalls since a lot of different
 		// systems rely on mesh bounds. If mesh description is available and doesn't require being
 		// loaded, compute bounds right now so we can unlock the property before going async.
-		if (IsSourceModelValid(0) && GetSourceModel(0).MeshDescription.IsValid())
+		if (IsSourceModelValid(0) && GetSourceModel(0).StaticMeshDescription != nullptr)
 		{
-			CachedMeshDescriptionBounds = GetSourceModel(0).MeshDescription.Get()->GetBounds();
+			CachedMeshDescriptionBounds = GetSourceModel(0).StaticMeshDescription->GetMeshDescription().GetBounds();
 		}
 
 		// CommitMeshDescription will also fill out CachedMeshDescriptionBounds
diff --git a/Engine/Source/Runtime/MeshDescription/Private/MeshDescription.cpp b/Engine/Source/Runtime/MeshDescription/Private/MeshDescription.cpp
--- a/Engine/Source/Runtime/MeshDescription/Private/MeshDescription.cpp
+++ b/Engine/Source/Runtime/MeshDescription/Private/MeshDescription.cpp
@@ -4,6 +4,7 @@
 #include "MeshAttributes.h"
 #include "Algo/Copy.h"
 #include "Misc/SecureHash.h"
+#include "Serialization/NameAsStringProxyArchive.h"
 #include "Serialization/BulkDataReader.h"
 #include "Serialization/BulkDataWriter.h"
 #include "UObject/EnterpriseObjectVersion.h"
@@ -158,8 +159,10 @@ void FMeshDescription::Cache()
 }
 
 
-void FMeshDescription::Serialize(FArchive& Ar)
+void FMeshDescription::Serialize(FArchive& BaseAr)
 {
+	FNameAsStringProxyArchive Ar(BaseAr);
+
 	Ar.UsingCustomVersion(FReleaseObjectVersion::GUID);
 	Ar.UsingCustomVersion(FEditorObjectVersion::GUID);
 	Ar.UsingCustomVersion(FUE5MainStreamObjectVersion::GUID);
@@ -2168,3 +2171,16 @@ void FMeshDescriptionBulkData::UseHashAsGuid()
 }
 
 #endif // #if WITH_EDITORONLY_DATA
+
+
+void UMeshDescriptionBulkDataWrapper::Serialize(FArchive& Ar)
+{
+	Super::Serialize(Ar);
+
+#if WITH_EDITORONLY_DATA
+	if (!Ar.IsCooking())
+	{
+		BulkData.Serialize(Ar, GetOuter());
+	}
+#endif
+}
diff --git a/Engine/Source/Runtime/MeshDescription/Private/MeshDescriptionBase.cpp b/Engine/Source/Runtime/MeshDescription/Private/MeshDescriptionBase.cpp
--- a/Engine/Source/Runtime/MeshDescription/Private/MeshDescriptionBase.cpp
+++ b/Engine/Source/Runtime/MeshDescription/Private/MeshDescriptionBase.cpp
@@ -1,13 +1,21 @@
 // Copyright Epic Games, Inc. All Rights Reserved.
 
 #include "MeshDescriptionBase.h"
+#include "UObject/UE5MainStreamObjectVersion.h"
 #include "Algo/Copy.h"
 
 
 
 void UMeshDescriptionBase::Serialize(FArchive& Ar)
 {
+	Ar.UsingCustomVersion(FUE5MainStreamObjectVersion::GUID);
+
 	Super::Serialize(Ar);
+
+	if (!Ar.IsLoading() || Ar.CustomVer(FUE5MainStreamObjectVersion::GUID) >= FUE5MainStreamObjectVersion::SerializeMeshDescriptionBase)
+	{
+		Ar << OwnedMeshDescription;
+	}
 }
 
 void UMeshDescriptionBase::RegisterAttributes()
@@ -18,7 +26,7 @@ void UMeshDescriptionBase::RegisterAttributes()
 
 void UMeshDescriptionBase::Reset()
 {
-	SetMeshDescriptionRef(nullptr);
+	OwnedMeshDescription = FMeshDescription();
 	RegisterAttributes();
 }
 
diff --git a/Engine/Source/Runtime/MeshDescription/Public/MeshDescription.h b/Engine/Source/Runtime/MeshDescription/Public/MeshDescription.h
--- a/Engine/Source/Runtime/MeshDescription/Public/MeshDescription.h
+++ b/Engine/Source/Runtime/MeshDescription/Public/MeshDescription.h
@@ -66,7 +66,6 @@ enum class EComputeNTBsOptions : uint32
 ENUM_CLASS_FLAGS(EComputeNTBsOptions);
 
 
-//USTRUCT()
 struct MESHDESCRIPTION_API FMeshDescription
 {
 public:
@@ -1339,5 +1338,27 @@ private:
 	/** Uses hash instead of guid to identify content to improve DDC cache hit. */
 	bool bGuidIsHash;
 };
+#endif // WITH_EDITORONLY_DATA
+
+
+/**
+ * UObject wrapper for FMeshDescriptionBulkData
+ */
+UCLASS()
+class MESHDESCRIPTION_API UMeshDescriptionBulkDataWrapper : public UObject
+{
+	GENERATED_BODY()
+
+public:
+	virtual void Serialize(FArchive& Ar) override;
+
+#if WITH_EDITORONLY_DATA
+
+	const FMeshDescriptionBulkData& GetBulkData() const { return BulkData; }
+	FMeshDescriptionBulkData& GetBulkData() { return BulkData; }
+
+private:
+	FMeshDescriptionBulkData BulkData;
 
 #endif //WITH_EDITORONLY_DATA
+};
diff --git a/Engine/Source/Runtime/MeshDescription/Public/MeshDescriptionBase.h b/Engine/Source/Runtime/MeshDescription/Public/MeshDescriptionBase.h
--- a/Engine/Source/Runtime/MeshDescription/Public/MeshDescriptionBase.h
+++ b/Engine/Source/Runtime/MeshDescription/Public/MeshDescriptionBase.h
@@ -21,40 +21,18 @@ public:
 	/** Get a reference to the actual mesh description */
 	FMeshDescription& GetMeshDescription()
 	{
-		if (MeshDescriptionPtr)
-		{
-			return *MeshDescriptionPtr;
-		}
-		else
-		{
-			return OwnedMeshDescription;
-		}
+		return OwnedMeshDescription;
 	}
 
 	const FMeshDescription& GetMeshDescription() const
 	{
-		if (MeshDescriptionPtr)
-		{
-			return *MeshDescriptionPtr;
-		}
-		else
-		{
-			return OwnedMeshDescription;
-		}
+		return OwnedMeshDescription;
 	}
 
 	/** Set the mesh description */
 	void SetMeshDescription(FMeshDescription InMeshDescription)
 	{ 
 		OwnedMeshDescription = MoveTemp(InMeshDescription);
-		MeshDescriptionPtr = nullptr;
-	}
-
-	/** Sets a reference to another mesh description */
-	void SetMeshDescriptionRef(FMeshDescription* InMeshDescriptionRef)
-	{
-		MeshDescriptionPtr = InMeshDescriptionRef;
-		OwnedMeshDescription.Empty();
 	}
 
 	// UMeshDescriptionBase interface
@@ -505,6 +483,5 @@ public:
 
 protected:
 	FMeshDescription OwnedMeshDescription;
-	FMeshDescription* MeshDescriptionPtr = nullptr;
 	TUniquePtr<FMeshAttributes> RequiredAttributes;
 };
diff --git a/Engine/Source/Runtime/RawMesh/Public/RawMesh.h b/Engine/Source/Runtime/RawMesh/Public/RawMesh.h
--- a/Engine/Source/Runtime/RawMesh/Public/RawMesh.h
+++ b/Engine/Source/Runtime/RawMesh/Public/RawMesh.h
@@ -129,6 +129,7 @@ public:
 
 	/** Uses a hash as the GUID, useful to prevent creating new GUIDs on load for legacy assets. */
 	RAWMESH_API void UseHashAsGuid(class UObject* Owner);
+	RAWMESH_API bool UsesHashAsGuid() const { return bGuidIsHash; }
 
 	RAWMESH_API const FByteBulkData& GetBulkData() const;
 
