Added more HLOD settings to Data Layers
* Generates HLODs
* Default HLOD Layer

diff --git a/Engine/Source/Runtime/Engine/Private/WorldPartition/DataLayer/DataLayer.cpp b/Engine/Source/Runtime/Engine/Private/WorldPartition/DataLayer/DataLayer.cpp
--- a/Engine/Source/Runtime/Engine/Private/WorldPartition/DataLayer/DataLayer.cpp
+++ b/Engine/Source/Runtime/Engine/Private/WorldPartition/DataLayer/DataLayer.cpp
@@ -15,6 +15,8 @@ UDataLayer::UDataLayer(const FObjectInitializer& ObjectInitializer)
 , bIsVisible(true)
 , bIsDynamicallyLoaded(false)
 , bIsDynamicallyLoadedInEditor(true)
+, bGeneratesHLODs(false)
+, DefaultHLODLayer()
 #endif
 {
 }
diff --git a/Engine/Source/Runtime/Engine/Private/WorldPartition/RuntimeSpatialHash/RuntimeSpatialHashHLOD.cpp b/Engine/Source/Runtime/Engine/Private/WorldPartition/RuntimeSpatialHash/RuntimeSpatialHashHLOD.cpp
--- a/Engine/Source/Runtime/Engine/Private/WorldPartition/RuntimeSpatialHash/RuntimeSpatialHashHLOD.cpp
+++ b/Engine/Source/Runtime/Engine/Private/WorldPartition/RuntimeSpatialHash/RuntimeSpatialHashHLOD.cpp
@@ -15,6 +15,7 @@
 #include "Misc/PackageName.h"
 #include "Misc/ScopedSlowTask.h"
 
+#include "Algo/AnyOf.h"
 #include "Algo/ForEach.h"
 
 #include "EngineUtils.h"
@@ -116,18 +117,28 @@ static TArray<FGuid> GenerateHLODsForGrid(UWorldPartition* WorldPartition, const
 	const FBox WorldBounds = WorldPartition->GetWorldBounds();
 
 	const FSquare2DGridHelper PartitionedActors = GetPartitionedActors(WorldPartition, WorldBounds, RuntimeGrid, ActorClusters);
+	const FSquare2DGridHelper::FGridLevel::FGridCell& AlwaysLoadedCell = PartitionedActors.GetAlwaysLoadedCell();
+
+	auto ShouldGenerateHLODs = [&AlwaysLoadedCell](const FSquare2DGridHelper::FGridLevel::FGridCell& GridCell, const FSquare2DGridHelper::FGridLevel::FGridCellDataChunk& GridCellDataChunk)
+	{
+		const bool bIsCellAlwaysLoaded = &GridCell == &AlwaysLoadedCell;
+		const bool bShouldGenerateHLODForDataLayers = GridCellDataChunk.GetDataLayers().IsEmpty() ? true : Algo::AnyOf(GridCellDataChunk.GetDataLayers(), [](const UDataLayer* DataLayer) { return DataLayer->ShouldGenerateHLODs(); });
+		const bool bChunkHasActors = !GridCellDataChunk.GetActors().IsEmpty();
+
+		const bool bShouldGenerateHLODs = !bIsCellAlwaysLoaded && bShouldGenerateHLODForDataLayers && bChunkHasActors;
+		return bShouldGenerateHLODs;
+	};
 
 	// Quick pass to compute the number of cells we'll have to process, to provide a meaningful progress display
 	int32 NbCellsToProcess = 0;
-	PartitionedActors.ForEachCells([&PartitionedActors, &NbCellsToProcess](const FIntVector& CellCoord)
+	PartitionedActors.ForEachCells([&](const FIntVector& CellCoord)
 	{
 		const FSquare2DGridHelper::FGridLevel::FGridCell& GridCell = PartitionedActors.GetCell(CellCoord);
 
 		for (const FSquare2DGridHelper::FGridLevel::FGridCellDataChunk& GridCellDataChunk : GridCell.GetDataChunks())
 		{
-			const bool bIsCellAlwaysLoaded = &GridCell == &PartitionedActors.GetAlwaysLoadedCell();
-
-			if (!bIsCellAlwaysLoaded && GridCellDataChunk.GetActors().Num() != 0)
+			const bool bShouldGenerateHLODs = ShouldGenerateHLODs(GridCell, GridCellDataChunk);
+			if (bShouldGenerateHLODs)
 			{
 				NbCellsToProcess++;
 			}
@@ -141,15 +152,14 @@ static TArray<FGuid> GenerateHLODsForGrid(UWorldPartition* WorldPartition, const
 	FAssetRegistryModule& AssetRegistryModule = FModuleManager::LoadModuleChecked<FAssetRegistryModule>(TEXT("AssetRegistry"));
 
 	TArray<FGuid> GridHLODActors;
-	PartitionedActors.ForEachCells([&PartitionedActors, &GridHLODActors, &SlowTask, RuntimeGrid, HLODLevel, WorldPartition, WorldBounds, &Context, SourceControlHelper, &DirectoryWatcherModule, &AssetRegistryModule, &HasExceededMaxMemory, &DoCollectGarbage](const FIntVector& CellCoord)
+	PartitionedActors.ForEachCells([&](const FIntVector& CellCoord)
 	{
 		const FSquare2DGridHelper::FGridLevel::FGridCell& GridCell = PartitionedActors.GetCell(CellCoord);
 
 		for (const FSquare2DGridHelper::FGridLevel::FGridCellDataChunk& GridCellDataChunk : GridCell.GetDataChunks())
 		{
-			const bool bIsCellAlwaysLoaded = &GridCell == &PartitionedActors.GetAlwaysLoadedCell();
-
-			if (!bIsCellAlwaysLoaded && GridCellDataChunk.GetActors().Num() != 0)
+			const bool bShouldGenerateHLODs = ShouldGenerateHLODs(GridCell, GridCellDataChunk);
+			if (bShouldGenerateHLODs)
 			{
 				SlowTask.EnterProgressFrame(1);
 
@@ -180,10 +190,8 @@ static TArray<FGuid> GenerateHLODsForGrid(UWorldPartition* WorldPartition, const
 					DoCollectGarbage();
 				}
 
-				TArray<AWorldPartitionHLOD*> CellHLODActors;
-				CellHLODActors = UHLODLayer::GenerateHLODForCell(WorldPartition, &Context, CellName, CellBounds, HLODLevel, CellActors);
-
-				if (CellHLODActors.Num())
+				TArray<AWorldPartitionHLOD*> CellHLODActors = UHLODLayer::GenerateHLODForCell(WorldPartition, &Context, CellName, CellBounds, HLODLevel, CellActors);
+				if (!CellHLODActors.IsEmpty())
 				{
 					TArray<AWorldPartitionHLOD*> NewCellHLODActors;
 
diff --git a/Engine/Source/Runtime/Engine/Public/WorldPartition/DataLayer/DataLayer.h b/Engine/Source/Runtime/Engine/Public/WorldPartition/DataLayer/DataLayer.h
--- a/Engine/Source/Runtime/Engine/Public/WorldPartition/DataLayer/DataLayer.h
+++ b/Engine/Source/Runtime/Engine/Public/WorldPartition/DataLayer/DataLayer.h
@@ -21,7 +21,11 @@ public:
 	void SetVisible(bool bInIsVisible);
 	void SetIsDynamicallyLoaded(bool bInIsDynamicallyLoaded);
 	void SetIsDynamicallyLoadedInEditor(bool bInIsDynamicallyLoadedInEditor);
+
 	bool IsDynamicallyLoadedInEditor() const { return !IsDynamicallyLoaded() || bIsDynamicallyLoadedInEditor; }
+	bool ShouldGenerateHLODs() const { return IsDynamicallyLoaded() && bGeneratesHLODs; }
+	class UHLODLayer* GetDefaultHLODLayer() const { return ShouldGenerateHLODs() ? DefaultHLODLayer : nullptr; }
+
 	static FText GetDataLayerText(const UDataLayer* InDataLayer);
 #endif
 	
@@ -46,5 +50,13 @@ private:
 	/** Whether the DataLayer affects actor editor loading */
 	UPROPERTY(Transient)
 	uint32 bIsDynamicallyLoadedInEditor : 1;
+
+	/** Whether HLODs should be generated for actors in this layer */
+	UPROPERTY(Category = DataLayer, EditAnywhere, meta=(EditConditionHides, EditCondition="bIsDynamicallyLoaded"))
+	uint32 bGeneratesHLODs : 1;
+
+	// Default HLOD layer
+	UPROPERTY(Category = DataLayer, EditAnywhere, meta=(EditConditionHides, EditCondition="bIsDynamicallyLoaded && bGeneratesHLODs"))
+	TObjectPtr<class UHLODLayer> DefaultHLODLayer;
 #endif
 };
\ No newline at end of file
