Fixed bug where actor's runtime grid would be read-only.

diff --git a/Engine/Source/Runtime/Engine/Private/ActorEditor.cpp b/Engine/Source/Runtime/Engine/Private/ActorEditor.cpp
--- a/Engine/Source/Runtime/Engine/Private/ActorEditor.cpp
+++ b/Engine/Source/Runtime/Engine/Private/ActorEditor.cpp
@@ -57,10 +57,11 @@ bool AActor::CanEditChange(const FProperty* PropertyThatWillChange) const
 		return false;
 	}
 
-	const bool bIsDataLayersProperty = (PropertyThatWillChange->GetFName() == GET_MEMBER_NAME_CHECKED(AActor, DataLayers));
-	if ((PropertyThatWillChange->GetFName() == GET_MEMBER_NAME_CHECKED(AActor, GridPlacement)) || 
-		(PropertyThatWillChange->GetFName() == GET_MEMBER_NAME_CHECKED(AActor, RuntimeGrid)) ||
-		bIsDataLayersProperty)
+	const bool bIsGridPlacement = PropertyThatWillChange->GetFName() == GET_MEMBER_NAME_CHECKED(AActor, GridPlacement);
+	const bool bIsRuntimeGrid = PropertyThatWillChange->GetFName() == GET_MEMBER_NAME_CHECKED(AActor, RuntimeGrid);
+	const bool bIsDataLayers = PropertyThatWillChange->GetFName() == GET_MEMBER_NAME_CHECKED(AActor, DataLayers);
+
+	if (bIsGridPlacement || bIsRuntimeGrid || bIsDataLayers)
 	{
 		if (!IsTemplate())
 		{
@@ -71,7 +72,7 @@ bool AActor::CanEditChange(const FProperty* PropertyThatWillChange) const
 			}
 		}
 
-		if (!bIsDataLayersProperty && (GetDefaultGridPlacement() != EActorGridPlacement::None))
+		if (bIsGridPlacement && (GetDefaultGridPlacement() != EActorGridPlacement::None))
 		{
 			return false;
 		}
