Add experimental support for material control flow (if-statement, for-loop, local variables)
- Very WIP, not intended/ready for production use
- Hidden behind CVAR to enable support, and per-material flag to opt in

diff --git a/Engine/Shaders/Private/LocalVertexFactory.ush b/Engine/Shaders/Private/LocalVertexFactory.ush
--- a/Engine/Shaders/Private/LocalVertexFactory.ush
+++ b/Engine/Shaders/Private/LocalVertexFactory.ush
@@ -653,6 +653,8 @@ FMaterialVertexParameters GetMaterialVertexParameters(FVertexFactoryInput Input,
 	Result.Particle.WorldToParticle = GetPrimitiveData(Result.PrimitiveId).WorldToLocal;
 #endif
 
+	Result.MaterialVertexAttributes = EvaluateVertexMaterialAttributes(Result);
+
 	return Result;
 }
 
diff --git a/Engine/Shaders/Private/MaterialTemplate.ush b/Engine/Shaders/Private/MaterialTemplate.ush
--- a/Engine/Shaders/Private/MaterialTemplate.ush
+++ b/Engine/Shaders/Private/MaterialTemplate.ush
@@ -223,6 +223,9 @@ struct FMaterialAttributes
 %s
 };
 
+/** FMaterialAttributes utilities */
+%s
+
 /** 
  * Parameters calculated from the pixel material inputs.
  */
@@ -477,6 +480,8 @@ struct FMaterialVertexParameters
 #if WATER_MESH_FACTORY
 	uint WaterWaveParamIndex;
 #endif
+
+	FMaterialAttributes MaterialVertexAttributes;
 };
 
 /**
@@ -2378,6 +2383,29 @@ float3 CalculateAnisotropyTangent(FMaterialPixelParameters Parameters, FPixelMat
 	return Tangent;
 }
 
+FMaterialAttributes FMaterialAttributes_GetDefault()
+{
+	FMaterialAttributes Result;
+%s
+	return Result;
+}
+
+FMaterialAttributes EvaluateVertexMaterialAttributes(in out FMaterialVertexParameters Parameters)
+{
+	FMaterialAttributes DefaultMaterialAttributes = FMaterialAttributes_GetDefault();
+%s
+	// Return the default, this may be skipped if the above code returns different attributes
+	return DefaultMaterialAttributes;
+}
+
+FMaterialAttributes EvaluatePixelMaterialAttributes(in out FMaterialPixelParameters Parameters)
+{
+	FMaterialAttributes DefaultMaterialAttributes = FMaterialAttributes_GetDefault();
+%s
+	// Return the default, this may be skipped if the above code returns different attributes
+	return DefaultMaterialAttributes;
+}
+
 // Be very, very careful about changing CalcPixelMaterialInputs() or CalcPixelMaterialInputsAnalyticDerivatives(). They apply the same basic calculation,
 // but CalcPixelMaterialInputs() uses finite difference derivates from hardware whereas CalcPixelMaterialInputsAnalyticDerivatives() uses analytic
 // derivatives. If you change anything, make sure that both functions stay identical. 
diff --git a/Engine/Source/Editor/MaterialEditor/Private/MaterialEditor.cpp b/Engine/Source/Editor/MaterialEditor/Private/MaterialEditor.cpp
--- a/Engine/Source/Editor/MaterialEditor/Private/MaterialEditor.cpp
+++ b/Engine/Source/Editor/MaterialEditor/Private/MaterialEditor.cpp
@@ -6049,6 +6049,10 @@ void FMaterialEditor::NotifyPostChange( const FPropertyChangedEvent& PropertyCha
 				SetPreviewAsset(GUnrealEd->GetThumbnailManager()->EditorSphere);
 			}
 		}
+		else if (NameOfPropertyThatChanged == GET_MEMBER_NAME_CHECKED(UMaterial, bEnableExecWire))
+		{
+			Material->MaterialGraph->RebuildGraph();
+		}
 		else if( NameOfPropertyThatChanged == GET_MEMBER_NAME_CHECKED(UMaterial, MaterialDomain) ||
 				 NameOfPropertyThatChanged == GET_MEMBER_NAME_CHECKED(UMaterial, ShadingModel))
 		{
diff --git a/Engine/Source/Editor/MaterialEditor/Private/MaterialEditorDetailCustomization.cpp b/Engine/Source/Editor/MaterialEditor/Private/MaterialEditorDetailCustomization.cpp
--- a/Engine/Source/Editor/MaterialEditor/Private/MaterialEditorDetailCustomization.cpp
+++ b/Engine/Source/Editor/MaterialEditor/Private/MaterialEditorDetailCustomization.cpp
@@ -736,6 +736,8 @@ TSharedRef<class IDetailCustomization> FMaterialDetailCustomization::MakeInstanc
 
 void FMaterialDetailCustomization::CustomizeDetails( IDetailLayoutBuilder& DetailLayout )
 {
+	static const auto CVarMaterialEnableControlFlow = IConsoleManager::Get().FindTConsoleVariableDataInt(TEXT("r.MaterialEnableControlFlow"));
+
 	TArray<TWeakObjectPtr<UObject> > Objects;
 	DetailLayout.GetObjectsBeingCustomized( Objects );
 
@@ -781,6 +783,11 @@ void FMaterialDetailCustomization::CustomizeDetails( IDetailLayoutBuilder& Detai
 				}
 			}
 
+			if (PropertyName == GET_MEMBER_NAME_CHECKED(UMaterial, bEnableExecWire) && !CVarMaterialEnableControlFlow->GetValueOnAnyThread())
+			{
+				DetailLayout.HideProperty(PropertyHandle);
+			}
+
 #if WITH_EDITORONLY_DATA
 			// Hide the shading model 
 			if (!bIsShadingModelFromMaterialExpression && PropertyName == GET_MEMBER_NAME_CHECKED(UMaterial, UsedShadingModels))
diff --git a/Engine/Source/Editor/MaterialEditor/Private/MaterialExpressionClasses.cpp b/Engine/Source/Editor/MaterialEditor/Private/MaterialExpressionClasses.cpp
--- a/Engine/Source/Editor/MaterialEditor/Private/MaterialExpressionClasses.cpp
+++ b/Engine/Source/Editor/MaterialEditor/Private/MaterialExpressionClasses.cpp
@@ -17,7 +17,13 @@
 #include "Materials/MaterialExpressionParameter.h"
 #include "Materials/MaterialExpressionMaterialLayerOutput.h"
 #include "Materials/MaterialExpressionNamedReroute.h"
+#include "Materials/MaterialExpressionExecBegin.h"
 
+#include "Materials/MaterialExpressionIfThenElse.h"
+#include "Materials/MaterialExpressionForLoop.h"
+#include "Materials/MaterialExpressionGetLocal.h"
+#include "Materials/MaterialExpressionSetLocal.h"
+#include "Materials/MaterialExpressionReturnMaterialAttributes.h"
 
 ///////////////////////////////////////////////////////////////////////////////////////////////////////////////////////
 //
@@ -78,6 +84,9 @@ void MaterialExpressionClasses::InitMaterialExpressionClasses()
 {
 	if(!bInitialized)
 	{
+		static const auto CVarMaterialEnableControlFlow = IConsoleManager::Get().FindTConsoleVariableDataInt(TEXT("r.MaterialEnableControlFlow"));
+		const bool bEnableControlFlow = CVarMaterialEnableControlFlow->GetValueOnAnyThread() != 0;
+
 		UMaterialEditorOptions* TempEditorOptions = NewObject<UMaterialEditorOptions>();
 		UClass* BaseType = UMaterialExpression::StaticClass();
 		if( BaseType )
@@ -98,11 +107,18 @@ void MaterialExpressionClasses::InitMaterialExpressionClasses()
 						{
 							continue;
 						}
-						
+
+						// Hide node types related to control flow, unless it's enabled
+						if (!bEnableControlFlow && Class->HasMetaData("MaterialControlFlow"))
+						{
+							continue;
+						}
+
 						// Exclude comments from the expression list, as well as the base parameter expression and local variable usage, as they should not be used directly
 						if ( Class != UMaterialExpressionComment::StaticClass() 
 							&& Class != UMaterialExpressionParameter::StaticClass()
-							&& Class != UMaterialExpressionNamedRerouteUsage::StaticClass())
+							&& Class != UMaterialExpressionNamedRerouteUsage::StaticClass()
+							&& Class != UMaterialExpressionExecBegin::StaticClass() )
 						{
 							FMaterialExpression MaterialExpression;
 							// Trim the material expression name and add it to the list used for filtering.
diff --git a/Engine/Source/Editor/UnrealEd/Classes/MaterialGraph/MaterialGraphNode.h b/Engine/Source/Editor/UnrealEd/Classes/MaterialGraph/MaterialGraphNode.h
--- a/Engine/Source/Editor/UnrealEd/Classes/MaterialGraph/MaterialGraphNode.h
+++ b/Engine/Source/Editor/UnrealEd/Classes/MaterialGraph/MaterialGraphNode.h
@@ -47,10 +47,8 @@ public:
 	/** Recreate this node's pins and relink the Graph from the Material */
 	UNREALED_API void RecreateAndLinkNode();
 
-	/** Get the Material Expression output index from an output pin */
-	UNREALED_API int32 GetOutputIndex(const UEdGraphPin* OutputPin);
 	/** Get the Material value type of an output pin */
-	uint32 GetOutputType(const UEdGraphPin* OutputPin);
+	virtual uint32 GetOutputType(const UEdGraphPin* OutputPin) const override;
 
 	//~ Begin UObject Interface
 	virtual void PostEditChangeProperty(FPropertyChangedEvent& PropertyChangedEvent) override;
diff --git a/Engine/Source/Editor/UnrealEd/Classes/MaterialGraph/MaterialGraphNode_Base.h b/Engine/Source/Editor/UnrealEd/Classes/MaterialGraph/MaterialGraphNode_Base.h
--- a/Engine/Source/Editor/UnrealEd/Classes/MaterialGraph/MaterialGraphNode_Base.h
+++ b/Engine/Source/Editor/UnrealEd/Classes/MaterialGraph/MaterialGraphNode_Base.h
@@ -37,6 +37,12 @@ class UMaterialGraphNode_Base : public UEdGraphNode
 	/** Get the Material value type of an input pin */
 	virtual uint32 GetInputType(const UEdGraphPin* InputPin) const;
 
+	/** Get the Material Expression output index from an output pin */
+	UNREALED_API int32 GetOutputIndex(const UEdGraphPin* OutputPin) const;
+
+	/** Get the Material value type of an output pin */
+	virtual uint32 GetOutputType(const UEdGraphPin* OutputPin) const { return 0; }
+
 	/**
 	 * Handles inserting the node between the FromPin and what the FromPin was original connected to
 	 *
diff --git a/Engine/Source/Editor/UnrealEd/Classes/MaterialGraph/MaterialGraphSchema.h b/Engine/Source/Editor/UnrealEd/Classes/MaterialGraph/MaterialGraphSchema.h
--- a/Engine/Source/Editor/UnrealEd/Classes/MaterialGraph/MaterialGraphSchema.h
+++ b/Engine/Source/Editor/UnrealEd/Classes/MaterialGraph/MaterialGraphSchema.h
@@ -182,6 +182,7 @@ class UMaterialGraphSchema : public UEdGraphSchema
 	UNREALED_API static const FName PC_Required;
 	UNREALED_API static const FName PC_Optional;
 	UNREALED_API static const FName PC_MaterialInput;
+	UNREALED_API static const FName PC_Exec;
 
 	// Common PinType.PinSubCategory values
 	UNREALED_API static const FName PSC_Red;
@@ -190,6 +191,8 @@ class UMaterialGraphSchema : public UEdGraphSchema
 	UNREALED_API static const FName PSC_Alpha;
 	UNREALED_API static const FName PSC_RGBA;
 
+	UNREALED_API static const FName PN_Execute; // Category=PC_Exec, singleton, input
+
 	// Color of certain pins/connections
 	UNREALED_API static const FLinearColor ActivePinColor;
 	UNREALED_API static const FLinearColor InactivePinColor;
diff --git a/Engine/Source/Editor/UnrealEd/Private/MaterialGraph.cpp b/Engine/Source/Editor/UnrealEd/Private/MaterialGraph.cpp
--- a/Engine/Source/Editor/UnrealEd/Private/MaterialGraph.cpp
+++ b/Engine/Source/Editor/UnrealEd/Private/MaterialGraph.cpp
@@ -92,11 +92,16 @@ void UMaterialGraph::RebuildGraphInternal(const TMap<UMaterialExpression*, TArra
 		//^^^ New material properties go above here. ^^^^
 		MaterialInputs.Add(FMaterialInputInfo(LOCTEXT("MaterialAttributes", "Material Attributes"), MP_MaterialAttributes, LOCTEXT("MaterialAttributesToolTip", "Material Attributes")));
 
-		// Add Root Node
-		FGraphNodeCreator<UMaterialGraphNode_Root> NodeCreator(*this);
-		RootNode = NodeCreator.CreateNode();
-		RootNode->Material = Material;
-		NodeCreator.Finalize();
+		// No root node if material is using exec pin
+		// Attribute assignments should be done using explicit nodes in this case
+		if (!Material->IsCompiledWithExecutionFlow())
+		{
+			// Add Root Node
+			FGraphNodeCreator<UMaterialGraphNode_Root> NodeCreator(*this);
+			RootNode = NodeCreator.CreateNode();
+			RootNode->Material = Material;
+			NodeCreator.Finalize();
+		}
 	}
 
 	TArray<UMaterialExpression*> ChildSubGraphExpressions;
diff --git a/Engine/Source/Editor/UnrealEd/Private/MaterialGraphNode.cpp b/Engine/Source/Editor/UnrealEd/Private/MaterialGraphNode.cpp
--- a/Engine/Source/Editor/UnrealEd/Private/MaterialGraphNode.cpp
+++ b/Engine/Source/Editor/UnrealEd/Private/MaterialGraphNode.cpp
@@ -38,6 +38,8 @@
 #include "Materials/MaterialExpressionNamedReroute.h"
 #include "Materials/MaterialExpressionReroute.h"
 #include "Materials/MaterialExpressionCurveAtlasRowParameter.h"
+#include "Materials/MaterialExpressionReturnMaterialAttributes.h"
+#include "Materials/MaterialExpressionExecBegin.h"
 #include "Materials/MaterialExpressionStrata.h"
 #include "MaterialEditorUtilities.h"
 #include "MaterialEditorActions.h"
@@ -229,8 +231,6 @@ FLinearColor UMaterialGraphNode::GetNodeTitleColor() const
 		return Settings->PreviewNodeTitleColor;
 	}
 
-
-
 	if (UsesBoolColour(MaterialExpression))
 	{
 		return Settings->BooleanPinTypeColor;
@@ -279,6 +279,18 @@ FLinearColor UMaterialGraphNode::GetNodeTitleColor() const
 		// Previously FColor(255, 155, 0);
 		return Settings->ResultNodeTitleColor;
 	}
+	else if (MaterialExpression->IsA(UMaterialExpressionReturnMaterialAttributes::StaticClass()))
+	{
+		return Settings->ResultNodeTitleColor;
+	}
+	else if (MaterialExpression->IsA(UMaterialExpressionExecBegin::StaticClass()))
+	{
+		return Settings->FunctionTerminatorNodeTitleColor;
+	}
+	else if (MaterialExpression->GetExecInput())
+	{
+		return Settings->ExecutionPinTypeColor;
+	}
 	else if (const UMaterialExpressionNamedRerouteDeclaration* RerouteDeclaration = Cast<UMaterialExpressionNamedRerouteDeclaration>(MaterialExpression))
 	{
 		// If it's a declaration node, we simply get the color from it 
@@ -615,11 +627,26 @@ void UMaterialGraphNode::CreateInputPins()
 	for (int32 Index = 0; Index < ExpressionInputs.Num() ; ++Index)
 	{
 		FExpressionInput* Input = ExpressionInputs[Index];
-		FName InputName = MaterialExpression->GetInputName(Index);
-
-		InputName = GetShortenPinName(InputName);
-
-		const FName PinCategory = MaterialExpression->IsInputConnectionRequired(Index) ? UMaterialGraphSchema::PC_Required : UMaterialGraphSchema::PC_Optional;
+		const uint32 InputType = MaterialExpression->GetInputType(Index);
+		FName PinCategory;
+		FName InputName;
+		if (InputType == MCT_Execution)
+		{
+			PinCategory = UMaterialGraphSchema::PC_Exec;
+		}
+		else
+		{
+			InputName = MaterialExpression->GetInputName(Index);
+			InputName = GetShortenPinName(InputName);
+			if (MaterialExpression->IsInputConnectionRequired(Index))
+			{
+				PinCategory = UMaterialGraphSchema::PC_Required;
+			}
+			else
+			{
+				PinCategory = UMaterialGraphSchema::PC_Optional;
+			}
+		}
 
 		UEdGraphPin* NewPin = CreatePin(EGPD_Input, PinCategory, InputName);
 		if (NewPin->PinName.IsNone())
@@ -635,15 +662,26 @@ void UMaterialGraphNode::CreateOutputPins()
 {
 	TArray<FExpressionOutput>& Outputs = MaterialExpression->GetOutputs();
 
-	for (const FExpressionOutput& ExpressionOutput : Outputs)
+	for(int32 Index = 0; Index < Outputs.Num(); ++Index)
 	{
+		const FExpressionOutput& ExpressionOutput = Outputs[Index];
+		const uint32 OutputType = MaterialExpression->GetOutputType(Index);
+
 		FName PinCategory;
 		FName PinSubCategory;
 		FName OutputName;
+		if (MaterialExpression->bShowOutputNameOnPin)
+		{
+			OutputName = ExpressionOutput.OutputName;
+		}
 
-		if (MaterialExpression->bShowMaskColorsOnPin)
+		if (OutputType == MCT_Execution)
+		{
+			PinCategory = UMaterialGraphSchema::PC_Exec;
+		}
+		else
 		{
-			if (ExpressionOutput.Mask)
+			if (MaterialExpression->bShowMaskColorsOnPin && ExpressionOutput.Mask)
 			{
 				PinCategory = UMaterialGraphSchema::PC_Mask;
 
@@ -670,11 +708,6 @@ void UMaterialGraphNode::CreateOutputPins()
 			}
 		}
 
-		if (MaterialExpression->bShowOutputNameOnPin)
-		{
-			OutputName = ExpressionOutput.OutputName;
-		}
-
 		UEdGraphPin* NewPin = CreatePin(EGPD_Output, PinCategory, PinSubCategory, OutputName);
 		if (NewPin->PinName.IsNone())
 		{
@@ -685,23 +718,7 @@ void UMaterialGraphNode::CreateOutputPins()
 	}
 }
 
-int32 UMaterialGraphNode::GetOutputIndex(const UEdGraphPin* OutputPin)
-{
-	TArray<UEdGraphPin*> OutputPins;
-	GetOutputPins(OutputPins);
-
-	for (int32 Index = 0; Index < OutputPins.Num(); ++Index)
-	{
-		if (OutputPin == OutputPins[Index])
-		{
-			return Index;
-		}
-	}
-
-	return -1;
-}
-
-uint32 UMaterialGraphNode::GetOutputType(const UEdGraphPin* OutputPin)
+uint32 UMaterialGraphNode::GetOutputType(const UEdGraphPin* OutputPin) const
 {
 	return MaterialExpression->GetOutputType(GetOutputIndex(OutputPin));
 }
diff --git a/Engine/Source/Editor/UnrealEd/Private/MaterialGraphNode_Base.cpp b/Engine/Source/Editor/UnrealEd/Private/MaterialGraphNode_Base.cpp
--- a/Engine/Source/Editor/UnrealEd/Private/MaterialGraphNode_Base.cpp
+++ b/Engine/Source/Editor/UnrealEd/Private/MaterialGraphNode_Base.cpp
@@ -83,6 +83,22 @@ void UMaterialGraphNode_Base::GetOutputPins(TArray<class UEdGraphPin*>& OutOutpu
 	}
 }
 
+int32 UMaterialGraphNode_Base::GetOutputIndex(const UEdGraphPin* OutputPin) const
+{
+	TArray<UEdGraphPin*> OutputPins;
+	GetOutputPins(OutputPins);
+
+	for (int32 Index = 0; Index < OutputPins.Num(); ++Index)
+	{
+		if (OutputPin == OutputPins[Index])
+		{
+			return Index;
+		}
+	}
+
+	return -1;
+}
+
 void UMaterialGraphNode_Base::ReplaceNode(UMaterialGraphNode_Base* OldNode)
 {
 	check(OldNode);
diff --git a/Engine/Source/Editor/UnrealEd/Private/MaterialGraphSchema.cpp b/Engine/Source/Editor/UnrealEd/Private/MaterialGraphSchema.cpp
--- a/Engine/Source/Editor/UnrealEd/Private/MaterialGraphSchema.cpp
+++ b/Engine/Source/Editor/UnrealEd/Private/MaterialGraphSchema.cpp
@@ -228,6 +228,7 @@ const FName UMaterialGraphSchema::PC_Mask(TEXT("mask"));
 const FName UMaterialGraphSchema::PC_Required(TEXT("required"));
 const FName UMaterialGraphSchema::PC_Optional(TEXT("optional"));
 const FName UMaterialGraphSchema::PC_MaterialInput(TEXT("materialinput"));
+const FName UMaterialGraphSchema::PC_Exec(TEXT("exec"));
 
 const FName UMaterialGraphSchema::PSC_Red(TEXT("red"));
 const FName UMaterialGraphSchema::PSC_Green(TEXT("green"));
@@ -235,6 +236,8 @@ const FName UMaterialGraphSchema::PSC_Blue(TEXT("blue"));
 const FName UMaterialGraphSchema::PSC_Alpha(TEXT("alpha"));
 const FName UMaterialGraphSchema::PSC_RGBA(TEXT("rgba"));
 
+const FName UMaterialGraphSchema::PN_Execute("execute");
+
 const FLinearColor UMaterialGraphSchema::ActivePinColor = FLinearColor::White;
 const FLinearColor UMaterialGraphSchema::InactivePinColor = FLinearColor(0.05f, 0.05f, 0.05f);
 const FLinearColor UMaterialGraphSchema::AlphaPinColor = FLinearColor(0.5f, 0.5f, 0.5f);
@@ -351,15 +354,15 @@ void UMaterialGraphSchema::GetPaletteActions(FGraphActionMenuBuilder& ActionMenu
 
 bool UMaterialGraphSchema::ConnectionCausesLoop(const UEdGraphPin* InputPin, const UEdGraphPin* OutputPin) const
 {
-	// Only nodes representing Expressions have outputs
-	UMaterialGraphNode* OutputNode = CastChecked<UMaterialGraphNode>(OutputPin->GetOwningNode());
-
-	TArray<UMaterialExpression*> InputExpressions;
-	OutputNode->MaterialExpression->GetAllInputExpressions(InputExpressions);
-
-	if (UMaterialGraphNode* InputNode = Cast<UMaterialGraphNode>(InputPin->GetOwningNode()))
+	if (UMaterialGraphNode* OutputNode = Cast<UMaterialGraphNode>(OutputPin->GetOwningNode()))
 	{
-		return InputExpressions.Contains(InputNode->MaterialExpression);
+		TArray<UMaterialExpression*> InputExpressions;
+		OutputNode->MaterialExpression->GetAllInputExpressions(InputExpressions);
+
+		if (UMaterialGraphNode* InputNode = Cast<UMaterialGraphNode>(InputPin->GetOwningNode()))
+		{
+			return InputExpressions.Contains(InputNode->MaterialExpression);
+		}
 	}
 
 	// Simple connection to root node
@@ -409,14 +412,13 @@ bool UMaterialGraphSchema::ArePinsCompatible_Internal(const UEdGraphPin* InputPi
 
 uint32 UMaterialGraphSchema::GetMaterialValueType(const UEdGraphPin* MaterialPin)
 {
+	const UMaterialGraphNode_Base* OwningNode = CastChecked<UMaterialGraphNode_Base>(MaterialPin->GetOwningNode());
 	if (MaterialPin->Direction == EGPD_Output)
 	{
-		UMaterialGraphNode* OwningNode = CastChecked<UMaterialGraphNode>(MaterialPin->GetOwningNode());
 		return OwningNode->GetOutputType(MaterialPin);
 	}
 	else
 	{
-		UMaterialGraphNode_Base* OwningNode = CastChecked<UMaterialGraphNode_Base>(MaterialPin->GetOwningNode());
 		return OwningNode->GetInputType(MaterialPin);
 	}
 }
@@ -561,20 +563,25 @@ const FPinConnectionResponse UMaterialGraphSchema::CanCreateConnection(const UEd
 	// Break existing connections on inputs only - multiple output connections are acceptable
 	if (InputPin->LinkedTo.Num() > 0)
 	{
-		ECanCreateConnectionResponse ReplyBreakOutputs;
-		if (InputPin == A)
-		{
-			ReplyBreakOutputs = CONNECT_RESPONSE_BREAK_OTHERS_A;
-		}
-		else
+		const uint32 InputType = GetMaterialValueType(InputPin);
+		// TODO - Allow multiple exec inputs, will generate a compile error where not supported
+		//if (!(InputType & MCT_Execution))
 		{
-			ReplyBreakOutputs = CONNECT_RESPONSE_BREAK_OTHERS_B;
-		}
-		if (ResponseMessage.IsEmpty())
-		{
-			ResponseMessage = LOCTEXT("ConnectionReplace", "Replace existing connections");
+			ECanCreateConnectionResponse ReplyBreakOutputs;
+			if (InputPin == A)
+			{
+				ReplyBreakOutputs = CONNECT_RESPONSE_BREAK_OTHERS_A;
+			}
+			else
+			{
+				ReplyBreakOutputs = CONNECT_RESPONSE_BREAK_OTHERS_B;
+			}
+			if (ResponseMessage.IsEmpty())
+			{
+				ResponseMessage = LOCTEXT("ConnectionReplace", "Replace existing connections");
+			}
+			return FPinConnectionResponse(ReplyBreakOutputs, ResponseMessage);
 		}
-		return FPinConnectionResponse(ReplyBreakOutputs, ResponseMessage);
 	}
 
 	return FPinConnectionResponse(CONNECT_RESPONSE_MAKE, ResponseMessage);
diff --git a/Engine/Source/Runtime/Engine/Classes/Materials/Material.h b/Engine/Source/Runtime/Engine/Classes/Materials/Material.h
--- a/Engine/Source/Runtime/Engine/Classes/Materials/Material.h
+++ b/Engine/Source/Runtime/Engine/Classes/Materials/Material.h
@@ -833,6 +833,10 @@ public:
 	UPROPERTY()
 	TArray<TObjectPtr<class UMaterialExpressionComment>> EditorComments;
 
+	/** The execution begin expression, if material is using an exec wire */
+	UPROPERTY()
+	class UMaterialExpressionExecBegin* ExpressionExecBegin;
+
 	/** Controls where this parameter group is displayed in a material instance parameter list.  The lower the number the higher up in the parameter list. */
 	UPROPERTY(EditAnywhere, EditFixedSize, Category = "Group Sorting")
 	TArray<FParameterGroupData> ParameterGroupData;
@@ -858,6 +862,9 @@ public:
 	UPROPERTY(EditAnywhere, Category=Material)
 	uint8 bUseMaterialAttributes : 1;
 
+	UPROPERTY(EditAnywhere, Category = Material)
+	uint8 bEnableExecWire : 1;
+
 	/** when true, the material casts ray tracing shadows. */
 	UPROPERTY(EditAnywhere, Category = Material)
 	uint8 bCastRayTracedShadows : 1;
@@ -1083,6 +1090,8 @@ public:
 	/** Like IsPropertyActive(), but considers any state overriden by DerivedMaterial */
 	ENGINE_API bool IsPropertyActiveInDerived(EMaterialProperty InProperty, const UMaterialInterface* DerivedMaterial) const;
 
+	ENGINE_API bool IsCompiledWithExecutionFlow() const;
+
 #if WITH_EDITOR
 	/** Allows material properties to be compiled with the option of being overridden by the material attributes input. */
 	ENGINE_API virtual int32 CompilePropertyEx( class FMaterialCompiler* Compiler, const FGuid& AttributeID ) override;
diff --git a/Engine/Source/Runtime/Engine/Classes/Materials/MaterialExpression.h b/Engine/Source/Runtime/Engine/Classes/Materials/MaterialExpression.h
--- a/Engine/Source/Runtime/Engine/Classes/Materials/MaterialExpression.h
+++ b/Engine/Source/Runtime/Engine/Classes/Materials/MaterialExpression.h
@@ -104,6 +104,8 @@ class ENGINE_API UMaterialExpression : public UObject
 {
 	GENERATED_UCLASS_BODY()
 
+	static constexpr int32 CompileExecutionOutputIndex = -2;
+
 #if WITH_EDITORONLY_DATA
 	UPROPERTY()
 	int32 MaterialExpressionEditorX;
@@ -267,6 +269,11 @@ class ENGINE_API UMaterialExpression : public UObject
 	virtual uint32 GetInputType(int32 InputIndex);
 	virtual uint32 GetOutputType(int32 OutputIndex);
 
+	/** Gets this expression's execution input, if any */
+	virtual FExpressionInput* GetExecInput();
+
+	bool IsExecInputConnection(UMaterialExpression* InExpression, int32 InOutputIndex);
+
 	virtual FText GetCreationDescription() const;
 	virtual FText GetCreationName() const;
 
diff --git a/Engine/Source/Runtime/Engine/Classes/Materials/MaterialExpressionExecBegin.h b/Engine/Source/Runtime/Engine/Classes/Materials/MaterialExpressionExecBegin.h
new file mode 100644
--- /dev/null
+++ b/Engine/Source/Runtime/Engine/Classes/Materials/MaterialExpressionExecBegin.h
@@ -0,0 +1,25 @@
+// Copyright Epic Games, Inc. All Rights Reserved.
+
+
+#pragma once
+
+#include "CoreMinimal.h"
+#include "UObject/ObjectMacros.h"
+#include "MaterialExpressionIO.h"
+#include "Materials/MaterialExpression.h"
+#include "MaterialExpressionExecBegin.generated.h"
+
+UCLASS(MinimalAPI, collapsecategories, hidecategories = Object, meta = (MaterialControlFlow))
+class UMaterialExpressionExecBegin : public UMaterialExpression
+{
+	GENERATED_UCLASS_BODY()
+
+	//~ Begin UMaterialExpression Interface
+#if WITH_EDITOR
+	virtual int32 Compile(class FMaterialCompiler* Compiler, int32 OutputIndex) override;
+	virtual void GetCaption(TArray<FString>& OutCaptions) const override;
+	virtual uint32 GetOutputType(int32 OutputIndex);
+	virtual bool CanUserDeleteExpression() const override { return false; }
+#endif
+	//~ End UMaterialExpression Interface
+};
diff --git a/Engine/Source/Runtime/Engine/Classes/Materials/MaterialExpressionForLoop.h b/Engine/Source/Runtime/Engine/Classes/Materials/MaterialExpressionForLoop.h
new file mode 100644
--- /dev/null
+++ b/Engine/Source/Runtime/Engine/Classes/Materials/MaterialExpressionForLoop.h
@@ -0,0 +1,41 @@
+// Copyright Epic Games, Inc. All Rights Reserved.
+
+
+#pragma once
+
+#include "CoreMinimal.h"
+#include "UObject/ObjectMacros.h"
+#include "MaterialExpressionIO.h"
+#include "Materials/MaterialExpression.h"
+#include "MaterialExpressionForLoop.generated.h"
+
+UCLASS(collapsecategories, hidecategories = Object, MinimalAPI, meta = (MaterialControlFlow))
+class UMaterialExpressionForLoop : public UMaterialExpression
+{
+	GENERATED_UCLASS_BODY()
+
+	UPROPERTY()
+	FExpressionInput Exec;
+
+	UPROPERTY()
+	FExpressionInput StartIndex;
+
+	UPROPERTY()
+	FExpressionInput EndIndex;
+
+	UPROPERTY()
+	FExpressionInput IndexStep;
+
+	//~ Begin UMaterialExpression Interface
+#if WITH_EDITOR
+	virtual int32 Compile(class FMaterialCompiler* Compiler, int32 OutputIndex) override;
+	virtual void GetCaption(TArray<FString>& OutCaptions) const override;
+	virtual uint32 GetInputType(int32 InputIndex) override;
+	virtual uint32 GetOutputType(int32 InputIndex) override;
+	virtual FExpressionInput* GetExecInput() override { return &Exec; }
+#endif
+	//~ End UMaterialExpression Interface
+};
+
+
+
diff --git a/Engine/Source/Runtime/Engine/Classes/Materials/MaterialExpressionGetLocal.h b/Engine/Source/Runtime/Engine/Classes/Materials/MaterialExpressionGetLocal.h
new file mode 100644
--- /dev/null
+++ b/Engine/Source/Runtime/Engine/Classes/Materials/MaterialExpressionGetLocal.h
@@ -0,0 +1,27 @@
+// Copyright Epic Games, Inc. All Rights Reserved.
+
+
+#pragma once
+
+#include "CoreMinimal.h"
+#include "UObject/ObjectMacros.h"
+#include "MaterialExpressionIO.h"
+#include "Materials/MaterialExpression.h"
+#include "MaterialExpressionGetLocal.generated.h"
+
+UCLASS(collapsecategories, hidecategories = Object, MinimalAPI, meta = (MaterialControlFlow))
+class UMaterialExpressionGetLocal : public UMaterialExpression
+{
+	GENERATED_UCLASS_BODY()
+
+	UPROPERTY(EditAnywhere, Category = MaterialLocals)
+	FName LocalName;
+
+	//~ Begin UMaterialExpression Interface
+#if WITH_EDITOR
+	virtual int32 Compile(class FMaterialCompiler* Compiler, int32 OutputIndex) override;
+	virtual void GetCaption(TArray<FString>& OutCaptions) const override;
+	virtual uint32 GetOutputType(int32 InputIndex) override;
+#endif
+	//~ End UMaterialExpression Interface
+};
diff --git a/Engine/Source/Runtime/Engine/Classes/Materials/MaterialExpressionIfThenElse.h b/Engine/Source/Runtime/Engine/Classes/Materials/MaterialExpressionIfThenElse.h
new file mode 100644
--- /dev/null
+++ b/Engine/Source/Runtime/Engine/Classes/Materials/MaterialExpressionIfThenElse.h
@@ -0,0 +1,35 @@
+// Copyright Epic Games, Inc. All Rights Reserved.
+
+
+#pragma once
+
+#include "CoreMinimal.h"
+#include "UObject/ObjectMacros.h"
+#include "MaterialExpressionIO.h"
+#include "Materials/MaterialExpression.h"
+#include "MaterialExpressionIfThenElse.generated.h"
+
+UCLASS(collapsecategories, hidecategories = Object, MinimalAPI, meta=(MaterialControlFlow))
+class UMaterialExpressionIfThenElse : public UMaterialExpression
+{
+	GENERATED_UCLASS_BODY()
+
+	UPROPERTY()
+	FExpressionInput Exec;
+
+	UPROPERTY()
+	FExpressionInput Condition;
+
+	//~ Begin UMaterialExpression Interface
+#if WITH_EDITOR
+	virtual int32 Compile(class FMaterialCompiler* Compiler, int32 OutputIndex) override;
+	virtual void GetCaption(TArray<FString>& OutCaptions) const override;
+	virtual uint32 GetInputType(int32 InputIndex) override;
+	virtual uint32 GetOutputType(int32 InputIndex) override;
+	virtual FExpressionInput* GetExecInput() override { return &Exec; }
+#endif
+	//~ End UMaterialExpression Interface
+};
+
+
+
diff --git a/Engine/Source/Runtime/Engine/Classes/Materials/MaterialExpressionReturnMaterialAttributes.h b/Engine/Source/Runtime/Engine/Classes/Materials/MaterialExpressionReturnMaterialAttributes.h
new file mode 100644
--- /dev/null
+++ b/Engine/Source/Runtime/Engine/Classes/Materials/MaterialExpressionReturnMaterialAttributes.h
@@ -0,0 +1,33 @@
+// Copyright Epic Games, Inc. All Rights Reserved.
+
+
+#pragma once
+
+#include "CoreMinimal.h"
+#include "UObject/ObjectMacros.h"
+#include "MaterialExpressionIO.h"
+#include "Materials/MaterialExpression.h"
+#include "MaterialExpressionReturnMaterialAttributes.generated.h"
+
+UCLASS(MinimalAPI, collapsecategories, hidecategories = Object, meta = (MaterialControlFlow))
+class UMaterialExpressionReturnMaterialAttributes : public UMaterialExpression
+{
+	GENERATED_UCLASS_BODY()
+
+	UPROPERTY()
+	FExpressionInput Exec;
+
+	UPROPERTY()
+	FMaterialAttributesInput MaterialAttributes;
+
+	//~ Begin UMaterialExpression Interface
+#if WITH_EDITOR
+	virtual int32 Compile(class FMaterialCompiler* Compiler, int32 OutputIndex) override;
+	virtual void GetCaption(TArray<FString>& OutCaptions) const override;
+	virtual const TArray<FExpressionInput*> GetInputs() override;
+	virtual FExpressionInput* GetInput(int32 InputIndex) override;
+	virtual uint32 GetInputType(int32 InputIndex) override;
+	virtual FExpressionInput* GetExecInput() override { return &Exec; }
+#endif
+	//~ End UMaterialExpression Interface
+};
diff --git a/Engine/Source/Runtime/Engine/Classes/Materials/MaterialExpressionSetLocal.h b/Engine/Source/Runtime/Engine/Classes/Materials/MaterialExpressionSetLocal.h
new file mode 100644
--- /dev/null
+++ b/Engine/Source/Runtime/Engine/Classes/Materials/MaterialExpressionSetLocal.h
@@ -0,0 +1,35 @@
+// Copyright Epic Games, Inc. All Rights Reserved.
+
+
+#pragma once
+
+#include "CoreMinimal.h"
+#include "UObject/ObjectMacros.h"
+#include "MaterialExpressionIO.h"
+#include "Materials/MaterialExpression.h"
+#include "MaterialExpressionSetLocal.generated.h"
+
+UCLASS(collapsecategories, hidecategories = Object, MinimalAPI, meta = (MaterialControlFlow))
+class UMaterialExpressionSetLocal : public UMaterialExpression
+{
+	GENERATED_UCLASS_BODY()
+
+	UPROPERTY()
+	FExpressionInput Exec;
+
+	UPROPERTY()
+	FExpressionInput Value;
+
+	UPROPERTY(EditAnywhere, Category = MaterialLocals)
+	FName LocalName;
+
+	//~ Begin UMaterialExpression Interface
+#if WITH_EDITOR
+	virtual int32 Compile(class FMaterialCompiler* Compiler, int32 OutputIndex) override;
+	virtual void GetCaption(TArray<FString>& OutCaptions) const override;
+	virtual uint32 GetInputType(int32 InputIndex) override;
+	virtual uint32 GetOutputType(int32 InputIndex) override;
+	virtual FExpressionInput* GetExecInput() override { return &Exec; }
+#endif
+	//~ End UMaterialExpression Interface
+};
diff --git a/Engine/Source/Runtime/Engine/Private/Materials/HLSLMaterialTranslator.cpp b/Engine/Source/Runtime/Engine/Private/Materials/HLSLMaterialTranslator.cpp
--- a/Engine/Source/Runtime/Engine/Private/Materials/HLSLMaterialTranslator.cpp
+++ b/Engine/Source/Runtime/Engine/Private/Materials/HLSLMaterialTranslator.cpp
@@ -696,6 +696,10 @@ FHLSLMaterialTranslator::FHLSLMaterialTranslator(FMaterial* InMaterial,
 {
 	FMemory::Memzero(SharedPixelProperties);
 
+	FMemory::Memzero(NumForLoops);
+
+	FMemory::Memzero(MaterialAttributesReturned);
+
 	SharedPixelProperties[MP_Normal] = true;
 	SharedPixelProperties[MP_Tangent] = true;
 	SharedPixelProperties[MP_EmissiveColor] = true;
@@ -1116,12 +1120,16 @@ bool FHLSLMaterialTranslator::Translate()
 		SCOPE_SECONDS_COUNTER(HLSLTranslateTime);
 		bSuccess = true;
 
+		check(ScopeStack.Num() == 0);
+
 		// WARNING: No compile outputs should be stored on the UMaterial / FMaterial / FMaterialResource, unless they are transient editor-only data (like error expressions)
 		// Compile outputs that need to be saved must be stored in MaterialCompilationOutput, which will be saved to the DDC.
 
 		Material->CompileErrors.Empty();
 		Material->ErrorExpressions.Empty();
 
+		bEnableExecutionFlow = Material->IsCompiledWithExecutionFlow();
+
 		bCompileForComputeShader = Material->IsLightFunction();
 
 		// Generate code:
@@ -1148,6 +1156,8 @@ bool FHLSLMaterialTranslator::Translate()
 
 		int32 NormalCodeChunkEnd = -1;
 		int32 Chunk[CompiledMP_MAX];
+		int32 VertexAttributesChunk = INDEX_NONE;
+		int32 PixelAttributesChunk = INDEX_NONE;
 			
 		memset(Chunk, INDEX_NONE, sizeof(Chunk));
 
@@ -1201,14 +1211,6 @@ bool FHLSLMaterialTranslator::Translate()
 		// Some custom outputs must be pre-compiled so they can be re-used as shared inputs
 		CompileCustomOutputs(CustomOutputExpressions, SeenCustomOutputExpressionsClasses, true);			
 
-		// Normal must always be compiled first; this will ensure its chunk calculations are the first to be added
-		{
-			// Verify that start chunk is 0
-			check(SharedPropertyCodeChunks[NormalShaderFrequency].Num() == 0);
-			Chunk[MP_Normal]					= Material->CompilePropertyAndSetMaterialProperty(MP_Normal					,this);
-			NormalCodeChunkEnd = SharedPropertyCodeChunks[NormalShaderFrequency].Num();
-		}
-
 		// Validate some things on the VT system. Since generated code for expressions shared between multiple properties
 		// (e.g. a texture sample connected to both diffuse and opacity mask) is reused we can't check based on the MaterialProperty
 		// variable inside the actual code generation pass. So we do a pre-pass over it here.
@@ -1217,26 +1219,42 @@ bool FHLSLMaterialTranslator::Translate()
 			ValidateVtPropertyLimits();
 		}
 
-		// Rest of properties
-		Chunk[MP_EmissiveColor]					= Material->CompilePropertyAndSetMaterialProperty(MP_EmissiveColor			,this);
-		Chunk[MP_DiffuseColor]					= Material->CompilePropertyAndSetMaterialProperty(MP_DiffuseColor			,this);
-		Chunk[MP_SpecularColor]					= Material->CompilePropertyAndSetMaterialProperty(MP_SpecularColor			,this);
-		Chunk[MP_BaseColor]						= Material->CompilePropertyAndSetMaterialProperty(MP_BaseColor				,this);
-		Chunk[MP_Metallic]						= Material->CompilePropertyAndSetMaterialProperty(MP_Metallic				,this);
-		Chunk[MP_Specular]						= Material->CompilePropertyAndSetMaterialProperty(MP_Specular				,this);
-		Chunk[MP_Roughness]						= Material->CompilePropertyAndSetMaterialProperty(MP_Roughness				,this);
-		Chunk[MP_Anisotropy]					= Material->CompilePropertyAndSetMaterialProperty(MP_Anisotropy				,this);
-		Chunk[MP_Opacity]						= Material->CompilePropertyAndSetMaterialProperty(MP_Opacity				,this);
-		Chunk[MP_OpacityMask]					= Material->CompilePropertyAndSetMaterialProperty(MP_OpacityMask			,this);
-		Chunk[MP_Tangent]						= Material->CompilePropertyAndSetMaterialProperty(MP_Tangent				,this);
-		Chunk[MP_WorldPositionOffset]			= Material->CompilePropertyAndSetMaterialProperty(MP_WorldPositionOffset	,this);
-		Chunk[MP_WorldDisplacement]				= Material->CompilePropertyAndSetMaterialProperty(MP_WorldDisplacement		,this);
-		Chunk[MP_TessellationMultiplier]		= Material->CompilePropertyAndSetMaterialProperty(MP_TessellationMultiplier	,this);			
-
-		// Make sure to compile this property before using ShadingModelsFromCompilation
-		Chunk[MP_ShadingModel]					= Material->CompilePropertyAndSetMaterialProperty(MP_ShadingModel			,this);
+		if (bEnableExecutionFlow)
+		{
+			PixelAttributesChunk = Material->CompilePropertyAndSetMaterialProperty(MP_MaterialAttributes, this, SF_Pixel);
+			VertexAttributesChunk = Material->CompilePropertyAndSetMaterialProperty(MP_MaterialAttributes, this, SF_Vertex);
+		}
+		else
+		{
+			// Normal must always be compiled first; this will ensure its chunk calculations are the first to be added
+			{
+				// Verify that start chunk is 0
+				check(SharedPropertyCodeChunks[NormalShaderFrequency].Num() == 0);
+				Chunk[MP_Normal]					= Material->CompilePropertyAndSetMaterialProperty(MP_Normal					,this);
+				NormalCodeChunkEnd = SharedPropertyCodeChunks[NormalShaderFrequency].Num();
+			}
+
+			// Rest of properties
+			Chunk[MP_EmissiveColor]					= Material->CompilePropertyAndSetMaterialProperty(MP_EmissiveColor			,this);
+			Chunk[MP_DiffuseColor]					= Material->CompilePropertyAndSetMaterialProperty(MP_DiffuseColor			,this);
+			Chunk[MP_SpecularColor]					= Material->CompilePropertyAndSetMaterialProperty(MP_SpecularColor			,this);
+			Chunk[MP_BaseColor]						= Material->CompilePropertyAndSetMaterialProperty(MP_BaseColor				,this);
+			Chunk[MP_Metallic]						= Material->CompilePropertyAndSetMaterialProperty(MP_Metallic				,this);
+			Chunk[MP_Specular]						= Material->CompilePropertyAndSetMaterialProperty(MP_Specular				,this);
+			Chunk[MP_Roughness]						= Material->CompilePropertyAndSetMaterialProperty(MP_Roughness				,this);
+			Chunk[MP_Anisotropy]					= Material->CompilePropertyAndSetMaterialProperty(MP_Anisotropy				,this);
+			Chunk[MP_Opacity]						= Material->CompilePropertyAndSetMaterialProperty(MP_Opacity				,this);
+			Chunk[MP_OpacityMask]					= Material->CompilePropertyAndSetMaterialProperty(MP_OpacityMask			,this);
+			Chunk[MP_Tangent]						= Material->CompilePropertyAndSetMaterialProperty(MP_Tangent				,this);
+			Chunk[MP_WorldPositionOffset]			= Material->CompilePropertyAndSetMaterialProperty(MP_WorldPositionOffset	,this);
+			Chunk[MP_WorldDisplacement]				= Material->CompilePropertyAndSetMaterialProperty(MP_WorldDisplacement		,this);
+			Chunk[MP_TessellationMultiplier]		= Material->CompilePropertyAndSetMaterialProperty(MP_TessellationMultiplier	,this);			
+
+			// Make sure to compile this property before using ShadingModelsFromCompilation
+			Chunk[MP_ShadingModel]					= Material->CompilePropertyAndSetMaterialProperty(MP_ShadingModel			,this);
 		
-		Chunk[MP_FrontMaterial]					= Material->CompilePropertyAndSetMaterialProperty(MP_FrontMaterial			,this);
+			Chunk[MP_FrontMaterial]					= Material->CompilePropertyAndSetMaterialProperty(MP_FrontMaterial			,this);
+		}
 
 		// Get shading models from material.
 		FMaterialShadingModelField MaterialShadingModels = Material->GetShadingModels(); 
@@ -1250,47 +1268,49 @@ bool FHLSLMaterialTranslator::Translate()
 		
 		ValidateShadingModelsForFeatureLevel(MaterialShadingModels);
 		
-		if (Domain == MD_Volume || (Domain == MD_Surface && IsSubsurfaceShadingModel(MaterialShadingModels)))
+		if (!bEnableExecutionFlow)
 		{
-			// Note we don't test for the blend mode as you can have a translucent material using the subsurface shading model
-
-			// another ForceCast as CompilePropertyAndSetMaterialProperty() can return MCT_Float which we don't want here
-			int32 SubsurfaceColor = Material->CompilePropertyAndSetMaterialProperty(MP_SubsurfaceColor, this);
-			SubsurfaceColor = ForceCast(SubsurfaceColor, FMaterialAttributeDefinitionMap::GetValueType(MP_SubsurfaceColor), MFCF_ExactMatch | MFCF_ReplicateValue);
+			if (Domain == MD_Volume || (Domain == MD_Surface && IsSubsurfaceShadingModel(MaterialShadingModels)))
+			{
+				// Note we don't test for the blend mode as you can have a translucent material using the subsurface shading model
 
-			static FName NameSubsurfaceProfile(TEXT("__SubsurfaceProfile"));
+				// another ForceCast as CompilePropertyAndSetMaterialProperty() can return MCT_Float which we don't want here
+				int32 SubsurfaceColor = Material->CompilePropertyAndSetMaterialProperty(MP_SubsurfaceColor, this);
+				SubsurfaceColor = ForceCast(SubsurfaceColor, FMaterialAttributeDefinitionMap::GetValueType(MP_SubsurfaceColor), MFCF_ExactMatch | MFCF_ReplicateValue);
 
-			// 1.0f is is a not used profile - later this gets replaced with the actual profile
-			int32 CodeSubsurfaceProfile = ForceCast(ScalarParameter(NameSubsurfaceProfile, 1.0f), MCT_Float1);
+				static FName NameSubsurfaceProfile(TEXT("__SubsurfaceProfile"));
 
-			Chunk[MP_SubsurfaceColor] = AppendVector(SubsurfaceColor, CodeSubsurfaceProfile);		
-		}
+				// 1.0f is is a not used profile - later this gets replaced with the actual profile
+				int32 CodeSubsurfaceProfile = ForceCast(ScalarParameter(NameSubsurfaceProfile, 1.0f), MCT_Float1);
 
-		Chunk[MP_CustomData0]					= Material->CompilePropertyAndSetMaterialProperty(MP_CustomData0		,this);
-		Chunk[MP_CustomData1]					= Material->CompilePropertyAndSetMaterialProperty(MP_CustomData1		,this);
-		Chunk[MP_AmbientOcclusion]				= Material->CompilePropertyAndSetMaterialProperty(MP_AmbientOcclusion	,this);
+				Chunk[MP_SubsurfaceColor] = AppendVector(SubsurfaceColor, CodeSubsurfaceProfile);
+			}
 
-		if (IsTranslucentBlendMode(BlendMode) || MaterialShadingModels.HasShadingModel(MSM_SingleLayerWater))
-		{
-			int32 UserRefraction = ForceCast(Material->CompilePropertyAndSetMaterialProperty(MP_Refraction, this), MCT_Float1);
-			int32 RefractionDepthBias = ForceCast(ScalarParameter(FName(TEXT("RefractionDepthBias")), Material->GetRefractionDepthBiasValue()), MCT_Float1);
+			Chunk[MP_CustomData0] = Material->CompilePropertyAndSetMaterialProperty(MP_CustomData0, this);
+			Chunk[MP_CustomData1] = Material->CompilePropertyAndSetMaterialProperty(MP_CustomData1, this);
+			Chunk[MP_AmbientOcclusion] = Material->CompilePropertyAndSetMaterialProperty(MP_AmbientOcclusion, this);
 
-			Chunk[MP_Refraction] = AppendVector(UserRefraction, RefractionDepthBias);
-		}
+			if (IsTranslucentBlendMode(BlendMode) || MaterialShadingModels.HasShadingModel(MSM_SingleLayerWater))
+			{
+				int32 UserRefraction = ForceCast(Material->CompilePropertyAndSetMaterialProperty(MP_Refraction, this), MCT_Float1);
+				int32 RefractionDepthBias = ForceCast(ScalarParameter(FName(TEXT("RefractionDepthBias")), Material->GetRefractionDepthBiasValue()), MCT_Float1);
 
-		if (bCompileForComputeShader)
-		{
-			Chunk[CompiledMP_EmissiveColorCS]	= Material->CompilePropertyAndSetMaterialProperty(MP_EmissiveColor, this, SF_Compute);
-		}
+				Chunk[MP_Refraction] = AppendVector(UserRefraction, RefractionDepthBias);
+			}
 
-		if (Chunk[MP_WorldPositionOffset] != INDEX_NONE)
-		{
-			// Only calculate previous WPO if there is a current WPO
-			Chunk[CompiledMP_PrevWorldPositionOffset] = Material->CompilePropertyAndSetMaterialProperty(MP_WorldPositionOffset, this, SF_Vertex, true);
-		}
+			if (bCompileForComputeShader)
+			{
+				Chunk[CompiledMP_EmissiveColorCS] = Material->CompilePropertyAndSetMaterialProperty(MP_EmissiveColor, this, SF_Compute);
+			}
 
-		Chunk[MP_PixelDepthOffset] = Material->CompilePropertyAndSetMaterialProperty(MP_PixelDepthOffset, this);
+			if (Chunk[MP_WorldPositionOffset] != INDEX_NONE)
+			{
+				// Only calculate previous WPO if there is a current WPO
+				Chunk[CompiledMP_PrevWorldPositionOffset] = Material->CompilePropertyAndSetMaterialProperty(MP_WorldPositionOffset, this, SF_Vertex, true);
+			}
 
+			Chunk[MP_PixelDepthOffset] = Material->CompilePropertyAndSetMaterialProperty(MP_PixelDepthOffset, this);
+		}
 
 		ResourcesString = TEXT("");
 
@@ -1366,14 +1386,17 @@ bool FHLSLMaterialTranslator::Translate()
 		// No more calls to non-vertex shader CompilePropertyAndSetMaterialProperty beyond this point
 		const uint32 SavedNumUserTexCoords = GetNumUserTexCoords();
 
-		for (uint32 CustomUVIndex = MP_CustomizedUVs0; CustomUVIndex <= MP_CustomizedUVs7; CustomUVIndex++)
+		if (!bEnableExecutionFlow)
 		{
-			// Only compile custom UV inputs for UV channels requested by the pixel shader inputs
-			// Any unconnected inputs will have a texcoord generated for them in Material->CompileProperty, which will pass through the vertex (uncustomized) texture coordinates
-			// Note: this is using NumUserTexCoords, which is set by translating all the pixel properties above
-			if (CustomUVIndex - MP_CustomizedUVs0 < SavedNumUserTexCoords)
+			for (uint32 CustomUVIndex = MP_CustomizedUVs0; CustomUVIndex <= MP_CustomizedUVs7; CustomUVIndex++)
 			{
-				Chunk[CustomUVIndex] = Material->CompilePropertyAndSetMaterialProperty((EMaterialProperty)CustomUVIndex, this);
+				// Only compile custom UV inputs for UV channels requested by the pixel shader inputs
+				// Any unconnected inputs will have a texcoord generated for them in Material->CompileProperty, which will pass through the vertex (uncustomized) texture coordinates
+				// Note: this is using NumUserTexCoords, which is set by translating all the pixel properties above
+				if (CustomUVIndex - MP_CustomizedUVs0 < SavedNumUserTexCoords)
+				{
+					Chunk[CustomUVIndex] = Material->CompilePropertyAndSetMaterialProperty((EMaterialProperty)CustomUVIndex, this);
+				}
 			}
 		}
 
@@ -1586,81 +1609,102 @@ bool FHLSLMaterialTranslator::Translate()
 
 		MaterialCompilationOutput.NumUsedUVScalars = GetNumUserTexCoords() * 2;
 		MaterialCompilationOutput.NumUsedCustomInterpolatorScalars = CurrentCustomVertexInterpolatorOffset;
-
-		for (int32 VariationIter = 0; VariationIter < CompiledPDV_MAX; VariationIter++)
+		
+		if (bEnableExecutionFlow)
 		{
-			ECompiledPartialDerivativeVariation Variation = (ECompiledPartialDerivativeVariation)VariationIter;
-
-			// Do Normal Chunk first
+			if (VertexAttributesChunk != INDEX_NONE)
 			{
-				GetFixedParameterCode(
-					0,
-					NormalCodeChunkEnd,
-					Chunk[MP_Normal],
-					SharedPropertyCodeChunks[NormalShaderFrequency],
-					DerivativeVariations[Variation].TranslatedCodeChunkDefinitions[MP_Normal],
-					DerivativeVariations[Variation].TranslatedCodeChunks[MP_Normal],
-					Variation);
+				LinkParentScopes(SharedPropertyCodeChunks[SF_Vertex]);
 
-				// Always gather MP_Normal definitions as they can be shared by other properties
-				if (DerivativeVariations[Variation].TranslatedCodeChunkDefinitions[MP_Normal].IsEmpty())
-				{
-					DerivativeVariations[Variation].TranslatedCodeChunkDefinitions[MP_Normal] = GetDefinitions(SharedPropertyCodeChunks[NormalShaderFrequency], 0, NormalCodeChunkEnd, Variation);
-				}
+				TSet<int32> EmittedChunks;
+				GetScopeCode(1, VertexAttributesChunk, SharedPropertyCodeChunks[SF_Vertex], EmittedChunks, TranslatedAttributesCodeChunks[SF_Vertex]);
 			}
 
-			// Now the rest, skipping Normal
-			for(uint32 PropertyId = 0; PropertyId < MP_MAX; ++PropertyId)
+			if (PixelAttributesChunk != INDEX_NONE)
 			{
-				if (PropertyId == MP_MaterialAttributes || PropertyId == MP_Normal || PropertyId == MP_CustomOutput)
-				{
-					continue;
-				}
+				LinkParentScopes(SharedPropertyCodeChunks[SF_Pixel]);
 
-				const EShaderFrequency PropertyShaderFrequency = FMaterialAttributeDefinitionMap::GetShaderFrequency((EMaterialProperty)PropertyId);
+				TSet<int32> EmittedChunks;
+				GetScopeCode(1, PixelAttributesChunk, SharedPropertyCodeChunks[SF_Pixel], EmittedChunks, TranslatedAttributesCodeChunks[SF_Pixel]);
+			}
+		}
+		else
+		{
+			for (int32 VariationIter = 0; VariationIter < CompiledPDV_MAX; VariationIter++)
+			{
+				ECompiledPartialDerivativeVariation Variation = (ECompiledPartialDerivativeVariation)VariationIter;
 
-				int32 StartChunk = 0;
-				if (PropertyShaderFrequency == NormalShaderFrequency && SharedPixelProperties[PropertyId])
+				// Do Normal Chunk first
 				{
-					// When processing shared properties, do not generate the code before the Normal was generated as those are already handled
-					StartChunk = NormalCodeChunkEnd;
+					GetFixedParameterCode(
+						0,
+						NormalCodeChunkEnd,
+						Chunk[MP_Normal],
+						SharedPropertyCodeChunks[NormalShaderFrequency],
+						DerivativeVariations[Variation].TranslatedCodeChunkDefinitions[MP_Normal],
+						DerivativeVariations[Variation].TranslatedCodeChunks[MP_Normal],
+						Variation);
+
+					// Always gather MP_Normal definitions as they can be shared by other properties
+					if (DerivativeVariations[Variation].TranslatedCodeChunkDefinitions[MP_Normal].IsEmpty())
+					{
+						DerivativeVariations[Variation].TranslatedCodeChunkDefinitions[MP_Normal] = GetDefinitions(SharedPropertyCodeChunks[NormalShaderFrequency], 0, NormalCodeChunkEnd, Variation);
+					}
 				}
 
-				GetFixedParameterCode(
-					StartChunk,
-					SharedPropertyCodeChunks[PropertyShaderFrequency].Num(),
-					Chunk[PropertyId],
-					SharedPropertyCodeChunks[PropertyShaderFrequency],
-					DerivativeVariations[Variation].TranslatedCodeChunkDefinitions[PropertyId],
-					DerivativeVariations[Variation].TranslatedCodeChunks[PropertyId],
-					Variation);
+				// Now the rest, skipping Normal
+				for (uint32 PropertyId = 0; PropertyId < MP_MAX; ++PropertyId)
+				{
+					if (PropertyId == MP_MaterialAttributes || PropertyId == MP_Normal || PropertyId == MP_CustomOutput)
+					{
+						continue;
+					}
+
+					const EShaderFrequency PropertyShaderFrequency = FMaterialAttributeDefinitionMap::GetShaderFrequency((EMaterialProperty)PropertyId);
 
-			}
+					int32 StartChunk = 0;
+					if (PropertyShaderFrequency == NormalShaderFrequency && SharedPixelProperties[PropertyId])
+					{
+						// When processing shared properties, do not generate the code before the Normal was generated as those are already handled
+						StartChunk = NormalCodeChunkEnd;
+					}
 
-			for(uint32 PropertyId = MP_MAX; PropertyId < CompiledMP_MAX; ++PropertyId)
-			{
-				switch(PropertyId)
+					GetFixedParameterCode(
+						StartChunk,
+						SharedPropertyCodeChunks[PropertyShaderFrequency].Num(),
+						Chunk[PropertyId],
+						SharedPropertyCodeChunks[PropertyShaderFrequency],
+						DerivativeVariations[Variation].TranslatedCodeChunkDefinitions[PropertyId],
+						DerivativeVariations[Variation].TranslatedCodeChunks[PropertyId],
+						Variation);
+
+				}
+
+				for (uint32 PropertyId = MP_MAX; PropertyId < CompiledMP_MAX; ++PropertyId)
 				{
-				case CompiledMP_EmissiveColorCS:
-					if (bCompileForComputeShader)
+					switch (PropertyId)
 					{
-						GetFixedParameterCode(Chunk[PropertyId], SharedPropertyCodeChunks[SF_Compute], DerivativeVariations[Variation].TranslatedCodeChunkDefinitions[PropertyId], DerivativeVariations[Variation].TranslatedCodeChunks[PropertyId], Variation);
-					}
-					break;
-				case CompiledMP_PrevWorldPositionOffset:
+					case CompiledMP_EmissiveColorCS:
+						if (bCompileForComputeShader)
+						{
+							GetFixedParameterCode(Chunk[PropertyId], SharedPropertyCodeChunks[SF_Compute], DerivativeVariations[Variation].TranslatedCodeChunkDefinitions[PropertyId], DerivativeVariations[Variation].TranslatedCodeChunks[PropertyId], Variation);
+						}
+						break;
+					case CompiledMP_PrevWorldPositionOffset:
 					{
 						GetFixedParameterCode(Chunk[PropertyId], SharedPropertyCodeChunks[SF_Vertex], DerivativeVariations[Variation].TranslatedCodeChunkDefinitions[PropertyId], DerivativeVariations[Variation].TranslatedCodeChunks[PropertyId], Variation);
 					}
 					break;
-				default: check(0);
-					break;
+					default: check(0);
+						break;
+					}
 				}
-			}
 
-			// Output the implementation for any custom output expressions
-			for (int32 ExpressionIndex = 0; ExpressionIndex < DerivativeVariations[Variation].CustomOutputImplementations.Num(); ExpressionIndex++)
-			{
-				ResourcesString += DerivativeVariations[Variation].CustomOutputImplementations[ExpressionIndex] + "\r\n\r\n";
+				// Output the implementation for any custom output expressions
+				for (int32 ExpressionIndex = 0; ExpressionIndex < DerivativeVariations[Variation].CustomOutputImplementations.Num(); ExpressionIndex++)
+				{
+					ResourcesString += DerivativeVariations[Variation].CustomOutputImplementations[ExpressionIndex] + "\r\n\r\n";
+				}
 			}
 		}
 
@@ -2267,22 +2311,25 @@ void FHLSLMaterialTranslator::GetSharedInputsMaterialCode(FString& PixelMembersD
 			check(PropertyName.Len() > 0);				
 			const EMaterialValueType Type = Property == MP_SubsurfaceColor ? MCT_Float4 : FMaterialAttributeDefinitionMap::GetValueType(Property);
 
-			// Normal requires its own separate initializer
-			if (Property == MP_Normal)
+			if (!bEnableExecutionFlow)
 			{
-				NormalInitializerValue = FString::Printf(TEXT("\tPixelMaterialInputs.%s = %s;\n"), *PropertyName, *DerivativeVariations[DerivativeVariation].TranslatedCodeChunks[Property]);
-			}
-			else
-			{
-				if (DerivativeVariations[DerivativeVariation].TranslatedCodeChunkDefinitions[Property].Len() > 0)
+				// Normal requires its own separate initializer
+				if (Property == MP_Normal)
+				{
+					NormalInitializerValue = FString::Printf(TEXT("\tPixelMaterialInputs.%s = %s;\n"), *PropertyName, *DerivativeVariations[DerivativeVariation].TranslatedCodeChunks[Property]);
+				}
+				else
 				{
-					if (LastProperty >= 0)
+					if (DerivativeVariations[DerivativeVariation].TranslatedCodeChunkDefinitions[Property].Len() > 0)
 					{
-						// Verify that all code chunks have the same contents
-						check(DerivativeVariations[DerivativeVariation].TranslatedCodeChunkDefinitions[Property].Len() == DerivativeVariations[DerivativeVariation].TranslatedCodeChunkDefinitions[LastProperty].Len());
-					}
+						if (LastProperty >= 0)
+						{
+							// Verify that all code chunks have the same contents
+							check(DerivativeVariations[DerivativeVariation].TranslatedCodeChunkDefinitions[Property].Len() == DerivativeVariations[DerivativeVariation].TranslatedCodeChunkDefinitions[LastProperty].Len());
+						}
 
-					LastProperty = Property;
+						LastProperty = Property;
+					}
 				}
 
 				PixelInputInitializerValues += FString::Printf(TEXT("\tPixelMaterialInputs.%s = %s;\n"), *PropertyName, *DerivativeVariations[DerivativeVariation].TranslatedCodeChunks[Property]);
@@ -2344,24 +2391,55 @@ FString FHLSLMaterialTranslator::GetMaterialShaderCode()
 	LazyPrintf.PushParam(*VertexInterpolatorsOffsetsDefinition);
 
 	FString MaterialAttributesDeclaration;
+	FString MaterialAttributesUtilities;
+	FString MaterialAttributesDefault;
+
+	const EMaterialShadingModel DefaultShadingModel = Material->GetShadingModels().GetFirstShadingModel();
 
 	const TArray<FGuid>& OrderedVisibleAttributes = FMaterialAttributeDefinitionMap::GetOrderedVisibleAttributeList();
-	for(const FGuid& AttributeID : OrderedVisibleAttributes)
+	for (const FGuid& AttributeID : OrderedVisibleAttributes)
 	{
 		const FString PropertyName = FMaterialAttributeDefinitionMap::GetAttributeName(AttributeID);
 		const EMaterialValueType PropertyType = FMaterialAttributeDefinitionMap::GetValueType(AttributeID);
+		const TCHAR* HLSLType = nullptr;
+
 		switch (PropertyType)
 		{
-		case MCT_Float1: case MCT_Float: MaterialAttributesDeclaration += FString::Printf(TEXT("\tfloat %s;") LINE_TERMINATOR, *PropertyName); break;
-		case MCT_Float2: MaterialAttributesDeclaration += FString::Printf(TEXT("\tfloat2 %s;") LINE_TERMINATOR, *PropertyName); break;
-		case MCT_Float3: MaterialAttributesDeclaration += FString::Printf(TEXT("\tfloat3 %s;") LINE_TERMINATOR, *PropertyName); break;
-		case MCT_Float4: MaterialAttributesDeclaration += FString::Printf(TEXT("\tfloat4 %s;") LINE_TERMINATOR, *PropertyName); break;
-		case MCT_ShadingModel: MaterialAttributesDeclaration += FString::Printf(TEXT("\tuint %s;") LINE_TERMINATOR, *PropertyName); break;
-		case MCT_Strata: MaterialAttributesDeclaration += FString::Printf(TEXT("\tFStrataData %s;") LINE_TERMINATOR, *PropertyName); break;
+		case MCT_Float1: case MCT_Float: HLSLType = TEXT("float"); break;
+		case MCT_Float2: HLSLType = TEXT("float2"); break;
+		case MCT_Float3: HLSLType = TEXT("float3"); break;
+		case MCT_Float4: HLSLType = TEXT("float4"); break;
+		case MCT_ShadingModel: HLSLType = TEXT("uint"); break;
+		case MCT_Strata: HLSLType = TEXT("FStrataData"); break;
+		default: break;
+		}
+
+		if (HLSLType)
+		{
+			const FVector4& DefaultValue = FMaterialAttributeDefinitionMap::GetDefaultValue(AttributeID);
+
+			MaterialAttributesDeclaration += FString::Printf(TEXT("\t%s %s;") LINE_TERMINATOR, HLSLType, *PropertyName);
+
+			// Chainable method to set the attribute
+			MaterialAttributesUtilities += FString::Printf(TEXT("FMaterialAttributes FMaterialAttributes_Set%s(FMaterialAttributes InAttributes, %s InValue) { InAttributes.%s = InValue; return InAttributes; }") LINE_TERMINATOR,
+				*PropertyName, HLSLType, *PropertyName);
+
+			// Standard value type
+			switch (PropertyType)
+			{
+			case MCT_Float: case MCT_Float1: MaterialAttributesDefault += FString::Printf(TEXT("\tResult.%s = %0.8f;") LINE_TERMINATOR, *PropertyName, DefaultValue.X); break;
+			case MCT_Float2: MaterialAttributesDefault += FString::Printf(TEXT("\tResult.%s = MaterialFloat2(%0.8f,%0.8f);") LINE_TERMINATOR, *PropertyName, DefaultValue.X, DefaultValue.Y); break;
+			case MCT_Float3: MaterialAttributesDefault += FString::Printf(TEXT("\tResult.%s = MaterialFloat3(%0.8f,%0.8f,%0.8f);") LINE_TERMINATOR, *PropertyName, DefaultValue.X, DefaultValue.Y, DefaultValue.Z); break;
+			case MCT_Float4: MaterialAttributesDefault += FString::Printf(TEXT("\tResult.%s = MaterialFloat4(%0.8f,%0.8f,%0.8f,%0.8f);") LINE_TERMINATOR, *PropertyName, DefaultValue.X, DefaultValue.Y, DefaultValue.Z, DefaultValue.W); break;
+			case MCT_ShadingModel: MaterialAttributesDefault += FString::Printf(TEXT("\tResult.%s = %d;") LINE_TERMINATOR, *PropertyName, (int32)DefaultShadingModel); break;
+			case MCT_Strata: MaterialAttributesDefault += FString::Printf(TEXT("\tResult.%s = GetInitialisedStrataData();") LINE_TERMINATOR, *PropertyName); break; // TODO
+			default: checkNoEntry(); break;
+			}
 		}
 	}
 
 	LazyPrintf.PushParam(*MaterialAttributesDeclaration);
+	LazyPrintf.PushParam(*MaterialAttributesUtilities);
 
 	// Stores the shared shader results member declarations
 	FString PixelMembersDeclaration[CompiledPDV_MAX];
@@ -2418,13 +2496,13 @@ FString FHLSLMaterialTranslator::GetMaterialShaderCode()
 
 	LazyPrintf.PushParam(*FString::Printf(TEXT("return %.5f"), Material->GetOpacityMaskClipValue()));
 
-	LazyPrintf.PushParam(*GenerateFunctionCode(MP_WorldPositionOffset, BaseDerivativeVariation));
-	LazyPrintf.PushParam(*GenerateFunctionCode(CompiledMP_PrevWorldPositionOffset, BaseDerivativeVariation));
-	LazyPrintf.PushParam(*GenerateFunctionCode(MP_WorldDisplacement, BaseDerivativeVariation));
+	LazyPrintf.PushParam(!bEnableExecutionFlow ? *GenerateFunctionCode(MP_WorldPositionOffset, BaseDerivativeVariation) : TEXT("return Parameters.MaterialVertexAttributes.WorldPositionOffset"));
+	LazyPrintf.PushParam(!bEnableExecutionFlow ? *GenerateFunctionCode(CompiledMP_PrevWorldPositionOffset, BaseDerivativeVariation) : TEXT("return 0.0f"));
+	LazyPrintf.PushParam(!bEnableExecutionFlow ? *GenerateFunctionCode(MP_WorldDisplacement, BaseDerivativeVariation) : TEXT("return 0.0f"));
 	LazyPrintf.PushParam(*FString::Printf(TEXT("return %.5f"), Material->GetMaxDisplacement()));
-	LazyPrintf.PushParam(*GenerateFunctionCode(MP_TessellationMultiplier, BaseDerivativeVariation));
-	LazyPrintf.PushParam(*GenerateFunctionCode(MP_CustomData0, BaseDerivativeVariation));
-	LazyPrintf.PushParam(*GenerateFunctionCode(MP_CustomData1, BaseDerivativeVariation));
+	LazyPrintf.PushParam(!bEnableExecutionFlow ? *GenerateFunctionCode(MP_TessellationMultiplier, BaseDerivativeVariation) : TEXT("return 0.0f"));
+	LazyPrintf.PushParam(!bEnableExecutionFlow ? *GenerateFunctionCode(MP_CustomData0, BaseDerivativeVariation) : TEXT("return 0.0f"));
+	LazyPrintf.PushParam(!bEnableExecutionFlow ? *GenerateFunctionCode(MP_CustomData1, BaseDerivativeVariation) : TEXT("return 0.0f"));
 
 	// Print custom texture coordinate assignments, should be fine with regular derivatives
 	FString CustomUVAssignments;
@@ -2432,20 +2510,28 @@ FString FHLSLMaterialTranslator::GetMaterialShaderCode()
 	int32 LastProperty = -1;
 	for (uint32 CustomUVIndex = 0; CustomUVIndex < NumUserTexCoords; CustomUVIndex++)
 	{
-		if (CustomUVIndex == 0)
+		if (bEnableExecutionFlow)
 		{
-			CustomUVAssignments += DerivativeVariations[BaseDerivativeVariation].TranslatedCodeChunkDefinitions[MP_CustomizedUVs0 + CustomUVIndex];
+			const FString AttributeName = FMaterialAttributeDefinitionMap::GetAttributeName((EMaterialProperty)(MP_CustomizedUVs0 + CustomUVIndex));
+			CustomUVAssignments += FString::Printf(TEXT("\tOutTexCoords[%u] = Parameters.MaterialVertexAttributes.%s;") LINE_TERMINATOR, CustomUVIndex, *AttributeName);
 		}
-
-		if (DerivativeVariations[BaseDerivativeVariation].TranslatedCodeChunkDefinitions[MP_CustomizedUVs0 + CustomUVIndex].Len() > 0)
+		else
 		{
-			if (LastProperty >= 0)
+			if (CustomUVIndex == 0)
+			{
+				CustomUVAssignments += DerivativeVariations[BaseDerivativeVariation].TranslatedCodeChunkDefinitions[MP_CustomizedUVs0 + CustomUVIndex];
+			}
+
+			if (DerivativeVariations[BaseDerivativeVariation].TranslatedCodeChunkDefinitions[MP_CustomizedUVs0 + CustomUVIndex].Len() > 0)
 			{
-				check(DerivativeVariations[BaseDerivativeVariation].TranslatedCodeChunkDefinitions[LastProperty].Len() == DerivativeVariations[BaseDerivativeVariation].TranslatedCodeChunkDefinitions[MP_CustomizedUVs0 + CustomUVIndex].Len());
+				if (LastProperty >= 0)
+				{
+					check(DerivativeVariations[BaseDerivativeVariation].TranslatedCodeChunkDefinitions[LastProperty].Len() == DerivativeVariations[BaseDerivativeVariation].TranslatedCodeChunkDefinitions[MP_CustomizedUVs0 + CustomUVIndex].Len());
+				}
+				LastProperty = MP_CustomizedUVs0 + CustomUVIndex;
 			}
-			LastProperty = MP_CustomizedUVs0 + CustomUVIndex;
+			CustomUVAssignments += FString::Printf(TEXT("\tOutTexCoords[%u] = %s;") LINE_TERMINATOR, CustomUVIndex, *DerivativeVariations[BaseDerivativeVariation].TranslatedCodeChunks[MP_CustomizedUVs0 + CustomUVIndex]);
 		}
-		CustomUVAssignments += FString::Printf(TEXT("\tOutTexCoords[%u] = %s;") LINE_TERMINATOR, CustomUVIndex, *DerivativeVariations[BaseDerivativeVariation].TranslatedCodeChunks[MP_CustomizedUVs0 + CustomUVIndex]);
 	}
 
 	LazyPrintf.PushParam(*CustomUVAssignments);
@@ -2487,15 +2573,75 @@ FString FHLSLMaterialTranslator::GetMaterialShaderCode()
 
 	LazyPrintf.PushParam(*CustomInterpolatorAssignments);
 
-	for (int32 Iter = 0; Iter < CompiledPDV_MAX; Iter++)
+	LazyPrintf.PushParam(*MaterialAttributesDefault);
+
+	if (bEnableExecutionFlow)
 	{
-		ECompiledPartialDerivativeVariation Variation = (ECompiledPartialDerivativeVariation)Iter;
+		FString EvaluateVertexCode;
+
+		// Set default texcoords in the VS
+		for (uint32 TexCoordIndex = 0; TexCoordIndex < NumUserVertexTexCoords; ++TexCoordIndex)
+		{
+			const FString AttributeName = FMaterialAttributeDefinitionMap::GetAttributeName((EMaterialProperty)(MP_CustomizedUVs0 + TexCoordIndex));
 
-		// Initializers required for Normal
-		LazyPrintf.PushParam(*DerivativeVariations[Variation].TranslatedCodeChunkDefinitions[MP_Normal]);
-		LazyPrintf.PushParam(*NormalAssignment[Variation]);
-		// Finally the rest of common code followed by assignment into each input
-		LazyPrintf.PushParam(*PixelMembersSetupAndAssignments[Variation]);
+			EvaluateVertexCode += FString::Printf(TEXT("\tDefaultMaterialAttributes.%s = Parameters.TexCoords[%d];") LINE_TERMINATOR, *AttributeName, TexCoordIndex);
+		}
+
+		EvaluateVertexCode += TranslatedAttributesCodeChunks[SF_Vertex];
+
+		LazyPrintf.PushParam(*EvaluateVertexCode);
+		LazyPrintf.PushParam(*TranslatedAttributesCodeChunks[SF_Pixel]);
+
+		FString EvaluateMaterialAttributesCode = TEXT("    FMaterialAttributes MaterialAttributes = EvaluatePixelMaterialAttributes(Parameters);" LINE_TERMINATOR);
+
+		for (int32 PropertyIndex = 0; PropertyIndex < MP_MAX; ++PropertyIndex)
+		{
+			// Skip non-shared properties
+			if (!SharedPixelProperties[PropertyIndex])
+			{
+				continue;
+			}
+
+			const EMaterialProperty Property = (EMaterialProperty)PropertyIndex;
+			check(FMaterialAttributeDefinitionMap::GetShaderFrequency(Property) == SF_Pixel);
+			// Special case MP_SubsurfaceColor as the actual property is a combination of the color and the profile but we don't want to expose the profile
+			const FString PropertyName = FMaterialAttributeDefinitionMap::GetAttributeName(Property);
+
+			if (PropertyIndex == MP_SubsurfaceColor)
+			{
+				// TODO - properly handle subsurface profile
+				EvaluateMaterialAttributesCode += FString::Printf("    PixelMaterialInputs.Subsurface = float4(MaterialAttributes.%s, 0.0f);" LINE_TERMINATOR, *PropertyName);
+			}
+			else
+			{
+				EvaluateMaterialAttributesCode += FString::Printf("    PixelMaterialInputs.%s = MaterialAttributes.%s;" LINE_TERMINATOR, *PropertyName, *PropertyName);
+			}
+		}
+
+		// TODO - deriv
+		for (int32 Iter = 0; Iter < CompiledPDV_MAX; Iter++)
+		{
+			LazyPrintf.PushParam(*EvaluateMaterialAttributesCode);
+			LazyPrintf.PushParam(TEXT(""));
+			LazyPrintf.PushParam(TEXT(""));
+		}
+	}
+	else
+	{
+		// skip material attributes code
+		LazyPrintf.PushParam(TEXT(""));
+		LazyPrintf.PushParam(TEXT(""));
+
+		for (int32 Iter = 0; Iter < CompiledPDV_MAX; Iter++)
+		{
+			ECompiledPartialDerivativeVariation Variation = (ECompiledPartialDerivativeVariation)Iter;
+
+			// Initializers required for Normal
+			LazyPrintf.PushParam(*DerivativeVariations[Variation].TranslatedCodeChunkDefinitions[MP_Normal]);
+			LazyPrintf.PushParam(*NormalAssignment[Variation]);
+			// Finally the rest of common code followed by assignment into each input
+			LazyPrintf.PushParam(*PixelMembersSetupAndAssignments[Variation]);
+		}
 	}
 
 	LazyPrintf.PushParam(*FString::Printf(TEXT("%u"),MaterialTemplateLineNumber));
@@ -2508,15 +2654,20 @@ FString FHLSLMaterialTranslator::GetMaterialShaderCode()
 
 bool FHLSLMaterialTranslator::IsMaterialPropertyUsed(EMaterialProperty Property, int32 PropertyChunkIndex, const FLinearColor& ReferenceValue, int32 NumComponents) const
 {
+	const int32 Frequency = (int32)FMaterialAttributeDefinitionMap::GetShaderFrequency(Property);
 	bool bPropertyUsed = false;
 
-	if (PropertyChunkIndex == -1)
+	if ((MaterialAttributesReturned[Frequency] & (1ull << Property)) != 0u)
+	{
+		// Property was set via a 'Return Material Attributes' expression
+		bPropertyUsed = true;
+	}
+	else if (PropertyChunkIndex == -1)
 	{
 		bPropertyUsed = false;
 	}
 	else
 	{
-		int32 Frequency = (int32)FMaterialAttributeDefinitionMap::GetShaderFrequency(Property);
 		const FShaderCodeChunk& PropertyChunk = SharedPropertyCodeChunks[Frequency][PropertyChunkIndex];
 
 		// Determine whether the property is used. 
@@ -2642,6 +2793,8 @@ FString FHLSLMaterialTranslator::GetParameterCodeRaw(int32 Index, const TCHAR* D
 			check(AccessedCodeChunk.SymbolName.Len() > 0);
 			return AccessedCodeChunk.SymbolName;
 		}
+
+		ReferencedCodeChunks.AddUnique(Index);
 			
 		// Return the symbol used to reference this code chunk
 		check(CodeChunk.SymbolName.Len() > 0);
@@ -2670,6 +2823,30 @@ uint64 FHLSLMaterialTranslator::GetParameterHash(int32 Index)
 	return CodeChunk.Hash;
 }
 
+uint64 FHLSLMaterialTranslator::GetParameterMaterialAttributeMask(int32 Index)
+{
+	if (Index == INDEX_NONE)
+	{
+		return 0u;
+	}
+
+	checkf(Index >= 0 && Index < CurrentScopeChunks->Num(), TEXT("Index %d/%d, Platform=%d"), Index, CurrentScopeChunks->Num(), (int)Platform);
+	const FShaderCodeChunk& CodeChunk = (*CurrentScopeChunks)[Index];
+	check(CodeChunk.Type == MCT_MaterialAttributes);
+	check(!CodeChunk.UniformExpression);
+
+	return CodeChunk.MaterialAttributeMask;
+}
+
+void FHLSLMaterialTranslator::SetParameterMaterialAttributes(int32 Index, uint64 Mask)
+{
+	checkf(Index >= 0 && Index < CurrentScopeChunks->Num(), TEXT("Index %d/%d, Platform=%d"), Index, CurrentScopeChunks->Num(), (int)Platform);
+	FShaderCodeChunk& CodeChunk = (*CurrentScopeChunks)[Index];
+	check(CodeChunk.Type == MCT_MaterialAttributes);
+	CodeChunk.MaterialAttributeMask |= Mask;
+}
+
+
 /** Creates a string of all definitions needed for the given material input. */
 FString FHLSLMaterialTranslator::GetDefinitions(TArray<FShaderCodeChunk>& CodeChunks, int32 StartChunk, int32 EndChunk, ECompiledPartialDerivativeVariation Variation) const
 {
@@ -2678,7 +2855,8 @@ FString FHLSLMaterialTranslator::GetDefinitions(TArray<FShaderCodeChunk>& CodeCh
 	{
 		const FShaderCodeChunk& CodeChunk = CodeChunks[ChunkIndex];
 		// Uniform expressions (both constant and variable) and inline expressions don't have definitions.
-		if (!CodeChunk.UniformExpression && !CodeChunk.bInline)
+		if (!CodeChunk.UniformExpression &&
+			(!CodeChunk.bInline || CodeChunk.Type == MCT_VoidStatement))
 		{
 			Definitions += CodeChunk.AtDefinition(Variation);
 		}
@@ -2723,6 +2901,78 @@ void FHLSLMaterialTranslator::GetFixedParameterCode(int32 ResultIndex, TArray<FS
 	GetFixedParameterCode(0, CodeChunks.Num(), ResultIndex, CodeChunks, OutDefinitions, OutValue, Variation);
 }
 
+static void AppendIndent(int32 Level, FString& OutValue)
+{
+	for (int32 i = 0; i < Level; ++i)
+	{
+		OutValue += '\t';
+	}
+}
+
+void FHLSLMaterialTranslator::LinkParentScopes(TArray<FShaderCodeChunk>& CodeChunks)
+{
+	// Add all chunks to parent's child array
+	for (int32 ChunkIndex = 0; ChunkIndex < CodeChunks.Num(); ++ChunkIndex)
+	{
+		FShaderCodeChunk& ChildChunk = CodeChunks[ChunkIndex];
+		if (ChildChunk.UsedScopeIndex != INDEX_NONE)
+		{
+			FShaderCodeChunk& ParentChunk = CodeChunks[ChildChunk.UsedScopeIndex];
+			ParentChunk.ScopedChunks.Add(ChunkIndex);
+		}
+	}
+}
+
+void FHLSLMaterialTranslator::GetScopeCode(int32 IndentLevel, int32 ScopeChunkIndex, const TArray<FShaderCodeChunk>& CodeChunks, TSet<int32>& EmittedChunks, FString& OutValue)
+{
+	bool bAlreadyAddedChunk = false;
+	EmittedChunks.Add(ScopeChunkIndex, &bAlreadyAddedChunk);
+	if (bAlreadyAddedChunk)
+	{
+		// Only emit code for the chunk the first time we visit it
+		return;
+	}
+
+	const FShaderCodeChunk& CodeChunk = CodeChunks[ScopeChunkIndex];
+
+	// Add code for any dependencies of the current chunk
+	for (int32 ReferencedChunkIndex : CodeChunk.ReferencedCodeChunks)
+	{
+		GetScopeCode(IndentLevel, ReferencedChunkIndex, CodeChunks, EmittedChunks, OutValue);
+	}
+
+	if ((!CodeChunk.UniformExpression && !CodeChunk.bInline) || CodeChunk.Type == MCT_VoidStatement)
+	{
+		AppendIndent(IndentLevel, OutValue);
+		OutValue += CodeChunk.AtDefinition(CompiledPDV_FiniteDifferences); // TODO - deriv
+	}
+
+	if (CodeChunk.ScopedChunks.Num() > 0)
+	{
+		AppendIndent(IndentLevel, OutValue);
+		OutValue += TEXT("{\n");
+		for (int32 ChildChunkIndex : CodeChunk.ScopedChunks)
+		{
+			const FShaderCodeChunk& ChildChunk = CodeChunks[ChildChunkIndex];
+			check(ChildChunk.UsedScopeIndex == ScopeChunkIndex);
+			if (ChildChunk.DeclaredScopeIndex != ScopeChunkIndex)
+			{
+				GetScopeCode(IndentLevel + 1, ChildChunkIndex, CodeChunks, EmittedChunks, OutValue);
+			}
+		}
+		for (int32 ChildChunkIndex : CodeChunk.ScopedChunks)
+		{
+			const FShaderCodeChunk& ChildChunk = CodeChunks[ChildChunkIndex];
+			if (ChildChunk.DeclaredScopeIndex == ScopeChunkIndex)
+			{
+				GetScopeCode(IndentLevel + 1, ChildChunkIndex, CodeChunks, EmittedChunks, OutValue);
+			}
+		}
+		AppendIndent(IndentLevel, OutValue);
+		OutValue += TEXT("}\n");
+	}
+}
+
 /** Used to get a user friendly type from EMaterialValueType */
 const TCHAR* FHLSLMaterialTranslator::DescribeType(EMaterialValueType Type) const
 {
@@ -2861,12 +3111,18 @@ int32 FHLSLMaterialTranslator::AddCodeChunkInner(uint64 Hash, const TCHAR* Forma
 		return INDEX_NONE;
 	}
 
-	if (bInlined)
+	int32 CodeIndex = INDEX_NONE;
+	if (Type == MCT_VoidStatement)
 	{
-		const int32 CodeIndex = CurrentScopeChunks->Num();
+		CodeIndex = CurrentScopeChunks->Num();
+		const FString Statement = FString("") + FormattedCode + LINE_TERMINATOR;
+		new(*CurrentScopeChunks) FShaderCodeChunk(Hash, *Statement, *Statement, TEXT(""), Type, true);
+	}
+	else if (bInlined)
+	{
+		CodeIndex = CurrentScopeChunks->Num();
 		// Adding an inline code chunk, the definition will be the code to inline
 		new(*CurrentScopeChunks) FShaderCodeChunk(Hash, FormattedCode,FormattedCode,TEXT(""),Type,true);
-		return CodeIndex;
 	}
 	// Can only create temporaries for certain types
 	else if ((Type & (MCT_Float | MCT_VTPageTableResult)) || Type == MCT_ShadingModel || Type == MCT_MaterialAttributes || Type == MCT_Strata)
@@ -2876,20 +3132,23 @@ int32 FHLSLMaterialTranslator::AddCodeChunkInner(uint64 Hash, const TCHAR* Forma
 		{
 			if ((*CurrentScopeChunks)[i].Hash == Hash)
 			{
-				return i;
+				CodeIndex = i;
+				break;
 			}
 		}
 
-		const int32 CodeIndex = CurrentScopeChunks->Num();
-		// Allocate a local variable name
-		const FString SymbolName = CreateSymbolName(TEXT("Local"));
-		// Construct the definition string which stores the result in a temporary and adds a newline for readability
-		const FString LocalVariableDefinitionFinite = FString("	") + HLSLTypeString(Type) + TEXT(" ") + SymbolName + TEXT(" = ") + FormattedCode + TEXT(";") + LINE_TERMINATOR;
-		// Construct the definition string which stores the result in a temporary and adds a newline for readability
-		const FString LocalVariableDefinitionAnalytic = FString("	") + HLSLTypeString(Type) + TEXT(" ") + SymbolName + TEXT(" = ") + FormattedCode + TEXT(";") + LINE_TERMINATOR;
-		// Adding a code chunk that creates a local variable
-		new(*CurrentScopeChunks) FShaderCodeChunk(Hash, *LocalVariableDefinitionFinite, *LocalVariableDefinitionAnalytic,SymbolName,Type,false);
-		return CodeIndex;
+		if (CodeIndex == INDEX_NONE)
+		{
+			CodeIndex = CurrentScopeChunks->Num();
+			// Allocate a local variable name
+			const FString SymbolName = CreateSymbolName(TEXT("Local"));
+			// Construct the definition string which stores the result in a temporary and adds a newline for readability
+			const FString LocalVariableDefinitionFinite = FString("	") + HLSLTypeString(Type) + TEXT(" ") + SymbolName + TEXT(" = ") + FormattedCode + TEXT(";") + LINE_TERMINATOR;
+			// Construct the definition string which stores the result in a temporary and adds a newline for readability
+			const FString LocalVariableDefinitionAnalytic = FString("	") + HLSLTypeString(Type) + TEXT(" ") + SymbolName + TEXT(" = ") + FormattedCode + TEXT(";") + LINE_TERMINATOR;
+			// Adding a code chunk that creates a local variable
+			new(*CurrentScopeChunks) FShaderCodeChunk(Hash, *LocalVariableDefinitionFinite, *LocalVariableDefinitionAnalytic,SymbolName,Type,false);
+		}
 	}
 	else
 	{
@@ -2902,9 +3161,10 @@ int32 FHLSLMaterialTranslator::AddCodeChunkInner(uint64 Hash, const TCHAR* Forma
 		{
 			return Errorf(TEXT("Operation not supported on a Static Bool"));
 		}
-			
-		return INDEX_NONE;
 	}
+	
+	AddCodeChunkToCurrentScope(CodeIndex);
+	return CodeIndex;
 }
 
 /** Adds an already formatted inline or referenced code chunk, and notes the derivative status. */
@@ -2926,12 +3186,14 @@ int32 FHLSLMaterialTranslator::AddCodeChunkInnerDeriv(uint64 Hash, const TCHAR*
 		return INDEX_NONE;
 	}
 
+	check(Type != MCT_VoidStatement);
+
+	int32 CodeIndex = INDEX_NONE;
 	if (bInlined)
 	{
-		const int32 CodeIndex = CurrentScopeChunks->Num();
+		CodeIndex = CurrentScopeChunks->Num();
 		// Adding an inline code chunk, the definition will be the code to inline
 		new(*CurrentScopeChunks) FShaderCodeChunk(Hash, FormattedCodeFinite,FormattedCodeAnalytic,TEXT(""),Type,true,DerivativeStatus);
-		return CodeIndex;
 	}
 	// Can only create temporaries for certain types
 	else if ((Type & (MCT_Float | MCT_VTPageTableResult)) || Type == MCT_ShadingModel || Type == MCT_MaterialAttributes || Type == MCT_Strata)
@@ -2941,20 +3203,23 @@ int32 FHLSLMaterialTranslator::AddCodeChunkInnerDeriv(uint64 Hash, const TCHAR*
 		{
 			if ((*CurrentScopeChunks)[i].Hash == Hash)
 			{
-				return i;
+				CodeIndex = i;
+				break;
 			}
 		}
 
-		const int32 CodeIndex = CurrentScopeChunks->Num();
-		// Allocate a local variable name
-		const FString SymbolName = CreateSymbolName(TEXT("Local"));
-		// Construct the definition string which stores the result in a temporary and adds a newline for readability
-		const FString LocalVariableDefinitionFinite = FString("	") + HLSLTypeString(Type) + TEXT(" ") + SymbolName + TEXT(" = ") + FormattedCodeFinite + TEXT(";") + LINE_TERMINATOR;
-		// Analytic version too
-		const FString LocalVariableDefinitionAnalytic = FString("	") + HLSLTypeStringDeriv(Type,DerivativeStatus) + TEXT(" ") + SymbolName + TEXT(" = (") + FormattedCodeAnalytic + TEXT("); // deriv") + LINE_TERMINATOR;
-		// Adding a code chunk that creates a local variable
-		new(*CurrentScopeChunks) FShaderCodeChunk(Hash, *LocalVariableDefinitionFinite, *LocalVariableDefinitionAnalytic,SymbolName,Type,false,DerivativeStatus);
-		return CodeIndex;
+		if(CodeIndex == INDEX_NONE)
+		{
+			CodeIndex = CurrentScopeChunks->Num();
+			// Allocate a local variable name
+			const FString SymbolName = CreateSymbolName(TEXT("Local"));
+			// Construct the definition string which stores the result in a temporary and adds a newline for readability
+			const FString LocalVariableDefinitionFinite = FString("	") + HLSLTypeString(Type) + TEXT(" ") + SymbolName + TEXT(" = ") + FormattedCodeFinite + TEXT(";") + LINE_TERMINATOR;
+			// Analytic version too
+			const FString LocalVariableDefinitionAnalytic = FString("	") + HLSLTypeStringDeriv(Type,DerivativeStatus) + TEXT(" ") + SymbolName + TEXT(" = (") + FormattedCodeAnalytic + TEXT("); // deriv") + LINE_TERMINATOR;
+			// Adding a code chunk that creates a local variable
+			new(*CurrentScopeChunks) FShaderCodeChunk(Hash, *LocalVariableDefinitionFinite, *LocalVariableDefinitionAnalytic,SymbolName,Type,false,DerivativeStatus);
+		}
 	}
 	else
 	{
@@ -2968,8 +3233,10 @@ int32 FHLSLMaterialTranslator::AddCodeChunkInnerDeriv(uint64 Hash, const TCHAR*
 			return Errorf(TEXT("Operation not supported on a Static Bool"));
 		}
 
-		return INDEX_NONE;
 	}
+
+	AddCodeChunkToCurrentScope(CodeIndex);
+	return CodeIndex;
 }
 
 /** 
@@ -3130,6 +3397,7 @@ int32 FHLSLMaterialTranslator::AddUniformExpressionInner(uint64 Hash, FMaterialU
 				{
 					delete UniformExpression;
 					// Reuse the entry in CurrentScopeChunks
+					AddCodeChunkToCurrentScope(ChunkIndex);
 					return ChunkIndex;
 				}
 			}
@@ -3185,6 +3453,7 @@ int32 FHLSLMaterialTranslator::AddUniformExpressionInner(uint64 Hash, FMaterialU
 		new(UniformExpressions) FShaderCodeChunk(Hash, UniformExpression, FormattedCode, FormattedCode, Type, MEADS_Zero);
 	}
 
+	AddCodeChunkToCurrentScope(ReturnIndex);
 	return ReturnIndex;
 }
 
@@ -3622,9 +3891,11 @@ int32 FHLSLMaterialTranslator::CallExpression(FMaterialExpressionKey ExpressionK
 	FMaterialFunctionCompileState* CurrentFunctionState = CurrentFunctionStack.Last();
 
 	int32* ExistingCodeIndex = CurrentFunctionState->ExpressionCodeMap.Find(ExpressionKey);
+	int32 Result = INDEX_NONE;
 	if (ExistingCodeIndex)
 	{
-		return *ExistingCodeIndex;
+		Result = *ExistingCodeIndex;
+		AddCodeChunkToCurrentScope(Result);
 	}
 	else
 	{
@@ -3650,7 +3921,9 @@ int32 FHLSLMaterialTranslator::CallExpression(FMaterialExpressionKey ExpressionK
 			FunctionCall->SetSharedCompileState(SharedFunctionState);
 		}
 
-		int32 Result = ExpressionKey.Expression->Compile(Compiler, ExpressionKey.OutputIndex);
+		ReferencedCodeChunks.Reset();
+
+		Result = ExpressionKey.Expression->Compile(Compiler, ExpressionKey.OutputIndex);
 
 		// Restore state
 		if (FunctionCall)
@@ -3667,7 +3940,57 @@ int32 FHLSLMaterialTranslator::CallExpression(FMaterialExpressionKey ExpressionK
 		// Cache the translation.
 		CurrentFunctionStack.Last()->ExpressionCodeMap.Add(ExpressionKey,Result);
 
-		return Result;
+		if (Result != INDEX_NONE)
+		{
+			FShaderCodeChunk& ResultChunk = (*CurrentScopeChunks)[Result];
+			ResultChunk.ReferencedCodeChunks = MoveTemp(ReferencedCodeChunks);
+		}
+		ReferencedCodeChunks.Reset();
+	}
+
+	return Result;
+}
+
+void FHLSLMaterialTranslator::AddCodeChunkToCurrentScope(int32 ChunkIndex)
+{
+	if (ChunkIndex != INDEX_NONE && ScopeStack.Num() > 0)
+	{
+		const int32 CurrentScopeIndex = ScopeStack.Last();
+		const FShaderCodeChunk& CurrentScope = (*CurrentScopeChunks)[CurrentScopeIndex];
+
+		FShaderCodeChunk& Chunk = (*CurrentScopeChunks)[ChunkIndex];
+		if (Chunk.DeclaredScopeIndex == INDEX_NONE)
+		{
+			check(Chunk.UsedScopeIndex == INDEX_NONE);
+			Chunk.DeclaredScopeIndex = CurrentScopeIndex;
+			Chunk.UsedScopeIndex = CurrentScopeIndex;
+			Chunk.ScopeLevel = CurrentScope.ScopeLevel + 1;
+		}
+		else if(Chunk.UsedScopeIndex != CurrentScopeIndex)
+		{
+			// Find the most derived scope that's shared by the current scope, and the scope this code was previously referenced from
+			int32 ScopeIndex0 = CurrentScopeIndex;
+			int32 ScopeIndex1 = Chunk.UsedScopeIndex;
+			while (ScopeIndex0 != ScopeIndex1)
+			{
+				const FShaderCodeChunk& Scope0 = (*CurrentScopeChunks)[ScopeIndex0];
+				const FShaderCodeChunk& Scope1 = (*CurrentScopeChunks)[ScopeIndex1];
+				if (Scope0.ScopeLevel > Scope1.ScopeLevel)
+				{
+					check(Scope0.UsedScopeIndex != INDEX_NONE);
+					ScopeIndex0 = Scope0.UsedScopeIndex;
+				}
+				else
+				{
+					check(Scope1.UsedScopeIndex != INDEX_NONE);
+					ScopeIndex1 = Scope1.UsedScopeIndex;
+				}
+			}
+
+			const FShaderCodeChunk& Scope = (*CurrentScopeChunks)[ScopeIndex0];
+			Chunk.UsedScopeIndex = ScopeIndex0;
+			Chunk.ScopeLevel = Scope.ScopeLevel + 1;
+		}
 	}
 }
 
@@ -8428,6 +8751,184 @@ int32 FHLSLMaterialTranslator::ShadingModel(EMaterialShadingModel InSelectedShad
 	return AddInlinedCodeChunk(MCT_ShadingModel, TEXT("%d"), InSelectedShadingModel);
 }
 
+int32 FHLSLMaterialTranslator::DefaultMaterialAttributes()
+{
+	return AddInlinedCodeChunk(MCT_MaterialAttributes, TEXT("DefaultMaterialAttributes"));
+}
+
+int32 FHLSLMaterialTranslator::SetMaterialAttribute(int32 MaterialAttributes, int32 Value, const FGuid& AttributeID)
+{
+	const EMaterialProperty Property = FMaterialAttributeDefinitionMap::GetProperty(AttributeID);
+	const FString PropertyName = FMaterialAttributeDefinitionMap::GetAttributeName(AttributeID);
+	const EMaterialValueType PropertyType = FMaterialAttributeDefinitionMap::GetValueType(AttributeID);
+	const EShaderFrequency Frequency = FMaterialAttributeDefinitionMap::GetShaderFrequency(AttributeID);
+	
+	if (MaterialAttributes == INDEX_NONE ||
+		GetParameterType(MaterialAttributes) != MCT_MaterialAttributes)
+	{
+		return Error(TEXT("Expected MaterialAttributes"));
+	}
+
+	if (Frequency != ShaderFrequency)
+	{
+		return Errorf(TEXT("Can't set material attribute %s from shader stage %d"), *PropertyName, ShaderFrequency);
+	}
+
+	const int32 CastValue = ValidCast(Value, PropertyType);
+	FMaterialUniformExpression* UniformExpression = GetParameterUniformExpression(CastValue);
+	bool bSetDefaultValue = false;
+	if (UniformExpression && UniformExpression->IsConstant())
+	{
+		FMaterialRenderContext Context(nullptr, *Material, nullptr);
+		FLinearColor ConstantValue(ForceInitToZero);
+		UniformExpression->GetNumberValue(Context, ConstantValue);
+		const FVector4 DefaultValue = FMaterialAttributeDefinitionMap::GetDefaultValue(AttributeID);
+		bSetDefaultValue = (ConstantValue == FLinearColor(DefaultValue));
+	}
+
+	uint64 AttributeMask = GetParameterMaterialAttributeMask(MaterialAttributes);
+	if (bSetDefaultValue)
+	{
+		if (!(AttributeMask & (1ull << Property)))
+		{
+			// Setting default value to an input that already has the default value set, this is a NOP
+			return MaterialAttributes;
+		}
+		// Otherwise, explicitly set the value back to default
+		AttributeMask &= ~(1ull << Property);
+	}
+	else
+	{
+		// Setting a non-default value
+		AttributeMask |= (1ull << Property);
+	}
+
+	const int32 Result = AddCodeChunk(MCT_MaterialAttributes, TEXT("FMaterialAttributes_Set%s(%s, %s)"),
+		*PropertyName,
+		*GetParameterCode(MaterialAttributes),
+		*GetParameterCode(CastValue));
+
+	SetParameterMaterialAttributes(Result, AttributeMask);
+
+	return Result;
+}
+
+int32 FHLSLMaterialTranslator::BeginScope()
+{
+	const int32 Result = AddCodeChunk(MCT_VoidStatement, TEXT(""));
+	ScopeStack.Add(Result);
+	return Result;
+}
+
+int32 FHLSLMaterialTranslator::BeginScope_If(int32 Condition)
+{
+	const int32 ConditionAsFloat = ForceCast(Condition, MCT_Float1);
+	if (ConditionAsFloat == INDEX_NONE)
+	{
+		return INDEX_NONE;
+	}
+
+	const int32 Result = AddCodeChunk(MCT_VoidStatement, TEXT("if (%s != 0.0f)"), *GetParameterCode(ConditionAsFloat));
+	ScopeStack.Add(Result);
+
+	return Result;
+}
+
+int32 FHLSLMaterialTranslator::BeginScope_Else()
+{
+	const int32 Result = AddCodeChunk(MCT_VoidStatement, TEXT("else"));
+	ScopeStack.Add(Result);
+
+	return Result;
+}
+
+int32 FHLSLMaterialTranslator::BeginScope_For(const UMaterialExpression* Expression, int32 StartIndex, int32 EndIndex, int32 IndexStep)
+{
+	const int32 StartIndexAsFloat = ForceCast(StartIndex, MCT_Float1);
+	const int32 EndIndexAsFloat = ForceCast(EndIndex, MCT_Float1);
+	const int32 IndexStepAsFloat = ForceCast(IndexStep, MCT_Float1);
+	if (StartIndexAsFloat == INDEX_NONE || EndIndexAsFloat == INDEX_NONE || IndexStepAsFloat == INDEX_NONE)
+	{
+		return INDEX_NONE;
+	}
+
+	const int32 LoopIndex = NumForLoops[ShaderFrequency]++;
+	check(!ForLoopMap[ShaderFrequency].Contains(Expression));
+	ForLoopMap[ShaderFrequency].Add(Expression, LoopIndex);
+
+	const int32 Result = AddCodeChunk(MCT_VoidStatement, TEXT("for (float ForLoopCounter%d = %s; ForLoopCounter%d < %s; ForLoopCounter%d += %s)"),
+		LoopIndex, *GetParameterCode(StartIndexAsFloat),
+		LoopIndex, *GetParameterCode(EndIndexAsFloat),
+		LoopIndex, *GetParameterCode(IndexStepAsFloat));
+	ScopeStack.Add(Result);
+	return Result;
+}
+
+int32 FHLSLMaterialTranslator::EndScope()
+{
+	return ScopeStack.Pop(false);
+}
+
+int32 FHLSLMaterialTranslator::ForLoopIndex(const UMaterialExpression* Expression)
+{
+	const int32* Result = ForLoopMap[ShaderFrequency].Find(Expression);
+	if (Result)
+	{
+		return AddInlinedCodeChunk(MCT_Float1, TEXT("ForLoopCounter%d"), *Result);
+	}
+	return Error(TEXT("Expression is not a for-loop"));
+}
+
+int32 FHLSLMaterialTranslator::ReturnMaterialAttributes(int32 MaterialAttributes)
+{
+	if (MaterialAttributes == INDEX_NONE ||
+		GetParameterType(MaterialAttributes) != MCT_MaterialAttributes)
+	{
+		return Error(TEXT("Expected MaterialAttributes"));
+	}
+
+	const uint64 AttributeMask = GetParameterMaterialAttributeMask(MaterialAttributes);
+	MaterialAttributesReturned[ShaderFrequency] |= AttributeMask;
+
+	return AddCodeChunk(MCT_VoidStatement, TEXT("return %s;"), *GetParameterCode(MaterialAttributes));
+}
+
+int32 FHLSLMaterialTranslator::SetLocal(const FName& LocalName, int32 Value)
+{
+	const int32 ValueAsFloat = ForceCast(Value, MCT_Float1);
+	if (ValueAsFloat == INDEX_NONE)
+	{
+		return INDEX_NONE;
+	}
+
+	FMaterialLocalVariableEntry& Entry = LocalVariables[ShaderFrequency].FindOrAdd(LocalName);
+	if (Entry.DeclarationCodeIndex == INDEX_NONE)
+	{
+		Entry.Name = LocalName.ToString();
+		Entry.DeclarationCodeIndex = AddCodeChunk(MCT_VoidStatement, TEXT("float %s = %s;"), *Entry.Name, *GetParameterCode(ValueAsFloat));
+		return Entry.DeclarationCodeIndex;
+	}
+	else
+	{
+		// ensure the declaration is visible in the current scope
+		AddCodeChunkToCurrentScope(Entry.DeclarationCodeIndex);
+		return AddCodeChunk(MCT_VoidStatement, TEXT("%s = %s;"), *Entry.Name, *GetParameterCode(ValueAsFloat));
+	}
+}
+
+int32 FHLSLMaterialTranslator::GetLocal(const FName& LocalName)
+{
+	const FMaterialLocalVariableEntry* Entry = LocalVariables[ShaderFrequency].Find(LocalName);
+	if (!Entry)
+	{
+		return Errorf(TEXT("Local %s used before being set"), *LocalName.ToString());
+	}
+
+	// ensure the declaration is visible in the current scope
+	AddCodeChunkToCurrentScope(Entry->DeclarationCodeIndex);
+	return AddInlinedCodeChunk(MCT_Float1, TEXT("%s"), *Entry->Name);
+}
+
 void FHLSLMaterialTranslator::StrataCompilationInfoRegisterCodeChunk(int32 CodeChunk, FStrataMaterialCompilationInfo& StrataMaterialCompilationInfo)
 {
 	check(CodeChunk != INDEX_NONE);
diff --git a/Engine/Source/Runtime/Engine/Private/Materials/HLSLMaterialTranslator.h b/Engine/Source/Runtime/Engine/Private/Materials/HLSLMaterialTranslator.h
--- a/Engine/Source/Runtime/Engine/Private/Materials/HLSLMaterialTranslator.h
+++ b/Engine/Source/Runtime/Engine/Private/Materials/HLSLMaterialTranslator.h
@@ -91,6 +91,10 @@ struct FShaderCodeChunk
 	 * By default this is simply the hash of the code string
 	 */
 	uint64 Hash;
+
+
+	uint64 MaterialAttributeMask;
+
 	/** 
 	 * Definition string of the code chunk. 
 	 * If !bInline && !UniformExpression || UniformExpression->IsConstant(), this is the definition of a local variable named by SymbolName.
@@ -109,7 +113,18 @@ struct FShaderCodeChunk
 	FString SymbolName;
 	/** Reference to a uniform expression, if this code chunk has one. */
 	TRefCountPtr<FMaterialUniformExpression> UniformExpression;
+
+	/** All the chunks that are scoped under this chunk.  This is populated once translation is complete, rather than on-the-fly, as a chunk's parent scope may change during translation */
+	TArray<int32> ScopedChunks;
+
+	TArray<int32> ReferencedCodeChunks;
+
 	EMaterialValueType Type;
+
+	int32 DeclaredScopeIndex;
+	int32 UsedScopeIndex;
+	int32 ScopeLevel;
+
 	/** Whether the code chunk should be inlined or not.  If true, SymbolName is empty and Definition contains the code to inline. */
 	bool bInline;
 	/** The status of partial derivatives at this expression. **/
@@ -131,11 +146,15 @@ struct FShaderCodeChunk
 	FShaderCodeChunk(uint64 InHash, const TCHAR* InDefinitionFinite,const TCHAR* InDefinitionAnalytic,const FString& InSymbolName,EMaterialValueType InType,bool bInInline,
 					 EMaterialExpressionAnalyticDerivativeStatus InDerivativeStatus = MEADS_NotAware):
 		Hash(InHash),
+		MaterialAttributeMask(0u),
 		DefinitionFinite(InDefinitionFinite),
 		DefinitionAnalytic(InDefinitionAnalytic),
 		SymbolName(InSymbolName),
 		UniformExpression(NULL),
 		Type(InType),
+		DeclaredScopeIndex(INDEX_NONE),
+		UsedScopeIndex(INDEX_NONE),
+		ScopeLevel(0),
 		bInline(bInInline),
 		DerivativeStatus(InDerivativeStatus)
 	{}
@@ -144,10 +163,14 @@ struct FShaderCodeChunk
 	FShaderCodeChunk(uint64 InHash, FMaterialUniformExpression* InUniformExpression,const TCHAR* InDefinitionFinite,const TCHAR* InDefinitionAnalytic,EMaterialValueType InType,
 					 EMaterialExpressionAnalyticDerivativeStatus InDerivativeStatus = MEADS_NotAware):
 		Hash(InHash),
+		MaterialAttributeMask(0u),
 		DefinitionFinite(InDefinitionFinite),
 		DefinitionAnalytic(InDefinitionAnalytic),
 		UniformExpression(InUniformExpression),
 		Type(InType),
+		DeclaredScopeIndex(INDEX_NONE),
+		UsedScopeIndex(INDEX_NONE),
+		ScopeLevel(0),
 		bInline(false),
 		DerivativeStatus(InDerivativeStatus)
 	{}
@@ -246,6 +269,12 @@ struct FMaterialDerivativeVariation
 
 };
 
+struct FMaterialLocalVariableEntry
+{
+	FString Name;
+	int32 DeclarationCodeIndex = INDEX_NONE;
+};
+
 class FHLSLMaterialTranslator : public FMaterialCompiler
 {
 protected:
@@ -286,15 +315,28 @@ protected:
 	/** Feature level being compiled for. */
 	ERHIFeatureLevel::Type FeatureLevel;
 
+	FString TranslatedAttributesCodeChunks[SF_NumFrequencies];
+
+	uint64 MaterialAttributesReturned[SF_NumFrequencies];
+
 	/** Line number of the #line in MaterialTemplate.usf */
 	int32 MaterialTemplateLineNumber;
 
 	/** Contents of the MaterialTemplate.usf file */
 	FString MaterialTemplate;
 
+	TArray<int32> ScopeStack;
+
+	TArray<int32> ReferencedCodeChunks;
+
 	// Array of code chunks per material property
 	TArray<FShaderCodeChunk> SharedPropertyCodeChunks[SF_NumFrequencies];
 
+	TMap<const UMaterialExpression*, int32> ForLoopMap[SF_NumFrequencies];
+	int32 NumForLoops[SF_NumFrequencies];
+
+	TMap<FName, FMaterialLocalVariableEntry> LocalVariables[SF_NumFrequencies];
+
 	// Uniform expressions used across all material properties
 	TArray<FShaderCodeChunk> UniformExpressions;
 
@@ -400,11 +442,12 @@ protected:
 	/** True if this material write anisotropy material property */
 	uint32 bUsesAnisotropy : 1;
 
-	/** 
-	 * True if the material is detected as a strata material at compile time.
+	/** True if the material is detected as a strata material at compile time.
 	 * This is decoupled from runtime FMaterialResource::IsStrataMaterial but practically fine since this is only temporary until Strata is the main shading system. Only really used at runtime for translucency dual source blending.
 	 */
 	uint32 bMaterialIsStrata : 1; // 
+	
+	uint32 bEnableExecutionFlow : 1;
 
 	/** Tracks the texture coordinates used by this material. */
 	TBitArray<> AllocatedUserTexCoords;
@@ -498,12 +541,19 @@ protected:
 
 	uint64 GetParameterHash(int32 Index);
 
+	uint64 GetParameterMaterialAttributeMask(int32 Index);
+	void SetParameterMaterialAttributes(int32 Index, uint64 Mask);
+
 	/** Creates a string of all definitions needed for the given material input. */
 	FString GetDefinitions(TArray<FShaderCodeChunk>& CodeChunks, int32 StartChunk, int32 EndChunk, ECompiledPartialDerivativeVariation Variation) const;
 
 	// GetFixedParameterCode
 	void GetFixedParameterCode(int32 StartChunk, int32 EndChunk, int32 ResultIndex, TArray<FShaderCodeChunk>& CodeChunks, FString& OutDefinitions, FString& OutValue, ECompiledPartialDerivativeVariation Variation);
 
+	void LinkParentScopes(TArray<FShaderCodeChunk>& CodeChunks);
+
+	void GetScopeCode(int32 IndentLevel, int32 ScopeChunkIndex, const TArray<FShaderCodeChunk>& CodeChunks, TSet<int32>& EmittedChunks, FString& OutValue);
+
 	void GetFixedParameterCode(int32 ResultIndex, TArray<FShaderCodeChunk>& CodeChunks, FString& OutDefinitions, FString& OutValue, ECompiledPartialDerivativeVariation Variation);
 
 	/** Used to get a user friendly type from EMaterialValueType */
@@ -527,6 +577,8 @@ protected:
 	/** Creates a unique symbol name and adds it to the symbol list. */
 	FString CreateSymbolName(const TCHAR* SymbolNameHint);
 
+	void AddCodeChunkToCurrentScope(int32 ChunkIndex);
+
 	/** Adds an already formatted inline or referenced code chunk */
 	int32 AddCodeChunkInner(uint64 Hash, const TCHAR* FormattedCode, EMaterialValueType Type, bool bInlined);
 
@@ -918,6 +970,21 @@ protected:
 
 	virtual int32 VirtualTextureOutput(uint8 MaterialAttributeMask) override;
 
+	// Material attributes
+	virtual int32 DefaultMaterialAttributes() override;
+	virtual int32 SetMaterialAttribute(int32 MaterialAttributes, int32 Value, const FGuid& AttributeID) override;
+
+	// Exec
+	virtual int32 BeginScope() override;
+	virtual int32 BeginScope_If(int32 Condition) override;
+	virtual int32 BeginScope_Else() override;
+	virtual int32 BeginScope_For(const UMaterialExpression* Expression, int32 StartIndex, int32 EndIndex, int32 IndexStep) override;
+	virtual int32 EndScope() override;
+	virtual int32 ForLoopIndex(const UMaterialExpression* Expression) override;
+	virtual int32 ReturnMaterialAttributes(int32 MaterialAttributes) override;
+	virtual int32 SetLocal(const FName& LocalName, int32 Value) override;
+	virtual int32 GetLocal(const FName& LocalName) override;
+
 	// Strata
 	virtual int32 StrataCreateAndRegisterNullMaterial() override;
 	virtual int32 StrataSlabBSDF(
diff --git a/Engine/Source/Runtime/Engine/Private/Materials/Material.cpp b/Engine/Source/Runtime/Engine/Private/Materials/Material.cpp
--- a/Engine/Source/Runtime/Engine/Private/Materials/Material.cpp
+++ b/Engine/Source/Runtime/Engine/Private/Materials/Material.cpp
@@ -44,6 +44,7 @@
 #include "Materials/MaterialExpressionSceneColor.h"
 #include "Materials/MaterialExpressionShadingModel.h"
 #include "Materials/MaterialExpressionTransform.h"
+#include "Materials/MaterialExpressionExecBegin.h"
 #include "Materials/MaterialFunction.h"
 #include "Materials/MaterialFunctionInstance.h"
 #include "Materials/MaterialExpressionMaterialFunctionCall.h"
@@ -121,6 +122,12 @@ static TAutoConsoleVariable<int32> CVarMaterialParameterLegacyChecks(
 	TEXT("Note that this can be slow"),
 	ECVF_RenderThreadSafe);
 
+static TAutoConsoleVariable<int32> CVarMaterialEnableControlFlow(
+	TEXT("r.MaterialEnableControlFlow"),
+	0,
+	TEXT("Allows experemental control flow to be used in the material editor.\n"),
+	ECVF_RenderThreadSafe | ECVF_ReadOnly);
+
 #if WITH_EDITOR
 const FMaterialsWithDirtyUsageFlags FMaterialsWithDirtyUsageFlags::DefaultAnnotation;
 
@@ -229,7 +236,7 @@ int32 FMaterialResource::CompilePropertyAndSetMaterialProperty(EMaterialProperty
 			break;
 
 		case MP_MaterialAttributes:
-			Ret = INDEX_NONE;
+			Ret = MaterialInterface->CompileProperty(Compiler, Property);
 			break;
 
 		default:
@@ -260,8 +267,14 @@ int32 FMaterialResource::CompilePropertyAndSetMaterialProperty(EMaterialProperty
 		}
 	}
 
-	// Output should always be the right type for this property
-	return Compiler->ForceCast(Ret, AttributeType);
+	// MaterialAttributes are expected to give a void statement, don't need to cast that
+	if (Property != MP_MaterialAttributes)
+	{
+		// Output should always be the right type for this property
+		Ret = Compiler->ForceCast(Ret, AttributeType);
+	}
+	return Ret;
+
 #else // WITH_EDITOR
 	check(0); // This is editor-only function
 	return INDEX_NONE;
@@ -4086,6 +4099,15 @@ void UMaterial::PropagateDataToMaterialProxy()
 	UpdateMaterialRenderProxy(*DefaultMaterialInstance);
 }
 
+bool UMaterial::IsCompiledWithExecutionFlow() const
+{
+	if (bEnableExecWire)
+	{
+		return CVarMaterialEnableControlFlow.GetValueOnAnyThread() != 0;
+	}
+	return false;
+}
+
 #if WITH_EDITOR
 void UMaterial::BeginCacheForCookedPlatformData( const ITargetPlatform *TargetPlatform )
 {
@@ -4410,6 +4432,23 @@ void UMaterial::PostEditChangePropertyInternal(FPropertyChangedEvent& PropertyCh
 		}
 	}
 
+	if (PropertyChangedEvent.GetPropertyName() == GET_MEMBER_NAME_CHECKED(UMaterial, bEnableExecWire))
+	{
+		if (IsCompiledWithExecutionFlow())
+		{
+			check(!ExpressionExecBegin);
+			ExpressionExecBegin = NewObject<UMaterialExpressionExecBegin>(this);
+			ExpressionExecBegin->Material = this;
+			Expressions.Add(ExpressionExecBegin);
+		}
+		else
+		{
+			check(ExpressionExecBegin);
+			verify(Expressions.Remove(ExpressionExecBegin) == 1);
+			ExpressionExecBegin = nullptr;
+		}
+	}
+
 	TranslucencyDirectionalLightingIntensity = FMath::Clamp(TranslucencyDirectionalLightingIntensity, .1f, 10.0f);
 
 	// Don't want to recompile after a duplicate because it's just been done by PostLoad, nor during interactive changes to prevent constant recompilation while spinning properties.
@@ -5937,6 +5976,12 @@ int32 UMaterial::CompilePropertyEx( FMaterialCompiler* Compiler, const FGuid& At
 {
 	const EMaterialProperty Property = FMaterialAttributeDefinitionMap::GetProperty(AttributeID);
 
+	if (IsCompiledWithExecutionFlow())
+	{
+		check(ExpressionExecBegin);
+		return ExpressionExecBegin->Compile(Compiler, UMaterialExpression::CompileExecutionOutputIndex);
+	}
+
 	if( bUseMaterialAttributes && MP_DiffuseColor != Property && MP_SpecularColor != Property )
 	{
 		return MaterialAttributes.CompileWithDefault(Compiler, AttributeID);
diff --git a/Engine/Source/Runtime/Engine/Private/Materials/MaterialExpressions.cpp b/Engine/Source/Runtime/Engine/Private/Materials/MaterialExpressions.cpp
--- a/Engine/Source/Runtime/Engine/Private/Materials/MaterialExpressions.cpp
+++ b/Engine/Source/Runtime/Engine/Private/Materials/MaterialExpressions.cpp
@@ -250,6 +250,12 @@
 #include "Materials/MaterialExpressionStrata.h"
 #include "Materials/MaterialUniformExpressions.h"
 #include "Materials/MaterialExpressionSamplePhysicsField.h"
+#include "Materials/MaterialExpressionExecBegin.h"
+#include "Materials/MaterialExpressionReturnMaterialAttributes.h"
+#include "Materials/MaterialExpressionIfThenElse.h"
+#include "Materials/MaterialExpressionForLoop.h"
+#include "Materials/MaterialExpressionSetLocal.h"
+#include "Materials/MaterialExpressionGetLocal.h"
 #include "EditorSupportDelegates.h"
 #include "MaterialCompiler.h"
 #if WITH_EDITOR
@@ -408,6 +414,12 @@ void GetMaterialValueTypeDescriptions(const uint32 MaterialValueType, TArray<FTe
 
 bool CanConnectMaterialValueTypes(const uint32 InputType, const uint32 OutputType)
 {
+	if ((InputType & MCT_Execution) || (OutputType & MCT_Execution))
+	{
+		// exec pins can only connect to other exec pins
+		return InputType == OutputType;
+	}
+
 	if (InputType & MCT_Unknown)
 	{
 		// can plug anything into unknown inputs
@@ -1116,6 +1128,22 @@ uint32 UMaterialExpression::GetOutputType(int32 OutputIndex)
 	}
 }
 
+FExpressionInput* UMaterialExpression::GetExecInput()
+{
+	return nullptr;
+}
+
+bool UMaterialExpression::IsExecInputConnection(UMaterialExpression* InExpression, int32 InOutputIndex)
+{
+	FExpressionInput* Input = GetExecInput();
+	if (Input)
+	{
+		FExpressionInput TracedInput = Input->GetTracedInput();
+		return TracedInput.Expression == InExpression && TracedInput.OutputIndex == InOutputIndex;
+	}
+	return false;
+}
+
 int32 UMaterialExpression::GetWidth() const
 {
 	return ME_STD_THUMBNAIL_SZ+(ME_STD_BORDER*2);
@@ -6290,27 +6318,62 @@ int32 UMaterialExpressionSetMaterialAttributes::Compile(class FMaterialCompiler*
 	}
 
 	// Compile attribute
-	const FGuid AttributeID = Compiler->GetMaterialAttribute();
-	FExpressionInput* AttributeInput = nullptr;
-
-	int32 PinIndex;
-	if (AttributeSetTypes.Find(AttributeID, PinIndex))
+	const FGuid CompilingAttributeID = Compiler->GetMaterialAttribute();
+	if (CompilingAttributeID == FMaterialAttributeDefinitionMap::GetID(MP_MaterialAttributes))
 	{
-		checkf(PinIndex + 1 < Inputs.Num(), TEXT("Requested non-existent pin."));
-		AttributeInput = &Inputs[PinIndex + 1];
-	}
+		int32 Result = INDEX_NONE;
+		if (Inputs[0].GetTracedInput().Expression)
+		{
+			Result = Inputs[0].GetTracedInput().Compile(Compiler);
+		}
+		else
+		{
+			Result = Compiler->DefaultMaterialAttributes();
+		}
 
-	if (AttributeInput && AttributeInput->GetTracedInput().Expression)
-	{
-		EMaterialValueType ValueType = FMaterialAttributeDefinitionMap::GetValueType(AttributeID);
-		return Compiler->ValidCast(AttributeInput->GetTracedInput().Compile(Compiler), ValueType);
+		for (int32 PinIndex = 0; PinIndex < AttributeSetTypes.Num(); ++PinIndex)
+		{
+			const FExpressionInput& AttributeInput = Inputs[PinIndex + 1];
+			if (AttributeInput.GetTracedInput().Expression)
+			{
+				const FGuid& AttributeID = AttributeSetTypes[PinIndex];
+				// Only compile code to set attributes of the current shader frequency
+				const EShaderFrequency AttributeFrequency = FMaterialAttributeDefinitionMap::GetShaderFrequency(AttributeID);
+				if (AttributeFrequency == Compiler->GetCurrentShaderFrequency())
+				{
+					const int32 AttributeResult = AttributeInput.GetTracedInput().Compile(Compiler);
+					if (AttributeResult != INDEX_NONE)
+					{
+						Result = Compiler->SetMaterialAttribute(Result, AttributeResult, AttributeID);
+					}
+				}
+			}
+		}
+		return Result;
 	}
-	else if (Inputs[0].GetTracedInput().Expression)
+	else
 	{
-		return Inputs[0].GetTracedInput().Compile(Compiler);
-	}
+		FExpressionInput* AttributeInput = nullptr;
 
-	return FMaterialAttributeDefinitionMap::CompileDefaultExpression(Compiler, AttributeID);
+		int32 PinIndex;
+		if (AttributeSetTypes.Find(CompilingAttributeID, PinIndex))
+		{
+			checkf(PinIndex + 1 < Inputs.Num(), TEXT("Requested non-existent pin."));
+			AttributeInput = &Inputs[PinIndex + 1];
+		}
+
+		if (AttributeInput && AttributeInput->GetTracedInput().Expression)
+		{
+			EMaterialValueType ValueType = FMaterialAttributeDefinitionMap::GetValueType(CompilingAttributeID);
+			return Compiler->ValidCast(AttributeInput->GetTracedInput().Compile(Compiler), ValueType);
+		}
+		else if (Inputs[0].GetTracedInput().Expression)
+		{
+			return Inputs[0].GetTracedInput().Compile(Compiler);
+		}
+
+		return FMaterialAttributeDefinitionMap::CompileDefaultExpression(Compiler, CompilingAttributeID);
+	}
 }
 
 void UMaterialExpressionSetMaterialAttributes::GetCaption(TArray<FString>& OutCaptions) const
@@ -20619,4 +20682,372 @@ uint32 UMaterialExpressionStrataAnisotropyToRoughness::GetInputType(int32 InputI
 }
 #endif // WITH_EDITOR
 
+UMaterialExpressionExecBegin::UMaterialExpressionExecBegin(const FObjectInitializer& ObjectInitializer)
+	: Super(ObjectInitializer)
+{
+#if WITH_EDITORONLY_DATA
+	Outputs.Reset();
+	Outputs.Add(FExpressionOutput(TEXT("exec")));
+	bHidePreviewWindow = true;
+#endif
+}
+
+#if WITH_EDITOR
+void UMaterialExpressionExecBegin::GetCaption(TArray<FString>& OutCaptions) const
+{
+	OutCaptions.Add(TEXT("Material Execution Begin"));
+}
+
+uint32 UMaterialExpressionExecBegin::GetOutputType(int32 OutputIndex)
+{
+	check(OutputIndex == 0);
+	return MCT_Execution;
+}
+
+static UMaterialExpression* GetExecConnection(UMaterialExpression* InExpression, int32 InOutputIndex)
+{
+	const uint32 OutputType = InExpression->GetOutputType(InOutputIndex);
+	check(OutputType == MCT_Execution);
+
+	const auto& Expressions = InExpression->Material->Expressions;
+	for (UMaterialExpression* Expression : Expressions)
+	{
+		if (Expression->IsExecInputConnection(InExpression, InOutputIndex))
+		{
+			return Expression;
+		}
+	}
+
+	return nullptr;
+}
+
+int32 UMaterialExpressionExecBegin::Compile(class FMaterialCompiler* Compiler, int32 OutputIndex)
+{
+	check(OutputIndex == CompileExecutionOutputIndex);
+	UMaterialExpression* NextExpression = GetExecConnection(this, 0);
+	if (NextExpression)
+	{
+		const int32 RootScope = Compiler->BeginScope();
+		NextExpression->Compile(Compiler, CompileExecutionOutputIndex);
+		Compiler->EndScope();
+		return RootScope;
+	}
+	return Compiler->Error(TEXT("Exec pin must be connected"));
+}
+#endif
+
+UMaterialExpressionReturnMaterialAttributes::UMaterialExpressionReturnMaterialAttributes(const FObjectInitializer& ObjectInitializer)
+	: Super(ObjectInitializer)
+{
+	struct FConstructorStatics
+	{
+		FText NAME_Execution;
+		FConstructorStatics() : NAME_Execution(LOCTEXT("Execution", "Execution")) { }
+	};
+	static FConstructorStatics ConstructorStatics;
+#if WITH_EDITORONLY_DATA
+	MenuCategories.Add(ConstructorStatics.NAME_Execution);
+	Outputs.Reset();
+	bHidePreviewWindow = true;
+#endif
+}
+
+#if WITH_EDITOR
+void UMaterialExpressionReturnMaterialAttributes::GetCaption(TArray<FString>& OutCaptions) const
+{
+	OutCaptions.Add(TEXT("Return Material Attributes"));
+}
+
+int32 UMaterialExpressionReturnMaterialAttributes::Compile(class FMaterialCompiler* Compiler, int32 OutputIndex)
+{
+	check(OutputIndex == CompileExecutionOutputIndex);
+	const int32 AttributesIndex = MaterialAttributes.Compile(Compiler);
+	return Compiler->ReturnMaterialAttributes(AttributesIndex);
+}
+
+const TArray<FExpressionInput*> UMaterialExpressionReturnMaterialAttributes::GetInputs()
+{
+	TArray<FExpressionInput*> Result;
+	Result.Add(&Exec);
+	Result.Add(&MaterialAttributes);
+	return Result;
+}
+
+FExpressionInput* UMaterialExpressionReturnMaterialAttributes::GetInput(int32 InputIndex)
+{
+	switch (InputIndex)
+	{
+	case 0: return &Exec;
+	case 1: return &MaterialAttributes;
+	default: checkNoEntry(); return nullptr;
+	}
+}
+
+uint32 UMaterialExpressionReturnMaterialAttributes::GetInputType(int32 InputIndex)
+{
+	switch (InputIndex)
+	{
+	case 0: return MCT_Execution;
+	case 1: return MCT_MaterialAttributes;
+	default: checkNoEntry(); return 0u;
+	}
+}
+#endif
+
+UMaterialExpressionIfThenElse::UMaterialExpressionIfThenElse(const FObjectInitializer& ObjectInitializer)
+	: Super(ObjectInitializer)
+{
+	struct FConstructorStatics
+	{
+		FText NAME_Execution;
+		FConstructorStatics() : NAME_Execution(LOCTEXT("Execution", "Execution")) { }
+	};
+	static FConstructorStatics ConstructorStatics;
+#if WITH_EDITORONLY_DATA
+	MenuCategories.Add(ConstructorStatics.NAME_Execution);
+	Outputs.Reset();
+	Outputs.Add(FExpressionOutput(TEXT("Then")));
+	Outputs.Add(FExpressionOutput(TEXT("Else")));
+	bShowOutputNameOnPin = true;
+	bHidePreviewWindow = true;
+#endif
+}
+
+#if WITH_EDITOR
+int32 UMaterialExpressionIfThenElse::Compile(class FMaterialCompiler* Compiler, int32 OutputIndex)
+{
+	check(OutputIndex == CompileExecutionOutputIndex);
+	UMaterialExpression* ThenExpression = GetExecConnection(this, 0);
+	if (!ThenExpression)
+	{
+		return Compiler->Error(TEXT("Then pin must be connected"));
+	}
+
+	const int32 ConditionResult = Condition.GetTracedInput().Compile(Compiler);
+	if (ConditionResult == INDEX_NONE)
+	{
+		return Compiler->Error(TEXT("Condition pin must be connected"));
+	}
+
+	const int32 IfScopeIndex = Compiler->BeginScope_If(ConditionResult);
+	ThenExpression->Compile(Compiler, CompileExecutionOutputIndex);
+	Compiler->EndScope();
+
+	UMaterialExpression* ElseExpression = GetExecConnection(this, 1);
+	int32 ElseScope = INDEX_NONE;
+	if (ElseExpression)
+	{
+		ElseScope = Compiler->BeginScope_Else();
+		ElseExpression->Compile(Compiler, CompileExecutionOutputIndex);
+		Compiler->EndScope();
+	}
+
+	return INDEX_NONE;
+}
+
+void UMaterialExpressionIfThenElse::GetCaption(TArray<FString>& OutCaptions) const
+{
+	OutCaptions.Add(TEXT("If"));
+}
+
+uint32 UMaterialExpressionIfThenElse::GetInputType(int32 InputIndex)
+{
+	switch (InputIndex)
+	{
+	case 0: return MCT_Execution;
+	case 1: return MCT_Float1;
+	default: checkNoEntry(); return 0u;
+	}
+}
+
+uint32 UMaterialExpressionIfThenElse::GetOutputType(int32 InputIndex)
+{
+	return MCT_Execution;
+}
+#endif
+
+UMaterialExpressionForLoop::UMaterialExpressionForLoop(const FObjectInitializer& ObjectInitializer)
+	: Super(ObjectInitializer)
+{
+	struct FConstructorStatics
+	{
+		FText NAME_Execution;
+		FConstructorStatics() : NAME_Execution(LOCTEXT("Execution", "Execution")) { }
+	};
+	static FConstructorStatics ConstructorStatics;
+#if WITH_EDITORONLY_DATA
+	MenuCategories.Add(ConstructorStatics.NAME_Execution);
+	Outputs.Reset();
+	Outputs.Add(FExpressionOutput(TEXT("Loop Body")));
+	Outputs.Add(FExpressionOutput(TEXT("Index")));
+	Outputs.Add(FExpressionOutput(TEXT("Completed")));
+	bShowOutputNameOnPin = true;
+	bHidePreviewWindow = true;
+#endif
+}
+
+#if WITH_EDITOR
+int32 UMaterialExpressionForLoop::Compile(class FMaterialCompiler* Compiler, int32 OutputIndex)
+{
+	if (OutputIndex == CompileExecutionOutputIndex)
+	{
+		UMaterialExpression* LoopExpression = GetExecConnection(this, 0);
+		if (!LoopExpression)
+		{
+			return Compiler->Error(TEXT("Loop pin must be connected"));
+		}
+
+		UMaterialExpression* CompletedExpression = GetExecConnection(this, 2);
+
+		const int32 StartResult = StartIndex.GetTracedInput().Compile(Compiler);
+		if (StartResult == INDEX_NONE)
+		{
+			return Compiler->Error(TEXT("Start index pin must be connected"));
+		}
+
+		const int32 EndResult = EndIndex.GetTracedInput().Compile(Compiler);
+		if (EndResult == INDEX_NONE)
+		{
+			return Compiler->Error(TEXT("End index pin must be connected"));
+		}
+
+		const int32 LoopScopeIndex = Compiler->BeginScope_For(this, StartResult, EndResult, Compiler->Constant(1.0f));
+		LoopExpression->Compile(Compiler, CompileExecutionOutputIndex);
+		Compiler->EndScope();
+
+		if (CompletedExpression)
+		{
+			CompletedExpression->Compile(Compiler, CompileExecutionOutputIndex);
+		}
+
+		return INDEX_NONE;
+	}
+
+	// Loop index output
+	check(OutputIndex == 1);
+	return Compiler->ForLoopIndex(this);
+}
+
+void UMaterialExpressionForLoop::GetCaption(TArray<FString>& OutCaptions) const
+{
+	OutCaptions.Add(TEXT("For"));
+}
+
+uint32 UMaterialExpressionForLoop::GetInputType(int32 InputIndex)
+{
+	switch (InputIndex)
+	{
+	case 0: return MCT_Execution;
+	case 1: return MCT_Float1;
+	case 2: return MCT_Float1;
+	case 3: return MCT_Float1;
+	default: checkNoEntry(); return 0u;
+	}
+}
+
+uint32 UMaterialExpressionForLoop::GetOutputType(int32 OutputIndex)
+{
+	switch (OutputIndex)
+	{
+	case 0: return MCT_Execution;
+	case 1: return MCT_Float1;
+	case 2: return MCT_Execution;
+	default: checkNoEntry(); return 0u;
+	}
+}
+#endif
+
+UMaterialExpressionGetLocal::UMaterialExpressionGetLocal(const FObjectInitializer& ObjectInitializer)
+	: Super(ObjectInitializer)
+{
+	struct FConstructorStatics
+	{
+		FText NAME_Execution;
+		FConstructorStatics() : NAME_Execution(LOCTEXT("Execution", "Execution")) { }
+	};
+	static FConstructorStatics ConstructorStatics;
+#if WITH_EDITORONLY_DATA
+	MenuCategories.Add(ConstructorStatics.NAME_Execution);
+	Outputs.Reset();
+	Outputs.Add(FExpressionOutput(TEXT("Result")));
+#endif
+}
+
+#if WITH_EDITOR
+int32 UMaterialExpressionGetLocal::Compile(class FMaterialCompiler* Compiler, int32 OutputIndex)
+{
+	return Compiler->GetLocal(LocalName);
+}
+
+void UMaterialExpressionGetLocal::GetCaption(TArray<FString>& OutCaptions) const
+{
+	OutCaptions.Add(FString::Printf(TEXT("Get %s"), *LocalName.ToString()));
+}
+
+uint32 UMaterialExpressionGetLocal::GetOutputType(int32 InputIndex)
+{
+	return MCT_Float1;
+}
+#endif
+
+UMaterialExpressionSetLocal::UMaterialExpressionSetLocal(const FObjectInitializer& ObjectInitializer)
+	: Super(ObjectInitializer)
+{
+	struct FConstructorStatics
+	{
+		FText NAME_Execution;
+		FConstructorStatics() : NAME_Execution(LOCTEXT("Execution", "Execution")) { }
+	};
+	static FConstructorStatics ConstructorStatics;
+#if WITH_EDITORONLY_DATA
+	MenuCategories.Add(ConstructorStatics.NAME_Execution);
+	Outputs.Reset();
+	Outputs.Add(FExpressionOutput(TEXT("Exec")));
+	bHidePreviewWindow = true;
+#endif
+}
+
+#if WITH_EDITOR
+uint32 UMaterialExpressionSetLocal::GetInputType(int32 InputIndex)
+{
+	switch (InputIndex)
+	{
+	case 0: return MCT_Execution;
+	case 1: return MCT_Float1;
+	default: checkNoEntry(); return 0u;
+	}
+}
+
+int32 UMaterialExpressionSetLocal::Compile(class FMaterialCompiler* Compiler, int32 OutputIndex)
+{
+	check(OutputIndex == CompileExecutionOutputIndex);
+	
+	const int32 ValueIndex = Value.GetTracedInput().Compile(Compiler);
+	if (ValueIndex == INDEX_NONE)
+	{
+		return Compiler->Error(TEXT("Value must be connected"));
+	}
+
+	Compiler->SetLocal(LocalName, ValueIndex);
+
+	UMaterialExpression* NextExpression = GetExecConnection(this, 0);
+	if (NextExpression)
+	{
+		NextExpression->Compile(Compiler, CompileExecutionOutputIndex);
+	}
+
+	return INDEX_NONE;
+}
+
+void UMaterialExpressionSetLocal::GetCaption(TArray<FString>& OutCaptions) const
+{
+	OutCaptions.Add(FString::Printf(TEXT("Set %s"), *LocalName.ToString()));
+}
+
+uint32 UMaterialExpressionSetLocal::GetOutputType(int32 InputIndex)
+{
+	return MCT_Execution;
+}
+#endif
+
 #undef LOCTEXT_NAMESPACE
diff --git a/Engine/Source/Runtime/Engine/Private/Materials/MaterialShared.cpp b/Engine/Source/Runtime/Engine/Private/Materials/MaterialShared.cpp
--- a/Engine/Source/Runtime/Engine/Private/Materials/MaterialShared.cpp
+++ b/Engine/Source/Runtime/Engine/Private/Materials/MaterialShared.cpp
@@ -1197,7 +1197,7 @@ bool FMaterialResource::IsVolumetricPrimitive() const { return Material->Materia
 bool FMaterialResource::IsSpecialEngineMaterial() const { return Material->bUsedAsSpecialEngineMaterial; }
 bool FMaterialResource::HasVertexPositionOffsetConnected() const { return HasMaterialAttributesConnected() || (!Material->bUseMaterialAttributes && Material->WorldPositionOffset.IsConnected()); }
 bool FMaterialResource::HasPixelDepthOffsetConnected() const { return HasMaterialAttributesConnected() || (!Material->bUseMaterialAttributes && Material->PixelDepthOffset.IsConnected()); }
-bool FMaterialResource::HasMaterialAttributesConnected() const { return Material->bUseMaterialAttributes && Material->MaterialAttributes.IsConnected(); }
+bool FMaterialResource::HasMaterialAttributesConnected() const { return (Material->bUseMaterialAttributes && Material->MaterialAttributes.IsConnected()) || Material->bEnableExecWire; }
 EMaterialShadingRate FMaterialResource::GetShadingRate() const { return Material->ShadingRate; }
 FString FMaterialResource::GetBaseMaterialPathName() const { return Material->GetPathName(); }
 FString FMaterialResource::GetDebugName() const
@@ -1681,6 +1681,15 @@ bool FMaterialResource::ShouldInlineShaderCode() const
 	UE_LOG(LogMaterial, Display, TEXT("%s (package %s): shader code is %s to be inlined as a workaround"), *GetFriendlyName(), *OutermostName, bNeedsToBeInlined ? TEXT("configured") : TEXT("NOT configured"));
 	return bNeedsToBeInlined;
 }
+
+bool FMaterialResource::IsCompiledWithExecutionFlow() const
+{
+	if (Material)
+	{
+		return Material->IsCompiledWithExecutionFlow();
+	}
+	return false;
+}
 #endif
 
 FString FMaterialResource::GetFullPath() const
@@ -4086,7 +4095,7 @@ FMaterialAttributeDefintion::FMaterialAttributeDefintion(
 	, BlendFunction(InBlendFunction)
 	, bIsHidden(bInIsHidden)
 {
-	checkf(ValueType & MCT_Float || ValueType == MCT_ShadingModel || ValueType == MCT_Strata, TEXT("Unsupported type, only Float1 through Float4 or MCT_ShadingModel are allowed."));
+	checkf(ValueType & MCT_Float || ValueType == MCT_ShadingModel || ValueType == MCT_Strata || ValueType == MCT_MaterialAttributes, TEXT("Unsupported material attribute type %d"), ValueType);
 }
 
 int32 FMaterialAttributeDefintion::CompileDefaultValue(FMaterialCompiler* Compiler)
@@ -4180,6 +4189,9 @@ void FMaterialAttributeDefinitionMap::InitializeAttributeMap()
 	Add(FGuid(0x0AC97EC3, 0xE3D047BA, 0xB610167D, 0xC4D919FF), TEXT("PixelDepthOffset"),		MP_PixelDepthOffset,		MCT_Float,	FVector4(0,0,0,0),	SF_Pixel);
 	Add(FGuid(0xD9423FFF, 0xD77E4D82, 0x8FF9CF5E, 0x055D1255), TEXT("ShadingModel"),			MP_ShadingModel,			MCT_ShadingModel, FVector4(0, 0, 0, 0), SF_Pixel, INDEX_NONE, false, &CompileShadingModelBlendFunction);
 	Add(FGuid(0x5973A03E, 0x13A74E08, 0x92D0CEDD, 0xF2936CF8), TEXT("FrontMaterial"),			MP_FrontMaterial,			MCT_Strata, FVector4(0,0,0,0),	SF_Pixel, INDEX_NONE, false, &CompileStrataBlendFunction);
+
+	// Used when compiling material with execution pins, which are compiling all attributes together
+	Add(FGuid(0xE0ED040B, 0x82794D93, 0xBD2D59B2, 0xA5BBF41C), TEXT("MaterialAttributes"),		MP_MaterialAttributes,		MCT_MaterialAttributes, FVector4(0,0,0,0), SF_Pixel, INDEX_NONE, bHideAttribute);
 	   
 	// Texture coordinates
 	Add(FGuid(0xD30EC284, 0xE13A4160, 0x87BB5230, 0x2ED115DC), TEXT("CustomizedUV0"), MP_CustomizedUVs0, MCT_Float2, FVector4(0,0,0,0), SF_Vertex, 0);
diff --git a/Engine/Source/Runtime/Engine/Public/MaterialCompiler.h b/Engine/Source/Runtime/Engine/Public/MaterialCompiler.h
--- a/Engine/Source/Runtime/Engine/Public/MaterialCompiler.h
+++ b/Engine/Source/Runtime/Engine/Public/MaterialCompiler.h
@@ -390,6 +390,21 @@ public:
 	virtual int32 CustomPrimitiveData(int32 OutputIndex, EMaterialValueType Type) = 0;
 	virtual int32 ShadingModel(EMaterialShadingModel InSelectedShadingModel) = 0;
 
+	// Material attributes
+	virtual int32 DefaultMaterialAttributes() = 0;
+	virtual int32 SetMaterialAttribute(int32 MaterialAttributes, int32 Value, const FGuid& AttributeID) = 0;
+
+	// Exec
+	virtual int32 BeginScope() = 0;
+	virtual int32 BeginScope_If(int32 Condition) = 0;
+	virtual int32 BeginScope_Else() = 0;
+	virtual int32 BeginScope_For(const UMaterialExpression* Expression, int32 StartIndex, int32 EndIndex, int32 IndexStep) = 0;
+	virtual int32 EndScope() = 0;
+	virtual int32 ForLoopIndex(const UMaterialExpression* Expression) = 0;
+	virtual int32 ReturnMaterialAttributes(int32 MaterialAttributes) = 0;
+	virtual int32 SetLocal(const FName& LocalName, int32 Value) = 0;
+	virtual int32 GetLocal(const FName& LocalName) = 0;
+
 	// Strata
 	virtual int32 StrataCreateAndRegisterNullMaterial() = 0;
 	virtual int32 StrataSlabBSDF(
@@ -435,7 +450,7 @@ class FProxyMaterialCompiler : public FMaterialCompiler
 public:
 
 	// Constructor.
-	FProxyMaterialCompiler(FMaterialCompiler* InCompiler):
+	FProxyMaterialCompiler(FMaterialCompiler* InCompiler) :
 		Compiler(InCompiler)
 	{}
 
@@ -458,7 +473,7 @@ public:
 	virtual int32 Error(const TCHAR* Text) override { return Compiler->Error(Text); }
 	virtual void AppendExpressionError(UMaterialExpression* Expression, const TCHAR* Text) override { return Compiler->AppendExpressionError(Expression, Text); }
 
-	virtual int32 CallExpression(FMaterialExpressionKey ExpressionKey,FMaterialCompiler* InCompiler) override { return Compiler->CallExpression(ExpressionKey,InCompiler); }
+	virtual int32 CallExpression(FMaterialExpressionKey ExpressionKey, FMaterialCompiler* InCompiler) override { return Compiler->CallExpression(ExpressionKey, InCompiler); }
 
 	virtual void PushFunction(FMaterialFunctionCompileState* FunctionState) override { Compiler->PushFunction(FunctionState); }
 	virtual FMaterialFunctionCompileState* PopFunction() override { return Compiler->PopFunction(); }
@@ -470,19 +485,21 @@ public:
 	virtual EShaderPlatform GetShaderPlatform() override { return Compiler->GetShaderPlatform(); }
 	virtual const ITargetPlatform* GetTargetPlatform() const override { return Compiler->GetTargetPlatform(); }
 	virtual bool IsMaterialPropertyUsed(EMaterialProperty Property, int32 CodeChunkIdx) const override { return Compiler->IsMaterialPropertyUsed(Property, CodeChunkIdx); }
-	virtual int32 ValidCast(int32 Code,EMaterialValueType DestType) override { return Compiler->ValidCast(Code, DestType); }
-	virtual int32 ForceCast(int32 Code,EMaterialValueType DestType,uint32 ForceCastFlags = 0) override
-	{ return Compiler->ForceCast(Code,DestType,ForceCastFlags); }
+	virtual int32 ValidCast(int32 Code, EMaterialValueType DestType) override { return Compiler->ValidCast(Code, DestType); }
+	virtual int32 ForceCast(int32 Code, EMaterialValueType DestType, uint32 ForceCastFlags = 0) override
+	{
+		return Compiler->ForceCast(Code, DestType, ForceCastFlags);
+	}
 
 	virtual int32 AccessCollectionParameter(UMaterialParameterCollection* ParameterCollection, int32 ParameterIndex, int32 ComponentIndex) override { return Compiler->AccessCollectionParameter(ParameterCollection, ParameterIndex, ComponentIndex); }
-	virtual int32 ScalarParameter(FName ParameterName, float DefaultValue) override { return Compiler->ScalarParameter(ParameterName,DefaultValue); }
-	virtual int32 VectorParameter(FName ParameterName, const FLinearColor& DefaultValue) override { return Compiler->VectorParameter(ParameterName,DefaultValue); }
+	virtual int32 ScalarParameter(FName ParameterName, float DefaultValue) override { return Compiler->ScalarParameter(ParameterName, DefaultValue); }
+	virtual int32 VectorParameter(FName ParameterName, const FLinearColor& DefaultValue) override { return Compiler->VectorParameter(ParameterName, DefaultValue); }
 
 	virtual int32 Constant(float X) override { return Compiler->Constant(X); }
-	virtual int32 Constant2(float X,float Y) override { return Compiler->Constant2(X,Y); }
-	virtual int32 Constant3(float X,float Y,float Z) override { return Compiler->Constant3(X,Y,Z); }
-	virtual int32 Constant4(float X,float Y,float Z,float W) override { return Compiler->Constant4(X,Y,Z,W); }
-	
+	virtual int32 Constant2(float X, float Y) override { return Compiler->Constant2(X, Y); }
+	virtual int32 Constant3(float X, float Y, float Z) override { return Compiler->Constant3(X, Y, Z); }
+	virtual int32 Constant4(float X, float Y, float Z, float W) override { return Compiler->Constant4(X, Y, Z, W); }
+
 	virtual	int32 ViewProperty(EMaterialExposedViewProperty Property, bool InvProperty) override { return Compiler->ViewProperty(Property, InvProperty); }
 
 	virtual int32 GameTime(bool bPeriodic, float Period) override { return Compiler->GameTime(bPeriodic, Period); }
@@ -509,7 +526,7 @@ public:
 	virtual int32 Truncate(int32 X) override { return Compiler->Truncate(X); }
 	virtual int32 Sign(int32 X) override { return Compiler->Sign(X); }
 	virtual int32 Frac(int32 X) override { return Compiler->Frac(X); }
-	virtual int32 Fmod(int32 A, int32 B) override { return Compiler->Fmod(A,B); }
+	virtual int32 Fmod(int32 A, int32 B) override { return Compiler->Fmod(A, B); }
 	virtual int32 Abs(int32 X) override { return Compiler->Abs(X); }
 
 	virtual int32 ReflectionVector() override { return Compiler->ReflectionVector(); }
@@ -533,12 +550,16 @@ public:
 	virtual int32 ParticleRadius() override { return Compiler->ParticleRadius(); }
 	virtual int32 SphericalParticleOpacity(int32 Density) override { return Compiler->SphericalParticleOpacity(Density); }
 
-	virtual int32 If(int32 A,int32 B,int32 AGreaterThanB,int32 AEqualsB,int32 ALessThanB,int32 Threshold) override { return Compiler->If(A,B,AGreaterThanB,AEqualsB,ALessThanB,Threshold); }
+	virtual int32 If(int32 A, int32 B, int32 AGreaterThanB, int32 AEqualsB, int32 ALessThanB, int32 Threshold) override { return Compiler->If(A, B, AGreaterThanB, AEqualsB, ALessThanB, Threshold); }
 
-	virtual int32 TextureSample(int32 InTexture,int32 Coordinate,enum EMaterialSamplerType SamplerType,int32 MipValue0Index,int32 MipValue1Index,ETextureMipValueMode MipValueMode,ESamplerSourceMode SamplerSource,int32 TextureReferenceIndex, bool AutomaticViewMipBias, bool AdaptiveVirtualTexture) override
-		{ return Compiler->TextureSample(InTexture,Coordinate,SamplerType,MipValue0Index,MipValue1Index,MipValueMode,SamplerSource,TextureReferenceIndex, AutomaticViewMipBias, AdaptiveVirtualTexture); }
-	virtual int32 TextureProperty(int32 InTexture, EMaterialExposedTextureProperty Property) override 
-		{ return Compiler->TextureProperty(InTexture, Property); }
+	virtual int32 TextureSample(int32 InTexture, int32 Coordinate, enum EMaterialSamplerType SamplerType, int32 MipValue0Index, int32 MipValue1Index, ETextureMipValueMode MipValueMode, ESamplerSourceMode SamplerSource, int32 TextureReferenceIndex, bool AutomaticViewMipBias, bool AdaptiveVirtualTexture) override
+	{
+		return Compiler->TextureSample(InTexture, Coordinate, SamplerType, MipValue0Index, MipValue1Index, MipValueMode, SamplerSource, TextureReferenceIndex, AutomaticViewMipBias, AdaptiveVirtualTexture);
+	}
+	virtual int32 TextureProperty(int32 InTexture, EMaterialExposedTextureProperty Property) override
+	{
+		return Compiler->TextureProperty(InTexture, Property);
+	}
 
 	virtual int32 TextureCoordinate(uint32 CoordinateIndex, bool UnMirrorU, bool UnMirrorV) override { return Compiler->TextureCoordinate(CoordinateIndex, UnMirrorU, UnMirrorV); }
 
@@ -546,16 +567,16 @@ public:
 	virtual int32 TextureDecalDerivative(bool bDDY) override { return Compiler->TextureDecalDerivative(bDDY); }
 	virtual int32 DecalLifetimeOpacity() override { return Compiler->DecalLifetimeOpacity(); }
 
-	virtual int32 Texture(UTexture* InTexture,int32& TextureReferenceIndex,EMaterialSamplerType SamplerType,ESamplerSourceMode SamplerSource=SSM_FromTextureAsset,ETextureMipValueMode MipValueMode=TMVM_None) override { return Compiler->Texture(InTexture,TextureReferenceIndex, SamplerType, SamplerSource,MipValueMode); }
-	virtual int32 TextureParameter(FName ParameterName,UTexture* DefaultValue,int32& TextureReferenceIndex, EMaterialSamplerType SamplerType, ESamplerSourceMode SamplerSource=SSM_FromTextureAsset) override { return Compiler->TextureParameter(ParameterName,DefaultValue,TextureReferenceIndex, SamplerType, SamplerSource); }
+	virtual int32 Texture(UTexture* InTexture, int32& TextureReferenceIndex, EMaterialSamplerType SamplerType, ESamplerSourceMode SamplerSource = SSM_FromTextureAsset, ETextureMipValueMode MipValueMode = TMVM_None) override { return Compiler->Texture(InTexture, TextureReferenceIndex, SamplerType, SamplerSource, MipValueMode); }
+	virtual int32 TextureParameter(FName ParameterName, UTexture* DefaultValue, int32& TextureReferenceIndex, EMaterialSamplerType SamplerType, ESamplerSourceMode SamplerSource = SSM_FromTextureAsset) override { return Compiler->TextureParameter(ParameterName, DefaultValue, TextureReferenceIndex, SamplerType, SamplerSource); }
 
 	virtual int32 VirtualTexture(URuntimeVirtualTexture* InTexture, int32 TextureLayerIndex, int32 PageTableLayerIndex, int32& TextureReferenceIndex, EMaterialSamplerType SamplerType) override
 	{
-		return Compiler->VirtualTexture(InTexture, TextureLayerIndex, PageTableLayerIndex, TextureReferenceIndex, SamplerType); 
+		return Compiler->VirtualTexture(InTexture, TextureLayerIndex, PageTableLayerIndex, TextureReferenceIndex, SamplerType);
 	}
 	virtual int32 VirtualTextureParameter(FName ParameterName, URuntimeVirtualTexture* DefaultValue, int32 TextureLayerIndex, int32 PageTableLayerIndex, int32& TextureReferenceIndex, EMaterialSamplerType SamplerType) override
 	{
-		return Compiler->VirtualTextureParameter(ParameterName, DefaultValue, TextureLayerIndex, PageTableLayerIndex, TextureReferenceIndex, SamplerType); 
+		return Compiler->VirtualTextureParameter(ParameterName, DefaultValue, TextureLayerIndex, PageTableLayerIndex, TextureReferenceIndex, SamplerType);
 	}
 	virtual int32 VirtualTextureUniform(int32 TextureIndex, int32 VectorIndex) override { return Compiler->VirtualTextureUniform(TextureIndex, VectorIndex); }
 	virtual int32 VirtualTextureUniform(FName ParameterName, int32 TextureIndex, int32 VectorIndex) override { return Compiler->VirtualTextureUniform(ParameterName, TextureIndex, VectorIndex); }
@@ -572,49 +593,49 @@ public:
 
 	virtual UObject* GetReferencedTexture(int32 Index) override { return Compiler->GetReferencedTexture(Index); }
 
-	virtual	int32 PixelDepth() override { return Compiler->PixelDepth();	}
+	virtual	int32 PixelDepth() override { return Compiler->PixelDepth(); }
 	virtual int32 SceneDepth(int32 Offset, int32 ViewportUV, bool bUseOffset) override { return Compiler->SceneDepth(Offset, ViewportUV, bUseOffset); }
 	virtual int32 SceneColor(int32 Offset, int32 ViewportUV, bool bUseOffset) override { return Compiler->SceneColor(Offset, ViewportUV, bUseOffset); }
 	virtual int32 SceneTextureLookup(int32 ViewportUV, uint32 InSceneTextureId, bool bFiltered) override { return Compiler->SceneTextureLookup(ViewportUV, InSceneTextureId, bFiltered); }
 	virtual int32 GetSceneTextureViewSize(int32 SceneTextureId, bool InvProperty) override { return Compiler->GetSceneTextureViewSize(SceneTextureId, InvProperty); }
 
 	virtual int32 StaticBool(bool Value) override { return Compiler->StaticBool(Value); }
-	virtual int32 StaticBoolParameter(FName ParameterName,bool bDefaultValue) override { return Compiler->StaticBoolParameter(ParameterName,bDefaultValue); }
-	virtual int32 StaticComponentMask(int32 Vector,FName ParameterName,bool bDefaultR,bool bDefaultG,bool bDefaultB,bool bDefaultA) override { return Compiler->StaticComponentMask(Vector,ParameterName,bDefaultR,bDefaultG,bDefaultB,bDefaultA); }
+	virtual int32 StaticBoolParameter(FName ParameterName, bool bDefaultValue) override { return Compiler->StaticBoolParameter(ParameterName, bDefaultValue); }
+	virtual int32 StaticComponentMask(int32 Vector, FName ParameterName, bool bDefaultR, bool bDefaultG, bool bDefaultB, bool bDefaultA) override { return Compiler->StaticComponentMask(Vector, ParameterName, bDefaultR, bDefaultG, bDefaultB, bDefaultA); }
 	virtual const FMaterialLayersFunctions* StaticMaterialLayersParameter(FName ParameterName) override { return Compiler->StaticMaterialLayersParameter(ParameterName); }
 	virtual bool GetStaticBoolValue(int32 BoolIndex, bool& bSucceeded) override { return Compiler->GetStaticBoolValue(BoolIndex, bSucceeded); }
-	virtual int32 StaticTerrainLayerWeight(FName ParameterName,int32 Default) override { return Compiler->StaticTerrainLayerWeight(ParameterName,Default); }
+	virtual int32 StaticTerrainLayerWeight(FName ParameterName, int32 Default) override { return Compiler->StaticTerrainLayerWeight(ParameterName, Default); }
 
 	virtual int32 VertexColor() override { return Compiler->VertexColor(); }
-	
+
 	virtual int32 PreSkinnedPosition() override { return Compiler->PreSkinnedPosition(); }
 	virtual int32 PreSkinnedNormal() override { return Compiler->PreSkinnedNormal(); }
 
-	virtual int32 Add(int32 A,int32 B) override { return Compiler->Add(A,B); }
-	virtual int32 Sub(int32 A,int32 B) override { return Compiler->Sub(A,B); }
-	virtual int32 Mul(int32 A,int32 B) override { return Compiler->Mul(A,B); }
-	virtual int32 Div(int32 A,int32 B) override { return Compiler->Div(A,B); }
-	virtual int32 Dot(int32 A,int32 B) override { return Compiler->Dot(A,B); }
-	virtual int32 Cross(int32 A,int32 B) override { return Compiler->Cross(A,B); }
+	virtual int32 Add(int32 A, int32 B) override { return Compiler->Add(A, B); }
+	virtual int32 Sub(int32 A, int32 B) override { return Compiler->Sub(A, B); }
+	virtual int32 Mul(int32 A, int32 B) override { return Compiler->Mul(A, B); }
+	virtual int32 Div(int32 A, int32 B) override { return Compiler->Div(A, B); }
+	virtual int32 Dot(int32 A, int32 B) override { return Compiler->Dot(A, B); }
+	virtual int32 Cross(int32 A, int32 B) override { return Compiler->Cross(A, B); }
 
-	virtual int32 Power(int32 Base,int32 Exponent) override { return Compiler->Power(Base,Exponent); }
+	virtual int32 Power(int32 Base, int32 Exponent) override { return Compiler->Power(Base, Exponent); }
 	virtual int32 Logarithm2(int32 X) override { return Compiler->Logarithm2(X); }
 	virtual int32 Logarithm10(int32 X) override { return Compiler->Logarithm10(X); }
 	virtual int32 SquareRoot(int32 X) override { return Compiler->SquareRoot(X); }
 	virtual int32 Length(int32 X) override { return Compiler->Length(X); }
 
-	virtual int32 Lerp(int32 X,int32 Y,int32 A) override { return Compiler->Lerp(X,Y,A); }
-	virtual int32 Min(int32 A,int32 B) override { return Compiler->Min(A,B); }
-	virtual int32 Max(int32 A,int32 B) override { return Compiler->Max(A,B); }
-	virtual int32 Clamp(int32 X,int32 A,int32 B) override { return Compiler->Clamp(X,A,B); }
+	virtual int32 Lerp(int32 X, int32 Y, int32 A) override { return Compiler->Lerp(X, Y, A); }
+	virtual int32 Min(int32 A, int32 B) override { return Compiler->Min(A, B); }
+	virtual int32 Max(int32 A, int32 B) override { return Compiler->Max(A, B); }
+	virtual int32 Clamp(int32 X, int32 A, int32 B) override { return Compiler->Clamp(X, A, B); }
 	virtual int32 Saturate(int32 X) override { return Compiler->Saturate(X); }
 
-	virtual int32 SmoothStep(int32 X,int32 Y,int32 A) override { return Compiler->SmoothStep(X,Y,A); }
-	virtual int32 Step(int32 Y,int32 X) override { return Compiler->Step(Y,X); }
-	virtual int32 InvLerp(int32 X,int32 Y,int32 A) override { return Compiler->InvLerp(X,Y,A); }
+	virtual int32 SmoothStep(int32 X, int32 Y, int32 A) override { return Compiler->SmoothStep(X, Y, A); }
+	virtual int32 Step(int32 Y, int32 X) override { return Compiler->Step(Y, X); }
+	virtual int32 InvLerp(int32 X, int32 Y, int32 A) override { return Compiler->InvLerp(X, Y, A); }
 
-	virtual int32 ComponentMask(int32 Vector,bool R,bool G,bool B,bool A) override { return Compiler->ComponentMask(Vector,R,G,B,A); }
-	virtual int32 AppendVector(int32 A,int32 B) override { return Compiler->AppendVector(A,B); }
+	virtual int32 ComponentMask(int32 Vector, bool R, bool G, bool B, bool A) override { return Compiler->ComponentMask(Vector, R, G, B, A); }
+	virtual int32 AppendVector(int32 A, int32 B) override { return Compiler->AppendVector(A, B); }
 	virtual int32 TransformVector(EMaterialCommonBasis SourceCoordBasis, EMaterialCommonBasis DestCoordBasis, int32 A) override
 	{
 		return Compiler->TransformVector(SourceCoordBasis, DestCoordBasis, A);
@@ -646,7 +667,7 @@ public:
 	virtual int32 PixelNormalWS() override { return Compiler->PixelNormalWS(); }
 
 	virtual int32 CustomExpression(class UMaterialExpressionCustom* Custom, int32 OutputIndex, TArray<int32>& CompiledInputs) override { return Compiler->CustomExpression(Custom, OutputIndex, CompiledInputs); }
-	virtual int32 CustomOutput(class UMaterialExpressionCustomOutput* Custom, int32 OutputIndex, int32 OutputCode) override{ return Compiler->CustomOutput(Custom, OutputIndex, OutputCode); }
+	virtual int32 CustomOutput(class UMaterialExpressionCustomOutput* Custom, int32 OutputIndex, int32 OutputCode) override { return Compiler->CustomOutput(Custom, OutputIndex, OutputCode); }
 	virtual int32 VirtualTextureOutput(uint8 AttributeMask) override { return Compiler->VirtualTextureOutput(AttributeMask); }
 
 	virtual int32 DDX(int32 X) override { return Compiler->DDX(X); }
@@ -656,7 +677,7 @@ public:
 	{
 		return Compiler->AntialiasedTextureMask(Tex, UV, Threshold, Channel);
 	}
-	virtual int32 Sobol(int32 Cell, int32 Index, int32 Seed) override {	return Compiler->Sobol(Cell, Index, Seed); }
+	virtual int32 Sobol(int32 Cell, int32 Index, int32 Seed) override { return Compiler->Sobol(Cell, Index, Seed); }
 	virtual int32 TemporalSobol(int32 Index, int32 Seed) override { return Compiler->TemporalSobol(Index, Seed); }
 	virtual int32 Noise(int32 Position, float Scale, int32 Quality, uint8 NoiseFunction, bool bTurbulence, int32 Levels, float OutputMin, float OutputMax, float LevelScale, int32 FilterWidth, bool bTiling, uint32 TileSize) override
 	{
@@ -666,7 +687,7 @@ public:
 	{
 		return Compiler->VectorNoise(Position, Quality, NoiseFunction, bTiling, TileSize);
 	}
-	virtual int32 BlackBody( int32 Temp ) override { return Compiler->BlackBody(Temp); }
+	virtual int32 BlackBody(int32 Temp) override { return Compiler->BlackBody(Temp); }
 	virtual int32 DistanceToNearestSurface(int32 PositionArg) override { return Compiler->DistanceToNearestSurface(PositionArg); }
 	virtual int32 DistanceFieldGradient(int32 PositionArg) override { return Compiler->DistanceFieldGradient(PositionArg); }
 	virtual int32 SamplePhysicsField(int32 PositionArg, const int32 OutputType, const int32 TargetIndex)  override { return Compiler->SamplePhysicsField(PositionArg, OutputType, TargetIndex); }
@@ -697,7 +718,7 @@ public:
 	}
 
 	virtual int32 SpeedTree(int32 GeometryArg, int32 WindArg, int32 LODArg, float BillboardThreshold, bool bAccurateWindVelocities, bool bExtraBend, int32 ExtraBendArg) override
-	{ 
+	{
 		return Compiler->SpeedTree(GeometryArg, WindArg, LODArg, BillboardThreshold, bAccurateWindVelocities, bExtraBend, ExtraBendArg);
 	}
 
@@ -765,7 +786,7 @@ public:
 	{
 		return Compiler->GetVolumeSampleConservativeDensity();
 	}
-	
+
 	virtual int32 SceneDepthWithoutWater(int32 Offset, int32 ViewportUV, bool bUseOffset, float FallbackDepth) override
 	{
 		return Compiler->SceneDepthWithoutWater(Offset, ViewportUV, bUseOffset, FallbackDepth);
@@ -806,6 +827,61 @@ public:
 		return Compiler->IsDevelopmentFeatureEnabled(FeatureName);
 	}
 
+	virtual int32 DefaultMaterialAttributes() override
+	{
+		return Compiler->DefaultMaterialAttributes();
+	}
+
+	virtual int32 SetMaterialAttribute(int32 MaterialAttributes, int32 Value, const FGuid& AttributeID)
+	{
+		return Compiler->SetMaterialAttribute(MaterialAttributes, Value, AttributeID);
+	}
+
+	virtual int32 BeginScope() override
+	{
+		return Compiler->BeginScope();
+	}
+
+	virtual int32 BeginScope_If(int32 Condition) override
+	{
+		return Compiler->BeginScope_If(Condition);
+	}
+
+	virtual int32 BeginScope_Else() override
+	{
+		return Compiler->BeginScope_Else();
+	}
+
+	virtual int32 BeginScope_For(const UMaterialExpression* Expression, int32 StartIndex, int32 EndIndex, int32 IndexStep) override
+	{
+		return Compiler->BeginScope_For(Expression, StartIndex, EndIndex, IndexStep);
+	}
+
+	virtual int32 EndScope() override
+	{
+		return Compiler->EndScope();
+	}
+
+	virtual int32 ForLoopIndex(const UMaterialExpression* Expression) override
+	{
+		return Compiler->ForLoopIndex(Expression);
+	}
+
+	virtual int32 ReturnMaterialAttributes(int32 MaterialAttributes) override
+	{
+		return Compiler->ReturnMaterialAttributes(MaterialAttributes);
+	}
+
+	virtual int32 SetLocal(const FName& LocalName, int32 Value) override
+	{
+		return Compiler->SetLocal(LocalName, Value);
+	}
+
+	virtual int32 GetLocal(const FName& LocalName) override
+	{
+		return Compiler->GetLocal(LocalName);
+	}
+
 	virtual int32 StrataCreateAndRegisterNullMaterial() override
 	{
 		return Compiler->StrataCreateAndRegisterNullMaterial();
diff --git a/Engine/Source/Runtime/Engine/Public/MaterialShared.h b/Engine/Source/Runtime/Engine/Public/MaterialShared.h
--- a/Engine/Source/Runtime/Engine/Public/MaterialShared.h
+++ b/Engine/Source/Runtime/Engine/Public/MaterialShared.h
@@ -146,6 +146,11 @@ enum EMaterialValueType
 	MCT_ShadingModel		  = 1 << 14,
 
 	MCT_Strata				  = 1 << 15,
+
+	MCT_Execution             = 1 << 16,
+
+	/** Used for code chunks that are statements with no value, rather than expressions */
+	MCT_VoidStatement         = 1 << 17,
 };
 
 /**
@@ -1799,6 +1804,9 @@ public:
 	 * @return returns true if compilation is complete false otherwise
 	 */
 	ENGINE_API bool IsCompilationFinished() const;
+
+	/** Should the material be compiled using exec pin? */
+	virtual bool IsCompiledWithExecutionFlow() const { return false; }
 #endif // WITH_EDITOR
 
 	/**
@@ -2632,6 +2640,7 @@ public:
 	/** Allows to associate the shader resources with the asset for load order. */
 	ENGINE_API virtual FName GetAssetPath() const override;
 	ENGINE_API virtual bool ShouldInlineShaderCode() const override;
+	ENGINE_API virtual bool IsCompiledWithExecutionFlow() const override;
 #endif // WITH_EDITOR
 
 	ENGINE_API virtual FString GetFullPath() const override;
