Fix crash when unlinking animbp's during pending kill

diff --git a/Engine/Source/Runtime/Engine/Private/Animation/AnimInstance.cpp b/Engine/Source/Runtime/Engine/Private/Animation/AnimInstance.cpp
--- a/Engine/Source/Runtime/Engine/Private/Animation/AnimInstance.cpp
+++ b/Engine/Source/Runtime/Engine/Private/Animation/AnimInstance.cpp
@@ -2680,9 +2680,12 @@ void UAnimInstance::PerformLinkedLayerOverlayOperation(TSubclassOf<UAnimInstance
 					if(FoundFunction->bImplemented)
 					{
 						UClass* ClassToSet = InClassSelectorFunction(NewClass, Layer);
-						TMap<FName, TArray<FAnimNode_LinkedAnimLayer*, TInlineAllocator<4>>, TInlineSetAllocator<4>>& ClassLayerNodesToSet = LayerNodesToSet.FindOrAdd(ClassToSet);
-						TArray<FAnimNode_LinkedAnimLayer*, TInlineAllocator<4>>& LayerNodes = ClassLayerNodesToSet.FindOrAdd(FoundFunction->Group);
-						LayerNodes.Add(Layer);
+						if (ClassToSet != nullptr)
+						{
+							TMap<FName, TArray<FAnimNode_LinkedAnimLayer*, TInlineAllocator<4>>, TInlineSetAllocator<4>>& ClassLayerNodesToSet = LayerNodesToSet.FindOrAdd(ClassToSet);
+							TArray<FAnimNode_LinkedAnimLayer*, TInlineAllocator<4>>& LayerNodes = ClassLayerNodesToSet.FindOrAdd(FoundFunction->Group);
+							LayerNodes.Add(Layer);
+						}
 					}
 				}
 			}
@@ -2690,9 +2693,12 @@ void UAnimInstance::PerformLinkedLayerOverlayOperation(TSubclassOf<UAnimInstance
 			{
 				// Add null classes so we clear the node's instance below
 				UClass* ClassToSet = InClassSelectorFunction(NewClass, Layer);
-				TMap<FName, TArray<FAnimNode_LinkedAnimLayer*, TInlineAllocator<4>>, TInlineSetAllocator<4>>& ClassLayerNodesToSet = LayerNodesToSet.FindOrAdd(ClassToSet);
-				TArray<FAnimNode_LinkedAnimLayer*, TInlineAllocator<4>>& LayerNodes = ClassLayerNodesToSet.FindOrAdd(NAME_None);
-				LayerNodes.Add(Layer);
+				if (ClassToSet != nullptr)
+				{
+					TMap<FName, TArray<FAnimNode_LinkedAnimLayer*, TInlineAllocator<4>>, TInlineSetAllocator<4>>& ClassLayerNodesToSet = LayerNodesToSet.FindOrAdd(ClassToSet);
+					TArray<FAnimNode_LinkedAnimLayer*, TInlineAllocator<4>>& LayerNodes = ClassLayerNodesToSet.FindOrAdd(NAME_None);
+					LayerNodes.Add(Layer);
+				}
 			}
 		}
 
@@ -2874,8 +2880,13 @@ void UAnimInstance::LinkAnimClassLayers(TSubclassOf<UAnimInstance> InClass)
 
 void UAnimInstance::UnlinkAnimClassLayers(TSubclassOf<UAnimInstance> InClass)
 {
-	auto ConditionallySelectDefaultClass = [](UClass* InResolvedClass, FAnimNode_LinkedAnimLayer* InLayerNode)
+	auto ConditionallySelectDefaultClass = [](UClass* InResolvedClass, FAnimNode_LinkedAnimLayer* InLayerNode) -> UClass*
 	{
+		if (InLayerNode->GetTargetInstance<UAnimInstance>() == nullptr)
+		{
+			return nullptr;
+		}
+
 		if (InResolvedClass != nullptr && InLayerNode->GetTargetInstance<UAnimInstance>()->GetClass() == InResolvedClass)
 		{
 			// Reset to default if the classes match
