Expose constraints properties to blueprints ( Linear/Angular Limits/Motor and newly added Plasticity properties )
- Create a blueprint libraries with get/set constraint related functions that use a Constraint object as input
- Added methods to ConstraintInstance to be called by the blueprint library
- Added blueprint function in PhysicsConstraintComponent to expose the Constraint object
- Added blueprint function in SkeletalMeshComponent to get a all constraints / constraint object by name and with a "terminated" filter

diff --git a/Engine/Source/Runtime/Engine/Classes/Components/SkeletalMeshComponent.h b/Engine/Source/Runtime/Engine/Classes/Components/SkeletalMeshComponent.h
--- a/Engine/Source/Runtime/Engine/Classes/Components/SkeletalMeshComponent.h
+++ b/Engine/Source/Runtime/Engine/Classes/Components/SkeletalMeshComponent.h
@@ -1965,6 +1965,9 @@ public:
 	/** Find instance of the constraint that matches the name supplied. */
 	FConstraintInstance* FindConstraintInstance(FName ConName);
 
+	/** get instance of the constraint that matches the index . */
+	FConstraintInstance* GetConstraintInstanceByIndex(uint32 Index);
+
 	/** Utility which returns total mass of all bones below the supplied one in the hierarchy (including this one). */
 	float GetTotalMassBelowBone(FName InBoneName);
 
@@ -2110,7 +2113,21 @@ public:
 	UFUNCTION(BlueprintCallable, Category = "Components|SkeletalMesh", meta = (Keywords = "Constraint"))
 	void BreakConstraint(FVector Impulse, FVector HitLocation, FName InBoneName);
 
-	/** Sets the Angular Motion Ranges for a named bone
+	/** Gets a constraint by its name 
+	* @param ConstraintName		name of the constraint
+	* @param IncludesTerminated wether or not to return a terminated constraint
+	* */
+	UFUNCTION(BlueprintCallable, Category = "Components|SkeletalMesh", meta = (Keywords = "Components|SkeletalMesh"))
+	FConstraintInstanceAccessor GetConstraintByName(FName ConstraintName, bool bIncludesTerminated);
+
+	/** Gets all the constraints
+	* @param IncludesTerminated wether or not to return terminated constraints
+	* @param OutConstraints returned list of constraints matching the parameters
+	* */
+	UFUNCTION(BlueprintCallable, Category = "Components|SkeletalMesh", meta = (Keywords = "Components|SkeletalMesh"))
+	void GetConstraints(bool bIncludesTerminated, TArray<FConstraintInstanceAccessor>& OutConstraints);
+
+	/** Sets the Angular Motion Ranges for a named constraint
 	*  @param InBoneName  Name of bone to adjust constraint ranges for
 	*  @param Swing1LimitAngle	 Size of limit in degrees, 0 means locked, 180 means free
 	*  @param TwistLimitAngle	 Size of limit in degrees, 0 means locked, 180 means free
diff --git a/Engine/Source/Runtime/Engine/Classes/PhysicsEngine/ConstraintDrives.h b/Engine/Source/Runtime/Engine/Classes/PhysicsEngine/ConstraintDrives.h
--- a/Engine/Source/Runtime/Engine/Classes/PhysicsEngine/ConstraintDrives.h
+++ b/Engine/Source/Runtime/Engine/Classes/PhysicsEngine/ConstraintDrives.h
@@ -115,6 +115,7 @@ private:
 	void SetLinearPositionDrive(bool bEnableXDrive, bool bEnableYDrive, bool bEnableZDrive);
 	void SetLinearVelocityDrive(bool bEnableXDrive, bool bEnableYDrive, bool bEnableZDrive);
 	void SetDriveParams(float InStiffness, float InDamping, float InForceLimit);
+	void GetDriveParams(float& OutStiffness, float& OutDamping, float& OutForceLimit) const;
 };
 
 
@@ -189,5 +190,6 @@ private:
 	void SetAngularVelocityDriveTwistAndSwing(bool InEnableTwistDrive, bool InEnableSwingDrive);
 	void SetAngularVelocityDriveSLERP(bool InEnableSLERP);
 	void SetDriveParams(float InStiffness, float InDamping, float InForceLimit);
+	void GetDriveParams(float& OutStiffness, float& OutDamping, float& OutForceLimit) const;
 	void SetAngularDriveMode(EAngularDriveMode::Type DriveMode);
 };
diff --git a/Engine/Source/Runtime/Engine/Classes/PhysicsEngine/ConstraintInstance.h b/Engine/Source/Runtime/Engine/Classes/PhysicsEngine/ConstraintInstance.h
--- a/Engine/Source/Runtime/Engine/Classes/PhysicsEngine/ConstraintInstance.h
+++ b/Engine/Source/Runtime/Engine/Classes/PhysicsEngine/ConstraintInstance.h
@@ -256,6 +256,16 @@ public:
 		return ProfileInstance.LinearLimit.Limit;
 	}
 
+	/** Sets the Linear XYZ Motion Type and the limit distance (Note distance is the same for all 3 axes) */
+	void SetLinearLimits(ELinearConstraintMotion XConstraintType, ELinearConstraintMotion YConstraintType, ELinearConstraintMotion ZConstraintType, float InLinearLimitSize)
+	{
+		ProfileInstance.LinearLimit.XMotion = XConstraintType;
+		ProfileInstance.LinearLimit.YMotion = YConstraintType;
+		ProfileInstance.LinearLimit.ZMotion = ZConstraintType;
+		ProfileInstance.LinearLimit.Limit = InLinearLimitSize;
+		UpdateLinearLimit();
+	}
+
 	/** Gets the motion type for the linear X-axis limit. */
 	ELinearConstraintMotion GetLinearXMotion() const
 	{
@@ -263,18 +273,17 @@ public:
 	}
 
 	/** Sets the Linear XMotion type */
-	void SetLinearXMotion(ELinearConstraintMotion ConstraintType)
+	void SetLinearXMotion(ELinearConstraintMotion XConstraintType)
 	{
-		ProfileInstance.LinearLimit.XMotion = ConstraintType;
-		UpdateLinearLimit();
+		const FLinearConstraint& prevLimits = ProfileInstance.LinearLimit;
+		SetLinearLimits(XConstraintType, prevLimits.YMotion, prevLimits.ZMotion, prevLimits.Limit);
 	}
 
 	/** Sets the LinearX Motion Type and the limit distance (Note distance is the same for all 3 axes) */
-	void SetLinearXLimit(ELinearConstraintMotion ConstraintType, float InLinearLimitSize)
+	void SetLinearXLimit(ELinearConstraintMotion XConstraintType, float InLinearLimitSize)
 	{
-		ProfileInstance.LinearLimit.XMotion = ConstraintType;
-		ProfileInstance.LinearLimit.Limit = InLinearLimitSize;
-		UpdateLinearLimit();
+		const FLinearConstraint& prevLimits = ProfileInstance.LinearLimit;
+		SetLinearLimits(XConstraintType, prevLimits.YMotion, prevLimits.ZMotion, InLinearLimitSize);
 	}
 
 	/** Gets the motion type for the linear Y-axis limit. */
@@ -284,18 +293,17 @@ public:
 	}
 
 	/** Sets the Linear YMotion type */
-	void SetLinearYMotion(ELinearConstraintMotion ConstraintType)
+	void SetLinearYMotion(ELinearConstraintMotion YConstraintType)
 	{
-		ProfileInstance.LinearLimit.YMotion = ConstraintType;
-		UpdateLinearLimit();
+		const FLinearConstraint& prevLimits = ProfileInstance.LinearLimit;
+		SetLinearLimits(prevLimits.XMotion, YConstraintType, prevLimits.ZMotion, prevLimits.Limit);
 	}
 
 	/** Sets the LinearY Motion Type and the limit distance (Note distance is the same for all 3 axes) */
-	void SetLinearYLimit(ELinearConstraintMotion ConstraintType, float InLinearLimitSize)
+	void SetLinearYLimit(ELinearConstraintMotion YConstraintType, float InLinearLimitSize)
 	{
-		ProfileInstance.LinearLimit.YMotion = ConstraintType;
-		ProfileInstance.LinearLimit.Limit = InLinearLimitSize;
-		UpdateLinearLimit();
+		const FLinearConstraint& prevLimits = ProfileInstance.LinearLimit;
+		SetLinearLimits(prevLimits.XMotion, YConstraintType, prevLimits.ZMotion, InLinearLimitSize);
 	}
 
 	/** Gets the motion type for the linear Z-axis limit. */
@@ -305,18 +313,17 @@ public:
 	}
 
 	/** Sets the Linear ZMotion type */
-	void SetLinearZMotion(ELinearConstraintMotion ConstraintType)
+	void SetLinearZMotion(ELinearConstraintMotion ZConstraintType)
 	{
-		ProfileInstance.LinearLimit.ZMotion = ConstraintType;
-		UpdateLinearLimit();
+		const FLinearConstraint& prevLimits = ProfileInstance.LinearLimit;
+		SetLinearLimits(prevLimits.XMotion, prevLimits.YMotion, ZConstraintType, prevLimits.Limit);
 	}
 
 	/** Sets the LinearZ Motion Type and the limit distance (Note distance is the same for all 3 axes) */
-	void SetLinearZLimit(ELinearConstraintMotion ConstraintType, float InLinearLimitSize)
+	void SetLinearZLimit(ELinearConstraintMotion ZConstraintType, float InLinearLimitSize)
 	{
-		ProfileInstance.LinearLimit.ZMotion = ConstraintType;
-		ProfileInstance.LinearLimit.Limit = InLinearLimitSize;
-		UpdateLinearLimit();
+		const FLinearConstraint& prevLimits = ProfileInstance.LinearLimit;
+		SetLinearLimits(prevLimits.XMotion, prevLimits.YMotion, ZConstraintType, InLinearLimitSize);
 	}
 
 	/** Gets the motion type for the swing1 of the cone constraint */
@@ -480,6 +487,18 @@ public:
 		UpdateBreakable();
 	}
 
+	/** Gets Whether it is possible to break the joint with linear force */
+	bool IsLinearBreakable() const
+	{
+		return ProfileInstance.bLinearBreakable;
+	}
+
+	/** Gets linear force needed to break the joint */
+	float GetLinearBreakThreshold() const
+	{
+		return ProfileInstance.LinearBreakThreshold;
+	}
+
 	/** Sets the Linear Plasticity properties
 	*	@param bInLinearPlasticity 	Whether it is possible to reset the target angles
 	*	@param InLinearPlasticityThreshold	Delta from target needed to reset the target joint
@@ -502,6 +521,17 @@ public:
 		UpdateBreakable();
 	}
 
+	/** Gets Whether it is possible to break the joint with angular force */
+	bool IsAngularBreakable() const
+	{
+		return ProfileInstance.bAngularBreakable;
+	}
+
+	/** Gets Torque needed to break the joint */
+	float GetAngularBreakThreshold() const
+	{
+		return ProfileInstance.AngularBreakThreshold;
+	}
 
 	/** Sets the Angular Plasticity properties
 	*	@param bInAngularPlasticity 	Whether it is possible to reset the target angles
@@ -514,6 +544,18 @@ public:
 		UpdatePlasticity();
 	}
 
+	/** Gets Whether it is possible to reset the target angles */
+	bool HasAngularPlasticity() const
+	{
+		return ProfileInstance.bAngularPlasticity;
+	}
+
+	/** Gets Delta from target needed to reset the target joint */
+	float GetAngularPlasticityThreshold() const
+	{
+		return ProfileInstance.AngularPlasticityThreshold;
+	}
+
 	// @todo document
 	void CopyConstraintGeometryFrom(const FConstraintInstance* FromInstance);
 
@@ -529,6 +571,24 @@ public:
 	/** Set which linear position drives are enabled */
 	void SetLinearPositionDrive(bool bEnableXDrive, bool bEnableYDrive, bool bEnableZDrive);
 
+	/** Get which linear position drives is enabled on XAxis */
+	bool IsLinearPositionDriveXEnabled() const
+	{
+		return ProfileInstance.LinearDrive.XDrive.bEnablePositionDrive;
+	}
+
+	/** Get which linear position drives is enabled on YAxis */
+	bool IsLinearPositionDriveYEnabled() const
+	{
+		return ProfileInstance.LinearDrive.YDrive.bEnablePositionDrive;
+	}
+
+	/** Get which linear position drives is enabled on ZAxis */
+	bool IsLinearPositionDriveZEnabled() const
+	{
+		return ProfileInstance.LinearDrive.ZDrive.bEnablePositionDrive;
+	}
+
 	/** Whether the linear position drive is enabled */
 	bool IsLinearPositionDriveEnabled() const
 	{
@@ -538,21 +598,54 @@ public:
 	/** Set the linear drive's target position position */
 	void SetLinearPositionTarget(const FVector& InPosTarget);
 
+	/** Get the linear drive's target position position */
+	const FVector& GetLinearPositionTarget()
+	{
+		return ProfileInstance.LinearDrive.PositionTarget;
+	}
+
 	/** Set which linear velocity drives are enabled */
 	void SetLinearVelocityDrive(bool bEnableXDrive, bool bEnableYDrive, bool bEnableZDrive);
 
+	/** Get which linear position drives is enabled on XAxis */
+	bool IsLinearVelocityDriveXEnabled() const
+	{
+		return ProfileInstance.LinearDrive.XDrive.bEnableVelocityDrive;
+	}
+
+	/** Get which linear position drives is enabled on YAxis */
+	bool IsLinearVelocityDriveYEnabled() const
+	{
+		return ProfileInstance.LinearDrive.YDrive.bEnableVelocityDrive;
+	}
+
+	/** Get which linear position drives is enabled on ZAxis */
+	bool IsLinearVelocityDriveZEnabled() const
+	{
+		return ProfileInstance.LinearDrive.ZDrive.bEnableVelocityDrive;
+	}
+
 	/** Whether the linear velocity drive is enabled */
 	bool IsLinearVelocityDriveEnabled() const
 	{
 		return ProfileInstance.LinearDrive.IsVelocityDriveEnabled();
 	}
 
-	/** Set the linear drive's target velocity*/
+	/** Set the linear drive's target velocity */
 	void SetLinearVelocityTarget(const FVector& InVelTarget);
 
+	/** Get the linear drive's target velocity */
+	const FVector& GetLinearVelocityTarget()
+	{
+		return ProfileInstance.LinearDrive.VelocityTarget;
+	}
+
 	/** Set the linear drive's strength parameters */
 	void SetLinearDriveParams(float InPositionStrength, float InVelocityStrength, float InForceLimit);
 
+	/** Get the linear drive's strength parameters */
+	void GetLinearDriveParams(float& OutPositionStrength, float& OutVelocityStrength, float& OutForceLimit);
+
 	UE_DEPRECATED(4.15, "Please call SetOrientationDriveTwistAndSwing. Note the order of bools is reversed (Make sure to pass Twist and then Swing)")
 	void SetAngularPositionDrive(bool bInEnableSwingDrive, bool bInEnableTwistDrive)
 	{
@@ -562,9 +655,18 @@ public:
 	/** Set which twist and swing orientation drives are enabled. Only applicable when Twist And Swing drive mode is used */
 	void SetOrientationDriveTwistAndSwing(bool bInEnableTwistDrive, bool bInEnableSwingDrive);
 
+	/** Get which twist and swing orientation drives are enabled. Only applicable when Twist And Swing drive mode is used */
+	void GetOrientationDriveTwistAndSwing(bool& bOutEnableTwistDrive, bool& bOutEnableSwingDrive);
+
 	/** Set whether the SLERP angular position drive is enabled. Only applicable when SLERP drive mode is used */
 	void SetOrientationDriveSLERP(bool bInEnableSLERP);
 
+	/** Get whether the SLERP angular position drive is enabled. Only applicable when SLERP drive mode is used */
+	bool GetOrientationDriveSLERP()
+	{
+		return ProfileInstance.AngularDrive.SlerpDrive.bEnablePositionDrive;
+	}
+
 	/** Whether the angular orientation drive is enabled */
 	bool IsAngularOrientationDriveEnabled() const
 	{
@@ -574,6 +676,12 @@ public:
 	/** Set the angular drive's orientation target*/
 	void SetAngularOrientationTarget(const FQuat& InPosTarget);
 
+	/** Get the angular drive's orientation target*/
+	const FRotator& GetAngularOrientationTarget() const
+	{
+		return ProfileInstance.AngularDrive.OrientationTarget;
+	}
+
 	UE_DEPRECATED(4.15, "Please call SetAngularVelocityDriveTwistAndSwing. Note the order of bools is reversed (Make sure to pass Twist and then Swing)")
 	void SetAngularVelocityDrive(bool bInEnableSwingDrive, bool bInEnableTwistDrive)
 	{
@@ -583,9 +691,18 @@ public:
 	/** Set which twist and swing angular velocity drives are enabled. Only applicable when Twist And Swing drive mode is used */
 	void SetAngularVelocityDriveTwistAndSwing(bool bInEnableTwistDrive, bool bInEnableSwingDrive);
 
+	/** Get which twist and swing angular velocity drives are enabled. Only applicable when Twist And Swing drive mode is used */
+	void GetAngularVelocityDriveTwistAndSwing(bool& bOutEnableTwistDrive, bool& bOutEnableSwingDrive);
+
 	/** Set whether the SLERP angular velocity drive is enabled. Only applicable when SLERP drive mode is used */
 	void SetAngularVelocityDriveSLERP(bool bInEnableSLERP);
 
+	/** Get whether the SLERP angular velocity drive is enabled. Only applicable when SLERP drive mode is used */
+	bool GetAngularVelocityDriveSLERP()
+	{
+		return ProfileInstance.AngularDrive.SlerpDrive.bEnableVelocityDrive;
+	}
+
 	/** Whether the angular velocity drive is enabled */
 	bool IsAngularVelocityDriveEnabled() const
 	{
@@ -595,12 +712,27 @@ public:
 	/** Set the angular drive's angular velocity target*/
 	void SetAngularVelocityTarget(const FVector& InVelTarget);
 
+	/** Get the angular drive's angular velocity target*/
+	const FVector& GetAngularVelocityTarget() const
+	{
+		return ProfileInstance.AngularDrive.AngularVelocityTarget;
+	}
+
 	/** Set the angular drive's strength parameters*/
 	void SetAngularDriveParams(float InSpring, float InDamping, float InForceLimit);
 
+	/** Get the angular drive's strength parameters*/
+	void GetAngularDriveParams(float& OutSpring, float& OutDamping, float& OutForceLimit) const;
+
 	/** Set the angular drive mode */
 	void SetAngularDriveMode(EAngularDriveMode::Type DriveMode);
 
+	/** Set the angular drive mode */
+	EAngularDriveMode::Type GetAngularDriveMode()
+	{
+		return ProfileInstance.AngularDrive.AngularDriveMode;
+	}
+
 	/** Refreshes the physics engine joint's linear limits. Only applicable if the joint has been created already.*/
 	void UpdateLinearLimit();
 
@@ -689,6 +821,14 @@ public:
 	/** Turn off linear and angular projection */
 	void DisableProjection();
 
+	/** Set projection parameters */
+	void SetProjectionParams(bool bEnableProjection, float ProjectionLinearAlphaOrTolerance, float ProjectionAngularAlphaOrTolerance);
+
+	/** Get projection parameters 
+	* Chaos returns alphas, PhysX returns tolerances
+	*/
+	void GetProjectionAlphasOrTolerances(float& ProjectionLinearAlphaOrTolerance, float& ProjectionAngularAlphaOrTolerance) const;
+
 	/** Whether parent domination is enabled (meaning the parent body cannot be be affected at all by a child) */
 	bool IsParentDominatesEnabled() const
 	{
@@ -887,4 +1027,33 @@ struct TStructOpsTypeTraits<FConstraintInstance> : public TStructOpsTypeTraitsBa
 		WithPostSerialize = true
 #endif
 	};
+};
+
+
+// Wrapping type around instance pointer to be returned per value in Blueprints
+USTRUCT(BlueprintType)
+struct ENGINE_API FConstraintInstanceAccessor
+{
+	GENERATED_USTRUCT_BODY()
+
+public:
+	FConstraintInstanceAccessor()
+		: Owner(nullptr)
+		, Index(0)
+	{}
+
+	FConstraintInstanceAccessor(const TWeakObjectPtr<UObject>& Owner, uint32 Index = 0)
+		: Owner(Owner)
+		, Index(Index)
+	{}
+
+	FConstraintInstance* Get() const;
+
+private:
+	UPROPERTY()
+	TWeakObjectPtr<UObject> Owner;
+
+	UPROPERTY()
+	uint32 Index;
+	
 };
\ No newline at end of file
diff --git a/Engine/Source/Runtime/Engine/Classes/PhysicsEngine/ConstraintInstanceBlueprintLibrary.h b/Engine/Source/Runtime/Engine/Classes/PhysicsEngine/ConstraintInstanceBlueprintLibrary.h
new file mode 100644
--- /dev/null
+++ b/Engine/Source/Runtime/Engine/Classes/PhysicsEngine/ConstraintInstanceBlueprintLibrary.h
@@ -0,0 +1,550 @@
+// Copyright Epic Games, Inc. All Rights Reserved.
+
+#pragma once
+
+#include "Kismet/BlueprintFunctionLibrary.h"
+#include "Engine/EngineTypes.h"
+#include "PhysicsEngine/ConstraintInstance.h"
+
+#include "ConstraintInstanceBlueprintLibrary.generated.h"
+
+UCLASS()
+class ENGINE_API UConstraintInstanceBlueprintLibrary : public UBlueprintFunctionLibrary
+{
+	GENERATED_BODY()
+
+public:
+	//---------------------------------------------------------------------------------------------------
+	//
+	// CONSTRAINT BEHAVIOR 
+	//
+	//---------------------------------------------------------------------------------------------------
+
+	/** Sets whether bodies attched to the constraint can collide or not
+	*	@param ConstraintInstance	Instance of the constraint to change
+	*	@param bDisableCollision	true to disable collision between constrained bodies
+	*/
+	UFUNCTION(BlueprintCallable, Category = "Physics|Constraints")
+	static void SetDisableCollision(
+			UPARAM(ref) FConstraintInstanceAccessor& Accessor,
+			bool bDisableCollision
+		);
+
+	/** Gets whether bodies attched to the constraint can collide or not
+	*	@param ConstraintInstance	Instance of the constraint to change
+	*/
+	UFUNCTION(BlueprintCallable, Category = "Physics|Constraints")
+	static bool GetDisableCollsion(FConstraintInstanceAccessor& Accessor);
+
+	/** Sets projection parameters of the constraint
+	*	@param ConstraintInstance		Instance of the constraint to change
+	*	@param bEnableProjection		true to enable projection
+	*	@param ProjectionLinearAlpha	how much linear projection to apply in [0,1] range
+	*	@param ProjectionAngularAlpha	how much angular projection to apply in [0,1] range
+	*/
+	UFUNCTION(BlueprintCallable, Category = "Physics|Constraints")
+	static void SetProjectionParams(
+		UPARAM(ref) FConstraintInstanceAccessor& Accessor,
+		bool bEnableProjection,
+		float ProjectionLinearAlpha,
+		float ProjectionAngularAlpha
+	);
+
+	/** Gets projection parameters of the constraint
+	*	@param ConstraintInstance		Instance of the constraint to change
+	*	@param bEnableProjection		true to enable projection
+	*	@param ProjectionLinearAlpha	how much linear projection to apply in [0,1] range
+	*	@param ProjectionAngularAlpha	how much angular projection to apply in [0,1] range
+	*/
+	UFUNCTION(BlueprintCallable, Category = "Physics|Constraints")
+	static void GetProjectionParams(
+		UPARAM(ref) FConstraintInstanceAccessor& Accessor,
+		bool& bEnableProjection,
+		float& ProjectionLinearAlpha,
+		float& ProjectionAngularAlpha
+	);
+
+	/** Sets whether the parent body is not affected by it's child motion 
+	*	@param ConstraintInstance	Instance of the constraint to change
+	*	@param bParentDominates		true to avoid the parent being affected by its child motion
+	*/
+	UFUNCTION(BlueprintCallable, Category = "Physics|Constraints")
+	static void SetParentDominates(
+		UPARAM(ref) FConstraintInstanceAccessor& Accessor,
+		bool bParentDominates
+	);
+
+	/** Gets whether the parent body is not affected by it's child motion
+	*	@param ConstraintInstance	Instance of the constraint to change
+	*/
+	UFUNCTION(BlueprintCallable, Category = "Physics|Constraints")
+	static bool GetParentDominates(FConstraintInstanceAccessor& Accessor);
+
+	//---------------------------------------------------------------------------------------------------
+	//
+	// LINEAR LIMITS
+	//
+	//---------------------------------------------------------------------------------------------------
+
+	/** Sets Constraint Linear Motion Ranges
+	*	@param ConstraintInstance	Instance of the constraint to change
+	*	@param XMotion	Type of motion along the X axis
+	*	@param YMotion	Type of motion along the Y axis
+	*	@param ZMotion	Type of motion along the Z axis
+	*	@param Limit	linear limit to apply to all axis
+	*/
+	UFUNCTION(BlueprintCallable, Category = "Physics|Constraints")
+	static void SetLinearLimits(
+		UPARAM(ref) FConstraintInstanceAccessor& Accessor,
+		TEnumAsByte<ELinearConstraintMotion> XMotion,
+		TEnumAsByte<ELinearConstraintMotion> YMotion,
+		TEnumAsByte<ELinearConstraintMotion> ZMotion,
+		float Limit
+	);
+
+	/** Gets Constraint Linear Motion Ranges
+	*	@param ConstraintInstance	Instance of the constraint to query
+	*	@param XMotion	Type of motion along the X axis
+	*	@param YMotion	Type of motion along the Y axis
+	*	@param ZMotion	Type of motion along the Z axis
+	*	@param Limit	linear limit applied to all axis
+	*/
+	UFUNCTION(BlueprintCallable, Category = "Physics|Constraints")
+	static void GetLinearLimits(
+		FConstraintInstanceAccessor& Accessor,
+		TEnumAsByte<ELinearConstraintMotion>& XMotion,
+		TEnumAsByte<ELinearConstraintMotion>& YMotion,
+		TEnumAsByte<ELinearConstraintMotion>& ZMotion,
+		float& Limit
+	);
+
+	/** Sets the Linear Breakable properties
+	*	@param ConstraintInstance	Instance of the constraint to change
+	*	@param bLinearBreakable		Whether it is possible to break the joint with linear force
+	*	@param LinearBreakThreshold	Force needed to break the joint
+	*/
+	UFUNCTION(BlueprintCallable, Category = "Physics|Constraints")
+	static void SetLinearBreakable(
+		UPARAM(ref) FConstraintInstanceAccessor& Accessor, 
+		bool bLinearBreakable, 
+		float LinearBreakThreshold
+	);
+
+	/** Gets Constraint Linear Breakable properties
+	*	@param ConstraintInstance	Instance of the constraint to query
+	*	@param bLinearBreakable		Whether it is possible to break the joint with linear force
+	*	@param LinearBreakThreshold	Force needed to break the joint
+	*/
+	UFUNCTION(BlueprintCallable, Category = "Physics|Constraints")
+	static void GetLinearBreakable(
+		UPARAM(ref) FConstraintInstanceAccessor& Accessor,
+		bool& bLinearBreakable,
+		float& LinearBreakThreshold
+	);
+
+	//---------------------------------------------------------------------------------------------------
+	//
+	// ANGULAR LIMITS 
+	//
+	//---------------------------------------------------------------------------------------------------
+
+	/** Sets COnstraint Angular Motion Ranges
+	*	@param ConstraintInstance	Instance of the constraint to change
+	*	@param Swing1MotionType		Type of swing motion ( first axis )
+	*	@param Swing1LimitAngle		Size of limit in degrees in [0, 180] range
+	*   @param Swing2MotionType		Type of swing motion ( second axis )
+	*	@param Swing2LimitAngle		Size of limit in degrees in [0, 180] range
+	*   @param TwistMotionType		Type of twist motion
+	*	@param TwistLimitAngle		Size of limit in degrees in [0, 180] range	
+	*/
+	UFUNCTION(BlueprintCallable, Category = "Physics|Constraints")
+	static void SetAngularLimits(
+		UPARAM(ref) FConstraintInstanceAccessor& Accessor, 
+		TEnumAsByte<EAngularConstraintMotion> Swing1MotionType,
+		float Swing1LimitAngle, 
+		TEnumAsByte<EAngularConstraintMotion> Swing2MotionType,
+		float Swing2LimitAngle,
+		TEnumAsByte<EAngularConstraintMotion> TwistMotionType,
+		float TwistLimitAngle
+	);
+
+	/** Gets Constraint Angular Motion Ranges
+	*	@param ConstraintInstance	Instance of the constraint to query
+	*	@param Swing1MotionType		Type of swing motion ( first axis )
+	*	@param Swing1LimitAngle		Size of limit in degrees in [0, 180] range
+	*   @param Swing2MotionType		Type of swing motion ( second axis )
+	*	@param Swing2LimitAngle		Size of limit in degrees in [0, 180] range
+	*   @param TwistMotionType		Type of twist motion
+	*	@param TwistLimitAngle		Size of limit in degrees in [0, 180] range
+	*/
+	UFUNCTION(BlueprintCallable, Category = "Physics|Constraints")
+	static void GetAngularLimits(
+		UPARAM(ref) FConstraintInstanceAccessor& Accessor,
+		TEnumAsByte<EAngularConstraintMotion>& Swing1MotionType,
+		float& Swing1LimitAngle,
+		TEnumAsByte<EAngularConstraintMotion>& Swing2MotionType,
+		float& Swing2LimitAngle,
+		TEnumAsByte<EAngularConstraintMotion>& TwistMotionType,
+		float& TwistLimitAngle
+	);
+
+	/** Sets Constraint Angular Breakable properties
+	*	@param ConstraintInstance		Instance of the constraint to change
+	*	@param bAngularBreakable		Whether it is possible to break the joint with angular force
+	*	@param AngularBreakThreshold	Torque needed to break the joint
+	*/
+	UFUNCTION(BlueprintCallable, Category = "Physics|Constraints")
+	static void SetAngularBreakable(
+		UPARAM(ref) FConstraintInstanceAccessor& Accessor,
+		bool bAngularBreakable, 
+		float AngularBreakThreshold
+	);
+
+	/** Gets Constraint Angular Breakable properties
+	*	@param ConstraintInstance		Instance of the constraint to query
+	*	@param bAngularBreakable		Whether it is possible to break the joint with angular force
+	*	@param AngularBreakThreshold	Torque needed to break the joint
+	*/
+	UFUNCTION(BlueprintCallable, Category = "Physics|Constraints")
+	static void GetAngularBreakable(
+		UPARAM(ref) FConstraintInstanceAccessor& Accessor,
+		bool& bAngularBreakable,
+		float& AngularBreakThreshold
+	);
+
+	/** Sets Constraint Angular Plasticity properties
+	*	@param ConstraintInstance			Instance of the constraint to change
+	*	@param bAngularPlasticity			Whether it is possible to reset the target angle from the angular displacement
+	*	@param AngularPlasticityThreshold	Degrees needed to reset the rest state of the joint
+	*/
+	UFUNCTION(BlueprintCallable, Category = "Physics|Constraints")
+	static void SetAngularPlasticity(
+		UPARAM(ref) FConstraintInstanceAccessor& Accessor,
+		bool bAngularPlasticity, 
+		float AngularPlasticityThreshold
+	);
+
+	/** Sets Constraint Angular Plasticity properties
+	*	@param ConstraintInstance			Instance of the constraint to query
+	*	@param bAngularPlasticity			Whether it is possible to reset the target angle from the angular displacement
+	*	@param AngularPlasticityThreshold	Degrees needed to reset the rest state of the joint
+	*/
+	UFUNCTION(BlueprintCallable, Category = "Physics|Constraints")
+	static void GetAngularPlasticity(
+		UPARAM(ref) FConstraintInstanceAccessor& Accessor,
+		bool& bAngularPlasticity,
+		float& AngularPlasticityThreshold
+	);
+
+	//---------------------------------------------------------------------------------------------------
+	//
+	// LINEAR MOTOR
+	//
+	//---------------------------------------------------------------------------------------------------
+
+	/** Enables/Disables linear position drive
+	*	@param ConstraintInstance	Instance of the constraint to change
+	*	@param bEnableDriveX		Indicates whether the drive for the X-Axis should be enabled
+	*	@param bEnableDriveY		Indicates whether the drive for the Y-Axis should be enabled
+	*	@param bEnableDriveZ		Indicates whether the drive for the Z-Axis should be enabled
+	*/
+	UFUNCTION(BlueprintCallable, Category = "Physics|Constraints")
+	static void SetLinearPositionDrive(
+		UPARAM(ref) FConstraintInstanceAccessor& Accessor, 
+		bool bEnableDriveX, 
+		bool bEnableDriveY, 
+		bool bEnableDriveZ
+	);
+
+	/** Gets whether linear position drive is enabled or not
+	*	@param ConstraintInstance	Instance of the constraint to query
+	*	@param bOutEnableDriveX		Indicates whether the drive for the X-Axis is enabled
+	*	@param bOutEnableDriveY		Indicates whether the drive for the Y-Axis is enabled
+	*	@param bOutEnableDriveZ		Indicates whether the drive for the Z-Axis is enabled
+	*/
+	UFUNCTION(BlueprintCallable, Category = "Physics|Constraints")
+	static void GetLinearPositionDrive(
+		UPARAM(ref) FConstraintInstanceAccessor& Accessor,
+		bool& bOutEnableDriveX,
+		bool& bOutEnableDriveY,
+		bool& bOutEnableDriveZ
+	);
+
+	/** Enables/Disables linear velocity drive
+	*	@param ConstraintInstance	Instance of the constraint to change
+	*	@param bEnableDriveX		Indicates whether the drive for the X-Axis should be enabled
+	*	@param bEnableDriveY		Indicates whether the drive for the Y-Axis should be enabled
+	*	@param bEnableDriveZ		Indicates whether the drive for the Z-Axis should be enabled
+	*/
+	UFUNCTION(BlueprintCallable, Category = "Physics|Constraints")
+	static void SetLinearVelocityDrive(
+		UPARAM(ref) FConstraintInstanceAccessor& Accessor,
+		bool bEnableDriveX, 
+		bool bEnableDriveY, 
+		bool bEnableDriveZ
+	);
+
+	/** Gets whether linear velocity drive is enabled or not
+	*	@param ConstraintInstance	Instance of the constraint to query
+	*	@param bOutEnableDriveX		Indicates whether the drive for the X-Axis is enabled
+	*	@param bOutEnableDriveY		Indicates whether the drive for the Y-Axis is enabled
+	*	@param bOutEnableDriveZ		Indicates whether the drive for the Z-Axis is enabled
+	*/
+	UFUNCTION(BlueprintCallable, Category = "Physics|Constraints")
+	static void GetLinearVelocityDrive(
+		UPARAM(ref) FConstraintInstanceAccessor& Accessor,
+		bool& bOutEnableDriveX,
+		bool& bOutEnableDriveY,
+		bool& bOutEnableDriveZ
+	);
+
+	/** Sets the target position for the linear drive.
+	*	@param ConstraintInstance	Instance of the constraint to change
+	*	@param InPosTarget			Target position
+	*/
+	UFUNCTION(BlueprintCallable, Category = "Physics|Constraints")
+	static void SetLinearPositionTarget(
+		UPARAM(ref) FConstraintInstanceAccessor& Accessor, 
+		const FVector& InPosTarget
+	);
+
+	/** Gets the target position for the linear drive.
+	*	@param ConstraintInstance	Instance of the constraint to query
+	*	@param OutPosTarget			Target position
+	*/
+	UFUNCTION(BlueprintCallable, Category = "Physics|Constraints")
+	static void GetLinearPositionTarget(
+		UPARAM(ref) FConstraintInstanceAccessor& Accessor,
+		FVector& OutPosTarget
+	);
+
+	/** Sets the target velocity for the linear drive.
+	*	@param ConstraintInstance	Instance of the constraint to change
+	*	@param InVelTarget			Target velocity
+	*/
+	UFUNCTION(BlueprintCallable, Category = "Physics|Constraints")
+	static void SetLinearVelocityTarget(
+		UPARAM(ref) FConstraintInstanceAccessor& Accessor, 
+		const FVector& InVelTarget
+	);
+
+	/** Gets the target velocity for the linear drive.
+	*	@param ConstraintInstance	Instance of the constraint to query
+	*	@param OutVelTarget			Target velocity
+	*/
+	UFUNCTION(BlueprintCallable, Category = "Physics|Constraints")
+	static void GetLinearVelocityTarget(
+		UPARAM(ref) FConstraintInstanceAccessor& Accessor,
+		FVector& OutVelTarget
+	);
+
+	/** Sets the drive params for the linear drive.
+	*	@param ConstraintInstance	Instance of the constraint to change
+	*	@param PositionStrength		Positional strength for the drive (stiffness)
+	*	@param VelocityStrength		Velocity strength of the drive (damping)
+	*	@param InForceLimit			Max force applied by the drive
+	*/
+	UFUNCTION(BlueprintCallable, Category = "Physics|Constraints")
+	static void SetLinearDriveParams(
+		UPARAM(ref) FConstraintInstanceAccessor& Accessor,
+		float PositionStrength, 
+		float VelocityStrength, 
+		float InForceLimit
+	);
+
+	/** Gets the drive params for the linear drive.
+	*	@param ConstraintInstance	Instance of the constraint to query
+	*	@param OutPositionStrength	Positional strength for the drive (stiffness)
+	*	@param OutVelocityStrength	Velocity strength of the drive (damping)
+	*	@param OutForceLimit		Max force applied by the drive
+	*/
+	UFUNCTION(BlueprintCallable, Category = "Physics|Constraints")
+	static void GetLinearDriveParams(
+		UPARAM(ref) FConstraintInstanceAccessor& Accessor,
+		float& OutPositionStrength,
+		float& OutVelocityStrength,
+		float& OutForceLimit
+	);
+
+	//---------------------------------------------------------------------------------------------------
+	//
+	// ANGULAR MOTOR
+	//
+	//---------------------------------------------------------------------------------------------------
+
+	/** Enables/Disables angular orientation drive. Only relevant if the AngularDriveMode is set to Twist and Swing
+	*	@param ConstraintInstance	Instance of the constraint to change
+	*	@param bEnableSwingDrive	Indicates whether the drive for the swing axis should be enabled. Only relevant if the AngularDriveMode is set to Twist and Swing
+	*	@param bEnableTwistDrive	Indicates whether the drive for the twist axis should be enabled. Only relevant if the AngularDriveMode is set to Twist and Swing
+	*/
+	UFUNCTION(BlueprintCallable, Category = "Physics|Constraints")
+	static void SetOrientationDriveTwistAndSwing(
+		UPARAM(ref) FConstraintInstanceAccessor& Accessor, 
+		bool bEnableTwistDrive, 
+		bool bEnableSwingDrive
+	);
+
+	/** Gets whether angular orientation drive are enabled. Only relevant if the AngularDriveMode is set to Twist and Swing
+	*	@param ConstraintInstance	Instance of the constraint to query
+	*	@param bOutEnableTwistDrive	Indicates whether the drive for the twist axis is enabled. Only relevant if the AngularDriveMode is set to Twist and Swing
+	*	@param bOutEnableSwingDrive	Indicates whether the drive for the swing axis is enabled. Only relevant if the AngularDriveMode is set to Twist and Swing
+	*/
+	UFUNCTION(BlueprintCallable, Category = "Physics|Constraints")
+	static void GetOrientationDriveTwistAndSwing(
+		UPARAM(ref) FConstraintInstanceAccessor& Accessor,
+		bool& bOutEnableTwistDrive,
+		bool& bOutEnableSwingDrive
+	);
+
+	/** Enables/Disables the angular orientation slerp drive. Only relevant if the AngularDriveMode is set to SLERP
+	*	@param ConstraintInstance	Instance of the constraint to change
+	*	@param bEnableSLERP			Indicates whether the SLERP drive should be enabled. Only relevant if the AngularDriveMode is set to SLERP
+	*/
+	UFUNCTION(BlueprintCallable, Category = "Physics|Constraints")
+	static void SetOrientationDriveSLERP(
+		UPARAM(ref) FConstraintInstanceAccessor& Accessor, 
+		bool bEnableSLERP
+	);
+
+	/** Gets whether the angular orientation slerp drive is enabled or not. Only relevant if the AngularDriveMode is set to SLERP
+	*	@param ConstraintInstance	Instance of the constraint to query
+	*	@param bOutEnableSLERP		Indicates whether the SLERP drive should be enabled. Only relevant if the AngularDriveMode is set to SLERP
+	*/
+	UFUNCTION(BlueprintCallable, Category = "Physics|Constraints")
+	static void GetOrientationDriveSLERP(
+		UPARAM(ref) FConstraintInstanceAccessor& Accessor,
+		bool& bOutEnableSLERP
+	);
+
+	/** Enables/Disables angular velocity twist and swing drive. Only relevant if the AngularDriveMode is set to Twist and Swing
+	*	@param ConstraintInstance	Instance of the constraint to change
+	*	@param bEnableTwistDrive	Indicates whether the drive for the twist axis should be enabled. Only relevant if the AngularDriveMode is set to Twist and Swing
+	*	@param bEnableSwingDrive	Indicates whether the drive for the swing axis should be enabled. Only relevant if the AngularDriveMode is set to Twist and Swing
+	*/
+	UFUNCTION(BlueprintCallable, Category = "Physics|Constraints")
+	static void SetAngularVelocityDriveTwistAndSwing(
+		UPARAM(ref) FConstraintInstanceAccessor& Accessor, 
+		bool bEnableTwistDrive, 
+		bool bEnableSwingDrive
+	);
+
+	/** Gets whether angular velocity twist and swing drive is enabled or not. Only relevant if the AngularDriveMode is set to Twist and Swing
+	*	@param ConstraintInstance	Instance of the constraint to query	
+	*	@param bOutEnableTwistDrive	Indicates whether the drive for the twist axis should be enabled. Only relevant if the AngularDriveMode is set to Twist and Swing
+	*   @param bOutEnableSwingDrive	Indicates whether the drive for the swing axis should be enabled. Only relevant if the AngularDriveMode is set to Twist and Swing
+	*/
+	UFUNCTION(BlueprintCallable, Category = "Physics|Constraints")
+	static void GetAngularVelocityDriveTwistAndSwing(
+		UPARAM(ref) FConstraintInstanceAccessor& Accessor,
+		bool& bOutEnableTwistDrive,
+		bool& bOutEnableSwingDrive
+	);
+
+	/** Enables/Disables the angular velocity slerp drive. Only relevant if the AngularDriveMode is set to SLERP
+	*	@param ConstraintInstance	Instance of the constraint to change
+	*	@param bEnableSLERP			Indicates whether the SLERP drive should be enabled. Only relevant if the AngularDriveMode is set to SLERP
+	*/
+	UFUNCTION(BlueprintCallable, Category = "Physics|Constraints")
+	static void SetAngularVelocityDriveSLERP(
+		UPARAM(ref) FConstraintInstanceAccessor& Accessor,
+		bool bEnableSLERP
+	);
+
+	/** Gets whether the angular velocity slerp drive is enabled or not. Only relevant if the AngularDriveMode is set to SLERP
+	*	@param ConstraintInstance	Instance of the constraint to change
+	*	@param bOutEnableSLERP		Indicates whether the SLERP drive is enabled. Only relevant if the AngularDriveMode is set to SLERP
+	*/
+	UFUNCTION(BlueprintCallable, Category = "Physics|Constraints")
+	static void GetAngularVelocityDriveSLERP(
+		UPARAM(ref) FConstraintInstanceAccessor& Accessor,
+		bool& bOutEnableSLERP
+	);
+
+	/** Switches the angular drive mode between SLERP and Twist And Swing
+	*	@param ConstraintInstance	Instance of the constraint to change
+	*	@param DriveMode	The angular drive mode to use. SLERP uses shortest spherical path, but will not work if any angular constraints are locked. Twist and Swing decomposes the path into the different angular degrees of freedom but may experience gimbal lock
+	*/
+	UFUNCTION(BlueprintCallable, Category = "Physics|Constraints")
+	static void SetAngularDriveMode(
+		UPARAM(ref) FConstraintInstanceAccessor& Accessor,
+		EAngularDriveMode::Type DriveMode
+	);
+
+	/** Gets the angular drive mode ( SLERP or Twist And Swing)
+	*	@param ConstraintInstance	Instance of the constraint to change
+	*	@param OutDriveMode	The angular drive mode to use. SLERP uses shortest spherical path, but will not work if any angular constraints are locked. Twist and Swing decomposes the path into the different angular degrees of freedom but may experience gimbal lock
+	*/
+	UFUNCTION(BlueprintCallable, Category = "Physics|Constraints")
+	static void GetAngularDriveMode(
+		UPARAM(ref) FConstraintInstanceAccessor& Accessor,
+		TEnumAsByte<EAngularDriveMode::Type>& OutDriveMode
+	);
+
+	/** Sets the target orientation for the angular drive.
+	*	@param ConstraintInstance	Instance of the constraint to change
+	*	@param InPosTarget			Target orientation
+	*/
+	UFUNCTION(BlueprintCallable, Category = "Physics|Constraints")
+	static void SetAngularOrientationTarget(
+		UPARAM(ref) FConstraintInstanceAccessor& Accessor,
+		const FRotator& InPosTarget
+	);
+
+	/** Gets the target orientation for the angular drive.
+	*	@param ConstraintInstance	Instance of the constraint to change
+	*	@param OutPosTarget			Target orientation
+	*/
+	UFUNCTION(BlueprintCallable, Category = "Physics|Constraints")
+	static void GetAngularOrientationTarget(
+		UPARAM(ref) FConstraintInstanceAccessor& Accessor,
+		FRotator& OutPosTarget
+	);
+
+	/** Sets the target velocity for the angular drive.
+	*	@param ConstraintInstance	Instance of the constraint to change
+	*	@param InVelTarget			Target velocity
+	*/
+	UFUNCTION(BlueprintCallable, Category = "Physics|Constraints")
+	static void SetAngularVelocityTarget(
+		UPARAM(ref) FConstraintInstanceAccessor& Accessor,
+		const FVector& InVelTarget
+	);
+
+	/** Gets the target velocity for the angular drive.
+	*	@param ConstraintInstance	Instance of the constraint to change
+	*	@param OutVelTarget			Target velocity
+	*/
+	UFUNCTION(BlueprintCallable, Category = "Physics|Constraints")
+	static void GetAngularVelocityTarget(
+		UPARAM(ref) FConstraintInstanceAccessor& Accessor,
+		FVector& OutVelTarget
+	);
+
+	/** Sets the drive params for the angular drive.
+	*	@param ConstraintInstance	Instance of the constraint to change
+	*	@param PositionStrength		Positional strength for the drive (stiffness)
+	*	@param VelocityStrength 	Velocity strength of the drive (damping)
+	*	@param InForceLimit			Max force applied by the drive
+	*/
+	UFUNCTION(BlueprintCallable, Category = "Physics|Constraints")
+	static void SetAngularDriveParams(
+		UPARAM(ref) FConstraintInstanceAccessor& Accessor,
+		float PositionStrength, 
+		float VelocityStrength, 
+		float InForceLimit
+	);
+
+	/** Gets the drive params for the angular drive.
+*	@param ConstraintInstance	Instance of the constraint to change
+*	@param OutPositionStrength	Positional strength for the drive (stiffness)
+*	@param OutVelocityStrength 	Velocity strength of the drive (damping)
+*	@param OutForceLimit		Max force applied by the drive
+*/
+	UFUNCTION(BlueprintCallable, Category = "Physics|Constraints")
+	static void GetAngularDriveParams(
+		UPARAM(ref) FConstraintInstanceAccessor& Accessor,
+		float OutPositionStrength,
+		float OutVelocityStrength,
+		float OutForceLimit
+	);
+};
\ No newline at end of file
diff --git a/Engine/Source/Runtime/Engine/Classes/PhysicsEngine/PhysicsConstraintComponent.h b/Engine/Source/Runtime/Engine/Classes/PhysicsEngine/PhysicsConstraintComponent.h
--- a/Engine/Source/Runtime/Engine/Classes/PhysicsEngine/PhysicsConstraintComponent.h
+++ b/Engine/Source/Runtime/Engine/Classes/PhysicsEngine/PhysicsConstraintComponent.h
@@ -333,6 +333,10 @@ public:
 	UFUNCTION(BlueprintCallable, Category = "Physics|Components|PhysicsConstraint")
 	bool IsBroken();
 
+	/** Gets the constraint object */
+	UFUNCTION(BlueprintCallable, Category = "Physics|Components|PhysicsConstraint")
+	FConstraintInstanceAccessor GetConstraint();
+
 #if WITH_EDITOR
 	void UpdateSpriteTexture();
 #endif
diff --git a/Engine/Source/Runtime/Engine/Private/PhysicsEngine/ConstraintDrives.cpp b/Engine/Source/Runtime/Engine/Private/PhysicsEngine/ConstraintDrives.cpp
--- a/Engine/Source/Runtime/Engine/Private/PhysicsEngine/ConstraintDrives.cpp
+++ b/Engine/Source/Runtime/Engine/Private/PhysicsEngine/ConstraintDrives.cpp
@@ -79,6 +79,15 @@ void FLinearDriveConstraint::SetDriveParams(float InStiffness, float InDamping,
 	ZDrive.SetDriveParams(InStiffness, InDamping, InForceLimit);
 }
 
+void FLinearDriveConstraint::GetDriveParams(float& OutStiffness, float& OutDamping, float& OutForceLimit) const
+{
+	// we set the same value on all drives, just return XDrive ones
+	OutStiffness = XDrive.Stiffness;
+	OutDamping = XDrive.Damping;
+	OutForceLimit = XDrive.MaxForce;
+}
+
+
 void FAngularDriveConstraint::SetAngularDriveMode(EAngularDriveMode::Type DriveMode)
 {
 	AngularDriveMode = DriveMode;
@@ -90,3 +99,11 @@ void FAngularDriveConstraint::SetDriveParams(float InStiffness, float InDamping,
 	TwistDrive.SetDriveParams(InStiffness, InDamping, InForceLimit);
 	SlerpDrive.SetDriveParams(InStiffness, InDamping, InForceLimit);
 }
+
+void FAngularDriveConstraint::GetDriveParams(float& OutStiffness, float& OutDamping, float& OutForceLimit) const
+{
+	// we set the same value on all drives, just return SwingDrive ones
+	OutStiffness = SwingDrive.Stiffness;
+	OutDamping = SwingDrive.Damping;
+	OutForceLimit = SwingDrive.MaxForce;
+}
diff --git a/Engine/Source/Runtime/Engine/Private/PhysicsEngine/ConstraintInstance.cpp b/Engine/Source/Runtime/Engine/Private/PhysicsEngine/ConstraintInstance.cpp
--- a/Engine/Source/Runtime/Engine/Private/PhysicsEngine/ConstraintInstance.cpp
+++ b/Engine/Source/Runtime/Engine/Private/PhysicsEngine/ConstraintInstance.cpp
@@ -6,6 +6,8 @@
 #include "UObject/AnimPhysObjectVersion.h"
 #include "HAL/IConsoleManager.h"
 #include "Components/PrimitiveComponent.h"
+#include "Components/SkeletalMeshComponent.h"
+#include "PhysicsEngine/PhysicsConstraintComponent.h"
 #include "PhysicsPublic.h"
 #include "Physics/PhysicsInterfaceTypes.h"
 
@@ -726,6 +728,12 @@ void FConstraintInstance::SetOrientationDriveTwistAndSwing(bool InEnableTwistDri
 	});
 }
 
+void FConstraintInstance::GetOrientationDriveTwistAndSwing(bool& bOutEnableTwistDrive, bool& bOutEnableSwingDrive)
+{
+	bOutEnableTwistDrive = ProfileInstance.AngularDrive.TwistDrive.bEnablePositionDrive;
+	bOutEnableSwingDrive = ProfileInstance.AngularDrive.SwingDrive.bEnablePositionDrive;
+}
+
 void FConstraintInstance::SetOrientationDriveSLERP(bool InEnableSLERP)
 {
 	ProfileInstance.AngularDrive.SetOrientationDriveSLERP(InEnableSLERP);
@@ -747,6 +755,12 @@ void FConstraintInstance::SetAngularVelocityDriveTwistAndSwing(bool bInEnableTwi
 	});
 }
 
+void FConstraintInstance::GetAngularVelocityDriveTwistAndSwing(bool& bOutEnableTwistDrive, bool& bOutEnableSwingDrive)
+{
+	bOutEnableTwistDrive = ProfileInstance.AngularDrive.TwistDrive.bEnableVelocityDrive;
+	bOutEnableSwingDrive = ProfileInstance.AngularDrive.SwingDrive.bEnableVelocityDrive;	
+}
+
 /** Set whether the SLERP angular velocity drive is enabled. Only applicable when SLERP drive mode is used */
 void FConstraintInstance::SetAngularVelocityDriveSLERP(bool bInEnableSLERP)
 {
@@ -806,6 +820,13 @@ void FConstraintInstance::SetLinearDriveParams(float InSpring, float InDamping,
 	});
 }
 
+/** Get the linear drive's strength parameters */
+void FConstraintInstance::GetLinearDriveParams(float& OutPositionStrength, float& OutVelocityStrength, float& OutForceLimit)
+{
+	ProfileInstance.LinearDrive.GetDriveParams(OutPositionStrength, OutVelocityStrength, OutForceLimit);
+}
+
+
 /** Function for setting target angular position. */
 void FConstraintInstance::SetAngularOrientationTarget(const FQuat& InOrientationTarget)
 {
@@ -861,6 +882,11 @@ void FConstraintInstance::SetAngularDriveParams(float InSpring, float InDamping,
 	});
 }
 
+void FConstraintInstance::GetAngularDriveParams(float& OutSpring, float& OutDamping, float& OutForceLimit) const
+{
+	ProfileInstance.AngularDrive.GetDriveParams(OutSpring, OutDamping, OutForceLimit);
+}
+
 /** Scale Angular Limit Constraints (as defined in RB_ConstraintSetup) */
 void FConstraintInstance::SetAngularDOFLimitScale(float InSwing1LimitScale, float InSwing2LimitScale, float InTwistLimitScale)
 {
@@ -1082,24 +1108,59 @@ FConstraintInstance * FConstraintInstance::Alloc()
 	return new (Memory)FConstraintInstance();
 }
 
-void FConstraintInstance::EnableProjection()
+void FConstraintInstance::GetProjectionAlphasOrTolerances(float& ProjectionLinearAlphaOrTolerance, float& ProjectionAngularAlphaOrTolerance) const
 {
-	ProfileInstance.bEnableProjection = true;
-	
+#if WITH_CHAOS
+	ProjectionLinearAlphaOrTolerance = ProfileInstance.ProjectionLinearAlpha;
+	ProjectionAngularAlphaOrTolerance = ProfileInstance.ProjectionAngularAlpha;
+#else
+	ProjectionLinearAlphaOrTolerance = ProfileInstance.ProjectionLinearTolerance;
+	ProjectionAngularAlphaOrTolerance = ProfileInstance.ProjectionAngularTolerance;
+
+#endif
+}
+
+void FConstraintInstance::SetProjectionParams(bool bEnableProjection, float ProjectionLinearAlphaOrTolerance, float ProjectionAngularAlphaOrTolerance)
+{
+	ProfileInstance.bEnableProjection = bEnableProjection;
+#if WITH_CHAOS
+	ProfileInstance.ProjectionLinearAlpha = ProjectionLinearAlphaOrTolerance;
+	ProfileInstance.ProjectionAngularAlpha = ProjectionAngularAlphaOrTolerance;
+#else
+	ProfileInstance.ProjectionLinearTolerance = ProjectionLinearAlphaOrTolerance;
+	ProfileInstance.ProjectionAngularTolerance = ProjectionAngularAlphaOrTolerance;
+
+#endif
+
 	FPhysicsCommand::ExecuteWrite(ConstraintHandle, [&](const FPhysicsConstraintHandle& Constraint)
-	{
-		FPhysicsInterface::SetProjectionEnabled_AssumesLocked(Constraint, true, ProfileInstance.ProjectionLinearTolerance, ProfileInstance.ProjectionAngularTolerance);
-	});
+		{
+			if (bEnableProjection)
+			{
+				FPhysicsInterface::SetProjectionEnabled_AssumesLocked(Constraint, true, ProjectionLinearAlphaOrTolerance, ProjectionAngularAlphaOrTolerance);
+			}
+			else
+			{
+				FPhysicsInterface::SetProjectionEnabled_AssumesLocked(Constraint, false);
+			}
+		});
+}
+
+void FConstraintInstance::EnableProjection()
+{
+#if WITH_CHAOS
+	SetProjectionParams(true, ProfileInstance.ProjectionLinearAlpha, ProfileInstance.ProjectionAngularAlpha);
+#else
+	SetProjectionParams(true, ProfileInstance.ProjectionLinearTolerance, ProfileInstance.ProjectionAngularTolerance);
+#endif
 }
 
 void FConstraintInstance::DisableProjection()
 {
-	ProfileInstance.bEnableProjection = false;
-	
-	FPhysicsCommand::ExecuteWrite(ConstraintHandle, [&](const FPhysicsConstraintHandle& Constraint)
-	{
-		FPhysicsInterface::SetProjectionEnabled_AssumesLocked(Constraint, false);
-	});
+#if WITH_CHAOS
+	SetProjectionParams(false, ProfileInstance.ProjectionLinearAlpha, ProfileInstance.ProjectionAngularAlpha);
+#else
+	SetProjectionParams(false, ProfileInstance.ProjectionLinearTolerance, ProfileInstance.ProjectionAngularTolerance);
+#endif
 }
 
 void FConstraintInstance::EnableParentDominates()
@@ -1122,5 +1183,20 @@ void FConstraintInstance::DisableParentDominates()
 	});
 }
 
+FConstraintInstance* FConstraintInstanceAccessor::Get() const
+{
+	if (Owner.IsValid())
+	{
+		if (USkeletalMeshComponent* SkeletalMeshComponent = Cast<USkeletalMeshComponent>(Owner.Get()))
+		{
+			return SkeletalMeshComponent->GetConstraintInstanceByIndex(Index);
+		}
+		if (UPhysicsConstraintComponent* PhysicsConstraintComponent = Cast<UPhysicsConstraintComponent>(Owner.Get()))
+		{
+			return &(PhysicsConstraintComponent->ConstraintInstance);
+		}
+	}
+	return nullptr;
+}
 
 #undef LOCTEXT_NAMESPACE
diff --git a/Engine/Source/Runtime/Engine/Private/PhysicsEngine/ConstraintInstanceBlueprintLibrary.cpp b/Engine/Source/Runtime/Engine/Private/PhysicsEngine/ConstraintInstanceBlueprintLibrary.cpp
new file mode 100644
--- /dev/null
+++ b/Engine/Source/Runtime/Engine/Private/PhysicsEngine/ConstraintInstanceBlueprintLibrary.cpp
@@ -0,0 +1,666 @@
+// Copyright Epic Games, Inc. All Rights Reserved.
+
+#include "PhysicsEngine/ConstraintInstanceBlueprintLibrary.h"
+
+//---------------------------------------------------------------------------------------------------
+//
+// CONSTRAINT BEHAVIOR 
+//
+//---------------------------------------------------------------------------------------------------
+
+void UConstraintInstanceBlueprintLibrary::SetDisableCollision(
+	FConstraintInstanceAccessor& Accessor,
+	bool bDisableCollision
+	)
+{
+	if (FConstraintInstance* ConstraintInstance = Accessor.Get())
+	{
+		ConstraintInstance->SetDisableCollision(bDisableCollision);
+	}
+}
+
+bool UConstraintInstanceBlueprintLibrary::GetDisableCollsion(FConstraintInstanceAccessor& Accessor)
+{
+	if (FConstraintInstance* ConstraintInstance = Accessor.Get())
+	{
+		return ConstraintInstance->IsCollisionDisabled();
+	}
+	return true;
+}
+
+void UConstraintInstanceBlueprintLibrary::SetProjectionParams(
+	FConstraintInstanceAccessor& Accessor,
+	bool bEnableProjection,
+	float ProjectionLinearAlpha,
+	float ProjectionAngularAlpha
+	)
+{
+	if (FConstraintInstance* ConstraintInstance = Accessor.Get())
+	{
+		ConstraintInstance->SetProjectionParams(bEnableProjection, ProjectionLinearAlpha, ProjectionAngularAlpha);
+	}
+}
+
+void UConstraintInstanceBlueprintLibrary::GetProjectionParams(
+	FConstraintInstanceAccessor& Accessor,
+	bool& bEnableProjection,
+	float& ProjectionLinearAlpha,
+	float& ProjectionAngularAlpha
+	)
+{
+	if (FConstraintInstance* ConstraintInstance = Accessor.Get())
+	{
+		bEnableProjection = ConstraintInstance->IsProjectionEnabled();
+		ConstraintInstance->GetProjectionAlphasOrTolerances(ProjectionLinearAlpha, ProjectionAngularAlpha);
+	}
+	else
+	{
+		bEnableProjection = false;
+		ProjectionLinearAlpha = 0.0f;
+		ProjectionAngularAlpha = 0.0f;
+	}
+}
+
+void UConstraintInstanceBlueprintLibrary::SetParentDominates(
+	FConstraintInstanceAccessor& Accessor,
+	bool bParentDominates
+	)
+{
+	if (FConstraintInstance* ConstraintInstance = Accessor.Get())
+	{
+		if (bParentDominates)
+		{
+			ConstraintInstance->EnableParentDominates();
+		}
+		else
+		{
+			ConstraintInstance->DisableParentDominates();
+		}
+	}
+}
+
+bool UConstraintInstanceBlueprintLibrary::GetParentDominates(FConstraintInstanceAccessor& Accessor)
+{
+	if (FConstraintInstance* ConstraintInstance = Accessor.Get())
+	{
+		return ConstraintInstance->IsParentDominatesEnabled();
+	}
+	return false;
+}
+
+//---------------------------------------------------------------------------------------------------
+//
+// LINEAR LIMITS
+//
+//---------------------------------------------------------------------------------------------------
+
+void UConstraintInstanceBlueprintLibrary::SetLinearLimits(
+	FConstraintInstanceAccessor& Accessor,
+	TEnumAsByte<ELinearConstraintMotion> XMotion,
+	TEnumAsByte<ELinearConstraintMotion> YMotion,
+	TEnumAsByte<ELinearConstraintMotion> ZMotion,
+	float Limit
+	)
+{
+	if (FConstraintInstance* ConstraintInstance = Accessor.Get())
+	{
+		ConstraintInstance->SetLinearLimits(XMotion, YMotion, ZMotion, Limit);
+	}
+}
+
+void UConstraintInstanceBlueprintLibrary::GetLinearLimits(
+	FConstraintInstanceAccessor& Accessor,
+	TEnumAsByte<ELinearConstraintMotion>& XMotion,
+	TEnumAsByte<ELinearConstraintMotion>& YMotion,
+	TEnumAsByte<ELinearConstraintMotion>& ZMotion,
+	float& Limit
+)
+{
+	if (FConstraintInstance* ConstraintInstance = Accessor.Get())
+	{
+		XMotion = ConstraintInstance->GetLinearXMotion();
+		YMotion = ConstraintInstance->GetLinearYMotion();
+		ZMotion = ConstraintInstance->GetLinearZMotion();
+		Limit = ConstraintInstance->GetLinearLimit();
+	} 
+	else
+	{
+		XMotion = ELinearConstraintMotion::LCM_Free;
+		YMotion = ELinearConstraintMotion::LCM_Free;
+		ZMotion = ELinearConstraintMotion::LCM_Free;
+		Limit = 0.f;
+	}
+}
+
+void UConstraintInstanceBlueprintLibrary::SetLinearBreakable(
+	FConstraintInstanceAccessor& Accessor,
+	bool bLinearBreakable,
+	float LinearBreakThreshold
+	)
+{
+	if (FConstraintInstance* ConstraintInstance = Accessor.Get())
+	{
+		ConstraintInstance->SetLinearBreakable(bLinearBreakable, LinearBreakThreshold);
+	}
+}
+
+void UConstraintInstanceBlueprintLibrary::GetLinearBreakable(
+	FConstraintInstanceAccessor& Accessor,
+	bool& bLinearBreakable,
+	float& LinearBreakThreshold
+	)
+{
+	if (FConstraintInstance* ConstraintInstance = Accessor.Get())
+	{
+		bLinearBreakable = ConstraintInstance->IsLinearBreakable();
+		LinearBreakThreshold = ConstraintInstance->GetLinearBreakThreshold();
+	}
+	else
+	{
+		bLinearBreakable = false;
+		LinearBreakThreshold = 0.0f;
+	}
+}
+
+//---------------------------------------------------------------------------------------------------
+//
+// ANGULAR LIMITS 
+//
+//---------------------------------------------------------------------------------------------------
+
+void UConstraintInstanceBlueprintLibrary::SetAngularLimits(
+	FConstraintInstanceAccessor& Accessor,
+	TEnumAsByte<EAngularConstraintMotion> Swing1MotionType,
+	float Swing1LimitAngle,
+	TEnumAsByte<EAngularConstraintMotion> Swing2MotionType,
+	float Swing2LimitAngle,
+	TEnumAsByte<EAngularConstraintMotion> TwistMotionType,
+	float TwistLimitAngle
+	)
+{
+	if (FConstraintInstance* ConstraintInstance = Accessor.Get())
+	{
+		ConstraintInstance->SetAngularSwing1Limit(Swing1MotionType, Swing1LimitAngle);
+		ConstraintInstance->SetAngularSwing2Limit(Swing2MotionType, Swing2LimitAngle);
+		ConstraintInstance->SetAngularTwistLimit(TwistMotionType, TwistLimitAngle);
+	}
+}
+
+void UConstraintInstanceBlueprintLibrary::GetAngularLimits(
+	FConstraintInstanceAccessor& Accessor,
+	TEnumAsByte<EAngularConstraintMotion>& Swing1MotionType,
+	float& Swing1LimitAngle,
+	TEnumAsByte<EAngularConstraintMotion>& Swing2MotionType,
+	float& Swing2LimitAngle,
+	TEnumAsByte<EAngularConstraintMotion>& TwistMotionType,
+	float& TwistLimitAngle
+	)
+{
+	if (FConstraintInstance* ConstraintInstance = Accessor.Get())
+	{
+		Swing1MotionType = ConstraintInstance->GetAngularSwing1Motion();
+		Swing1LimitAngle = ConstraintInstance->GetAngularSwing1Limit();
+		Swing2MotionType = ConstraintInstance->GetAngularSwing2Motion();
+		Swing2LimitAngle = ConstraintInstance->GetAngularSwing2Limit();
+		TwistMotionType = ConstraintInstance->GetAngularTwistMotion();
+		TwistLimitAngle = ConstraintInstance->GetAngularTwistLimit();
+	} 
+	else
+	{
+		Swing1MotionType = EAngularConstraintMotion::ACM_Free;
+		Swing1LimitAngle = 0.0f;
+		Swing2MotionType = EAngularConstraintMotion::ACM_Free;
+		Swing2LimitAngle = 0.0f;
+		TwistMotionType = EAngularConstraintMotion::ACM_Free;
+		TwistLimitAngle = 0.0f;
+	}
+}
+
+void UConstraintInstanceBlueprintLibrary::SetAngularBreakable(
+	FConstraintInstanceAccessor& Accessor,
+	bool bAngularBreakable,
+	float AngularBreakThreshold
+	)
+{
+	if (FConstraintInstance* ConstraintInstance = Accessor.Get())
+	{
+		ConstraintInstance->SetAngularBreakable(bAngularBreakable, AngularBreakThreshold);
+	}
+}
+
+void UConstraintInstanceBlueprintLibrary::GetAngularBreakable(
+	FConstraintInstanceAccessor& Accessor,
+	bool& bAngularBreakable,
+	float& AngularBreakThreshold
+	)
+{
+	if (FConstraintInstance* ConstraintInstance = Accessor.Get())
+	{
+		bAngularBreakable = ConstraintInstance->IsAngularBreakable();
+		AngularBreakThreshold = ConstraintInstance->GetAngularBreakThreshold();
+	}
+	else
+	{
+		bAngularBreakable = false;
+		AngularBreakThreshold = 0.0f;
+	}
+
+}
+
+void UConstraintInstanceBlueprintLibrary::SetAngularPlasticity(
+	FConstraintInstanceAccessor& Accessor,
+	bool bAngularPlasticity,
+	float AngularPlasticityThreshold
+	)
+{
+	if (FConstraintInstance* ConstraintInstance = Accessor.Get())
+	{
+		ConstraintInstance->SetAngularPlasticity(bAngularPlasticity, AngularPlasticityThreshold);
+	}
+}
+
+void UConstraintInstanceBlueprintLibrary::GetAngularPlasticity(
+	FConstraintInstanceAccessor& Accessor,
+	bool& bAngularPlasticity,
+	float& AngularPlasticityThreshold
+	)
+{
+	if (FConstraintInstance* ConstraintInstance = Accessor.Get())
+	{
+		bAngularPlasticity = ConstraintInstance->HasAngularPlasticity();
+		AngularPlasticityThreshold = ConstraintInstance->GetAngularPlasticityThreshold();
+	}
+	else
+	{
+		bAngularPlasticity = false;
+		AngularPlasticityThreshold = 0.0f;
+	}
+}
+
+
+//---------------------------------------------------------------------------------------------------
+//
+// LINEAR MOTOR
+//
+//---------------------------------------------------------------------------------------------------
+
+void UConstraintInstanceBlueprintLibrary::SetLinearPositionDrive(
+	FConstraintInstanceAccessor& Accessor, 
+	bool bEnableDriveX, 
+	bool bEnableDriveY, 
+	bool bEnableDriveZ
+	)
+{
+	if (FConstraintInstance* ConstraintInstance = Accessor.Get())
+	{
+		ConstraintInstance->SetLinearPositionDrive(bEnableDriveX, bEnableDriveY, bEnableDriveZ);
+	}
+}
+
+void UConstraintInstanceBlueprintLibrary::GetLinearPositionDrive(
+	FConstraintInstanceAccessor& Accessor,
+	bool& bEnableDriveX,
+	bool& bEnableDriveY,
+	bool& bEnableDriveZ
+	)
+{
+	if (FConstraintInstance* ConstraintInstance = Accessor.Get())
+	{
+		bEnableDriveX = ConstraintInstance->IsLinearPositionDriveXEnabled();
+		bEnableDriveY = ConstraintInstance->IsLinearPositionDriveYEnabled();
+		bEnableDriveZ = ConstraintInstance->IsLinearPositionDriveZEnabled();
+	}
+	else
+	{
+		bEnableDriveX = false;
+		bEnableDriveY = false;
+		bEnableDriveZ = false;
+	}
+}
+
+void UConstraintInstanceBlueprintLibrary::SetLinearVelocityDrive(
+	FConstraintInstanceAccessor& Accessor,
+	bool bEnableDriveX, 
+	bool bEnableDriveY, 
+	bool bEnableDriveZ
+	)
+{
+	if (FConstraintInstance* ConstraintInstance = Accessor.Get())
+	{
+		ConstraintInstance->SetLinearVelocityDrive(bEnableDriveX, bEnableDriveY, bEnableDriveZ);
+	}
+}
+
+void UConstraintInstanceBlueprintLibrary::GetLinearVelocityDrive(
+	FConstraintInstanceAccessor& Accessor,
+	bool& bEnableDriveX,
+	bool& bEnableDriveY,
+	bool& bEnableDriveZ
+)
+{
+	if (FConstraintInstance* ConstraintInstance = Accessor.Get())
+	{
+		bEnableDriveX = ConstraintInstance->IsLinearVelocityDriveXEnabled();
+		bEnableDriveY = ConstraintInstance->IsLinearVelocityDriveYEnabled();
+		bEnableDriveZ = ConstraintInstance->IsLinearVelocityDriveZEnabled();
+	}
+	else
+	{
+		bEnableDriveX = false;
+		bEnableDriveY = false;
+		bEnableDriveZ = false;
+	}
+}
+
+void UConstraintInstanceBlueprintLibrary::SetLinearPositionTarget(
+	FConstraintInstanceAccessor& Accessor, 
+	const FVector& InPosTarget
+	)
+{
+	if (FConstraintInstance* ConstraintInstance = Accessor.Get())
+	{
+		ConstraintInstance->SetLinearPositionTarget(InPosTarget);
+	}
+}
+
+void UConstraintInstanceBlueprintLibrary::GetLinearPositionTarget(
+	FConstraintInstanceAccessor& Accessor,
+	FVector& OutPosTarget
+	)
+{
+	if (FConstraintInstance* ConstraintInstance = Accessor.Get())
+	{
+		OutPosTarget = ConstraintInstance->GetLinearPositionTarget();
+	}
+	else
+	{
+		OutPosTarget = FVector::ZeroVector;
+	}
+}
+
+
+void UConstraintInstanceBlueprintLibrary::SetLinearVelocityTarget(
+	FConstraintInstanceAccessor& Accessor, 
+	const FVector& InVelTarget
+	)
+{
+	if (FConstraintInstance* ConstraintInstance = Accessor.Get())
+	{
+		ConstraintInstance->SetLinearVelocityTarget(InVelTarget);
+	}
+}
+
+void UConstraintInstanceBlueprintLibrary::GetLinearVelocityTarget(
+	FConstraintInstanceAccessor& Accessor,
+	FVector& OutVelTarget
+	)
+{
+	if (FConstraintInstance* ConstraintInstance = Accessor.Get())
+	{
+		OutVelTarget = ConstraintInstance->GetLinearVelocityTarget();
+	}
+	else
+	{
+		OutVelTarget = FVector::ZeroVector;
+	}
+}
+
+void UConstraintInstanceBlueprintLibrary::SetLinearDriveParams(
+	FConstraintInstanceAccessor& Accessor,
+	float PositionStrength, 
+	float VelocityStrength, 
+	float InForceLimit
+	)
+{
+	if (FConstraintInstance* ConstraintInstance = Accessor.Get())
+	{
+		ConstraintInstance->SetLinearDriveParams(PositionStrength, VelocityStrength, InForceLimit);
+	}
+}
+
+void UConstraintInstanceBlueprintLibrary::GetLinearDriveParams(
+	FConstraintInstanceAccessor& Accessor,
+	float& OutPositionStrength,
+	float& OutVelocityStrength,
+	float& OutForceLimit
+	)
+{
+	if (FConstraintInstance* ConstraintInstance = Accessor.Get())
+	{
+		ConstraintInstance->GetLinearDriveParams(OutPositionStrength, OutVelocityStrength, OutForceLimit);
+	}
+	else
+	{
+		OutPositionStrength = 0.0f;
+		OutVelocityStrength = 0.0f;
+		OutForceLimit = 0.0f;
+	}
+}
+
+
+//---------------------------------------------------------------------------------------------------
+//
+// ANGULAR MOTOR
+//
+//---------------------------------------------------------------------------------------------------
+
+void UConstraintInstanceBlueprintLibrary::SetOrientationDriveTwistAndSwing(
+	FConstraintInstanceAccessor& Accessor, 
+	bool bEnableTwistDrive, 
+	bool bEnableSwingDrive
+	)
+{
+	if (FConstraintInstance* ConstraintInstance = Accessor.Get())
+	{
+		ConstraintInstance->SetOrientationDriveTwistAndSwing(bEnableTwistDrive, bEnableSwingDrive);
+	}
+}
+
+void UConstraintInstanceBlueprintLibrary::GetOrientationDriveTwistAndSwing(
+	UPARAM(ref) FConstraintInstanceAccessor& Accessor,
+	bool& bOutEnableTwistDrive,
+	bool& bOutEnableSwingDrive
+	)
+{
+	if (FConstraintInstance* ConstraintInstance = Accessor.Get())
+	{
+		ConstraintInstance->GetOrientationDriveTwistAndSwing(bOutEnableTwistDrive, bOutEnableSwingDrive);
+	}
+	else
+	{
+		bOutEnableTwistDrive = false;
+		bOutEnableSwingDrive = false;
+	}
+}
+
+void UConstraintInstanceBlueprintLibrary::SetOrientationDriveSLERP(
+	FConstraintInstanceAccessor& Accessor, 
+	bool bEnableSLERP
+	)
+{
+	if (FConstraintInstance* ConstraintInstance = Accessor.Get())
+	{
+		ConstraintInstance->SetOrientationDriveSLERP(bEnableSLERP);
+	}
+}
+
+void UConstraintInstanceBlueprintLibrary::GetOrientationDriveSLERP(
+	FConstraintInstanceAccessor& Accessor,
+	bool& bOutEnableSLERP
+	)
+{
+	if (FConstraintInstance* ConstraintInstance = Accessor.Get())
+	{
+		bOutEnableSLERP = ConstraintInstance->GetOrientationDriveSLERP();
+	}
+	else
+	{
+		bOutEnableSLERP = false;
+	}
+}
+
+void UConstraintInstanceBlueprintLibrary::SetAngularVelocityDriveTwistAndSwing(
+	FConstraintInstanceAccessor& Accessor, 
+	bool bEnableTwistDrive, 
+	bool bEnableSwingDrive
+	)
+{
+	if (FConstraintInstance* ConstraintInstance = Accessor.Get())
+	{
+		ConstraintInstance->SetAngularVelocityDriveTwistAndSwing(bEnableTwistDrive, bEnableSwingDrive);
+	}
+}
+
+void UConstraintInstanceBlueprintLibrary::GetAngularVelocityDriveTwistAndSwing(
+	FConstraintInstanceAccessor& Accessor,
+	bool& bOutEnableTwistDrive,
+	bool& bOutEnableSwingDrive
+)
+{
+	if (FConstraintInstance* ConstraintInstance = Accessor.Get())
+	{
+		ConstraintInstance->GetAngularVelocityDriveTwistAndSwing(bOutEnableTwistDrive, bOutEnableSwingDrive);
+	}
+	else
+	{
+		bOutEnableTwistDrive = false;
+		bOutEnableSwingDrive = false;
+	}
+}
+
+void UConstraintInstanceBlueprintLibrary::SetAngularVelocityDriveSLERP(
+	FConstraintInstanceAccessor& Accessor,
+	bool bEnableSLERP
+	)
+{
+	if (FConstraintInstance* ConstraintInstance = Accessor.Get())
+	{
+		ConstraintInstance->SetAngularVelocityDriveSLERP(bEnableSLERP);
+	}
+}
+
+void UConstraintInstanceBlueprintLibrary::GetAngularVelocityDriveSLERP(
+	FConstraintInstanceAccessor& Accessor,
+	bool& bOutEnableSLERP
+	)
+{
+	if (FConstraintInstance* ConstraintInstance = Accessor.Get())
+	{
+		bOutEnableSLERP = ConstraintInstance->GetAngularVelocityDriveSLERP();
+	} 
+	else
+	{
+		bOutEnableSLERP = false;
+	}
+}
+
+void UConstraintInstanceBlueprintLibrary::SetAngularDriveMode(
+	FConstraintInstanceAccessor& Accessor,
+	EAngularDriveMode::Type DriveMode
+	)
+{
+	if (FConstraintInstance* ConstraintInstance = Accessor.Get())
+	{
+		ConstraintInstance->SetAngularDriveMode(DriveMode);
+	}
+}
+
+void UConstraintInstanceBlueprintLibrary::GetAngularDriveMode(
+	FConstraintInstanceAccessor& Accessor,
+	TEnumAsByte<EAngularDriveMode::Type>& OutDriveMode
+)
+{
+	if (FConstraintInstance* ConstraintInstance = Accessor.Get())
+	{
+		OutDriveMode = ConstraintInstance->GetAngularDriveMode();
+	} 
+	else
+	{
+		OutDriveMode = EAngularDriveMode::Type::SLERP;
+	}
+
+}
+
+void UConstraintInstanceBlueprintLibrary::SetAngularOrientationTarget(
+	FConstraintInstanceAccessor& Accessor,
+	const FRotator& InPosTarget
+	)
+{
+	if (FConstraintInstance* ConstraintInstance = Accessor.Get())
+	{
+		ConstraintInstance->SetAngularOrientationTarget(InPosTarget.Quaternion());
+	}
+}
+
+void UConstraintInstanceBlueprintLibrary::GetAngularOrientationTarget(
+	FConstraintInstanceAccessor& Accessor,
+	FRotator& OutPosTarget
+)
+{
+	if (FConstraintInstance* ConstraintInstance = Accessor.Get())
+	{
+		OutPosTarget = ConstraintInstance->GetAngularOrientationTarget();
+	}
+	else
+	{
+		OutPosTarget = FRotator::ZeroRotator;
+	}
+}
+
+void UConstraintInstanceBlueprintLibrary::SetAngularVelocityTarget(
+	FConstraintInstanceAccessor& Accessor,
+	const FVector& InVelTarget
+	)
+{
+	if (FConstraintInstance* ConstraintInstance = Accessor.Get())
+	{
+		ConstraintInstance->SetAngularVelocityTarget(InVelTarget);
+	}
+}
+
+void UConstraintInstanceBlueprintLibrary::GetAngularVelocityTarget(
+	FConstraintInstanceAccessor& Accessor,
+	FVector& OutVelTarget
+)
+{
+	if (FConstraintInstance* ConstraintInstance = Accessor.Get())
+	{
+		OutVelTarget = ConstraintInstance->GetAngularVelocityTarget();
+	}
+	else
+	{
+		OutVelTarget = FVector::ZeroVector;
+	}
+}
+
+void UConstraintInstanceBlueprintLibrary::SetAngularDriveParams(
+	FConstraintInstanceAccessor& Accessor,
+	float PositionStrength,
+	float VelocityStrength,
+	float InForceLimit
+	)
+{
+	if (FConstraintInstance* ConstraintInstance = Accessor.Get())
+	{
+		ConstraintInstance->SetAngularDriveParams(PositionStrength, VelocityStrength, InForceLimit);
+	}
+}
+
+void UConstraintInstanceBlueprintLibrary::GetAngularDriveParams(
+	FConstraintInstanceAccessor& Accessor,
+	float OutPositionStrength,
+	float OutVelocityStrength,
+	float OutForceLimit
+)
+{
+	if (FConstraintInstance* ConstraintInstance = Accessor.Get())
+	{
+		ConstraintInstance->GetAngularDriveParams(OutPositionStrength, OutVelocityStrength, OutForceLimit);
+	}
+	else
+	{
+		OutPositionStrength = 0.0f;
+		OutVelocityStrength = 0.0f;
+		OutForceLimit = 0.0f;
+	}
+}
\ No newline at end of file
diff --git a/Engine/Source/Runtime/Engine/Private/PhysicsEngine/PhysicsConstraintComponent.cpp b/Engine/Source/Runtime/Engine/Private/PhysicsEngine/PhysicsConstraintComponent.cpp
--- a/Engine/Source/Runtime/Engine/Private/PhysicsEngine/PhysicsConstraintComponent.cpp
+++ b/Engine/Source/Runtime/Engine/Private/PhysicsEngine/PhysicsConstraintComponent.cpp
@@ -532,6 +532,10 @@ void UPhysicsConstraintComponent::SetDisableCollision(bool bDisableCollision)
 	ConstraintInstance.SetDisableCollision(bDisableCollision);
 }
 
+FConstraintInstanceAccessor UPhysicsConstraintComponent::GetConstraint()
+{
+	return FConstraintInstanceAccessor(this);
+}
 
 #if WITH_EDITOR
 void UPhysicsConstraintComponent::UpdateSpriteTexture()
diff --git a/Engine/Source/Runtime/Engine/Private/SkeletalMeshComponentPhysics.cpp b/Engine/Source/Runtime/Engine/Private/SkeletalMeshComponentPhysics.cpp
--- a/Engine/Source/Runtime/Engine/Private/SkeletalMeshComponentPhysics.cpp
+++ b/Engine/Source/Runtime/Engine/Private/SkeletalMeshComponentPhysics.cpp
@@ -1399,6 +1399,15 @@ FConstraintInstance* USkeletalMeshComponent::FindConstraintInstance(FName ConNam
 	return NULL;
 }
 
+FConstraintInstance* USkeletalMeshComponent::GetConstraintInstanceByIndex(uint32 Index)
+{
+	if (Index < (uint32)Constraints.Num())
+	{
+		return Constraints[Index];
+	}
+	return nullptr;
+}
+
 void USkeletalMeshComponent::AddForceToAllBodiesBelow(FVector Force, FName BoneName, bool bAccelChange, bool bIncludeSelf)
 {
 	ForEachBodyBelow(BoneName, bIncludeSelf, /*bSkipCustomPhysics=*/false, [Force, bAccelChange](FBodyInstance* BI)
@@ -1824,6 +1833,41 @@ void USkeletalMeshComponent::SetNotifyRigidBodyCollisionBelow(bool bNewNotifyRig
 	}
 }
 
+FConstraintInstanceAccessor USkeletalMeshComponent::GetConstraintByName(FName ConstraintName, bool bIncludesTerminated)
+{
+	int32 ConstraintIndex = FindConstraintIndex(ConstraintName);
+	if (ConstraintIndex == INDEX_NONE || ConstraintIndex >= Constraints.Num())
+	{
+		return FConstraintInstanceAccessor();
+	}
+
+	if (FConstraintInstance* Constraint = Constraints[ConstraintIndex])
+	{
+		if (bIncludesTerminated || !Constraint->IsTerminated())
+		{
+			return FConstraintInstanceAccessor(this, ConstraintIndex);
+		}
+	}
+	return FConstraintInstanceAccessor();
+}
+
+void USkeletalMeshComponent::GetConstraints(bool bIncludesTerminated, TArray<FConstraintInstanceAccessor>& OutConstraints)
+{
+	if (UPhysicsAsset* const PhysicsAsset = GetPhysicsAsset())
+	{
+		for (int32 i = 0; i < Constraints.Num(); i++)
+		{
+			if (FConstraintInstance* ConstraintInstance = Constraints[i])
+			{
+				if (bIncludesTerminated || !ConstraintInstance->IsTerminated())
+				{
+					OutConstraints.Add(FConstraintInstanceAccessor(this, i));
+				}
+			}
+		}
+	}
+}
+
 void USkeletalMeshComponent::BreakConstraint(FVector Impulse, FVector HitLocation, FName InBoneName)
 {
 	// you can enable/disable the instanced weights by calling
