Restore old behavior of BeginReleaseOverrideVertexColors to fix mesh painting
Rework the PrivateFixupOverrideColors to still avoid doing FlushRenderingCommands during PostLoad

diff --git a/Engine/Source/Runtime/Engine/Private/Components/StaticMeshComponent.cpp b/Engine/Source/Runtime/Engine/Private/Components/StaticMeshComponent.cpp
--- a/Engine/Source/Runtime/Engine/Private/Components/StaticMeshComponent.cpp
+++ b/Engine/Source/Runtime/Engine/Private/Components/StaticMeshComponent.cpp
@@ -1142,17 +1142,9 @@ void UStaticMeshComponent::CopyInstanceVertexColorsIfCompatible( UStaticMeshComp
 
 					if (TargetLODInfo.OverrideVertexColors != nullptr || CopiedColors.Num() > 0)
 					{
-						if (TargetLODInfo.OverrideVertexColors != nullptr)
-						{
-							TargetLODInfo.BeginReleaseOverrideVertexColors();
-						}
-
-						if (TargetLODInfo.OverrideVertexColors == nullptr)
-						{
-							TargetLODInfo.OverrideVertexColors = new FColorVertexBuffer;
-						}
-
-						TargetLODInfo.OverrideVertexColors->InitFromColorArray(CopiedColors);
+						TargetLODInfo.CleanUp();
+						TargetLODInfo.OverrideVertexColors = new FColorVertexBuffer;
+						TargetLODInfo.OverrideVertexColors->InitFromColorArray( CopiedColors );
 
 						check(TargetLODInfo.OverrideVertexColors->GetStride() > 0);
 						BeginInitResource( TargetLODInfo.OverrideVertexColors );
@@ -1307,17 +1299,6 @@ void UStaticMeshComponent::PrivateFixupOverrideColors()
 	for (uint32 LODIndex = 0; LODIndex < NumLODs; ++LODIndex)
 	{
 		FStaticMeshComponentLODInfo& LODInfo = LODData[LODIndex];
-		if (LODInfo.OverrideVertexColors == nullptr)
-		{
-			if (bCustomOverrideVertexColorPerLOD) //No fixup required if the component is in custom LOD paint and there is no paint on a LOD
-				continue;
-			LODInfo.OverrideVertexColors = new FColorVertexBuffer;
-		}
-		else
-		{
-			LODInfo.BeginReleaseOverrideVertexColors();
-		}
-
 		FStaticMeshLODResources& CurRenderData = GetStaticMesh()->GetRenderData()->LODResources[LODIndex];
 		TArray<FColor> NewOverrideColors;
 		if (bCustomOverrideVertexColorPerLOD)
@@ -1327,7 +1308,7 @@ void UStaticMeshComponent::PrivateFixupOverrideColors()
 			//Use the existing LOD custom paint and remap it on the new mesh
 			RemapPaintedVertexColors(
 				LODInfo.PaintedVertices,
-				LODInfo.OverrideVertexColors,
+				nullptr,
 				SourceRenderData.VertexBuffers.PositionVertexBuffer,
 				SourceRenderData.VertexBuffers.StaticMeshVertexBuffer,
 				CurRenderData.VertexBuffers.PositionVertexBuffer,
@@ -1339,7 +1320,7 @@ void UStaticMeshComponent::PrivateFixupOverrideColors()
 		{
 			RemapPaintedVertexColors(
 				LOD0Info.PaintedVertices,
-				LOD0Info.OverrideVertexColors,
+				nullptr,
 				SourceRenderData.VertexBuffers.PositionVertexBuffer,
 				SourceRenderData.VertexBuffers.StaticMeshVertexBuffer,
 				CurRenderData.VertexBuffers.PositionVertexBuffer,
@@ -1347,8 +1328,11 @@ void UStaticMeshComponent::PrivateFixupOverrideColors()
 				NewOverrideColors
 				);
 		}
+
+		LODInfo.CleanUp();
 		if (NewOverrideColors.Num())
 		{
+			LODInfo.OverrideVertexColors = new FColorVertexBuffer;
 			LODInfo.OverrideVertexColors->InitFromColorArray(NewOverrideColors);
 
 			// Update the PaintedVertices array
@@ -1367,12 +1351,6 @@ void UStaticMeshComponent::PrivateFixupOverrideColors()
 			BeginInitResource(LODInfo.OverrideVertexColors);
 			UpdateStaticMeshDeriveDataKey = true;
 		}
-		else
-		{
-			delete LODInfo.OverrideVertexColors;
-			LODInfo.OverrideVertexColors = nullptr;
-		}
-		
 	}
 
 	if (UpdateStaticMeshDeriveDataKey)
@@ -2761,22 +2739,8 @@ void FStaticMeshComponentLODInfo::BeginReleaseOverrideVertexColors()
 {
 	if (OverrideVertexColors)
 	{
-		DEC_DWORD_STAT_BY(STAT_InstVertexColorMemory, OverrideVertexColors->GetAllocatedSize());
-
-		// The old OverrideVertexColors will be deleted on the render thread so we can start
-		// using a new instance right away without the need of a costly FlushRenderingCommand.
-		FColorVertexBuffer* LocalOverrideVertexColors = OverrideVertexColors;
-		OverrideVertexColors = new FColorVertexBuffer;
-
-		if (LocalOverrideVertexColors != nullptr)
-		{
-			ENQUEUE_RENDER_COMMAND(FStaticMeshComponentLODInfoCleanUp)(
-				[LocalOverrideVertexColors](FRHICommandList&)
-				{
-					LocalOverrideVertexColors->ReleaseResource();
-					delete LocalOverrideVertexColors;
-				});
-		}
+		// enqueue a rendering command to release
+		BeginReleaseResource(OverrideVertexColors);
 	}
 }
 
@@ -2786,6 +2750,8 @@ void FStaticMeshComponentLODInfo::ReleaseOverrideVertexColorsAndBlock()
 	{
 		// The RT thread has no access to it any more so it's safe to delete it.
 		CleanUp();
+		// Ensure the RT no longer accessed the data, might slow down
+		FlushRenderingCommands();
 	}
 }
 
