Add missing undef for loctext namespace

diff --git a/Engine/Source/Runtime/Engine/Private/ObjectCacheContext.cpp b/Engine/Source/Runtime/Engine/Private/ObjectCacheContext.cpp
--- a/Engine/Source/Runtime/Engine/Private/ObjectCacheContext.cpp
+++ b/Engine/Source/Runtime/Engine/Private/ObjectCacheContext.cpp
@@ -714,3 +714,5 @@ void FObjectCacheEventSink::NotifyStaticMeshChanged_Concurrent(UStaticMeshCompon
 }
 
 #endif
+
+#undef LOCTEXT_NAMESPACE
\ No newline at end of file
