Fixing accidential removal of pre-engine tick for non-editor builds
#rnx

diff --git a/Engine/Source/Runtime/Engine/Private/GameEngine.cpp b/Engine/Source/Runtime/Engine/Private/GameEngine.cpp
--- a/Engine/Source/Runtime/Engine/Private/GameEngine.cpp
+++ b/Engine/Source/Runtime/Engine/Private/GameEngine.cpp
@@ -1708,6 +1708,17 @@ void UGameEngine::Tick( float DeltaSeconds, bool bIdleMode )
 		UE_LOG(LogEngine, Log, TEXT("Slow GT frame detected (GT frame %u, delta time %f s)"), GFrameCounter - 1, DeltaSeconds);
 	}
 
+#if !UE_SERVER
+	// tick media framework
+	static const FName MediaModuleName(TEXT("Media"));
+	IMediaModule* MediaModule = FModuleManager::LoadModulePtr<IMediaModule>(MediaModuleName);
+
+	if (MediaModule != nullptr)
+	{
+		MediaModule->TickPreEngine();
+	}
+#endif
+
 	if (IsRunningDedicatedServer())
 	{
 		double CurrentTime = FPlatformTime::Seconds();
@@ -1870,10 +1881,6 @@ void UGameEngine::Tick( float DeltaSeconds, bool bIdleMode )
 	}
 
 #if !UE_SERVER
-	// tick media framework
-	static const FName MediaModuleName(TEXT("Media"));
-	IMediaModule* MediaModule = FModuleManager::LoadModulePtr<IMediaModule>(MediaModuleName);
-
 	if (MediaModule != nullptr)
 	{
 		MediaModule->TickPostEngine();
