World Partition
- Added accessors for actor's grid placement and runtime grid properties.
- Grid placement and runtime grid properties are now editor only.

diff --git a/Engine/Source/Editor/UnrealEd/Private/Commandlets/WorldPartitionConvertCommandlet.cpp b/Engine/Source/Editor/UnrealEd/Private/Commandlets/WorldPartitionConvertCommandlet.cpp
--- a/Engine/Source/Editor/UnrealEd/Private/Commandlets/WorldPartitionConvertCommandlet.cpp
+++ b/Engine/Source/Editor/UnrealEd/Private/Commandlets/WorldPartitionConvertCommandlet.cpp
@@ -300,7 +300,7 @@ bool UWorldPartitionConvertCommandlet::PrepareStreamingLevelForConversion(const
 			{
 				if (Actor)
 				{
-					Actor->GridPlacement = EActorGridPlacement::AlwaysLoaded;
+					Actor->SetGridPlacement(EActorGridPlacement::AlwaysLoaded);
 				}
 			}
 		}
@@ -316,7 +316,7 @@ bool UWorldPartitionConvertCommandlet::PrepareStreamingLevelForConversion(const
 		{
 			if (Actor)
 			{
-				Actor->RuntimeGrid = RuntimeHash->GetActorRuntimeGrid(Actor);
+				Actor->SetRuntimeGrid(RuntimeHash->GetActorRuntimeGrid(Actor));
 			}
 		}
 	}
@@ -986,7 +986,7 @@ int32 UWorldPartitionConvertCommandlet::Main(const FString& Params)
 						LandscapeInfos.Add(LandscapeInfo);
 					}
 					// Only override default grid placement on actors that are not marked as always loaded
-					else if (Actor->GridPlacement != EActorGridPlacement::AlwaysLoaded)
+					else if (Actor->GetGridPlacement() != EActorGridPlacement::AlwaysLoaded)
 					{
 						FVector BoundsLocation;
 						FVector BoundsExtent;
@@ -996,11 +996,11 @@ int32 UWorldPartitionConvertCommandlet::Main(const FString& Params)
 
 						if (!WorldBounds.IsInside(ActorBounds))
 						{
-							Actor->GridPlacement = EActorGridPlacement::AlwaysLoaded;
+							Actor->SetGridPlacement(EActorGridPlacement::AlwaysLoaded);
 						}
 						else
 						{
-							Actor->GridPlacement = DefaultGridPlacement;
+							Actor->SetGridPlacement(DefaultGridPlacement);
 						}
 					}
 
diff --git a/Engine/Source/Runtime/Engine/Classes/GameFramework/Actor.h b/Engine/Source/Runtime/Engine/Classes/GameFramework/Actor.h
--- a/Engine/Source/Runtime/Engine/Classes/GameFramework/Actor.h
+++ b/Engine/Source/Runtime/Engine/Classes/GameFramework/Actor.h
@@ -48,6 +48,7 @@ enum class EActorUpdateOverlapsMethod : uint8
 	NeverUpdate			// Never update overlap state on initialization.
 };
 
+#if WITH_EDITORONLY_DATA
 /** Enum defining how actor will be placed in the partition */
 UENUM()
 enum class EActorGridPlacement : uint8
@@ -57,7 +58,7 @@ enum class EActorGridPlacement : uint8
 	AlwaysLoaded,				// Actor is always loaded (not placed in the grid), also affects editor.
 	None UMETA(Hidden)
 };
-
+#endif
 
 ENGINE_API DECLARE_LOG_CATEGORY_EXTERN(LogActor, Log, Warning);
 
@@ -546,6 +547,8 @@ public:
 	 */
 	float CreationTime;
 
+protected:
+#if WITH_EDITORONLY_DATA
 	/** 
 	 * Determine how this actor will be placed in the partition (if the world is partitioned). This value is not taken into account
 	 * if the function ActorClass->GetDefaultGridPlacement returns other than EActorGridPlacement::None.
@@ -559,8 +562,8 @@ public:
 	 */
 	UPROPERTY(EditAnywhere, Category=WorldPartition)
 	FName RuntimeGrid;
+#endif
 
-private:
 	/**
 	 * Used for replicating attachment of this actor's RootComponent to another actor.
 	 * This is filled in via GatherCurrentMovement() when the RootComponent has an AttachParent.
@@ -799,6 +802,18 @@ public:
 	 / @return EActorGridPlacement::None if the placement should be on a per-instance basis, otherwise the default value wins.
 	 */
 	virtual EActorGridPlacement GetDefaultGridPlacement() const;
+	
+	/** Returns this actor's current grid placement. */
+	virtual EActorGridPlacement GetGridPlacement() const { return GridPlacement; }
+
+	/** Sets this actor's current grid placement. */
+	void SetGridPlacement(EActorGridPlacement InGridPLacement) { GridPlacement = InGridPLacement; }
+
+	/** Returns this actor's current target runtime grid. */
+	virtual FName GetRuntimeGrid() const { return RuntimeGrid; }
+
+	/** Sets this actor's current target runtime grid. */
+	void SetRuntimeGrid(FName InRuntimeGrid) { RuntimeGrid = InRuntimeGrid; }
 
 	/** Returns this actor's current Guid. Actor Guids are only available in development builds. */
 	inline const FGuid& GetActorGuid() const { return ActorGuid; }
diff --git a/Engine/Source/Runtime/Engine/Private/Actor.cpp b/Engine/Source/Runtime/Engine/Private/Actor.cpp
--- a/Engine/Source/Runtime/Engine/Private/Actor.cpp
+++ b/Engine/Source/Runtime/Engine/Private/Actor.cpp
@@ -172,7 +172,9 @@ void AActor::InitializeDefaults()
 	}
 #endif // (CSV_PROFILER && !UE_BUILD_SHIPPING)
 
+#if WITH_EDITORONLY_DATA
 	GridPlacement = EActorGridPlacement::Bounds;
+#endif
 }
 
 void FActorTickFunction::ExecuteTick(float DeltaTime, enum ELevelTick TickType, ENamedThreads::Type CurrentThread, const FGraphEventRef& MyCompletionGraphEvent)
diff --git a/Engine/Source/Runtime/Engine/Private/LevelActor.cpp b/Engine/Source/Runtime/Engine/Private/LevelActor.cpp
--- a/Engine/Source/Runtime/Engine/Private/LevelActor.cpp
+++ b/Engine/Source/Runtime/Engine/Private/LevelActor.cpp
@@ -683,7 +683,7 @@ AActor* UWorld::SpawnActor( UClass* Class, FTransform const* UserTransformPtr, c
 	EActorGridPlacement DefaultGridPlacement = Class->GetDefaultObject<AActor>()->GetDefaultGridPlacement();
 	if (DefaultGridPlacement != EActorGridPlacement::None)
 	{
-		Actor->GridPlacement = DefaultGridPlacement;
+		Actor->SetGridPlacement(DefaultGridPlacement);
 	}
 #endif
 
diff --git a/Engine/Source/Runtime/Engine/Private/WorldPartition/HLOD/HLODBuilder.cpp b/Engine/Source/Runtime/Engine/Private/WorldPartition/HLOD/HLODBuilder.cpp
--- a/Engine/Source/Runtime/Engine/Private/WorldPartition/HLOD/HLODBuilder.cpp
+++ b/Engine/Source/Runtime/Engine/Private/WorldPartition/HLOD/HLODBuilder.cpp
@@ -66,7 +66,7 @@ public:
 			HLODActor->Modify();
 			HLODActor->SetHLODPrimitives(HLODPrimitives);
 			HLODActor->SetChildrenPrimitives(InSubComponents);
-			HLODActor->RuntimeGrid = HLODLayer->GetRuntimeGrid(HLODLevel);
+			HLODActor->SetRuntimeGrid(HLODLayer->GetRuntimeGrid(HLODLevel));
 			HLODActor->SetLODLevel(HLODLevel);
 			HLODActor->SetHLODLayer(HLODLayer->GetParentLayer().LoadSynchronous());
 			HLODActor->SetSubActorsHLODLayer(HLODLayer);
diff --git a/Engine/Source/Runtime/Engine/Private/WorldPartition/RuntimeSpatialHash/RuntimeSpatialHashNavigation.cpp b/Engine/Source/Runtime/Engine/Private/WorldPartition/RuntimeSpatialHash/RuntimeSpatialHashNavigation.cpp
--- a/Engine/Source/Runtime/Engine/Private/WorldPartition/RuntimeSpatialHash/RuntimeSpatialHashNavigation.cpp
+++ b/Engine/Source/Runtime/Engine/Private/WorldPartition/RuntimeSpatialHash/RuntimeSpatialHashNavigation.cpp
@@ -98,7 +98,7 @@ bool UWorldPartitionRuntimeSpatialHash::GenerateNavigationData()
 				DataChunkActor->SetActorLabel(FString::Printf(TEXT("NavDataChunkActor_%s_%s"), *GetName(), *CellName.ToString()));
 
 				// Set target grid
-				DataChunkActor->RuntimeGrid = RuntimeGrid.GridName;
+				DataChunkActor->SetRuntimeGrid(RuntimeGrid.GridName);
 				ValidNavigationDataChunkActors.Add(DataChunkActor);
 
 				UE_LOG(LogWorldPartitionRuntimeSpatialHashNav, Verbose, TEXT("%i) %s added."), ActorCount, *DataChunkActor->GetName());
diff --git a/Engine/Source/Runtime/Engine/Private/WorldPartition/WorldPartitionActorDesc.cpp b/Engine/Source/Runtime/Engine/Private/WorldPartition/WorldPartitionActorDesc.cpp
--- a/Engine/Source/Runtime/Engine/Private/WorldPartition/WorldPartitionActorDesc.cpp
+++ b/Engine/Source/Runtime/Engine/Private/WorldPartition/WorldPartitionActorDesc.cpp
@@ -57,10 +57,10 @@ void FWorldPartitionActorDesc::Init(const AActor* InActor)
 	}
 	else
 	{
-		GridPlacement = InActor->GridPlacement;
+		GridPlacement = InActor->GetGridPlacement();
 	}
 
-	RuntimeGrid = InActor->RuntimeGrid;
+	RuntimeGrid = InActor->GetRuntimeGrid();
 	bActorIsEditorOnly = InActor->IsEditorOnly();
 	bLevelBoundsRelevant = InActor->IsLevelBoundsRelevant();
 	DataLayers = InActor->GetDataLayerNames();
diff --git a/Engine/Source/Runtime/Engine/Private/WorldPartition/WorldPartitionMiniMapHelper.cpp b/Engine/Source/Runtime/Engine/Private/WorldPartition/WorldPartitionMiniMapHelper.cpp
--- a/Engine/Source/Runtime/Engine/Private/WorldPartition/WorldPartitionMiniMapHelper.cpp
+++ b/Engine/Source/Runtime/Engine/Private/WorldPartition/WorldPartitionMiniMapHelper.cpp
@@ -140,7 +140,7 @@ void FWorldPartitionMiniMapHelper::CaptureWorldMiniMapToTexture(UWorld* InWorld,
 
 bool FWorldPartitionMiniMapHelper::DoesActorContributeToBounds(AActor* Actor)
 {
-	return Actor && Actor->GridPlacement != EActorGridPlacement::AlwaysLoaded;
+	return Actor && Actor->GetGridPlacement() != EActorGridPlacement::AlwaysLoaded;
 }
 
 FBox FWorldPartitionMiniMapHelper::GetWorldBounds(UWorld* World)
