PostProcessVolume grid placement is now forced to be always loaded when bUnbound is true.

diff --git a/Engine/Source/Runtime/Engine/Classes/Engine/PostProcessVolume.h b/Engine/Source/Runtime/Engine/Classes/Engine/PostProcessVolume.h
--- a/Engine/Source/Runtime/Engine/Classes/Engine/PostProcessVolume.h
+++ b/Engine/Source/Runtime/Engine/Classes/Engine/PostProcessVolume.h
@@ -69,6 +69,10 @@ class ENGINE_API APostProcessVolume : public AVolume, public IInterface_PostProc
 	//~ Begin AActor Interface
 	virtual void PostUnregisterAllComponents( void ) override;
 
+#if WITH_EDITOR
+	virtual EActorGridPlacement GetGridPlacement() const override;
+#endif
+
 protected:
 	virtual void PostRegisterAllComponents() override;
 	//~ End AActor Interface
diff --git a/Engine/Source/Runtime/Engine/Private/PostProcessVolume.cpp b/Engine/Source/Runtime/Engine/Private/PostProcessVolume.cpp
--- a/Engine/Source/Runtime/Engine/Private/PostProcessVolume.cpp
+++ b/Engine/Source/Runtime/Engine/Private/PostProcessVolume.cpp
@@ -36,6 +36,10 @@ void APostProcessVolume::Serialize(FArchive& Ar)
 }
 
 #if WITH_EDITOR
+EActorGridPlacement APostProcessVolume::GetGridPlacement() const
+{
+	return bUnbound ? EActorGridPlacement::AlwaysLoaded : AActor::GetGridPlacement();
+}
 
 void APostProcessVolume::PostEditChangeProperty(FPropertyChangedEvent& PropertyChangedEvent)
 {
@@ -199,6 +203,11 @@ bool APostProcessVolume::CanEditChange(const FProperty* InProperty) const
 				return false;
 			}
 		}
+
+		if (PropertyName == GET_MEMBER_NAME_STRING_CHECKED(APostProcessVolume, GridPlacement))
+		{
+			return !bUnbound;
+		}
 	}
 
 	return Super::CanEditChange(InProperty);
