Added option to have dynamic data layers initially active

diff --git a/Engine/Source/Runtime/Engine/Private/WorldPartition/DataLayer/DataLayer.cpp b/Engine/Source/Runtime/Engine/Private/WorldPartition/DataLayer/DataLayer.cpp
--- a/Engine/Source/Runtime/Engine/Private/WorldPartition/DataLayer/DataLayer.cpp
+++ b/Engine/Source/Runtime/Engine/Private/WorldPartition/DataLayer/DataLayer.cpp
@@ -14,6 +14,7 @@ UDataLayer::UDataLayer(const FObjectInitializer& ObjectInitializer)
 , DataLayerLabel(GetFName())
 , bIsVisible(true)
 , bIsDynamicallyLoaded(false)
+, bIsInitiallyActive(false)
 , bIsDynamicallyLoadedInEditor(true)
 , bGeneratesHLODs(false)
 , DefaultHLODLayer()
diff --git a/Engine/Source/Runtime/Engine/Private/WorldPartition/DataLayer/WorldDataLayers.cpp b/Engine/Source/Runtime/Engine/Private/WorldPartition/DataLayer/WorldDataLayers.cpp
--- a/Engine/Source/Runtime/Engine/Private/WorldPartition/DataLayer/WorldDataLayers.cpp
+++ b/Engine/Source/Runtime/Engine/Private/WorldPartition/DataLayer/WorldDataLayers.cpp
@@ -6,6 +6,7 @@
 
 #include "WorldPartition/DataLayer/WorldDataLayers.h"
 #include "WorldPartition/DataLayer/DataLayer.h"
+#include "WorldPartition/DataLayer/DataLayerSubsystem.h"
 #include "EngineUtils.h"
 #if WITH_EDITOR
 #include "WorldPartition/DataLayer/DataLayerEditorPerProjectUserSettings.h"
@@ -252,4 +253,21 @@ void AWorldDataLayers::PostLoad()
 #endif
 }
 
+void AWorldDataLayers::BeginPlay()
+{
+	Super::BeginPlay();
+
+	// Activate initially loaded data layers
+	if (UDataLayerSubsystem* DataLayerSubsystem = GetWorld()->GetSubsystem<UDataLayerSubsystem>())
+	{
+		for (const UDataLayer* DataLayer : WorldDataLayers)
+		{
+			if (DataLayer && DataLayer->IsInitiallyActive())
+			{
+				DataLayerSubsystem->ActivateDataLayer(DataLayer, true);
+			}
+		}
+	}
+}
+
 #undef LOCTEXT_NAMESPACE
\ No newline at end of file
diff --git a/Engine/Source/Runtime/Engine/Public/WorldPartition/DataLayer/DataLayer.h b/Engine/Source/Runtime/Engine/Public/WorldPartition/DataLayer/DataLayer.h
--- a/Engine/Source/Runtime/Engine/Public/WorldPartition/DataLayer/DataLayer.h
+++ b/Engine/Source/Runtime/Engine/Public/WorldPartition/DataLayer/DataLayer.h
@@ -32,6 +32,7 @@ public:
 	FName GetDataLayerLabel() const  { return DataLayerLabel; }
 	bool IsVisible() const { return bIsVisible; }
 	bool IsDynamicallyLoaded() const { return bIsDynamicallyLoaded; }
+	bool IsInitiallyActive() const { return IsDynamicallyLoaded() && bIsInitiallyActive; }
 
 private:
 	/** The display name of the DataLayer */
@@ -46,6 +47,10 @@ private:
 	UPROPERTY(Category = DataLayer, EditAnywhere)
 	uint32 bIsDynamicallyLoaded : 1;
 
+	/** Whether a dynamically loaded DataLayer should be initially active at runtime */
+	UPROPERTY(Category = DataLayer, EditAnywhere, meta = (EditConditionHides, EditCondition = "bIsDynamicallyLoaded"))
+	uint32 bIsInitiallyActive : 1;
+
 #if WITH_EDITORONLY_DATA
 	/** Whether the DataLayer affects actor editor loading */
 	UPROPERTY(Transient)
diff --git a/Engine/Source/Runtime/Engine/Public/WorldPartition/DataLayer/WorldDataLayers.h b/Engine/Source/Runtime/Engine/Public/WorldPartition/DataLayer/WorldDataLayers.h
--- a/Engine/Source/Runtime/Engine/Public/WorldPartition/DataLayer/WorldDataLayers.h
+++ b/Engine/Source/Runtime/Engine/Public/WorldPartition/DataLayer/WorldDataLayers.h
@@ -20,6 +20,8 @@ class ENGINE_API AWorldDataLayers : public AInfo
 public:
 	static const AWorldDataLayers* Get(UWorld* World);
 	virtual void PostLoad() override;
+	
+	virtual void BeginPlay() override;
 
 #if WITH_EDITOR
 	static AWorldDataLayers* Get(UWorld* World, bool bCreateIfNotFound);
