World Partition streaming sources
- Added interface for streaming source providers: a streaming source keeps data loaded at a specific loacation.
- Added WorldPartitionStreamingSourceComponent that will register itself as a streaming source when attached to an actor.

diff --git a/Engine/Source/Runtime/Engine/Classes/Components/WorldPartitionStreamingSourceComponent.h b/Engine/Source/Runtime/Engine/Classes/Components/WorldPartitionStreamingSourceComponent.h
new file mode 100644
--- /dev/null
+++ b/Engine/Source/Runtime/Engine/Classes/Components/WorldPartitionStreamingSourceComponent.h
@@ -0,0 +1,21 @@
+// Copyright Epic Games, Inc. All Rights Reserved.
+
+#pragma once
+
+#include "CoreMinimal.h"
+#include "UObject/ObjectMacros.h"
+#include "Components/ActorComponent.h"
+#include "WorldPartition/WorldPartitionStreamingSource.h"
+#include "WorldPartitionStreamingSourceComponent.generated.h"
+
+UCLASS(Meta = (BlueprintSpawnableComponent))
+class ENGINE_API UWorldPartitionStreamingSourceComponent : public UActorComponent, public IWorldPartitionStreamingSourceProvider
+{
+	GENERATED_UCLASS_BODY()
+
+	virtual void OnRegister() override;
+	virtual void OnUnregister() override;
+
+	// IWorldPartitionStreamingSourceProvider interface
+	virtual bool GetStreamingSource(FWorldPartitionStreamingSource& OutStreamingSource) override;
+};
\ No newline at end of file
diff --git a/Engine/Source/Runtime/Engine/Private/Components/WorldPartitionStreamingSourceComponent.cpp b/Engine/Source/Runtime/Engine/Private/Components/WorldPartitionStreamingSourceComponent.cpp
new file mode 100644
--- /dev/null
+++ b/Engine/Source/Runtime/Engine/Private/Components/WorldPartitionStreamingSourceComponent.cpp
@@ -0,0 +1,69 @@
+// Copyright Epic Games, Inc. All Rights Reserved.
+
+#include "Components/WorldPartitionStreamingSourceComponent.h"
+#include "WorldPartition/WorldPartition.h"
+
+UWorldPartitionStreamingSourceComponent::UWorldPartitionStreamingSourceComponent(const FObjectInitializer& ObjectInitializer)
+	: Super(ObjectInitializer)
+{
+	PrimaryComponentTick.bCanEverTick = false;
+}
+
+void UWorldPartitionStreamingSourceComponent::OnRegister()
+{
+	Super::OnRegister();
+
+	UWorld* World = GetWorld();
+
+#if WITH_EDITOR
+	if (!World->IsPlayInEditor())
+	{
+		return;
+	}
+#endif
+
+	if (UWorldPartition* WorldPartition = World->GetWorldPartition())
+	{
+		AActor* Actor = GetOwner();
+		ULevel* ActorLevel = Actor->GetLevel();
+
+		if (ActorLevel == World->PersistentLevel)
+		{
+			WorldPartition->RegisterStreamingSourceProvider(this);
+		}
+		else
+		{
+			UE_LOG(LogWorldPartition, Warning, TEXT("Attaching a WorldPartitionStreamingSourceComponent to an actor that is in the grid will result in never unloading the affected cells (%s)"), *Actor->GetName());
+		}
+	}
+}
+
+void UWorldPartitionStreamingSourceComponent::OnUnregister()
+{
+	Super::OnUnregister();
+
+	UWorld* World = GetWorld();
+
+#if WITH_EDITOR
+	if (!World->IsPlayInEditor())
+	{
+		return;
+	}
+#endif
+
+	if (UWorldPartition* WorldPartition = World->GetWorldPartition())
+	{
+		AActor* Actor = GetOwner();
+		ULevel* ActorLevel = Actor->GetLevel();
+
+		verify(WorldPartition->UnregisterStreamingSourceProvider(this) || (ActorLevel != World->PersistentLevel));
+	}
+}
+
+bool UWorldPartitionStreamingSourceComponent::GetStreamingSource(FWorldPartitionStreamingSource& OutStreamingSource)
+{
+	AActor* Actor = GetOwner();
+	OutStreamingSource.Location = Actor->GetActorLocation();
+	OutStreamingSource.Rotation = Actor->GetActorRotation();
+	return true;
+}
\ No newline at end of file
diff --git a/Engine/Source/Runtime/Engine/Private/WorldPartition/WorldPartition.cpp b/Engine/Source/Runtime/Engine/Private/WorldPartition/WorldPartition.cpp
--- a/Engine/Source/Runtime/Engine/Private/WorldPartition/WorldPartition.cpp
+++ b/Engine/Source/Runtime/Engine/Private/WorldPartition/WorldPartition.cpp
@@ -10,6 +10,7 @@
 #include "WorldPartition/WorldPartitionLevelStreamingPolicy.h"
 #include "WorldPartition/WorldPartitionActorDesc.h"
 #include "WorldPartition/WorldPartitionSubsystem.h"
+#include "WorldPartition/WorldPartitionStreamingSource.h"
 #include "Algo/Accumulate.h"
 #include "Algo/Transform.h"
 #include "Algo/RemoveIf.h"
@@ -283,6 +284,16 @@ UWorld* UWorldPartition::GetWorld() const
 	return Super::GetWorld();
 }
 
+void UWorldPartition::RegisterStreamingSourceProvider(IWorldPartitionStreamingSourceProvider* StreamingSource)
+{
+	StreamingSourceProviders.Add(StreamingSource);
+}
+
+bool UWorldPartition::UnregisterStreamingSourceProvider(IWorldPartitionStreamingSourceProvider* StreamingSource)
+{
+	return !!StreamingSourceProviders.Remove(StreamingSource);
+}
+
 // This will trap all broadcast of LevelRemovedFromWorld and Uninitialize world partition if existing
 void UWorldPartition::WorldPartitionOnLevelRemovedFromWorld(ULevel* Level, UWorld* InWorld)
 {
diff --git a/Engine/Source/Runtime/Engine/Private/WorldPartition/WorldPartitionStreamingPolicy.cpp b/Engine/Source/Runtime/Engine/Private/WorldPartition/WorldPartitionStreamingPolicy.cpp
--- a/Engine/Source/Runtime/Engine/Private/WorldPartition/WorldPartitionStreamingPolicy.cpp
+++ b/Engine/Source/Runtime/Engine/Private/WorldPartition/WorldPartitionStreamingPolicy.cpp
@@ -7,6 +7,7 @@
 #include "WorldPartition/WorldPartitionStreamingPolicy.h"
 #include "WorldPartition/WorldPartitionRuntimeCell.h"
 #include "WorldPartition/WorldPartitionRuntimeHash.h"
+#include "WorldPartition/WorldPartitionStreamingSource.h"
 #include "WorldPartition/WorldPartition.h"
 #include "GameFramework/PlayerController.h"
 #include "Engine/LocalPlayer.h"
@@ -83,6 +84,19 @@ void UWorldPartitionStreamingPolicy::UpdateStreamingSources()
 			}
 		}
 	}
+
+	for (IWorldPartitionStreamingSourceProvider* StreamingSourceProvider : WorldPartition->StreamingSourceProviders)
+	{
+		FWorldPartitionStreamingSource StreamingSource;
+		if (StreamingSourceProvider->GetStreamingSource(StreamingSource))
+		{
+			// Transform to Local
+			StreamingSource.Location = WorldToLocal.TransformPosition(StreamingSource.Location);
+			StreamingSource.Rotation = WorldToLocal.TransformRotation(StreamingSource.Rotation.Quaternion()).Rotator();
+
+			StreamingSources.Add(StreamingSource);
+		}
+	}
 }
 
 void UWorldPartitionStreamingPolicy::UpdateStreamingState()
@@ -137,11 +151,11 @@ void UWorldPartitionStreamingPolicy::UpdateStreamingState()
 			{
 				UE_LOG(LogWorldPartition, Verbose, TEXT("UWorldPartitionStreamingPolicy: CellsToLoad(%d), CellsToUnload(%d)"), ToLoadCells.Num(), ToUnloadCells.Num());
 				
-				FTransform WorldToLocal = WorldPartition->GetInstanceTransform();
+				FTransform LocalToWorld = WorldPartition->GetInstanceTransform();
 				for (int i = 0; i < StreamingSources.Num(); ++i)
 				{
-					FVector ViewLocation = WorldToLocal.TransformPosition(StreamingSources[i].Location);
-					FRotator ViewRotation = WorldToLocal.TransformRotation(StreamingSources[i].Rotation.Quaternion()).Rotator();
+					FVector ViewLocation = LocalToWorld.TransformPosition(StreamingSources[i].Location);
+					FRotator ViewRotation = LocalToWorld.TransformRotation(StreamingSources[i].Rotation.Quaternion()).Rotator();
 					UE_LOG(LogWorldPartition, Verbose, TEXT("UWorldPartitionStreamingPolicy: Sources[%d] = %s,%s"), i, *ViewLocation.ToString(), *ViewRotation.ToString());
 				}
 			}
diff --git a/Engine/Source/Runtime/Engine/Public/WorldPartition/WorldPartition.h b/Engine/Source/Runtime/Engine/Public/WorldPartition/WorldPartition.h
--- a/Engine/Source/Runtime/Engine/Public/WorldPartition/WorldPartition.h
+++ b/Engine/Source/Runtime/Engine/Public/WorldPartition/WorldPartition.h
@@ -7,6 +7,7 @@
 #include "GameFramework/Actor.h"
 #include "Templates/SubclassOf.h"
 #include "WorldPartition/WorldPartitionActorDesc.h"
+#include "WorldPartition/WorldPartitionStreamingSource.h"
 #include "WorldPartition/ActorDescContainer.h"
 
 #if WITH_EDITOR
@@ -23,6 +24,8 @@ class UWorldPartitionRuntimeHash;
 class UWorldPartitionStreamingPolicy;
 class FHLODActorDesc;
 
+struct IWorldPartitionStreamingSourceProvider;
+
 DECLARE_LOG_CATEGORY_EXTERN(LogWorldPartition, Log, All);
 
 enum class EWorldPartitionStreamingMode
@@ -156,6 +159,9 @@ public:
 
 	virtual UWorld* GetWorld() const override;
 
+	void RegisterStreamingSourceProvider(IWorldPartitionStreamingSourceProvider* StreamingSource);
+	bool UnregisterStreamingSourceProvider(IWorldPartitionStreamingSourceProvider* StreamingSource);
+
 	UPROPERTY(Transient)
 	TObjectPtr<UWorld> World;
 
@@ -187,6 +193,8 @@ private:
 	UPROPERTY(Transient, DuplicateTransient)
 	mutable TObjectPtr<UWorldPartitionStreamingPolicy> StreamingPolicy;
 
+	TArray<IWorldPartitionStreamingSourceProvider*> StreamingSourceProviders;
+
 #if WITH_EDITORONLY_DATA
 	FLinkerInstancingContext InstancingContext;
 #endif
diff --git a/Engine/Source/Runtime/Engine/Public/WorldPartition/WorldPartitionStreamingPolicy.h b/Engine/Source/Runtime/Engine/Public/WorldPartition/WorldPartitionStreamingPolicy.h
--- a/Engine/Source/Runtime/Engine/Public/WorldPartition/WorldPartitionStreamingPolicy.h
+++ b/Engine/Source/Runtime/Engine/Public/WorldPartition/WorldPartitionStreamingPolicy.h
@@ -11,21 +11,11 @@
 
 #include "Containers/Set.h"
 #include "WorldPartition/WorldPartitionRuntimeCell.h"
+#include "WorldPartition/WorldPartitionStreamingSource.h"
 #include "WorldPartitionStreamingPolicy.generated.h"
 
 class UWorldPartition;
 
-struct FWorldPartitionStreamingSource
-{
-	FWorldPartitionStreamingSource(const FVector& InLocation, const FRotator& InRotation)
-		: Location(InLocation)
-		, Rotation(InRotation)
-	{}
-
-	FVector Location;
-	FRotator Rotation;
-};
-
 UCLASS(Abstract, Within = WorldPartition)
 class UWorldPartitionStreamingPolicy : public UObject
 {
diff --git a/Engine/Source/Runtime/Engine/Public/WorldPartition/WorldPartitionStreamingSource.h b/Engine/Source/Runtime/Engine/Public/WorldPartition/WorldPartitionStreamingSource.h
new file mode 100644
--- /dev/null
+++ b/Engine/Source/Runtime/Engine/Public/WorldPartition/WorldPartitionStreamingSource.h
@@ -0,0 +1,29 @@
+// Copyright Epic Games, Inc. All Rights Reserved.
+#pragma once
+
+#include "CoreMinimal.h"
+
+/**
+ * Structure containing all properties required to stream from a source
+ */
+struct ENGINE_API FWorldPartitionStreamingSource
+{
+	FWorldPartitionStreamingSource()
+	{}
+
+	FWorldPartitionStreamingSource(const FVector& InLocation, const FRotator& InRotation)
+		: Location(InLocation)
+		, Rotation(InRotation)
+	{}
+
+	FVector Location;
+	FRotator Rotation;
+};
+
+/**
+ * Interface for world partition streaming sources
+ */
+struct ENGINE_API IWorldPartitionStreamingSourceProvider
+{
+	virtual bool GetStreamingSource(FWorldPartitionStreamingSource& StreamingSource) = 0;
+};
\ No newline at end of file
