Fix crash after rebuild of virtual texture referenced by two volumes.

diff --git a/Engine/Source/Runtime/Engine/Private/Components/RuntimeVirtualTextureComponent.cpp b/Engine/Source/Runtime/Engine/Private/Components/RuntimeVirtualTextureComponent.cpp
--- a/Engine/Source/Runtime/Engine/Private/Components/RuntimeVirtualTextureComponent.cpp
+++ b/Engine/Source/Runtime/Engine/Private/Components/RuntimeVirtualTextureComponent.cpp
@@ -8,6 +8,7 @@
 #include "Misc/UObjectToken.h"
 #include "Misc/MapErrors.h"
 #include "SceneInterface.h"
+#include "UObject/UObjectIterator.h"
 #include "VT/RuntimeVirtualTexture.h"
 #include "VT/VirtualTextureBuilder.h"
 
@@ -160,14 +161,46 @@ bool URuntimeVirtualTextureComponent::IsStreamingLowMips() const
 
 #if WITH_EDITOR
 
+// RAII class to release and recreate runtime virtual texture producers associated with a UVirtualTextureBuilder.
+// Required around modifications of a UVirtualTextureBuilder because virtual producers hold pointers to the internal data.
+class FScopedRuntimeVirtualTextureRecreate
+{
+public:
+	FScopedRuntimeVirtualTextureRecreate(UVirtualTextureBuilder* VirtualTextureBuilder)
+	{
+		for (TObjectIterator<URuntimeVirtualTextureComponent> It; It; ++It)
+		{
+			if (It->GetStreamingTexture() == VirtualTextureBuilder)
+			{
+				URuntimeVirtualTexture* VirtualTexture = It->GetVirtualTexture();
+				if (VirtualTexture != nullptr)
+				{
+					VirtualTextures.Add(VirtualTexture);
+					VirtualTexture->Release();
+				}
+			}
+		}
+	}
+
+	~FScopedRuntimeVirtualTextureRecreate()
+	{
+		for (URuntimeVirtualTexture* VirtualTexture : VirtualTextures)
+		{
+			// PostEditChange will trigger the correct notifications and recreation of virtual texture producers.
+			VirtualTexture->PostEditChange();
+		}
+	}
+
+private:
+	TArray<URuntimeVirtualTexture*> VirtualTextures;
+};
+
 void URuntimeVirtualTextureComponent::InitializeStreamingTexture(uint32 InSizeX, uint32 InSizeY, uint8* InData)
 {
 	// We need an existing StreamingTexture object to update.
 	if (VirtualTexture != nullptr && StreamingTexture != nullptr)
 	{
-		// Release current runtime virtual texture producer.
-		// It may reference data inside the old StreamingTexture which could be garbage collected any time from now.
-		VirtualTexture->Release();
+		FScopedRuntimeVirtualTextureRecreate ProducerRecreate(StreamingTexture);
 
 		FVirtualTextureBuildDesc BuildDesc;
 		BuildDesc.bContinuousUpdate = VirtualTexture->GetContinuousUpdate();
@@ -203,9 +236,6 @@ void URuntimeVirtualTextureComponent::InitializeStreamingTexture(uint32 InSizeX,
 
 		StreamingTexture->Modify();
 		StreamingTexture->BuildTexture(BuildDesc);
-
-		// Trigger refresh of the runtime virtual texture producer.
-		VirtualTexture->PostEditChange();
 	}
 }
 
