LevelInstance: Prevent breaking level instances which have level scripts.

diff --git a/Engine/Source/Editor/LevelInstanceEditor/Private/LevelInstanceEditorModule.cpp b/Engine/Source/Editor/LevelInstanceEditor/Private/LevelInstanceEditorModule.cpp
--- a/Engine/Source/Editor/LevelInstanceEditor/Private/LevelInstanceEditorModule.cpp
+++ b/Engine/Source/Editor/LevelInstanceEditor/Private/LevelInstanceEditorModule.cpp
@@ -495,7 +495,7 @@ namespace LevelInstanceMenuUtils
 					return true;
 				});
 
-			if (ContextLevelInstance && !ContextLevelInstance->IsEditing())
+			if (ContextLevelInstance && !ContextLevelInstance->IsEditing() && !LevelInstanceSubsystem->LevelInstanceHasLevelScriptBlueprint(ContextLevelInstance))
 			{
 				FToolMenuSection& Section = CreateLevelInstanceSection(Menu);
 
diff --git a/Engine/Source/Runtime/Engine/Private/LevelInstance/LevelInstanceSubsystem.cpp b/Engine/Source/Runtime/Engine/Private/LevelInstance/LevelInstanceSubsystem.cpp
--- a/Engine/Source/Runtime/Engine/Private/LevelInstance/LevelInstanceSubsystem.cpp
+++ b/Engine/Source/Runtime/Engine/Private/LevelInstance/LevelInstanceSubsystem.cpp
@@ -33,6 +33,8 @@
 #include "Engine/Blueprint.h"
 #include "LevelInstance/Packed/PackedLevelInstanceActor.h"
 #include "LevelInstance/Packed/PackedLevelInstanceBuilder.h"
+#include "Engine/LevelScriptBlueprint.h"
+#include "EdGraph/EdGraph.h"
 #endif
 
 #include "HAL/IConsoleManager.h"
@@ -783,6 +785,7 @@ bool ULevelInstanceSubsystem::BreakLevelInstance(ALevelInstance* LevelInstanceAc
 		// Can only break the top level LevelInstance
 		check(LevelInstanceActor->GetLevel() == GetWorld()->GetCurrentLevel());
 
+
 		// Actors in a packed level instance will not be streamed in unless they are editing. Must force this before moving.
 		if (LevelInstanceActor->IsA<APackedLevelInstance>())
 		{
@@ -792,6 +795,18 @@ bool ULevelInstanceSubsystem::BreakLevelInstance(ALevelInstance* LevelInstanceAc
 		// need to ensure that LevelInstanceActor has been streamed in fully
 		GEngine->BlockTillLevelStreamingCompleted(LevelInstanceActor->GetWorld());
 
+		// Cannot break a level instance which has a level script
+		if (LevelInstanceHasLevelScriptBlueprint(LevelInstanceActor))
+		{
+			UE_LOG(LogLevelInstance, Warning, TEXT("Failed to completely break Level Instance because some children have Level Scripts."));
+
+			if (LevelInstanceActor->IsA<APackedLevelInstance>())
+			{
+				BlockUnloadLevelInstance(LevelInstanceActor);
+			}
+			return false;
+		}
+
 		TArray<AActor*> ActorsToMove;
 		ForEachActorInLevelInstance(LevelInstanceActor, [this, &ActorsToMove](AActor* Actor)
 			{
@@ -869,6 +884,33 @@ ULevel* ULevelInstanceSubsystem::GetLevelInstanceLevel(const ALevelInstance* Lev
 	return nullptr;
 }
 
+bool ULevelInstanceSubsystem::LevelInstanceHasLevelScriptBlueprint(const ALevelInstance* LevelInstance) const
+{
+	if (LevelInstance)
+	{
+		if (ULevel* LevelInstanceLevel = GetLevelInstanceLevel(LevelInstance))
+		{
+			if (ULevelScriptBlueprint* LevelScriptBP = LevelInstanceLevel->GetLevelScriptBlueprint(true))
+			{
+				TArray<UEdGraph*> AllGraphs;
+				LevelScriptBP->GetAllGraphs(AllGraphs);
+				for (UEdGraph* CurrentGraph : AllGraphs)
+				{
+					for (UEdGraphNode* Node : CurrentGraph->Nodes)
+					{
+						if (!Node->IsAutomaticallyPlacedGhostNode())
+						{
+							return true;
+						}
+					}
+				}
+			}
+		}
+	}
+
+	return false;
+}
+
 void ULevelInstanceSubsystem::RemoveLevelFromWorld(ULevel* Level, bool bResetTrans)
 {
 	if (LevelsToRemoveScope)
diff --git a/Engine/Source/Runtime/Engine/Public/LevelInstance/LevelInstanceSubsystem.h b/Engine/Source/Runtime/Engine/Public/LevelInstance/LevelInstanceSubsystem.h
--- a/Engine/Source/Runtime/Engine/Public/LevelInstance/LevelInstanceSubsystem.h
+++ b/Engine/Source/Runtime/Engine/Public/LevelInstance/LevelInstanceSubsystem.h
@@ -83,6 +83,8 @@ public:
 	void OnActorDeleted(AActor* Actor);
 	ULevel* GetLevelInstanceLevel(const ALevelInstance* LevelInstanceActor) const;
 
+	bool LevelInstanceHasLevelScriptBlueprint(const ALevelInstance* LevelInstance) const;
+
 	ALevelInstance* GetParentLevelInstance(const AActor* Actor) const;
 	void BlockLoadLevelInstance(ALevelInstance* LevelInstanceActor);
 	void BlockUnloadLevelInstance(ALevelInstance* LevelInstanceActor);
