Add flag to support texture build not stripping the alpha channel.
This has a single use case: We build streaming virtual textures from runtime virtual textures and we want to ensure the format doesn't change just because the alpha channel is all 1.

diff --git a/Engine/Source/Developer/TextureCompressor/Private/TextureCompressorModule.cpp b/Engine/Source/Developer/TextureCompressor/Private/TextureCompressorModule.cpp
--- a/Engine/Source/Developer/TextureCompressor/Private/TextureCompressorModule.cpp
+++ b/Engine/Source/Developer/TextureCompressor/Private/TextureCompressorModule.cpp
@@ -2060,7 +2060,7 @@ static bool CompressMipChain(
 {
 	TRACE_CPUPROFILER_EVENT_SCOPE(CompressMipChain)
 
-	const bool bImageHasAlphaChannel = DetectAlphaChannel(MipChain[0]);
+	const bool bImageHasAlphaChannel = Settings.bForceAlphaChannel || DetectAlphaChannel(MipChain[0]);
 
 	// now call the Ex version now that we have the proper MipChain
 	const FTextureFormatCompressorCaps CompressorCaps = TextureFormat->GetFormatCapabilitiesEx(Settings, MipChain.Num(), MipChain[0], bImageHasAlphaChannel);
diff --git a/Engine/Source/Developer/TextureCompressor/Public/TextureCompressorModule.h b/Engine/Source/Developer/TextureCompressor/Public/TextureCompressorModule.h
--- a/Engine/Source/Developer/TextureCompressor/Public/TextureCompressorModule.h
+++ b/Engine/Source/Developer/TextureCompressor/Public/TextureCompressorModule.h
@@ -99,6 +99,8 @@ struct FTextureBuildSettings
 	uint32 bUseLegacyGamma : 1;
 	/** Whether the border of the image should be maintained during mipmap generation. */
 	uint32 bPreserveBorder : 1;
+	/** Whether we should not discard the alpha channel when it contains 1 for the entire texture. */
+	uint32 bForceAlphaChannel : 1;
 	/** Whether the alpha channel should contain a dithered alpha value. */
 	uint32 bDitherMipMapAlpha : 1;
 	/** Whether bokeh alpha values should be computed for the texture. */
@@ -185,6 +187,7 @@ struct FTextureBuildSettings
 		, bSRGB(false)
 		, bUseLegacyGamma(false)
 		, bPreserveBorder(false)
+		, bForceAlphaChannel(false)
 		, bDitherMipMapAlpha(false)
 		, bComputeBokehAlpha(false)
 		, bReplicateRed(false)
diff --git a/Engine/Source/Runtime/Engine/Classes/Engine/Texture.h b/Engine/Source/Runtime/Engine/Classes/Engine/Texture.h
--- a/Engine/Source/Runtime/Engine/Classes/Engine/Texture.h
+++ b/Engine/Source/Runtime/Engine/Classes/Engine/Texture.h
@@ -614,6 +614,7 @@ struct FTextureFormatSettings
 	FTextureFormatSettings()
 		: CompressionSettings(TC_Default)
 		, CompressionNoAlpha(false)
+		, CompressionForceAlpha(false)
 		, CompressionNone(false)
 		, CompressionYCoCg(false)
 		, SRGB(false)
@@ -625,6 +626,9 @@ struct FTextureFormatSettings
 	UPROPERTY()
 	uint8 CompressionNoAlpha : 1;
 
+	UPROPERTY()
+	uint8 CompressionForceAlpha : 1;
+
 	UPROPERTY()
 	uint8 CompressionNone : 1;
 
diff --git a/Engine/Source/Runtime/Engine/Private/Components/RuntimeVirtualTextureComponent.cpp b/Engine/Source/Runtime/Engine/Private/Components/RuntimeVirtualTextureComponent.cpp
--- a/Engine/Source/Runtime/Engine/Private/Components/RuntimeVirtualTextureComponent.cpp
+++ b/Engine/Source/Runtime/Engine/Private/Components/RuntimeVirtualTextureComponent.cpp
@@ -224,6 +224,7 @@ void URuntimeVirtualTextureComponent::InitializeStreamingTexture(uint32 InSizeX,
 			BuildDesc.LayerFormatSettings[Layer].CompressionSettings = LayerFormat == PF_BC5 ? TC_Normalmap : TC_Default;
 			BuildDesc.LayerFormatSettings[Layer].CompressionNone = LayerFormat == PF_B8G8R8A8 || LayerFormat == PF_G16;
 			BuildDesc.LayerFormatSettings[Layer].CompressionNoAlpha = LayerFormat == PF_DXT1 || LayerFormat == PF_BC5;
+			BuildDesc.LayerFormatSettings[Layer].CompressionForceAlpha = LayerFormat == PF_DXT5;
 			BuildDesc.LayerFormatSettings[Layer].CompressionYCoCg = VirtualTexture->IsLayerYCoCg(Layer);
 			BuildDesc.LayerFormatSettings[Layer].SRGB = VirtualTexture->IsLayerSRGB(Layer);
 		}
diff --git a/Engine/Source/Runtime/Engine/Private/TextureDerivedData.cpp b/Engine/Source/Runtime/Engine/Private/TextureDerivedData.cpp
--- a/Engine/Source/Runtime/Engine/Private/TextureDerivedData.cpp
+++ b/Engine/Source/Runtime/Engine/Private/TextureDerivedData.cpp
@@ -82,6 +82,7 @@ static void SerializeForKey(FArchive& Ar, const FTextureBuildSettings& Settings)
 	uint8 TempByte;
 	FColor TempColor;
 	FVector4 TempVector4;
+	FGuid TempGuid;
 
 	TempFloat = Settings.ColorAdjustment.AdjustBrightness; Ar << TempFloat;
 	TempFloat = Settings.ColorAdjustment.AdjustBrightnessCurve; Ar << TempFloat;
@@ -158,6 +159,12 @@ static void SerializeForKey(FArchive& Ar, const FTextureBuildSettings& Settings)
 		TempFloat = Settings.Downscale; Ar << TempFloat;
 		TempByte = Settings.DownscaleOptions; Ar << TempByte;
 	}
+
+	if (Settings.bForceAlphaChannel)
+	{
+		TempGuid = FGuid(0x2C9DF7E3, 0xBC9D413B, 0xBF963C7A, 0x3F27E8B1); // Guid reserved for bForceAlphaChannel feature
+		Ar << TempGuid;
+	}
 }
 
 /**
@@ -320,6 +327,7 @@ static void FinalizeBuildSettingsForLayer(const UTexture& Texture, int32 LayerIn
 
 	OutSettings.bHDRSource = Texture.HasHDRSource(LayerIndex);
 	OutSettings.bSRGB = FormatSettings.SRGB;
+	OutSettings.bForceAlphaChannel = FormatSettings.CompressionForceAlpha;
 	OutSettings.bApplyYCoCgBlockScale = FormatSettings.CompressionYCoCg;
 
 	if (FormatSettings.CompressionSettings == TC_Displacementmap || FormatSettings.CompressionSettings == TC_DistanceFieldFont)
diff --git a/Engine/Source/Runtime/Engine/Private/VT/VirtualTextureDataBuilder.cpp b/Engine/Source/Runtime/Engine/Private/VT/VirtualTextureDataBuilder.cpp
--- a/Engine/Source/Runtime/Engine/Private/VT/VirtualTextureDataBuilder.cpp
+++ b/Engine/Source/Runtime/Engine/Private/VT/VirtualTextureDataBuilder.cpp
@@ -607,6 +607,7 @@ void FVirtualTextureDataBuilder::BuildTiles(const TArray<FVTSourceTileEntry>& Ti
 		TBSettings.bSRGB = BuildSettingsForLayer.bSRGB;
 		TBSettings.bUseLegacyGamma = BuildSettingsForLayer.bUseLegacyGamma;
 		TBSettings.MipGenSettings = TMGS_NoMipmaps;
+		TBSettings.bForceAlphaChannel = BuildSettingsForLayer.bForceAlphaChannel;
 		TBSettings.bVirtualStreamable = true;
 
 		check(TBSettings.GetGammaSpace() == BuildSettingsForLayer.GetGammaSpace());
@@ -829,7 +830,7 @@ void FVirtualTextureDataBuilder::BuildSourcePixels(const FTextureSourceData& Sou
 		FVirtualTextureSourceLayerData& LayerData = SourceLayers[LayerIndex];
 
 		LayerData.GammaSpace = BuildSettingsForLayer.GetGammaSpace();
-		LayerData.bHasAlpha = false;
+		LayerData.bHasAlpha = BuildSettingsForLayer.bForceAlphaChannel;
 
 		const FName TextureFormatName = RemovePlatformPrefixFromName(BuildSettingsForLayer.TextureFormatName);
 
@@ -894,6 +895,7 @@ void FVirtualTextureDataBuilder::BuildSourcePixels(const FTextureSourceData& Sou
 			TBSettings.TextureFormatName = LayerData.FormatName;
 			TBSettings.bSRGB = BuildSettingsForLayer.bSRGB;
 			TBSettings.bUseLegacyGamma = BuildSettingsForLayer.bUseLegacyGamma;
+			TBSettings.bForceAlphaChannel = BuildSettingsForLayer.bForceAlphaChannel;
 			TBSettings.bApplyYCoCgBlockScale = BuildSettingsForLayer.bApplyYCoCgBlockScale;
 			TBSettings.bReplicateRed = BuildSettingsForLayer.bReplicateRed;
 			TBSettings.bReplicateAlpha = BuildSettingsForLayer.bReplicateAlpha;
