Fix null deref static analysis warning.

diff --git a/Engine/Source/Runtime/Engine/Private/World.cpp b/Engine/Source/Runtime/Engine/Private/World.cpp
--- a/Engine/Source/Runtime/Engine/Private/World.cpp
+++ b/Engine/Source/Runtime/Engine/Private/World.cpp
@@ -1155,13 +1155,14 @@ void UWorld::PostLoad()
 		}
 
 		// Ensure the DefaultBrush's model has the same outer as the default brush itself. Older packages erroneously stored this object as a top-level package
-		ABrush* DefaultBrush = PersistentLevel->Actors.Num() < 2 ? NULL : Cast<ABrush>(PersistentLevel->Actors[1]);
-		UModel* Model = DefaultBrush ? DefaultBrush->Brush : NULL;
-		if (Model != NULL)
+		if (ABrush* DefaultBrush = PersistentLevel->Actors.Num() < 2 ? NULL : Cast<ABrush>(PersistentLevel->Actors[1]))
 		{
-			if (Model->GetOuter() != DefaultBrush->GetOuter())
+			if (UModel* Model = DefaultBrush->Brush)
 			{
-				Model->Rename(TEXT("Brush"), DefaultBrush->GetOuter(), REN_DoNotDirty | REN_DontCreateRedirectors | REN_ForceNoResetLoaders | REN_NonTransactional);
+				if (Model->GetOuter() != DefaultBrush->GetOuter())
+				{
+					Model->Rename(TEXT("Brush"), DefaultBrush->GetOuter(), REN_DoNotDirty | REN_DontCreateRedirectors | REN_ForceNoResetLoaders | REN_NonTransactional);
+				}
 			}
 		}
 
