Refactor FLocalVertexFactory::InitRHI() to use SetPrimitiveIdStreamIndex() instead of directly poking the array

diff --git a/Engine/Source/Runtime/Engine/Private/LocalVertexFactory.cpp b/Engine/Source/Runtime/Engine/Private/LocalVertexFactory.cpp
--- a/Engine/Source/Runtime/Engine/Private/LocalVertexFactory.cpp
+++ b/Engine/Source/Runtime/Engine/Private/LocalVertexFactory.cpp
@@ -287,13 +287,11 @@ void FLocalVertexFactory::InitRHI()
 				StreamElements.Add(AccessStreamComponent(Data.TangentBasisComponents[1], 2, InputStreamType));
 			}
 
-			const uint8 TypeIndex = static_cast<uint8>(InputStreamType);
-			PrimitiveIdStreamIndex[TypeIndex] = -1;
 			if (GetType()->SupportsPrimitiveIdStream() && bCanUseGPUScene)
 			{
 				// When the VF is used for rendering in normal mesh passes, this vertex buffer and offset will be overridden
 				StreamElements.Add(AccessStreamComponent(FVertexStreamComponent(&GPrimitiveIdDummy, 0, 0, PrimitiveIdStreamStride, VET_UInt, EVertexStreamUsage::Instancing), 1, InputStreamType));
-				PrimitiveIdStreamIndex[TypeIndex] = StreamElements.Last().StreamIndex;
+				SetPrimitiveIdStreamIndex(InputStreamType, StreamElements.Last().StreamIndex);
 			}
 
 			InitDeclaration(StreamElements, InputStreamType);
