Fix last TObjectPtr ternary operator warning from PVS static analysis.

diff --git a/Engine/Source/Runtime/Engine/Private/World.cpp b/Engine/Source/Runtime/Engine/Private/World.cpp
--- a/Engine/Source/Runtime/Engine/Private/World.cpp
+++ b/Engine/Source/Runtime/Engine/Private/World.cpp
@@ -5776,7 +5776,7 @@ void UWorld::NotifyControlMessage(UNetConnection* Connection, uint8 MessageType,
 				{
 					if (ChildIndex >= 0)
 					{
-						SwapConnection = Connection->Children.IsValidIndex(ChildIndex) ? Connection->Children[ChildIndex] : NULL;
+						SwapConnection = Connection->Children.IsValidIndex(ChildIndex) ? ToRawPtr(Connection->Children[ChildIndex]) : NULL;
 					}
 					bool bSuccess = false;
 					if (SwapConnection != NULL)
