Fix non uniform aspect ratio changing when virtual texture is built with a max resolution.

diff --git a/Engine/Source/Runtime/Engine/Private/VT/VirtualTextureDataBuilder.cpp b/Engine/Source/Runtime/Engine/Private/VT/VirtualTextureDataBuilder.cpp
--- a/Engine/Source/Runtime/Engine/Private/VT/VirtualTextureDataBuilder.cpp
+++ b/Engine/Source/Runtime/Engine/Private/VT/VirtualTextureDataBuilder.cpp
@@ -310,8 +310,14 @@ void FVirtualTextureDataBuilder::Build(const FTextureSourceData& InSourceData, c
 	}
 
 	check(SettingsPerLayer[0].MaxTextureResolution >= (uint32)TileSize);
-	BlockSizeX = FMath::Min<uint32>(BlockSizeX, SettingsPerLayer[0].MaxTextureResolution);
-	BlockSizeY = FMath::Min<uint32>(BlockSizeY, SettingsPerLayer[0].MaxTextureResolution);
+	
+	// Clamp BlockSizeX and BlockSizeY to MaxTextureResolution, but don't change aspect ratio
+	const uint32 ClampBlockSize = SettingsPerLayer[0].MaxTextureResolution;
+	if (FMath::Max<uint32>(BlockSizeX, BlockSizeY) > ClampBlockSize)
+	{
+		BlockSizeX = BlockSizeX >= BlockSizeY ? ClampBlockSize : FMath::Max(ClampBlockSize * BlockSizeX / BlockSizeY, 1u);
+		BlockSizeY = BlockSizeY >= BlockSizeX ? ClampBlockSize : FMath::Max(ClampBlockSize * BlockSizeY / BlockSizeX, 1u);
+	}
 
 	// We require VT blocks (UDIM pages) to be PoT, but multi block textures may have full logical dimension that's not PoT
 	check(FMath::IsPowerOfTwo(BlockSizeX));
