Add missing hooks to detect when the StaticMesh property is modified
Reenable the ensure to validate that we catched all the case where the property can be overwritten

diff --git a/Engine/Source/Runtime/Engine/Classes/Components/StaticMeshComponent.h b/Engine/Source/Runtime/Engine/Classes/Components/StaticMeshComponent.h
--- a/Engine/Source/Runtime/Engine/Classes/Components/StaticMeshComponent.h
+++ b/Engine/Source/Runtime/Engine/Classes/Components/StaticMeshComponent.h
@@ -383,13 +383,10 @@ public:
 	UStaticMesh* GetStaticMesh() const 
 	{ 
 #if WITH_EDITOR
+		// This should never happen and is a last resort, we should have catched the property overwrite well before we reach this code
 		if (KnownStaticMesh != StaticMesh)
 		{
-			//ensureMsgf(KnownStaticMesh == StaticMesh, TEXT("There is a missing call to NotifyIfStaticMeshChanged after StaticMesh has been overwritten"));
-
-			// This is a last resort, we should have catched the property overwrite well before we reach this code
-			UStaticMeshComponent* MutableThis = const_cast<UStaticMeshComponent*>(this);
-			MutableThis->NotifyIfStaticMeshChanged();
+			OutdatedKnownStaticMeshDetected();
 		}
 #endif
 		return StaticMesh; 
@@ -427,6 +424,8 @@ public:
 	virtual void ImportCustomProperties(const TCHAR* SourceText, FFeedbackContext* Warn) override;	
 	virtual void Serialize(FArchive& Ar) override;
 	virtual void PostInitProperties() override;
+	virtual void PostReinitProperties() override;
+	virtual void PostApplyToComponent() override;
 #if WITH_EDITOR
 	virtual void PostEditUndo() override;
 	virtual void PreEditUndo() override;
@@ -665,6 +664,9 @@ private:
 #if WITH_EDITOR
 	/** Update the vertex override colors */
 	void PrivateFixupOverrideColors();
+
+	/** Called when the StaticMesh property gets overwritten without us knowing about it */
+	void OutdatedKnownStaticMeshDetected() const;
 #endif
 protected:
 
diff --git a/Engine/Source/Runtime/Engine/Private/Components/StaticMeshComponent.cpp b/Engine/Source/Runtime/Engine/Private/Components/StaticMeshComponent.cpp
--- a/Engine/Source/Runtime/Engine/Private/Components/StaticMeshComponent.cpp
+++ b/Engine/Source/Runtime/Engine/Private/Components/StaticMeshComponent.cpp
@@ -334,8 +334,24 @@ void UStaticMeshComponent::Serialize(FArchive& Ar)
 	NotifyIfStaticMeshChanged();
 }
 
+void UStaticMeshComponent::PostApplyToComponent()
+{
+	NotifyIfStaticMeshChanged();
+
+	Super::PostApplyToComponent();
+}
+
+void UStaticMeshComponent::PostReinitProperties()
+{
+	NotifyIfStaticMeshChanged();
+
+	Super::PostReinitProperties();
+}
+
 void UStaticMeshComponent::PostInitProperties()
 {
+	NotifyIfStaticMeshChanged();
+
 	Super::PostInitProperties();
 
 	for (int32 LODIndex = 0; LODIndex < LODData.Num(); LODIndex++)
@@ -564,6 +580,24 @@ void UStaticMeshComponent::NotifyIfStaticMeshChanged()
 
 #if WITH_EDITOR
 
+void UStaticMeshComponent::OutdatedKnownStaticMeshDetected() const
+{
+	ensureMsgf(
+		KnownStaticMesh == StaticMesh, 
+		TEXT("There is a missing call to NotifyIfStaticMeshChanged for %s (%s -> %s) after StaticMesh has been overwritten"),
+		this, 
+		*GetFullName(),
+		KnownStaticMesh,
+		KnownStaticMesh ? *KnownStaticMesh->GetFullName() : TEXT("nullptr"),
+		StaticMesh,
+		StaticMesh ? *StaticMesh->GetFullName() : TEXT("nullptr")
+		);
+
+	// This is a last resort, call the notification now
+	UStaticMeshComponent* MutableThis = const_cast<UStaticMeshComponent*>(this);
+	MutableThis->NotifyIfStaticMeshChanged();
+}
+
 void UStaticMeshComponent::InitializeComponent()
 {
 	NotifyIfStaticMeshChanged();
