World Partition streaming sources
- Removed the code to filter out streaming source components if they are in the grid.
- Added the possibility to enable/disable a streaming source components through Blueprint.
- Added streaming sources debug display (can be enabled via WP.Runtime.ToggleDrawStreamingSources).

diff --git a/Engine/Source/Runtime/Engine/Classes/Components/WorldPartitionStreamingSourceComponent.h b/Engine/Source/Runtime/Engine/Classes/Components/WorldPartitionStreamingSourceComponent.h
--- a/Engine/Source/Runtime/Engine/Classes/Components/WorldPartitionStreamingSourceComponent.h
+++ b/Engine/Source/Runtime/Engine/Classes/Components/WorldPartitionStreamingSourceComponent.h
@@ -8,7 +8,7 @@
 #include "WorldPartition/WorldPartitionStreamingSource.h"
 #include "WorldPartitionStreamingSourceComponent.generated.h"
 
-UCLASS(Meta = (BlueprintSpawnableComponent))
+UCLASS(Meta = (BlueprintSpawnableComponent), HideCategories = (Tags, Sockets, ComponentTick, ComponentReplication, Activation, Cooking, Events, AssetUserData, Collision))
 class ENGINE_API UWorldPartitionStreamingSourceComponent : public UActorComponent, public IWorldPartitionStreamingSourceProvider
 {
 	GENERATED_UCLASS_BODY()
@@ -16,6 +16,23 @@ class ENGINE_API UWorldPartitionStreamingSourceComponent : public UActorComponen
 	virtual void OnRegister() override;
 	virtual void OnUnregister() override;
 
+	/** Enable the component */
+	UFUNCTION(BlueprintCallable, Category = "Streaming")
+	void EnableStreamingSource() { bStreamingSourceEnabled = true; }
+
+	/** Disable the component */
+	UFUNCTION(BlueprintCallable, Category = "Streaming")
+	void DisableStreamingSource() { bStreamingSourceEnabled = false; }
+
+	/** Returns true if the component is active. */
+	UFUNCTION(BlueprintPure, Category = "Streaming")
+	bool IsStreamingSourceEnabled() const { return bStreamingSourceEnabled; }
+
 	// IWorldPartitionStreamingSourceProvider interface
 	virtual bool GetStreamingSource(FWorldPartitionStreamingSource& OutStreamingSource) override;
+
+private:
+	/** Whether this component is enabled or not */
+	UPROPERTY(EditAnywhere, Category = "Streaming")
+	bool bStreamingSourceEnabled;
 };
\ No newline at end of file
diff --git a/Engine/Source/Runtime/Engine/Private/Components/WorldPartitionStreamingSourceComponent.cpp b/Engine/Source/Runtime/Engine/Private/Components/WorldPartitionStreamingSourceComponent.cpp
--- a/Engine/Source/Runtime/Engine/Private/Components/WorldPartitionStreamingSourceComponent.cpp
+++ b/Engine/Source/Runtime/Engine/Private/Components/WorldPartitionStreamingSourceComponent.cpp
@@ -5,6 +5,7 @@
 
 UWorldPartitionStreamingSourceComponent::UWorldPartitionStreamingSourceComponent(const FObjectInitializer& ObjectInitializer)
 	: Super(ObjectInitializer)
+	, bStreamingSourceEnabled(true)
 {
 	PrimaryComponentTick.bCanEverTick = false;
 }
@@ -24,17 +25,7 @@ void UWorldPartitionStreamingSourceComponent::OnRegister()
 
 	if (UWorldPartition* WorldPartition = World->GetWorldPartition())
 	{
-		AActor* Actor = GetOwner();
-		ULevel* ActorLevel = Actor->GetLevel();
-
-		if (ActorLevel == World->PersistentLevel)
-		{
-			WorldPartition->RegisterStreamingSourceProvider(this);
-		}
-		else
-		{
-			UE_LOG(LogWorldPartition, Warning, TEXT("Attaching a WorldPartitionStreamingSourceComponent to an actor that is in the grid will result in never unloading the affected cells (%s)"), *Actor->GetName());
-		}
+		WorldPartition->RegisterStreamingSourceProvider(this);
 	}
 }
 
@@ -53,17 +44,29 @@ void UWorldPartitionStreamingSourceComponent::OnUnregister()
 
 	if (UWorldPartition* WorldPartition = World->GetWorldPartition())
 	{
-		AActor* Actor = GetOwner();
-		ULevel* ActorLevel = Actor->GetLevel();
-
-		verify(WorldPartition->UnregisterStreamingSourceProvider(this) || (ActorLevel != World->PersistentLevel));
+		verify(WorldPartition->UnregisterStreamingSourceProvider(this));
 	}
 }
 
 bool UWorldPartitionStreamingSourceComponent::GetStreamingSource(FWorldPartitionStreamingSource& OutStreamingSource)
 {
-	AActor* Actor = GetOwner();
-	OutStreamingSource.Location = Actor->GetActorLocation();
-	OutStreamingSource.Rotation = Actor->GetActorRotation();
-	return true;
+	if (bStreamingSourceEnabled)
+	{
+		AActor* Actor = GetOwner();
+
+		FName SourceName = Actor->GetFName();
+
+#if WITH_EDITOR
+		const FString& ActorLabel = Actor->GetActorLabel(false);
+		if (!ActorLabel.IsEmpty())
+		{
+			SourceName = *ActorLabel;
+		}
+#endif
+		OutStreamingSource.Name = SourceName;
+		OutStreamingSource.Location = Actor->GetActorLocation();
+		OutStreamingSource.Rotation = Actor->GetActorRotation();
+		return true;
+	}
+	return false;
 }
\ No newline at end of file
diff --git a/Engine/Source/Runtime/Engine/Private/WorldPartition/WorldPartition.cpp b/Engine/Source/Runtime/Engine/Private/WorldPartition/WorldPartition.cpp
--- a/Engine/Source/Runtime/Engine/Private/WorldPartition/WorldPartition.cpp
+++ b/Engine/Source/Runtime/Engine/Private/WorldPartition/WorldPartition.cpp
@@ -414,7 +414,20 @@ void UWorldPartition::ForEachActorDesc(TSubclassOf<AActor> ActorClass, TFunction
 		}
 	}
 }
+#endif
 
+const TArray<FWorldPartitionStreamingSource>& UWorldPartition::GetStreamingSources() const
+{
+	if (GetWorld()->IsGameWorld())
+	{
+		return GetStreamingPolicy()->GetStreamingSources();
+	}
+
+	static TArray<FWorldPartitionStreamingSource> EmptyStreamingSources;
+	return EmptyStreamingSources;
+}
+
+#if WITH_EDITOR
 const FWorldPartitionActorDesc* UWorldPartition::GetActorDesc(const FGuid& Guid) const
 {
 	const TUniquePtr<FWorldPartitionActorDesc>* const * ActorDesc = Actors.Find(Guid);
diff --git a/Engine/Source/Runtime/Engine/Private/WorldPartition/WorldPartitionStreamingPolicy.cpp b/Engine/Source/Runtime/Engine/Private/WorldPartition/WorldPartitionStreamingPolicy.cpp
--- a/Engine/Source/Runtime/Engine/Private/WorldPartition/WorldPartitionStreamingPolicy.cpp
+++ b/Engine/Source/Runtime/Engine/Private/WorldPartition/WorldPartitionStreamingPolicy.cpp
@@ -54,7 +54,8 @@ void UWorldPartitionStreamingPolicy::UpdateStreamingSources()
 		const FRotator ViewRotation = GCurrentLevelEditingViewportClient->GetViewRotation();
 		const FVector ViewLocationLocal = WorldToLocal.TransformPosition(ViewLocation);
 		const FRotator ViewRotationLocal = WorldToLocal.TransformRotation(ViewRotation.Quaternion()).Rotator();
-		StreamingSources.Add(FWorldPartitionStreamingSource(ViewLocationLocal, ViewRotationLocal));
+		static const FName NAME_SIE(TEXT("SIE"));
+		StreamingSources.Add(FWorldPartitionStreamingSource(NAME_SIE, ViewLocationLocal, ViewRotationLocal));
 		return;
 	}
 #endif
@@ -79,7 +80,7 @@ void UWorldPartitionStreamingPolicy::UpdateStreamingSources()
 					// Transform to Local
 					ViewLocation = WorldToLocal.TransformPosition(ViewLocation);
 					ViewRotation = WorldToLocal.TransformRotation(ViewRotation.Quaternion()).Rotator();
-					StreamingSources.Add(FWorldPartitionStreamingSource(ViewLocation, ViewRotation));
+					StreamingSources.Add(FWorldPartitionStreamingSource(Player->GetFName(), ViewLocation, ViewRotation));
 				}
 			}
 		}
diff --git a/Engine/Source/Runtime/Engine/Private/WorldPartition/WorldPartitionSubsystem.cpp b/Engine/Source/Runtime/Engine/Private/WorldPartition/WorldPartitionSubsystem.cpp
--- a/Engine/Source/Runtime/Engine/Private/WorldPartition/WorldPartitionSubsystem.cpp
+++ b/Engine/Source/Runtime/Engine/Private/WorldPartition/WorldPartitionSubsystem.cpp
@@ -22,6 +22,12 @@ static FAutoConsoleCommand CVarDrawRuntimeHash2D(
 	TEXT("Toggles 2D debug display of world partition runtime hash."),
 	FConsoleCommandDelegate::CreateLambda([] { GDrawRuntimeHash2D = !GDrawRuntimeHash2D; }));
 
+static int32 GDrawStreamingSources = 0;
+static FAutoConsoleCommand CVarDrawStreamingSources(
+	TEXT("wp.Runtime.ToggleDrawStreamingSources"),
+	TEXT("Toggles debug display of world partition streaming sources."),
+	FConsoleCommandDelegate::CreateLambda([] { GDrawStreamingSources = !GDrawStreamingSources; }));
+
 UWorldPartitionSubsystem::UWorldPartitionSubsystem()
 {}
 
@@ -60,17 +66,17 @@ void UWorldPartitionSubsystem::PostInitialize()
 
 		if (GetWorld()->IsGameWorld() && (GetWorld()->GetNetMode() != NM_DedicatedServer))
 		{
-			DrawRuntimeHash2DHandle = UDebugDrawService::Register(TEXT("Game"), FDebugDrawDelegate::CreateUObject(this, &UWorldPartitionSubsystem::DrawRuntimeHash2D));
+			DrawHandle = UDebugDrawService::Register(TEXT("Game"), FDebugDrawDelegate::CreateUObject(this, &UWorldPartitionSubsystem::Draw));
 		}
 	}
 }
 
 void UWorldPartitionSubsystem::Deinitialize()
 {
-	if (DrawRuntimeHash2DHandle.IsValid())
+	if (DrawHandle.IsValid())
 	{
-		UDebugDrawService::Unregister(DrawRuntimeHash2DHandle);
-		DrawRuntimeHash2DHandle.Reset();
+		UDebugDrawService::Unregister(DrawHandle);
+		DrawHandle.Reset();
 	}
 
 	while (RegisteredWorldPartitions.Num() > 0)
@@ -152,36 +158,67 @@ void UWorldPartitionSubsystem::UpdateStreamingState()
 	}
 }
 
-void UWorldPartitionSubsystem::DrawRuntimeHash2D(UCanvas* Canvas, class APlayerController* PC)
+void UWorldPartitionSubsystem::Draw(UCanvas* Canvas, class APlayerController* PC)
 {
-	if (!GDrawRuntimeHash2D || !Canvas || !Canvas->SceneView)
+	if (!Canvas || !Canvas->SceneView)
 	{
 		return;
 	}
 
-	const float MaxScreenRatio = 0.75f;
 	const FVector2D CanvasTopLeftPadding(10.f, 10.f);
-	const FVector2D CanvasBottomRightPadding(10.f, 10.f);
-	const FVector2D CanvasMinimumSize(100.f, 100.f);
-	const FVector2D CanvasMaxScreenSize = FVector2D::Max(MaxScreenRatio*FVector2D(Canvas->ClipX, Canvas->ClipY) - CanvasBottomRightPadding - CanvasTopLeftPadding, CanvasMinimumSize);
+	FVector2D Pos = CanvasTopLeftPadding;
 
-	FVector2D TotalFootprint(ForceInitToZero);
-	for (UWorldPartition* Partition : RegisteredWorldPartitions)
+	auto DrawText = [](UCanvas* Canvas, const FString& Text, UFont* Font, const FColor& Color, FVector2D& Pos)
 	{
-		FVector2D Footprint = Partition->GetDrawRuntimeHash2DDesiredFootprint(CanvasMaxScreenSize);
-		TotalFootprint.X += Footprint.X;
-	}
-
-	if (TotalFootprint.X > 0.f)
+		float TextWidth, TextHeight;
+		Canvas->StrLen(Font, Text, TextWidth, TextHeight);
+		Canvas->SetDrawColor(Color);
+		Canvas->DrawText(Font, Text, Pos.X, Pos.Y);
+		Pos.Y += TextHeight + 1;
+	};
+
+	if (GDrawRuntimeHash2D)
 	{
-		FVector2D PartitionCanvasOffset(CanvasTopLeftPadding);
+		const float MaxScreenRatio = 0.75f;
+		const FVector2D CanvasBottomRightPadding(10.f, 10.f);
+		const FVector2D CanvasMinimumSize(100.f, 100.f);
+		const FVector2D CanvasMaxScreenSize = FVector2D::Max(MaxScreenRatio*FVector2D(Canvas->ClipX, Canvas->ClipY) - CanvasBottomRightPadding - CanvasTopLeftPadding, CanvasMinimumSize);
+
+		FVector2D TotalFootprint(ForceInitToZero);
 		for (UWorldPartition* Partition : RegisteredWorldPartitions)
 		{
 			FVector2D Footprint = Partition->GetDrawRuntimeHash2DDesiredFootprint(CanvasMaxScreenSize);
-			float FootprintRatio = Footprint.X / TotalFootprint.X;
-			FVector2D PartitionCanvasSize = FVector2D(CanvasMaxScreenSize.X * FootprintRatio, CanvasMaxScreenSize.Y);
-			Partition->DrawRuntimeHash2D(Canvas, PartitionCanvasOffset, PartitionCanvasSize);
-			PartitionCanvasOffset.X += PartitionCanvasSize.X;
+			TotalFootprint.X += Footprint.X;
+		}
+
+		if (TotalFootprint.X > 0.f)
+		{
+			FVector2D PartitionCanvasOffset(CanvasTopLeftPadding);
+			for (UWorldPartition* Partition : RegisteredWorldPartitions)
+			{
+				FVector2D Footprint = Partition->GetDrawRuntimeHash2DDesiredFootprint(CanvasMaxScreenSize);
+				float FootprintRatio = Footprint.X / TotalFootprint.X;
+				FVector2D PartitionCanvasSize = FVector2D(CanvasMaxScreenSize.X * FootprintRatio, CanvasMaxScreenSize.Y);
+				Partition->DrawRuntimeHash2D(Canvas, PartitionCanvasOffset, PartitionCanvasSize);
+				PartitionCanvasOffset.X += PartitionCanvasSize.X;
+			}
+		}
+	}
+
+	if (GDrawStreamingSources)
+	{
+		if (const UWorldPartition* WorldPartition = GetMainWorldPartition())
+		{
+			const TArray<FWorldPartitionStreamingSource>& StreamingSources = WorldPartition->GetStreamingSources();
+
+			FString Title(TEXT("Streaming Sources"));
+			DrawText(Canvas, Title, GEngine->GetSmallFont(), FColor::Yellow, Pos);
+			UFont* Font = GEngine->GetTinyFont();
+			for (const FWorldPartitionStreamingSource& StreamingSource : StreamingSources)
+			{
+				const FString Text = FString::Printf(TEXT("%s %s %s"), *StreamingSource.Name.ToString(), *StreamingSource.Location.ToString(), *StreamingSource.Rotation.ToString());
+				DrawText(Canvas, Text, Font, FColor::White, Pos);
+			}
 		}
 	}
 }
diff --git a/Engine/Source/Runtime/Engine/Public/WorldPartition/WorldPartition.h b/Engine/Source/Runtime/Engine/Public/WorldPartition/WorldPartition.h
--- a/Engine/Source/Runtime/Engine/Public/WorldPartition/WorldPartition.h
+++ b/Engine/Source/Runtime/Engine/Public/WorldPartition/WorldPartition.h
@@ -103,12 +103,16 @@ private:
 
 	static void WorldPartitionOnLevelRemovedFromWorld(class ULevel* Level, UWorld* InWorld);
 
-#if WITH_EDITOR
 	// UWorldPartitionSubsystem interface+
 	friend class UWorldPartitionSubsystem;
+#if WITH_EDITOR
 	void ForEachIntersectingActorDesc(const FBox& Box, TSubclassOf<AActor> ActorClass, TFunctionRef<bool(const FWorldPartitionActorDesc*)> Predicate) const;
 	void ForEachActorDesc(TSubclassOf<AActor> ActorClass, TFunctionRef<bool(const FWorldPartitionActorDesc*)> Predicate) const;
+#endif
+	const TArray<FWorldPartitionStreamingSource>& GetStreamingSources() const;
 	// UActorPartitionSubsystem interface-
+
+#if WITH_EDITOR
 public:
 	FName GetWorldPartitionEditorName();
 	bool IsSimulating() const;
diff --git a/Engine/Source/Runtime/Engine/Public/WorldPartition/WorldPartitionStreamingPolicy.h b/Engine/Source/Runtime/Engine/Public/WorldPartition/WorldPartitionStreamingPolicy.h
--- a/Engine/Source/Runtime/Engine/Public/WorldPartition/WorldPartitionStreamingPolicy.h
+++ b/Engine/Source/Runtime/Engine/Public/WorldPartition/WorldPartitionStreamingPolicy.h
@@ -38,6 +38,8 @@ public:
 	virtual void OnPreFixupForPIE(int32 InPIEInstanceID, FSoftObjectPath& ObjectPath) {}
 #endif
 
+	const TArray<FWorldPartitionStreamingSource>& GetStreamingSources() const { return StreamingSources; }
+
 protected:
 	void UpdateStreamingSources();
 
diff --git a/Engine/Source/Runtime/Engine/Public/WorldPartition/WorldPartitionStreamingSource.h b/Engine/Source/Runtime/Engine/Public/WorldPartition/WorldPartitionStreamingSource.h
--- a/Engine/Source/Runtime/Engine/Public/WorldPartition/WorldPartitionStreamingSource.h
+++ b/Engine/Source/Runtime/Engine/Public/WorldPartition/WorldPartitionStreamingSource.h
@@ -11,11 +11,13 @@ struct ENGINE_API FWorldPartitionStreamingSource
 	FWorldPartitionStreamingSource()
 	{}
 
-	FWorldPartitionStreamingSource(const FVector& InLocation, const FRotator& InRotation)
-		: Location(InLocation)
+	FWorldPartitionStreamingSource(FName InName, const FVector& InLocation, const FRotator& InRotation)
+		: Name(InName)
+		, Location(InLocation)
 		, Rotation(InRotation)
 	{}
 
+	FName Name;
 	FVector Location;
 	FRotator Rotation;
 };
diff --git a/Engine/Source/Runtime/Engine/Public/WorldPartition/WorldPartitionSubsystem.h b/Engine/Source/Runtime/Engine/Public/WorldPartition/WorldPartitionSubsystem.h
--- a/Engine/Source/Runtime/Engine/Public/WorldPartition/WorldPartitionSubsystem.h
+++ b/Engine/Source/Runtime/Engine/Public/WorldPartition/WorldPartitionSubsystem.h
@@ -52,11 +52,11 @@ private:
 	const UWorldPartition* GetMainWorldPartition() const;
 	void RegisterWorldPartition(UWorldPartition* WorldPartition);
 	void UnregisterWorldPartition(UWorldPartition* WorldPartition);
-	void DrawRuntimeHash2D(class UCanvas* Canvas, class APlayerController* PC);
+	void Draw(class UCanvas* Canvas, class APlayerController* PC);
 	friend class UWorldPartition;
 
 	UPROPERTY()
 	TArray<TObjectPtr<UWorldPartition>> RegisteredWorldPartitions;
 
-	FDelegateHandle	DrawRuntimeHash2DHandle;
+	FDelegateHandle	DrawHandle;
 };
