LevelInstance: optional out param when breaking level instance for list of moved actors

diff --git a/Engine/Source/Runtime/Engine/Private/LevelInstance/LevelInstanceSubsystem.cpp b/Engine/Source/Runtime/Engine/Private/LevelInstance/LevelInstanceSubsystem.cpp
--- a/Engine/Source/Runtime/Engine/Private/LevelInstance/LevelInstanceSubsystem.cpp
+++ b/Engine/Source/Runtime/Engine/Private/LevelInstance/LevelInstanceSubsystem.cpp
@@ -778,14 +778,35 @@ ALevelInstance* ULevelInstanceSubsystem::CreateLevelInstanceFrom(const TArray<AA
 	return CommitLevelInstance(LevelStreaming->GetLevelInstanceActor());
 }
 
-bool ULevelInstanceSubsystem::BreakLevelInstance(ALevelInstance* LevelInstanceActor, uint32 Levels)
+bool ULevelInstanceSubsystem::BreakLevelInstance(ALevelInstance* LevelInstanceActor, uint32 Levels /* = 1 */, TArray<AActor*>* OutMovedActors /* = nullptr */)
+{
+	TArray<AActor*> MovedActors;
+	BreakLevelInstance_Impl(LevelInstanceActor, Levels, MovedActors);
+
+	USelection* ActorSelection = GEditor->GetSelectedActors();
+	ActorSelection->BeginBatchSelectOperation();
+	for (AActor* MovedActor : MovedActors)
+	{
+		GEditor->SelectActor(MovedActor, true, false);
+	}
+	ActorSelection->EndBatchSelectOperation(false);
+
+	bool bStatus = MovedActors.Num() > 0;
+	if (OutMovedActors)
+	{
+		*OutMovedActors = MoveTemp(MovedActors);
+	}
+
+	return bStatus;
+}
+
+void ULevelInstanceSubsystem::BreakLevelInstance_Impl(ALevelInstance* LevelInstanceActor, uint32 Levels, TArray<AActor*>& OutMovedActors)
 {
 	if (Levels > 0)
 	{
 		// Can only break the top level LevelInstance
 		check(LevelInstanceActor->GetLevel() == GetWorld()->GetCurrentLevel());
 
-
 		// Actors in a packed level instance will not be streamed in unless they are editing. Must force this before moving.
 		if (LevelInstanceActor->IsA<APackedLevelInstance>())
 		{
@@ -804,7 +825,7 @@ bool ULevelInstanceSubsystem::BreakLevelInstance(ALevelInstance* LevelInstanceAc
 			{
 				BlockUnloadLevelInstance(LevelInstanceActor);
 			}
-			return false;
+			return;
 		}
 
 		TArray<AActor*> ActorsToMove;
@@ -833,7 +854,7 @@ bool ULevelInstanceSubsystem::BreakLevelInstance(ALevelInstance* LevelInstanceAc
 		if (!EditorLevelUtils::CopyActorsToLevel(ActorsToMove, DestinationLevel, bWarnAboutReferences, bWarnAboutRenaming, bMoveAllOrFail))
 		{
 			UE_LOG(LogLevelInstance, Warning, TEXT("Failed to break Level Instance because not all actors could be moved"));
-			return false;
+			return;
 		}
 
 		if (LevelInstanceActor->IsA<APackedLevelInstance>())
@@ -844,27 +865,37 @@ bool ULevelInstanceSubsystem::BreakLevelInstance(ALevelInstance* LevelInstanceAc
 		// Destroy the old LevelInstance instance actor
 		GetWorld()->DestroyActor(LevelInstanceActor);
 	
-		// Break up any sub LevelInstances if more levels are requested
-		if (Levels > 1)
+		const bool bContinueBreak = Levels > 1;
+		TArray<ALevelInstance*> Children;
+
+		for (FSelectionIterator It(GEditor->GetSelectedActorIterator()); It; ++It)
 		{
-			TArray<ALevelInstance*> Children;
-			for (FSelectionIterator It(GEditor->GetSelectedActorIterator()); It; ++It)
+			AActor* Actor = Cast<AActor>(*It);
+
+			if (Actor)
+			{
+				OutMovedActors.Add(Actor);
+			}
+
+			// Break up any sub LevelInstances if more levels are requested
+			if (bContinueBreak)
 			{
-				AActor* Actor = Cast<AActor>(*It);
 				if (ALevelInstance* ChildLevelInstance = Cast<ALevelInstance>(Actor))
 				{
+					OutMovedActors.Remove(ChildLevelInstance);
+
 					Children.Add(ChildLevelInstance);
 				}
 			}
+		}
 
-			for (auto& Child : Children)
-			{
-				BreakLevelInstance(Child, Levels - 1);
-			}
+		for (auto& Child : Children)
+		{
+			BreakLevelInstance_Impl(Child, Levels - 1, OutMovedActors);
 		}
 	}
 
-	return true;
+	return;
 }
 
 ULevel* ULevelInstanceSubsystem::GetLevelInstanceLevel(const ALevelInstance* LevelInstanceActor) const
diff --git a/Engine/Source/Runtime/Engine/Public/LevelInstance/LevelInstanceSubsystem.h b/Engine/Source/Runtime/Engine/Public/LevelInstance/LevelInstanceSubsystem.h
--- a/Engine/Source/Runtime/Engine/Public/LevelInstance/LevelInstanceSubsystem.h
+++ b/Engine/Source/Runtime/Engine/Public/LevelInstance/LevelInstanceSubsystem.h
@@ -76,7 +76,7 @@ public:
 	ALevelInstance* CreateLevelInstanceFrom(const TArray<AActor*>& ActorsToMove, const FNewLevelInstanceParams& CreationParams);
 	bool MoveActorsToLevel(const TArray<AActor*>& ActorsToRemove, ULevel* DestinationLevel) const;
 	bool MoveActorsTo(ALevelInstance* LevelInstanceActor, const TArray<AActor*>& ActorsToMove);
-	bool BreakLevelInstance(ALevelInstance* LevelInstanceActor, uint32 Levels = 1);
+	bool BreakLevelInstance(ALevelInstance* LevelInstanceActor, uint32 Levels = 1, TArray<AActor*>* OutMovedActors = nullptr);
 
 	bool CanMoveActorToLevel(const AActor* Actor, FText* OutReason = nullptr) const;
 	void DiscardEdits();
@@ -99,6 +99,7 @@ private:
 	FLevelInstanceID ComputeLevelInstanceID(ALevelInstance* LevelInstanceActor) const;
 #if WITH_EDITOR
 	void CommitChildrenLevelInstances(ALevelInstance* LevelInstanceActor);
+	void BreakLevelInstance_Impl(ALevelInstance* LevelInstanceActor, uint32 Levels, TArray<AActor*>& OutMovedActors);
 
 	static bool ShouldIgnoreDirtyPackage(UPackage* DirtyPackage, const UWorld* EditingWorld);
 
