LevelInstance: Fixed crash when creating new world asset in the level instance "World Asset" property

diff --git a/Engine/Source/Runtime/Engine/Public/LevelInstance/LevelInstanceActor.h b/Engine/Source/Runtime/Engine/Public/LevelInstance/LevelInstanceActor.h
--- a/Engine/Source/Runtime/Engine/Public/LevelInstance/LevelInstanceActor.h
+++ b/Engine/Source/Runtime/Engine/Public/LevelInstance/LevelInstanceActor.h
@@ -16,7 +16,7 @@ class ENGINE_API ALevelInstance : public AActor
 
 protected:
 	/** Level LevelInstance */
-	UPROPERTY(EditAnywhere, Category = LevelInstance)
+	UPROPERTY(EditAnywhere, Category = LevelInstance, meta = (NoCreate))
 	TSoftObjectPtr<UWorld> WorldAsset;
 public:
 
