Moved conversion of legacy static mesh assets from PostLoad to Serialize, to prevent creation of the UMeshDescriptionBulkDataWrapper object on a worker thread.

diff --git a/Engine/Source/Runtime/Engine/Private/StaticMesh.cpp b/Engine/Source/Runtime/Engine/Private/StaticMesh.cpp
--- a/Engine/Source/Runtime/Engine/Private/StaticMesh.cpp
+++ b/Engine/Source/Runtime/Engine/Private/StaticMesh.cpp
@@ -4966,9 +4966,6 @@ void UStaticMesh::CacheDerivedData()
 {
 	TRACE_CPUPROFILER_EVENT_SCOPE(UStaticMesh::CacheDerivedData);
 
-#if WITH_EDITORONLY_DATA
-	CacheMeshData();
-#endif
 	// Cache derived data for the running platform.
 	ITargetPlatformManagerModule& TargetPlatformManager = GetTargetPlatformManagerRef();
 	ITargetPlatform* RunningPlatform = TargetPlatformManager.GetRunningTargetPlatform();
@@ -5205,6 +5202,12 @@ void UStaticMesh::Serialize(FArchive& Ar)
 			SrcModel.SerializeBulkData(Ar, this);
 		}
 
+		// If loading, perform immediate conversion from legacy RawMesh to MeshDescription if necessary
+		if (Ar.IsLoading())
+		{
+			CacheMeshData();
+		}
+
 		if (Ar.CustomVer(FEditorObjectVersion::GUID) < FEditorObjectVersion::UPropertryForMeshSection)
 		{
 			GetSectionInfoMap().Serialize(Ar);
@@ -5513,14 +5516,10 @@ void UStaticMesh::BeginPostLoadInternal(FStaticMeshPostLoadContext& Context)
 
 			for (int32 i = 0; i < GetNumSourceModels(); ++i)
 			{
-				// Access RawMesh directly instead of through the FStaticMeshSourceModel API,
-				// because we don't want to perform an automatic conversion to MeshDescription at this point -
-				// this will be done below in CacheDerivedData().
-				// This is a path for legacy assets.
-				if (!GetSourceModel(i).RawMeshBulkData->IsEmpty())
+				FMeshDescription MeshDescription;
+				if (GetSourceModel(i).CloneMeshDescription(MeshDescription))
 				{
-					GetSourceModel(i).RawMeshBulkData->LoadRawMesh(TempRawMesh);
-					TotalIndexCount += TempRawMesh.WedgeIndices.Num();
+					TotalIndexCount += MeshDescription.VertexInstances().Num();
 				}
 			}
 
