Strata material roughness conversion fixup after edge color broke it.

diff --git a/Engine/Source/Runtime/Engine/Private/Materials/Material.cpp b/Engine/Source/Runtime/Engine/Private/Materials/Material.cpp
--- a/Engine/Source/Runtime/Engine/Private/Materials/Material.cpp
+++ b/Engine/Source/Runtime/Engine/Private/Materials/Material.cpp
@@ -3673,7 +3673,7 @@ void UMaterial::ConvertMaterialToStrataMaterial()
 			}
 			else
 			{
-				MoveConnectionTo(Roughness, SlabBSDF, 3);		// RoughnessX
+				MoveConnectionTo(Roughness, SlabBSDF, 4);	// RoughnessX
 			}
 
 			MoveConnectionTo(BaseColor, SlabBSDF, 0);		// BaseColor
