HISM
- Don't call BuildTreeIfOutdated for reference collector archives.
- Call ApplyEmpty when the component has no static mesh set.

diff --git a/Engine/Source/Runtime/Engine/Private/HierarchicalInstancedStaticMesh.cpp b/Engine/Source/Runtime/Engine/Private/HierarchicalInstancedStaticMesh.cpp
--- a/Engine/Source/Runtime/Engine/Private/HierarchicalInstancedStaticMesh.cpp
+++ b/Engine/Source/Runtime/Engine/Private/HierarchicalInstancedStaticMesh.cpp
@@ -2132,7 +2132,7 @@ void UHierarchicalInstancedStaticMeshComponent::Serialize(FArchive& Ar)
 
 	// If we are saving make sure the tree is up to date. For Undo/Redo the PostEditUndo will rebuild the tree.
 	// Properly building the tree here will avoid the need to call BuildTreeIfOutdated in different duplication/load use cases.
-	if (Ar.IsSaving() && !Ar.IsTransacting())
+	if (Ar.IsSaving() && !Ar.IsTransacting() && !Ar.IsObjectReferenceCollector())
 	{
 		BuildTreeIfOutdated(/*Async*/false, /*ForceUpdate*/false);
 	}
@@ -2934,10 +2934,15 @@ bool UHierarchicalInstancedStaticMeshComponent::BuildTreeIfOutdated(bool Async,
 		|| GetLinkerUE4Version() < VER_UE4_REBUILD_HIERARCHICAL_INSTANCE_TREES
 		|| GetLinkerCustomVersion(FReleaseObjectVersion::GUID) < FReleaseObjectVersion::HISMCClusterTreeMigration)
 	{
-		// Make sure if any of those conditions is true, we mark ourselves out of date so the Async Build completes
-		bIsOutOfDate = true;
-		if (GetStaticMesh() != nullptr && !GetStaticMesh()->HasAnyFlags(RF_NeedLoad)) // we can build the tree if the static mesh is not even loaded, and we can't call PostLoad as the load is not even done
+		if (!GetStaticMesh())
+		{
+			ApplyEmpty();
+		}
+		else if (!GetStaticMesh()->HasAnyFlags(RF_NeedLoad)) // we can build the tree if the static mesh is not even loaded, and we can't call PostLoad as the load is not even done
 		{
+			// Make sure if any of those conditions is true, we mark ourselves out of date so the Async Build completes
+			bIsOutOfDate = true;
+
 			GetStaticMesh()->ConditionalPostLoad();
 
 			if (Async)
