Made AActor::EditorGetUnderlyingActors() const

diff --git a/Engine/Source/Runtime/Engine/Classes/GameFramework/Actor.h b/Engine/Source/Runtime/Engine/Classes/GameFramework/Actor.h
--- a/Engine/Source/Runtime/Engine/Classes/GameFramework/Actor.h
+++ b/Engine/Source/Runtime/Engine/Classes/GameFramework/Actor.h
@@ -2051,7 +2051,7 @@ public:
 	virtual void EditorApplyMirror(const FVector& MirrorScale, const FVector& PivotLocation);	
 
 	/** Get underlying actors */
-	virtual void EditorGetUnderlyingActors(TSet<AActor*>& OutUnderlyingActors);
+	virtual void EditorGetUnderlyingActors(TSet<AActor*>& OutUnderlyingActors) const;
 
 	/** Returns true if the actor is hidden upon editor startup/by default, false if it is not */
 	UFUNCTION(BlueprintCallable, Category = "Editor Scripting | Actor Editing")
diff --git a/Engine/Source/Runtime/Engine/Private/ActorEditor.cpp b/Engine/Source/Runtime/Engine/Private/ActorEditor.cpp
--- a/Engine/Source/Runtime/Engine/Private/ActorEditor.cpp
+++ b/Engine/Source/Runtime/Engine/Private/ActorEditor.cpp
@@ -772,7 +772,7 @@ void AActor::EditorApplyMirror(const FVector& MirrorScale, const FVector& PivotL
 	}
 }
 
-void AActor::EditorGetUnderlyingActors(TSet<AActor*>& OutUnderlyingActors)
+void AActor::EditorGetUnderlyingActors(TSet<AActor*>& OutUnderlyingActors) const
 {
 	TInlineComponentArray<UChildActorComponent*> ChildActorComponents;
 	GetComponents(ChildActorComponents);
diff --git a/Engine/Source/Runtime/Engine/Private/LevelInstance/LevelInstanceActor.cpp b/Engine/Source/Runtime/Engine/Private/LevelInstance/LevelInstanceActor.cpp
--- a/Engine/Source/Runtime/Engine/Private/LevelInstance/LevelInstanceActor.cpp
+++ b/Engine/Source/Runtime/Engine/Private/LevelInstance/LevelInstanceActor.cpp
@@ -488,7 +488,7 @@ void ALevelInstance::SetIsTemporarilyHiddenInEditor(bool bIsHidden)
 	}
 }
 
-void ALevelInstance::EditorGetUnderlyingActors(TSet<AActor*>& OutUnderlyingActors)
+void ALevelInstance::EditorGetUnderlyingActors(TSet<AActor*>& OutUnderlyingActors) const
 {
 	Super::EditorGetUnderlyingActors(OutUnderlyingActors);
 
diff --git a/Engine/Source/Runtime/Engine/Public/LevelInstance/LevelInstanceActor.h b/Engine/Source/Runtime/Engine/Public/LevelInstance/LevelInstanceActor.h
--- a/Engine/Source/Runtime/Engine/Public/LevelInstance/LevelInstanceActor.h
+++ b/Engine/Source/Runtime/Engine/Public/LevelInstance/LevelInstanceActor.h
@@ -51,7 +51,7 @@ public:
 	virtual TUniquePtr<class FWorldPartitionActorDesc> CreateClassActorDesc() const override;
 	virtual bool CanDeleteSelectedActor(FText& OutReason) const override;
 	virtual void SetIsTemporarilyHiddenInEditor(bool bIsHidden) override;
-	virtual void EditorGetUnderlyingActors(TSet<AActor*>& OutUnderlyingActors) override;
+	virtual void EditorGetUnderlyingActors(TSet<AActor*>& OutUnderlyingActors) const override;
 	virtual void PushSelectionToProxies() override;
 	virtual void PushLevelInstanceEditingStateToProxies(bool bInEditingState) override;
 	virtual FBox GetComponentsBoundingBox(bool bNonColliding = false, bool bIncludeFromChildActors = false) const override;
