New Draw Circle Arc Helper
#rnx

diff --git a/Engine/Source/Runtime/Engine/Private/DrawDebugHelpers.cpp b/Engine/Source/Runtime/Engine/Private/DrawDebugHelpers.cpp
--- a/Engine/Source/Runtime/Engine/Private/DrawDebugHelpers.cpp
+++ b/Engine/Source/Runtime/Engine/Private/DrawDebugHelpers.cpp
@@ -479,6 +479,44 @@ void DrawDebugCircle(const UWorld* InWorld, FVector Center, float Radius, int32
 	);
 }
 
+void DrawDebugCircleArc(const UWorld* InWorld, const FVector& Center, float Radius, const FVector& Direction, float AngleWidth, int32 Segments, const FColor& Color, bool PersistentLines, float LifeTime, uint8 DepthPriority, float Thickness)
+{
+	if (GEngine->GetNetMode(InWorld) != NM_DedicatedServer)
+	{
+		if (ULineBatchComponent* const LineBatcher = GetDebugLineBatcher(InWorld, PersistentLines, LifeTime, (DepthPriority == SDPG_Foreground)))
+		{
+			const float LineLifeTime = GetDebugLineLifeTime(LineBatcher, LifeTime, PersistentLines);
+
+			// Need at least 4 segments
+			Segments = FMath::Max(Segments, 4);
+			const float AngleStep = AngleWidth / float(Segments) * 2.f;
+
+			FVector AxisY, AxisZ;
+			FVector DirectionNorm = Direction.GetSafeNormal();
+			DirectionNorm.FindBestAxisVectors(AxisZ, AxisY);
+
+			TArray<FBatchedLine> Lines;
+			Lines.Empty(Segments);
+			float Angle = -AngleWidth;
+			FVector PrevVertex = Center + Radius * (AxisY * -FMath::Sin(Angle) + DirectionNorm * FMath::Cos(Angle));
+			while (Segments--)
+			{
+				Angle += AngleStep;
+				FVector NextVertex = Center + Radius * (AxisY * -FMath::Sin(Angle) + DirectionNorm * FMath::Cos(Angle));
+				Lines.Emplace(FBatchedLine(PrevVertex, NextVertex, Color, LineLifeTime, Thickness, DepthPriority));
+				PrevVertex = NextVertex;
+			}
+
+			LineBatcher->DrawLines(Lines);
+		}
+	}
+	else
+	{
+		UE_DRAW_SERVER_DEBUG_ON_EACH_CLIENT(DrawDebugCircleArc, Center, Radius, Direction, AngleWidth, Segments, AdjustColorForServer(Color), PersistentLines, LifeTime, DepthPriority, Thickness);
+	}
+
+}
+
 void DrawDebug2DDonut(const UWorld* InWorld, const FMatrix& TransformMatrix, float InnerRadius, float OuterRadius, int32 Segments, const FColor& Color, bool bPersistentLines, float LifeTime, uint8 DepthPriority, float Thickness)
 {
 	// no debug line drawing on dedicated server
diff --git a/Engine/Source/Runtime/Engine/Public/DrawDebugHelpers.h b/Engine/Source/Runtime/Engine/Public/DrawDebugHelpers.h
--- a/Engine/Source/Runtime/Engine/Public/DrawDebugHelpers.h
+++ b/Engine/Source/Runtime/Engine/Public/DrawDebugHelpers.h
@@ -34,6 +34,8 @@ ENGINE_API void DrawDebugCrosshairs(const UWorld* InWorld, FVector const& AxisLo
 ENGINE_API void DrawDebugCircle(const UWorld* InWorld, const FMatrix& TransformMatrix, float Radius, int32 Segments, const FColor& Color, bool bPersistentLines = false, float LifeTime=-1.f, uint8 DepthPriority = 0, float Thickness = 0.f, bool bDrawAxis=true);
 /** Draw Debug Circle */
 ENGINE_API void DrawDebugCircle(const UWorld* InWorld, FVector Center, float Radius, int32 Segments, const FColor& Color, bool bPersistentLines = false, float LifeTime=-1.f, uint8 DepthPriority = 0, float Thickness = 0.f, FVector YAxis=FVector(0.f,1.f,0.f), FVector ZAxis=FVector(0.f,0.f,1.f), bool bDrawAxis=true);
+/** Draw Debug Circle Arc */
+ENGINE_API void DrawDebugCircleArc(const UWorld* InWorld, const FVector& Center, float Radius, const FVector& Direction, float AngleWidth, int32 Segments, const FColor& Color, bool PersistentLines = false, float LifeTime = -1.f, uint8 DepthPriority = 0, float Thickness = 0.f);
 /** Draw Debug 2D donut */
 ENGINE_API void DrawDebug2DDonut(const UWorld* InWorld, const FMatrix& TransformMatrix, float InnerRadius, float OuterRadius, int32 Segments, const FColor& Color, bool bPersistentLines = false, float LifeTime=-1.f, uint8 DepthPriority = 0, float Thickness = 0.f);
 /** Draw a debug sphere */
@@ -148,6 +150,7 @@ FORCEINLINE void DrawDebugCoordinateSystem(const UWorld* InWorld, FVector const&
 FORCEINLINE void DrawDebugCrosshairs(const UWorld* InWorld, FVector const& AxisLoc, FRotator const& AxisRot, float Scale, const FColor& Color = FColor::White, bool bPersistentLines = false, float LifeTime = -1.f, uint8 DepthPriority = 0) {}
 FORCEINLINE void DrawDebugCircle(const UWorld* InWorld, const FMatrix& TransformMatrix, float Radius, int32 Segments, const FColor& Color, bool bPersistentLines = false, float LifeTime = -1.f, uint8 DepthPriority = 0, float Thickness = 0.f, bool bDrawAxis = true) {}
 FORCEINLINE void DrawDebugCircle(const UWorld* InWorld, FVector Center, float Radius, int32 Segments, const FColor& Color, bool bPersistentLines = false, float LifeTime = -1.f, uint8 DepthPriority = 0, float Thickness = 0.f, FVector YAxis = FVector(0.f, 1.f, 0.f), FVector ZAxis = FVector(0.f, 0.f, 1.f), bool bDrawAxis = true) {}
+FORCEINLINE void DrawDebugCircleArc(const UWorld* InWorld, FVector const& Center, float Radius, FVector const& Direction, float AngleWidth, int32 Segments, const FColor& Color, bool bPersistentLines = false, float LifeTime = -1.f, uint8 DepthPriority = 0, float Thickness = 0.f) {}
 FORCEINLINE void DrawDebug2DDonut(const UWorld* InWorld, const FMatrix& TransformMatrix, float InnerRadius, float OuterRadius, int32 Segments, const FColor& Color, bool bPersistentLines = false, float LifeTime = -1.f, uint8 DepthPriority = 0, float Thickness = 0.f) {}
 FORCEINLINE void DrawDebugSphere(const UWorld* InWorld, FVector const& Center, float Radius, int32 Segments, FColor const& Color, bool bPersistentLines = false, float LifeTime = -1.f, uint8 DepthPriority = 0, float Thickness = 0.f) {}
 FORCEINLINE void DrawDebugCylinder(const UWorld* InWorld, FVector const& Start, FVector const& End, float Radius, int32 Segments, FColor const& Color, bool bPersistentLines = false, float LifeTime = -1.f, uint8 DepthPriority = 0, float Thickness = 0.f) {}
