HLOD support for Level Instances

diff --git a/Engine/Source/Editor/UnrealEd/Private/Commandlets/WorldPartitionBuilderCommandlet.cpp b/Engine/Source/Editor/UnrealEd/Private/Commandlets/WorldPartitionBuilderCommandlet.cpp
--- a/Engine/Source/Editor/UnrealEd/Private/Commandlets/WorldPartitionBuilderCommandlet.cpp
+++ b/Engine/Source/Editor/UnrealEd/Private/Commandlets/WorldPartitionBuilderCommandlet.cpp
@@ -12,6 +12,7 @@
 #include "ProfilingDebugging/ScopedTimers.h"
 #include "WorldPartition/WorldPartition.h"
 #include "WorldPartition/WorldPartitionSubsystem.h"
+#include "LevelInstance/LevelInstanceSubsystem.h"
 #include "Trace/Trace.h"
 
 DEFINE_LOG_CATEGORY_STATIC(LogWorldPartitionBuilderCommandlet, All, All);
@@ -97,6 +98,12 @@ int32 UWorldPartitionBuilderCommandlet::Main(const FString& Params)
 		return 1;
 	}
 
+	// Commandlets aren't loading level instances by default, change that behavior
+	if (ULevelInstanceSubsystem* LevelInstanceSubsystem = World->GetSubsystem<ULevelInstanceSubsystem>())
+	{
+		LevelInstanceSubsystem->SetLoadInstancesOnRegistration(true);
+	}
+
 	// Retrieve the world partition.
 	UWorldPartition* WorldPartition = World->GetWorldPartition();
 	check(WorldPartition);
diff --git a/Engine/Source/Runtime/Engine/Private/LevelInstance/LevelInstanceActor.cpp b/Engine/Source/Runtime/Engine/Private/LevelInstance/LevelInstanceActor.cpp
--- a/Engine/Source/Runtime/Engine/Private/LevelInstance/LevelInstanceActor.cpp
+++ b/Engine/Source/Runtime/Engine/Private/LevelInstance/LevelInstanceActor.cpp
@@ -77,7 +77,7 @@ void ALevelInstance::PostRegisterAllComponents()
 	{
 		LevelInstanceID = LevelInstanceSubsystem->RegisterLevelInstance(this);
 
-		if (!IsRunningCommandlet())
+		if (LevelInstanceSubsystem->ShouldLoadInstancesOnRegistration())
 		{
 			LoadLevelInstance();
 		}
@@ -98,7 +98,7 @@ void ALevelInstance::PostUnregisterAllComponents()
 
 		LevelInstanceSubsystem->UnregisterLevelInstance(this);
 
-		if (!IsRunningCommandlet())
+		if (LevelInstanceSubsystem->ShouldLoadInstancesOnRegistration())
 		{
 			UnloadLevelInstance();
 		}
diff --git a/Engine/Source/Runtime/Engine/Private/LevelInstance/LevelInstanceSubsystem.cpp b/Engine/Source/Runtime/Engine/Private/LevelInstance/LevelInstanceSubsystem.cpp
--- a/Engine/Source/Runtime/Engine/Private/LevelInstance/LevelInstanceSubsystem.cpp
+++ b/Engine/Source/Runtime/Engine/Private/LevelInstance/LevelInstanceSubsystem.cpp
@@ -63,6 +63,8 @@ void ULevelInstanceSubsystem::Initialize(FSubsystemCollectionBase& Collection)
 	{
 		FModuleManager::LoadModuleChecked<ILevelInstanceEditorModule>("LevelInstanceEditor");
 	}
+
+	bLoadInstancesOnRegistration = !IsRunningCommandlet();
 #endif
 }
 
@@ -1613,7 +1615,21 @@ void ULevelInstanceSubsystem::BlockUnloadLevelInstance(ALevelInstance* LevelInst
 	}
 }
 
+void ULevelInstanceSubsystem::SetLoadInstancesOnRegistration(bool bInLoadInstancesOnRegistration)
+{
+	bLoadInstancesOnRegistration = bInLoadInstancesOnRegistration;
+}
+
 #endif
 
+bool ULevelInstanceSubsystem::ShouldLoadInstancesOnRegistration() const
+{
+#if WITH_EDITOR
+	return bLoadInstancesOnRegistration;
+#else
+	return !IsRunningCommandlet();
+#endif
+}
+
 #undef LOCTEXT_NAMESPACE
 
diff --git a/Engine/Source/Runtime/Engine/Private/WorldPartition/HLOD/HLODActor.cpp b/Engine/Source/Runtime/Engine/Private/WorldPartition/HLOD/HLODActor.cpp
--- a/Engine/Source/Runtime/Engine/Private/WorldPartition/HLOD/HLODActor.cpp
+++ b/Engine/Source/Runtime/Engine/Private/WorldPartition/HLOD/HLODActor.cpp
@@ -50,7 +50,7 @@ void AWorldPartitionHLOD::EndPlay(const EEndPlayReason::Type EndPlayReason)
 
 #if WITH_EDITOR
 
-void AWorldPartitionHLOD::OnSubActorLoaded(AActor& Actor)
+void AWorldPartitionHLOD::OnSubActorLoaded(const AActor& Actor)
 {
 	bool bIsAlreadyInSet = false;
 	LoadedSubActors.Add(&Actor, &bIsAlreadyInSet);
@@ -64,7 +64,7 @@ void AWorldPartitionHLOD::OnSubActorLoaded(AActor& Actor)
 	}
 }
 
-void AWorldPartitionHLOD::OnSubActorUnloaded(AActor& Actor)
+void AWorldPartitionHLOD::OnSubActorUnloaded(const AActor& Actor)
 {
 	LoadedSubActors.Remove(&Actor);
 
@@ -88,7 +88,7 @@ void AWorldPartitionHLOD::SetupLoadedSubActors()
 		for (const FGuid& SubActorGuid : SubActors)
 		{
 			const FWorldPartitionActorDesc* SubActorDesc = WorldPartition->GetActorDesc(SubActorGuid);
-			AActor* SubActor = SubActorDesc->IsLoaded() ? SubActorDesc->GetActor() : nullptr;
+			AActor* SubActor = SubActorDesc && SubActorDesc->IsLoaded() ? SubActorDesc->GetActor() : nullptr;
 			if (SubActor)
 			{
 				OnSubActorLoaded(*SubActor);
@@ -104,7 +104,7 @@ void AWorldPartitionHLOD::ResetLoadedSubActors()
 	UWorld* World = GetWorld();
 	if (World && !World->IsGameWorld())
 	{
-		for (const TWeakObjectPtr<AActor>& SubActor : LoadedSubActors)
+		for (const TWeakObjectPtr<const AActor>& SubActor : LoadedSubActors)
 		{
 			if (!SubActor.IsValid())
 			{
@@ -192,20 +192,15 @@ void AWorldPartitionHLOD::SetHLODPrimitives(const TArray<UPrimitiveComponent*>&
 	}
 }
 
-void AWorldPartitionHLOD::SetChildrenPrimitives(const TArray<UPrimitiveComponent*>& InChildrenPrimitives)
+void AWorldPartitionHLOD::SetSubActors(const TArray<const AActor*>& InSubActors)
 {
 	check(GetHLODComponent());
 
 	ResetLoadedSubActors();
 	SubActors.Empty();
 
-	UPrimitiveComponent* HLODComponent = GetHLODComponent();
-	check(HLODComponent);
-
-	for (UPrimitiveComponent* ChildPrimitive : InChildrenPrimitives)
+	for (const AActor* SubActor : InSubActors)
 	{
-		AActor* SubActor = ChildPrimitive->GetOwner();
-		
 		if (!LoadedSubActors.Contains(SubActor))
 		{
 			OnSubActorLoaded(*SubActor);
diff --git a/Engine/Source/Runtime/Engine/Private/WorldPartition/HLOD/HLODBuilder.cpp b/Engine/Source/Runtime/Engine/Private/WorldPartition/HLOD/HLODBuilder.cpp
--- a/Engine/Source/Runtime/Engine/Private/WorldPartition/HLOD/HLODBuilder.cpp
+++ b/Engine/Source/Runtime/Engine/Private/WorldPartition/HLOD/HLODBuilder.cpp
@@ -30,12 +30,17 @@ class FHLODBuilder
 public:
 	virtual ~FHLODBuilder() {}
 
-	virtual void Build(const TArray<UPrimitiveComponent*>& InSubComponents) = 0;
+	virtual const TCHAR* GetActorSuffix() const = 0;
+	virtual TArray<UPrimitiveComponent*> CreateComponents(AWorldPartitionHLOD* InHLODActor, const TArray<UPrimitiveComponent*>& InSubComponents) = 0;
 
-	typedef TFunction<TArray<UPrimitiveComponent*>(AWorldPartitionHLOD*)> FCreateComponentsFunction;
-		
-	void SpawnHLODActor(const TCHAR* InName, const TArray<UPrimitiveComponent*>& InSubComponents, FCreateComponentsFunction InCreateComponentsFunc)
+	void Build(const TArray<const AActor*>& InSubActors)
 	{
+		TArray<UPrimitiveComponent*> SubComponents = GatherPrimitiveComponents(0, InSubActors);
+		if (SubComponents.IsEmpty())
+		{
+			return;
+		}
+
 		AWorldPartitionHLOD* HLODActor = nullptr;
 
 		// Compute HLODActor hash
@@ -52,20 +57,20 @@ public:
 		if (!HLODActor)
 		{
 			FActorSpawnParameters SpawnParams;
-			SpawnParams.Name = *FString::Printf(TEXT("%s_%016llx_%s"), *HLODLayer->GetName(), CellHash, InName);
+			SpawnParams.Name = *FString::Printf(TEXT("%s_%016llx_%s"), *HLODLayer->GetName(), CellHash, GetActorSuffix());
 			SpawnParams.NameMode = FActorSpawnParameters::ESpawnActorNameMode::Required_Fatal;
 			HLODActor = World->SpawnActor<AWorldPartitionHLOD>(SpawnParams);
 			HLODActor->SetActorLabel(CellName.ToString());
 		}
 
-		TArray<UPrimitiveComponent*> HLODPrimitives = InCreateComponentsFunc(HLODActor);
+		TArray<UPrimitiveComponent*> HLODPrimitives = CreateComponents(HLODActor, SubComponents);
 		HLODPrimitives.RemoveSwap(nullptr);
 
 		if (!HLODPrimitives.IsEmpty())
 		{
 			HLODActor->Modify();
 			HLODActor->SetHLODPrimitives(HLODPrimitives);
-			HLODActor->SetChildrenPrimitives(InSubComponents);
+			HLODActor->SetSubActors(InSubActors);
 			HLODActor->SetRuntimeGrid(HLODLayer->GetRuntimeGrid(HLODLevel));
 			HLODActor->SetLODLevel(HLODLevel);
 			HLODActor->SetHLODLayer(HLODLayer->GetParentLayer().LoadSynchronous());
@@ -92,22 +97,33 @@ public:
 		}
 	}
 
-	static TArray<UPrimitiveComponent*> GatherPrimitiveComponents(uint32 InHLODLevel, const TArray<const AActor*> InActors)
+	TArray<UPrimitiveComponent*> GatherPrimitiveComponents(uint32 InHLODLevel, const TArray<const AActor*> InActors)
 	{
 		TArray<UPrimitiveComponent*> PrimitiveComponents;
 		for (const AActor* SubActor : InActors)
 		{
-			for (UActorComponent* SubComponent : SubActor->GetComponents())
+			// Gather all subactors
+			TSet<AActor*> UnderlyingActors;
+			SubActor->EditorGetUnderlyingActors(UnderlyingActors);
+
+			for (AActor* UnderlyingActor : UnderlyingActors)
 			{
-				if (UPrimitiveComponent* PrimitiveComponent = Cast<UPrimitiveComponent>(SubComponent))
+				if (UHLODLayer::ShouldIncludeInHLOD(UnderlyingActor, true))
 				{
-					if (UHLODLayer::ShouldIncludeInHLOD(PrimitiveComponent, InHLODLevel))
+					for (UActorComponent* SubComponent : UnderlyingActor->GetComponents())
 					{
-						PrimitiveComponents.Add(PrimitiveComponent);
+						if (UPrimitiveComponent* PrimitiveComponent = Cast<UPrimitiveComponent>(SubComponent))
+						{
+							if (UHLODLayer::ShouldIncludeInHLOD(PrimitiveComponent, InHLODLevel))
+							{
+								PrimitiveComponents.Add(PrimitiveComponent);
+							}
+						}
 					}
 				}
 			}
 		}
+
 		return PrimitiveComponents;
 	}
 
@@ -183,60 +199,57 @@ class FHLODBuilder_Instancing : public FHLODBuilder
 		uint32						Hash;
 	};
 
-	virtual void Build(const TArray<UPrimitiveComponent*>& InSubComponents) override
+	virtual const TCHAR* GetActorSuffix() const override { return TEXT("InstancedMeshes"); }
+
+	virtual TArray<UPrimitiveComponent*> CreateComponents(AWorldPartitionHLOD* InHLODActor, const TArray<UPrimitiveComponent*>& InSubComponents) override
 	{
-		TRACE_CPUPROFILER_EVENT_SCOPE(FHLODBuilder_Instancing::BuildHLOD);
+		TRACE_CPUPROFILER_EVENT_SCOPE(FHLODBuilder_Instancing::CreateComponents);
 
-		FCreateComponentsFunction CreateComponentLambda = [&InSubComponents](AWorldPartitionHLOD* HLODActor)
-		{
-			TArray<UPrimitiveComponent*> Components;
+		TArray<UPrimitiveComponent*> Components;
 
-			// Gather all meshes to instantiate along with their transforms
-			TMap<FInstancingKey, TArray<UPrimitiveComponent*>> Instances;
-			for (UPrimitiveComponent* Primitive : InSubComponents)
+		// Gather all meshes to instantiate along with their transforms
+		TMap<FInstancingKey, TArray<UPrimitiveComponent*>> Instances;
+		for (UPrimitiveComponent* Primitive : InSubComponents)
+		{
+			if (UStaticMeshComponent* SMC = Cast<UStaticMeshComponent>(Primitive))
 			{
-				if (UStaticMeshComponent* SMC = Cast<UStaticMeshComponent>(Primitive))
-				{
-					Instances.FindOrAdd(FInstancingKey(SMC)).Add(SMC);
-				}
+				Instances.FindOrAdd(FInstancingKey(SMC)).Add(SMC);
 			}
+		}
 
-			// Create an ISMC for each SM asset we found
-			for (const auto& Entry : Instances)
-			{
-				const FInstancingKey EntryInstancingKey = Entry.Key;
-				const TArray<UPrimitiveComponent*>& EntryComponents = Entry.Value;
+		// Create an ISMC for each SM asset we found
+		for (const auto& Entry : Instances)
+		{
+			const FInstancingKey EntryInstancingKey = Entry.Key;
+			const TArray<UPrimitiveComponent*>& EntryComponents = Entry.Value;
 			
-				UInstancedStaticMeshComponent* Component = NewObject<UInstancedStaticMeshComponent>(HLODActor);
-				EntryInstancingKey.ApplyTo(Component);
-				Component->SetForcedLodModel(Component->GetStaticMesh()->GetNumLODs());
+			UInstancedStaticMeshComponent* Component = NewObject<UInstancedStaticMeshComponent>(InHLODActor);
+			EntryInstancingKey.ApplyTo(Component);
+			Component->SetForcedLodModel(Component->GetStaticMesh()->GetNumLODs());
 
-				// Add all instances
-				for (UPrimitiveComponent* SMC : EntryComponents)
+			// Add all instances
+			for (UPrimitiveComponent* SMC : EntryComponents)
+			{
+				// If we have an ISMC, retrieve all instances
+				if (UInstancedStaticMeshComponent* InstancedStaticMeshComponent = Cast<UInstancedStaticMeshComponent>(SMC))
 				{
-					// If we have an ISMC, retrieve all instances
-					if (UInstancedStaticMeshComponent* InstancedStaticMeshComponent = Cast<UInstancedStaticMeshComponent>(SMC))
+					for (int32 InstanceIdx = 0; InstanceIdx < InstancedStaticMeshComponent->GetInstanceCount(); InstanceIdx++)
 					{
-						for (int32 InstanceIdx = 0; InstanceIdx < InstancedStaticMeshComponent->GetInstanceCount(); InstanceIdx++)
-						{
-							FTransform InstanceTransform;
-							InstancedStaticMeshComponent->GetInstanceTransform(InstanceIdx, InstanceTransform, true);
-							Component->AddInstanceWorldSpace(InstanceTransform);
-						}
-					}
-					else
-					{
-						Component->AddInstanceWorldSpace(SMC->GetComponentTransform());
+						FTransform InstanceTransform;
+						InstancedStaticMeshComponent->GetInstanceTransform(InstanceIdx, InstanceTransform, true);
+						Component->AddInstanceWorldSpace(InstanceTransform);
 					}
 				}
+				else
+				{
+					Component->AddInstanceWorldSpace(SMC->GetComponentTransform());
+				}
+			}
 
-				Components.Add(Component);
-			};
-
-			return Components;
+			Components.Add(Component);
 		};
 
-		SpawnHLODActor(TEXT("InstancedMeshes"), InSubComponents, CreateComponentLambda);
+		return Components;
 	}
 };
 
@@ -246,36 +259,33 @@ class FHLODBuilder_Instancing : public FHLODBuilder
  */
 class FHLODBuilder_MeshMerge : public FHLODBuilder
 {
-	virtual void Build(const TArray<UPrimitiveComponent*>& InSubComponents) override
+	virtual const TCHAR* GetActorSuffix() const override { return TEXT("MergedMesh"); }
+
+	virtual TArray<UPrimitiveComponent*> CreateComponents(AWorldPartitionHLOD* InHLODActor, const TArray<UPrimitiveComponent*>& InSubComponents) override
 	{
-		TRACE_CPUPROFILER_EVENT_SCOPE(UHLODLayer::BuildHLOD_MeshMerge);
+		TRACE_CPUPROFILER_EVENT_SCOPE(FHLODBuilder_MeshMerge::CreateComponents);
 
-		FCreateComponentsFunction CreateComponentLambda = [&InSubComponents, this](AWorldPartitionHLOD* HLODActor)
-		{
-			TArray<UObject*> Assets;
-			FVector MergedActorLocation;
+		TArray<UObject*> Assets;
+		FVector MergedActorLocation;
 
-			const IMeshMergeUtilities& MeshMergeUtilities = FModuleManager::Get().LoadModuleChecked<IMeshMergeModule>("MeshMergeUtilities").GetUtilities();
-			MeshMergeUtilities.MergeComponentsToStaticMesh(InSubComponents, HLODActor->GetWorld(), HLODLayer->GetMeshMergeSettings(), HLODLayer->GetHLODMaterial().LoadSynchronous(), HLODActor->GetPackage(), CellName.ToString(), Assets, MergedActorLocation, 0.25f, false);
+		const IMeshMergeUtilities& MeshMergeUtilities = FModuleManager::Get().LoadModuleChecked<IMeshMergeModule>("MeshMergeUtilities").GetUtilities();
+		MeshMergeUtilities.MergeComponentsToStaticMesh(InSubComponents, InHLODActor->GetWorld(), HLODLayer->GetMeshMergeSettings(), HLODLayer->GetHLODMaterial().LoadSynchronous(), InHLODActor->GetPackage(), CellName.ToString(), Assets, MergedActorLocation, 0.25f, false);
 
-			UStaticMeshComponent* Component = nullptr;
-			Algo::ForEach(Assets, [HLODActor, &Component, &MergedActorLocation](UObject* Asset)
-			{
-				Asset->ClearFlags(RF_Public | RF_Standalone);
-				Asset->Rename(nullptr, HLODActor);
+		UStaticMeshComponent* Component = nullptr;
+		Algo::ForEach(Assets, [InHLODActor, &Component, &MergedActorLocation](UObject* Asset)
+		{
+			Asset->ClearFlags(RF_Public | RF_Standalone);
+			Asset->Rename(nullptr, InHLODActor);
 
-				if (Cast<UStaticMesh>(Asset))
-				{
-					Component = NewObject<UStaticMeshComponent>(HLODActor);
-					Component->SetStaticMesh(static_cast<UStaticMesh*>(Asset));
-					Component->SetWorldLocation(MergedActorLocation);
-				}
-			});
+			if (Cast<UStaticMesh>(Asset))
+			{
+				Component = NewObject<UStaticMeshComponent>(InHLODActor);
+				Component->SetStaticMesh(static_cast<UStaticMesh*>(Asset));
+				Component->SetWorldLocation(MergedActorLocation);
+			}
+		});
 
-			return TArray<UPrimitiveComponent*>({ Component });
-		};
-		
-		SpawnHLODActor(TEXT("MergedMesh"), InSubComponents, CreateComponentLambda);
+		return TArray<UPrimitiveComponent*>({ Component });
 	}
 };
 
@@ -284,39 +294,36 @@ class FHLODBuilder_MeshMerge : public FHLODBuilder
  */
 class FHLODBuilder_MeshSimplify : public FHLODBuilder
 {
-	virtual void Build(const TArray<UPrimitiveComponent*>& InSubComponents) override
+	virtual const TCHAR* GetActorSuffix() const override { return TEXT("SimplifiedMesh"); }
+
+	virtual TArray<UPrimitiveComponent*> CreateComponents(AWorldPartitionHLOD* InHLODActor, const TArray<UPrimitiveComponent*>& InSubComponents) override
 	{
-		TRACE_CPUPROFILER_EVENT_SCOPE(UHLODLayer::BuildHLOD_MeshProxy);
+		TRACE_CPUPROFILER_EVENT_SCOPE(FHLODBuilder_MeshSimplify::CreateComponents);
 
-		FCreateComponentsFunction CreateComponentLambda = [&InSubComponents, this] (AWorldPartitionHLOD* HLODActor)
-		{
-			TArray<UObject*> Assets;
-			FCreateProxyDelegate ProxyDelegate;
-			ProxyDelegate.BindLambda([&Assets](const FGuid Guid, TArray<UObject*>& InAssetsCreated) { Assets = InAssetsCreated; });
+		TArray<UObject*> Assets;
+		FCreateProxyDelegate ProxyDelegate;
+		ProxyDelegate.BindLambda([&Assets](const FGuid Guid, TArray<UObject*>& InAssetsCreated) { Assets = InAssetsCreated; });
 
-			TArray<UStaticMeshComponent*> StaticMeshComponents;
-			Algo::TransformIf(InSubComponents, StaticMeshComponents, [](UPrimitiveComponent* InPrimitiveComponent) { return InPrimitiveComponent->IsA<UStaticMeshComponent>(); }, [](UPrimitiveComponent* InPrimitiveComponent) { return Cast<UStaticMeshComponent>(InPrimitiveComponent); });
+		TArray<UStaticMeshComponent*> StaticMeshComponents;
+		Algo::TransformIf(InSubComponents, StaticMeshComponents, [](UPrimitiveComponent* InPrimitiveComponent) { return InPrimitiveComponent->IsA<UStaticMeshComponent>(); }, [](UPrimitiveComponent* InPrimitiveComponent) { return Cast<UStaticMeshComponent>(InPrimitiveComponent); });
 
-			const IMeshMergeUtilities& MeshMergeUtilities = FModuleManager::Get().LoadModuleChecked<IMeshMergeModule>("MeshMergeUtilities").GetUtilities();
-			MeshMergeUtilities.CreateProxyMesh(StaticMeshComponents, HLODLayer->GetMeshSimplifySettings(), HLODLayer->GetHLODMaterial().LoadSynchronous(), HLODActor->GetPackage(), CellName.ToString(), FGuid::NewGuid(), ProxyDelegate, true);
+		const IMeshMergeUtilities& MeshMergeUtilities = FModuleManager::Get().LoadModuleChecked<IMeshMergeModule>("MeshMergeUtilities").GetUtilities();
+		MeshMergeUtilities.CreateProxyMesh(StaticMeshComponents, HLODLayer->GetMeshSimplifySettings(), HLODLayer->GetHLODMaterial().LoadSynchronous(), InHLODActor->GetPackage(), CellName.ToString(), FGuid::NewGuid(), ProxyDelegate, true);
 
-			UStaticMeshComponent* Component = nullptr;
-			Algo::ForEach(Assets, [HLODActor, &Component](UObject* Asset)
-			{
-				Asset->ClearFlags(RF_Public | RF_Standalone);
-				Asset->Rename(nullptr, HLODActor);
+		UStaticMeshComponent* Component = nullptr;
+		Algo::ForEach(Assets, [InHLODActor, &Component](UObject* Asset)
+		{
+			Asset->ClearFlags(RF_Public | RF_Standalone);
+			Asset->Rename(nullptr, InHLODActor);
 
-				if (Cast<UStaticMesh>(Asset))
-				{
-					Component = NewObject<UStaticMeshComponent>(HLODActor);
-					Component->SetStaticMesh(static_cast<UStaticMesh*>(Asset));
-				}
-			});
+			if (Cast<UStaticMesh>(Asset))
+			{
+				Component = NewObject<UStaticMeshComponent>(InHLODActor);
+				Component->SetStaticMesh(static_cast<UStaticMesh*>(Asset));
+			}
+		});
 
-			return TArray<UPrimitiveComponent*>({ Component });
-		};
-		
-		SpawnHLODActor(TEXT("SimplifiedMesh"), InSubComponents, CreateComponentLambda);
+		return TArray<UPrimitiveComponent*>({ Component });
 	}
 };
 
@@ -343,10 +350,9 @@ TArray<AWorldPartitionHLOD*> FHLODBuilderUtilities::BuildHLODs(UWorldPartition*
 		checkf(false, TEXT("Unsupported type"));
 	}
 
-	TArray<UPrimitiveComponent*> SubComponents = FHLODBuilder::GatherPrimitiveComponents(0, InSubActors);
 	TArray<AWorldPartitionHLOD*> HLODActors;
 
-	if (HLODBuilder && !SubComponents.IsEmpty())
+	if (HLODBuilder)
 	{
 		HLODBuilder->World = InWorldPartition->GetWorld();
 		HLODBuilder->WorldPartition = InWorldPartition;
@@ -356,7 +362,7 @@ TArray<AWorldPartitionHLOD*> FHLODBuilderUtilities::BuildHLODs(UWorldPartition*
 		HLODBuilder->CellBounds = InCellBounds;
 		HLODBuilder->Context = InContext;
 
-		HLODBuilder->Build(SubComponents);
+		HLODBuilder->Build(InSubActors);
 			
 		HLODActors = HLODBuilder->HLODActors;
 	}
diff --git a/Engine/Source/Runtime/Engine/Private/WorldPartition/HLOD/HLODLayer.cpp b/Engine/Source/Runtime/Engine/Private/WorldPartition/HLOD/HLODLayer.cpp
--- a/Engine/Source/Runtime/Engine/Private/WorldPartition/HLOD/HLODLayer.cpp
+++ b/Engine/Source/Runtime/Engine/Private/WorldPartition/HLOD/HLODLayer.cpp
@@ -8,12 +8,16 @@
 #include "WorldPartition/HLOD/HLODActor.h"
 
 #if WITH_EDITOR
+#include "Algo/Copy.h"
 #include "UObject/ConstructorHelpers.h"
 #include "Materials/MaterialInterface.h"
 
 #include "WorldPartition/WorldPartition.h"
 #include "WorldPartition/WorldPartitionActorDesc.h"
 #include "WorldPartition/HLOD/HLODBuilder.h"
+
+#include "LevelInstance/LevelInstanceActor.h"
+#include "LevelInstance/LevelInstanceSubsystem.h"
 #endif
 
 UHLODLayer::UHLODLayer(const FObjectInitializer& ObjectInitializer)
@@ -35,13 +39,19 @@ TArray<AWorldPartitionHLOD*> UHLODLayer::GenerateHLODForCell(UWorldPartition* In
 	TRACE_CPUPROFILER_EVENT_SCOPE(UHLODLayer::GenerateHLODForCell);
 
 	TMap<UHLODLayer*, TArray<const AActor*>> HLODLayersActors;
-	for (const AActor* Actor : InCellActors)
+	for (AActor* Actor : InCellActors)
 	{
 		if (ShouldIncludeInHLOD(Actor))
 		{
 			UHLODLayer* HLODLayer = UHLODLayer::GetHLODLayer(Actor);
 			if (HLODLayer)
 			{
+				if (ALevelInstance* LevelInstance = Cast<ALevelInstance>(Actor))
+				{
+					// Wait for level instance loading
+					LevelInstance->GetLevelInstanceSubsystem()->BlockLoadLevelInstance(LevelInstance);
+				}
+
 				HLODLayersActors.FindOrAdd(HLODLayer).Add(Actor);
 			}
 		}
@@ -50,7 +60,10 @@ TArray<AWorldPartitionHLOD*> UHLODLayer::GenerateHLODForCell(UWorldPartition* In
 	TArray<AWorldPartitionHLOD*> HLODActors;
 	for (const auto& HLODLayerActors : HLODLayersActors)
 	{
-		HLODActors += FHLODBuilderUtilities::BuildHLODs(InWorldPartition, InContext, InCellName, InCellBounds, HLODLayerActors.Key, InHLODLevel, HLODLayerActors.Value);
+		if (!HLODLayerActors.Value.IsEmpty())
+		{
+			HLODActors += FHLODBuilderUtilities::BuildHLODs(InWorldPartition, InContext, InCellName, InCellBounds, HLODLayerActors.Key, InHLODLevel, HLODLayerActors.Value);
+		}
 	}
 	return HLODActors;
 }
@@ -94,11 +107,23 @@ UHLODLayer* UHLODLayer::GetHLODLayer(const AActor* InActor)
 	bool bIsHLOD0 = !InActor->IsA<AWorldPartitionHLOD>();
 	if (bIsHLOD0) 
 	{
+		// Check if this actor is part of a DataLayer that has a default HLOD layer
+		for(const UDataLayer* DataLayer : InActor->GetDataLayerObjects())
+		{
+			UHLODLayer* HLODLayer = DataLayer ? DataLayer->GetDefaultHLODLayer() : nullptr;
+			if (HLODLayer)
+			{
+				return HLODLayer;
+			}
+		}
+
+		// Fallback to the world partition default HLOD layer
 		if (UWorldPartition* WorldPartition = InActor->GetWorld()->GetWorldPartition())
 		{
 			return WorldPartition->DefaultHLODLayer;
 		}
 	}
+
 	return nullptr;
 }
 
@@ -118,7 +143,7 @@ FName UHLODLayer::GetRuntimeGrid(uint32 InHLODLevel) const
 
 #endif // WITH_EDITOR
 
-bool UHLODLayer::ShouldIncludeInHLOD(const AActor* InActor)
+bool UHLODLayer::ShouldIncludeInHLOD(const AActor* InActor, bool bInFromLevelInstance)
 {
 	if (!InActor)
 	{
@@ -135,7 +160,7 @@ bool UHLODLayer::ShouldIncludeInHLOD(const AActor* InActor)
 		return false;
 	}
 
-	if (InActor->HasAnyFlags(RF_Transient))
+	if (!bInFromLevelInstance && InActor->HasAnyFlags(RF_Transient))
 	{
 		return false;
 	}
diff --git a/Engine/Source/Runtime/Engine/Private/WorldPartition/RuntimeSpatialHash/RuntimeSpatialHashHLOD.cpp b/Engine/Source/Runtime/Engine/Private/WorldPartition/RuntimeSpatialHash/RuntimeSpatialHashHLOD.cpp
--- a/Engine/Source/Runtime/Engine/Private/WorldPartition/RuntimeSpatialHash/RuntimeSpatialHashHLOD.cpp
+++ b/Engine/Source/Runtime/Engine/Private/WorldPartition/RuntimeSpatialHash/RuntimeSpatialHashHLOD.cpp
@@ -79,7 +79,26 @@ static void DeletePackage(const FString& PackageName, ISourceControlHelper* Sour
 
 static void DeletePackage(UPackage* Package, ISourceControlHelper* SourceControlHelper)
 {
-	DeletePackage(Package->GetName(), SourceControlHelper);
+	if (SourceControlHelper)
+	{
+		SourceControlHelper->Delete(Package);
+	}
+	else
+	{
+		DeletePackage(Package->GetName(), SourceControlHelper);
+	}
+}
+
+static void DeletePackage(FWorldPartitionActorDesc* ActorDesc, ISourceControlHelper* SourceControlHelper)
+{
+	if (ActorDesc->IsLoaded())
+	{
+		DeletePackage(ActorDesc->GetActor()->GetPackage(), SourceControlHelper);
+	}
+	else
+	{
+		DeletePackage(ActorDesc->GetActorPackage().ToString(), SourceControlHelper);
+	}
 }
 
 static TArray<FGuid> GenerateHLODsForGrid(UWorldPartition* WorldPartition, const FSpatialHashRuntimeGrid& RuntimeGrid, uint32 HLODLevel, FHLODGenerationContext& Context, ISourceControlHelper* SourceControlHelper, const TArray<FActorCluster>& ActorClusters)
@@ -474,9 +493,9 @@ bool UWorldPartitionRuntimeSpatialHash::GenerateHLOD(ISourceControlHelper* Sourc
 
 	// Create actor clusters - ignore HLOD actors
 	TArray<FActorCluster> ActorClusters = CreateActorClusters(WorldPartition, [](const FWorldPartitionActorDesc& ActorDesc)
-		{
-			return !ActorDesc.GetActorClass()->IsChildOf<AWorldPartitionHLOD>();
-		});
+	{
+		return !ActorDesc.GetActorClass()->IsChildOf<AWorldPartitionHLOD>();
+	});
 
 	TArray<TArray<FActorCluster>> GridsClusters;
 	GridsClusters.InsertDefaulted(0, Grids.Num());
@@ -549,7 +568,7 @@ bool UWorldPartitionRuntimeSpatialHash::GenerateHLOD(ISourceControlHelper* Sourc
 		FWorldPartitionActorDesc* HLODActorDesc = ActorHandle.Get();
 		check(HLODActorDesc);
 
-		DeletePackage(HLODActorDesc->GetActorPackage().ToString(), SourceControlHelper);
+		DeletePackage(HLODActorDesc, SourceControlHelper);
 	};
 
 	// Destroy all unreferenced HLOD actors
diff --git a/Engine/Source/Runtime/Engine/Public/LevelInstance/LevelInstanceSubsystem.h b/Engine/Source/Runtime/Engine/Public/LevelInstance/LevelInstanceSubsystem.h
--- a/Engine/Source/Runtime/Engine/Public/LevelInstance/LevelInstanceSubsystem.h
+++ b/Engine/Source/Runtime/Engine/Public/LevelInstance/LevelInstanceSubsystem.h
@@ -43,6 +43,7 @@ public:
 	void RequestLoadLevelInstance(ALevelInstance* LevelInstanceActor, bool bUpdate);
 	void RequestUnloadLevelInstance(ALevelInstance* LevelInstanceActor);
 	bool IsLoaded(const ALevelInstance* LevelInstanceActor) const;
+	bool ShouldLoadInstancesOnRegistration() const;
 	void ForEachLevelInstanceAncestorsAndSelf(AActor* Actor, TFunctionRef<bool(ALevelInstance*)> Operation) const;
 
 #if WITH_EDITOR
@@ -88,6 +89,8 @@ public:
 	ALevelInstance* GetParentLevelInstance(const AActor* Actor) const;
 	void BlockLoadLevelInstance(ALevelInstance* LevelInstanceActor);
 	void BlockUnloadLevelInstance(ALevelInstance* LevelInstanceActor);
+
+	void SetLoadInstancesOnRegistration(bool bLoadInstancesOnRegistration);
 #endif
 
 private:
@@ -131,6 +134,8 @@ private:
 	TUniquePtr<FLevelsToRemoveScope> LevelsToRemoveScope;
 	
 	TMap<FName, FLevelInstanceEdit> LevelInstanceEdits;
+
+	bool bLoadInstancesOnRegistration;
 #endif
 
 	struct FLevelInstance
diff --git a/Engine/Source/Runtime/Engine/Public/WorldPartition/HLOD/HLODActor.h b/Engine/Source/Runtime/Engine/Public/WorldPartition/HLOD/HLODActor.h
--- a/Engine/Source/Runtime/Engine/Public/WorldPartition/HLOD/HLODActor.h
+++ b/Engine/Source/Runtime/Engine/Public/WorldPartition/HLOD/HLODActor.h
@@ -25,12 +25,12 @@ public:
 	inline uint32 GetLODLevel() const { return LODLevel; }
 
 #if WITH_EDITOR
-	void OnSubActorLoaded(AActor& Actor);
-	void OnSubActorUnloaded(AActor& Actor);
+	void OnSubActorLoaded(const AActor& Actor);
+	void OnSubActorUnloaded(const AActor& Actor);
 
 	void SetHLODPrimitives(const TArray<UPrimitiveComponent*>& InHLODPrimitives);
-	void SetChildrenPrimitives(const TArray<UPrimitiveComponent*>& InChildrenPrimitives);
 
+	void SetSubActors(const TArray<const AActor*>& InSubActors);
 	const TArray<FGuid>& GetSubActors() const;
 
 	void SetSubActorsHLODLayer(const UHLODLayer* InSubActorsHLODLayer) { SubActorsHLODLayer = InSubActorsHLODLayer; }
@@ -95,7 +95,7 @@ private:
 	UPROPERTY()
 	int64 GridIndexZ;
 
-	TSet<TWeakObjectPtr<AActor>> LoadedSubActors;
+	TSet<TWeakObjectPtr<const AActor>> LoadedSubActors;
 #endif
 
 	UPROPERTY(NonPIEDuplicateTransient, TextExportTransient, NonTransactional)
diff --git a/Engine/Source/Runtime/Engine/Public/WorldPartition/HLOD/HLODLayer.h b/Engine/Source/Runtime/Engine/Public/WorldPartition/HLOD/HLODLayer.h
--- a/Engine/Source/Runtime/Engine/Public/WorldPartition/HLOD/HLODLayer.h
+++ b/Engine/Source/Runtime/Engine/Public/WorldPartition/HLOD/HLODLayer.h
@@ -52,7 +52,7 @@ public:
 	static FName GetRuntimeGridName(uint32 InLODLevel, int32 InCellSize, float InLoadingRange);
 #endif
 
-	static bool ShouldIncludeInHLOD(const AActor* InActor);
+	static bool ShouldIncludeInHLOD(const AActor* InActor, bool bInFromLevelInstance = false);
 
 #if WITH_EDITORONLY_DATA
 private:
