Physics: Add UFUNCTION to get constrained component from a physics constraint component

diff --git a/Engine/Source/Runtime/Engine/Classes/PhysicsEngine/PhysicsConstraintComponent.h b/Engine/Source/Runtime/Engine/Classes/PhysicsEngine/PhysicsConstraintComponent.h
--- a/Engine/Source/Runtime/Engine/Classes/PhysicsEngine/PhysicsConstraintComponent.h
+++ b/Engine/Source/Runtime/Engine/Classes/PhysicsEngine/PhysicsConstraintComponent.h
@@ -103,6 +103,10 @@ public:
 	UFUNCTION(BlueprintCallable, Category="Physics|Components|PhysicsConstraint")
 	void SetConstrainedComponents(UPrimitiveComponent* Component1, FName BoneName1, UPrimitiveComponent* Component2, FName BoneName2);
 
+	/** Get connected components and potential related attachement bones */
+	UFUNCTION(BlueprintCallable, Category = "Physics|Components|PhysicsConstraint")
+	void GetConstrainedComponents(UPrimitiveComponent*& OutComponent1, FName& OutBoneName1, UPrimitiveComponent*& OutComponent2, FName& OutBoneName2);
+
 	/** Break this constraint */
 	UFUNCTION(BlueprintCallable, Category="Physics|Components|PhysicsConstraint")
 	void BreakConstraint();
diff --git a/Engine/Source/Runtime/Engine/Private/PhysicsEngine/PhysicsConstraintComponent.cpp b/Engine/Source/Runtime/Engine/Private/PhysicsEngine/PhysicsConstraintComponent.cpp
--- a/Engine/Source/Runtime/Engine/Private/PhysicsEngine/PhysicsConstraintComponent.cpp
+++ b/Engine/Source/Runtime/Engine/Private/PhysicsEngine/PhysicsConstraintComponent.cpp
@@ -274,6 +274,15 @@ float UPhysicsConstraintComponent::GetConstraintScale() const
 	return GetComponentScale().GetAbsMin();
 }
 
+void UPhysicsConstraintComponent::GetConstrainedComponents(UPrimitiveComponent*& OutComponent1, FName& OutBoneName1, UPrimitiveComponent*& OutComponent2, FName& OutBoneName2)
+{
+	OutComponent1 = GetComponentInternal(EConstraintFrame::Frame1);
+	OutBoneName1 = ConstraintInstance.ConstraintBone1;
+
+	OutComponent2 = GetComponentInternal(EConstraintFrame::Frame2);
+	OutBoneName2 = ConstraintInstance.ConstraintBone2;
+}
+
 void UPhysicsConstraintComponent::SetConstrainedComponents(UPrimitiveComponent* Component1, FName BoneName1, UPrimitiveComponent* Component2, FName BoneName2)
 {
 	if(Component1 != NULL)
