Chaos - Fix post-merge compiler and linker errors (ChaosCooking.cpp)

diff --git a/Engine/Source/Runtime/Engine/Private/PhysicsEngine/Experimental/ChaosCooking.cpp b/Engine/Source/Runtime/Engine/Private/PhysicsEngine/Experimental/ChaosCooking.cpp
--- a/Engine/Source/Runtime/Engine/Private/PhysicsEngine/Experimental/ChaosCooking.cpp
+++ b/Engine/Source/Runtime/Engine/Private/PhysicsEngine/Experimental/ChaosCooking.cpp
@@ -161,17 +161,17 @@ namespace Chaos
 					}
 
 					// Create the surface particle for the convex
-					TParticles<FReal, 3> ConvexParticles;
-					ConvexParticles.AddParticles(NumHullVerts);
+					TArray<FVec3> ConvexVertices;
+					ConvexVertices.SetNumZeroed(NumHullVerts);
 
 					for(int32 VertIndex = 0; VertIndex < NumHullVerts; ++VertIndex)
 					{
 						const FVector& HullVert = HullVerts[VertIndex];
-						ConvexParticles.X(VertIndex) = FVector(bMirrored ? -HullVert.X : HullVert.X, HullVert.Y, HullVert.Z);
+						ConvexVertices[VertIndex] = FVector(bMirrored ? -HullVert.X : HullVert.X, HullVert.Y, HullVert.Z);
 					}
 
 					// Margin is always zero on convex shapes - they are intended to be instanced
-					OutConvexes.Emplace(new Chaos::FConvex(ConvexParticles, 0.0f));
+					OutConvexes.Emplace(new Chaos::FConvex(ConvexVertices, 0.0f));
 				}
 			};
 
diff --git a/Engine/Source/Runtime/Experimental/Chaos/Public/Chaos/CollisionConvexMesh.h b/Engine/Source/Runtime/Experimental/Chaos/Public/Chaos/CollisionConvexMesh.h
--- a/Engine/Source/Runtime/Experimental/Chaos/Public/Chaos/CollisionConvexMesh.h
+++ b/Engine/Source/Runtime/Experimental/Chaos/Public/Chaos/CollisionConvexMesh.h
@@ -611,19 +611,11 @@ namespace Chaos
 		}
 
 		// CVars variables for controlling geometry complexity checking and simplification
-#if PLATFORM_MAC || PLATFORM_LINUX
 		static CHAOS_API int32 PerformGeometryCheck;
 		static CHAOS_API int32 PerformGeometryReduction;
 		static CHAOS_API int32 VerticesThreshold;
 		static CHAOS_API int32 ParticlesThreshold;
 		static CHAOS_API int32 ComputeHorizonEpsilonFromMeshExtends;
-#else
-		static int32 PerformGeometryCheck;
-		static int32 PerformGeometryReduction;
-		static int32 VerticesThreshold;
-		static int32 ParticlesThreshold;
-		static int32 ComputeHorizonEpsilonFromMeshExtends;
-#endif
 
 	private:
 
