Renamed presented strata node to user.
Fixed shortcut category listing.

diff --git a/Engine/Source/Runtime/Engine/Classes/Materials/MaterialExpressionStrata.h b/Engine/Source/Runtime/Engine/Classes/Materials/MaterialExpressionStrata.h
--- a/Engine/Source/Runtime/Engine/Classes/Materials/MaterialExpressionStrata.h
+++ b/Engine/Source/Runtime/Engine/Classes/Materials/MaterialExpressionStrata.h
@@ -23,14 +23,13 @@ extern int32 CompileStrataBlendFunction(FMaterialCompiler* Compiler, const int32
 ///////////////////////////////////////////////////////////////////////////////
 // BSDF nodes
 
-// This would be needed to for a common node interface and weight input and normal too?
-UCLASS(MinimalAPI, collapsecategories, hidecategories = Object, Abstract)
+UCLASS(MinimalAPI, collapsecategories, hidecategories = Object, Abstract, DisplayName = "Strata Expression")
 class UMaterialExpressionStrataBSDF : public UMaterialExpression
 {
 	GENERATED_UCLASS_BODY()
 };
 
-UCLASS(MinimalAPI, collapsecategories, hidecategories = Object)
+UCLASS(MinimalAPI, collapsecategories, hidecategories = Object, DisplayName = "Strata Slab BSDF")
 class UMaterialExpressionStrataSlabBSDF : public UMaterialExpressionStrataBSDF
 {
 	GENERATED_UCLASS_BODY()
@@ -136,7 +135,7 @@ class UMaterialExpressionStrataSlabBSDF : public UMaterialExpressionStrataBSDF
 	//~ End UMaterialExpression Interface
 };
 
-UCLASS(MinimalAPI, collapsecategories, hidecategories = Object)
+UCLASS(MinimalAPI, collapsecategories, hidecategories = Object, DisplayName = "Strata Sheen BSDF")
 class UMaterialExpressionStrataSheenBSDF : public UMaterialExpressionStrataBSDF
 {
 	GENERATED_UCLASS_BODY()
@@ -171,7 +170,7 @@ class UMaterialExpressionStrataSheenBSDF : public UMaterialExpressionStrataBSDF
 	//~ End UMaterialExpression Interface
 };
 
-UCLASS(MinimalAPI, collapsecategories, hidecategories = Object)
+UCLASS(MinimalAPI, collapsecategories, hidecategories = Object, DisplayName = "Strata Volumetric-Fog-Cloud BSDF")
 class UMaterialExpressionStrataVolumetricFogCloudBSDF : public UMaterialExpressionStrataBSDF
 {
 	GENERATED_UCLASS_BODY()
@@ -212,7 +211,7 @@ class UMaterialExpressionStrataVolumetricFogCloudBSDF : public UMaterialExpressi
 	//~ End UMaterialExpression Interface
 };
 
-UCLASS(MinimalAPI, collapsecategories, hidecategories = Object)
+UCLASS(MinimalAPI, collapsecategories, hidecategories = Object, DisplayName = "Strata Unlit BSDF")
 class UMaterialExpressionStrataUnlitBSDF : public UMaterialExpressionStrataBSDF
 {
 	GENERATED_UCLASS_BODY()
@@ -241,7 +240,7 @@ class UMaterialExpressionStrataUnlitBSDF : public UMaterialExpressionStrataBSDF
 	//~ End UMaterialExpression Interface
 };
 
-UCLASS(MinimalAPI, collapsecategories, hidecategories = Object)
+UCLASS(MinimalAPI, collapsecategories, hidecategories = Object, DisplayName = "Strata Hair BSDF")
 class UMaterialExpressionStrataHairBSDF : public UMaterialExpressionStrataBSDF
 {
 	GENERATED_UCLASS_BODY()
@@ -300,7 +299,7 @@ class UMaterialExpressionStrataHairBSDF : public UMaterialExpressionStrataBSDF
 	//~ End UMaterialExpression Interface
 };
 
-UCLASS(MinimalAPI, collapsecategories, hidecategories = Object)
+UCLASS(MinimalAPI, collapsecategories, hidecategories = Object, DisplayName = "Strata Single Layer Water BSDF")
 class UMaterialExpressionStrataSingleLayerWaterBSDF : public UMaterialExpressionStrataBSDF
 {
 	GENERATED_UCLASS_BODY()
@@ -388,7 +387,7 @@ class UMaterialExpressionStrataSingleLayerWaterBSDF : public UMaterialExpression
 ///////////////////////////////////////////////////////////////////////////////
 // Operator nodes
 
-UCLASS(MinimalAPI, collapsecategories, hidecategories = Object)
+UCLASS(MinimalAPI, collapsecategories, hidecategories = Object, DisplayName = "Strata BSDF Horizontal Blend")
 class UMaterialExpressionStrataHorizontalMixing : public UMaterialExpressionStrataBSDF
 {
 	GENERATED_UCLASS_BODY()
@@ -423,7 +422,7 @@ class UMaterialExpressionStrataHorizontalMixing : public UMaterialExpressionStra
 	//~ End UMaterialExpression Interface
 };
 
-UCLASS(MinimalAPI, collapsecategories, hidecategories = Object)
+UCLASS(MinimalAPI, collapsecategories, hidecategories = Object, DisplayName = "Strata BSDF Vertical Layer")
 class UMaterialExpressionStrataVerticalLayering : public UMaterialExpressionStrataBSDF
 {
 	GENERATED_UCLASS_BODY()
@@ -452,7 +451,7 @@ class UMaterialExpressionStrataVerticalLayering : public UMaterialExpressionStra
 	//~ End UMaterialExpression Interface
 };
 
-UCLASS(MinimalAPI, collapsecategories, hidecategories = Object)
+UCLASS(MinimalAPI, collapsecategories, hidecategories = Object, DisplayName = "Strata BSDF Add")
 class UMaterialExpressionStrataAdd : public UMaterialExpressionStrataBSDF
 {
 	GENERATED_UCLASS_BODY()
@@ -481,7 +480,7 @@ class UMaterialExpressionStrataAdd : public UMaterialExpressionStrataBSDF
 	//~ End UMaterialExpression Interface
 };
 
-UCLASS(MinimalAPI, collapsecategories, hidecategories = Object)
+UCLASS(MinimalAPI, collapsecategories, hidecategories = Object, DisplayName = "Strata BSDF Weight")
 class UMaterialExpressionStrataMultiply : public UMaterialExpressionStrataBSDF
 {
 	GENERATED_UCLASS_BODY()
@@ -515,8 +514,8 @@ class UMaterialExpressionStrataMultiply : public UMaterialExpressionStrataBSDF
 ///////////////////////////////////////////////////////////////////////////////
 // Utilities
 
-UCLASS(MinimalAPI, collapsecategories, hidecategories = Object)
-class UMaterialExpressionStrataAnisotropyToRoughness : public UMaterialExpression
+UCLASS(MinimalAPI, collapsecategories, hidecategories = Object, DisplayName = "Strata Anisotropy-To-Roughnesses")
+class UMaterialExpressionStrataAnisotropyToRoughness : public UMaterialExpressionStrataBSDF
 {
 	GENERATED_UCLASS_BODY()
 
diff --git a/Engine/Source/Runtime/Engine/Private/Materials/MaterialExpressions.cpp b/Engine/Source/Runtime/Engine/Private/Materials/MaterialExpressions.cpp
--- a/Engine/Source/Runtime/Engine/Private/Materials/MaterialExpressions.cpp
+++ b/Engine/Source/Runtime/Engine/Private/Materials/MaterialExpressions.cpp
@@ -19719,15 +19719,6 @@ static int32 CompileWithDefaultTangentWS(class FMaterialCompiler* Compiler, FExp
 UMaterialExpressionStrataBSDF::UMaterialExpressionStrataBSDF(const FObjectInitializer& ObjectInitializer)
 	: Super(ObjectInitializer)
 {
-	struct FConstructorStatics
-	{
-		FText NAME_Strata;
-		FConstructorStatics() : NAME_Strata(LOCTEXT("Strata BSDFs", "Strata BSDFs")) { }
-	};
-	static FConstructorStatics ConstructorStatics;
-#if WITH_EDITORONLY_DATA
-	MenuCategories.Add(ConstructorStatics.NAME_Strata);
-#endif
 }
 
 UMaterialExpressionStrataSlabBSDF::UMaterialExpressionStrataSlabBSDF(const FObjectInitializer& ObjectInitializer)
@@ -20050,7 +20041,7 @@ int32 UMaterialExpressionStrataVolumetricFogCloudBSDF::Compile(class FMaterialCo
 
 void UMaterialExpressionStrataVolumetricFogCloudBSDF::GetCaption(TArray<FString>& OutCaptions) const
 {
-	OutCaptions.Add(TEXT("Strata Volumetric Fog Cloud BSDF"));
+	OutCaptions.Add(TEXT("Strata Volumetric-Fog-Cloud BSDF"));
 }
 
 uint32 UMaterialExpressionStrataVolumetricFogCloudBSDF::GetOutputType(int32 OutputIndex)
@@ -20400,7 +20391,7 @@ int32 UMaterialExpressionStrataHorizontalMixing::Compile(class FMaterialCompiler
 
 void UMaterialExpressionStrataHorizontalMixing::GetCaption(TArray<FString>& OutCaptions) const
 {
-	OutCaptions.Add(TEXT("Strata Horizontal Mixing"));
+	OutCaptions.Add(TEXT("Strata BSDF Horizontal Blend"));
 }
 
 uint32 UMaterialExpressionStrataHorizontalMixing::GetOutputType(int32 OutputIndex)
@@ -20476,7 +20467,7 @@ int32 UMaterialExpressionStrataVerticalLayering::Compile(class FMaterialCompiler
 
 void UMaterialExpressionStrataVerticalLayering::GetCaption(TArray<FString>& OutCaptions) const
 {
-	OutCaptions.Add(TEXT("Strata Vertical Layering"));
+	OutCaptions.Add(TEXT("Strata BSDF Vertical Layer"));
 }
 
 uint32 UMaterialExpressionStrataVerticalLayering::GetOutputType(int32 OutputIndex)
@@ -20552,7 +20543,7 @@ int32 UMaterialExpressionStrataAdd::Compile(class FMaterialCompiler* Compiler, i
 
 void UMaterialExpressionStrataAdd::GetCaption(TArray<FString>& OutCaptions) const
 {
-	OutCaptions.Add(TEXT("Strata Add"));
+	OutCaptions.Add(TEXT("Strata BSDF Add"));
 }
 
 uint32 UMaterialExpressionStrataAdd::GetOutputType(int32 OutputIndex)
@@ -20624,7 +20615,7 @@ int32 UMaterialExpressionStrataMultiply::Compile(class FMaterialCompiler* Compil
 
 void UMaterialExpressionStrataMultiply::GetCaption(TArray<FString>& OutCaptions) const
 {
-	OutCaptions.Add(TEXT("Strata Multiply"));
+	OutCaptions.Add(TEXT("Strata BSDF Weight"));
 }
 
 uint32 UMaterialExpressionStrataMultiply::GetOutputType(int32 OutputIndex)
@@ -20687,7 +20678,7 @@ int32 UMaterialExpressionStrataAnisotropyToRoughness::Compile(class FMaterialCom
 
 void UMaterialExpressionStrataAnisotropyToRoughness::GetCaption(TArray<FString>& OutCaptions) const
 {
-	OutCaptions.Add(TEXT("Strata Anisotropy to roughnesses"));
+	OutCaptions.Add(TEXT("Strata Anisotropy-To-Roughnesses"));
 }
 
 uint32 UMaterialExpressionStrataAnisotropyToRoughness::GetOutputType(int32 OutputIndex)
