Don't limit size of VT UDIM mip tale,
We've already clamped the max texture size for each source block, so this is already accounted for in the mip tale size

diff --git a/Engine/Source/Runtime/Engine/Private/VT/VirtualTextureDataBuilder.cpp b/Engine/Source/Runtime/Engine/Private/VT/VirtualTextureDataBuilder.cpp
--- a/Engine/Source/Runtime/Engine/Private/VT/VirtualTextureDataBuilder.cpp
+++ b/Engine/Source/Runtime/Engine/Private/VT/VirtualTextureDataBuilder.cpp
@@ -1061,7 +1061,7 @@ void FVirtualTextureDataBuilder::BuildSourcePixels(const FTextureSourceData& Sou
 			// Adjust the build settings to generate an uncompressed texture with mips but leave other settings
 			// like color correction, ... in place
 			FTextureBuildSettings TBSettings = SettingsPerLayer[0];
-			//TBSettings.MaxTextureResolution = TNumericLimits<uint32>::Max();
+			TBSettings.MaxTextureResolution = TNumericLimits<uint32>::Max(); // don't limit the size of the mip-tail, this limit only applies to each source block
 			TBSettings.TextureFormatName = LayerData.FormatName;
 			TBSettings.bSRGB = BuildSettingsForLayer.bSRGB;
 			TBSettings.bUseLegacyGamma = BuildSettingsForLayer.bUseLegacyGamma;
