Ensure curve validation for Transform Curve names is done before populating the DataModel. Otherwise curves missing from the skeleton would not be propagated to the model

diff --git a/Engine/Source/Runtime/Engine/Private/Animation/AnimSequenceBase.cpp b/Engine/Source/Runtime/Engine/Private/Animation/AnimSequenceBase.cpp
--- a/Engine/Source/Runtime/Engine/Private/Animation/AnimSequenceBase.cpp
+++ b/Engine/Source/Runtime/Engine/Private/Animation/AnimSequenceBase.cpp
@@ -88,6 +88,7 @@ void UAnimSequenceBase::PostLoad()
 
 			PRAGMA_DISABLE_DEPRECATION_WARNINGS
 			VerifyCurveNames<FFloatCurve>(*MySkeleton, USkeleton::AnimCurveMappingName, RawCurveData.FloatCurves);
+			VerifyCurveNames<FTransformCurve>(*MySkeleton, USkeleton::AnimTrackCurveMappingName, RawCurveData.TransformCurves);
 			PRAGMA_ENABLE_DEPRECATION_WARNINGS
 		}
 
