Add support in WorldPartition to resolve unloaded objects through FSoftObjectPath.
This is built on top of 15307418, which defer loading of objects to their toplevel object if not loaded.

diff --git a/Engine/Source/Runtime/Engine/Classes/Engine/Level.h b/Engine/Source/Runtime/Engine/Classes/Engine/Level.h
--- a/Engine/Source/Runtime/Engine/Classes/Engine/Level.h
+++ b/Engine/Source/Runtime/Engine/Classes/Engine/Level.h
@@ -781,6 +781,7 @@ public:
 	virtual void PostEditChangeProperty(FPropertyChangedEvent& PropertyChangedEvent) override;
 	virtual void BeginCacheForCookedPlatformData(const ITargetPlatform *TargetPlatform) override;
 	virtual bool CanEditChange(const FProperty* PropertyThatWillChange) const override;
+	virtual UObject* LoadSubobject(const TCHAR* SubObjectPath) override;
 #endif // WITH_EDITOR
 	virtual void PostLoad() override;
 	virtual void PreSave(const class ITargetPlatform* TargetPlatform) override;
diff --git a/Engine/Source/Runtime/Engine/Classes/Engine/World.h b/Engine/Source/Runtime/Engine/Classes/Engine/World.h
--- a/Engine/Source/Runtime/Engine/Classes/Engine/World.h
+++ b/Engine/Source/Runtime/Engine/Classes/Engine/World.h
@@ -2607,6 +2607,7 @@ public:
 	virtual bool Rename(const TCHAR* NewName = NULL, UObject* NewOuter = NULL, ERenameFlags Flags = REN_None) override;
 	virtual void GetAssetRegistryTags(TArray<FAssetRegistryTag>& OutTags) const override;
 	virtual bool IsNameStableForNetworking() const override;
+	virtual UObject* LoadSubobject(const TCHAR* SubObjectPath) override;
 #endif
 	virtual void PostDuplicate(bool bDuplicateForPIE) override;
 	//~ End UObject Interface
diff --git a/Engine/Source/Runtime/Engine/Private/Level.cpp b/Engine/Source/Runtime/Engine/Private/Level.cpp
--- a/Engine/Source/Runtime/Engine/Private/Level.cpp
+++ b/Engine/Source/Runtime/Engine/Private/Level.cpp
@@ -2609,6 +2609,21 @@ bool ULevel::CanEditChange(const FProperty* PropertyThatWillChange) const
 	return Super::CanEditChange(PropertyThatWillChange);
 }
 
+UObject* ULevel::LoadSubobject(const TCHAR* SubObjectPath)
+{
+	UObject* LoadedObject = StaticFindObject(nullptr, this, SubObjectPath);
+
+	if (!LoadedObject)
+	{
+		if (UWorldPartition* WorldPartition = GetWorldPartition())
+		{
+			LoadedObject = WorldPartition->LoadSubobject(SubObjectPath);
+		}
+	}
+
+	return LoadedObject;
+}
+
 void ULevel::FixupForPIE(int32 InPIEInstanceID, TFunctionRef<void(int32, FSoftObjectPath&)> InCustomFixupFunction)
 {
 	struct FSoftPathPIEFixupSerializer : public FArchiveUObject
diff --git a/Engine/Source/Runtime/Engine/Private/World.cpp b/Engine/Source/Runtime/Engine/Private/World.cpp
--- a/Engine/Source/Runtime/Engine/Private/World.cpp
+++ b/Engine/Source/Runtime/Engine/Private/World.cpp
@@ -7977,6 +7977,22 @@ bool UWorld::IsNameStableForNetworking() const
 	return bIsNameStableForNetworking || Super::IsNameStableForNetworking();
 }
 
+UObject* UWorld::LoadSubobject(const TCHAR* SubObjectPath)
+{
+	UObject* LoadedObject = nullptr;
+
+	FString SubObjectName;
+	FString SubObjectContext;	
+	if (FString(SubObjectPath).Split(TEXT("."), &SubObjectContext, &SubObjectName))
+	{
+		if (UObject* SubObject = StaticFindObject(nullptr, this, *SubObjectContext))
+		{
+			LoadedObject = SubObject->LoadSubobject(*SubObjectName);
+		}
+	}
+
+	return LoadedObject;
+}
 #endif
 
 FPrimaryAssetId UWorld::GetPrimaryAssetId() const
diff --git a/Engine/Source/Runtime/Engine/Private/WorldPartition/WorldPartition.cpp b/Engine/Source/Runtime/Engine/Private/WorldPartition/WorldPartition.cpp
--- a/Engine/Source/Runtime/Engine/Private/WorldPartition/WorldPartition.cpp
+++ b/Engine/Source/Runtime/Engine/Private/WorldPartition/WorldPartition.cpp
@@ -769,6 +769,8 @@ void UWorldPartition::BeginDestroy()
 	Super::BeginDestroy();
 
 #if WITH_EDITOR
+	LoadedSubobjects.Empty();
+
 	for (TUniquePtr<FWorldPartitionActorDesc>& ActorDescPtr : ActorDescList)
 	{
 		ActorDescPtr.Release();
@@ -776,6 +778,24 @@ void UWorldPartition::BeginDestroy()
 #endif
 }
 
+#if WITH_EDITOR
+UObject* UWorldPartition::LoadSubobject(const TCHAR* SubObjectPath)
+{
+	for (const auto& Pair : Actors)
+	{
+		const FWorldPartitionActorDesc* ActorDesc = Pair.Value->Get();
+
+		if (FString(ActorDesc->ActorPath.ToString()).EndsWith(SubObjectPath))
+		{
+			int32 ReferenceIndex = LoadedSubobjects.Add(Pair.Value);
+			return LoadedSubobjects[ReferenceIndex]->GetActor();
+		}
+	}
+
+	return nullptr;
+}
+#endif
+
 UWorldPartitionStreamingPolicy* UWorldPartition::GetStreamingPolicy() const
 {
 	check(GetWorld());
diff --git a/Engine/Source/Runtime/Engine/Public/WorldPartition/WorldPartition.h b/Engine/Source/Runtime/Engine/Public/WorldPartition/WorldPartition.h
--- a/Engine/Source/Runtime/Engine/Public/WorldPartition/WorldPartition.h
+++ b/Engine/Source/Runtime/Engine/Public/WorldPartition/WorldPartition.h
@@ -8,6 +8,7 @@
 #include "Templates/SubclassOf.h"
 #include "WorldPartition/WorldPartitionActorDesc.h"
 #include "WorldPartition/WorldPartitionStreamingSource.h"
+#include "WorldPartition/WorldPartitionHandle.h"
 #include "WorldPartition/ActorDescContainer.h"
 
 #if WITH_EDITOR
@@ -114,6 +115,10 @@ private:
 
 #if WITH_EDITOR
 public:
+	//~ Begin UObject Interface
+	virtual UObject* LoadSubobject(const TCHAR* SubObjectPath) override;
+	//~ End UObject Interface
+
 	FName GetWorldPartitionEditorName();
 	bool IsSimulating() const;
 
@@ -188,6 +193,8 @@ public:
 	// Default HLOD layer
 	UPROPERTY(EditAnywhere, Category=HLOD)
 	TObjectPtr<class UHLODLayer> DefaultHLODLayer;
+
+	TArray<FWorldPartitionReference> LoadedSubobjects;
 #endif
 
 private:
