Ease in/out for the target weight interpolation in BlendSpace.
This also allows per-bone target weight interpolation even if the overal weight is zero

diff --git a/Engine/Source/Runtime/Core/Public/Math/UnrealMathUtility.h b/Engine/Source/Runtime/Core/Public/Math/UnrealMathUtility.h
--- a/Engine/Source/Runtime/Core/Public/Math/UnrealMathUtility.h
+++ b/Engine/Source/Runtime/Core/Public/Math/UnrealMathUtility.h
@@ -1135,6 +1135,64 @@ public:
 	/** Interpolate quaternion from Current to Target. Scaled by angle to Target, so it has a strong start speed and ease out. */
 	static CORE_API FQuat QInterpTo(const FQuat& Current, const FQuat& Target, float DeltaTime, float InterpSpeed);
 
+	/**
+	 * Smooths a value using a critically damped spring. Works for any type that supports basic arithmetic operations.
+	 * 
+	 * Note that InSmoothingTime is the time lag when tracking constant motion (so if you tracked a predicted position 
+	 * TargetVel * InSmoothingTime ahead, you'd match the target position)
+	 * 
+	 * When starting from zero velocity, the maximum velocity is reached after time InSmoothingTime * 0.5
+	 *
+	 * @param  InOutValue      The value to be smoothed
+	 * @param  InOutValueRate  The rate of change of the value
+	 * @para   InTargetValue   The target to smooth towards
+	 * @param  InDeltaTime     Time interval
+	 * @param  InSmoothingTime Timescale over which to smooth. Larger values result in more smoothed behaviour. Can be zero.
+	 */
+	template< class T >
+	static void CriticallyDampedSmoothing(
+		T& InOutValue,
+		T& InOutValueRate,
+		const T& InTargetValue,
+		float InDeltaTime,
+		float InSmoothingTime)
+	{
+		if (InSmoothingTime > 0.0f)
+		{
+			// The closed form solution for critically damped motion towards zero is:
+			//
+			// x = ( x0 + t (v0 + w x0) ) exp(-w t)
+			//
+			// where w is the natural frequency, x0, v0 are x and dx/dt at time 0 (e.g. Brian Stone - Chatter and Machine Tools pdf)
+			//
+			// Differentiate to get velocity (remember d(v0)/dt = 0):
+			//
+			// v = ( v0 - w t (v0 + w x0) ) exp(-w t)
+			//
+			// We're damping towards a target, so convert - i.e. x = value - target
+			// 
+			// SmoothingTime is 
+			float W = 2.0f / InSmoothingTime;
+			float A = W * InDeltaTime;
+			// Taylor expansion for Exp(-W * InDeltaTime)
+			float Exp = 1.0f / (1.0f + A + 0.48f * A * A + 0.235f * A * A * A);
+			T X0 = InOutValue - InTargetValue;
+			T B = (InOutValueRate + X0 * W) * InDeltaTime;
+			InOutValue = InTargetValue + (X0 + B) * Exp;
+			InOutValueRate = (InOutValueRate - B * W) * Exp;
+		}
+		else if (InDeltaTime > 0.0f)
+		{
+			InOutValueRate = (InTargetValue - InOutValue) / InDeltaTime;
+			InOutValue = InTargetValue;
+		}
+		else
+		{
+			InOutValue = InTargetValue;
+			InOutValueRate = T(0);
+		}
+	}
+
 	/**
 	 * Simple function to create a pulsating scalar value
 	 *
diff --git a/Engine/Source/Runtime/Engine/Classes/Animation/AnimationAsset.h b/Engine/Source/Runtime/Engine/Classes/Animation/AnimationAsset.h
--- a/Engine/Source/Runtime/Engine/Classes/Animation/AnimationAsset.h
+++ b/Engine/Source/Runtime/Engine/Classes/Animation/AnimationAsset.h
@@ -82,6 +82,10 @@ struct FBlendSampleData
 	UPROPERTY()
 	float TotalWeight;
 
+	// Rate of change of the Weight - used in smoothed BlendSpace blends
+	UPROPERTY()
+	float WeightRate;
+
 	UPROPERTY()
 	float Time;
 
@@ -95,24 +99,29 @@ struct FBlendSampleData
 
 	FMarkerTickRecord MarkerTickRecord;
 
-	// transient perbone interpolation data
+	// transient per-bone interpolation data
 	TArray<float> PerBoneBlendData;
 
+	// transient per-bone weight rate - only allocated when used
+	TArray<float> PerBoneWeightRate;
+
 	FBlendSampleData()
-		:	SampleDataIndex(0)
-		,	Animation(nullptr)
-		,	TotalWeight(0.f)
-		,	Time(0.f)
-		,	PreviousTime(0.f)
-		,	SamplePlayRate(0.0f)
+		: SampleDataIndex(0)
+		, Animation(nullptr)
+		, TotalWeight(0.f)
+		, WeightRate(0.f)
+		, Time(0.f)
+		, PreviousTime(0.f)
+		, SamplePlayRate(0.0f)
 	{}
 	FBlendSampleData(int32 Index)
-		:	SampleDataIndex(Index)
-		,	Animation(nullptr)
-		,	TotalWeight(0.f)
-		,	Time(0.f)
-		,	PreviousTime(0.f)
-		,	SamplePlayRate(0.0f)
+		: SampleDataIndex(Index)
+		, Animation(nullptr)
+		, TotalWeight(0.f)
+		, WeightRate(0.f)
+		, Time(0.f)
+		, PreviousTime(0.f)
+		, SamplePlayRate(0.0f)
 	{}
 	bool operator==( const FBlendSampleData& Other ) const 
 	{
diff --git a/Engine/Source/Runtime/Engine/Classes/Animation/BlendSpaceBase.h b/Engine/Source/Runtime/Engine/Classes/Animation/BlendSpaceBase.h
--- a/Engine/Source/Runtime/Engine/Classes/Animation/BlendSpaceBase.h
+++ b/Engine/Source/Runtime/Engine/Classes/Animation/BlendSpaceBase.h
@@ -204,13 +204,17 @@ struct FPerBoneInterpolation
 	FBoneReference BoneReference;
 
 	/**
-	* This is the speed at which we interpolate towards the target weights for this specific bone, measured in 'how many times per second' we can get to the target.
-	* A value of 0 means it would instantly set itself to the target value, while a value of one means it will take one second to get there.
-	* A value of 2 would mean it goes there in twice the speed of a second, so in half a second, a value of 3 would mean in a third of a second, and so on.
-	* Smaller values mean slower interpolation speeds.
-	* This value overrides the global interpolation speed, so the global interpolation speed has no impact anymore on the interpolation speed of this bone.
+	* If greater than zero, this is the speed at which the sample weights are allowed to change for this specific bone.
+	*
+	* A speed of 1 means a sample weight can change from zero to one (or one to zero) in one second.
+	* A speed of 2 means that this would take half a second.
+	*
+	* Smaller values mean slower adjustments of the sample weights, and thus more smoothing. However, a
+	* value of zero disables this smoothing entirely.
+	* 
+	* If set, the value overrides the overall Sample Weight Speed which will no longer affect this bone.
 	*/
-	UPROPERTY(EditAnywhere, Category=FPerBoneInterpolation, meta=(DisplayName="Interpolation Speed"))
+	UPROPERTY(EditAnywhere, Category=FPerBoneInterpolation, meta=(DisplayName="Sample Weight Speed"))
 	float InterpolationSpeedPerSec;
 
 	FPerBoneInterpolation()
@@ -479,17 +483,31 @@ public:
 	FInterpolationParameter	InterpolationParam[3];
 
 	/**
-	* This is the speed at which we interpolate towards the target weights, measured in 'how many times per second' we can get to the target.
-	* A value of 0 means it would instantly set itself to the target value, while a value of one means it will take one second to get there.
-	* A value of 2 would mean it goes there in twice the speed of a second, so in half a second, a value of 3 would mean in a third of a second, and so on.
-	* Smaller values mean slower interpolation speeds.
-	* Imagine we have a blend space for locomotion, moving left, forward and right. Now if you interpolate the inputs of the blend space itself, from one extreme to the other, you will
-	* go from left, to forward, to right. As an alternative, by setting this global interpolation speed to a value higher than zero, it will go directly from left to right, without going through moving forward first.
+	* If greater than zero, this is the speed at which the sample weights are allowed to change.
+	* 
+	* A speed of 1 means a sample weight can change from zero to one (or one to zero) in one second.
+	* A speed of 2 means that this would take half a second.
+	* 
+	* This allows the Blend Space to switch to new parameters without going through intermediate states, 
+	* effectively blending between where it was and where the new target is. For example, imagine we have 
+	* a blend space for locomotion, moving left, forward and right. Now if you interpolate the inputs of 
+	* the blend space itself, from one extreme to the other, you will go from left, to forward, to right. 
+	* As an alternative, by setting this Sample Weight Speed to a value higher than zero, it will go 
+	* directly from left to right, without going through moving forward first.
+	* 
+	* Smaller values mean slower adjustments of the sample weights, and thus more smoothing. However, a 
+	* value of zero disables this smoothing entirely.
 	*/
-	UPROPERTY(EditAnywhere, Category = SampleInterpolation, meta = (DisplayName = "Global Interpolation Speed"))
+	UPROPERTY(EditAnywhere, Category = SampleInterpolation, meta = (DisplayName = "Sample Weight Speed"))
 	float TargetWeightInterpolationSpeedPerSec;
 
-	/** The current mode used by the blendspace to decide which animation notifies to fire. Valid options are:
+	/**
+	 * If set then this eases in/out the sample weight adjustments, using the speed to determine how much smoothing to apply.
+	 */
+	UPROPERTY(EditAnywhere, Category = SampleInterpolation, meta = (DisplayName = "Sample Weight Speed Smoothing "))
+	bool bTargetWeightInterpolationEaseInOut;
+
+	/** The current mode used by the BlendSpace to decide which animation notifies to fire. Valid options are:
 	- AllAnimations - All notify events will fire
 	- HighestWeightedAnimation - Notify events will only fire from the highest weighted animation
 	- None - No notify events will fire from any animations
diff --git a/Engine/Source/Runtime/Engine/Private/Animation/BlendSpaceBase.cpp b/Engine/Source/Runtime/Engine/Private/Animation/BlendSpaceBase.cpp
--- a/Engine/Source/Runtime/Engine/Private/Animation/BlendSpaceBase.cpp
+++ b/Engine/Source/Runtime/Engine/Private/Animation/BlendSpaceBase.cpp
@@ -19,7 +19,7 @@
 
 DECLARE_CYCLE_STAT(TEXT("BlendSpace GetAnimPose"), STAT_BlendSpace_GetAnimPose, STATGROUP_Anim);
 
-/** Scratch buffers for multithreaded usage */
+/** Scratch buffers for multi-threaded usage */
 struct FBlendSpaceScratchData : public TThreadSingleton<FBlendSpaceScratchData>
 {
 	TArray<FBlendSampleData> OldSampleDataList;
@@ -144,7 +144,7 @@ bool UBlendSpaceBase::UpdateBlendSamples_Internal(const FVector& InBlendSpacePos
 	if (GetSamplesFromBlendInput(InBlendSpacePosition, NewSampleDataList))
 	{
 		// if target weight interpolation is set
-		if (TargetWeightInterpolationSpeedPerSec > 0.f)
+		if (TargetWeightInterpolationSpeedPerSec > 0.f || PerBoneBlend.Num() > 0)
 		{
 			// target weight interpolation
 			if (InterpolateWeightOfSampleData(InDeltaTime, InOutOldSampleDataList, NewSampleDataList, InOutSampleDataCache))
@@ -748,7 +748,7 @@ bool UBlendSpaceBase::GetSamplesFromBlendInput(const FVector &BlendInput, TArray
 	OutSampleDataList.Reset();
 	OutSampleDataList.Reserve(RawGridSamples.Num() * FEditorElement::MAX_VERTICES);
 
-	// consolidate all samples
+	// Consolidate all samples
 	for (int32 SampleNum=0; SampleNum<RawGridSamples.Num(); ++SampleNum)
 	{
 		FGridBlendSample& GridSample = RawGridSamples[SampleNum];
@@ -770,7 +770,8 @@ bool UBlendSpaceBase::GetSamplesFromBlendInput(const FVector &BlendInput, TArray
 		}
 	}
 
-	// go through merge down to first sample 
+	// At this point we'll only have one of each sample, but different samples can point to the same
+	// animation. We can combine those, making sure to interpolate the parameters like play rate too.
 	for (int32 Index1 = 0; Index1 < OutSampleDataList.Num(); ++Index1)
 	{
 		FBlendSampleData& FirstSample = OutSampleDataList[Index1];
@@ -778,7 +779,7 @@ bool UBlendSpaceBase::GetSamplesFromBlendInput(const FVector &BlendInput, TArray
 		{
 			FBlendSampleData& SecondSample = OutSampleDataList[Index2];
 			// if they have same sample, remove the Index2, and get out
-			if (FirstSample.SampleDataIndex == SecondSample.SampleDataIndex || 
+			if (FirstSample.SampleDataIndex == SecondSample.SampleDataIndex || // Shouldn't happen
 				(FirstSample.Animation != nullptr && FirstSample.Animation == SecondSample.Animation))
 			{
 				//Calc New Sample Playrate
@@ -1272,11 +1273,71 @@ const FEditorElement* UBlendSpaceBase::GetGridSampleInternal(int32 Index) const
 	return GridSamples.IsValidIndex(Index) ? &GridSamples[Index] : NULL;
 }
 
-bool UBlendSpaceBase::InterpolateWeightOfSampleData(float DeltaTime, const TArray<FBlendSampleData> & OldSampleDataList, const TArray<FBlendSampleData> & NewSampleDataList, TArray<FBlendSampleData> & FinalSampleDataList) const
+// When using CriticallyDampedSmoothing, how to go from the interpolation speed to the smooth
+// time? What would the critically damped velocity be as it goes from a starting value of 0 to a
+// target of 1 (see eq in CriticallyDampedSmoothing), starting with v = 0?
+//
+// v = W^2 t exp(-W t)
+//
+// Differentiate and set equal to zero to find maximum v is at t = 1 / W
+//
+// vMax = W / e = 2 / (SmoothingTime * e)
+//
+// Set this equal to TargetWeightInterpolationSpeedPerSec, we get
+//
+// SmoothingTime = 2 / (e * TargetWeightInterpolationSpeedPerSec)
+//
+// However - this looks significantly slower than when using a constant speed, because we're
+// easing in/out, so aim for twice this speed (i.e. smooth over half the time)
+static float SmoothingTimeFromSpeed(float Speed)
+{
+	return 1.0f / (EULERS_NUMBER * Speed);	
+}
+
+static void SmoothWeight(FBlendSampleData& Output, const FBlendSampleData& Input, float TargetWeight, float DeltaTime, float Speed, bool bUseEaseInOut)
+{
+	if (Speed <= 0.0f)
+	{
+		Output.TotalWeight = TargetWeight;
+		return;
+	}
+
+	if (bUseEaseInOut)
+	{
+		Output.TotalWeight = Input.TotalWeight;
+		Output.WeightRate = Input.WeightRate;
+		FMath::CriticallyDampedSmoothing(Output.TotalWeight, Output.WeightRate, TargetWeight, DeltaTime, SmoothingTimeFromSpeed(Speed));
+	}
+	else
+	{
+		Output.TotalWeight = FMath::FInterpConstantTo(Input.TotalWeight, TargetWeight, DeltaTime, Speed);
+	}
+}
+
+static void SmoothWeight(float& Output, float& OutputRate, float Input, float InputRate, float Target, float DeltaTime, float Speed, bool bUseEaseInOut)
 {
-	check (TargetWeightInterpolationSpeedPerSec > 0.f);
+	if (Speed <= 0.0f)
+	{
+		Output = Target;
+		return;
+	}
+
+	if (bUseEaseInOut)
+	{
+		Output = Input;
+		OutputRate = InputRate;
+		FMath::CriticallyDampedSmoothing(Output, OutputRate, Target, DeltaTime, SmoothingTimeFromSpeed(Speed));
+	}
+	else
+	{
+		Output = FMath::FInterpConstantTo(Input, Target, DeltaTime, Speed);
+	}
+}
 
+bool UBlendSpaceBase::InterpolateWeightOfSampleData(float DeltaTime, const TArray<FBlendSampleData> & OldSampleDataList, const TArray<FBlendSampleData> & NewSampleDataList, TArray<FBlendSampleData> & FinalSampleDataList) const
+{
 	float TotalFinalWeight = 0.f;
+	float TotalFinalPerBoneWeight = 0.0f;
 
 	// now interpolate from old to new target, this is brute-force
 	for (auto OldIt=OldSampleDataList.CreateConstIterator(); OldIt; ++OldIt)
@@ -1288,9 +1349,9 @@ bool UBlendSpaceBase::InterpolateWeightOfSampleData(float DeltaTime, const TArra
 		if (OldSample.PerBoneBlendData.Num()!=PerBoneBlend.Num())
 		{
 			OldSample.PerBoneBlendData.Init(OldSample.TotalWeight, PerBoneBlend.Num());
+			OldSample.PerBoneWeightRate.Init(OldSample.WeightRate, PerBoneBlend.Num());
 		}
 
-		// i'd like to change this later
 		for (auto NewIt=NewSampleDataList.CreateConstIterator(); NewIt; ++NewIt)
 		{
 			const FBlendSampleData& NewSample = *NewIt;
@@ -1298,26 +1359,29 @@ bool UBlendSpaceBase::InterpolateWeightOfSampleData(float DeltaTime, const TArra
 			if (NewSample.SampleDataIndex == OldSample.SampleDataIndex)
 			{
 				FBlendSampleData InterpData = NewSample;
-				InterpData.TotalWeight = FMath::FInterpConstantTo(OldSample.TotalWeight, NewSample.TotalWeight, DeltaTime, TargetWeightInterpolationSpeedPerSec);
+				SmoothWeight(InterpData, OldSample, NewSample.TotalWeight, DeltaTime, TargetWeightInterpolationSpeedPerSec, bTargetWeightInterpolationEaseInOut);
 				InterpData.PerBoneBlendData = OldSample.PerBoneBlendData;
+				InterpData.PerBoneWeightRate = OldSample.PerBoneWeightRate;
 
 				// now interpolate the per bone weights
+				float TotalPerBoneWeight = 0.0f;
 				for (int32 Iter = 0; Iter<InterpData.PerBoneBlendData.Num(); ++Iter)
 				{
-					if (PerBoneBlend[Iter].InterpolationSpeedPerSec > 0.f)
-					{
-						InterpData.PerBoneBlendData[Iter] = FMath::FInterpConstantTo(OldSample.PerBoneBlendData[Iter], NewSample.TotalWeight, DeltaTime, PerBoneBlend[Iter].InterpolationSpeedPerSec);
-					}
-					else
-					{
-						InterpData.PerBoneBlendData[Iter] = NewSample.TotalWeight;
-					}
+					SmoothWeight(
+						InterpData.PerBoneBlendData[Iter], InterpData.PerBoneWeightRate[Iter], 
+						OldSample.PerBoneBlendData[Iter], OldSample.PerBoneWeightRate[Iter], NewSample.TotalWeight, 
+						DeltaTime, PerBoneBlend[Iter].InterpolationSpeedPerSec, bTargetWeightInterpolationEaseInOut);
+					TotalPerBoneWeight += InterpData.PerBoneBlendData[Iter];
 				}
 
-				FinalSampleDataList.Add(InterpData);
-				TotalFinalWeight += InterpData.GetWeight();
-				bTargetSampleExists = true;
-				break;
+				if (InterpData.TotalWeight > ZERO_ANIMWEIGHT_THRESH || TotalPerBoneWeight > ZERO_ANIMWEIGHT_THRESH)
+				{
+					FinalSampleDataList.Add(InterpData);
+					TotalFinalWeight += InterpData.GetWeight();
+					TotalFinalPerBoneWeight += TotalPerBoneWeight;
+					bTargetSampleExists = true;
+					break;
+				}
 			}
 		}
 
@@ -1325,25 +1389,24 @@ bool UBlendSpaceBase::InterpolateWeightOfSampleData(float DeltaTime, const TArra
 		if (bTargetSampleExists == false)
 		{
 			FBlendSampleData InterpData = OldSample;
-			InterpData.TotalWeight = FMath::FInterpConstantTo(OldSample.TotalWeight, 0.f, DeltaTime, TargetWeightInterpolationSpeedPerSec);
+			SmoothWeight(InterpData, OldSample, 0.0f, DeltaTime, TargetWeightInterpolationSpeedPerSec, bTargetWeightInterpolationEaseInOut);
 			// now interpolate the per bone weights
+			float TotalPerBoneWeight = 0.0f;
 			for (int32 Iter = 0; Iter<InterpData.PerBoneBlendData.Num(); ++Iter)
 			{
-				if (PerBoneBlend[Iter].InterpolationSpeedPerSec > 0.f)
-				{
-					InterpData.PerBoneBlendData[Iter] = FMath::FInterpConstantTo(OldSample.PerBoneBlendData[Iter], 0.f, DeltaTime, PerBoneBlend[Iter].InterpolationSpeedPerSec);
-				}
-				else
-				{
-					InterpData.PerBoneBlendData[Iter] = 0.f;
-				}
+				SmoothWeight(
+					InterpData.PerBoneBlendData[Iter], InterpData.PerBoneWeightRate[Iter],
+					OldSample.PerBoneBlendData[Iter], OldSample.PerBoneWeightRate[Iter], 0.0f,
+					DeltaTime, PerBoneBlend[Iter].InterpolationSpeedPerSec, bTargetWeightInterpolationEaseInOut);
+				TotalPerBoneWeight += InterpData.PerBoneBlendData[Iter];
 			}
 
 			// add it if it's not zero
-			if ( InterpData.TotalWeight > ZERO_ANIMWEIGHT_THRESH )
+			if ( InterpData.TotalWeight > ZERO_ANIMWEIGHT_THRESH || TotalPerBoneWeight > ZERO_ANIMWEIGHT_THRESH )
 			{
 				FinalSampleDataList.Add(InterpData);
 				TotalFinalWeight += InterpData.GetWeight();
+				TotalFinalPerBoneWeight += TotalPerBoneWeight;
 			}
 		}
 	}
@@ -1358,6 +1421,7 @@ bool UBlendSpaceBase::InterpolateWeightOfSampleData(float DeltaTime, const TArra
 		if (OldSample.PerBoneBlendData.Num()!=PerBoneBlend.Num())
 		{
 			OldSample.PerBoneBlendData.Init(OldSample.TotalWeight, PerBoneBlend.Num());
+			OldSample.PerBoneWeightRate.Init(OldSample.WeightRate, PerBoneBlend.Num());
 		}
 
 		for (auto NewIt=FinalSampleDataList.CreateConstIterator(); NewIt; ++NewIt)
@@ -1374,25 +1438,33 @@ bool UBlendSpaceBase::InterpolateWeightOfSampleData(float DeltaTime, const TArra
 		if (bOldSampleExists == false)
 		{
 			FBlendSampleData InterpData = OldSample;
-			InterpData.TotalWeight = FMath::FInterpConstantTo(0.f, OldSample.TotalWeight, DeltaTime, TargetWeightInterpolationSpeedPerSec);
+			float TargetWeight = InterpData.TotalWeight;
+			OldSample.TotalWeight = 0.0f;
+			OldSample.WeightRate = 0.0f;
+			SmoothWeight(InterpData, OldSample, TargetWeight, DeltaTime, TargetWeightInterpolationSpeedPerSec, bTargetWeightInterpolationEaseInOut);
 			// now interpolate the per bone weights
+			float TotalPerBoneWeight = 0.0f;
 			for (int32 Iter = 0; Iter<InterpData.PerBoneBlendData.Num(); ++Iter)
 			{
-				if (PerBoneBlend[Iter].InterpolationSpeedPerSec > 0.f)
-				{
-					InterpData.PerBoneBlendData[Iter] = FMath::FInterpConstantTo(0.f, OldSample.PerBoneBlendData[Iter], DeltaTime, PerBoneBlend[Iter].InterpolationSpeedPerSec);
-				}
-				else
-				{
-					InterpData.PerBoneBlendData[Iter] = OldSample.PerBoneBlendData[Iter];
-				}
+				float Target = OldSample.PerBoneBlendData[Iter];
+				OldSample.PerBoneBlendData[Iter] = 0.0f;
+				OldSample.PerBoneWeightRate[Iter] = 0.0f;
+				SmoothWeight(
+					InterpData.PerBoneBlendData[Iter], InterpData.PerBoneWeightRate[Iter],
+					OldSample.PerBoneBlendData[Iter], OldSample.PerBoneWeightRate[Iter], Target,
+					DeltaTime, PerBoneBlend[Iter].InterpolationSpeedPerSec, bTargetWeightInterpolationEaseInOut);
+				TotalPerBoneWeight += InterpData.PerBoneBlendData[Iter];
+			}
+			if (InterpData.TotalWeight > ZERO_ANIMWEIGHT_THRESH || TotalPerBoneWeight > ZERO_ANIMWEIGHT_THRESH)
+			{
+				FinalSampleDataList.Add(InterpData);
+				TotalFinalWeight += InterpData.GetWeight();
+				TotalFinalPerBoneWeight += TotalPerBoneWeight;
 			}
-			FinalSampleDataList.Add(InterpData);
-			TotalFinalWeight += InterpData.GetWeight();
 		}
 	}
 
-	return (TotalFinalWeight > ZERO_ANIMWEIGHT_THRESH);
+	return TotalFinalWeight > ZERO_ANIMWEIGHT_THRESH || TotalFinalPerBoneWeight > ZERO_ANIMWEIGHT_THRESH;
 }
 
 FVector UBlendSpaceBase::FilterInput(FBlendFilter * Filter, const FVector& BlendInput, float DeltaTime) const
diff --git a/Engine/Source/Runtime/Engine/Private/AnimationAsset.cpp b/Engine/Source/Runtime/Engine/Private/AnimationAsset.cpp
--- a/Engine/Source/Runtime/Engine/Private/AnimationAsset.cpp
+++ b/Engine/Source/Runtime/Engine/Private/AnimationAsset.cpp
@@ -707,7 +707,7 @@ void FBlendSampleData::NormalizeDataWeight(TArray<FBlendSampleData>& SampleDataL
 	}
 
 	// Re-normalize Pose weight
-	if (ensure(TotalSum > ZERO_ANIMWEIGHT_THRESH))
+	if (TotalSum > ZERO_ANIMWEIGHT_THRESH)
 	{
 		if (FMath::Abs<float>(TotalSum - 1.f) > ZERO_ANIMWEIGHT_THRESH)
 		{
@@ -717,11 +717,18 @@ void FBlendSampleData::NormalizeDataWeight(TArray<FBlendSampleData>& SampleDataL
 			}
 		}
 	}
+	else
+	{
+		for (int32 PoseIndex = 0; PoseIndex < SampleDataList.Num(); PoseIndex++)
+		{
+			SampleDataList[PoseIndex].TotalWeight = 1.0f / SampleDataList.Num();
+		}
+	}
 
 	// Re-normalize per bone weights.
 	for (int32 BoneIndex = 0; BoneIndex < NumBones; BoneIndex++)
 	{
-		if (ensure(PerBoneTotalSums[BoneIndex] > ZERO_ANIMWEIGHT_THRESH))
+		if (PerBoneTotalSums[BoneIndex] > ZERO_ANIMWEIGHT_THRESH)
 		{
 			if (FMath::Abs<float>(PerBoneTotalSums[BoneIndex] - 1.f) > ZERO_ANIMWEIGHT_THRESH)
 			{
@@ -731,6 +738,13 @@ void FBlendSampleData::NormalizeDataWeight(TArray<FBlendSampleData>& SampleDataL
 				}
 			}
 		}
+		else
+		{
+			for (int32 PoseIndex = 0; PoseIndex < SampleDataList.Num(); PoseIndex++)
+			{
+				SampleDataList[PoseIndex].PerBoneBlendData[BoneIndex] = 1.0f / SampleDataList.Num();
+			}
+		}
 	}
 }
 
