Actor Descriptors Containers
- Moved actor descriptors accessors from UWorldPartition to UActorDescContainer.
- Moved TWorldPartitionActorDescIterator into UActorDescContainer and made const and non-const versions.
- Made UActorDescContainer actor descriptors list private and converted all existing code to use new iterators.

diff --git a/Engine/Source/Editor/UnrealEd/Private/WorldPartition/WorldPartitionHLODsBuilder.cpp b/Engine/Source/Editor/UnrealEd/Private/WorldPartition/WorldPartitionHLODsBuilder.cpp
--- a/Engine/Source/Editor/UnrealEd/Private/WorldPartition/WorldPartitionHLODsBuilder.cpp
+++ b/Engine/Source/Editor/UnrealEd/Private/WorldPartition/WorldPartitionHLODsBuilder.cpp
@@ -11,7 +11,6 @@
 #include "WorldPartition/WorldPartitionSubsystem.h"
 #include "WorldPartition/HLOD/HLODActor.h"
 #include "WorldPartition/HLOD/HLODActorDesc.h"
-#include "WorldPartition/WorldPartitionActorDescIterator.h"
 
 DEFINE_LOG_CATEGORY_STATIC(LogWorldPartitionHLODsBuilder, All, All);
 
diff --git a/Engine/Source/Runtime/Engine/Private/WorldPartition/ActorDescContainer.cpp b/Engine/Source/Runtime/Engine/Private/WorldPartition/ActorDescContainer.cpp
--- a/Engine/Source/Runtime/Engine/Private/WorldPartition/ActorDescContainer.cpp
+++ b/Engine/Source/Runtime/Engine/Private/WorldPartition/ActorDescContainer.cpp
@@ -196,6 +196,18 @@ void UActorDescContainer::OnAssetUpdated(const FAssetData& InAssetData)
 	}
 }
 
+const FWorldPartitionActorDesc* UActorDescContainer::GetActorDesc(const FGuid& Guid) const
+{
+	const TUniquePtr<FWorldPartitionActorDesc>* const * ActorDesc = Actors.Find(Guid);
+	return ActorDesc ? (*ActorDesc)->Get() : nullptr;
+}
+
+FWorldPartitionActorDesc* UActorDescContainer::GetActorDesc(const FGuid& Guid)
+{
+	TUniquePtr<FWorldPartitionActorDesc>** ActorDesc = Actors.Find(Guid);
+	return ActorDesc ? (*ActorDesc)->Get() : nullptr;
+}
+
 void UActorDescContainer::RegisterDelegates()
 {
 	if (GEditor && !IsTemplate())
@@ -220,5 +232,4 @@ void UActorDescContainer::UnregisterDelegates()
 		}
 	}
 }
-
 #endif
\ No newline at end of file
diff --git a/Engine/Source/Runtime/Engine/Private/WorldPartition/RuntimeSpatialHash/RuntimeSpatialHashActorCluster.cpp b/Engine/Source/Runtime/Engine/Private/WorldPartition/RuntimeSpatialHash/RuntimeSpatialHashActorCluster.cpp
--- a/Engine/Source/Runtime/Engine/Private/WorldPartition/RuntimeSpatialHash/RuntimeSpatialHashActorCluster.cpp
+++ b/Engine/Source/Runtime/Engine/Private/WorldPartition/RuntimeSpatialHash/RuntimeSpatialHashActorCluster.cpp
@@ -147,9 +147,9 @@ TArray<FActorCluster> CreateActorClustersImpl(UWorldPartition* WorldPartition, T
 		LevelScriptExternalActorReferences.Append(ActorsReferencesUtils::GetExternalActorReferences(LevelScriptBlueprint));
 	}
 
-	for (const auto& Pair : WorldPartition->Actors)
+	for (UActorDescContainer::TIterator ActorDescIterator(WorldPartition); ActorDescIterator; ++ActorDescIterator)
 	{
-		const FWorldPartitionActorDesc* ActorDesc = Pair.Value->Get();
+		FWorldPartitionActorDesc* ActorDesc = *ActorDescIterator;
 		EActorGridPlacement GridPlacement = ActorDesc->GetGridPlacement();
 
 		// Check if the actor is loaded (potentially referenced by the level script)
diff --git a/Engine/Source/Runtime/Engine/Private/WorldPartition/RuntimeSpatialHash/RuntimeSpatialHashHLOD.cpp b/Engine/Source/Runtime/Engine/Private/WorldPartition/RuntimeSpatialHash/RuntimeSpatialHashHLOD.cpp
--- a/Engine/Source/Runtime/Engine/Private/WorldPartition/RuntimeSpatialHash/RuntimeSpatialHashHLOD.cpp
+++ b/Engine/Source/Runtime/Engine/Private/WorldPartition/RuntimeSpatialHash/RuntimeSpatialHashHLOD.cpp
@@ -5,7 +5,6 @@
 #if WITH_EDITOR
 
 #include "WorldPartition/RuntimeSpatialHash/RuntimeSpatialHashGridHelper.h"
-#include "WorldPartition/WorldPartitionActorDescIterator.h"
 #include "WorldPartition/HLOD/HLODBuilder.h"
 #include "WorldPartition/HLOD/HLODLayer.h"
 #include "WorldPartition/HLOD/HLODActor.h"
@@ -477,17 +476,22 @@ bool UWorldPartitionRuntimeSpatialHash::GenerateHLOD(ISourceControlHelper* Sourc
 	// Create HLODs generation context
 	FHLODGenerationContext Context;
 	TArray<FWorldPartitionHandle> InvalidHLODActors;
-	for (TWorldPartitionActorDescIterator<AWorldPartitionHLOD, FHLODActorDesc> HLODIterator(WorldPartition); HLODIterator; ++HLODIterator)
-	{
-		FWorldPartitionHandle HLODActorHandle(WorldPartition, HLODIterator->GetGuid());
 
-		if (HLODIterator->GetCellHash())
-		{
-			Context.HLODActorDescs.Add(HLODIterator->GetCellHash(), HLODActorHandle);
-		}
-		else
+	for (UActorDescContainer::TIterator HLODIterator(WorldPartition); HLODIterator; ++HLODIterator)
+	{
+		if (HLODIterator->GetActorClass()->IsChildOf<AWorldPartitionHLOD>())
 		{
-			InvalidHLODActors.Add(HLODActorHandle);
+			FHLODActorDesc* HLODActorDesc = (FHLODActorDesc*)*HLODIterator;
+			FWorldPartitionHandle HLODActorHandle(WorldPartition, HLODActorDesc->GetGuid());
+
+			if (HLODActorDesc->GetCellHash())
+			{
+				Context.HLODActorDescs.Add(HLODActorDesc->GetCellHash(), HLODActorHandle);
+			}
+			else
+			{
+				InvalidHLODActors.Add(HLODActorHandle);
+			}
 		}
 	}
 
diff --git a/Engine/Source/Runtime/Engine/Private/WorldPartition/WorldPartition.cpp b/Engine/Source/Runtime/Engine/Private/WorldPartition/WorldPartition.cpp
--- a/Engine/Source/Runtime/Engine/Private/WorldPartition/WorldPartition.cpp
+++ b/Engine/Source/Runtime/Engine/Private/WorldPartition/WorldPartition.cpp
@@ -213,24 +213,24 @@ void UWorldPartition::Initialize(UWorld* InWorld, const FTransform& InTransform)
 		UActorDescContainer::Initialize(PackageName, bRegisterDelegates);
 		check(bContainerInitialized);
 
-		for (TPair<FGuid, TUniquePtr<FWorldPartitionActorDesc>*>& Pair : Actors)
+		for (UActorDescContainer::TIterator ActorDescIterator(this); ActorDescIterator; ++ActorDescIterator)
 		{
-			TUniquePtr<FWorldPartitionActorDesc>* NewActorDesc = Pair.Value;
+			FWorldPartitionActorDesc* ActorDesc = *ActorDescIterator;
 
 			if (bIsInstanced)
 			{
-				const FString LongActorPackageName = (*NewActorDesc)->ActorPackage.ToString();
+				const FString LongActorPackageName = ActorDesc->ActorPackage.ToString();
 				const FString ActorPackageName = FPaths::GetBaseFilename(LongActorPackageName);
 				const FString InstancedName = FString::Printf(TEXT("%s_InstanceOf_%s"), *LevelPackage->GetName(), *ActorPackageName);
 
 				InstancingContext.AddMapping(*LongActorPackageName, *InstancedName);
 
-				(*NewActorDesc)->TransformInstance(ReplaceFrom, ReplaceTo, InstanceTransform);
+				ActorDesc->TransformInstance(ReplaceFrom, ReplaceTo, InstanceTransform);
 			}
 
 			if (bEditorOnly)
 			{
-				HashActorDesc(NewActorDesc->Get());
+				HashActorDesc(ActorDesc);
 			}
 		}
 
@@ -402,9 +402,9 @@ void UWorldPartition::ForEachIntersectingActorDesc(const FBox& Box, TSubclassOf<
 
 void UWorldPartition::ForEachActorDesc(TSubclassOf<AActor> ActorClass, TFunctionRef<bool(const FWorldPartitionActorDesc*)> Predicate) const
 {
-	for (const TPair<FGuid, TUniquePtr<FWorldPartitionActorDesc>*> Pair : Actors)
+	for (UActorDescContainer::TConstIterator ActorDescIterator(this); ActorDescIterator; ++ActorDescIterator)
 	{
-		FWorldPartitionActorDesc* ActorDesc = Pair.Value->Get();
+		const FWorldPartitionActorDesc* ActorDesc = *ActorDescIterator;
 		if (ActorDesc->GetActorClass()->IsChildOf(ActorClass))
 		{
 			if (!Predicate(ActorDesc))
@@ -428,18 +428,6 @@ const TArray<FWorldPartitionStreamingSource>& UWorldPartition::GetStreamingSourc
 }
 
 #if WITH_EDITOR
-const FWorldPartitionActorDesc* UWorldPartition::GetActorDesc(const FGuid& Guid) const
-{
-	const TUniquePtr<FWorldPartitionActorDesc>* const * ActorDesc = Actors.Find(Guid);
-	return ActorDesc ? (*ActorDesc)->Get() : nullptr;
-}
-
-FWorldPartitionActorDesc* UWorldPartition::GetActorDesc(const FGuid& Guid)
-{
-	TUniquePtr<FWorldPartitionActorDesc>** ActorDesc = Actors.Find(Guid);
-	return ActorDesc ? (*ActorDesc)->Get() : nullptr;
-}
-
 bool UWorldPartition::IsSimulating() const
 {
 	return GEditor->bIsSimulatingInEditor || !!GEditor->PlayWorld;
@@ -781,13 +769,13 @@ void UWorldPartition::BeginDestroy()
 #if WITH_EDITOR
 UObject* UWorldPartition::LoadSubobject(const TCHAR* SubObjectPath)
 {
-	for (const auto& Pair : Actors)
+	for (UActorDescContainer::TIterator ActorDescIterator(this); ActorDescIterator; ++ActorDescIterator)
 	{
-		const FWorldPartitionActorDesc* ActorDesc = Pair.Value->Get();
+		FWorldPartitionActorDesc* ActorDesc = *ActorDescIterator;
 
 		if (FString(ActorDesc->ActorPath.ToString()).EndsWith(SubObjectPath))
 		{
-			int32 ReferenceIndex = LoadedSubobjects.Add(Pair.Value);
+			int32 ReferenceIndex = LoadedSubobjects.Emplace(this, ActorDesc->GetGuid());
 			return LoadedSubobjects[ReferenceIndex]->GetActor();
 		}
 	}
@@ -950,9 +938,9 @@ void UWorldPartition::OnPreFixupForPIE(int32 InPIEInstanceID, FSoftObjectPath& O
 FBox UWorldPartition::GetWorldBounds() const
 {
 	FBox WorldBounds(ForceInit);
-	for (const auto& Pair : Actors)
+	for (UActorDescContainer::TConstIterator ActorDescIterator(this); ActorDescIterator; ++ActorDescIterator)
 	{
-		const FWorldPartitionActorDesc* ActorDesc = Pair.Value->Get();
+		const FWorldPartitionActorDesc* ActorDesc = *ActorDescIterator;
 
 		switch (ActorDesc->GetGridPlacement())
 		{
diff --git a/Engine/Source/Runtime/Engine/Private/WorldPartition/WorldPartitionRuntimeSpatialHash.cpp b/Engine/Source/Runtime/Engine/Private/WorldPartition/WorldPartitionRuntimeSpatialHash.cpp
--- a/Engine/Source/Runtime/Engine/Private/WorldPartition/WorldPartitionRuntimeSpatialHash.cpp
+++ b/Engine/Source/Runtime/Engine/Private/WorldPartition/WorldPartitionRuntimeSpatialHash.cpp
@@ -6,7 +6,6 @@
 #include "WorldPartition/WorldPartitionRuntimeSpatialHash.h"
 #include "WorldPartition/WorldPartitionRuntimeSpatialHashCell.h"
 #include "WorldPartition/WorldPartitionActorDesc.h"
-#include "WorldPartition/WorldPartitionActorDescIterator.h"
 #include "WorldPartition/WorldPartitionStreamingPolicy.h"
 #include "WorldPartition/WorldPartitionHandle.h"
 
@@ -489,11 +488,15 @@ void UWorldPartitionRuntimeSpatialHash::CacheHLODParents()
 
 	UWorldPartition* WorldPartition = GetOuterUWorldPartition();
 
-	for (TWorldPartitionActorDescIterator<AWorldPartitionHLOD, FHLODActorDesc> HLODIterator(WorldPartition); HLODIterator; ++HLODIterator)
+	for (UActorDescContainer::TIterator HLODIterator(WorldPartition); HLODIterator; ++HLODIterator)
 	{
-		for (const auto& SubActor : HLODIterator->GetSubActors())
+		if (HLODIterator->GetActorClass()->IsChildOf<AWorldPartitionHLOD>())
 		{
-			CachedHLODParents.Emplace(SubActor, HLODIterator->GetGuid());
+			FHLODActorDesc* HLODActorDesc = (FHLODActorDesc*)*HLODIterator;
+			for (const auto& SubActor : HLODActorDesc->GetSubActors())
+			{
+				CachedHLODParents.Emplace(SubActor, HLODActorDesc->GetGuid());
+			}
 		}
 	}
 }
diff --git a/Engine/Source/Runtime/Engine/Public/WorldPartition/ActorDescContainer.h b/Engine/Source/Runtime/Engine/Public/WorldPartition/ActorDescContainer.h
--- a/Engine/Source/Runtime/Engine/Public/WorldPartition/ActorDescContainer.h
+++ b/Engine/Source/Runtime/Engine/Public/WorldPartition/ActorDescContainer.h
@@ -12,6 +12,9 @@ UCLASS()
 class ENGINE_API UActorDescContainer : public UObject
 {
 	GENERATED_UCLASS_BODY()
+
+	friend struct FWorldPartitionHandleUtils;
+
 public:
 	void Initialize(FName InPackageName, bool bRegisterDelegates);
 	virtual void Uninitialize();
@@ -21,6 +24,9 @@ public:
 	virtual void OnAssetAdded(const FAssetData& InAssetData);
 	virtual void OnAssetRemoved(const FAssetData& InAssetData);
 	virtual void OnAssetUpdated(const FAssetData& InAssetData);
+
+	FWorldPartitionActorDesc* GetActorDesc(const FGuid& Guid);
+	const FWorldPartitionActorDesc* GetActorDesc(const FGuid& Guid) const;
 #endif
 
 protected:
@@ -43,8 +49,91 @@ protected:
 	TChunkedArray<TUniquePtr<FWorldPartitionActorDesc>> ActorDescList;
 
 	FName ContainerPackageName;
-public:
+
+private:
 	TMap<FGuid, TUniquePtr<FWorldPartitionActorDesc>*> Actors;
-	
+
+public:
+	template<bool bConst>
+	class TBaseIterator
+	{
+	private:
+		typedef TMap<FGuid, TUniquePtr<FWorldPartitionActorDesc>*> MapType;
+		typedef typename TChooseClass<bConst, MapType::TConstIterator, MapType::TIterator>::Result IteratorType;
+		typedef typename TChooseClass<bConst, const UActorDescContainer*, UActorDescContainer*>::Result ContainerType;
+		typedef typename TChooseClass<bConst, const FWorldPartitionActorDesc*, FWorldPartitionActorDesc*>::Result ValueType;
+
+	public:
+		explicit TBaseIterator(ContainerType InActorDescContainer)
+			: ActorsIterator(InActorDescContainer->Actors)
+		{
+			if (ShouldSkip())
+			{
+				operator++();
+			}
+		}
+
+		/**
+		 * Iterates to next suitable actor desc.
+		 */
+		void operator++()
+		{
+			do
+			{
+				++ActorsIterator;
+			} while (ShouldSkip());
+		}
+
+		/**
+		 * Returns the current suitable actor desc pointed at by the Iterator
+		 *
+		 * @return	Current suitable actor desc
+		 */
+		FORCEINLINE ValueType operator*() const
+		{
+			return ActorsIterator->Value->Get();
+		}
+
+		/**
+		 * Returns the current suitable actor desc pointed at by the Iterator
+		 *
+		 * @return	Current suitable actor desc
+		 */
+		FORCEINLINE ValueType operator->() const
+		{
+			return ActorsIterator->Value->Get();
+		}
+		/**
+		 * Returns whether the iterator has reached the end and no longer points
+		 * to a suitable actor desc.
+		 *
+		 * @return true if iterator points to a suitable actor desc, false if it has reached the end
+		 */
+		FORCEINLINE explicit operator bool() const
+		{
+			return (bool)ActorsIterator;
+		}
+
+	private:
+		/**
+		 * Determines whether the iterator currently points to a valid actor desc or not.
+		 * @return	true
+		 */
+		FORCEINLINE bool ShouldSkip() const
+		{
+			if (!ActorsIterator)
+			{
+				return false;
+			}
+
+			return false;
+		}
+
+	private:
+		IteratorType ActorsIterator;
+	};
+
+	typedef TBaseIterator<false> TIterator;
+	typedef TBaseIterator<true> TConstIterator;
 #endif
 };
\ No newline at end of file
diff --git a/Engine/Source/Runtime/Engine/Public/WorldPartition/WorldPartition.h b/Engine/Source/Runtime/Engine/Public/WorldPartition/WorldPartition.h
--- a/Engine/Source/Runtime/Engine/Public/WorldPartition/WorldPartition.h
+++ b/Engine/Source/Runtime/Engine/Public/WorldPartition/WorldPartition.h
@@ -122,9 +122,6 @@ public:
 	FName GetWorldPartitionEditorName();
 	bool IsSimulating() const;
 
-	FWorldPartitionActorDesc* GetActorDesc(const FGuid& Guid);
-	const FWorldPartitionActorDesc* GetActorDesc(const FGuid& Guid) const;
-
 	void LoadEditorCells(const FBox& Box);
 	void UnloadEditorCells(const FBox& Box);
 	bool AreEditorCellsLoaded(const FBox& Box);
diff --git a/Engine/Source/Runtime/Engine/Public/WorldPartition/WorldPartitionActorDescIterator.h b/Engine/Source/Runtime/Engine/Public/WorldPartition/WorldPartitionActorDescIterator.h
deleted file mode 100644
diff --git a/Engine/Source/Runtime/Landscape/Private/Landscape.cpp b/Engine/Source/Runtime/Landscape/Private/Landscape.cpp
--- a/Engine/Source/Runtime/Landscape/Private/Landscape.cpp
+++ b/Engine/Source/Runtime/Landscape/Private/Landscape.cpp
@@ -71,7 +71,6 @@ Landscape.cpp: Terrain rendering
 #include "Algo/BinarySearch.h"
 #include "WorldPartition/WorldPartition.h"
 #include "WorldPartition/WorldPartitionHandle.h"
-#include "WorldPartition/WorldPartitionActorDescIterator.h"
 #include "WorldPartition/Landscape/LandscapeActorDesc.h"
 #include "WorldPartition/Landscape/LandscapeSplineActorDesc.h"
 #if WITH_EDITOR
@@ -3134,7 +3133,7 @@ void ULandscapeInfo::Initialize(UWorld* InWorld, const FGuid& InLandscapeGuid)
 	// If WorldPartition world get the Proxy handles
 	if (UWorldPartition* WorldPartition = InWorld->GetWorldPartition())
 	{
-		for (TWorldPartitionActorDescIterator<AActor, FWorldPartitionActorDesc> ActorDescIterator(WorldPartition); ActorDescIterator; ++ActorDescIterator)
+		for (UActorDescContainer::TIterator ActorDescIterator(WorldPartition); ActorDescIterator; ++ActorDescIterator)
 		{
 			FWorldPartitionActorDesc* ActorDesc = *ActorDescIterator;
 			if (ActorDesc->GetActorClass()->IsChildOf(ALandscapeStreamingProxy::StaticClass()))
