Fix NonUnity build

diff --git a/Engine/Source/Runtime/Engine/Private/Components/WorldPartitionStreamingSourceComponent.cpp b/Engine/Source/Runtime/Engine/Private/Components/WorldPartitionStreamingSourceComponent.cpp
--- a/Engine/Source/Runtime/Engine/Private/Components/WorldPartitionStreamingSourceComponent.cpp
+++ b/Engine/Source/Runtime/Engine/Private/Components/WorldPartitionStreamingSourceComponent.cpp
@@ -2,6 +2,7 @@
 
 #include "Components/WorldPartitionStreamingSourceComponent.h"
 #include "WorldPartition/WorldPartition.h"
+#include "Engine/World.h"
 
 UWorldPartitionStreamingSourceComponent::UWorldPartitionStreamingSourceComponent(const FObjectInitializer& ObjectInitializer)
 	: Super(ObjectInitializer)
