Fix missing include

diff --git a/Engine/Source/Runtime/Engine/Classes/Engine/Scene.h b/Engine/Source/Runtime/Engine/Classes/Engine/Scene.h
--- a/Engine/Source/Runtime/Engine/Classes/Engine/Scene.h
+++ b/Engine/Source/Runtime/Engine/Classes/Engine/Scene.h
@@ -13,6 +13,7 @@
 #include "Engine/BlendableInterface.h"
 #include "RHIDefinitions.h"
 #include "SceneUtils.h"
+#include "Engine/EngineTypes.h"
 #include "Scene.generated.h"
 
 
