Make FinishCompilation be part of the stacks we gather to identity stalls to have a complete picture of time lost waiting on async tasks

diff --git a/Engine/Source/Runtime/Engine/Private/AsyncCompilationHelpers.cpp b/Engine/Source/Runtime/Engine/Private/AsyncCompilationHelpers.cpp
--- a/Engine/Source/Runtime/Engine/Private/AsyncCompilationHelpers.cpp
+++ b/Engine/Source/Runtime/Engine/Private/AsyncCompilationHelpers.cpp
@@ -91,86 +91,92 @@ namespace AsyncCompilationHelpers
 		TFunctionRef<void(ICompilable*)> PostCompileSingle
 	)
 	{
+		if (Num <= 0)
+		{
+			return;
+		}
+
 		check(IsInGameThread());
 
 		FObjectCacheContextScope ObjectCacheScope;
-		if (Num)
-		{
-			TOptional<FScopedSlowTask> SlowTask;
-
-			// Do not create a progress during PostLoad as TickSlate could endup calling too many things
-			if (!FUObjectThreadContext::Get().IsRoutingPostLoad)
-			{
-				FFormatNamedArguments Args;
-				Args.Add(TEXT("AssetType"), AssetType.ToLower());
 
-				SlowTask.Emplace((float)Num, FText::Format(LOCTEXT("WaitingOnFinishCompilation", "Waiting on {AssetType} preparation"), Args), true);
-				SlowTask->MakeDialogDelayed(1.0f, false /*bShowCancelButton*/, true /*bAllowInPIE*/);
-			}
+		uint64 StartTime = FPlatformTime::Cycles64();
+		TOptional<FScopedSlowTask> SlowTask;
 
-			struct FFinishTask : public IQueuedWork
-			{
-				ICompilable* Job = nullptr;
-				FEvent* Event;
-				FFinishTask() { Event = FPlatformProcess::GetSynchEventFromPool(true); }
-				~FFinishTask() { FPlatformProcess::ReturnSynchEventToPool(Event); }
-				void DoThreadedWork() override
-				{
-					FOptionalTaskTagScope Scope(ETaskTag::EParallelGameThread);
-					Job->EnsureCompletion();
-					Event->Trigger();
-				}
-
-				void Abandon() override { }
-			};
+		// Do not create a progress during PostLoad as TickSlate could endup calling too many things
+		if (!FUObjectThreadContext::Get().IsRoutingPostLoad)
+		{
+			FFormatNamedArguments Args;
+			Args.Add(TEXT("AssetType"), AssetType.ToLower());
 
-			// Perform forced compilation on as many thread as possible in high priority since the game-thread is waiting
-			TArray<FFinishTask> PendingTasks;
-			PendingTasks.SetNum(Num);
+			SlowTask.Emplace((float)Num, FText::Format(LOCTEXT("WaitingOnFinishCompilation", "Waiting on {AssetType} preparation"), Args), true);
+			SlowTask->MakeDialogDelayed(1.0f, false /*bShowCancelButton*/, true /*bAllowInPIE*/);
+		}
 
-			for (int32 Index = 0; Index < Num; ++Index)
+		struct FFinishTask : public IQueuedWork
+		{
+			ICompilable* Job = nullptr;
+			FEvent* Event;
+			FFinishTask() { Event = FPlatformProcess::GetSynchEventFromPool(true); }
+			~FFinishTask() { FPlatformProcess::ReturnSynchEventToPool(Event); }
+			void DoThreadedWork() override
 			{
-				PendingTasks[Index].Job = &Getter(Index);
-				GThreadPool->AddQueuedWork(&PendingTasks[Index], EQueuedWorkPriority::High);
+				FOptionalTaskTagScope Scope(ETaskTag::EParallelGameThread);
+				Job->EnsureCompletion();
+				Event->Trigger();
 			}
 
-			auto FormatProgress =
-				[&AssetType](int32 Done, int32 Total, FName ObjectName)
-			{
-				FFormatNamedArguments Args;
-				Args.Add(TEXT("AssetType"), AssetType.ToLower());
-				Args.Add(TEXT("Done"), Done);
-				Args.Add(TEXT("Total"), Total);
-				Args.Add(TEXT("ObjectName"), FText::FromName(ObjectName));
-				return FText::Format(LOCTEXT("WaitingOnFinishCompilationWithCount", "Waiting for {AssetType} to be ready {Done}/{Total} ({ObjectName}) ..."), Args);
-			};
-
-			int32 NumDone = 0;
-			for (FFinishTask& PendingTask : PendingTasks)
-			{
-				ICompilable* Job = PendingTask.Job;
+			void Abandon() override { }
+		};
 
-				FText Progress = FormatProgress(NumDone++, Num, Job->GetName());
+		// Perform forced compilation on as many thread as possible in high priority since the game-thread is waiting
+		TArray<FFinishTask> PendingTasks;
+		PendingTasks.SetNum(Num);
 
-				// Be nice with the game thread and tick the progress at 60 fps even when no progress is being made...
-				while (!PendingTask.Event->Wait(16))
-				{
-					if (SlowTask.IsSet())
-					{
-						SlowTask->EnterProgressFrame(0.0f, Progress);
-					}
-				}
+		for (int32 Index = 0; Index < Num; ++Index)
+		{
+			PendingTasks[Index].Job = &Getter(Index);
+			GThreadPool->AddQueuedWork(&PendingTasks[Index], EQueuedWorkPriority::High);
+		}
 
+		auto FormatProgress =
+			[&AssetType](int32 Done, int32 Total, FName ObjectName)
+		{
+			FFormatNamedArguments Args;
+			Args.Add(TEXT("AssetType"), AssetType.ToLower());
+			Args.Add(TEXT("Done"), Done);
+			Args.Add(TEXT("Total"), Total);
+			Args.Add(TEXT("ObjectName"), FText::FromName(ObjectName));
+			return FText::Format(LOCTEXT("WaitingOnFinishCompilationWithCount", "Waiting for {AssetType} to be ready {Done}/{Total} ({ObjectName}) ..."), Args);
+		};
+
+		int32 NumDone = 0;
+		for (FFinishTask& PendingTask : PendingTasks)
+		{
+			ICompilable* Job = PendingTask.Job;
+
+			FText Progress = FormatProgress(NumDone++, Num, Job->GetName());
+
+			// Be nice with the game thread and tick the progress at 60 fps even when no progress is being made...
+			while (!PendingTask.Event->Wait(16))
+			{
 				if (SlowTask.IsSet())
 				{
-					SlowTask->EnterProgressFrame(1.0f, Progress);
+					SlowTask->EnterProgressFrame(0.0f, Progress);
 				}
+			}
 
-				UE_LOG_REF(LogCategory, Display, TEXT("%s"), *Progress.ToString());
-
-				PostCompileSingle(Job);
+			if (SlowTask.IsSet())
+			{
+				SlowTask->EnterProgressFrame(1.0f, Progress);
 			}
+
+			UE_LOG_REF(LogCategory, Display, TEXT("%s"), *Progress.ToString());
+
+			PostCompileSingle(Job);
 		}
+
+		SaveStallStack(FPlatformTime::Cycles64() - StartTime);
 	}
 
 	void EnsureInitializedCVars(
@@ -350,6 +356,8 @@ namespace AsyncCompilationHelpers
 
 	void DumpStallStacks()
 	{
+		TRACE_CPUPROFILER_EVENT_SCOPE(AsyncCompilationHelpers::DumpStallStacks);
+
 		FRWScopeLock Scope(BackTracesLock, SLT_Write);
 		TArray<FStackData*> Stacks;
 		for (TPair<uint64, FStackData>& Pair : BackTraces)
@@ -377,7 +385,12 @@ namespace AsyncCompilationHelpers
 
 			TotalCount += StackData->Count;
 			TotalCycles += StackData->Cycles;
-			UE_LOG(LogAsyncCompilation, Display, TEXT("Async Compilation Stall Stack: (count: %llu, time: %s)\n%s"), StackData->Count, *FPlatformTime::PrettyTime(FPlatformTime::ToSeconds64(StackData->Cycles)), ANSI_TO_TCHAR(HumanReadableString));
+			UE_LOG(LogAsyncCompilation, Display, TEXT("Async Compilation Stall Stack: (Count: %llu, Time: %s)\n%s"), StackData->Count, *FPlatformTime::PrettyTime(FPlatformTime::ToSeconds64(StackData->Cycles)), ANSI_TO_TCHAR(HumanReadableString));
+		}
+
+		if (TotalCount)
+		{
+			UE_LOG(LogAsyncCompilation, Display, TEXT("Async Compilation Stall Stack: (Total Count: %llu, Total Time: %s)"), TotalCount, *FPlatformTime::PrettyTime(FPlatformTime::ToSeconds64(TotalCycles)));
 		}
 	}
 }
diff --git a/Engine/Source/Runtime/Engine/Private/StaticMesh.cpp b/Engine/Source/Runtime/Engine/Private/StaticMesh.cpp
--- a/Engine/Source/Runtime/Engine/Private/StaticMesh.cpp
+++ b/Engine/Source/Runtime/Engine/Private/StaticMesh.cpp
@@ -2928,9 +2928,7 @@ void UStaticMesh::WaitUntilAsyncPropertyReleased(EStaticMeshAsyncProperties Asyn
 				ToString(AsyncProperties)
 			);
 
-			uint64 StartTime = FPlatformTime::Cycles64();
 			FStaticMeshCompilingManager::Get().FinishCompilation({ const_cast<UStaticMesh*>(this) });
-			AsyncCompilationHelpers::SaveStallStack(FPlatformTime::Cycles64() - StartTime);
 		}
 		else
 		{
