Fixed race condition with skylight capturing the scene, where it would often be black due to async static mesh compilation

diff --git a/Engine/Source/Runtime/Engine/Private/Components/SkyLightComponent.cpp b/Engine/Source/Runtime/Engine/Private/Components/SkyLightComponent.cpp
--- a/Engine/Source/Runtime/Engine/Private/Components/SkyLightComponent.cpp
+++ b/Engine/Source/Runtime/Engine/Private/Components/SkyLightComponent.cpp
@@ -35,6 +35,7 @@
 #if WITH_EDITOR
 #include "Rendering/StaticLightingSystemInterface.h"
 #include "TextureCompiler.h"
+#include "StaticMeshCompiler.h"
 #endif
 
 #define LOCTEXT_NAMESPACE "SkyLightComponent"
@@ -719,9 +720,9 @@ void USkyLightComponent::ApplyComponentInstanceData(FPrecomputedSkyLightInstance
 void USkyLightComponent::UpdateSkyCaptureContentsArray(UWorld* WorldToUpdate, TArray<USkyLightComponent*>& ComponentArray, bool bOperateOnBlendSource)
 {
 	const bool bIsCompilingShaders = GShaderCompilingManager != nullptr && GShaderCompilingManager->IsCompiling();
-	bool bIsCompilingTextures = false;
+	bool bSceneIsAsyncCompiling = false;
 #if WITH_EDITOR
-	bIsCompilingTextures = FTextureCompilingManager::Get().GetNumRemainingTextures() > 0;
+	bSceneIsAsyncCompiling = FTextureCompilingManager::Get().GetNumRemainingTextures() > 0 || FStaticMeshCompilingManager::Get().GetNumRemainingMeshes() > 0;
 #endif
 
 	// Iterate backwards so we can remove elements without changing the index
@@ -747,7 +748,7 @@ void USkyLightComponent::UpdateSkyCaptureContentsArray(UWorld* WorldToUpdate, TA
 
 		if (((!Owner || !Owner->GetLevel() || Owner->GetLevel()->bIsVisible) && CaptureComponent->GetWorld() == WorldToUpdate)
 			// Only process sky capture requests once async texture and shader compiling completes, otherwise we will capture the scene with temporary shaders/textures
-			&& (((!bIsCompilingTextures) && (!bIsCompilingShaders)) || ((CaptureComponent->SourceType == SLS_SpecifiedCubemap) && (!bIsCubemapCompiling))))
+			&& (((!bSceneIsAsyncCompiling) && (!bIsCompilingShaders)) || ((CaptureComponent->SourceType == SLS_SpecifiedCubemap) && (!bIsCubemapCompiling))))
 		{
 			// Only capture valid sky light components
 			if (CaptureComponent->SourceType != SLS_SpecifiedCubemap || CaptureComponent->Cubemap)
