Exposed DistanceFieldVoxelDensity in project settings

diff --git a/Engine/Source/Runtime/Engine/Classes/Engine/RendererSettings.h b/Engine/Source/Runtime/Engine/Classes/Engine/RendererSettings.h
--- a/Engine/Source/Runtime/Engine/Classes/Engine/RendererSettings.h
+++ b/Engine/Source/Runtime/Engine/Classes/Engine/RendererSettings.h
@@ -369,6 +369,14 @@ class ENGINE_API URendererSettings : public UDeveloperSettings
 		ConfigRestartRequired=true))
 	uint32 bGenerateMeshDistanceFields:1;
 
+	UPROPERTY(config, EditAnywhere, Category=SoftwareRayTracing, meta=(
+		EditCondition = "bGenerateMeshDistanceFields",
+		ConsoleVariable="r.DistanceFields.DefaultVoxelDensity",
+		ClampMin = ".05", ClampMax = ".4",
+		ToolTip="Determines how the default scale of a mesh converts into distance field voxel dimensions. Changing this will cause all distance fields to be rebuilt.  Large values can consume memory very quickly!  Changing this setting requires restarting the editor.",
+		ConfigRestartRequired=true))
+	float DistanceFieldVoxelDensity;
+
 	UPROPERTY(config, EditAnywhere, Category=SoftwareRayTracing, meta=(
 		EditCondition = "bGenerateMeshDistanceFields",
 		ConsoleVariable="r.DistanceFieldBuild.EightBit",
