Added nullcheck on GEngine access in GetWorldForVisualLogger

diff --git a/Engine/Source/Runtime/Engine/Private/VisualLogger/VisualLogger.cpp b/Engine/Source/Runtime/Engine/Private/VisualLogger/VisualLogger.cpp
--- a/Engine/Source/Runtime/Engine/Private/VisualLogger/VisualLogger.cpp
+++ b/Engine/Source/Runtime/Engine/Private/VisualLogger/VisualLogger.cpp
@@ -41,7 +41,7 @@ namespace
 #endif
 		if (!GIsEditor && World == nullptr)
 		{
-			World = GEngine->GetWorld();
+			World = GEngine ? GEngine->GetWorld() : nullptr;
 		}
 
 		return World;
