Use texture buffers for a skeletal mesh bone data on all mobile platforms. Bump bone limit to 256 from 75 on mobile.

diff --git a/Engine/Source/Runtime/Engine/Private/GPUSkinVertexFactory.cpp b/Engine/Source/Runtime/Engine/Private/GPUSkinVertexFactory.cpp
--- a/Engine/Source/Runtime/Engine/Private/GPUSkinVertexFactory.cpp
+++ b/Engine/Source/Runtime/Engine/Private/GPUSkinVertexFactory.cpp
@@ -109,12 +109,14 @@ static TAutoConsoleVariable<int32> CVarVelocityTest(
 // These should match USE_BONES_SRV_BUFFER
 static inline bool SupportsBonesBufferSRV(EShaderPlatform Platform)
 {
-	return IsFeatureLevelSupported(Platform, ERHIFeatureLevel::SM5) || IsVulkanPlatform(Platform) || IsMetalPlatform(Platform);
+	// at some point we might switch GL to uniform buffers
+	return true;
 }
 
 static inline bool SupportsBonesBufferSRV(ERHIFeatureLevel::Type InFeatureLevel)
 {
-	return InFeatureLevel > ERHIFeatureLevel::ES3_1 || IsVulkanPlatform(GMaxRHIShaderPlatform) || IsMetalPlatform(GMaxRHIShaderPlatform);
+	// at some point we might switch GL to uniform buffers
+	return true;
 }
 // ---
 
diff --git a/Engine/Source/Runtime/Engine/Private/SkeletalMeshRenderData.cpp b/Engine/Source/Runtime/Engine/Private/SkeletalMeshRenderData.cpp
--- a/Engine/Source/Runtime/Engine/Private/SkeletalMeshRenderData.cpp
+++ b/Engine/Source/Runtime/Engine/Private/SkeletalMeshRenderData.cpp
@@ -600,8 +600,8 @@ bool FSkeletalMeshRenderData::RequiresCPUSkinning(ERHIFeatureLevel::Type Feature
 {
 	const int32 MaxGPUSkinBones = FMath::Min(GetFeatureLevelMaxNumberOfBones(FeatureLevel), FGPUBaseSkinVertexFactory::GetMaxGPUSkinBones());
 	const int32 MaxBonesPerChunk = GetMaxBonesPerSection(MinLODIndex);
-	// Do CPU skinning if we need too many bones per chunk, or if we have too many influences per vertex on lower end
-	return (MaxBonesPerChunk > MaxGPUSkinBones) || (GetNumBoneInfluences(MinLODIndex) > MAX_INFLUENCES_PER_STREAM && FeatureLevel < ERHIFeatureLevel::ES3_1);
+	// Do CPU skinning if we need too many bones per chunk
+	return (MaxBonesPerChunk > MaxGPUSkinBones);
 }
 
 bool FSkeletalMeshRenderData::RequiresCPUSkinning(ERHIFeatureLevel::Type FeatureLevel) const
diff --git a/Engine/Source/Runtime/Engine/Public/GPUSkinVertexFactory.h b/Engine/Source/Runtime/Engine/Public/GPUSkinVertexFactory.h
--- a/Engine/Source/Runtime/Engine/Public/GPUSkinVertexFactory.h
+++ b/Engine/Source/Runtime/Engine/Public/GPUSkinVertexFactory.h
@@ -26,7 +26,8 @@ END_GLOBAL_SHADER_PARAMETER_STRUCT()
 
 enum
 {
-	MAX_GPU_BONE_MATRICES_UNIFORMBUFFER = 75,
+	// 256 works for real uniform buffers, emulated UB can support up to 75 
+	MAX_GPU_BONE_MATRICES_UNIFORMBUFFER = 256,
 };
 
 BEGIN_GLOBAL_SHADER_PARAMETER_STRUCT(FBoneMatricesUniformShaderParameters,)
diff --git a/Engine/Source/Runtime/RHI/Public/RHIDefinitions.h b/Engine/Source/Runtime/RHI/Public/RHIDefinitions.h
--- a/Engine/Source/Runtime/RHI/Public/RHIDefinitions.h
+++ b/Engine/Source/Runtime/RHI/Public/RHIDefinitions.h
@@ -1903,7 +1903,7 @@ inline int32 GetFeatureLevelMaxNumberOfBones(const FStaticFeatureLevel FeatureLe
 	switch (FeatureLevel)
 	{
 	case ERHIFeatureLevel::ES3_1:
-		return 75;	
+		return 256;	
 	case ERHIFeatureLevel::SM5:
 	case ERHIFeatureLevel::SM6:
 		return 65536; // supports uint16
