Don't spawn actors using an external package when running in editor game mode.

diff --git a/Engine/Source/Runtime/Engine/Private/LevelActor.cpp b/Engine/Source/Runtime/Engine/Private/LevelActor.cpp
--- a/Engine/Source/Runtime/Engine/Private/LevelActor.cpp
+++ b/Engine/Source/Runtime/Engine/Private/LevelActor.cpp
@@ -548,6 +548,11 @@ AActor* UWorld::SpawnActor( UClass* Class, FTransform const* UserTransformPtr, c
 	{
 		bNeedGloballyUniqueName = CastChecked<AActor>(Class->GetDefaultObject())->SupportsExternalPackaging();
 	}
+
+	if (!GIsEditor)
+	{
+		bNeedGloballyUniqueName = false;
+	}
 #endif
 
 	if (NewActorName.IsNone())
@@ -736,13 +741,13 @@ AActor* UWorld::SpawnActor( UClass* Class, FTransform const* UserTransformPtr, c
 	OnActorSpawned.Broadcast(Actor);
 
 #if WITH_EDITOR
-	if (Actor->IsAsset())
-	{
-		FAssetRegistryModule::AssetCreated(Actor);
-	}
-
 	if (GIsEditor)
 	{
+		if (Actor->IsAsset())
+		{
+			FAssetRegistryModule::AssetCreated(Actor);
+		}
+
 		GEngine->BroadcastLevelActorAdded(Actor);
 	}
 #endif
