Delay generating legacy transient data within UAnimDataModel until end of controller bracket(s)

diff --git a/Engine/Source/Runtime/Engine/Classes/Animation/AnimData/AnimDataModel.h b/Engine/Source/Runtime/Engine/Classes/Animation/AnimData/AnimDataModel.h
--- a/Engine/Source/Runtime/Engine/Classes/Animation/AnimData/AnimDataModel.h
+++ b/Engine/Source/Runtime/Engine/Classes/Animation/AnimData/AnimDataModel.h
@@ -296,12 +296,29 @@ private:
 			ModifiedEventDynamic.Broadcast(NotifyType, this, Payload);
 		}
 
-		GenerateTransientData();
+		// Only regenerate transient data when not in a bracket, or at the end of one
+		{
+			if (NotifyType == EAnimDataModelNotifyType::BracketOpened)
+			{
+				++BracketCounter;
+			}
+			if (NotifyType == EAnimDataModelNotifyType::BracketClosed)
+			{
+				--BracketCounter;
+			}
+
+			if (BracketCounter == 0)
+			{
+				GenerateTransientData();
+			}
+		}
 	}
 
 private:
 	void GenerateTransientData();
 
+	UPROPERTY(Transient)
+	int32 BracketCounter = 0;
 private:
 	/** Dynamic delegate event allows scripting to register to any broadcasted notify. */
 	UPROPERTY(BlueprintAssignable, Transient, Category = AnimationDataModel, meta = (ScriptName = "ModifiedEvent", AllowPrivateAccess = "true"))
