Added support for automatic type deduction for UActorDescContainer::TIterator.

diff --git a/Engine/Source/Runtime/Engine/Classes/GameFramework/Actor.h b/Engine/Source/Runtime/Engine/Classes/GameFramework/Actor.h
--- a/Engine/Source/Runtime/Engine/Classes/GameFramework/Actor.h
+++ b/Engine/Source/Runtime/Engine/Classes/GameFramework/Actor.h
@@ -20,6 +20,7 @@
 
 #if WITH_EDITOR
 #include "WorldPartition/DataLayer/ActorDataLayer.h"
+#include "WorldPartition/WorldPartitionActorDescType.h"
 #endif
 
 #include "Actor.generated.h"
@@ -3862,6 +3863,8 @@ FORCEINLINE_DEBUGGABLE bool AActor::IsNetMode(ENetMode Mode) const
 #endif
 }
 
+DEFINE_ACTORDESC_TYPE(AActor, FWorldPartitionActorDesc);
+
 //////////////////////////////////////////////////////////////////////////
 // Macro to hide common Transform functions in native code for classes where they don't make sense.
 // Note that this doesn't prevent access through function calls from parent classes (ie an AActor*), but
diff --git a/Engine/Source/Runtime/Engine/Private/WorldPartition/RuntimeSpatialHash/RuntimeSpatialHashActorCluster.cpp b/Engine/Source/Runtime/Engine/Private/WorldPartition/RuntimeSpatialHash/RuntimeSpatialHashActorCluster.cpp
--- a/Engine/Source/Runtime/Engine/Private/WorldPartition/RuntimeSpatialHash/RuntimeSpatialHashActorCluster.cpp
+++ b/Engine/Source/Runtime/Engine/Private/WorldPartition/RuntimeSpatialHash/RuntimeSpatialHashActorCluster.cpp
@@ -147,7 +147,7 @@ TArray<FActorCluster> CreateActorClustersImpl(UWorldPartition* WorldPartition, T
 		LevelScriptExternalActorReferences.Append(ActorsReferencesUtils::GetExternalActorReferences(LevelScriptBlueprint));
 	}
 
-	for (UActorDescContainer::TIterator ActorDescIterator(WorldPartition); ActorDescIterator; ++ActorDescIterator)
+	for (UActorDescContainer::TIterator<> ActorDescIterator(WorldPartition); ActorDescIterator; ++ActorDescIterator)
 	{
 		FWorldPartitionActorDesc* ActorDesc = *ActorDescIterator;
 		EActorGridPlacement GridPlacement = ActorDesc->GetGridPlacement();
diff --git a/Engine/Source/Runtime/Engine/Private/WorldPartition/RuntimeSpatialHash/RuntimeSpatialHashHLOD.cpp b/Engine/Source/Runtime/Engine/Private/WorldPartition/RuntimeSpatialHash/RuntimeSpatialHashHLOD.cpp
--- a/Engine/Source/Runtime/Engine/Private/WorldPartition/RuntimeSpatialHash/RuntimeSpatialHashHLOD.cpp
+++ b/Engine/Source/Runtime/Engine/Private/WorldPartition/RuntimeSpatialHash/RuntimeSpatialHashHLOD.cpp
@@ -477,21 +477,17 @@ bool UWorldPartitionRuntimeSpatialHash::GenerateHLOD(ISourceControlHelper* Sourc
 	FHLODGenerationContext Context;
 	TArray<FWorldPartitionHandle> InvalidHLODActors;
 
-	for (UActorDescContainer::TIterator HLODIterator(WorldPartition); HLODIterator; ++HLODIterator)
+	for (UActorDescContainer::TIterator<AWorldPartitionHLOD> HLODIterator(WorldPartition); HLODIterator; ++HLODIterator)
 	{
-		if (HLODIterator->GetActorClass()->IsChildOf<AWorldPartitionHLOD>())
-		{
-			FHLODActorDesc* HLODActorDesc = (FHLODActorDesc*)*HLODIterator;
-			FWorldPartitionHandle HLODActorHandle(WorldPartition, HLODActorDesc->GetGuid());
+		FWorldPartitionHandle HLODActorHandle(WorldPartition, HLODIterator->GetGuid());
 
-			if (HLODActorDesc->GetCellHash())
-			{
-				Context.HLODActorDescs.Add(HLODActorDesc->GetCellHash(), HLODActorHandle);
-			}
-			else
-			{
-				InvalidHLODActors.Add(HLODActorHandle);
-			}
+		if (HLODIterator->GetCellHash())
+		{
+			Context.HLODActorDescs.Add(HLODIterator->GetCellHash(), HLODActorHandle);
+		}
+		else
+		{
+			InvalidHLODActors.Add(HLODActorHandle);
 		}
 	}
 
diff --git a/Engine/Source/Runtime/Engine/Private/WorldPartition/WorldPartition.cpp b/Engine/Source/Runtime/Engine/Private/WorldPartition/WorldPartition.cpp
--- a/Engine/Source/Runtime/Engine/Private/WorldPartition/WorldPartition.cpp
+++ b/Engine/Source/Runtime/Engine/Private/WorldPartition/WorldPartition.cpp
@@ -213,7 +213,7 @@ void UWorldPartition::Initialize(UWorld* InWorld, const FTransform& InTransform)
 		UActorDescContainer::Initialize(PackageName, bRegisterDelegates);
 		check(bContainerInitialized);
 
-		for (UActorDescContainer::TIterator ActorDescIterator(this); ActorDescIterator; ++ActorDescIterator)
+		for (UActorDescContainer::TIterator<> ActorDescIterator(this); ActorDescIterator; ++ActorDescIterator)
 		{
 			FWorldPartitionActorDesc* ActorDesc = *ActorDescIterator;
 
@@ -402,7 +402,7 @@ void UWorldPartition::ForEachIntersectingActorDesc(const FBox& Box, TSubclassOf<
 
 void UWorldPartition::ForEachActorDesc(TSubclassOf<AActor> ActorClass, TFunctionRef<bool(const FWorldPartitionActorDesc*)> Predicate) const
 {
-	for (UActorDescContainer::TConstIterator ActorDescIterator(this); ActorDescIterator; ++ActorDescIterator)
+	for (UActorDescContainer::TConstIterator<> ActorDescIterator(this); ActorDescIterator; ++ActorDescIterator)
 	{
 		const FWorldPartitionActorDesc* ActorDesc = *ActorDescIterator;
 		if (ActorDesc->GetActorClass()->IsChildOf(ActorClass))
@@ -769,7 +769,7 @@ void UWorldPartition::BeginDestroy()
 #if WITH_EDITOR
 UObject* UWorldPartition::LoadSubobject(const TCHAR* SubObjectPath)
 {
-	for (UActorDescContainer::TIterator ActorDescIterator(this); ActorDescIterator; ++ActorDescIterator)
+	for (UActorDescContainer::TIterator<> ActorDescIterator(this); ActorDescIterator; ++ActorDescIterator)
 	{
 		FWorldPartitionActorDesc* ActorDesc = *ActorDescIterator;
 
@@ -938,7 +938,7 @@ void UWorldPartition::OnPreFixupForPIE(int32 InPIEInstanceID, FSoftObjectPath& O
 FBox UWorldPartition::GetWorldBounds() const
 {
 	FBox WorldBounds(ForceInit);
-	for (UActorDescContainer::TConstIterator ActorDescIterator(this); ActorDescIterator; ++ActorDescIterator)
+	for (UActorDescContainer::TConstIterator<> ActorDescIterator(this); ActorDescIterator; ++ActorDescIterator)
 	{
 		const FWorldPartitionActorDesc* ActorDesc = *ActorDescIterator;
 
diff --git a/Engine/Source/Runtime/Engine/Private/WorldPartition/WorldPartitionRuntimeSpatialHash.cpp b/Engine/Source/Runtime/Engine/Private/WorldPartition/WorldPartitionRuntimeSpatialHash.cpp
--- a/Engine/Source/Runtime/Engine/Private/WorldPartition/WorldPartitionRuntimeSpatialHash.cpp
+++ b/Engine/Source/Runtime/Engine/Private/WorldPartition/WorldPartitionRuntimeSpatialHash.cpp
@@ -488,15 +488,11 @@ void UWorldPartitionRuntimeSpatialHash::CacheHLODParents()
 
 	UWorldPartition* WorldPartition = GetOuterUWorldPartition();
 
-	for (UActorDescContainer::TIterator HLODIterator(WorldPartition); HLODIterator; ++HLODIterator)
+	for (UActorDescContainer::TIterator<AWorldPartitionHLOD> HLODIterator(WorldPartition); HLODIterator; ++HLODIterator)
 	{
-		if (HLODIterator->GetActorClass()->IsChildOf<AWorldPartitionHLOD>())
+		for (const auto& SubActor : HLODIterator->GetSubActors())
 		{
-			FHLODActorDesc* HLODActorDesc = (FHLODActorDesc*)*HLODIterator;
-			for (const auto& SubActor : HLODActorDesc->GetSubActors())
-			{
-				CachedHLODParents.Emplace(SubActor, HLODActorDesc->GetGuid());
-			}
+			CachedHLODParents.Emplace(SubActor, HLODIterator->GetGuid());
 		}
 	}
 }
diff --git a/Engine/Source/Runtime/Engine/Public/ActorPartition/PartitionActor.h b/Engine/Source/Runtime/Engine/Public/ActorPartition/PartitionActor.h
--- a/Engine/Source/Runtime/Engine/Public/ActorPartition/PartitionActor.h
+++ b/Engine/Source/Runtime/Engine/Public/ActorPartition/PartitionActor.h
@@ -33,4 +33,6 @@ public:
 	UPROPERTY()
 	uint32 GridSize;
 #endif
-};
\ No newline at end of file
+};
+
+DEFINE_ACTORDESC_TYPE(APartitionActor, FPartitionActorDesc);
diff --git a/Engine/Source/Runtime/Engine/Public/LevelInstance/LevelInstanceActor.h b/Engine/Source/Runtime/Engine/Public/LevelInstance/LevelInstanceActor.h
--- a/Engine/Source/Runtime/Engine/Public/LevelInstance/LevelInstanceActor.h
+++ b/Engine/Source/Runtime/Engine/Public/LevelInstance/LevelInstanceActor.h
@@ -103,4 +103,6 @@ private:
 #endif
 
 	FLevelInstanceID LevelInstanceID;
-};
\ No newline at end of file
+};
+
+DEFINE_ACTORDESC_TYPE(ALevelInstance, FLevelInstanceActorDesc);
diff --git a/Engine/Source/Runtime/Engine/Public/WorldPartition/ActorDescContainer.h b/Engine/Source/Runtime/Engine/Public/WorldPartition/ActorDescContainer.h
--- a/Engine/Source/Runtime/Engine/Public/WorldPartition/ActorDescContainer.h
+++ b/Engine/Source/Runtime/Engine/Public/WorldPartition/ActorDescContainer.h
@@ -6,6 +6,7 @@
 #include "UObject/Object.h"
 #include "GameFramework/Actor.h"
 #include "WorldPartition/WorldPartitionActorDesc.h"
+#include "WorldPartition/WorldPartitionActorDescType.h"
 #include "ActorDescContainer.generated.h"
 
 UCLASS()
@@ -54,17 +55,20 @@ private:
 	TMap<FGuid, TUniquePtr<FWorldPartitionActorDesc>*> Actors;
 
 public:
-	template<bool bConst>
+	template<bool bConst, class ActorType>
 	class TBaseIterator
 	{
-	private:
+		static_assert(TIsDerivedFrom<ActorType, AActor>::IsDerived, "Type is not derived from AActor.");
+
+	protected:
 		typedef TMap<FGuid, TUniquePtr<FWorldPartitionActorDesc>*> MapType;
+		typedef typename FWorldPartitionActorDescType<ActorType>::Type ValueType;
 		typedef typename TChooseClass<bConst, MapType::TConstIterator, MapType::TIterator>::Result IteratorType;
-		typedef typename TChooseClass<bConst, const UActorDescContainer*, UActorDescContainer*>::Result ContainerType;
-		typedef typename TChooseClass<bConst, const FWorldPartitionActorDesc*, FWorldPartitionActorDesc*>::Result ValueType;
+		typedef typename TChooseClass<bConst, const UActorDescContainer*, UActorDescContainer*>::Result ContainerType;		
+		typedef typename TChooseClass<bConst, const ValueType*, ValueType*>::Result ReturnType;
 
 	public:
-		explicit TBaseIterator(ContainerType InActorDescContainer)
+		TBaseIterator(ContainerType InActorDescContainer)
 			: ActorsIterator(InActorDescContainer->Actors)
 		{
 			if (ShouldSkip())
@@ -89,9 +93,9 @@ public:
 		 *
 		 * @return	Current suitable actor desc
 		 */
-		FORCEINLINE ValueType operator*() const
+		FORCEINLINE ReturnType operator*() const
 		{
-			return ActorsIterator->Value->Get();
+			return StaticCast<ReturnType>(ActorsIterator->Value->Get());
 		}
 
 		/**
@@ -99,9 +103,9 @@ public:
 		 *
 		 * @return	Current suitable actor desc
 		 */
-		FORCEINLINE ValueType operator->() const
+		FORCEINLINE ReturnType operator->() const
 		{
-			return ActorsIterator->Value->Get();
+			return StaticCast<ReturnType>(ActorsIterator->Value->Get());
 		}
 		/**
 		 * Returns whether the iterator has reached the end and no longer points
@@ -114,10 +118,10 @@ public:
 			return (bool)ActorsIterator;
 		}
 
-	private:
+	protected:
 		/**
 		 * Determines whether the iterator currently points to a valid actor desc or not.
-		 * @return	true
+		 * @return true if we should skip the actor desc
 		 */
 		FORCEINLINE bool ShouldSkip() const
 		{
@@ -126,14 +130,32 @@ public:
 				return false;
 			}
 
-			return false;
+			return !ActorsIterator->Value->Get()->GetActorClass()->IsChildOf(ActorType::StaticClass());
 		}
 
-	private:
 		IteratorType ActorsIterator;
 	};
 
-	typedef TBaseIterator<false> TIterator;
-	typedef TBaseIterator<true> TConstIterator;
+	template <class ActorType = AActor>
+	class TIterator : public TBaseIterator<false, ActorType>
+	{
+		typedef TBaseIterator<false, ActorType> BaseType;
+
+	public:
+		TIterator(BaseType::ContainerType InActorDescContainer)
+			: BaseType(InActorDescContainer)
+		{}
+	};
+
+	template <class ActorType = AActor>
+	class TConstIterator : public TBaseIterator<true, ActorType>
+	{
+		typedef TBaseIterator<true, ActorType> BaseType;
+
+	public:
+		TConstIterator(BaseType::ContainerType InActorDescContainer)
+			: BaseType(InActorDescContainer)
+		{}
+	};
 #endif
 };
\ No newline at end of file
diff --git a/Engine/Source/Runtime/Engine/Public/WorldPartition/HLOD/HLODActor.h b/Engine/Source/Runtime/Engine/Public/WorldPartition/HLOD/HLODActor.h
--- a/Engine/Source/Runtime/Engine/Public/WorldPartition/HLOD/HLODActor.h
+++ b/Engine/Source/Runtime/Engine/Public/WorldPartition/HLOD/HLODActor.h
@@ -5,6 +5,7 @@
 #include "Containers/Array.h"
 #include "Containers/Set.h"
 #include "WorldPartition/WorldPartitionRuntimeCell.h"
+#include "WorldPartition/WorldPartitionActorDesc.h"
 #include "HLODActor.generated.h"
 
 class UHLODLayer;
@@ -105,6 +106,8 @@ private:
 	uint32 LODLevel;
 };
 
+DEFINE_ACTORDESC_TYPE(AWorldPartitionHLOD, FHLODActorDesc);
+
 UCLASS()
 class UWorldPartitionRuntimeHLODCellData : public UWorldPartitionRuntimeCellData
 {
diff --git a/Engine/Source/Runtime/Engine/Public/WorldPartition/WorldPartitionActorDescType.h b/Engine/Source/Runtime/Engine/Public/WorldPartition/WorldPartitionActorDescType.h
new file mode 100644
--- /dev/null
+++ b/Engine/Source/Runtime/Engine/Public/WorldPartition/WorldPartitionActorDescType.h
@@ -0,0 +1,19 @@
+// Copyright Epic Games, Inc. All Rights Reserved.
+#pragma once
+
+#include "CoreMinimal.h"
+
+#if WITH_EDITOR
+template <typename ActorType, TEMPLATE_REQUIRES(TIsDerivedFrom<ActorType, AActor>::IsDerived)>
+struct FWorldPartitionActorDescType
+{};
+
+#define DEFINE_ACTORDESC_TYPE(ActorType, ActorDescType) \
+template <>												\
+struct FWorldPartitionActorDescType<ActorType>			\
+{														\
+	typedef class ActorDescType Type;					\
+};
+#else
+#define DEFINE_ACTORDESC_TYPE(ActorType, ActorDescType)
+#endif
diff --git a/Engine/Source/Runtime/Landscape/Classes/LandscapeProxy.h b/Engine/Source/Runtime/Landscape/Classes/LandscapeProxy.h
--- a/Engine/Source/Runtime/Landscape/Classes/LandscapeProxy.h
+++ b/Engine/Source/Runtime/Landscape/Classes/LandscapeProxy.h
@@ -1213,4 +1213,6 @@ private:
 	UWorld* World;
 	int32 OudatedPhysicalMaterialComponentsCount;
 };
+
+DEFINE_ACTORDESC_TYPE(ALandscapeProxy, FLandscapeActorDesc);
 #endif
\ No newline at end of file
diff --git a/Engine/Source/Runtime/Landscape/Classes/LandscapeSplineActor.h b/Engine/Source/Runtime/Landscape/Classes/LandscapeSplineActor.h
--- a/Engine/Source/Runtime/Landscape/Classes/LandscapeSplineActor.h
+++ b/Engine/Source/Runtime/Landscape/Classes/LandscapeSplineActor.h
@@ -42,4 +42,6 @@ public:
 	virtual void CreateSplineComponent() override { check(false); }
 	virtual void CreateSplineComponent(const FVector& Scale3D) override { check(false); }
 #endif
-};
\ No newline at end of file
+};
+
+DEFINE_ACTORDESC_TYPE(ALandscapeSplineActor, FLandscapeSplineActorDesc);
\ No newline at end of file
diff --git a/Engine/Source/Runtime/Landscape/Private/Landscape.cpp b/Engine/Source/Runtime/Landscape/Private/Landscape.cpp
--- a/Engine/Source/Runtime/Landscape/Private/Landscape.cpp
+++ b/Engine/Source/Runtime/Landscape/Private/Landscape.cpp
@@ -3133,7 +3133,7 @@ void ULandscapeInfo::Initialize(UWorld* InWorld, const FGuid& InLandscapeGuid)
 	// If WorldPartition world get the Proxy handles
 	if (UWorldPartition* WorldPartition = InWorld->GetWorldPartition())
 	{
-		for (UActorDescContainer::TIterator ActorDescIterator(WorldPartition); ActorDescIterator; ++ActorDescIterator)
+		for (UActorDescContainer::TIterator<> ActorDescIterator(WorldPartition); ActorDescIterator; ++ActorDescIterator)
 		{
 			FWorldPartitionActorDesc* ActorDesc = *ActorDescIterator;
 			if (ActorDesc->GetActorClass()->IsChildOf(ALandscapeStreamingProxy::StaticClass()))
