Make PlayerStart always loaded by default in partitioned worlds (can be overriden).

diff --git a/Engine/Source/Runtime/Engine/Private/PlayerStart.cpp b/Engine/Source/Runtime/Engine/Private/PlayerStart.cpp
--- a/Engine/Source/Runtime/Engine/Private/PlayerStart.cpp
+++ b/Engine/Source/Runtime/Engine/Private/PlayerStart.cpp
@@ -61,6 +61,8 @@ APlayerStart::APlayerStart(const FObjectInitializer& ObjectInitializer)
 			ArrowComponent->bIsScreenSizeScaled = true;
 		}
 	}
+
+	GridPlacement = EActorGridPlacement::AlwaysLoaded;
 #endif // WITH_EDITORONLY_DATA
 }
 
