fix mac build

diff --git a/Engine/Source/Runtime/Engine/Public/WorldPartition/ActorDescContainer.h b/Engine/Source/Runtime/Engine/Public/WorldPartition/ActorDescContainer.h
--- a/Engine/Source/Runtime/Engine/Public/WorldPartition/ActorDescContainer.h
+++ b/Engine/Source/Runtime/Engine/Public/WorldPartition/ActorDescContainer.h
@@ -142,7 +142,7 @@ public:
 		typedef TBaseIterator<false, ActorType> BaseType;
 
 	public:
-		TIterator(BaseType::ContainerType InActorDescContainer)
+		TIterator(typename BaseType::ContainerType InActorDescContainer)
 			: BaseType(InActorDescContainer)
 		{}
 	};
@@ -153,7 +153,7 @@ public:
 		typedef TBaseIterator<true, ActorType> BaseType;
 
 	public:
-		TConstIterator(BaseType::ContainerType InActorDescContainer)
+		TConstIterator(typename BaseType::ContainerType InActorDescContainer)
 			: BaseType(InActorDescContainer)
 		{}
 	};
