LevelInstance: Fix for crash when assigning unsaved new level to level instance.

#jra UE-107560

diff --git a/Engine/Source/Runtime/Engine/Private/LevelInstance/LevelInstanceActor.cpp b/Engine/Source/Runtime/Engine/Private/LevelInstance/LevelInstanceActor.cpp
--- a/Engine/Source/Runtime/Engine/Private/LevelInstance/LevelInstanceActor.cpp
+++ b/Engine/Source/Runtime/Engine/Private/LevelInstance/LevelInstanceActor.cpp
@@ -359,6 +359,16 @@ bool ALevelInstance::CheckForLoop(TSoftObjectPtr<UWorld> InLevelInstance, TArray
 
 bool ALevelInstance::CanSetValue(TSoftObjectPtr<UWorld> InLevelInstance, FString* Reason) const
 {
+	FString PackageName;
+	if (!FPackageName::DoesPackageExist(InLevelInstance.GetLongPackageName(), nullptr, nullptr))
+	{
+		if (Reason)
+		{
+			*Reason = FString::Format(TEXT("Attempting to set Level Instance to package {0} which does not exist. Ensure the level was saved before attepting to set the level instance world asset."), { InLevelInstance.GetLongPackageName() });
+		}
+		return false;
+	}
+
 	TArray<TPair<FText, TSoftObjectPtr<UWorld>>> LoopInfo;
 	const ALevelInstance* LoopStart = nullptr;
 
