WP HLODs - Bugfix - Underlying actors does not contain the actor itself

diff --git a/Engine/Source/Runtime/Engine/Private/WorldPartition/HLOD/HLODBuilder.cpp b/Engine/Source/Runtime/Engine/Private/WorldPartition/HLOD/HLODBuilder.cpp
--- a/Engine/Source/Runtime/Engine/Private/WorldPartition/HLOD/HLODBuilder.cpp
+++ b/Engine/Source/Runtime/Engine/Private/WorldPartition/HLOD/HLODBuilder.cpp
@@ -22,6 +22,8 @@
 #include "Engine/StaticMesh.h"
 #include "Materials/Material.h"
 
+#include "LevelInstance/LevelInstanceActor.h"
+
 /**
  * Base class for all HLODBuilders
  */
@@ -97,31 +99,43 @@ public:
 		}
 	}
 
-	TArray<UPrimitiveComponent*> GatherPrimitiveComponents(uint32 InHLODLevel, const TArray<const AActor*> InActors)
+	TArray<UPrimitiveComponent*> GatherPrimitiveComponents(uint32 InHLODLevel, const TArray<const AActor*>& InActors)
 	{
 		TArray<UPrimitiveComponent*> PrimitiveComponents;
-		for (const AActor* SubActor : InActors)
-		{
-			// Gather all subactors
-			TSet<AActor*> UnderlyingActors;
-			SubActor->EditorGetUnderlyingActors(UnderlyingActors);
 
-			for (AActor* UnderlyingActor : UnderlyingActors)
+		auto GatherPrimitivesFromActor = [&PrimitiveComponents, InHLODLevel](const AActor* Actor)
+		{
+			if (UHLODLayer::ShouldIncludeInHLOD(Actor, Actor->IsA<ALevelInstance>()))
 			{
-				if (UHLODLayer::ShouldIncludeInHLOD(UnderlyingActor, true))
+				for (UActorComponent* SubComponent : Actor->GetComponents())
 				{
-					for (UActorComponent* SubComponent : UnderlyingActor->GetComponents())
+					if (UPrimitiveComponent* PrimitiveComponent = Cast<UPrimitiveComponent>(SubComponent))
 					{
-						if (UPrimitiveComponent* PrimitiveComponent = Cast<UPrimitiveComponent>(SubComponent))
+						if (UHLODLayer::ShouldIncludeInHLOD(PrimitiveComponent, InHLODLevel))
 						{
-							if (UHLODLayer::ShouldIncludeInHLOD(PrimitiveComponent, InHLODLevel))
-							{
-								PrimitiveComponents.Add(PrimitiveComponent);
-							}
+							PrimitiveComponents.Add(PrimitiveComponent);
 						}
 					}
 				}
 			}
+		};
+
+		TSet<AActor*> UnderlyingActors;
+
+		for (const AActor* Actor : InActors)
+		{
+			// Gather primitives from the Actor
+			GatherPrimitivesFromActor(Actor);
+
+			// Retrieve all underlying actors (ex: all sub actors of a LevelInstance)
+			UnderlyingActors.Reset();
+			Actor->EditorGetUnderlyingActors(UnderlyingActors);
+
+			// Gather primitives from underlying actors
+			for (const AActor* UnderlyingActor : UnderlyingActors)
+			{
+				GatherPrimitivesFromActor(UnderlyingActor);
+			}
 		}
 
 		return PrimitiveComponents;
