Fix non-unity build

diff --git a/Engine/Source/Runtime/Engine/Classes/GameFramework/Actor.h b/Engine/Source/Runtime/Engine/Classes/GameFramework/Actor.h
--- a/Engine/Source/Runtime/Engine/Classes/GameFramework/Actor.h
+++ b/Engine/Source/Runtime/Engine/Classes/GameFramework/Actor.h
@@ -21,6 +21,7 @@
 #if WITH_EDITOR
 #include "WorldPartition/DataLayer/ActorDataLayer.h"
 #include "WorldPartition/WorldPartitionActorDescType.h"
+#include "WorldPartition/WorldPartitionActorDesc.h"
 #endif
 
 #include "Actor.generated.h"
@@ -825,17 +826,17 @@ private:
 	/**
 	 * Creates an uninitialized actor descriptor. Can be overriden.
 	 */
-	virtual TUniquePtr<class FWorldPartitionActorDesc> CreateClassActorDesc() const;
+	virtual TUniquePtr<FWorldPartitionActorDesc> CreateClassActorDesc() const;
 
 	/**
 	 * Creates an initialized actor descriptor.
 	 */
-	TUniquePtr<class FWorldPartitionActorDesc> CreateActorDesc() const;
+	TUniquePtr<FWorldPartitionActorDesc> CreateActorDesc() const;
 
 	/**
 	 * Creates an uninitialized actor descriptor.
 	 */
-	static TUniquePtr<class FWorldPartitionActorDesc> CreateClassActorDesc(const TSubclassOf<AActor>& ActorClass);
+	static TUniquePtr<FWorldPartitionActorDesc> CreateClassActorDesc(const TSubclassOf<AActor>& ActorClass);
 #endif // WITH_EDITOR
 
 public:
diff --git a/Engine/Source/Runtime/Engine/Public/WorldPartition/WorldPartitionActorDesc.h b/Engine/Source/Runtime/Engine/Public/WorldPartition/WorldPartitionActorDesc.h
--- a/Engine/Source/Runtime/Engine/Public/WorldPartition/WorldPartitionActorDesc.h
+++ b/Engine/Source/Runtime/Engine/Public/WorldPartition/WorldPartitionActorDesc.h
@@ -29,6 +29,8 @@ namespace WorldPartitionTests
 }
 #endif
 
+enum class EActorGridPlacement : uint8;
+
 /**
  * Represents a potentially unloaded actor (editor-only)
  */
