Added buffer overrun prevention FSkeletalMeshLODModel::GetVertices() and a fatal error in cases it would have happened.

diff --git a/Engine/Source/Runtime/Engine/Private/SkeletalMeshLODModel.cpp b/Engine/Source/Runtime/Engine/Private/SkeletalMeshLODModel.cpp
--- a/Engine/Source/Runtime/Engine/Private/SkeletalMeshLODModel.cpp
+++ b/Engine/Source/Runtime/Engine/Private/SkeletalMeshLODModel.cpp
@@ -917,6 +917,21 @@ void FSkeletalMeshLODModel::GetVertices(TArray<FSoftSkinVertex>& Vertices) const
 	Vertices.Empty(NumVertices);
 	Vertices.AddUninitialized(NumVertices);
 
+	// validate NumVertices is correct
+	{
+		int32 TotalSoftVertices = 0;
+		for (int32 SectionIndex = 0; SectionIndex < Sections.Num(); SectionIndex++)
+		{
+			TotalSoftVertices += Sections[SectionIndex].SoftVertices.Num();
+		}
+		if (TotalSoftVertices != NumVertices)
+		{
+			// hitting this means NumVertices didn't match the sum of the vertex counts of all the sections,
+			// which could potentially overrun the Vertices buffer's allocation
+			UE_LOG(LogSkeletalMesh, Fatal, TEXT("NumVertices (%i) != TotalSoftVertices (%i)"), NumVertices, TotalSoftVertices);
+		}
+	}
+
 	// Initialize the vertex data
 	// All chunks are combined into one (rigid first, soft next)
 	FSoftSkinVertex* DestVertex = (FSoftSkinVertex*)Vertices.GetData();
