Fixed some enums comments syntax that had UHT associate comments to the wrong enum values.

diff --git a/Engine/Source/Runtime/Engine/Classes/GameFramework/Actor.h b/Engine/Source/Runtime/Engine/Classes/GameFramework/Actor.h
--- a/Engine/Source/Runtime/Engine/Classes/GameFramework/Actor.h
+++ b/Engine/Source/Runtime/Engine/Classes/GameFramework/Actor.h
@@ -44,10 +44,14 @@ class UDataLayer;
 UENUM(BlueprintType)
 enum class EActorUpdateOverlapsMethod : uint8
 {
-	UseConfigDefault,	// Use the default value specified by the native class or config value.
-	AlwaysUpdate,		// Always update overlap state on initialization.
-	OnlyUpdateMovable,	// Only update if root component has Movable mobility.
-	NeverUpdate			// Never update overlap state on initialization.
+	// Use the default value specified by the native class or config value.
+	UseConfigDefault,
+	// Always update overlap state on initialization.
+	AlwaysUpdate,
+	// Only update if root component has Movable mobility.
+	OnlyUpdateMovable,
+	// Never update overlap state on initialization.
+	NeverUpdate
 };
 
 #if WITH_EDITORONLY_DATA
@@ -55,9 +59,12 @@ enum class EActorUpdateOverlapsMethod : uint8
 UENUM()
 enum class EActorGridPlacement : uint8
 {
-	Bounds,						// Actor uses its bounds to determine in which runtime cells it's going to be placed.
-	Location,					// Actor uses its location to determine in which runtime cells it's going to be placed.
-	AlwaysLoaded,				// Actor is always loaded (not placed in the grid), also affects editor.
+	// Actor uses its bounds to determine in which runtime cells it's going to be placed.
+	Bounds,
+	// Actor uses its location to determine in which runtime cells it's going to be placed.
+	Location,
+	// Actor is always loaded (not placed in the grid), also affects editor.
+	AlwaysLoaded,
 	None UMETA(Hidden)
 };
 #endif
