Investigating stream caching crash when migrating in a project that had it disabled.

diff --git a/Engine/Source/Runtime/Engine/Private/AudioCompressionSettingsUtils.cpp b/Engine/Source/Runtime/Engine/Private/AudioCompressionSettingsUtils.cpp
--- a/Engine/Source/Runtime/Engine/Private/AudioCompressionSettingsUtils.cpp
+++ b/Engine/Source/Runtime/Engine/Private/AudioCompressionSettingsUtils.cpp
@@ -115,9 +115,17 @@ void CacheAudioCookOverrides(FPlatformAudioCookOverrides& OutOverrides, const TC
 	if (OutOverrides.bUseStreamCaching)
 	{
 		// Cache size:
-		int32 RetrievedCacheSize = 32 * 1024;
+		const int32 DefaultCacheSize = 32 * 1024;
+		int32 RetrievedCacheSize = DefaultCacheSize;
 		int32 RetrievedChunkSizeOverride = INDEX_NONE;
+
 		PlatformFile.GetInt(*CategoryName, TEXT("CacheSizeKB"), RetrievedCacheSize);
+		if (!RetrievedCacheSize)
+		{
+			PlatformFile.SetInt64(*CategoryName, TEXT("CacheSizeKB"), DefaultCacheSize);
+			RetrievedCacheSize = DefaultCacheSize;
+		}
+
 		OutOverrides.StreamCachingSettings.CacheSizeKB = RetrievedCacheSize;
 
 		PlatformFile.GetInt(*CategoryName, TEXT("MaxChunkSizeOverrideKB"), RetrievedChunkSizeOverride);
diff --git a/Engine/Source/Runtime/Engine/Private/AudioStreamingCache.cpp b/Engine/Source/Runtime/Engine/Private/AudioStreamingCache.cpp
--- a/Engine/Source/Runtime/Engine/Private/AudioStreamingCache.cpp
+++ b/Engine/Source/Runtime/Engine/Private/AudioStreamingCache.cpp
@@ -703,6 +703,7 @@ FAudioChunkCache::FAudioChunkCache(uint32 InMaxChunkSize, uint32 NumChunks, uint
 	, MemoryLimitBytes(InMemoryLimitInBytes)
 	, bLogCacheMisses(false)
 {
+	check(NumChunks > 0);
 	CachePool.Reset(NumChunks);
 	for (uint32 Index = 0; Index < NumChunks; Index++)
 	{
