Removed shader compiler log spam

diff --git a/Engine/Source/Runtime/Engine/Private/ShaderCompiler/ShaderCompiler.cpp b/Engine/Source/Runtime/Engine/Private/ShaderCompiler/ShaderCompiler.cpp
--- a/Engine/Source/Runtime/Engine/Private/ShaderCompiler/ShaderCompiler.cpp
+++ b/Engine/Source/Runtime/Engine/Private/ShaderCompiler/ShaderCompiler.cpp
@@ -1808,7 +1808,7 @@ int32 FShaderCompileThreadRunnable::PullTasksFromQueue()
 					NumPendingJobs = Manager->AllJobs.GetNumPendingJobs((EShaderCompileJobPriority)PriorityIndex);
 					if (NumPendingJobs > 0)
 					{
-						UE_LOG(LogShaderCompilers, Display, TEXT("Worker (%d/%d): shaders left to compile %i"), WorkerIndex + 1, WorkerInfos.Num(), NumPendingJobs);
+						UE_LOG(LogShaderCompilers, Verbose, TEXT("Worker (%d/%d): shaders left to compile %i"), WorkerIndex + 1, WorkerInfos.Num(), NumPendingJobs);
 
 						int32 MaxNumJobs = 1;
 						if (PriorityIndex < (int32)EShaderCompileJobPriority::High)
