Fix z-fighting between two redundant skylight icons when a sky light actor is placed in the scene

diff --git a/Engine/Source/Runtime/Engine/Private/Components/SkyLightComponent.cpp b/Engine/Source/Runtime/Engine/Private/Components/SkyLightComponent.cpp
--- a/Engine/Source/Runtime/Engine/Private/Components/SkyLightComponent.cpp
+++ b/Engine/Source/Runtime/Engine/Private/Components/SkyLightComponent.cpp
@@ -1066,32 +1066,10 @@ ASkyLight::ASkyLight(const FObjectInitializer& ObjectInitializer)
 	RootComponent = LightComponent;
 
 #if WITH_EDITORONLY_DATA
-	if (!IsRunningCommandlet())
-	{
-	// Structure to hold one-time initialization
-	struct FConstructorStatics
+	// Null out the sprite. The Skylight components sprite is the one we use.
+	if (GetSpriteComponent())
 	{
-			ConstructorHelpers::FObjectFinderOptional<UTexture2D> SkyLightTextureObject;
-		FName ID_Sky;
-		FText NAME_Sky;
-
-		FConstructorStatics()
-				: SkyLightTextureObject(TEXT("/Engine/EditorResources/LightIcons/SkyLight"))
-				, ID_Sky(TEXT("Sky"))
-			, NAME_Sky(NSLOCTEXT( "SpriteCategory", "Sky", "Sky" ))
-		{
-		}
-	};
-	static FConstructorStatics ConstructorStatics;
-
-		if (GetSpriteComponent())
-		{
-			GetSpriteComponent()->Sprite = ConstructorStatics.SkyLightTextureObject.Get();
-			GetSpriteComponent()->SpriteInfo.Category = ConstructorStatics.ID_Sky;
-			GetSpriteComponent()->SpriteInfo.DisplayName = ConstructorStatics.NAME_Sky;
-			GetSpriteComponent()->SetRelativeScale3D(FVector(0.5f, 0.5f, 0.5f));
-			GetSpriteComponent()->SetupAttachment(LightComponent);
-		}
+		GetSpriteComponent()->Sprite = nullptr;
 	}
 #endif // WITH_EDITORONLY_DATA
 }
