Merge fixes to the shader jobs cache from CLs 15308711, 15308851.

Original descriptions were:
Remove the duplicates of a job even if it falied.

- Otherwise a failed job would result in its duplicates being held up indefinitely

Prevent cooks on the buildfarm from hanging indefinitely in case of SC getting stuck.

diff --git a/Engine/Source/Runtime/Engine/Private/ShaderCompiler/ShaderCompiler.cpp b/Engine/Source/Runtime/Engine/Private/ShaderCompiler/ShaderCompiler.cpp
--- a/Engine/Source/Runtime/Engine/Private/ShaderCompiler/ShaderCompiler.cpp
+++ b/Engine/Source/Runtime/Engine/Private/ShaderCompiler/ShaderCompiler.cpp
@@ -503,12 +503,6 @@ void FShaderCompileJobCollection::AddToCacheAndProcessPending(FShaderCommonCompi
 		return;
 	}
 
-	if (!FinishedJob->bSucceeded)
-	{
-		// we only cache jobs that succeded
-		return;
-	}
-
 	// TODO: reduce the scopes
 	FWriteScopeLock JobLocker(Lock);
 
@@ -517,24 +511,16 @@ void FShaderCompileJobCollection::AddToCacheAndProcessPending(FShaderCommonCompi
 	FMemoryWriter Writer(Output);
 	FinishedJob->SerializeOutput(Writer);
 
-	CompletedJobsCache.Add(InputHash, Output);
-
 	// see if there are outstanding jobs that also need to be resolved
+	int32 NumOutstandingJobsWithSameHash = 0;
 	if (FShaderCommonCompileJob** WaitList = DuplicateJobsWaitList.Find(InputHash))
 	{
-		int32 NumOutstandingJobsWithSameHash = 0;
 		FShaderCommonCompileJob* CurHead = *WaitList;
 		while (CurHead)
 		{
 			checkf(CurHead != FinishedJob, TEXT("Job that is being added to cache was also on a waiting list! Error in bookkeeping."));
 
-#if 0 // while we could reuse the same ouput, let's go through the cache to influence the statistics (might be important for LRU later)
 			FMemoryReader MemReader(Output);
-#else
-			TArray<uint8>* ExistingOutput = CompletedJobsCache.Find(InputHash);
-			checkf(ExistingOutput, TEXT("At this point we expect a newly cached job to be found"));
-			FMemoryReader MemReader(*ExistingOutput);
-#endif
 			CurHead->SerializeOutput(MemReader);
 
 			// finish the job instantly
@@ -552,6 +538,12 @@ void FShaderCompileJobCollection::AddToCacheAndProcessPending(FShaderCommonCompi
 			UE_LOG(LogShaderCompilers, UE_SHADERCACHE_LOG_LEVEL, TEXT("Processed %d outstanding jobs with the same ihash %s."), NumOutstandingJobsWithSameHash, *InputHash.ToString());
 		}
 	}
+
+	if (FinishedJob->bSucceeded)
+	{
+		// we only cache jobs that succeded
+		CompletedJobsCache.Add(InputHash, Output, NumOutstandingJobsWithSameHash);
+	}
 }
 
 void FShaderCompileJobCollection::LogCachingStats(bool bForceLogIgnoringTimeInverval)
@@ -2685,6 +2677,12 @@ FShaderCompilingManager::FShaderCompilingManager() :
 		{
 			UE_LOG(LogShaderCompilers, Warning, TEXT("bForceUseSCWMemoryPressureLimits was set but missing one or more prerequisite setting(s): CookerMemoryUsedInGB, MemoryToLeaveForTheOSInGB, MemoryUsedPerSCWProcessInGB.  Ignoring bForceUseSCWMemoryPressureLimits") );
 		}
+
+		if (GIsBuildMachine)
+		{
+			// force crashes on hung shader maps on build machines, to prevent builds running for days
+			GCrashOnHungShaderMaps = 1;
+		}
 	}
 #endif
 
@@ -6305,7 +6303,7 @@ uint64 FShaderJobCache::GetCurrentMemoryBudget() const
 	return FMath::Min(AbsoluteLimit, RelativeLimit);
 }
 
-void FShaderJobCache::Add(const FJobInputHash& Hash, const FJobCachedOutput& Contents)
+void FShaderJobCache::Add(const FJobInputHash& Hash, const FJobCachedOutput& Contents, int32 InitialHitCount)
 {
 	if (!ShaderCompiler::IsJobCacheEnabled())
 	{
@@ -6380,7 +6378,7 @@ void FShaderJobCache::Add(const FJobInputHash& Hash, const FJobCachedOutput& Con
 		}
 
 		FStoredOutput* NewStoredOutput = new FStoredOutput();
-		NewStoredOutput->NumHits = 0;
+		NewStoredOutput->NumHits = InitialHitCount;
 		NewStoredOutput->NumReferences = 1;
 		NewStoredOutput->JobOutput = Contents;
 		Outputs.Add(OutputHash, NewStoredOutput);
diff --git a/Engine/Source/Runtime/Engine/Public/ShaderCompiler.h b/Engine/Source/Runtime/Engine/Public/ShaderCompiler.h
--- a/Engine/Source/Runtime/Engine/Public/ShaderCompiler.h
+++ b/Engine/Source/Runtime/Engine/Public/ShaderCompiler.h
@@ -329,7 +329,7 @@ public:
 	FJobCachedOutput* Find(const FJobInputHash& Hash);
 
 	/** Adds a job output to the cache */
-	void Add(const FJobInputHash& Hash, const FJobCachedOutput& Contents);
+	void Add(const FJobInputHash& Hash, const FJobCachedOutput& Contents, int InitialHitCount);
 
 	/** Calculates memory used by the cache*/
 	uint64 GetAllocatedMemory();
