Avoid repeatedly querying config system for supported shader targets from the UMaterial::IsCachedCookedPlatformDataLoaded polling method as it can be called millions of times during a cook.

diff --git a/Engine/Source/Runtime/Engine/Private/Materials/Material.cpp b/Engine/Source/Runtime/Engine/Private/Materials/Material.cpp
--- a/Engine/Source/Runtime/Engine/Private/Materials/Material.cpp
+++ b/Engine/Source/Runtime/Engine/Private/Materials/Material.cpp
@@ -4139,9 +4139,6 @@ void UMaterial::BeginCacheForCookedPlatformData( const ITargetPlatform *TargetPl
 
 bool UMaterial::IsCachedCookedPlatformDataLoaded( const ITargetPlatform* TargetPlatform ) 
 {
-	TArray<FName> DesiredShaderFormats;
-	TargetPlatform->GetAllTargetedShaderFormats(DesiredShaderFormats);
-
 	const TArray<FMaterialResource*>* CachedMaterialResourcesForPlatform = CachedMaterialResourcesForCooking.Find( TargetPlatform );
 
 	if ( CachedMaterialResourcesForPlatform != NULL ) // this should always succeed if begincacheforcookedcplatformdata is called first
