sepculative fix for mesh merge load order crashing incremental shootergame cooks
worked for aaron. :)

diff --git a/Engine/Source/Runtime/Engine/Private/MeshSimplificationSettings.cpp b/Engine/Source/Runtime/Engine/Private/MeshSimplificationSettings.cpp
--- a/Engine/Source/Runtime/Engine/Private/MeshSimplificationSettings.cpp
+++ b/Engine/Source/Runtime/Engine/Private/MeshSimplificationSettings.cpp
@@ -3,6 +3,10 @@
 #include "Engine/MeshSimplificationSettings.h"
 #include "UObject/UnrealType.h"
 
+#if WITH_EDITOR
+#include "Developer/MeshReductionInterface/Public/IMeshReductionManagerModule.h"
+#endif
+
 UMeshSimplificationSettings::UMeshSimplificationSettings(const FObjectInitializer& ObjectInitializer)
 	: Super(ObjectInitializer)
 {
@@ -26,6 +30,7 @@ void UMeshSimplificationSettings::PostInitProperties()
 	Super::PostInitProperties(); 
 
 #if WITH_EDITOR
+	IMeshReductionManagerModule& MeshReductionModule = FModuleManager::Get().LoadModuleChecked<IMeshReductionManagerModule>("MeshReductionInterface");
 	if(IsTemplate())
 	{
 		ImportConsoleVariableValues();
