Disable hardware ray-traced shadows by default. Global use of hardware ray-traced shadows is now selectable as a project setting in the Hardware Ray Tracing section: "Ray Tracing Shadows."

When disabled, any light component's "Cast Ray Tracing Shadows" parameter will be greyed out.

diff --git a/Engine/Source/Runtime/Engine/Classes/Engine/RendererSettings.h b/Engine/Source/Runtime/Engine/Classes/Engine/RendererSettings.h
--- a/Engine/Source/Runtime/Engine/Classes/Engine/RendererSettings.h
+++ b/Engine/Source/Runtime/Engine/Classes/Engine/RendererSettings.h
@@ -357,6 +357,11 @@ class ENGINE_API URendererSettings : public UDeveloperSettings
 		ConfigRestartRequired = true))
 		uint32 bEnableRayTracing : 1;
 
+	UPROPERTY(config, EditAnywhere, Category = HardwareRayTracing, meta = (
+		ConsoleVariable = "r.RayTracing.Shadows", DisplayName = "Ray Tracing Shadows",
+		ToolTip = "Enables hardware ray-traced shadows."))
+		uint32 bEnableRayTracingShadows : 1;
+
 	UPROPERTY(config, EditAnywhere, Category = HardwareRayTracing, meta = (
 		ConsoleVariable = "r.RayTracing.UseTextureLod", DisplayName = "Texture LOD",
 		ToolTip = "Enable automatic texture mip level selection in ray tracing material shaders. Unchecked: highest resolution mip level is used for all texture (default). Checked: texture LOD is approximated based on total ray length, output resolution and texel density at hit point (ray cone method).",
diff --git a/Engine/Source/Runtime/Engine/Private/Components/LightComponent.cpp b/Engine/Source/Runtime/Engine/Private/Components/LightComponent.cpp
--- a/Engine/Source/Runtime/Engine/Private/Components/LightComponent.cpp
+++ b/Engine/Source/Runtime/Engine/Private/Components/LightComponent.cpp
@@ -640,6 +640,12 @@ bool ULightComponent::CanEditChange(const FProperty* InProperty) const
 			return bCanEdit;
 		}
 
+		if (PropertyName == GET_MEMBER_NAME_STRING_CHECKED(ULightComponent, bCastRaytracedShadow))
+		{
+			static const auto CVar = IConsoleManager::Get().FindTConsoleVariableDataInt(TEXT("r.RayTracing.Shadows"));
+			return CVar->GetValueOnGameThread() != 0;
+		}
+
 		if (PropertyName == GET_MEMBER_NAME_STRING_CHECKED(ULightComponent, LightFunctionScale)
 			|| PropertyName == GET_MEMBER_NAME_STRING_CHECKED(ULightComponent, LightFunctionFadeDistance)
 			|| PropertyName == GET_MEMBER_NAME_STRING_CHECKED(ULightComponent, DisabledBrightness))
diff --git a/Engine/Source/Runtime/Engine/Private/RendererSettings.cpp b/Engine/Source/Runtime/Engine/Private/RendererSettings.cpp
--- a/Engine/Source/Runtime/Engine/Private/RendererSettings.cpp
+++ b/Engine/Source/Runtime/Engine/Private/RendererSettings.cpp
@@ -72,6 +72,7 @@ URendererSettings::URendererSettings(const FObjectInitializer& ObjectInitializer
 	GPUSimulationTextureSizeX = 1024;
 	GPUSimulationTextureSizeY = 1024;
 	bEnableRayTracing = 0;
+	bEnableRayTracingShadows = 0;
 	bEnableRayTracingTextureLOD = 0;
 }
 
diff --git a/Engine/Source/Runtime/Renderer/Private/LightRendering.cpp b/Engine/Source/Runtime/Renderer/Private/LightRendering.cpp
--- a/Engine/Source/Runtime/Renderer/Private/LightRendering.cpp
+++ b/Engine/Source/Runtime/Renderer/Private/LightRendering.cpp
@@ -61,9 +61,9 @@ static FAutoConsoleVariableRef CVarAllowSimpleLights(
 
 static TAutoConsoleVariable<int32> CVarRayTracingOcclusion(
 	TEXT("r.RayTracing.Shadows"),
-	1,
-	TEXT("0: use traditional rasterized shadow map\n")
-	TEXT("1: use ray tracing shadows (default)"),
+	0,
+	TEXT("0: use traditional rasterized shadow map (default)\n")
+	TEXT("1: use ray tracing shadows"),
 	ECVF_RenderThreadSafe | ECVF_Scalability);
 
 static int32 GShadowRayTracingSamplesPerPixel = 1;
