Always overwrite DDC version of FVirtualTextureBuiltData
This is important to ensure that a given machine doesn't get into a state where the FVirtualTextureBuiltData in the DDC is references (slightly) different chunks built by another machine, due to non-determinism.
If this happens, the local machine will constantly regenerate VT data every load, causing long delays

diff --git a/Engine/Source/Runtime/Engine/Private/TextureDerivedData.cpp b/Engine/Source/Runtime/Engine/Private/TextureDerivedData.cpp
--- a/Engine/Source/Runtime/Engine/Private/TextureDerivedData.cpp
+++ b/Engine/Source/Runtime/Engine/Private/TextureDerivedData.cpp
@@ -652,6 +652,7 @@ uint32 PutDerivedDataInCache(FTexturePlatformData* DerivedData, const FString& D
 	}
 
 	// Write out each VT chunk to the DDC
+	bool bReplaceExistingDerivedDataDDC = bReplaceExistingDDC;
 	if (DerivedData->VTData)
 	{
 		const int32 ChunkCount = DerivedData->VTData->Chunks.Num();
@@ -665,6 +666,10 @@ uint32 PutDerivedDataInCache(FTexturePlatformData* DerivedData, const FString& D
 
 			TotalBytesPut += Chunk.StoreInDerivedDataCache(ChunkDerivedDataKey, TextureName, bReplaceExistingDDC);
 		}
+
+		// VT always needs to replace the FVirtualTextureBuiltData in the DDC, otherwise we can be left in a situation where a local client is constantly attempting to rebuild chunks,
+		// but failing to generate chunks that match the FVirtualTextureBuiltData in the DDC, due to non-determinism in texture generation
+		bReplaceExistingDerivedDataDDC = true;
 	}
 
 	// Store derived data.
@@ -672,7 +677,7 @@ uint32 PutDerivedDataInCache(FTexturePlatformData* DerivedData, const FString& D
 	FMemoryWriter Ar(RawDerivedData, /*bIsPersistent=*/ true);
 	DerivedData->Serialize(Ar, NULL);
 	TotalBytesPut += RawDerivedData.Num();
-	GetDerivedDataCacheRef().Put(*DerivedDataKey, RawDerivedData, TextureName, bReplaceExistingDDC);
+	GetDerivedDataCacheRef().Put(*DerivedDataKey, RawDerivedData, TextureName, bReplaceExistingDerivedDataDDC);
 	UE_LOG(LogTexture,Verbose,TEXT("%s  Derived Data: %d bytes"),*LogString,RawDerivedData.Num());
 	return TotalBytesPut;
 }
