Fix VT mip bias calculation when clamping max texture size

diff --git a/Engine/Source/Runtime/Engine/Private/VT/VirtualTextureDataBuilder.cpp b/Engine/Source/Runtime/Engine/Private/VT/VirtualTextureDataBuilder.cpp
--- a/Engine/Source/Runtime/Engine/Private/VT/VirtualTextureDataBuilder.cpp
+++ b/Engine/Source/Runtime/Engine/Private/VT/VirtualTextureDataBuilder.cpp
@@ -960,8 +960,17 @@ void FVirtualTextureDataBuilder::BuildSourcePixels(const FTextureSourceData& Sou
 			check(bBuildTextureResult);
 
 			// Get size of block from Compressor output, since it may have been padded/adjusted
-			BlockData.SizeX = CompressedMips[0].SizeX;
-			BlockData.SizeY = CompressedMips[0].SizeY;
+			{
+				BlockData.SizeX = CompressedMips[0].SizeX;
+				BlockData.SizeY = CompressedMips[0].SizeY;
+
+				// re-compute mip bias to account for any resizing of this block (typically due to clamped max size)
+				const int32 MipBiasX = FMath::CeilLogTwo(BlockSizeX / BlockData.SizeX);
+				const int32 MipBiasY = FMath::CeilLogTwo(BlockSizeY / BlockData.SizeY);
+				checkf(MipBiasX == MipBiasY, TEXT("Mismatched aspect ratio (%d x %d), (%d x %d)"), BlockSizeX, BlockSizeY, BlockData.SizeX, BlockData.SizeY);
+				BlockData.MipBias = MipBiasX;
+			}
+
 			check(BlockData.SizeX << BlockData.MipBias == BlockSizeX);
 			check(BlockData.SizeY << BlockData.MipBias == BlockSizeY);
 
