Strata is now hidden in  UE5EA.

diff --git a/Engine/Source/Editor/DetailCustomizations/Private/MaterialAttributePropertyDetails.cpp b/Engine/Source/Editor/DetailCustomizations/Private/MaterialAttributePropertyDetails.cpp
--- a/Engine/Source/Editor/DetailCustomizations/Private/MaterialAttributePropertyDetails.cpp
+++ b/Engine/Source/Editor/DetailCustomizations/Private/MaterialAttributePropertyDetails.cpp
@@ -25,6 +25,10 @@ void FMaterialAttributePropertyDetails::CustomizeDetails(IDetailLayoutBuilder& D
 	AttributeDisplayNameList.Empty(AttributeNameToIDList.Num());
 	for (const TPair<FString, FGuid>& NameGUIDPair : AttributeNameToIDList)
 	{
+		if (NameGUIDPair.Key == TEXT("FrontMaterial"))
+		{
+			continue;// STRATA_DISABLED
+		}
 		AttributeDisplayNameList.Add(MakeShareable(new FString(NameGUIDPair.Key)));
 	}
 
diff --git a/Engine/Source/Editor/UnrealEd/Private/MaterialGraph.cpp b/Engine/Source/Editor/UnrealEd/Private/MaterialGraph.cpp
--- a/Engine/Source/Editor/UnrealEd/Private/MaterialGraph.cpp
+++ b/Engine/Source/Editor/UnrealEd/Private/MaterialGraph.cpp
@@ -87,7 +87,7 @@ void UMaterialGraph::RebuildGraphInternal(const TMap<UMaterialExpression*, TArra
 
 		MaterialInputs.Add(FMaterialInputInfo(FMaterialAttributeDefinitionMap::GetDisplayNameForMaterial(MP_PixelDepthOffset, Material), MP_PixelDepthOffset, LOCTEXT("PixelDepthOffsetToolTip", "Pixel Depth Offset")));
 		MaterialInputs.Add(FMaterialInputInfo(FMaterialAttributeDefinitionMap::GetDisplayNameForMaterial(MP_ShadingModel, Material), MP_ShadingModel, LOCTEXT("ShadingModelToolTip", "Selects which shading model should be used per pixel")));
-		MaterialInputs.Add(FMaterialInputInfo(FMaterialAttributeDefinitionMap::GetDisplayNameForMaterial(MP_FrontMaterial, Material), MP_FrontMaterial, LOCTEXT("FrontMaterialToolTip", "Specify the front facing material")));
+		// STRATA_DISABLED MaterialInputs.Add(FMaterialInputInfo(FMaterialAttributeDefinitionMap::GetDisplayNameForMaterial(MP_FrontMaterial, Material), MP_FrontMaterial, LOCTEXT("FrontMaterialToolTip", "Specify the front facing material")));
 
 		//^^^ New material properties go above here. ^^^^
 		MaterialInputs.Add(FMaterialInputInfo(LOCTEXT("MaterialAttributes", "Material Attributes"), MP_MaterialAttributes, LOCTEXT("MaterialAttributesToolTip", "Material Attributes")));
diff --git a/Engine/Source/Runtime/Engine/Classes/Engine/RendererSettings.h b/Engine/Source/Runtime/Engine/Classes/Engine/RendererSettings.h
--- a/Engine/Source/Runtime/Engine/Classes/Engine/RendererSettings.h
+++ b/Engine/Source/Runtime/Engine/Classes/Engine/RendererSettings.h
@@ -725,19 +725,19 @@ class ENGINE_API URendererSettings : public UDeveloperSettings
 	/**
 	"Enable Strata materials (Beta)."
 	*/
-	UPROPERTY(config, EditAnywhere, Category = Strata, meta = (
-		ConsoleVariable = "r.Strata", DisplayName = "Strata materials (Experimental)",
-		ToolTip = "Enable Strata materials (Experimental).",
-		ConfigRestartRequired = true))
+//STRATA_DISABLED	UPROPERTY(config, EditAnywhere, Category = Strata, meta = (
+//		ConsoleVariable = "r.Strata", DisplayName = "Strata materials (Experimental)",
+//		ToolTip = "Enable Strata materials (Experimental).",
+//		ConfigRestartRequired = true))
 		uint32 bEnableStrata : 1;
 
 	/**
 	"Strata allocated byte per pixel to store materials data. Higher value means more complex material can be represented."
 	*/
-	UPROPERTY(config, EditAnywhere, Category = Strata, meta = (
-		ConsoleVariable = "r.Strata.BytesPerPixel", DisplayName = "Strata bytes per pixel",
-		ToolTip = "Strata allocated byte per pixel to store materials data. Higher value means more complex material can be represented.",
-		ConfigRestartRequired = true))
+//STRATA_DISABLED	UPROPERTY(config, EditAnywhere, Category = Strata, meta = (
+//			ConsoleVariable = "r.Strata.BytesPerPixel", DisplayName = "Strata bytes per pixel",
+//			ToolTip = "Strata allocated byte per pixel to store materials data. Higher value means more complex material can be represented.",
+//			ConfigRestartRequired = true))
 		uint32 StrataBytePerPixel;
 
 	/**
diff --git a/Engine/Source/Runtime/Engine/Classes/Materials/MaterialExpressionStrata.h b/Engine/Source/Runtime/Engine/Classes/Materials/MaterialExpressionStrata.h
--- a/Engine/Source/Runtime/Engine/Classes/Materials/MaterialExpressionStrata.h
+++ b/Engine/Source/Runtime/Engine/Classes/Materials/MaterialExpressionStrata.h
@@ -29,7 +29,7 @@ class UMaterialExpressionStrataBSDF : public UMaterialExpression
 	GENERATED_UCLASS_BODY()
 };
 
-UCLASS(MinimalAPI, collapsecategories, hidecategories = Object, DisplayName = "Strata Slab BSDF")
+UCLASS(MinimalAPI, collapsecategories, hidecategories = Object, DisplayName = "Strata Slab BSDF", Abstract)// STRATA_DISABLED
 class UMaterialExpressionStrataSlabBSDF : public UMaterialExpressionStrataBSDF
 {
 	GENERATED_UCLASS_BODY()
@@ -135,7 +135,7 @@ class UMaterialExpressionStrataSlabBSDF : public UMaterialExpressionStrataBSDF
 	//~ End UMaterialExpression Interface
 };
 
-UCLASS(MinimalAPI, collapsecategories, hidecategories = Object, DisplayName = "Strata Sheen BSDF")
+UCLASS(MinimalAPI, collapsecategories, hidecategories = Object, DisplayName = "Strata Sheen BSDF", Abstract)// STRATA_DISABLED
 class UMaterialExpressionStrataSheenBSDF : public UMaterialExpressionStrataBSDF
 {
 	GENERATED_UCLASS_BODY()
@@ -170,7 +170,7 @@ class UMaterialExpressionStrataSheenBSDF : public UMaterialExpressionStrataBSDF
 	//~ End UMaterialExpression Interface
 };
 
-UCLASS(MinimalAPI, collapsecategories, hidecategories = Object, DisplayName = "Strata Volumetric-Fog-Cloud BSDF")
+UCLASS(MinimalAPI, collapsecategories, hidecategories = Object, DisplayName = "Strata Volumetric-Fog-Cloud BSDF", Abstract)// STRATA_DISABLED
 class UMaterialExpressionStrataVolumetricFogCloudBSDF : public UMaterialExpressionStrataBSDF
 {
 	GENERATED_UCLASS_BODY()
@@ -211,7 +211,7 @@ class UMaterialExpressionStrataVolumetricFogCloudBSDF : public UMaterialExpressi
 	//~ End UMaterialExpression Interface
 };
 
-UCLASS(MinimalAPI, collapsecategories, hidecategories = Object, DisplayName = "Strata Unlit BSDF")
+UCLASS(MinimalAPI, collapsecategories, hidecategories = Object, DisplayName = "Strata Unlit BSDF", Abstract)// STRATA_DISABLED
 class UMaterialExpressionStrataUnlitBSDF : public UMaterialExpressionStrataBSDF
 {
 	GENERATED_UCLASS_BODY()
@@ -240,7 +240,7 @@ class UMaterialExpressionStrataUnlitBSDF : public UMaterialExpressionStrataBSDF
 	//~ End UMaterialExpression Interface
 };
 
-UCLASS(MinimalAPI, collapsecategories, hidecategories = Object, DisplayName = "Strata Hair BSDF")
+UCLASS(MinimalAPI, collapsecategories, hidecategories = Object, DisplayName = "Strata Hair BSDF", Abstract)// STRATA_DISABLED
 class UMaterialExpressionStrataHairBSDF : public UMaterialExpressionStrataBSDF
 {
 	GENERATED_UCLASS_BODY()
@@ -299,7 +299,7 @@ class UMaterialExpressionStrataHairBSDF : public UMaterialExpressionStrataBSDF
 	//~ End UMaterialExpression Interface
 };
 
-UCLASS(MinimalAPI, collapsecategories, hidecategories = Object, DisplayName = "Strata Single Layer Water BSDF")
+UCLASS(MinimalAPI, collapsecategories, hidecategories = Object, DisplayName = "Strata Single Layer Water BSDF", Abstract)// STRATA_DISABLED
 class UMaterialExpressionStrataSingleLayerWaterBSDF : public UMaterialExpressionStrataBSDF
 {
 	GENERATED_UCLASS_BODY()
@@ -387,7 +387,7 @@ class UMaterialExpressionStrataSingleLayerWaterBSDF : public UMaterialExpression
 ///////////////////////////////////////////////////////////////////////////////
 // Operator nodes
 
-UCLASS(MinimalAPI, collapsecategories, hidecategories = Object, DisplayName = "Strata BSDF Horizontal Blend")
+UCLASS(MinimalAPI, collapsecategories, hidecategories = Object, DisplayName = "Strata BSDF Horizontal Blend", Abstract)// STRATA_DISABLED
 class UMaterialExpressionStrataHorizontalMixing : public UMaterialExpressionStrataBSDF
 {
 	GENERATED_UCLASS_BODY()
@@ -422,7 +422,7 @@ class UMaterialExpressionStrataHorizontalMixing : public UMaterialExpressionStra
 	//~ End UMaterialExpression Interface
 };
 
-UCLASS(MinimalAPI, collapsecategories, hidecategories = Object, DisplayName = "Strata BSDF Vertical Layer")
+UCLASS(MinimalAPI, collapsecategories, hidecategories = Object, DisplayName = "Strata BSDF Vertical Layer", Abstract)// STRATA_DISABLED
 class UMaterialExpressionStrataVerticalLayering : public UMaterialExpressionStrataBSDF
 {
 	GENERATED_UCLASS_BODY()
@@ -451,7 +451,7 @@ class UMaterialExpressionStrataVerticalLayering : public UMaterialExpressionStra
 	//~ End UMaterialExpression Interface
 };
 
-UCLASS(MinimalAPI, collapsecategories, hidecategories = Object, DisplayName = "Strata BSDF Add")
+UCLASS(MinimalAPI, collapsecategories, hidecategories = Object, DisplayName = "Strata BSDF Add", Abstract)// STRATA_DISABLED
 class UMaterialExpressionStrataAdd : public UMaterialExpressionStrataBSDF
 {
 	GENERATED_UCLASS_BODY()
@@ -480,7 +480,7 @@ class UMaterialExpressionStrataAdd : public UMaterialExpressionStrataBSDF
 	//~ End UMaterialExpression Interface
 };
 
-UCLASS(MinimalAPI, collapsecategories, hidecategories = Object, DisplayName = "Strata BSDF Weight")
+UCLASS(MinimalAPI, collapsecategories, hidecategories = Object, DisplayName = "Strata BSDF Weight", Abstract)// STRATA_DISABLED
 class UMaterialExpressionStrataMultiply : public UMaterialExpressionStrataBSDF
 {
 	GENERATED_UCLASS_BODY()
@@ -514,7 +514,7 @@ class UMaterialExpressionStrataMultiply : public UMaterialExpressionStrataBSDF
 ///////////////////////////////////////////////////////////////////////////////
 // Utilities
 
-UCLASS(MinimalAPI, collapsecategories, hidecategories = Object, DisplayName = "Strata Anisotropy-To-Roughnesses")
+UCLASS(MinimalAPI, collapsecategories, hidecategories = Object, DisplayName = "Strata Anisotropy-To-Roughnesses", Abstract)// STRATA_DISABLED
 class UMaterialExpressionStrataAnisotropyToRoughness : public UMaterialExpressionStrataBSDF
 {
 	GENERATED_UCLASS_BODY()
diff --git a/Engine/Source/Runtime/Renderer/Private/Strata/Strata.cpp b/Engine/Source/Runtime/Renderer/Private/Strata/Strata.cpp
--- a/Engine/Source/Runtime/Renderer/Private/Strata/Strata.cpp
+++ b/Engine/Source/Runtime/Renderer/Private/Strata/Strata.cpp
@@ -39,12 +39,12 @@ namespace Strata
 
 bool IsStrataEnabled()
 {
-	return CVarStrata.GetValueOnAnyThread() > 0;
+	return false;// STRATA_DISABLED CVarStrata.GetValueOnAnyThread() > 0;
 }
 
 bool IsClassificationEnabled()
 {
-	return CVarStrataClassification.GetValueOnAnyThread() > 0;
+	return false;// STRATA_DISABLED CVarStrataClassification.GetValueOnAnyThread() > 0;
 }
 
 uint32 GetStrataTileSize()
