Fix non-editor builds.

diff --git a/Engine/Source/Runtime/Engine/Private/World.cpp b/Engine/Source/Runtime/Engine/Private/World.cpp
--- a/Engine/Source/Runtime/Engine/Private/World.cpp
+++ b/Engine/Source/Runtime/Engine/Private/World.cpp
@@ -1518,10 +1518,12 @@ void UWorld::RepairWorldSettings()
 
 void UWorld::RepairWorldDataLayers()
 {
+#if WITH_EDITOR
 	if (const UWorldPartition* WorldPartition = GetWorldPartition())
 	{
 		AWorldDataLayers::Get(this, /*bCreateIfNotFound*/true);
 	}
+#endif
 }
 
 void UWorld::InitWorld(const InitializationValues IVS)
diff --git a/Engine/Source/Runtime/Engine/Public/WorldPartition/DataLayer/DataLayer.h b/Engine/Source/Runtime/Engine/Public/WorldPartition/DataLayer/DataLayer.h
--- a/Engine/Source/Runtime/Engine/Public/WorldPartition/DataLayer/DataLayer.h
+++ b/Engine/Source/Runtime/Engine/Public/WorldPartition/DataLayer/DataLayer.h
@@ -28,13 +28,14 @@ public:
 	class UHLODLayer* GetDefaultHLODLayer() const { return ShouldGenerateHLODs() ? DefaultHLODLayer : nullptr; }
 
 	static FText GetDataLayerText(const UDataLayer* InDataLayer);
+
+	bool IsLocked() const { return bIsLocked; }
 #endif
 	
 	FName GetDataLayerLabel() const  { return DataLayerLabel; }
 	bool IsVisible() const { return bIsVisible; }
 	bool IsDynamicallyLoaded() const { return bIsDynamicallyLoaded; }
-	bool IsInitiallyActive() const { return IsDynamicallyLoaded() && bIsInitiallyActive; }
-	bool IsLocked() const { return bIsLocked; }
+	bool IsInitiallyActive() const { return IsDynamicallyLoaded() && bIsInitiallyActive; }	
 
 private:
 	/** The display name of the DataLayer */
