Missing animation tracks when compressed
#fix Corrected condition for removing missing / identity bone tracks

diff --git a/Engine/Source/Runtime/Engine/Private/Animation/AnimCompressionDerivedData.cpp b/Engine/Source/Runtime/Engine/Private/Animation/AnimCompressionDerivedData.cpp
--- a/Engine/Source/Runtime/Engine/Private/Animation/AnimCompressionDerivedData.cpp
+++ b/Engine/Source/Runtime/Engine/Private/Animation/AnimCompressionDerivedData.cpp
@@ -35,7 +35,7 @@ const TCHAR* FDerivedDataAnimationCompression::GetVersionString() const
 	// This is a version string that mimics the old versioning scheme. If you
 	// want to bump this version, generate a new guid using VS->Tools->Create GUID and
 	// return it here. Ex.
-	return TEXT("565E584FC51D4ACB9C17A9F9C54CB545");	
+	return TEXT("C006D6DDA62143F0915F60671659A4FD");
 }
 
 bool FDerivedDataAnimationCompression::Build( TArray<uint8>& OutDataArray )
diff --git a/Engine/Source/Runtime/Engine/Private/Animation/AnimCompressionTypes.cpp b/Engine/Source/Runtime/Engine/Private/Animation/AnimCompressionTypes.cpp
--- a/Engine/Source/Runtime/Engine/Private/Animation/AnimCompressionTypes.cpp
+++ b/Engine/Source/Runtime/Engine/Private/Animation/AnimCompressionTypes.cpp
@@ -160,13 +160,14 @@ FCompressibleAnimData::FCompressibleAnimData(class UAnimSequence* InSeq, const b
 	};
 
 	// Verify bone track names and data, removing any bone that does not exist on the skeleton
+    // And for additive animations remove any track deemed not to add any additive animation (identity rotation and zero-vector translation and scale)
 	// Note on (TrackIndex > 0) below : deliberately stop before track 0, compression code doesn't like getting a completely empty animation
 	for (int32 TrackIndex = RawAnimationData.Num() - 1; TrackIndex > 0; --TrackIndex)
 	{
 		const FRawAnimSequenceTrack& Track = RawAnimationData[TrackIndex];
 		// Try find correct bone index
 		const int32 BoneIndex = RefSkeleton.FindBoneIndex(TempTrackNames[TrackIndex]);
-		if (IsRawTrackValidForRemoval(Track) || (BoneIndex == INDEX_NONE && !bBakeAdditive))
+		if ((bBakeAdditive && IsRawTrackValidForRemoval(Track)) || BoneIndex == INDEX_NONE)
 		{
 			RawAnimationData.RemoveAtSwap(TrackIndex, 1, false);
 			TempTrackNames.RemoveAtSwap(TrackIndex, 1, false);
