LevelInstance: Support Scaling

diff --git a/Engine/Source/Runtime/Engine/Classes/Engine/LevelStreamingDynamic.h b/Engine/Source/Runtime/Engine/Classes/Engine/LevelStreamingDynamic.h
--- a/Engine/Source/Runtime/Engine/Classes/Engine/LevelStreamingDynamic.h
+++ b/Engine/Source/Runtime/Engine/Classes/Engine/LevelStreamingDynamic.h
@@ -50,6 +50,8 @@ class ENGINE_API ULevelStreamingDynamic : public ULevelStreaming
  	UFUNCTION(BlueprintCallable, Category = LevelStreaming, meta=(DisplayName = "Load Level Instance (by Object Reference)", WorldContext="WorldContextObject"))
  	static ULevelStreamingDynamic* LoadLevelInstanceBySoftObjectPtr(UObject* WorldContextObject, TSoftObjectPtr<UWorld> Level, FVector Location, FRotator Rotation, bool& bOutSuccess, const FString& OptionalLevelNameOverride = TEXT(""), TSubclassOf<ULevelStreamingDynamic> OptionalLevelStreamingClass = nullptr);
  	
+	static ULevelStreamingDynamic* LoadLevelInstanceBySoftObjectPtr(UObject* WorldContextObject, TSoftObjectPtr<UWorld> Level, const FTransform LevelTransform, bool& bOutSuccess, const FString& OptionalLevelNameOverride = TEXT(""), TSubclassOf<ULevelStreamingDynamic> OptionalLevelStreamingClass = nullptr);
+
 	//~ Begin UObject Interface
 	virtual void PostLoad() override;
 	//~ End UObject Interface
@@ -65,7 +67,7 @@ private:
 	// Counter used by LoadLevelInstance to create unique level names
 	static int32 UniqueLevelInstanceId;
 
- 	static ULevelStreamingDynamic* LoadLevelInstance_Internal(UWorld* World, const FString& LongPackageName, FVector Location, FRotator Rotation, bool& bOutSuccess, const FString& OptionalLevelNameOverride, TSubclassOf<ULevelStreamingDynamic> OptionalLevelStreamingClass);
+ 	static ULevelStreamingDynamic* LoadLevelInstance_Internal(UWorld* World, const FString& LongPackageName, const FTransform LevelTransform, bool& bOutSuccess, const FString& OptionalLevelNameOverride, TSubclassOf<ULevelStreamingDynamic> OptionalLevelStreamingClass);
 
 };
 
diff --git a/Engine/Source/Runtime/Engine/Private/LevelInstance/LevelInstanceEditorInstanceActor.cpp b/Engine/Source/Runtime/Engine/Private/LevelInstance/LevelInstanceEditorInstanceActor.cpp
--- a/Engine/Source/Runtime/Engine/Private/LevelInstance/LevelInstanceEditorInstanceActor.cpp
+++ b/Engine/Source/Runtime/Engine/Private/LevelInstance/LevelInstanceEditorInstanceActor.cpp
@@ -35,7 +35,7 @@ ALevelInstanceEditorInstanceActor* ALevelInstanceEditorInstanceActor::Create(ALe
 	SpawnParams.ObjectFlags = RF_Transient;
 	SpawnParams.bNoFail = true;
 	ALevelInstanceEditorInstanceActor* InstanceActor = LevelInstanceActor->GetWorld()->SpawnActor<ALevelInstanceEditorInstanceActor>(LevelInstanceActor->GetActorLocation(), LevelInstanceActor->GetActorRotation(), SpawnParams);
-
+	InstanceActor->SetActorScale3D(LevelInstanceActor->GetActorScale3D());
 	InstanceActor->SetLevelInstanceID(LevelInstanceActor->GetLevelInstanceID());
 	
 	for (AActor* LevelActor : LoadedLevel->Actors)
diff --git a/Engine/Source/Runtime/Engine/Private/LevelInstance/LevelInstanceLevelStreaming.cpp b/Engine/Source/Runtime/Engine/Private/LevelInstance/LevelInstanceLevelStreaming.cpp
--- a/Engine/Source/Runtime/Engine/Private/LevelInstance/LevelInstanceLevelStreaming.cpp
+++ b/Engine/Source/Runtime/Engine/Private/LevelInstance/LevelInstanceLevelStreaming.cpp
@@ -59,7 +59,7 @@ ULevelStreamingLevelInstance* ULevelStreamingLevelInstance::LoadInstance(ALevelI
 	// Build a unique and deterministic LevelInstance level instance name by using LevelInstanceID. 
 	// Distinguish game from editor since we don't want to duplicate for PIE already loaded editor instances (not yet supported).
 	FString Suffix = FString::Printf(TEXT("%s_LevelInstance_%08X_%d"), *ShortPackageName, LevelInstanceActor->GetLevelInstanceID(), LevelInstanceActor->GetWorld()->IsGameWorld() ? 1 : 0);
-	ULevelStreamingLevelInstance* LevelStreaming = Cast<ULevelStreamingLevelInstance>(ULevelStreamingDynamic::LoadLevelInstanceBySoftObjectPtr(LevelInstanceActor->GetWorld(), LevelInstanceActor->GetWorldAsset(), LevelInstanceActor->GetActorLocation(), LevelInstanceActor->GetActorRotation(), bOutSuccess, Suffix, ULevelStreamingLevelInstance::StaticClass()));
+	ULevelStreamingLevelInstance* LevelStreaming = Cast<ULevelStreamingLevelInstance>(ULevelStreamingDynamic::LoadLevelInstanceBySoftObjectPtr(LevelInstanceActor->GetWorld(), LevelInstanceActor->GetWorldAsset(), LevelInstanceActor->GetActorTransform(), bOutSuccess, Suffix, ULevelStreamingLevelInstance::StaticClass()));
 	if (bOutSuccess)
 	{
 		LevelStreaming->LevelInstanceID = LevelInstanceActor->GetLevelInstanceID();
diff --git a/Engine/Source/Runtime/Engine/Private/LevelStreaming.cpp b/Engine/Source/Runtime/Engine/Private/LevelStreaming.cpp
--- a/Engine/Source/Runtime/Engine/Private/LevelStreaming.cpp
+++ b/Engine/Source/Runtime/Engine/Private/LevelStreaming.cpp
@@ -1905,10 +1905,15 @@ ULevelStreamingDynamic* ULevelStreamingDynamic::LoadLevelInstance(UObject* World
 		return nullptr;
 	}
 
-	return LoadLevelInstance_Internal(World, LongPackageName, Location, Rotation, bOutSuccess, OptionalLevelNameOverride, OptionalLevelStreamingClass);
+	return LoadLevelInstance_Internal(World, LongPackageName, FTransform(Rotation, Location), bOutSuccess, OptionalLevelNameOverride, OptionalLevelStreamingClass);
 }
 
 ULevelStreamingDynamic* ULevelStreamingDynamic::LoadLevelInstanceBySoftObjectPtr(UObject* WorldContextObject, const TSoftObjectPtr<UWorld> Level, const FVector Location, const FRotator Rotation, bool& bOutSuccess, const FString& OptionalLevelNameOverride, TSubclassOf<ULevelStreamingDynamic> OptionalLevelStreamingClass)
+{
+	return LoadLevelInstanceBySoftObjectPtr(WorldContextObject, Level, FTransform(Rotation, Location), bOutSuccess, OptionalLevelNameOverride, OptionalLevelStreamingClass);
+}
+
+ULevelStreamingDynamic* ULevelStreamingDynamic::LoadLevelInstanceBySoftObjectPtr(UObject* WorldContextObject, const TSoftObjectPtr<UWorld> Level, const FTransform LevelTransform, bool& bOutSuccess, const FString& OptionalLevelNameOverride, TSubclassOf<ULevelStreamingDynamic> OptionalLevelStreamingClass)
 {
 	bOutSuccess = false;
 	UWorld* World = GEngine->GetWorldFromContextObject(WorldContextObject, EGetWorldErrorMode::LogAndReturnNull);
@@ -1923,10 +1928,10 @@ ULevelStreamingDynamic* ULevelStreamingDynamic::LoadLevelInstanceBySoftObjectPtr
 		return nullptr;
 	}
 
-	return LoadLevelInstance_Internal(World, Level.GetLongPackageName(), Location, Rotation, bOutSuccess, OptionalLevelNameOverride, OptionalLevelStreamingClass);
+	return LoadLevelInstance_Internal(World, Level.GetLongPackageName(), LevelTransform, bOutSuccess, OptionalLevelNameOverride, OptionalLevelStreamingClass);
 }
 
-ULevelStreamingDynamic* ULevelStreamingDynamic::LoadLevelInstance_Internal(UWorld* World, const FString& LongPackageName, const FVector Location, const FRotator Rotation, bool& bOutSuccess, const FString& OptionalLevelNameOverride, TSubclassOf<ULevelStreamingDynamic> OptionalLevelStreamingClass)
+ULevelStreamingDynamic* ULevelStreamingDynamic::LoadLevelInstance_Internal(UWorld* World, const FString& LongPackageName, const FTransform LevelTransform, bool& bOutSuccess, const FString& OptionalLevelNameOverride, TSubclassOf<ULevelStreamingDynamic> OptionalLevelStreamingClass)
 {
 	const FString PackagePath = FPackageName::GetLongPackagePath(LongPackageName);
 	FString ShortPackageName = FPackageName::GetShortName(LongPackageName);
@@ -1999,7 +2004,7 @@ ULevelStreamingDynamic* ULevelStreamingDynamic::LoadLevelInstance_Internal(UWorl
     StreamingLevel->bInitiallyLoaded = true;
     StreamingLevel->bInitiallyVisible = true;
 	// Transform
-    StreamingLevel->LevelTransform = FTransform(Rotation, Location);
+    StreamingLevel->LevelTransform = LevelTransform;
 	// Map to Load
     StreamingLevel->PackageNameToLoad = FName(*OnDiskPackageName);
           
diff --git a/Engine/Source/Runtime/Engine/Private/LevelUtils.cpp b/Engine/Source/Runtime/Engine/Private/LevelUtils.cpp
--- a/Engine/Source/Runtime/Engine/Private/LevelUtils.cpp
+++ b/Engine/Source/Runtime/Engine/Private/LevelUtils.cpp
@@ -390,6 +390,7 @@ void FLevelUtils::ApplyLevelTransform(const FLevelUtils::FApplyLevelTransformPar
 				{
 					RootComponent->SetRelativeLocation_Direct(TransformParams.LevelTransform.TransformPosition(RootComponent->GetRelativeLocation()));
 					RootComponent->SetRelativeRotation_Direct(TransformParams.LevelTransform.TransformRotation(RootComponent->GetRelativeRotation().Quaternion()).Rotator());
+					RootComponent->SetRelativeScale3D_Direct(TransformParams.LevelTransform.GetScale3D() * RootComponent->GetRelativeScale3D());
 				}
 			}
 			else
@@ -399,6 +400,7 @@ void FLevelUtils::ApplyLevelTransform(const FLevelUtils::FApplyLevelTransformPar
 				if (RootComponent && RootComponent->GetAttachParent() == nullptr)
 				{
 					RootComponent->SetRelativeLocationAndRotation(TransformParams.LevelTransform.TransformPosition(RootComponent->GetRelativeLocation()), TransformParams.LevelTransform.TransformRotation(RootComponent->GetRelativeRotation().Quaternion()));
+					RootComponent->SetRelativeScale3D(TransformParams.LevelTransform.GetScale3D() * RootComponent->GetRelativeScale3D());
 				}
 			}
 #if WITH_EDITOR
@@ -428,6 +430,7 @@ void FLevelUtils::ApplyLevelTransform(const FLevelUtils::FApplyLevelTransformPar
 				{
 					ModelComponent->SetRelativeLocation_Direct(TransformParams.LevelTransform.TransformPosition(ModelComponent->GetRelativeLocation()));
 					ModelComponent->SetRelativeRotation_Direct(TransformParams.LevelTransform.TransformRotation(ModelComponent->GetRelativeRotation().Quaternion()).Rotator());
+					ModelComponent->SetRelativeScale3D_Direct(TransformParams.LevelTransform.GetScale3D() * ModelComponent->GetRelativeScale3D());
 				}
 			}
 
@@ -443,6 +446,7 @@ void FLevelUtils::ApplyLevelTransform(const FLevelUtils::FApplyLevelTransformPar
 					{
 						RootComponent->SetRelativeLocation_Direct(TransformParams.LevelTransform.TransformPosition(RootComponent->GetRelativeLocation()));
 						RootComponent->SetRelativeRotation_Direct(TransformParams.LevelTransform.TransformRotation(RootComponent->GetRelativeRotation().Quaternion()).Rotator());
+						RootComponent->SetRelativeScale3D_Direct(TransformParams.LevelTransform.GetScale3D() * RootComponent->GetRelativeScale3D());
 					}
 				}
 			}
@@ -455,6 +459,7 @@ void FLevelUtils::ApplyLevelTransform(const FLevelUtils::FApplyLevelTransformPar
 				if (ModelComponent)
 				{
 					ModelComponent->SetRelativeLocationAndRotation(TransformParams.LevelTransform.TransformPosition(ModelComponent->GetRelativeLocation()), TransformParams.LevelTransform.TransformRotation(ModelComponent->GetRelativeRotation().Quaternion()));
+					ModelComponent->SetRelativeScale3D(TransformParams.LevelTransform.GetScale3D() * ModelComponent->GetRelativeScale3D());
 				}
 			}
 
@@ -469,6 +474,7 @@ void FLevelUtils::ApplyLevelTransform(const FLevelUtils::FApplyLevelTransformPar
 					if (RootComponent && RootComponent->GetAttachParent() == nullptr)
 					{
 						RootComponent->SetRelativeLocationAndRotation(TransformParams.LevelTransform.TransformPosition(RootComponent->GetRelativeLocation()), TransformParams.LevelTransform.TransformRotation(RootComponent->GetRelativeRotation().Quaternion()));
+						RootComponent->SetRelativeScale3D(TransformParams.LevelTransform.GetScale3D() * RootComponent->GetRelativeScale3D());
 					}
 				}
 			}
