Temporarily disable UWorld::RepairWorldDataLayers, until I make a proper fix for PIE duplication.

diff --git a/Engine/Source/Runtime/Engine/Private/World.cpp b/Engine/Source/Runtime/Engine/Private/World.cpp
--- a/Engine/Source/Runtime/Engine/Private/World.cpp
+++ b/Engine/Source/Runtime/Engine/Private/World.cpp
@@ -1518,7 +1518,7 @@ void UWorld::RepairWorldSettings()
 
 void UWorld::RepairWorldDataLayers()
 {
-#if WITH_EDITOR
+#if 0
 	if (const UWorldPartition* WorldPartition = GetWorldPartition())
 	{
 		AWorldDataLayers::Get(this, /*bCreateIfNotFound*/true);
