Fix NonUnity + Static Analysis

diff --git a/Engine/Source/Runtime/Engine/Private/ISMPartition/ISMComponentData.cpp b/Engine/Source/Runtime/Engine/Private/ISMPartition/ISMComponentData.cpp
--- a/Engine/Source/Runtime/Engine/Private/ISMPartition/ISMComponentData.cpp
+++ b/Engine/Source/Runtime/Engine/Private/ISMPartition/ISMComponentData.cpp
@@ -3,6 +3,7 @@
 #include "ISMPartition/ISMComponentData.h"
 #include "Components/InstancedStaticMeshComponent.h"
 #include "Components/HierarchicalInstancedStaticMeshComponent.h"
+#include "Engine/StaticMesh.h"
 
 FISMComponentData::FISMComponentData()
 #if WITH_EDITORONLY_DATA
diff --git a/Engine/Source/Runtime/Engine/Private/ISMPartition/ISMComponentDescriptor.cpp b/Engine/Source/Runtime/Engine/Private/ISMPartition/ISMComponentDescriptor.cpp
--- a/Engine/Source/Runtime/Engine/Private/ISMPartition/ISMComponentDescriptor.cpp
+++ b/Engine/Source/Runtime/Engine/Private/ISMPartition/ISMComponentDescriptor.cpp
@@ -104,7 +104,7 @@ bool FISMComponentDescriptor::operator==(const FISMComponentDescriptor& Other) c
 	BodyInstance.GetCollisionEnabled() == Other.BodyInstance.GetCollisionEnabled() && 
 	BodyInstance.GetCollisionResponse() == Other.BodyInstance.GetCollisionResponse() &&
 	BodyInstance.DoesUseCollisionProfile() == Other.BodyInstance.DoesUseCollisionProfile() &&
-	(!BodyInstance.DoesUseCollisionProfile() || (BodyInstance.GetCollisionProfileName() == BodyInstance.GetCollisionProfileName()));
+	(!BodyInstance.DoesUseCollisionProfile() || (BodyInstance.GetCollisionProfileName() == Other.BodyInstance.GetCollisionProfileName()));
 }
 
 uint32 FISMComponentDescriptor::ComputeHash() const
diff --git a/Engine/Source/Runtime/Engine/Private/ISMPartition/ISMPartitionActor.cpp b/Engine/Source/Runtime/Engine/Private/ISMPartition/ISMPartitionActor.cpp
--- a/Engine/Source/Runtime/Engine/Private/ISMPartition/ISMPartitionActor.cpp
+++ b/Engine/Source/Runtime/Engine/Private/ISMPartition/ISMPartitionActor.cpp
@@ -4,6 +4,7 @@
 #include "Components/HierarchicalInstancedStaticMeshComponent.h"
 #include "Engine/World.h"
 #include "Engine/StaticMesh.h"
+#include "EngineUtils.h"
 
 #if WITH_EDITOR
 #include "Editor.h"
diff --git a/Engine/Source/Runtime/Engine/Public/ISMPartition/ISMComponentData.h b/Engine/Source/Runtime/Engine/Public/ISMPartition/ISMComponentData.h
--- a/Engine/Source/Runtime/Engine/Public/ISMPartition/ISMComponentData.h
+++ b/Engine/Source/Runtime/Engine/Public/ISMPartition/ISMComponentData.h
@@ -5,6 +5,8 @@
 #include "Engine/EngineTypes.h"
 #include "ISMComponentData.generated.h"
 
+class UInstancedStaticMeshComponent;
+
 /** Represents one component instance with a mapping back to the Client instance. The should be a 1 to 1 relationship between those and the actual Component instances. (equal count) */
 USTRUCT()
 struct FISMComponentInstance
