#fix incorrect trimming when using AnimationData::Trim
#misc incorrect number of frames used when inserting multiple in the Anim Sequenc editor

diff --git a/Engine/Source/Editor/Persona/Private/AnimTimeline/SAnimTimeline.cpp b/Engine/Source/Editor/Persona/Private/AnimTimeline/SAnimTimeline.cpp
--- a/Engine/Source/Editor/Persona/Private/AnimTimeline/SAnimTimeline.cpp
+++ b/Engine/Source/Editor/Persona/Private/AnimTimeline/SAnimTimeline.cpp
@@ -527,7 +527,7 @@ void SAnimTimeline::OnAppendAnimSequence( bool bFromStart, int32 NumOfFrames )
 
 			// Crop the raw anim data.
 			const int32 StartFrame = (bFromStart)? 0 : AnimSequence->GetDataModel()->GetNumberOfFrames() - 1;
-			UE::Anim::AnimationData::DuplicateKeys(AnimSequence, StartFrame, 1, StartFrame);
+			UE::Anim::AnimationData::DuplicateKeys(AnimSequence, StartFrame, NumOfFrames, StartFrame);
 		}
 	}
 }
diff --git a/Engine/Source/Runtime/Engine/Private/Animation/AnimSequenceHelpers.cpp b/Engine/Source/Runtime/Engine/Private/Animation/AnimSequenceHelpers.cpp
--- a/Engine/Source/Runtime/Engine/Private/Animation/AnimSequenceHelpers.cpp
+++ b/Engine/Source/Runtime/Engine/Private/Animation/AnimSequenceHelpers.cpp
@@ -775,10 +775,11 @@ bool AnimationData::Trim(UAnimSequence* InSequence, float TrimStart, float TrimE
 		const FFrameNumber StartFrameTrim = FrameRate.AsFrameNumber(TrimStart);
 		const FFrameNumber EndFrameTrim = FrameRate.AsFrameNumber(TrimEnd);
 
+		const int32 StartTrimKeyIndex = StartFrameTrim.Value + 1;
 		const int32 NumTrimmedFrames = EndFrameTrim.Value - StartFrameTrim.Value;
 		
 		UAnimDataController::FScopedBracket ScopedBracket(Controller, LOCTEXT("TrimRawAnimation_Bracket", "Trimming Animation Track Data"));
-		RemoveKeys(InSequence, StartFrameTrim.Value, NumTrimmedFrames);
+		RemoveKeys(InSequence, StartTrimKeyIndex, NumTrimmedFrames);
 		
 		return true;
 	}
@@ -885,11 +886,12 @@ void AnimationData::RemoveKeys(UAnimSequence* InSequence, int32 StartKeyIndex, i
 		const int32 NewNumberOfFrames = FMath::Max(NewNumberOfKeys - 1, 0);
 		const float NewSequenceLength = FrameRate.AsSeconds(NewNumberOfFrames);
 
-		const float StartTime = FrameRate.AsSeconds(StartKeyIndex);
+		const float StartFrameTime = FrameRate.AsSeconds(StartKeyIndex - 1);
 		const float RemovedTime = NumKeysToRemove * FrameRate.AsInterval();
+		const float EndTime = StartFrameTime + RemovedTime;
 
 		// Notify will happen with time slice that was removed
-		Controller->Resize(NewSequenceLength, StartTime, RemovedTime);
+		Controller->Resize(NewSequenceLength, StartFrameTime, EndTime);
 	}
 }
 
