Avoid creation of a temporary actor descriptor reference while iterating through loaded actors.

diff --git a/Engine/Source/Runtime/Engine/Private/WorldPartition/WorldPartition.cpp b/Engine/Source/Runtime/Engine/Private/WorldPartition/WorldPartition.cpp
--- a/Engine/Source/Runtime/Engine/Private/WorldPartition/WorldPartition.cpp
+++ b/Engine/Source/Runtime/Engine/Private/WorldPartition/WorldPartition.cpp
@@ -658,7 +658,16 @@ void UWorldPartition::UpdateLoadingEditorCell(UWorldPartitionEditorCell* Cell, b
 			}
 			else
 			{
-				Cell->LoadedActors.Remove(ActorHandle);
+				// Don't call LoadedActors.Remove(ActorHandle) right away, as this will create a temporary reference and might try to load.
+				for (const FWorldPartitionReference& ActorReference: Cell->LoadedActors)
+				{
+					if (ActorReference == ActorHandle)
+					{
+						Cell->LoadedActors.Remove(ActorHandle);
+						break;
+					}
+				}
+
 				bPotentiallyUnloadedActors = true;
 			}
 		}
