Added support for replacing the internal objects pointed to by an object-based element handle

This currently requires that either the old or new elements have no external references, as redirecting two elements with external references will require support for redirection within the element framework itself

#rnx

diff --git a/Engine/Source/Editor/UnrealEd/Classes/Editor/EditorEngine.h b/Engine/Source/Editor/UnrealEd/Classes/Editor/EditorEngine.h
--- a/Engine/Source/Editor/UnrealEd/Classes/Editor/EditorEngine.h
+++ b/Engine/Source/Editor/UnrealEd/Classes/Editor/EditorEngine.h
@@ -596,13 +596,13 @@ public:
 	/**	Broadcasts that a blueprint just finished being reinstanced. THIS SHOULD NOT BE PUBLIC */
 	void BroadcastBlueprintReinstanced() { BlueprintReinstanced.Broadcast(); }
 
-	/** Called when uobjects have been replaced to allow others a chance to fix their references. */
-	typedef TMap<UObject*, UObject*> ReplacementObjectMap;
-	DECLARE_EVENT_OneParam( UEditorEngine, FObjectsReplacedEvent, const ReplacementObjectMap& );
-	FObjectsReplacedEvent& OnObjectsReplaced() { return ObjectsReplacedEvent; }
+	/** Called when UObjects have been replaced to allow others a chance to fix their references. */
+	// TODO: Deprecate in favor of using FCoreUObjectDelegates::OnObjectsReplaced directly?
+	using FObjectsReplacedEvent = FCoreUObjectDelegates::FOnObjectsReplaced;
+	FObjectsReplacedEvent& OnObjectsReplaced() { return FCoreUObjectDelegates::OnObjectsReplaced; }
 
 	/**	Broadcasts that objects have been replaced*/
-	void BroadcastBlueprintCompiled(const TMap<UObject*, UObject*>& ReplacementMap) { ObjectsReplacedEvent.Broadcast(ReplacementMap); }
+	void BroadcastBlueprintCompiled(const TMap<UObject*, UObject*>& ReplacementMap) { FCoreUObjectDelegates::OnObjectsReplaced.Broadcast(ReplacementMap); }
 
 	/** Called when a package with data-driven classes becomes loaded or unloaded */
 	DECLARE_EVENT( UEditorEngine, FClassPackageLoadedOrUnloadedEvent );
@@ -2696,9 +2696,6 @@ private:
 	/** Helper function to show undo/redo notifications */
 	void ShowUndoRedoNotification(const FText& NotificationText, bool bSuccess);
 
-	/**	Broadcasts that the supplied objects have been replaced (map of old object to new object */
-	void BroadcastObjectsReplaced(const TMap<UObject*, UObject*>& ReplacementMap) { ObjectsReplacedEvent.Broadcast(ReplacementMap); }
-
 	/** Internal helper functions */
 	virtual void PostUndo (bool bSuccess);
 
@@ -2777,9 +2774,6 @@ private:
 	/** Delegate broadcast when blueprint is reinstanced */
 	FBlueprintReinstanced BlueprintReinstanced;
 
-	/** Delegate broadcast when objects have been replaced (e.g on blueprint compile) */
-	FObjectsReplacedEvent ObjectsReplacedEvent;
-
 	/** Delegate broadcast when a package has been loaded or unloaded */
 	FClassPackageLoadedOrUnloadedEvent ClassPackageLoadedOrUnloadedEvent;
 
diff --git a/Engine/Source/Editor/UnrealEd/Private/EditorEngine.cpp b/Engine/Source/Editor/UnrealEd/Private/EditorEngine.cpp
--- a/Engine/Source/Editor/UnrealEd/Private/EditorEngine.cpp
+++ b/Engine/Source/Editor/UnrealEd/Private/EditorEngine.cpp
@@ -5998,7 +5998,7 @@ void UEditorEngine::NotifyToolsOfObjectReplacement(const TMap<UObject*, UObject*
 	}
 
 	// Allow any other observers to act upon the object replacement
-	BroadcastObjectsReplaced(OldToNewInstanceMap);
+	FCoreUObjectDelegates::OnObjectsReplaced.Broadcast(OldToNewInstanceMap);
 
 	// Check to see if any selected components were reinstanced, as a final step.
 	USelection* ComponentSelection = GetSelectedComponents();
diff --git a/Engine/Source/Runtime/CoreUObject/Private/UObject/UObjectGlobals.cpp b/Engine/Source/Runtime/CoreUObject/Private/UObject/UObjectGlobals.cpp
--- a/Engine/Source/Runtime/CoreUObject/Private/UObject/UObjectGlobals.cpp
+++ b/Engine/Source/Runtime/CoreUObject/Private/UObject/UObjectGlobals.cpp
@@ -135,6 +135,7 @@ FCoreUObjectDelegates::FOnObjectPropertyChanged FCoreUObjectDelegates::OnObjectP
 TSet<UObject*> FCoreUObjectDelegates::ObjectsModifiedThisFrame;
 FCoreUObjectDelegates::FOnObjectModified FCoreUObjectDelegates::OnObjectModified;
 FCoreUObjectDelegates::FOnObjectTransacted FCoreUObjectDelegates::OnObjectTransacted;
+FCoreUObjectDelegates::FOnObjectsReplaced FCoreUObjectDelegates::OnObjectsReplaced;
 FCoreUObjectDelegates::FOnAssetLoaded FCoreUObjectDelegates::OnAssetLoaded;
 FCoreUObjectDelegates::FOnObjectSaved FCoreUObjectDelegates::OnObjectSaved;
 #endif // WITH_EDITOR
diff --git a/Engine/Source/Runtime/CoreUObject/Public/UObject/UObjectGlobals.h b/Engine/Source/Runtime/CoreUObject/Public/UObject/UObjectGlobals.h
--- a/Engine/Source/Runtime/CoreUObject/Public/UObject/UObjectGlobals.h
+++ b/Engine/Source/Runtime/CoreUObject/Public/UObject/UObjectGlobals.h
@@ -2315,6 +2315,11 @@ struct COREUOBJECT_API FCoreUObjectDelegates
 	DECLARE_MULTICAST_DELEGATE_TwoParams(FOnObjectTransacted, UObject*, const class FTransactionObjectEvent&);
 	static FOnObjectTransacted OnObjectTransacted;
 
+	/** Called when UObjects have been replaced to allow others a chance to fix their references */
+	using FReplacementObjectMap = TMap<UObject*, UObject*>; // Alias for use in the macro
+	DECLARE_MULTICAST_DELEGATE_OneParam(FOnObjectsReplaced, const FReplacementObjectMap&);
+	static FOnObjectsReplaced OnObjectsReplaced;
+
 	/** Callback for when an asset is saved. This is called from UObject::PreSave before it is actually written to disk, for every object saved */
 	DECLARE_MULTICAST_DELEGATE_OneParam(FOnObjectSaved, UObject*);
 	static FOnObjectSaved OnObjectSaved;
diff --git a/Engine/Source/Runtime/Engine/Private/Elements/Framework/EngineElementsLibrary.cpp b/Engine/Source/Runtime/Engine/Private/Elements/Framework/EngineElementsLibrary.cpp
--- a/Engine/Source/Runtime/Engine/Private/Elements/Framework/EngineElementsLibrary.cpp
+++ b/Engine/Source/Runtime/Engine/Private/Elements/Framework/EngineElementsLibrary.cpp
@@ -15,6 +15,9 @@
 #include "Elements/Component/ComponentElementData.h"
 #include "Components/ActorComponent.h"
 
+// The editor requires ref-counting for object replacement to function correctly
+static_assert(!WITH_EDITOR || UE_TYPED_ELEMENT_HAS_REFCOUNTING, "The editor requires that ref-counting is enabled for typed elements!");
+
 namespace EngineElementsLibraryUtil
 {
 template <typename ObjectClass, typename ElementDataType>
@@ -66,6 +69,50 @@ FTypedElementHandle AcquireEditorTypedElementHandle(const ObjectClass* Object, T
 
 	return FTypedElementHandle();
 }
+
+template <typename ObjectClass, typename ElementDataType>
+void ReplaceEditorTypedElementHandles(const TMap<UObject*, UObject*>& ReplacementObjects, TTypedElementOwnerStore<ElementDataType, const ObjectClass*>& ElementOwnerStore, TFunctionRef<void(const ObjectClass*, TTypedElementOwner<ElementDataType>&)> UpdateElement, TFunctionRef<void(const ObjectClass*, TTypedElementOwner<ElementDataType>&)> DestroyElement)
+{
+	for (const TTuple<UObject*, UObject*>& ReplacementObjectPair : ReplacementObjects)
+	{
+		const ObjectClass* OldObject = Cast<ObjectClass>(ReplacementObjectPair.Key);
+		if (!OldObject || !ReplacementObjectPair.Value)
+		{
+			// We can't do anything for objects of the incorrect type, or redirection to null...
+			continue;
+		}
+
+		const ObjectClass* NewObject = Cast<ObjectClass>(ReplacementObjectPair.Value);
+		if (!NewObject)
+		{
+			UE_LOG(LogCore, Error, TEXT("Typed element replacement failed! Objects must be compatible types - Expected: %s, Old: %s (%s), New: %s (%s)"), *ObjectClass::StaticClass()->GetName(), *OldObject->GetClass()->GetName(), *OldObject->GetPathName(), *ReplacementObjectPair.Key->GetClass()->GetName(), *ReplacementObjectPair.Key->GetPathName());
+			continue;
+		}
+
+		// We only need to attempt replacement if we actually have an element for the old object...
+		if (TTypedElementOwner<ElementDataType> OldEditorElement = ElementOwnerStore.UnregisterElementOwner(OldObject))
+		{
+			if (OldEditorElement.Private_GetInternalData()->GetRefCount() > 1)
+			{
+				// The old element has external references, so we're going to destroy the new element (if any) and re-point the old element to the new object...
+				// Note: This requires that the new element has no external references - if both old and new elements have external references then we'll need to support redirection at the element level!
+				if (TTypedElementOwner<ElementDataType> NewEditorElement = ElementOwnerStore.UnregisterElementOwner(NewObject))
+				{
+					checkf(NewEditorElement.Private_GetInternalData()->GetRefCount() <= 1, TEXT("The old and new element both have external references! Replacing these will require support for redirection at the element level!"));
+					DestroyElement(NewObject, NewEditorElement);
+				}
+
+				UpdateElement(NewObject, OldEditorElement);
+				ElementOwnerStore.RegisterElementOwner(NewObject, MoveTemp(OldEditorElement));
+			}
+			else
+			{
+				// The old element has no external references, so we can just destroy it...
+				DestroyElement(OldObject, OldEditorElement);
+			}
+		}
+	}
+}
 #endif
 } // namespace EngineElementsLibraryUtil
 
@@ -80,6 +127,9 @@ UEngineElementsLibrary::UEngineElementsLibrary()
 #if WITH_EDITOR
 	if (HasAnyFlags(RF_ClassDefaultObject))
 	{
+		// The editor may replace objects and perform fix-up from old->new, so we need to keep any object-based elements in-sync too
+		FCoreUObjectDelegates::OnObjectsReplaced.AddStatic(&UEngineElementsLibrary::OnObjectsReplaced);
+
 		// UObject exists inside CoreUObject, so it cannot call through directly to clean-up any element handles that have been created
 		// Instead we rely on this GC hook to clean-up any element handles for unreachable objects prior to them being destroyed
 		FCoreUObjectDelegates::PreGarbageCollectConditionalBeginDestroy.AddStatic(&UEngineElementsLibrary::DestroyUnreachableEditorObjectElements);
@@ -87,6 +137,26 @@ UEngineElementsLibrary::UEngineElementsLibrary()
 #endif
 }
 
+#if WITH_EDITOR
+void UEngineElementsLibrary::OnObjectsReplaced(const TMap<UObject*, UObject*>& InReplacementObjects)
+{
+	EngineElementsLibraryUtil::ReplaceEditorTypedElementHandles<UObject, FObjectElementData>(InReplacementObjects, GObjectElementOwnerStore, [](const UObject* InObject, TTypedElementOwner<FObjectElementData>& InOutObjectElement)
+	{
+		InOutObjectElement.GetDataChecked().Object = const_cast<UObject*>(InObject);
+	}, &UEngineElementsLibrary::DestroyObjectElement);
+
+	EngineElementsLibraryUtil::ReplaceEditorTypedElementHandles<AActor, FActorElementData>(InReplacementObjects, GActorElementOwnerStore, [](const AActor* InActor, TTypedElementOwner<FActorElementData>& InOutActorElement)
+	{
+		InOutActorElement.GetDataChecked().Actor = const_cast<AActor*>(InActor);
+	}, &UEngineElementsLibrary::DestroyActorElement);
+
+	EngineElementsLibraryUtil::ReplaceEditorTypedElementHandles<UActorComponent, FComponentElementData>(InReplacementObjects, GComponentElementOwnerStore, [](const UActorComponent* InComponent, TTypedElementOwner<FComponentElementData>& InOutComponentElement)
+	{
+		InOutComponentElement.GetDataChecked().Component = const_cast<UActorComponent*>(InComponent);
+	}, &UEngineElementsLibrary::DestroyComponentElement);
+}
+#endif
+
 TTypedElementOwner<FObjectElementData> UEngineElementsLibrary::CreateObjectElement(const UObject* InObject)
 {
 	TTypedElementOwner<FObjectElementData> ObjectElement = EngineElementsLibraryUtil::CreateTypedElement<UObject, FObjectElementData>(InObject, NAME_Object);
diff --git a/Engine/Source/Runtime/Engine/Public/Elements/Framework/EngineElementsLibrary.h b/Engine/Source/Runtime/Engine/Public/Elements/Framework/EngineElementsLibrary.h
--- a/Engine/Source/Runtime/Engine/Public/Elements/Framework/EngineElementsLibrary.h
+++ b/Engine/Source/Runtime/Engine/Public/Elements/Framework/EngineElementsLibrary.h
@@ -56,5 +56,6 @@ public:
 private:
 #if WITH_EDITOR
 	static void DestroyUnreachableEditorObjectElements();
+	static void OnObjectsReplaced(const TMap<UObject*, UObject*>& InReplacementObjects);
 #endif
 };
