Fixed static analysis warning in VisualLogger.cpp

diff --git a/Engine/Source/Runtime/Engine/Private/VisualLogger/VisualLogger.cpp b/Engine/Source/Runtime/Engine/Private/VisualLogger/VisualLogger.cpp
--- a/Engine/Source/Runtime/Engine/Private/VisualLogger/VisualLogger.cpp
+++ b/Engine/Source/Runtime/Engine/Private/VisualLogger/VisualLogger.cpp
@@ -29,7 +29,7 @@ namespace
 {
 	UWorld* GetWorldForVisualLogger(const UObject* Object)
 	{
-		UWorld* World = GEngine->GetWorldFromContextObject(Object, EGetWorldErrorMode::ReturnNull);
+		UWorld* World = GEngine ? GEngine->GetWorldFromContextObject(Object, EGetWorldErrorMode::ReturnNull) : nullptr;
 #if WITH_EDITOR
 		UEditorEngine *EEngine = Cast<UEditorEngine>(GEngine);
 		if (GIsEditor && EEngine != nullptr && World == nullptr)
