Fix an occasional crash in FStaticMeshCompilingManager

- The crash was seen when running Reverb in the editor with -game command line from an empty DDC
- The return value of Component->GetStaticMesh() is tested against the TSet RegisteredStaticMesh and so a nullptr should return false at this point
  - However the TSet contains TWeakObjectPtr
  - When checking the set to see if it contains a nullptr, it will be tested against an existing key and in the rare case that the TWeakObjectPtr is not valid anymore this test will aways return true, as two invalid TWeakObjectPtr are always considered equal even if they originally pointed to different UObjects.
  - This means that on rare occasional the check against RegisteredStaticMesh can succeed when testing vs a nullptr
  - We should check if Component->GetStaticMesh() returns nullptr or not and if so, skip the loop, before we ever check against RegisteredStaticMesh

#rnx

diff --git a/Engine/Source/Runtime/Engine/Private/StaticMeshCompiler.cpp b/Engine/Source/Runtime/Engine/Private/StaticMeshCompiler.cpp
--- a/Engine/Source/Runtime/Engine/Private/StaticMeshCompiler.cpp
+++ b/Engine/Source/Runtime/Engine/Private/StaticMeshCompiler.cpp
@@ -392,6 +392,7 @@ void FStaticMeshCompilingManager::FinishCompilationsForGame()
 			{
 				if (Component->IsRegistered() &&
 					PIEWorlds.Contains(Component->GetWorld()) &&
+					Component->GetStaticMesh() != nullptr &&
 					RegisteredStaticMesh.Contains(Component->GetStaticMesh()) &&
 					(PlayInEditorMode == 0 || Component->GetCollisionEnabled() != ECollisionEnabled::NoCollision || Component->IsNavigationRelevant()))
 				{
