Enable GPU-Scene Instance Culling.

diff --git a/Engine/Shaders/Public/ShaderVersion.ush b/Engine/Shaders/Public/ShaderVersion.ush
--- a/Engine/Shaders/Public/ShaderVersion.ush
+++ b/Engine/Shaders/Public/ShaderVersion.ush
@@ -3,5 +3,5 @@
 // in Platform.ush (which should be included in any shader) it allows to invalidate the shader DDC.
 // 
 // If you are merging streams and there is a conflict with this GUID you should make a new GUID rather than taking one or the other.
-// GUID = F5342E5D-ED6F-4D7D-871C-5BC9D901271E
+// GUID = D8A972B4-2356-49BC-9CA8-ABF9D6792E1B
 
diff --git a/Engine/Source/Runtime/Core/Core.Build.cs b/Engine/Source/Runtime/Core/Core.Build.cs
--- a/Engine/Source/Runtime/Core/Core.Build.cs
+++ b/Engine/Source/Runtime/Core/Core.Build.cs
@@ -265,6 +265,9 @@ public class Core : ModuleRules
 		// temporary thing.
 		PrivateDefinitions.Add("PLATFORM_SUPPORTS_BINARYCONFIG=" + (SupportsBinaryConfig(Target) ? "1" : "0"));
 
+		// temporary thing to enable backing out in case of disaster, remove after initial testing period.
+		PublicDefinitions.Add("GPUCULL_TODO=1");
+
 		PublicDefinitions.Add("WITH_MALLOC_STOMP=" + (bWithMallocStomp ? "1" : "0"));
 
 		PrivateDefinitions.Add("PLATFORM_COMPILER_OPTIMIZATION_LTCG=" + (Target.bAllowLTCG ? "1" : "0"));
diff --git a/Engine/Source/Runtime/Engine/Private/HierarchicalInstancedStaticMesh.cpp b/Engine/Source/Runtime/Engine/Private/HierarchicalInstancedStaticMesh.cpp
--- a/Engine/Source/Runtime/Engine/Private/HierarchicalInstancedStaticMesh.cpp
+++ b/Engine/Source/Runtime/Engine/Private/HierarchicalInstancedStaticMesh.cpp
@@ -1416,7 +1416,7 @@ void FHierarchicalStaticMeshSceneProxy::FillDynamicMeshElements(FMeshElementColl
 									LODIndex, InstancedRenderData.VertexFactories.Num(), SelectionGroupIndex, SectionIndex, LODModel.Sections.Num(), RunArray.Num() / 2,
 									TotalInstances, Tris, (int)MeshBatch.CastShadow, ElementParams.ShadowFrustum,
 									*StaticMesh->GetPathName(),
-									*MeshBatch.MaterialRenderProxy->GetMaterial(ElementParams.FeatureLevel)->GetFriendlyName());
+									*MeshBatch.MaterialRenderProxy->GetIncompleteMaterialWithFallback(ElementParams.FeatureLevel).GetFriendlyName());
 							}
 							else
 							{
@@ -1424,7 +1424,7 @@ void FHierarchicalStaticMeshSceneProxy::FillDynamicMeshElements(FMeshElementColl
 									LODIndex, InstancedRenderData.VertexFactories.Num(), SelectionGroupIndex, SectionIndex, LODModel.Sections.Num(), RunArray.Num() / 2,
 									TotalInstances, Tris, (int)MeshBatch.CastShadow, ElementParams.FinalCullDistance, ElementParams.ShadowFrustum,
 									*StaticMesh->GetPathName(),
-									*MeshBatch.MaterialRenderProxy->GetMaterial(ElementParams.FeatureLevel)->GetFriendlyName());
+									*MeshBatch.MaterialRenderProxy->GetIncompleteMaterialWithFallback(ElementParams.FeatureLevel).GetFriendlyName());
 							}
 						}
 #endif // STATS
diff --git a/Engine/Source/Runtime/Engine/Private/ShaderCompiler/ShaderCompiler.cpp b/Engine/Source/Runtime/Engine/Private/ShaderCompiler/ShaderCompiler.cpp
--- a/Engine/Source/Runtime/Engine/Private/ShaderCompiler/ShaderCompiler.cpp
+++ b/Engine/Source/Runtime/Engine/Private/ShaderCompiler/ShaderCompiler.cpp
@@ -4385,6 +4385,10 @@ void GlobalBeginCompileShader(
 		}
 	}
 
+#if defined(GPUCULL_TODO)
+	Input.Environment.SetDefine(TEXT("GPUCULL_TODO"), 1);
+#endif // defined(GPUCULL_TODO)	
+
 	// Set VR definitions
 	{
 		static const auto CVarInstancedStereo = IConsoleManager::Get().FindTConsoleVariableDataInt(TEXT("vr.InstancedStereo"));
diff --git a/Engine/Source/Runtime/Renderer/Private/VirtualShadowMaps/VirtualShadowMapConfig.h b/Engine/Source/Runtime/Renderer/Private/VirtualShadowMaps/VirtualShadowMapConfig.h
--- a/Engine/Source/Runtime/Renderer/Private/VirtualShadowMaps/VirtualShadowMapConfig.h
+++ b/Engine/Source/Runtime/Renderer/Private/VirtualShadowMaps/VirtualShadowMapConfig.h
@@ -7,7 +7,7 @@
 #pragma once
 
 #if defined(GPUCULL_TODO)
-	#define ENABLE_NON_NANITE_VSM 1
+	#define ENABLE_NON_NANITE_VSM 0
 #else // !defined(GPUCULL_TODO)
 	#define ENABLE_NON_NANITE_VSM 0
 #endif // defined(GPUCULL_TODO)
\ No newline at end of file
