Added previewing of WorldPartitionRuntimeSpatialHash grids

diff --git a/Engine/Source/Runtime/Engine/Private/WorldPartition/WorldPartition.cpp b/Engine/Source/Runtime/Engine/Private/WorldPartition/WorldPartition.cpp
--- a/Engine/Source/Runtime/Engine/Private/WorldPartition/WorldPartition.cpp
+++ b/Engine/Source/Runtime/Engine/Private/WorldPartition/WorldPartition.cpp
@@ -854,6 +854,11 @@ void UWorldPartition::DrawRuntimeHash3D()
 }
 
 #if WITH_EDITOR
+void UWorldPartition::DrawRuntimeHashPreview()
+{
+	RuntimeHash->DrawPreview();
+}
+
 void UWorldPartition::PrepareForPIE()
 {
 	check(World->IsPlayInEditor());
diff --git a/Engine/Source/Runtime/Engine/Private/WorldPartition/WorldPartitionRuntimeSpatialHash.cpp b/Engine/Source/Runtime/Engine/Private/WorldPartition/WorldPartitionRuntimeSpatialHash.cpp
--- a/Engine/Source/Runtime/Engine/Private/WorldPartition/WorldPartitionRuntimeSpatialHash.cpp
+++ b/Engine/Source/Runtime/Engine/Private/WorldPartition/WorldPartitionRuntimeSpatialHash.cpp
@@ -337,9 +337,17 @@ ASpatialHashRuntimeGridInfo::ASpatialHashRuntimeGridInfo(const FObjectInitialize
 
 UWorldPartitionRuntimeSpatialHash::UWorldPartitionRuntimeSpatialHash(const FObjectInitializer& ObjectInitializer)
 	: Super(ObjectInitializer)
+#if WITH_EDITOR
+	, bPreviewGrids(false)
+#endif
 {}
 
 #if WITH_EDITOR
+void UWorldPartitionRuntimeSpatialHash::DrawPreview() const
+{
+	GridPreviewer.Draw(GetWorld(), Grids, bPreviewGrids);
+}
+
 FName UWorldPartitionRuntimeSpatialHash::GetActorRuntimeGrid(const AActor* Actor) const
 {
 	if (ULevel* Level = Actor ? Actor->GetLevel() : nullptr)
diff --git a/Engine/Source/Runtime/Engine/Private/WorldPartition/WorldPartitionRuntimeSpatialHashGridPreviewer.cpp b/Engine/Source/Runtime/Engine/Private/WorldPartition/WorldPartitionRuntimeSpatialHashGridPreviewer.cpp
new file mode 100644
--- /dev/null
+++ b/Engine/Source/Runtime/Engine/Private/WorldPartition/WorldPartitionRuntimeSpatialHashGridPreviewer.cpp
@@ -0,0 +1,80 @@
+// Copyright Epic Games, Inc. All Rights Reserved.
+
+#include "WorldPartition/WorldPartitionRuntimeSpatialHashGridPreviewer.h"
+#include "WorldPartition/WorldPartitionRuntimeSpatialHash.h"
+#include "Engine/PostProcessVolume.h"
+#include "Materials/MaterialInstanceDynamic.h"
+
+/*
+ * FWorldPartitionRuntimeSpatialHashGridPreviewer
+ */
+
+FWorldPartitionRuntimeSpatialHashGridPreviewer::FWorldPartitionRuntimeSpatialHashGridPreviewer()
+#if WITH_EDITORONLY_DATA
+	: MID(nullptr)
+	, Volume(nullptr)
+#endif
+{
+#if WITH_EDITORONLY_DATA
+	Material = LoadObject<UMaterial>(nullptr, TEXT("/Engine/EditorMaterials/WorldPartition/WorldPartitionSpatialHashGridPreviewMaterial"));
+#endif
+}
+
+#if WITH_EDITOR
+void FWorldPartitionRuntimeSpatialHashGridPreviewer::Draw(UWorld* World, const TArray<FSpatialHashRuntimeGrid>& Grids, bool bEnabled)
+{
+	if (bEnabled && Material)
+	{
+		if (!Volume)
+		{
+			FActorSpawnParameters SpawnParameters;
+			SpawnParameters.ObjectFlags |= RF_Transient;
+			SpawnParameters.ObjectFlags &= ~RF_Transactional;
+			Volume = World->SpawnActor<APostProcessVolume>(SpawnParameters);
+		}
+
+		if (!MID && Volume)
+		{
+			MID = UMaterialInstanceDynamic::Create(Material, Volume);
+			Volume->Settings.WeightedBlendables.Array.Emplace(1.0f, MID);
+			Volume->bUnbound = true;
+		}
+
+		if (MID)
+		{
+			for (int i = 0; i < MAX_PREVIEW_GRIDS; ++i)
+			{
+				bool bGridEnabled = Grids.IsValidIndex(i) && Grids[i].CellSize > 0;
+				FGridParametersCache& CachedParameters = GridParameters[i];
+
+				if (CachedParameters.bEnabled != bGridEnabled)
+				{
+					MID->SetScalarParameterValue(*FString::Printf(TEXT("Grid%d_Enabled"), i), bGridEnabled ? 1.0f : 0.f);
+					CachedParameters.bEnabled = bGridEnabled;
+				}
+				if (bGridEnabled)
+				{
+					const FSpatialHashRuntimeGrid& Grid = Grids[i];
+
+					if (CachedParameters.CellSize != Grid.CellSize)
+					{
+						MID->SetScalarParameterValue(*FString::Printf(TEXT("Grid%d_CellSize"), i), (float)Grid.CellSize);
+						CachedParameters.CellSize = Grid.CellSize;
+					}
+
+					if (CachedParameters.GridColor != Grid.DebugColor)
+					{
+						MID->SetVectorParameterValue(*FString::Printf(TEXT("Grid%d_Color"), i), Grid.DebugColor);
+						CachedParameters.GridColor = Grid.DebugColor;
+					}
+				}
+			}
+		}
+	}
+
+	if (Volume)
+	{
+		Volume->bEnabled = bEnabled;
+	}
+}
+#endif
\ No newline at end of file
diff --git a/Engine/Source/Runtime/Engine/Private/WorldPartition/WorldPartitionSubsystem.cpp b/Engine/Source/Runtime/Engine/Private/WorldPartition/WorldPartitionSubsystem.cpp
--- a/Engine/Source/Runtime/Engine/Private/WorldPartition/WorldPartitionSubsystem.cpp
+++ b/Engine/Source/Runtime/Engine/Private/WorldPartition/WorldPartitionSubsystem.cpp
@@ -119,6 +119,13 @@ void UWorldPartitionSubsystem::Tick(float DeltaSeconds)
 		{
 			Partition->DrawRuntimeHash3D();
 		}
+
+#if WITH_EDITOR
+		if (!GetWorld()->IsGameWorld())
+		{
+			Partition->DrawRuntimeHashPreview();
+		}
+#endif
 	}
 }
 
diff --git a/Engine/Source/Runtime/Engine/Public/WorldPartition/WorldPartition.h b/Engine/Source/Runtime/Engine/Public/WorldPartition/WorldPartition.h
--- a/Engine/Source/Runtime/Engine/Public/WorldPartition/WorldPartition.h
+++ b/Engine/Source/Runtime/Engine/Public/WorldPartition/WorldPartition.h
@@ -143,6 +143,8 @@ public:
 	void GenerateHLOD(ISourceControlHelper* SourceControlHelper);
 	void GenerateNavigationData();
 
+	// Debugging Methods
+	void DrawRuntimeHashPreview();
 	void DumpActorDescs(const FString& Path);
 #endif
 
diff --git a/Engine/Source/Runtime/Engine/Public/WorldPartition/WorldPartitionRuntimeHash.h b/Engine/Source/Runtime/Engine/Public/WorldPartition/WorldPartitionRuntimeHash.h
--- a/Engine/Source/Runtime/Engine/Public/WorldPartition/WorldPartitionRuntimeHash.h
+++ b/Engine/Source/Runtime/Engine/Public/WorldPartition/WorldPartitionRuntimeHash.h
@@ -25,6 +25,7 @@ class ENGINE_API UWorldPartitionRuntimeHash : public UObject
 	virtual bool GenerateHLOD(ISourceControlHelper* SourceControlHelper) { return false; }
 	virtual bool GenerateNavigationData() { return false; }
 	virtual FName GetActorRuntimeGrid(const AActor* Actor) const { return NAME_None; }
+	virtual void DrawPreview() const {}
 #endif
 
 	// Streaming interface
diff --git a/Engine/Source/Runtime/Engine/Public/WorldPartition/WorldPartitionRuntimeSpatialHash.h b/Engine/Source/Runtime/Engine/Public/WorldPartition/WorldPartitionRuntimeSpatialHash.h
--- a/Engine/Source/Runtime/Engine/Public/WorldPartition/WorldPartitionRuntimeSpatialHash.h
+++ b/Engine/Source/Runtime/Engine/Public/WorldPartition/WorldPartitionRuntimeSpatialHash.h
@@ -11,6 +11,7 @@
 #include "WorldPartition/DataLayer/DataLayersID.h"
 #include "WorldPartitionRuntimeHash.h"
 #include "WorldPartitionRuntimeSpatialHashCell.h"
+#include "WorldPartitionRuntimeSpatialHashGridPreviewer.h"
 #include "WorldPartitionRuntimeSpatialHash.generated.h"
 
 DECLARE_LOG_CATEGORY_EXTERN(LogWorldPartitionRuntimeSpatialHash, Log, All);
@@ -198,6 +199,7 @@ public:
 	virtual bool GenerateHLOD(ISourceControlHelper* SourceControlHelper) override;
 	virtual bool GenerateNavigationData() override;
 	virtual FName GetActorRuntimeGrid(const AActor* Actor) const override;
+	virtual void DrawPreview() const override;
 
 	FName GetCellName(FName InGridName, int32 InLevel, int32 InCellX, int32 InCellY, const FDataLayersID& InDataLayerID) const;
 	static FName GetCellName(UWorldPartition* WorldPartition, FName InGridName, int32 InLevel, int32 InCellX, int32 InCellY, const FDataLayersID& InDataLayerID);
@@ -217,9 +219,15 @@ private:
 	static class FAutoConsoleCommand OverrideLoadingRangeCommand;
 
 #if WITH_EDITORONLY_DATA
-	UPROPERTY(EditAnywhere, Config, Category=RuntimeSettings)
+	UPROPERTY(EditAnywhere, Config, Category = RuntimeSettings)
 	TArray<FSpatialHashRuntimeGrid> Grids;
 
+	UPROPERTY(EditAnywhere, Category = RuntimeSettings, Transient, NonTransactional)
+	bool bPreviewGrids;
+
+	UPROPERTY(Transient)
+	mutable FWorldPartitionRuntimeSpatialHashGridPreviewer GridPreviewer;
+
 	TMap<FString, UWorldPartitionRuntimeCell*> PackagesToGenerateForCook;
 	TMap<FGuid, FGuid> CachedHLODParents;
 #endif
diff --git a/Engine/Source/Runtime/Engine/Public/WorldPartition/WorldPartitionRuntimeSpatialHashGridPreviewer.h b/Engine/Source/Runtime/Engine/Public/WorldPartition/WorldPartitionRuntimeSpatialHashGridPreviewer.h
new file mode 100644
--- /dev/null
+++ b/Engine/Source/Runtime/Engine/Public/WorldPartition/WorldPartitionRuntimeSpatialHashGridPreviewer.h
@@ -0,0 +1,49 @@
+// Copyright Epic Games, Inc. All Rights Reserved.
+#pragma once
+
+#include "CoreMinimal.h"
+#include "WorldPartitionRuntimeSpatialHashGridPreviewer.generated.h"
+
+class APostProcessVolume;
+class UMaterial;
+class UMaterialInstanceDynamic;
+struct FSpatialHashRuntimeGrid;
+
+USTRUCT()
+struct ENGINE_API FWorldPartitionRuntimeSpatialHashGridPreviewer
+{
+	GENERATED_USTRUCT_BODY()
+
+	FWorldPartitionRuntimeSpatialHashGridPreviewer();
+
+#if WITH_EDITOR
+	void Draw(UWorld* World, const TArray<FSpatialHashRuntimeGrid>& Grids, bool bEnabled);
+#endif
+
+#if WITH_EDITORONLY_DATA
+	UPROPERTY()
+	UMaterial* Material;
+
+	UPROPERTY()
+	UMaterialInstanceDynamic* MID;
+
+	UPROPERTY()
+	APostProcessVolume* Volume;
+
+	struct FGridParametersCache
+	{
+		FGridParametersCache()
+			: bEnabled(false)
+			, CellSize(0)
+			, GridColor(0.f, 0.f, 0.f, 0.f)
+		{}
+
+		bool bEnabled;
+		int32 CellSize;
+		FLinearColor GridColor;
+	};
+
+	enum { MAX_PREVIEW_GRIDS = 4 };
+	FGridParametersCache GridParameters[MAX_PREVIEW_GRIDS];
+#endif
+};
\ No newline at end of file
