U5 - Move Anim Compression into parallelFor

diff --git a/Engine/Source/Runtime/Engine/Private/Animation/AnimBoneCompressionSettings.cpp b/Engine/Source/Runtime/Engine/Private/Animation/AnimBoneCompressionSettings.cpp
--- a/Engine/Source/Runtime/Engine/Private/Animation/AnimBoneCompressionSettings.cpp
+++ b/Engine/Source/Runtime/Engine/Private/Animation/AnimBoneCompressionSettings.cpp
@@ -8,6 +8,7 @@
 #include "AnimationCompression.h"
 #include "Serialization/MemoryWriter.h"
 #include "UObject/Package.h"
+#include "Async/ParallelFor.h"
 
 #define DEBUG_DUMP_ANIM_COMPRESSION_STATS 0
 
@@ -108,28 +109,6 @@ static void CompressAnimSequenceImpl(FAnimBoneCompressionContext& Context)
 	}
 }
 
-class FAsyncAnimCompressionTask
-{
-public:
-	FAsyncAnimCompressionTask(FAnimBoneCompressionContext* Context_)
-		: Context(Context_)
-	{}
-
-	/** return the name of the task **/
-	static const TCHAR* GetTaskName() { return TEXT("FAsyncAnimCompressionTask"); }
-	FORCEINLINE static TStatId GetStatId() { RETURN_QUICK_DECLARE_CYCLE_STAT(FAsyncAnimCompressionTask, STATGROUP_TaskGraphTasks); }
-
-	static ENamedThreads::Type GetDesiredThread() { return ENamedThreads::AnyBackgroundThreadNormalTask; }
-	static ESubsequentsMode::Type GetSubsequentsMode() { return ESubsequentsMode::TrackSubsequents; }
-
-	void DoTask(ENamedThreads::Type CurrentThread, const FGraphEventRef& MyCompletionGraphEvent)
-	{
-		CompressAnimSequenceImpl(*Context);
-	}
-
-	FAnimBoneCompressionContext* Context;
-};
-
 bool UAnimBoneCompressionSettings::Compress(const FCompressibleAnimData& AnimSeq, FCompressibleAnimDataResult& OutCompressedData) const
 {
 	if (!AreSettingsValid())
@@ -148,7 +127,6 @@ bool UAnimBoneCompressionSettings::Compress(const FCompressibleAnimData& AnimSeq
 	const double CompressionStartTime = FPlatformTime::Seconds();
 #endif
 
-	FGraphEventArray AnimCompressionTask_CompletionEvents;
 	TArray<FAnimBoneCompressionContext*> ContextList;
 
 	for (UAnimBoneCompressionCodec* Codec : Codecs)
@@ -160,12 +138,12 @@ bool UAnimBoneCompressionSettings::Compress(const FCompressibleAnimData& AnimSeq
 
 		FAnimBoneCompressionContext* Context = new FAnimBoneCompressionContext(AnimSeq, Codec);
 		ContextList.Add(Context);
-
-		AnimCompressionTask_CompletionEvents.Add(TGraphTask<FAsyncAnimCompressionTask>::CreateTask(NULL).ConstructAndDispatchWhenReady(Context));
 	}
 
-	// Wait for async compression to finish
-	FTaskGraphInterface::Get().WaitUntilTasksComplete(AnimCompressionTask_CompletionEvents);
+	ParallelForTemplate(ContextList.Num(), [&ContextList](int32 TaskIndex)
+	{
+		CompressAnimSequenceImpl(*ContextList[TaskIndex]);
+	}, EParallelForFlags::BackgroundPriority | EParallelForFlags::Unbalanced);
 
 	const int32 NumContextes = ContextList.Num();
 	if (NumContextes == 0)
@@ -264,7 +242,6 @@ bool UAnimBoneCompressionSettings::Compress(const FCompressibleAnimData& AnimSeq
 	}
 
 	ContextList.Reset();
-	AnimCompressionTask_CompletionEvents.Reset();
 
 #if DEBUG_DUMP_ANIM_COMPRESSION_STATS
 	const double CompressionEndTime = FPlatformTime::Seconds();
