Data Layers Loading
- Remove AActor.ForceExternalActorLevelReference as it is already handled for cook builds (always loaded actors stay in the persistent level) and PIE.
- Sort actors to load list when loading editor cells to make sure AInfo actors are loaded first.

diff --git a/Engine/Source/Runtime/Engine/Classes/GameFramework/Actor.h b/Engine/Source/Runtime/Engine/Classes/GameFramework/Actor.h
--- a/Engine/Source/Runtime/Engine/Classes/GameFramework/Actor.h
+++ b/Engine/Source/Runtime/Engine/Classes/GameFramework/Actor.h
@@ -1851,9 +1851,6 @@ public:
 	/** Does this actor supports external packaging? */
 	virtual bool SupportsExternalPackaging() const;
 
-	/** Force the level to keep a hard reference on that external actor. */
-	virtual bool ForceExternalActorLevelReference() const { return false; }
-
 	/** Internal struct used to store information about an actor's components during reconstruction */
 	struct FActorRootComponentReconstructionData
 	{
diff --git a/Engine/Source/Runtime/Engine/Private/Level.cpp b/Engine/Source/Runtime/Engine/Private/Level.cpp
--- a/Engine/Source/Runtime/Engine/Private/Level.cpp
+++ b/Engine/Source/Runtime/Engine/Private/Level.cpp
@@ -451,7 +451,7 @@ void ULevel::Serialize( FArchive& Ar )
 				// Don't filter out external actors if duplicating the world to get the actors properly duplicated.
 				if (!(Ar.GetPortFlags() & PPF_Duplicate))
 				{
-					return Actor->ForceExternalActorLevelReference();
+					return false;
 				}
 				// When the world is partitioned and duplicating for PIE, we don't want to duplicate any external actors as they will be streamed by the runtime grids.
 				else if (bIsPartitioned && (Ar.GetPortFlags() & PPF_DuplicateForPIE))
diff --git a/Engine/Source/Runtime/Engine/Private/LevelActor.cpp b/Engine/Source/Runtime/Engine/Private/LevelActor.cpp
--- a/Engine/Source/Runtime/Engine/Private/LevelActor.cpp
+++ b/Engine/Source/Runtime/Engine/Private/LevelActor.cpp
@@ -720,13 +720,6 @@ AActor* UWorld::SpawnActor( UClass* Class, FTransform const* UserTransformPtr, c
 	if (ExternalPackage)
 	{
 		ExternalPackage->MarkPackageDirty();
-
-#if WITH_EDITOR
-		if (Actor->ForceExternalActorLevelReference())
-		{
-			LevelToSpawnIn->MarkPackageDirty();
-		}
-#endif
 	}
 
 	if (Actor->IsPendingKill() && !SpawnParameters.bNoFail)
diff --git a/Engine/Source/Runtime/Engine/Private/WorldPartition/WorldPartition.cpp b/Engine/Source/Runtime/Engine/Private/WorldPartition/WorldPartition.cpp
--- a/Engine/Source/Runtime/Engine/Private/WorldPartition/WorldPartition.cpp
+++ b/Engine/Source/Runtime/Engine/Private/WorldPartition/WorldPartition.cpp
@@ -12,8 +12,6 @@
 #include "WorldPartition/WorldPartitionSubsystem.h"
 #include "WorldPartition/WorldPartitionStreamingSource.h"
 #include "Algo/Accumulate.h"
-#include "Algo/Transform.h"
-#include "Algo/RemoveIf.h"
 #include "Engine/World.h"
 #include "Engine/LevelStreaming.h"
 #include "GameFramework/WorldSettings.h"
@@ -236,6 +234,11 @@ void UWorldPartition::Initialize(UWorld* InWorld, const FTransform& InTransform)
 
 		if (bEditorOnly)
 		{
+			// Make sure to preload any AInfo actors first (mainly for ShouldActorBeLoaded as it needs to have access to the WorldDataLayers actor)
+			TArray<FWorldPartitionHandle> AlwaysLoadedActors = EditorHash->GetAlwaysLoadedCell()->Actors.Array();			
+			AlwaysLoadedActors.Sort([](const FWorldPartitionHandle& A, const FWorldPartitionHandle& B) { return A->GetActorClass()->IsChildOf<AInfo>(); });
+			TArray<FWorldPartitionReference> LoadedActors(AlwaysLoadedActors);
+
 			// Load the always loaded cell, don't call LoadCells to avoid creating a transaction
 			UpdateLoadingEditorCell(EditorHash->GetAlwaysLoadedCell(), true);
 
@@ -634,15 +637,14 @@ void UWorldPartition::UpdateLoadingEditorCell(UWorldPartitionEditorCell* Cell, b
 
 	if (!bShouldBeLoaded)
 	{
+		bPotentiallyUnloadedActors = !Cell->LoadedActors.IsEmpty();
 		Cell->LoadedActors.Empty();
-		bPotentiallyUnloadedActors = true;
 	}
 	else
 	{
 		TGuardValue<bool> IsEditorLoadingPackageGuard(GIsEditorLoadingPackage, true);
 
-		TArray<FWorldPartitionActorDesc*> ActorsToUnload;
-		for (FWorldPartitionHandle& ActorHandle: Cell->Actors.Array())
+		for (FWorldPartitionHandle& ActorHandle: Cell->Actors)
 		{
 			AActor* Actor = ActorHandle->GetActor();
 
@@ -653,8 +655,7 @@ void UWorldPartition::UpdateLoadingEditorCell(UWorldPartitionEditorCell* Cell, b
 			// Filter actor against DataLayers
 			if (ShouldActorBeLoaded(*ActorHandle))
 			{
-				bool bIsAlreadyInLoadedActors = false;
-				Cell->LoadedActors.Add(ActorHandle, &bIsAlreadyInLoadedActors);
+				Cell->LoadedActors.Add(ActorHandle);
 			}
 			else
 			{
@@ -664,11 +665,10 @@ void UWorldPartition::UpdateLoadingEditorCell(UWorldPartitionEditorCell* Cell, b
 					if (ActorReference == ActorHandle)
 					{
 						Cell->LoadedActors.Remove(ActorHandle);
+						bPotentiallyUnloadedActors = true;
 						break;
 					}
 				}
-
-				bPotentiallyUnloadedActors = true;
 			}
 		}
 	}
diff --git a/Engine/Source/Runtime/Engine/Public/WorldPartition/DataLayer/WorldDataLayers.h b/Engine/Source/Runtime/Engine/Public/WorldPartition/DataLayer/WorldDataLayers.h
--- a/Engine/Source/Runtime/Engine/Public/WorldPartition/DataLayer/WorldDataLayers.h
+++ b/Engine/Source/Runtime/Engine/Public/WorldPartition/DataLayer/WorldDataLayers.h
@@ -30,10 +30,6 @@ public:
 	bool RemoveDataLayers(const TArray<UDataLayer*>& InDataLayers);
 	FName GenerateUniqueDataLayerLabel(const FName& InDataLayerLabel) const;
 
-	//~ Begin AActor Interface
-	virtual bool ForceExternalActorLevelReference() const override { return true; }
-	//~ End AActor Interface
-
 	//~ Begin Helper Functions
 	TArray<const UDataLayer*> GetDataLayerObjects(const TArray<FActorDataLayer>& DataLayers) const;
 	TArray<FName> GetDataLayerNames(const TArray<FActorDataLayer>& DataLayers) const;
