LevelInstance: Packed BP set bRunConstructionScriptOnDrag to false.

diff --git a/Engine/Source/Runtime/Engine/Private/LevelInstance/Packed/PackedLevelInstanceBuilder.cpp b/Engine/Source/Runtime/Engine/Private/LevelInstance/Packed/PackedLevelInstanceBuilder.cpp
--- a/Engine/Source/Runtime/Engine/Private/LevelInstance/Packed/PackedLevelInstanceBuilder.cpp
+++ b/Engine/Source/Runtime/Engine/Private/LevelInstance/Packed/PackedLevelInstanceBuilder.cpp
@@ -401,6 +401,8 @@ bool FPackedLevelInstanceBuilder::CreateOrUpdateBlueprintFromPacked(APackedLevel
 	}
 
 	BP->Modify();
+	// Avoid running construction script while dragging an instance of that BP for performance reasons
+	BP->bRunConstructionScriptOnDrag = false;
 	APackedLevelInstance* CDO = CastChecked<APackedLevelInstance>(BP->GeneratedClass->GetDefaultObject());
 	CDO->SetWorldAsset(InActor->GetWorldAsset());
 	
