Fix nonunity build error

diff --git a/Engine/Source/Runtime/Engine/Private/MeshSimplificationSettings.cpp b/Engine/Source/Runtime/Engine/Private/MeshSimplificationSettings.cpp
--- a/Engine/Source/Runtime/Engine/Private/MeshSimplificationSettings.cpp
+++ b/Engine/Source/Runtime/Engine/Private/MeshSimplificationSettings.cpp
@@ -2,6 +2,7 @@
 
 #include "Engine/MeshSimplificationSettings.h"
 #include "UObject/UnrealType.h"
+#include "Modules/ModuleManager.h"
 
 #if WITH_EDITOR
 #include "Developer/MeshReductionInterface/Public/IMeshReductionManagerModule.h"
