Add support for warning as errors when compiling shaders

This adds support for a new compilation flag that enables warnings as errors. Support was added for most of the main platforms. This feature has to opted into since many shaders currenntly contain small warnings, however this gives the ability for developpers to locally opt-in to this behavior to slowly sanitize code.

A global CVar: r.Shaders.WarningsAsErrors was also added, which will force every shader to be (re)built with this enabled for even more thorough checking.

diff --git a/Engine/Source/Developer/Windows/ShaderFormatD3D/Private/D3DShaderCompiler.cpp b/Engine/Source/Developer/Windows/ShaderFormatD3D/Private/D3DShaderCompiler.cpp
--- a/Engine/Source/Developer/Windows/ShaderFormatD3D/Private/D3DShaderCompiler.cpp
+++ b/Engine/Source/Developer/Windows/ShaderFormatD3D/Private/D3DShaderCompiler.cpp
@@ -53,6 +53,7 @@ static uint32 TranslateCompilerFlagD3D11(ECompilerFlags CompilerFlag)
 	{
 	case CFLAG_PreferFlowControl: return D3DCOMPILE_PREFER_FLOW_CONTROL;
 	case CFLAG_AvoidFlowControl: return D3DCOMPILE_AVOID_FLOW_CONTROL;
+	case CFLAG_WarningsAsErrors: return D3DCOMPILE_WARNINGS_ARE_ERRORS;
 	default: return 0;
 	};
 }
@@ -221,6 +222,12 @@ static FString D3D11CreateShaderCompileCommandLine(
 		FXCCommandline += FString(TEXT(" /Gec"));
 	}
 
+	if (CompileFlags & D3DCOMPILE_WARNINGS_ARE_ERRORS)
+	{
+		CompileFlags &= ~D3DCOMPILE_WARNINGS_ARE_ERRORS;
+		FXCCommandline += FString(TEXT(" /WX"));
+	}
+
 	switch (CompileFlags & SHADER_OPTIMIZATION_LEVEL_MASK)
 	{
 		case D3DCOMPILE_OPTIMIZATION_LEVEL2:
diff --git a/Engine/Source/Developer/Windows/ShaderFormatD3D/Private/D3DShaderCompiler.inl b/Engine/Source/Developer/Windows/ShaderFormatD3D/Private/D3DShaderCompiler.inl
--- a/Engine/Source/Developer/Windows/ShaderFormatD3D/Private/D3DShaderCompiler.inl
+++ b/Engine/Source/Developer/Windows/ShaderFormatD3D/Private/D3DShaderCompiler.inl
@@ -88,6 +88,12 @@ public:
 			ExtraArguments.Add(L"/Gec");
 		}
 
+		if (D3DCompileFlags & D3DCOMPILE_WARNINGS_ARE_ERRORS)
+		{
+			D3DCompileFlags &= ~D3DCOMPILE_WARNINGS_ARE_ERRORS;
+			ExtraArguments.Add(L"/WX");
+		}
+
 		switch (D3DCompileFlags & SHADER_OPTIMIZATION_LEVEL_MASK)
 		{
 		case D3DCOMPILE_OPTIMIZATION_LEVEL0:
diff --git a/Engine/Source/Runtime/Engine/Private/ShaderCompiler/ShaderCompiler.cpp b/Engine/Source/Runtime/Engine/Private/ShaderCompiler/ShaderCompiler.cpp
--- a/Engine/Source/Runtime/Engine/Private/ShaderCompiler/ShaderCompiler.cpp
+++ b/Engine/Source/Runtime/Engine/Private/ShaderCompiler/ShaderCompiler.cpp
@@ -799,6 +799,12 @@ static TAutoConsoleVariable<int32> CVarShaderBoundsChecking(
 	TEXT("Whether to enforce bounds-checking & flush-to-zero/ignore for buffer reads & writes in shaders. Defaults to 1 (enabled). Not all shader languages can omit bounds checking."),
 	ECVF_ReadOnly);
 
+static TAutoConsoleVariable<int32> CVarShaderWarningsAsErrors(
+	TEXT("r.Shaders.WarningsAsErrors"),
+	0,
+	TEXT("Whether to treat warnings as errors when running the shader compiler. Defaults to 0 (disabled). Not all compilers support this mode."),
+	ECVF_ReadOnly);
+
 static TAutoConsoleVariable<int32> CVarShaderFlowControl(
 	TEXT("r.Shaders.FlowControlMode"),
 	0,
@@ -4486,6 +4492,11 @@ void GlobalBeginCompileShader(
 		}
 	}
 
+	if (CVarShaderWarningsAsErrors.GetValueOnAnyThread())
+	{
+		Input.Environment.CompilerFlags.Add(CFLAG_WarningsAsErrors);
+	}
+
 	if (IsD3DPlatform((EShaderPlatform)Target.Platform) && IsPCPlatform((EShaderPlatform)Target.Platform))
 	{
 		if (CVarD3DRemoveUnusedInterpolators.GetValueOnAnyThread() != 0)
diff --git a/Engine/Source/Runtime/RenderCore/Private/Shader.cpp b/Engine/Source/Runtime/RenderCore/Private/Shader.cpp
--- a/Engine/Source/Runtime/RenderCore/Private/Shader.cpp
+++ b/Engine/Source/Runtime/RenderCore/Private/Shader.cpp
@@ -1428,6 +1428,11 @@ void ShaderMapAppendKeyString(EShaderPlatform Platform, FString& KeyString)
 		static const auto CVar = IConsoleManager::Get().FindConsoleVariable(TEXT("r.Shaders.FastMath"));
 		KeyString += (CVar && CVar->GetInt() == 0) ? TEXT("_NoFastMath") : TEXT("");
 	}
+
+	{
+		static const auto CVar = IConsoleManager::Get().FindConsoleVariable(TEXT("r.Shaders.WarningsAsErrors"));
+		KeyString += (CVar && CVar->GetInt() == 1) ? TEXT("_WX") : TEXT("");
+	}
 	
 	{
 		static const auto CVar = IConsoleManager::Get().FindConsoleVariable(TEXT("r.Shaders.FlowControlMode"));
diff --git a/Engine/Source/Runtime/RenderCore/Public/ShaderCompilerCore.h b/Engine/Source/Runtime/RenderCore/Public/ShaderCompilerCore.h
--- a/Engine/Source/Runtime/RenderCore/Public/ShaderCompilerCore.h
+++ b/Engine/Source/Runtime/RenderCore/Public/ShaderCompilerCore.h
@@ -90,6 +90,9 @@ enum ECompilerFlags
 	// Enable support of C-style data types for platforms that can. Check for PLATFORM_SUPPORTS_REAL_TYPES.
 	CFLAG_AllowRealTypes,
 
+	// Allow warnings to be treated as errors
+	CFLAG_WarningsAsErrors,
+
 	CFLAG_Max,
 };
 static_assert(CFLAG_Max < 32, "Out of bitfields!");
