When r.ShaderDevelopmentMode is enabled, only prompt to retry shader compiles on errors for global shaders
bPromptToRetryFailedShaderCompiles can be enabled to get prompts on any shader compile errors

diff --git a/Engine/Config/BaseEngine.ini b/Engine/Config/BaseEngine.ini
--- a/Engine/Config/BaseEngine.ini
+++ b/Engine/Config/BaseEngine.ini
@@ -1318,7 +1318,7 @@ NumUnusedShaderCompilingThreads=3
 NumUnusedShaderCompilingThreadsDuringGame=4
 ; Batching multiple jobs to reduce file overhead, but not so many that latency of blocking compiles is hurt
 MaxShaderJobBatchSize=10
-bPromptToRetryFailedShaderCompiles=True
+bPromptToRetryFailedShaderCompiles=False
 bLogJobCompletionTimes=False
 ; Only using 10ms of game thread time per frame to process async shader maps
 ProcessGameThreadTargetTime=.01
diff --git a/Engine/Source/Runtime/Engine/Private/ShaderCompiler/ShaderCompiler.cpp b/Engine/Source/Runtime/Engine/Private/ShaderCompiler/ShaderCompiler.cpp
--- a/Engine/Source/Runtime/Engine/Private/ShaderCompiler/ShaderCompiler.cpp
+++ b/Engine/Source/Runtime/Engine/Private/ShaderCompiler/ShaderCompiler.cpp
@@ -669,7 +669,7 @@ constexpr auto& CastEnumToUnderlyingTypeReference(EnumType& Type)
 // Default value comes from bPromptToRetryFailedShaderCompiles in BaseEngine.ini
 // This is set as a global variable to allow changing in the debugger even in release
 // For example if there are a lot of content shader compile errors you want to skip over without relaunching
-bool GRetryShaderCompilation = false;
+bool GRetryShaderCompilation = true;
 
 static FShaderCompilingManager::EDumpShaderDebugInfo GDumpShaderDebugInfo = FShaderCompilingManager::EDumpShaderDebugInfo::Never;
 static FAutoConsoleVariableRef CVarDumpShaderDebugInfo(
@@ -2578,8 +2578,6 @@ FShaderCompilingManager::FShaderCompilingManager() :
 	GConfig->GetFloat(TEXT("DevOptions.Shaders"), TEXT("BuildWorkerTimeToLive"), GBuildWorkerTimeToLive, GEngineIni);
 	GConfig->GetBool(TEXT("DevOptions.Shaders"), TEXT("bForceUseSCWMemoryPressureLimits"), bForceUseSCWMemoryPressureLimits, GEngineIni);
 
-	GRetryShaderCompilation = bPromptToRetryFailedShaderCompiles;
-
 	verify(GConfig->GetFloat( TEXT("DevOptions.Shaders"), TEXT("ProcessGameThreadTargetTime"), ProcessGameThreadTargetTime, GEngineIni ));
 
 #if UE_BUILD_DEBUG
@@ -3645,7 +3643,7 @@ bool FShaderCompilingManager::HandlePotentialRetryOnError(TMap<int32, FShaderMap
 					}
 				}
 
-				if (UE_LOG_ACTIVE(LogShaders, Log) && bPromptToRetryFailedShaderCompiles && bAnyErrorLikelyToBeCodeError)
+				if (UE_LOG_ACTIVE(LogShaders, Log) && (bAnyErrorLikelyToBeCodeError || bPromptToRetryFailedShaderCompiles))
 				{
 #if UE_BUILD_DEBUG
 					// Use debug break in debug with the debugger attached, otherwise message box
@@ -5143,6 +5141,8 @@ static void PrepareGlobalShaderCompileJob(EShaderPlatform Platform,
 
 	static FString GlobalName(TEXT("Global"));
 
+	NewJob->bErrorsAreLikelyToBeCode = true;
+
 	// Compile the shader environment passed in with the shader type's source code.
 	::GlobalBeginCompileShader(
 		GlobalName,
diff --git a/Engine/Source/Runtime/Engine/Public/ShaderCompiler.h b/Engine/Source/Runtime/Engine/Public/ShaderCompiler.h
--- a/Engine/Source/Runtime/Engine/Public/ShaderCompiler.h
+++ b/Engine/Source/Runtime/Engine/Public/ShaderCompiler.h
@@ -186,7 +186,7 @@ protected:
 		CurrentWorker(EShaderCompilerWorkerType::None),
 		bFinalized(false),
 		bSucceeded(false),
-		bErrorsAreLikelyToBeCode(true),
+		bErrorsAreLikelyToBeCode(false),
 		bReleased(false),
 		bInputHashSet(false)
 	{
