Allow passing a specific actor class to ActorDescContainer iterator.

diff --git a/Engine/Source/Runtime/Engine/Private/WorldPartition/WorldPartition.cpp b/Engine/Source/Runtime/Engine/Private/WorldPartition/WorldPartition.cpp
--- a/Engine/Source/Runtime/Engine/Private/WorldPartition/WorldPartition.cpp
+++ b/Engine/Source/Runtime/Engine/Private/WorldPartition/WorldPartition.cpp
@@ -405,15 +405,11 @@ void UWorldPartition::ForEachIntersectingActorDesc(const FBox& Box, TSubclassOf<
 
 void UWorldPartition::ForEachActorDesc(TSubclassOf<AActor> ActorClass, TFunctionRef<bool(const FWorldPartitionActorDesc*)> Predicate) const
 {
-	for (UActorDescContainer::TConstIterator<> ActorDescIterator(this); ActorDescIterator; ++ActorDescIterator)
+	for (UActorDescContainer::TConstIterator<> ActorDescIterator(this, ActorClass); ActorDescIterator; ++ActorDescIterator)
 	{
-		const FWorldPartitionActorDesc* ActorDesc = *ActorDescIterator;
-		if (ActorDesc->GetActorClass()->IsChildOf(ActorClass))
+		if (!Predicate(*ActorDescIterator))
 		{
-			if (!Predicate(ActorDesc))
-			{
-				return;
-			}
+			return;
 		}
 	}
 }
@@ -918,23 +914,22 @@ void UWorldPartition::DumpActorDescs(const FString& Path)
 		FString LineEntry = TEXT("Guid, Class, Name, BVCenterX, BVCenterY, BVCenterZ, BVExtentX, BVExtentY, BVExtentZ") LINE_TERMINATOR;
 		LogFile->Serialize(TCHAR_TO_ANSI(*LineEntry), LineEntry.Len());
 			
-		ForEachActorDesc(AActor::StaticClass(), [LogFile, &LineEntry](const FWorldPartitionActorDesc* ActorDesc)
+		for (UActorDescContainer::TConstIterator<> ActorDescIterator(this); ActorDescIterator; ++ActorDescIterator)
 		{
 			LineEntry = FString::Printf(
 				TEXT("%s, %s, %s, %.2f, %.2f, %.2f, %.2f, %.2f, %.2f") LINE_TERMINATOR, 
-				*ActorDesc->GetGuid().ToString(), 
-				*ActorDesc->GetClass().ToString(), 
-				*FPaths::GetExtension(ActorDesc->GetActorPath().ToString()), 
-				ActorDesc->GetBounds().GetCenter().X,
-				ActorDesc->GetBounds().GetCenter().Y,
-				ActorDesc->GetBounds().GetCenter().Z,
-				ActorDesc->GetBounds().GetExtent().X,
-				ActorDesc->GetBounds().GetExtent().Y,
-				ActorDesc->GetBounds().GetExtent().Z
+				*ActorDescIterator->GetGuid().ToString(), 
+				*ActorDescIterator->GetClass().ToString(), 
+				*FPaths::GetExtension(ActorDescIterator->GetActorPath().ToString()), 
+				ActorDescIterator->GetBounds().GetCenter().X,
+				ActorDescIterator->GetBounds().GetCenter().Y,
+				ActorDescIterator->GetBounds().GetCenter().Z,
+				ActorDescIterator->GetBounds().GetExtent().X,
+				ActorDescIterator->GetBounds().GetExtent().Y,
+				ActorDescIterator->GetBounds().GetExtent().Z
 			);
 			LogFile->Serialize(TCHAR_TO_ANSI(*LineEntry), LineEntry.Len());
-			return true;
-		});
+		}
 
 		LogFile->Close();
 		delete LogFile;
@@ -954,19 +949,17 @@ FBox UWorldPartition::GetWorldBounds() const
 	FBox WorldBounds(ForceInit);
 	for (UActorDescContainer::TConstIterator<> ActorDescIterator(this); ActorDescIterator; ++ActorDescIterator)
 	{
-		const FWorldPartitionActorDesc* ActorDesc = *ActorDescIterator;
-
-		switch (ActorDesc->GetGridPlacement())
+		switch (ActorDescIterator->GetGridPlacement())
 		{
 			case EActorGridPlacement::Location:
 			{
-				FVector Location = ActorDesc->GetOrigin();
+				FVector Location = ActorDescIterator->GetOrigin();
 				WorldBounds += FBox(Location, Location);
 			}
 			break;
 			case EActorGridPlacement::Bounds:
 			{
-				WorldBounds += ActorDesc->GetBounds();
+				WorldBounds += ActorDescIterator->GetBounds();
 			}
 			break;
 		}
diff --git a/Engine/Source/Runtime/Engine/Public/WorldPartition/ActorDescContainer.h b/Engine/Source/Runtime/Engine/Public/WorldPartition/ActorDescContainer.h
--- a/Engine/Source/Runtime/Engine/Public/WorldPartition/ActorDescContainer.h
+++ b/Engine/Source/Runtime/Engine/Public/WorldPartition/ActorDescContainer.h
@@ -68,9 +68,12 @@ public:
 		typedef typename TChooseClass<bConst, const ValueType*, ValueType*>::Result ReturnType;
 
 	public:
-		TBaseIterator(ContainerType InActorDescContainer)
+		TBaseIterator(ContainerType InActorDescContainer, UClass* InActorClass)
 			: ActorsIterator(InActorDescContainer->Actors)
+			, ActorClass(InActorClass)
 		{
+			check(ActorClass->IsChildOf(ActorType::StaticClass()));
+
 			if (ShouldSkip())
 			{
 				operator++();
@@ -130,10 +133,11 @@ public:
 				return false;
 			}
 
-			return !ActorsIterator->Value->Get()->GetActorClass()->IsChildOf(ActorType::StaticClass());
+			return !ActorsIterator->Value->Get()->GetActorClass()->IsChildOf(ActorClass);
 		}
 
 		IteratorType ActorsIterator;
+		UClass* ActorClass;
 	};
 
 	template <class ActorType = AActor>
@@ -142,8 +146,8 @@ public:
 		typedef TBaseIterator<false, ActorType> BaseType;
 
 	public:
-		TIterator(typename BaseType::ContainerType InActorDescContainer)
-			: BaseType(InActorDescContainer)
+		TIterator(typename BaseType::ContainerType InActorDescContainer, UClass* InActorClass = nullptr)
+			: BaseType(InActorDescContainer, InActorClass ? InActorClass : ActorType::StaticClass())
 		{}
 	};
 
@@ -153,8 +157,8 @@ public:
 		typedef TBaseIterator<true, ActorType> BaseType;
 
 	public:
-		TConstIterator(typename BaseType::ContainerType InActorDescContainer)
-			: BaseType(InActorDescContainer)
+		TConstIterator(typename BaseType::ContainerType InActorDescContainer, UClass* InActorClass = nullptr)
+			: BaseType(InActorDescContainer, InActorClass ? InActorClass : ActorType::StaticClass())
 		{}
 	};
 #endif
