Don't spawn external actors for PIE levels.

diff --git a/Engine/Source/Runtime/Engine/Classes/Engine/Level.h b/Engine/Source/Runtime/Engine/Classes/Engine/Level.h
--- a/Engine/Source/Runtime/Engine/Classes/Engine/Level.h
+++ b/Engine/Source/Runtime/Engine/Classes/Engine/Level.h
@@ -975,6 +975,9 @@ public:
 	/** Sets if the level uses external actors mode or not. */
 	ENGINE_API void SetUseExternalActors(bool bEnable);
 
+	/** Returns true if the level wants newly spawned actors to be external */
+	ENGINE_API bool ShouldCreateNewExternalActors() const;
+
 	/** 
 	 * Convert this level actors to the specified loading strategy
 	 * @param bExternal if true will convert internal actors to external, will convert external actors to internal otherwise
diff --git a/Engine/Source/Runtime/Engine/Private/Level.cpp b/Engine/Source/Runtime/Engine/Private/Level.cpp
--- a/Engine/Source/Runtime/Engine/Private/Level.cpp
+++ b/Engine/Source/Runtime/Engine/Private/Level.cpp
@@ -2306,6 +2306,11 @@ void ULevel::SetUseExternalActors(bool bEnable)
 	}
 }
 
+bool ULevel::ShouldCreateNewExternalActors() const
+{
+	return IsUsingExternalActors() && !GetPackage()->HasAnyPackageFlags(PKG_PlayInEditor);
+}
+
 void ULevel::ConvertAllActorsToPackaging(bool bExternal)
 {
 	// Make a copy of the current actor lists since packaging conversion may modify the actor list as a side effect
diff --git a/Engine/Source/Runtime/Engine/Private/LevelActor.cpp b/Engine/Source/Runtime/Engine/Private/LevelActor.cpp
--- a/Engine/Source/Runtime/Engine/Private/LevelActor.cpp
+++ b/Engine/Source/Runtime/Engine/Private/LevelActor.cpp
@@ -544,7 +544,7 @@ AActor* UWorld::SpawnActor( UClass* Class, FTransform const* UserTransformPtr, c
 		ExternalPackage = SpawnParameters.OverridePackage;
 		bNeedGloballyUniqueName = true;
 	}
-	else if (LevelToSpawnIn->IsUsingExternalActors() && SpawnParameters.bCreateActorPackage && !(SpawnParameters.ObjectFlags & RF_Transient))
+	else if (LevelToSpawnIn->ShouldCreateNewExternalActors() && SpawnParameters.bCreateActorPackage && !(SpawnParameters.ObjectFlags & RF_Transient))
 	{
 		bNeedGloballyUniqueName = CastChecked<AActor>(Class->GetDefaultObject())->SupportsExternalPackaging();
 	}
