Fixed body instance MassScale being able to go negative (or zero). Applied property clamps for details panel

diff --git a/Engine/Source/Runtime/Engine/Classes/PhysicsEngine/BodyInstance.h b/Engine/Source/Runtime/Engine/Classes/PhysicsEngine/BodyInstance.h
--- a/Engine/Source/Runtime/Engine/Classes/PhysicsEngine/BodyInstance.h
+++ b/Engine/Source/Runtime/Engine/Classes/PhysicsEngine/BodyInstance.h
@@ -470,7 +470,7 @@ public:
 	FVector COMNudge;
 
 	/** Per-instance scaling of mass */
-	UPROPERTY(EditAnywhere, AdvancedDisplay, BlueprintReadWrite, Category = Physics)
+	UPROPERTY(EditAnywhere, AdvancedDisplay, BlueprintReadWrite, Category = Physics, meta = (ClampMin = "0.001", UIMin = "0.001"))
 	float MassScale;
 
 	/** Per-instance scaling of inertia (bigger number means  it'll be harder to rotate) */
