Fix crash when reloading the global shader map.

Deleting and reloading the global shader map from the DDC was deleting shaders from underneath the renderer.  Ensure that the cache shaders are refreshed.  See CL# 14368615 for a similar fix.  Breaking out code to refresh primitives with global shaders into a function that can be called from the different pathways global shaders are compiled. Adding this call to refresh in CompileGlobalShaderMap() in situations where we are creating the global shader map.

diff --git a/Engine/Source/Runtime/Engine/Private/ShaderCompiler/ShaderCompiler.cpp b/Engine/Source/Runtime/Engine/Private/ShaderCompiler/ShaderCompiler.cpp
--- a/Engine/Source/Runtime/Engine/Private/ShaderCompiler/ShaderCompiler.cpp
+++ b/Engine/Source/Runtime/Engine/Private/ShaderCompiler/ShaderCompiler.cpp
@@ -3181,6 +3181,39 @@ void FShaderCompilingManager::BlockOnAllShaderMapCompletion(TMap<int32, FShaderM
 	}
 }
 
+namespace
+{
+	void PropagateGlobalShadersToAllPrimitives()
+	{
+		// Re-register everything to work around FShader lifetime issues - it currently lives and dies with the
+		// shadermap it is stored in, while cached MDCs can reference its memory. Re-registering will
+		// re-create the cache.
+		TRACE_CPUPROFILER_EVENT_SCOPE(PropagateGlobalShadersToAllPrimitives);
+
+		FObjectCacheContextScope ObjectCacheScope;
+		TSet<FSceneInterface*> ScenesToUpdate;
+		TIndirectArray<FComponentRecreateRenderStateContext> ComponentContexts;
+		for (UPrimitiveComponent* PrimitiveComponent : ObjectCacheScope.GetContext().GetPrimitiveComponents())
+		{
+			if (PrimitiveComponent->IsRenderStateCreated())
+			{
+				ComponentContexts.Add(new FComponentRecreateRenderStateContext(PrimitiveComponent, &ScenesToUpdate));
+#if WITH_EDITOR
+				if (PrimitiveComponent->HasValidSettingsForStaticLighting(false))
+				{
+					FStaticLightingSystemInterface::OnPrimitiveComponentUnregistered.Broadcast(PrimitiveComponent);
+					FStaticLightingSystemInterface::OnPrimitiveComponentRegistered.Broadcast(PrimitiveComponent);
+				}
+#endif
+			}
+		}
+
+		UpdateAllPrimitiveSceneInfosForScenes(ScenesToUpdate);
+		ComponentContexts.Empty();
+		UpdateAllPrimitiveSceneInfosForScenes(ScenesToUpdate);
+	}
+}
+
 void FShaderCompilingManager::ProcessCompiledShaderMaps(
 	TMap<int32, FShaderMapFinalizeResults>& CompiledShaderMaps, 
 	float TimeBudget)
@@ -3378,35 +3411,8 @@ void FShaderCompilingManager::ProcessCompiledShaderMaps(
 		{
 			ProcessCompiledGlobalShaders(CompileResults.FinishedJobs);
 			
-			// Re-register everything to work around FShader lifetime issues - it currently lives and dies with the
-			// shadermap it is stored in, while cached MDCs can reference its memory. Re-registering will
-			// re-create the cache.
-			{
-				TRACE_CPUPROFILER_EVENT_SCOPE(PropagateGlobalShadersToAllPrimitives);
+			PropagateGlobalShadersToAllPrimitives();
 
-				FObjectCacheContextScope ObjectCacheScope;
-				TSet<FSceneInterface*> ScenesToUpdate;
-				TIndirectArray<FComponentRecreateRenderStateContext> ComponentContexts;
-				for (UPrimitiveComponent* PrimitiveComponent : ObjectCacheScope.GetContext().GetPrimitiveComponents())
-				{
-					if (PrimitiveComponent->IsRenderStateCreated())
-					{
-						ComponentContexts.Add(new FComponentRecreateRenderStateContext(PrimitiveComponent, &ScenesToUpdate));
-#if WITH_EDITOR
-						if (PrimitiveComponent->HasValidSettingsForStaticLighting(false))
-						{
-							FStaticLightingSystemInterface::OnPrimitiveComponentUnregistered.Broadcast(PrimitiveComponent);
-							FStaticLightingSystemInterface::OnPrimitiveComponentRegistered.Broadcast(PrimitiveComponent);
-						}
-#endif
-					}
-				}
-
-				UpdateAllPrimitiveSceneInfosForScenes(ScenesToUpdate);
-				ComponentContexts.Empty();
-				UpdateAllPrimitiveSceneInfosForScenes(ScenesToUpdate);
-			}
-				
 			{
 				FScopeLock Lock(&CompileQueueSection);
 				for (auto& Job : CompileResults.FinishedJobs)
@@ -5742,6 +5748,8 @@ void CompileGlobalShaderMap(EShaderPlatform Platform, const ITargetPlatform* Tar
 		{
 			GGlobalShaderMap[Platform]->BeginCreateAllShaders();
 		}
+
+		PropagateGlobalShadersToAllPrimitives();
 	}
 }
 
@@ -6017,7 +6025,6 @@ void BeginRecompileGlobalShaders(const TArray<const FShaderType*>& OutdatedShade
 		// Now check if there is any work to be done wrt outdates types
 		if (OutdatedShaderTypes.Num() > 0 || OutdatedShaderPipelineTypes.Num() > 0)
 		{
-
 			VerifyGlobalShaders(ShaderPlatform, TargetPlatform, false, &OutdatedShaderTypes, &OutdatedShaderPipelineTypes);
 		}
 	}
