Fixed static analysis assert.

diff --git a/Engine/Source/Runtime/Engine/Private/ShaderCompiler/ShaderGenerationUtil.cpp b/Engine/Source/Runtime/Engine/Private/ShaderCompiler/ShaderGenerationUtil.cpp
--- a/Engine/Source/Runtime/Engine/Private/ShaderCompiler/ShaderGenerationUtil.cpp
+++ b/Engine/Source/Runtime/Engine/Private/ShaderCompiler/ShaderGenerationUtil.cpp
@@ -1681,7 +1681,7 @@ static void SetStandardGBufferSlots(bool Slots[], bool bWriteEmissive, bool bHas
 	Slots[GBS_SelectiveOutputMask] =	bIsStrataMaterial ? false : true;
 	Slots[GBS_BaseColor] =				bIsStrataMaterial ? false : true;
 	Slots[GBS_GenericAO] =				bIsStrataMaterial ? false : true;
-	Slots[GBS_AO] =						bIsStrataMaterial ? false : false;// true;		// Why false?
+	Slots[GBS_AO] =						false;//bIsStrataMaterial ? false : false;// true;		// Why only false?
 	Slots[GBS_WorldTangent] =			bIsStrataMaterial ? false : bHasTangent;
 	Slots[GBS_Anisotropy] =				bIsStrataMaterial ? false : bHasTangent;
 }
