Adding DumpShaderCompileStats exec command that simply calls GShaderCompilerStats->WriteStats() to write shader compilation stats to a csv file.  File will be written to the project's saved directory under: MaterialStats/Stats-*.csv

I've also only enabled this command for the editor because it is the only platform that should have shader compile stats.

diff --git a/Engine/Config/BaseInput.ini b/Engine/Config/BaseInput.ini
--- a/Engine/Config/BaseInput.ini
+++ b/Engine/Config/BaseInput.ini
@@ -68,6 +68,7 @@ MaxScrollbackSize=1024
 +ManualAutoCompleteList=(Command="DumpMaterialStats",Desc="Dump material information")
 +ManualAutoCompleteList=(Command="DumpShaderStats",Desc="Dump shader information")
 +ManualAutoCompleteList=(Command="DumpShaderPipelineStats",Desc="Dump shader pipeline information")
++ManualAutoCompleteList=(Command="DumpShaderCompileStats",Desc="Dump shader compilation information")
 +ManualAutoCompleteList=(Command="StartFPSChart",Desc="after that use StopFPSChart")
 +ManualAutoCompleteList=(Command="StopFPSChart",Desc="after that look for the output in Saved/Profiling/FPSChartStats")
 +ManualAutoCompleteList=(Command="FreezeAt",Desc="Locks the player view and rendering time.")
diff --git a/Engine/Source/Runtime/Engine/Classes/Engine/Engine.h b/Engine/Source/Runtime/Engine/Classes/Engine/Engine.h
--- a/Engine/Source/Runtime/Engine/Classes/Engine/Engine.h
+++ b/Engine/Source/Runtime/Engine/Classes/Engine/Engine.h
@@ -2018,6 +2018,9 @@ public:
 	bool HandleRecompileGlobalShadersCommand( const TCHAR* Cmd, FOutputDevice& Ar );
 	bool HandleDumpShaderStatsCommand( const TCHAR* Cmd, FOutputDevice& Ar );
 	bool HandleDumpMaterialStatsCommand( const TCHAR* Cmd, FOutputDevice& Ar );
+#if WITH_EDITOR
+	bool HandleDumpShaderCompileStatsCommand( const TCHAR* Cmd, FOutputDevice& Ar );
+#endif
 	bool HandleProfileCommand( const TCHAR* Cmd, FOutputDevice& Ar );
 	bool HandleProfileGPUHitchesCommand( const TCHAR* Cmd, FOutputDevice& Ar );
 	bool HandleShaderComplexityCommand( const TCHAR* Cmd, FOutputDevice& Ar );
diff --git a/Engine/Source/Runtime/Engine/Private/ShaderCompiler/ShaderCompiler.cpp b/Engine/Source/Runtime/Engine/Private/ShaderCompiler/ShaderCompiler.cpp
--- a/Engine/Source/Runtime/Engine/Private/ShaderCompiler/ShaderCompiler.cpp
+++ b/Engine/Source/Runtime/Engine/Private/ShaderCompiler/ShaderCompiler.cpp
@@ -120,12 +120,12 @@ namespace GlobalShaderCookStats
 	static int32 ShadersCompiled = 0;
 
 	static FCookStatsManager::FAutoRegisterCallback RegisterCookStats([](FCookStatsManager::AddStatFuncRef AddStat)
-	{
-		UsageStats.LogStats(AddStat, TEXT("GlobalShader.Usage"), TEXT(""));
-		AddStat(TEXT("GlobalShader.Misc"), FCookStatsManager::CreateKeyValueArray(
-			TEXT("ShadersCompiled"), ShadersCompiled
-		));
-	});
+		{
+			UsageStats.LogStats(AddStat, TEXT("GlobalShader.Usage"), TEXT(""));
+			AddStat(TEXT("GlobalShader.Misc"), FCookStatsManager::CreateKeyValueArray(
+				TEXT("ShadersCompiled"), ShadersCompiled
+			));
+		});
 }
 #endif
 
@@ -900,15 +900,15 @@ namespace ShaderCompilerCookStats
 	static double AsyncCompileTimeSec = 0.0;
 
 	static FCookStatsManager::FAutoRegisterCallback RegisterCookStats([](FCookStatsManager::AddStatFuncRef AddStat)
-	{
-		AddStat(TEXT("ShaderCompiler"), FCookStatsManager::CreateKeyValueArray(
-			TEXT("BlockingTimeSec"), BlockingTimeSec,
-			TEXT("AsyncCompileTimeSec"), AsyncCompileTimeSec,
-			TEXT("GlobalBeginCompileShaderTimeSec"), GlobalBeginCompileShaderTimeSec,
-			TEXT("GlobalBeginCompileShaderCalls"), GlobalBeginCompileShaderCalls,
-			TEXT("ProcessAsyncResultsTimeSec"), ProcessAsyncResultsTimeSec
+		{
+			AddStat(TEXT("ShaderCompiler"), FCookStatsManager::CreateKeyValueArray(
+				TEXT("BlockingTimeSec"), BlockingTimeSec,
+				TEXT("AsyncCompileTimeSec"), AsyncCompileTimeSec,
+				TEXT("GlobalBeginCompileShaderTimeSec"), GlobalBeginCompileShaderTimeSec,
+				TEXT("GlobalBeginCompileShaderCalls"), GlobalBeginCompileShaderCalls,
+				TEXT("ProcessAsyncResultsTimeSec"), ProcessAsyncResultsTimeSec
 			));
-	});
+		});
 }
 #endif
 
@@ -2283,13 +2283,14 @@ int32 FShaderCompileThreadRunnable::CompilingLoop()
 
 FShaderCompilerStats* GShaderCompilerStats = NULL;
 
-void FShaderCompilerStats::WriteStats()
+void FShaderCompilerStats::WriteStats(FOutputDevice* Ar)
 {
 #if ALLOW_DEBUG_FILES
+	static TCHAR DebugText[] = TEXT("Wrote shader compile stats to file '%s'.");
 	{
 		FlushRenderingCommands(true);
 
-		FString FileName = FPaths::Combine(*FPaths::ProjectSavedDir(), FString::Printf(TEXT("MaterialStats/Stats-%s.csv"),  *FDateTime::Now().ToString()));
+		FString FileName = FPaths::Combine(*FPaths::ProjectSavedDir(), FString::Printf(TEXT("MaterialStats/Stats-%s.csv"), *FDateTime::Now().ToString()));
 		auto DebugWriter = IFileManager::Get().CreateFileWriter(*FileName);
 		FDiagnosticTableWriterCSV StatWriter(DebugWriter);
 		const TSparseArray<ShaderCompilerStats>& PlatformStats = GetShaderCompilerStats();
@@ -2332,6 +2333,17 @@ void FShaderCompilerStats::WriteStats()
 			}
 		}
 		DebugWriter->Close();
+
+		FString FullFileName = FPaths::ConvertRelativePathToFull(FileName);
+		if (Ar)
+		{
+			Ar->Logf(DebugText, *FullFileName);
+		}
+		else
+		{
+			UE_LOG(LogShaderCompilers, Log, DebugText, *FullFileName);
+		}
+
 		if (FParse::Param(FCommandLine::Get(), TEXT("mirrorshaderstats")))
 		{
 			FString MirrorLocation;
@@ -2400,6 +2412,16 @@ void FShaderCompilerStats::WriteStats()
 
 			}
 		}
+
+		FString FullFileName = FPaths::ConvertRelativePathToFull(FileName);
+		if (Ar)
+		{
+			Ar->Logf(DebugText, *FullFileName);
+		}
+		else
+		{
+			UE_LOG(LogShaderCompilers, Log, DebugText, *FullFileName);
+		}
 	}
 #endif // ALLOW_DEBUG_FILES
 }
diff --git a/Engine/Source/Runtime/Engine/Private/UnrealEngine.cpp b/Engine/Source/Runtime/Engine/Private/UnrealEngine.cpp
--- a/Engine/Source/Runtime/Engine/Private/UnrealEngine.cpp
+++ b/Engine/Source/Runtime/Engine/Private/UnrealEngine.cpp
@@ -4257,6 +4257,12 @@ bool UEngine::Exec( UWorld* InWorld, const TCHAR* Cmd, FOutputDevice& Ar )
 	{
 		return HandleDumpMaterialStatsCommand( Cmd, Ar );	
 	}
+#if WITH_EDITOR
+	else if( FParse::Command(&Cmd,TEXT("DumpShaderCompileStats")) )
+	{
+		HandleDumpShaderCompileStatsCommand(Cmd, Ar);
+	}
+#endif
 	else if (FParse::Command(&Cmd, TEXT("DumpShaderPipelineStats")))
 	{
 		return HandleDumpShaderPipelineStatsCommand(Cmd, Ar);
@@ -5035,6 +5041,15 @@ bool UEngine::HandleDumpMaterialStatsCommand( const TCHAR* Cmd, FOutputDevice& A
 	return true;
 }
 
+#if WITH_EDITOR
+bool UEngine::HandleDumpShaderCompileStatsCommand(const TCHAR* Cmd, FOutputDevice& Ar)
+{
+	Ar.Logf(TEXT("Dumping shader compile stats"));
+	GShaderCompilerStats->WriteStats(&Ar);
+	return true;
+}
+#endif
+
 bool UEngine::HandleProfileCommand( const TCHAR* Cmd, FOutputDevice& Ar )
 {
 	if ( FParse::Command(&Cmd,TEXT("GPU")) )
diff --git a/Engine/Source/Runtime/Engine/Public/ShaderCompiler.h b/Engine/Source/Runtime/Engine/Public/ShaderCompiler.h
--- a/Engine/Source/Runtime/Engine/Public/ShaderCompiler.h
+++ b/Engine/Source/Runtime/Engine/Public/ShaderCompiler.h
@@ -809,7 +809,7 @@ public:
 	ENGINE_API void RegisterCookedShaders(uint32 NumCooked, float CompileTime, EShaderPlatform Platform, const FString MaterialPath, FString PermutationString = FString(""));
 	ENGINE_API void RegisterCompiledShaders(uint32 NumPermutations, EShaderPlatform Platform, const FString MaterialPath, FString PermutationString = FString(""));
 	ENGINE_API const TSparseArray<ShaderCompilerStats>& GetShaderCompilerStats() { return CompileStats; }
-	ENGINE_API void WriteStats();
+	ENGINE_API void WriteStats(class FOutputDevice* Ar = nullptr);
 
 private:
 	FCriticalSection CompileStatsLock;
