LevelInstance: Fix failed check + Avoid save prompt when child did not change

diff --git a/Engine/Source/Runtime/Engine/Private/LevelInstance/LevelInstanceSubsystem.cpp b/Engine/Source/Runtime/Engine/Private/LevelInstance/LevelInstanceSubsystem.cpp
--- a/Engine/Source/Runtime/Engine/Private/LevelInstance/LevelInstanceSubsystem.cpp
+++ b/Engine/Source/Runtime/Engine/Private/LevelInstance/LevelInstanceSubsystem.cpp
@@ -498,40 +498,60 @@ void ULevelInstanceSubsystem::ForEachLevelInstanceAncestors(const AActor* Actor,
 	while (ParentLevelInstance != nullptr && Operation(ParentLevelInstance));
 }
 
-void ULevelInstanceSubsystem::ForEachLevelInstanceChildren(const ALevelInstance* LevelInstanceActor, bool bRecursive, TFunctionRef<bool(const ALevelInstance*)> Operation) const
+void ULevelInstanceSubsystem::ForEachLevelInstanceChild(const ALevelInstance* LevelInstanceActor, bool bRecursive, TFunctionRef<bool(const ALevelInstance*)> Operation) const
 {
+	ForEachLevelInstanceChildImpl(LevelInstanceActor, bRecursive, Operation);
+}
+
+bool ULevelInstanceSubsystem::ForEachLevelInstanceChildImpl(const ALevelInstance* LevelInstanceActor, bool bRecursive, TFunctionRef<bool(const ALevelInstance*)> Operation) const
+{
+	bool bContinue = true;
 	if (ULevel* LevelInstanceLevel = GetLevelInstanceLevel(LevelInstanceActor))
 	{
-		ForEachActorInLevel(LevelInstanceLevel, [this, Operation,bRecursive](AActor* LevelActor)
+		ForEachActorInLevel(LevelInstanceLevel, [&bContinue, this, Operation,bRecursive](AActor* LevelActor)
 		{
 			if (const ALevelInstance* ChildLevelInstanceActor = Cast<ALevelInstance>(LevelActor))
 			{
-				if (Operation(ChildLevelInstanceActor) && bRecursive)
+				bContinue = Operation(ChildLevelInstanceActor);
+				
+				if (bContinue && bRecursive)
 				{
-					ForEachLevelInstanceChildren(ChildLevelInstanceActor, bRecursive, Operation);
+					bContinue = ForEachLevelInstanceChildImpl(ChildLevelInstanceActor, bRecursive, Operation);
 				}
 			}
-			return true;
+			return bContinue;
 		});
 	}
+
+	return bContinue;
 }
 
-void ULevelInstanceSubsystem::ForEachLevelInstanceChildren(ALevelInstance* LevelInstanceActor, bool bRecursive, TFunctionRef<bool(ALevelInstance*)> Operation) const
+void ULevelInstanceSubsystem::ForEachLevelInstanceChild(ALevelInstance* LevelInstanceActor, bool bRecursive, TFunctionRef<bool(ALevelInstance*)> Operation) const
 {
+	ForEachLevelInstanceChildImpl(LevelInstanceActor, bRecursive, Operation);
+}
+
+bool ULevelInstanceSubsystem::ForEachLevelInstanceChildImpl(ALevelInstance* LevelInstanceActor, bool bRecursive, TFunctionRef<bool(ALevelInstance*)> Operation) const
+{
+	bool bContinue = true;
 	if (ULevel* LevelInstanceLevel = GetLevelInstanceLevel(LevelInstanceActor))
 	{
-		ForEachActorInLevel(LevelInstanceLevel, [this, Operation, bRecursive](AActor* LevelActor)
+		ForEachActorInLevel(LevelInstanceLevel, [&bContinue, this, Operation, bRecursive](AActor* LevelActor)
 		{
 			if (ALevelInstance* ChildLevelInstanceActor = Cast<ALevelInstance>(LevelActor))
 			{
-				if (Operation(ChildLevelInstanceActor) && bRecursive)
+				bContinue = Operation(ChildLevelInstanceActor);
+
+				if (bContinue && bRecursive)
 				{
-					ForEachLevelInstanceChildren(ChildLevelInstanceActor, bRecursive, Operation);
+					bContinue = ForEachLevelInstanceChildImpl(ChildLevelInstanceActor, bRecursive, Operation);
 				}
 			}
-			return true;
+			return bContinue;
 		});
 	}
+
+	return bContinue;
 }
 
 void ULevelInstanceSubsystem::ForEachLevelInstanceEdit(TFunctionRef<bool(ALevelInstance*)> Operation) const
@@ -548,7 +568,7 @@ void ULevelInstanceSubsystem::ForEachLevelInstanceEdit(TFunctionRef<bool(ALevelI
 bool ULevelInstanceSubsystem::HasDirtyChildrenLevelInstances(const ALevelInstance* LevelInstanceActor) const
 {
 	bool bDirtyChildren = false;
-	ForEachLevelInstanceChildren(LevelInstanceActor, /*bRecursive=*/true, [this, &bDirtyChildren](const ALevelInstance* ChildLevelInstanceActor)
+	ForEachLevelInstanceChild(LevelInstanceActor, /*bRecursive=*/true, [this, &bDirtyChildren](const ALevelInstance* ChildLevelInstanceActor)
 	{
 		if (IsEditingLevelInstanceDirty(ChildLevelInstanceActor))
 		{
@@ -563,7 +583,7 @@ bool ULevelInstanceSubsystem::HasDirtyChildrenLevelInstances(const ALevelInstanc
 bool ULevelInstanceSubsystem::HasEditingChildrenLevelInstances(const ALevelInstance* LevelInstanceActor) const
 {
 	bool bEditingChildren = false;
-	ForEachLevelInstanceChildren(LevelInstanceActor, true, [this, &bEditingChildren](const ALevelInstance* ChildLevelInstanceActor)
+	ForEachLevelInstanceChild(LevelInstanceActor, true, [this, &bEditingChildren](const ALevelInstance* ChildLevelInstanceActor)
 		{
 			if (IsEditingLevelInstance(ChildLevelInstanceActor))
 			{
@@ -986,7 +1006,7 @@ bool ULevelInstanceSubsystem::CanMoveActorToLevel(const AActor* Actor, FText* Ou
 			}
 
 			bool bEditingChildren = false;
-			ForEachLevelInstanceChildren(LevelInstanceActor, true, [this, &bEditingChildren](const ALevelInstance* ChildLevelInstanceActor)
+			ForEachLevelInstanceChild(LevelInstanceActor, true, [this, &bEditingChildren](const ALevelInstance* ChildLevelInstanceActor)
 			{
 				if (IsEditingLevelInstance(ChildLevelInstanceActor))
 				{
@@ -1332,7 +1352,7 @@ void ULevelInstanceSubsystem::EditLevelInstance(ALevelInstance* LevelInstanceAct
 
 		if (!LevelInstanceToCommit)
 		{
-			ForEachLevelInstanceChildren(LevelInstanceActor, true, GetLevelInstanceToCommit);
+			ForEachLevelInstanceChild(LevelInstanceActor, true, GetLevelInstanceToCommit);
 		}
 		
 		if (LevelInstanceToCommit)
@@ -1410,7 +1430,7 @@ void ULevelInstanceSubsystem::EditLevelInstance(ALevelInstance* LevelInstanceAct
 void ULevelInstanceSubsystem::CommitChildrenLevelInstances(ALevelInstance* LevelInstanceActor)
 {
 	// We are ending editing. Discard Non dirty child edits
-	ForEachLevelInstanceChildren(LevelInstanceActor, /*bRecursive=*/true, [this](const ALevelInstance* ChildLevelInstanceActor)
+	ForEachLevelInstanceChild(LevelInstanceActor, /*bRecursive=*/true, [this](const ALevelInstance* ChildLevelInstanceActor)
 	{
 		if (const FLevelInstanceEdit* ChildLevelInstanceEdit = GetLevelInstanceEdit(ChildLevelInstanceActor))
 		{
@@ -1502,9 +1522,9 @@ ALevelInstance* ULevelInstanceSubsystem::CommitLevelInstance(ALevelInstance* Lev
 	LevelInstanceActor->OnCommit();
 
 	// Notify Ancestors
-	ForEachLevelInstanceAncestors(LevelInstanceActor, [](ALevelInstance* AncestorLevelInstance)
+	ForEachLevelInstanceAncestors(LevelInstanceActor, [bChangesCommitted](ALevelInstance* AncestorLevelInstance)
 	{
-		AncestorLevelInstance->OnCommitChild();
+		AncestorLevelInstance->OnCommitChild(bChangesCommitted);
 		return true;
 	});
 				
diff --git a/Engine/Source/Runtime/Engine/Private/LevelInstance/Packed/PackedLevelInstanceActor.cpp b/Engine/Source/Runtime/Engine/Private/LevelInstance/Packed/PackedLevelInstanceActor.cpp
--- a/Engine/Source/Runtime/Engine/Private/LevelInstance/Packed/PackedLevelInstanceActor.cpp
+++ b/Engine/Source/Runtime/Engine/Private/LevelInstance/Packed/PackedLevelInstanceActor.cpp
@@ -7,11 +7,13 @@
 
 APackedLevelInstance::APackedLevelInstance(const FObjectInitializer& ObjectInitializer)
 	: Super(ObjectInitializer)
-{
 #if WITH_EDITORONLY_DATA
-	ISMComponentClass = UInstancedStaticMeshComponent::StaticClass();
-	bEditing = false;
+	, ISMComponentClass(UInstancedStaticMeshComponent::StaticClass())
+	, bEditing(false)
+	, ChildEditing(0)
+	, bChildChanged(false)
 #endif
+{
 }
 
 bool APackedLevelInstance::SupportsLoading() const
@@ -65,26 +67,31 @@ void APackedLevelInstance::OnEditChild()
 	MarkComponentsRenderStateDirty();
 }
 
-void APackedLevelInstance::OnCommitChild()
+void APackedLevelInstance::OnCommitChild(bool bChanged)
 {
 	check(GetLevelInstanceSubsystem()->GetLevelInstanceLevel(this) != nullptr);
 	check(ChildEditing > 0);
 	ChildEditing--;
+	bChildChanged |= bChanged;
 	if (!ChildEditing)
 	{
 		UnloadLevelInstance();
 
-		// Reflect child changes
-		TSharedPtr<FPackedLevelInstanceBuilder> Builder = FPackedLevelInstanceBuilder::CreateDefaultBuilder();
-		
-		if (UBlueprint* GeneratedBy = Cast<UBlueprint>(GetClass()->ClassGeneratedBy))
-		{
-			check(GeneratedBy == BlueprintAsset.Get());
-			Builder->UpdateBlueprint(GeneratedBy);
-		}
-		else
+		if (bChildChanged)
 		{
-			Builder->PackActor(this, GetWorldAsset());
+			// Reflect child changes
+			TSharedPtr<FPackedLevelInstanceBuilder> Builder = FPackedLevelInstanceBuilder::CreateDefaultBuilder();
+
+			if (UBlueprint* GeneratedBy = Cast<UBlueprint>(GetClass()->ClassGeneratedBy))
+			{
+				check(GeneratedBy == BlueprintAsset.Get());
+				Builder->UpdateBlueprint(GeneratedBy);
+			}
+			else
+			{
+				Builder->PackActor(this, GetWorldAsset());
+			}
+			bChildChanged = false;
 		}
 
 		MarkComponentsRenderStateDirty();
diff --git a/Engine/Source/Runtime/Engine/Public/LevelInstance/LevelInstanceActor.h b/Engine/Source/Runtime/Engine/Public/LevelInstance/LevelInstanceActor.h
--- a/Engine/Source/Runtime/Engine/Public/LevelInstance/LevelInstanceActor.h
+++ b/Engine/Source/Runtime/Engine/Public/LevelInstance/LevelInstanceActor.h
@@ -82,7 +82,7 @@ public:
 	virtual void OnEdit();
 	virtual void OnEditChild() {}
 	virtual void OnCommit();
-	virtual void OnCommitChild() {}
+	virtual void OnCommitChild(bool bChanged) {}
 		
 	DECLARE_MULTICAST_DELEGATE_OneParam(FOnLevelInstanceActorPostLoad, ALevelInstance*);
 	static FOnLevelInstanceActorPostLoad OnLevelInstanceActorPostLoad;
diff --git a/Engine/Source/Runtime/Engine/Public/LevelInstance/LevelInstanceSubsystem.h b/Engine/Source/Runtime/Engine/Public/LevelInstance/LevelInstanceSubsystem.h
--- a/Engine/Source/Runtime/Engine/Public/LevelInstance/LevelInstanceSubsystem.h
+++ b/Engine/Source/Runtime/Engine/Public/LevelInstance/LevelInstanceSubsystem.h
@@ -64,8 +64,8 @@ public:
 	void ForEachActorInLevelInstance(const ALevelInstance* LevelInstanceActor, TFunctionRef<bool(AActor * LevelActor)> Operation) const;
 	void ForEachLevelInstanceAncestorsAndSelf(const AActor* Actor, TFunctionRef<bool(const ALevelInstance*)> Operation) const;
 	void ForEachLevelInstanceAncestors(const AActor* Actor, TFunctionRef<bool(const ALevelInstance*)> Operation) const;
-	void ForEachLevelInstanceChildren(const ALevelInstance* LevelInstanceActor, bool bRecursive, TFunctionRef<bool(const ALevelInstance*)> Operation) const;
-	void ForEachLevelInstanceChildren(ALevelInstance* LevelInstanceActor, bool bRecursive, TFunctionRef<bool(ALevelInstance*)> Operation) const;
+	void ForEachLevelInstanceChild(const ALevelInstance* LevelInstanceActor, bool bRecursive, TFunctionRef<bool(const ALevelInstance*)> Operation) const;
+	void ForEachLevelInstanceChild(ALevelInstance* LevelInstanceActor, bool bRecursive, TFunctionRef<bool(ALevelInstance*)> Operation) const;
 	void ForEachLevelInstanceEdit(TFunctionRef<bool(ALevelInstance*)> Operation) const;
 	bool HasDirtyChildrenLevelInstances(const ALevelInstance* LevelInstanceActor) const;
 	bool HasEditingChildrenLevelInstances(const ALevelInstance* LevelInstanceActor) const;
@@ -101,6 +101,9 @@ private:
 	ALevelInstance* GetOwningLevelInstance(const ULevel* Level) const;
 	FLevelInstanceID ComputeLevelInstanceID(ALevelInstance* LevelInstanceActor) const;
 #if WITH_EDITOR
+	bool ForEachLevelInstanceChildImpl(const ALevelInstance* LevelInstanceActor, bool bRecursive, TFunctionRef<bool(const ALevelInstance*)> Operation) const;
+	bool ForEachLevelInstanceChildImpl(ALevelInstance* LevelInstanceActor, bool bRecursive, TFunctionRef<bool(ALevelInstance*)> Operation) const;
+
 	void CommitChildrenLevelInstances(ALevelInstance* LevelInstanceActor);
 	void BreakLevelInstance_Impl(ALevelInstance* LevelInstanceActor, uint32 Levels, TArray<AActor*>& OutMovedActors);
 
diff --git a/Engine/Source/Runtime/Engine/Public/LevelInstance/Packed/PackedLevelInstanceActor.h b/Engine/Source/Runtime/Engine/Public/LevelInstance/Packed/PackedLevelInstanceActor.h
--- a/Engine/Source/Runtime/Engine/Public/LevelInstance/Packed/PackedLevelInstanceActor.h
+++ b/Engine/Source/Runtime/Engine/Public/LevelInstance/Packed/PackedLevelInstanceActor.h
@@ -38,7 +38,7 @@ public:
 	virtual void OnWorldAssetChanged() override;
 	virtual void OnWorldAssetSaved() override;
 	virtual void OnCommit() override;
-	virtual void OnCommitChild() override;
+	virtual void OnCommitChild(bool bChanged) override;
 	virtual void OnEdit() override;
 	virtual void OnEditChild() override;
 
@@ -73,9 +73,12 @@ public:
 
 private:
 	UPROPERTY(NonTransactional)
-	bool bEditing = false;
+	bool bEditing;
 
 	UPROPERTY(NonTransactional)
-	int32 ChildEditing = 0;
+	int32 ChildEditing;
+
+	UPROPERTY(NonTransactional)
+	bool bChildChanged;
 #endif
 };
\ No newline at end of file
