Fix the textures to be able to fetch content from a shared DDC instead of always recompiling them if not present in local DDC
Prefetch remote texture data on an async thread when using a shared DDC to reduce stutters on the game-thread and improve the editor experience

  - Tested with a 32 Core 3970X and a 1Gbps internet connection to the DDC (Jupiter)
    - 15m to 7m until the editor is totally idle when opening P_Construct with an empty local DDC

diff --git a/Engine/Source/Runtime/Engine/Private/TextureDerivedData.cpp b/Engine/Source/Runtime/Engine/Private/TextureDerivedData.cpp
--- a/Engine/Source/Runtime/Engine/Private/TextureDerivedData.cpp
+++ b/Engine/Source/Runtime/Engine/Private/TextureDerivedData.cpp
@@ -1388,8 +1388,33 @@ bool FTexturePlatformData::AreDerivedMipsAvailable() const
 			MipKeys.Add(Mip.DerivedDataKey);
 		}
 	}
-	return GetDerivedDataCacheRef().AllCachedDataProbablyExists(MipKeys);
+	
+	// If this returns true, everything is already present in the local cache with fast access
+	bool bDataProbablyExists = GetDerivedDataCacheRef().AllCachedDataProbablyExists(MipKeys);
+
+	// This is a temporary solution that should be replaced by a Prefetch command
+	// once the new DDC is in place. 
+	// CachedDataProbablyExists doesn't work on slow cache, so it means we rebuild the 
+	// textures every time we don't have these locally.. which is bad, so Get them instead.
+	if (!bDataProbablyExists && !IsInGameThread())
+	{
+		bDataProbablyExists = true;
+		// When using a shared DDC and performing async loading, 
+		// prefetch the data so we don't pay the latency on the game-thread
+		// to fetch the asset from a remote location once
+		// the mips needs to be accessed
+		TRACE_CPUPROFILER_EVENT_SCOPE(PrefetchMips);
+		TArray<uint8> ThrowAwayData;
+		for (FString& Key : MipKeys)
+		{
+			ThrowAwayData.Reset();
+			bDataProbablyExists &= GetDerivedDataCacheRef().GetSynchronous(*Key, ThrowAwayData);
+		}
+	}
+
+	return bDataProbablyExists;
 }
+
 bool FTexturePlatformData::AreDerivedVTChunksAvailable() const
 {
 	check(VTData);
@@ -1401,7 +1426,31 @@ bool FTexturePlatformData::AreDerivedVTChunksAvailable() const
 			ChunkKeys.Add(Chunk.DerivedDataKey);
 		}
 	}
-	return GetDerivedDataCacheRef().AllCachedDataProbablyExists(ChunkKeys);
+
+	// If this returns true, everything is already present in the local cache with fast access
+	bool bDataProbablyExists = GetDerivedDataCacheRef().AllCachedDataProbablyExists(ChunkKeys);
+
+	// This is a temporary solution that should be replaced by a Prefetch command
+	// once the new DDC is in place. 
+	// CachedDataProbablyExists doesn't work on slow cache, so it means we rebuild the 
+	// textures every time we don't have these locally.. which is bad, so Get them instead.
+	if (!bDataProbablyExists && !IsInGameThread())
+	{
+		bDataProbablyExists = true;
+		// When using a shared DDC and performing async loading, 
+		// prefetch the data so we don't pay the latency on the game-thread
+		// to fetch the asset from a remote location once
+		// the mips needs to be accessed
+		TRACE_CPUPROFILER_EVENT_SCOPE(PrefetchDerivedVTChunks);
+		TArray<uint8> ThrowAwayData;
+		for (FString& Key : ChunkKeys)
+		{
+			ThrowAwayData.Reset();
+			bDataProbablyExists &= GetDerivedDataCacheRef().GetSynchronous(*Key, ThrowAwayData);
+		}
+	}
+
+	return bDataProbablyExists;
 }
 #endif // #if WITH_EDITOR
 
