Fix missing DebugContext to avoid deprecation warnings

diff --git a/Engine/Source/Runtime/Engine/Private/TextureDerivedData.cpp b/Engine/Source/Runtime/Engine/Private/TextureDerivedData.cpp
--- a/Engine/Source/Runtime/Engine/Private/TextureDerivedData.cpp
+++ b/Engine/Source/Runtime/Engine/Private/TextureDerivedData.cpp
@@ -1408,7 +1408,7 @@ bool FTexturePlatformData::AreDerivedMipsAvailable() const
 		for (FString& Key : MipKeys)
 		{
 			ThrowAwayData.Reset();
-			bDataProbablyExists &= GetDerivedDataCacheRef().GetSynchronous(*Key, ThrowAwayData);
+			bDataProbablyExists &= GetDerivedDataCacheRef().GetSynchronous(*Key, ThrowAwayData, TEXT("DerivedMips"));
 		}
 	}
 
@@ -1440,13 +1440,13 @@ bool FTexturePlatformData::AreDerivedVTChunksAvailable() const
 		// When using a shared DDC and performing async loading, 
 		// prefetch the data so we don't pay the latency on the game-thread
 		// to fetch the asset from a remote location once
-		// the mips needs to be accessed
+		// the VT chunks needs to be accessed
 		TRACE_CPUPROFILER_EVENT_SCOPE(PrefetchDerivedVTChunks);
 		TArray<uint8> ThrowAwayData;
 		for (FString& Key : ChunkKeys)
 		{
 			ThrowAwayData.Reset();
-			bDataProbablyExists &= GetDerivedDataCacheRef().GetSynchronous(*Key, ThrowAwayData);
+			bDataProbablyExists &= GetDerivedDataCacheRef().GetSynchronous(*Key, ThrowAwayData, TEXT("DerivedVTChunks"));
 		}
 	}
 
