If migration is not able to reconnect a skeletalmesh skeleton, the engine will crash when loading the skeletalmesh and the user will not be able to re-assign the proper skeleton.

diff --git a/Engine/Source/Developer/MeshBoneReduction/Private/MeshBoneReduction.cpp b/Engine/Source/Developer/MeshBoneReduction/Private/MeshBoneReduction.cpp
--- a/Engine/Source/Developer/MeshBoneReduction/Private/MeshBoneReduction.cpp
+++ b/Engine/Source/Developer/MeshBoneReduction/Private/MeshBoneReduction.cpp
@@ -361,9 +361,10 @@ public:
 
 	bool ReduceBoneCounts(USkeletalMesh* SkeletalMesh, int32 DesiredLOD, const TArray<FName>* BoneNamesToRemove, bool bCallPostEditChange /*= true*/) override
 	{
-		check (SkeletalMesh);
-		USkeleton* Skeleton = SkeletalMesh->GetSkeleton();
-		check (Skeleton);
+		if (!ensure(SkeletalMesh))
+		{
+			return false;
+		}
 
 		// find all the bones to remove from Skeleton settings
 		TMap<FBoneIndexType, FBoneIndexType> BonesToRemove;
diff --git a/Engine/Source/Developer/MeshBuilder/Private/SkeletalMeshBuilder.cpp b/Engine/Source/Developer/MeshBuilder/Private/SkeletalMeshBuilder.cpp
--- a/Engine/Source/Developer/MeshBuilder/Private/SkeletalMeshBuilder.cpp
+++ b/Engine/Source/Developer/MeshBuilder/Private/SkeletalMeshBuilder.cpp
@@ -185,7 +185,7 @@ bool FSkeletalMeshBuilder::Build(USkeletalMesh* SkeletalMesh, const int32 LODInd
 		}
 		else
 		{
-			if (LODInfo->BonesToRemove.Num() > 0)
+			if (LODInfo->BonesToRemove.Num() > 0 && SkeletalMesh->GetSkeleton())
 			{
 				TArray<FName> BonesToRemove;
 				BonesToRemove.Reserve(LODInfo->BonesToRemove.Num());
diff --git a/Engine/Source/Runtime/Engine/Private/SkeletalMesh.cpp b/Engine/Source/Runtime/Engine/Private/SkeletalMesh.cpp
--- a/Engine/Source/Runtime/Engine/Private/SkeletalMesh.cpp
+++ b/Engine/Source/Runtime/Engine/Private/SkeletalMesh.cpp
@@ -4364,13 +4364,13 @@ const UAnimSequence* USkeletalMesh::GetBakePose(int32 LODIndex) const
 	const FSkeletalMeshLODInfo* LOD = GetLODInfo(LODIndex);
 	if (LOD)
 	{
-		if (LOD->BakePoseOverride && GetSkeleton() == LOD->BakePoseOverride->GetSkeleton())
+		if (LOD->BakePoseOverride && GetSkeleton() && GetSkeleton() == LOD->BakePoseOverride->GetSkeleton())
 		{
 			return LOD->BakePoseOverride;
 		}
 
 		// we make sure bake pose uses same skeleton
-		if (LOD->BakePose && GetSkeleton() == LOD->BakePose->GetSkeleton())
+		if (LOD->BakePose && GetSkeleton() && GetSkeleton() == LOD->BakePose->GetSkeleton())
 		{
 			return LOD->BakePose;
 		}
