World Partition
- Minimap renderer now get it's bounding volume through the editor hash to make sure it fits with the editor grid region.
- Sort actor descriptors by extent when using wp.editor.dumpactordescs console command.

diff --git a/Engine/Source/Runtime/Engine/Private/WorldPartition/WorldPartition.cpp b/Engine/Source/Runtime/Engine/Private/WorldPartition/WorldPartition.cpp
--- a/Engine/Source/Runtime/Engine/Private/WorldPartition/WorldPartition.cpp
+++ b/Engine/Source/Runtime/Engine/Private/WorldPartition/WorldPartition.cpp
@@ -913,14 +913,24 @@ void UWorldPartition::DumpActorDescs(const FString& Path)
 	{
 		FString LineEntry = TEXT("Guid, Class, Name, BVCenterX, BVCenterY, BVCenterZ, BVExtentX, BVExtentY, BVExtentZ") LINE_TERMINATOR;
 		LogFile->Serialize(TCHAR_TO_ANSI(*LineEntry), LineEntry.Len());
-			
+
+		TArray<const FWorldPartitionActorDesc*> ActorDescs;
 		for (UActorDescContainer::TConstIterator<> ActorDescIterator(this); ActorDescIterator; ++ActorDescIterator)
+		{
+			ActorDescs.Add(*ActorDescIterator);
+		}
+		ActorDescs.Sort([](const FWorldPartitionActorDesc& A, const FWorldPartitionActorDesc& B)
+		{
+			return A.GetBounds().GetExtent().GetMax() < B.GetBounds().GetExtent().GetMax();
+		});
+		for (const FWorldPartitionActorDesc* ActorDescIterator : ActorDescs)
 		{
 			LineEntry = FString::Printf(
-				TEXT("%s, %s, %s, %.2f, %.2f, %.2f, %.2f, %.2f, %.2f") LINE_TERMINATOR, 
+				TEXT("%s, %s, %s, %s, %.2f, %.2f, %.2f, %.2f, %.2f, %.2f") LINE_TERMINATOR, 
 				*ActorDescIterator->GetGuid().ToString(), 
 				*ActorDescIterator->GetClass().ToString(), 
 				*FPaths::GetExtension(ActorDescIterator->GetActorPath().ToString()), 
+				ActorDescIterator->GetActor() ? *ActorDescIterator->GetActor()->GetActorLabel(false) : TEXT("None"),
 				ActorDescIterator->GetBounds().GetCenter().X,
 				ActorDescIterator->GetBounds().GetCenter().Y,
 				ActorDescIterator->GetBounds().GetCenter().Z,
@@ -966,6 +976,11 @@ FBox UWorldPartition::GetWorldBounds() const
 	}
 	return WorldBounds;
 }
+
+FBox UWorldPartition::GetEditorWorldBounds() const
+{
+	return EditorHash->GetEditorWorldBounds();
+}
 #endif
 
 #undef LOCTEXT_NAMESPACE
\ No newline at end of file
diff --git a/Engine/Source/Runtime/Engine/Private/WorldPartition/WorldPartitionEditorSpatialHash.cpp b/Engine/Source/Runtime/Engine/Private/WorldPartition/WorldPartitionEditorSpatialHash.cpp
--- a/Engine/Source/Runtime/Engine/Private/WorldPartition/WorldPartitionEditorSpatialHash.cpp
+++ b/Engine/Source/Runtime/Engine/Private/WorldPartition/WorldPartitionEditorSpatialHash.cpp
@@ -39,6 +39,11 @@ FName UWorldPartitionEditorSpatialHash::GetWorldPartitionEditorName()
 	return TEXT("SpatialHash");
 }
 
+FBox UWorldPartitionEditorSpatialHash::GetEditorWorldBounds() const
+{
+	return Bounds;
+}
+
 void UWorldPartitionEditorSpatialHash::Tick(float DeltaSeconds)
 {
 	if (bBoundsDirty)
diff --git a/Engine/Source/Runtime/Engine/Private/WorldPartition/WorldPartitionMiniMapHelper.cpp b/Engine/Source/Runtime/Engine/Private/WorldPartition/WorldPartitionMiniMapHelper.cpp
--- a/Engine/Source/Runtime/Engine/Private/WorldPartition/WorldPartitionMiniMapHelper.cpp
+++ b/Engine/Source/Runtime/Engine/Private/WorldPartition/WorldPartitionMiniMapHelper.cpp
@@ -17,6 +17,7 @@
 #include "PackageSourceControlHelper.h"
 #include "ProfilingDebugging/ScopedTimers.h"
 #include "WorldPartition/WorldPartitionMiniMap.h"
+#include"WorldPartition/WorldPartition.h"
 
 DEFINE_LOG_CATEGORY_STATIC(LogWorldPartitionMiniMapHelper, All, All);
 
@@ -84,7 +85,7 @@ void FWorldPartitionMiniMapHelper::CaptureWorldMiniMapToTexture(UWorld* InWorld,
 	WaitForShaderCompilation();
 	
 	//Calculate bounds of the World
-	OutWorldBounds = FWorldPartitionMiniMapHelper::GetWorldBounds(InWorld);
+	OutWorldBounds = InWorld->GetWorldPartition()->GetEditorWorldBounds();
 	
 	//Calculate Viewport size
 	FBox2D WorldBounds2D(FVector2D(OutWorldBounds.Min), FVector2D(OutWorldBounds.Max));
@@ -138,36 +139,6 @@ void FWorldPartitionMiniMapHelper::CaptureWorldMiniMapToTexture(UWorld* InWorld,
 	InOutMiniMapTexture->UpdateResource();
 }
 
-bool FWorldPartitionMiniMapHelper::DoesActorContributeToBounds(AActor* Actor)
-{
-	return Actor && Actor->GetGridPlacement() != EActorGridPlacement::AlwaysLoaded;
-}
-
-FBox FWorldPartitionMiniMapHelper::GetWorldBounds(UWorld* World)
-{
-	FBox WorldBox(ForceInit);
-
-	TArray<ULevel*> LevelsToRender = World->GetLevels();
-	for (ULevel* Level : LevelsToRender)
-	{
-		if (Level && Level->bIsVisible)
-		{
-			for (AActor* Actor : Level->Actors)
-			{
-				if ( Actor && DoesActorContributeToBounds(Actor))
-				{
-					FVector Origin, Extent;
-					Actor->GetActorBounds(false, Origin, Extent);
-					FBox ActorBounds(Origin - Extent, Origin + Extent);
-					WorldBox += ActorBounds;
-				}
-			}
-		}
-	}
-
-	return WorldBox;
-}
-
 void FWorldPartitionMiniMapHelper::CalTopViewOfWorld(FMatrix& OutProjectionMatrix, const FBox& WorldBox, uint32 ViewportWidth, uint32 ViewportHeight)
 {
 	const FVector Origin = WorldBox.GetCenter();
diff --git a/Engine/Source/Runtime/Engine/Public/WorldPartition/WorldPartition.h b/Engine/Source/Runtime/Engine/Public/WorldPartition/WorldPartition.h
--- a/Engine/Source/Runtime/Engine/Public/WorldPartition/WorldPartition.h
+++ b/Engine/Source/Runtime/Engine/Public/WorldPartition/WorldPartition.h
@@ -140,6 +140,7 @@ public:
 	void FinalizeGeneratedPackageForCook();
 
 	FBox GetWorldBounds() const;
+	FBox GetEditorWorldBounds() const;
 	void GenerateHLOD(ISourceControlHelper* SourceControlHelper);
 	void GenerateNavigationData();
 
diff --git a/Engine/Source/Runtime/Engine/Public/WorldPartition/WorldPartitionEditorHash.h b/Engine/Source/Runtime/Engine/Public/WorldPartition/WorldPartitionEditorHash.h
--- a/Engine/Source/Runtime/Engine/Public/WorldPartition/WorldPartitionEditorHash.h
+++ b/Engine/Source/Runtime/Engine/Public/WorldPartition/WorldPartitionEditorHash.h
@@ -20,6 +20,7 @@ public:
 	virtual void Initialize() PURE_VIRTUAL(UWorldPartitionEditorHash::Initialize, return;);
 	virtual void SetDefaultValues() PURE_VIRTUAL(UWorldPartitionEditorHash::SetDefaultValues, return;);
 	virtual FName GetWorldPartitionEditorName() PURE_VIRTUAL(UWorldPartitionEditorHash::GetWorldPartitionEditorName, return FName(NAME_None););
+	virtual FBox GetEditorWorldBounds() const  PURE_VIRTUAL(UWorldPartitionEditorHash::GetEditorWorldBounds, return FBox(ForceInit););
 	virtual void Tick(float DeltaSeconds) PURE_VIRTUAL(UWorldPartitionEditorHash::Tick, return;);
 
 	virtual void HashActor(FWorldPartitionHandle& InActorHandle) PURE_VIRTUAL(UWorldPartitionEditorHash::HashActor, ;);
diff --git a/Engine/Source/Runtime/Engine/Public/WorldPartition/WorldPartitionEditorSpatialHash.h b/Engine/Source/Runtime/Engine/Public/WorldPartition/WorldPartitionEditorSpatialHash.h
--- a/Engine/Source/Runtime/Engine/Public/WorldPartition/WorldPartitionEditorSpatialHash.h
+++ b/Engine/Source/Runtime/Engine/Public/WorldPartition/WorldPartitionEditorSpatialHash.h
@@ -210,6 +210,7 @@ public:
 	virtual void Initialize() override;
 	virtual void SetDefaultValues() override;
 	virtual FName GetWorldPartitionEditorName() override;
+	virtual FBox GetEditorWorldBounds() const override;
 	virtual void Tick(float DeltaSeconds) override;
 
 	virtual void HashActor(FWorldPartitionHandle& InActorHandle) override;
diff --git a/Engine/Source/Runtime/Engine/Public/WorldPartition/WorldPartitionMiniMapHelper.h b/Engine/Source/Runtime/Engine/Public/WorldPartition/WorldPartitionMiniMapHelper.h
--- a/Engine/Source/Runtime/Engine/Public/WorldPartition/WorldPartitionMiniMapHelper.h
+++ b/Engine/Source/Runtime/Engine/Public/WorldPartition/WorldPartitionMiniMapHelper.h
@@ -23,8 +23,6 @@ public:
 	static AWorldPartitionMiniMap* GetWorldPartitionMiniMap(UWorld* World, bool bCreateNewMiniMap=false);
 	static void CaptureWorldMiniMapToTexture(UWorld* InWorld, UObject* InOuterForTexture, uint32 InMiniMapSize, UTexture2D*& InOutMiniMapTexture, FBox& OutWorldBounds);
 private:
-	static bool DoesActorContributeToBounds(AActor* Actor);
-	static FBox GetWorldBounds(UWorld* World);
 	static void CalTopViewOfWorld(FMatrix& OutProjectionMatrix, const FBox& WorldBox, uint32 ViewportWidth, uint32 ViewportHeight);
 };
 #endif
