Add Stall Detectors to report on overbudget routines
* Block/Sync routines for async processing tasks
* Open time for Asset Editor window

diff --git a/Engine/Source/Editor/UnrealEd/Private/Subsystems/AssetEditorSubsystem.cpp b/Engine/Source/Editor/UnrealEd/Private/Subsystems/AssetEditorSubsystem.cpp
--- a/Engine/Source/Editor/UnrealEd/Private/Subsystems/AssetEditorSubsystem.cpp
+++ b/Engine/Source/Editor/UnrealEd/Private/Subsystems/AssetEditorSubsystem.cpp
@@ -30,6 +30,7 @@
 #include "AssetEditorMessages.h"
 #include "EditorModeManager.h"
 #include "Tools/LegacyEdMode.h"
+#include "ProfilingDebugging/StallDetector.h"
 
 
 #define LOCTEXT_NAMESPACE "AssetEditorSubsystem"
@@ -275,6 +276,8 @@ bool UAssetEditorSubsystem::CloseAllAssetEditors()
 
 bool UAssetEditorSubsystem::OpenEditorForAsset(UObject* Asset, const EToolkitMode::Type ToolkitMode, TSharedPtr< IToolkitHost > OpenedFromLevelEditor, const bool bShowProgressWindow)
 {
+	SCOPE_STALL_REPORTER(UAssetEditorSubsystem::OpenEditorForAsset, 2.0);
+
 	const double OpenAssetStartTime = FStudioAnalytics::GetAnalyticSeconds();
 
 	if (!Asset)
diff --git a/Engine/Source/Runtime/Engine/Private/AsyncCompilationHelpers.cpp b/Engine/Source/Runtime/Engine/Private/AsyncCompilationHelpers.cpp
--- a/Engine/Source/Runtime/Engine/Private/AsyncCompilationHelpers.cpp
+++ b/Engine/Source/Runtime/Engine/Private/AsyncCompilationHelpers.cpp
@@ -16,6 +16,7 @@
 #include "HAL/PlatformStackWalk.h"
 #include "Hash/CityHash.h"
 #include "Misc/ScopeRWLock.h"
+#include "ProfilingDebugging/StallDetector.h"
 
 #define LOCTEXT_NAMESPACE "AsyncCompilation"
 
@@ -91,6 +92,7 @@ namespace AsyncCompilationHelpers
 		TFunctionRef<void(ICompilable*)> PostCompileSingle
 	)
 	{
+		SCOPE_STALL_REPORTER(AsyncCompilationHelpers::FinishCompilation, 10.0);
 		if (Num <= 0)
 		{
 			return;
diff --git a/Engine/Source/Runtime/Engine/Private/ShaderCompiler/ShaderCompiler.cpp b/Engine/Source/Runtime/Engine/Private/ShaderCompiler/ShaderCompiler.cpp
--- a/Engine/Source/Runtime/Engine/Private/ShaderCompiler/ShaderCompiler.cpp
+++ b/Engine/Source/Runtime/Engine/Private/ShaderCompiler/ShaderCompiler.cpp
@@ -55,6 +55,7 @@
 #include "DistributedBuildInterface/Public/DistributedBuildControllerInterface.h"
 #include "Misc/ScopeRWLock.h"
 #include "ObjectCacheContext.h"
+#include "ProfilingDebugging/StallDetector.h"
 
 #if WITH_EDITOR
 #include "TextureCompiler.h"
@@ -3009,6 +3010,8 @@ FProcHandle FShaderCompilingManager::LaunchWorker(const FString& WorkingDirector
 /** Flushes all pending jobs for the given shader maps. */
 void FShaderCompilingManager::BlockOnShaderMapCompletion(const TArray<int32>& ShaderMapIdsToFinishCompiling, TMap<int32, FShaderMapFinalizeResults>& CompiledShaderMaps)
 {
+	SCOPE_STALL_REPORTER(FShaderCompilingManager::BlockOnShaderMapCompletion, 2.0);
+
 	COOK_STAT(FScopedDurationTimer BlockingTimer(ShaderCompilerCookStats::BlockingTimeSec));
 	if (bAllowAsynchronousShaderCompiling)
 	{
@@ -3119,6 +3122,8 @@ void FShaderCompilingManager::BlockOnShaderMapCompletion(const TArray<int32>& Sh
 
 void FShaderCompilingManager::BlockOnAllShaderMapCompletion(TMap<int32, FShaderMapFinalizeResults>& CompiledShaderMaps)
 {
+	SCOPE_STALL_REPORTER(FShaderCompilingManager::BlockOnAllShaderMapCompletion, 2.0);
+
 	COOK_STAT(FScopedDurationTimer BlockingTimer(ShaderCompilerCookStats::BlockingTimeSec));
 	if (bAllowAsynchronousShaderCompiling)
 	{
