Added option to load last loaded editor cells when opening a partitioned world
- Last loaded cells list will be invalidated if the editor grid topology changes (it shouldn't).
- Can be disabled in the project setting with the bEnableLoadingOfLastLoadedCells property.

diff --git a/Engine/Source/Editor/WorldPartitionEditor/Private/WorldPartitionEditorModule.cpp b/Engine/Source/Editor/WorldPartitionEditor/Private/WorldPartitionEditorModule.cpp
--- a/Engine/Source/Editor/WorldPartitionEditor/Private/WorldPartitionEditorModule.cpp
+++ b/Engine/Source/Editor/WorldPartitionEditor/Private/WorldPartitionEditorModule.cpp
@@ -142,6 +142,11 @@ void FWorldPartitionEditorModule::SetConversionPromptEnabled(bool bEnabled)
 	GetMutableDefault<UWorldPartitionEditorSettings>()->bEnableConversionPrompt = bEnabled;
 }
 
+bool FWorldPartitionEditorModule::GetEnableLoadingOfLastLoadedCells() const
+{
+	return GetDefault<UWorldPartitionEditorSettings>()->bEnableLoadingOfLastLoadedCells;
+}
+	
 float FWorldPartitionEditorModule::GetAutoCellLoadingMaxWorldSize() const
 {
 	return GetDefault<UWorldPartitionEditorSettings>()->AutoCellLoadingMaxWorldSize;
diff --git a/Engine/Source/Editor/WorldPartitionEditor/Public/WorldPartition/WorldPartitionEditorSettings.h b/Engine/Source/Editor/WorldPartitionEditor/Public/WorldPartition/WorldPartitionEditorSettings.h
--- a/Engine/Source/Editor/WorldPartitionEditor/Public/WorldPartition/WorldPartitionEditorSettings.h
+++ b/Engine/Source/Editor/WorldPartitionEditor/Public/WorldPartition/WorldPartitionEditorSettings.h
@@ -25,6 +25,9 @@ public:
 	UPROPERTY(config, EditAnywhere, Category = MapConversion, meta = (ToolTip = "This option when enabled will show a conversion prompt when opening non World Partition maps"))
 	bool bEnableConversionPrompt = false;
 
+	UPROPERTY(config, EditAnywhere, Category = WorldPartition, meta = (ToolTip = "This option is used to enable the loading of the last loaded cells in the editor"))
+	bool bEnableLoadingOfLastLoadedCells = true;
+
 	UPROPERTY(config, EditAnywhere, Category = MapConversion, meta = (ToolTip = "Commandlet class to use for World Parition conversion"))
 	TSubclassOf<UWorldPartitionConvertCommandlet> CommandletClass;
 
diff --git a/Engine/Source/Editor/WorldPartitionEditor/Public/WorldPartitionEditorModule.h b/Engine/Source/Editor/WorldPartitionEditor/Public/WorldPartitionEditorModule.h
--- a/Engine/Source/Editor/WorldPartitionEditor/Public/WorldPartitionEditorModule.h
+++ b/Engine/Source/Editor/WorldPartitionEditor/Public/WorldPartitionEditorModule.h
@@ -40,6 +40,11 @@ public:
 	 */
 	virtual void SetConversionPromptEnabled(bool bEnabled) override;
 
+	/**
+	 *
+	 */
+	virtual bool GetEnableLoadingOfLastLoadedCells() const override;
+
 	/**
 	 *
 	 */
diff --git a/Engine/Source/Runtime/Engine/Private/WorldPartition/WorldPartition.cpp b/Engine/Source/Runtime/Engine/Private/WorldPartition/WorldPartition.cpp
--- a/Engine/Source/Runtime/Engine/Private/WorldPartition/WorldPartition.cpp
+++ b/Engine/Source/Runtime/Engine/Private/WorldPartition/WorldPartition.cpp
@@ -31,6 +31,7 @@
 #include "WorldPartition/WorldPartitionEditorSpatialHash.h"
 #include "WorldPartition/WorldPartitionRuntimeHash.h"
 #include "WorldPartition/WorldPartitionRuntimeSpatialHash.h"
+#include "WorldPartition/WorldPartitionEditorPerProjectUserSettings.h"
 #include "Modules/ModuleManager.h"
 #endif //WITH_EDITOR
 
@@ -84,6 +85,17 @@ struct FWorldPartionCellUpdateContext
 				{
 					WorldPartition->WorldPartitionEditor->Refresh();
 				}
+
+				TArray<FVector> EditorGridLastLoadedCells;
+				WorldPartition->EditorHash->ForEachCell([this, &EditorGridLastLoadedCells](UWorldPartitionEditorCell* Cell)
+				{
+					if ((Cell != WorldPartition->EditorHash->GetAlwaysLoadedCell()) && Cell->bLoaded)
+					{
+						EditorGridLastLoadedCells.Add(Cell->Bounds.GetCenter());
+					}
+				});
+
+				GetMutableDefault<UWorldPartitionEditorPerProjectUserSettings>()->SetEditorGridLastLoadedCells(WorldPartition->GetWorld(), EditorGridLastLoadedCells);
 			}
 		}
 	}
@@ -258,6 +270,20 @@ void UWorldPartition::Initialize(UWorld* InWorld, const FTransform& InTransform)
 					UpdateLoadingEditorCell(Cell, true);
 				});
 			}
+
+			// Load last loaded cells
+			if (WorldPartitionEditorModule.GetEnableLoadingOfLastLoadedCells())
+			{
+				const TArray<FVector>& EditorGridLastLoadedCells = GetMutableDefault<UWorldPartitionEditorPerProjectUserSettings>()->GetEditorGridLastLoadedCells(InWorld);
+
+				for (const FVector& EditorGridLastLoadedCellCenter : EditorGridLastLoadedCells)
+				{
+					EditorHash->ForEachIntersectingCell(FBox(&EditorGridLastLoadedCellCenter, 1), [this](UWorldPartitionEditorCell* Cell)
+					{
+						UpdateLoadingEditorCell(Cell, true);
+					});
+				}
+			}
 		}
 	}
 #endif //WITH_EDITOR
@@ -473,6 +499,8 @@ bool UWorldPartition::RefreshLoadedEditorCells()
 
 void UWorldPartition::UnloadEditorCells(const FBox& Box)
 {
+	FWorldPartionCellUpdateContext CellUpdateContext(this);
+
 	auto GetCellsToUnload = [this, Box](TArray<UWorldPartitionEditorCell*>& CellsToUnload)
 	{
 		return EditorHash->GetIntersectingCells(Box, CellsToUnload) > 0;
diff --git a/Engine/Source/Runtime/Engine/Private/WorldPartition/WorldPartitionEditorSpatialHash.cpp b/Engine/Source/Runtime/Engine/Private/WorldPartition/WorldPartitionEditorSpatialHash.cpp
--- a/Engine/Source/Runtime/Engine/Private/WorldPartition/WorldPartitionEditorSpatialHash.cpp
+++ b/Engine/Source/Runtime/Engine/Private/WorldPartition/WorldPartitionEditorSpatialHash.cpp
@@ -5,6 +5,7 @@
 #include "WorldPartition/WorldPartitionEditorCell.h"
 #include "WorldPartition/WorldPartitionActorDesc.h"
 #include "WorldPartition/WorldPartitionActorDesc.h"
+#include "WorldPartition/WorldPartitionEditorPerProjectUserSettings.h"
 
 #define LOCTEXT_NAMESPACE "WorldPartitionEditor"
 
@@ -27,6 +28,10 @@ void UWorldPartitionEditorSpatialHash::Initialize()
 
 	AlwaysLoadedCell = NewObject<UWorldPartitionEditorCell>(this, TEXT("AlwaysLoadedCell"), RF_Transient);
 	AlwaysLoadedCell->Bounds.Init();
+
+	FHashBuilder EditorGridConfigHashBuilder;
+	EditorGridConfigHashBuilder << CellSize;
+	GetMutableDefault<UWorldPartitionEditorPerProjectUserSettings>()->SetEditorGridConfigHash(EditorGridConfigHashBuilder.GetHash());
 }
 
 void UWorldPartitionEditorSpatialHash::SetDefaultValues()
diff --git a/Engine/Source/Runtime/Engine/Public/WorldPartition/IWorldPartitionEditorModule.h b/Engine/Source/Runtime/Engine/Public/WorldPartition/IWorldPartitionEditorModule.h
--- a/Engine/Source/Runtime/Engine/Public/WorldPartition/IWorldPartitionEditorModule.h
+++ b/Engine/Source/Runtime/Engine/Public/WorldPartition/IWorldPartitionEditorModule.h
@@ -15,6 +15,7 @@ public:
 	virtual bool IsWorldPartitionEnabled() const = 0;
 	virtual bool IsConversionPromptEnabled() const = 0;
 	virtual void SetConversionPromptEnabled(bool bEnabled) = 0;
+	virtual bool GetEnableLoadingOfLastLoadedCells() const = 0;
 	virtual float GetAutoCellLoadingMaxWorldSize() const = 0;
 	virtual bool ConvertMap(const FString& InLongPackageName) = 0;
 };
\ No newline at end of file
diff --git a/Engine/Source/Runtime/Engine/Public/WorldPartition/WorldPartitionEditorPerProjectUserSettings.h b/Engine/Source/Runtime/Engine/Public/WorldPartition/WorldPartitionEditorPerProjectUserSettings.h
new file mode 100644
--- /dev/null
+++ b/Engine/Source/Runtime/Engine/Public/WorldPartition/WorldPartitionEditorPerProjectUserSettings.h
@@ -0,0 +1,70 @@
+// Copyright Epic Games, Inc. All Rights Reserved.
+
+#pragma once
+
+#include "CoreMinimal.h"
+#include "WorldPartitionEditorPerProjectUserSettings.generated.h"
+
+USTRUCT()
+struct FWorldPartitionPerWorldEditorGridSettings
+{
+	GENERATED_BODY();
+
+#if WITH_EDITOR
+	FWorldPartitionPerWorldEditorGridSettings()
+	{}
+
+	FWorldPartitionPerWorldEditorGridSettings(TArray<FVector>& InLastLoadedCells)
+		: LastLoadedCells(InLastLoadedCells)
+	{}
+#endif
+
+#if WITH_EDITORONLY_DATA
+	UPROPERTY()
+	TArray<FVector> LastLoadedCells;
+#endif
+};
+
+UCLASS(MinimalAPI, config = EditorPerProjectUserSettings)
+class UWorldPartitionEditorPerProjectUserSettings : public UObject
+{
+	GENERATED_BODY()
+
+public:
+#if WITH_EDITOR
+	uint32 GetEditorGridConfigHash() const
+	{ 
+		return EditorGridConfigHash;
+	}
+
+	void SetEditorGridConfigHash(uint32 InEditorGridConfigHash)
+	{ 
+		if (EditorGridConfigHash != InEditorGridConfigHash)
+		{
+			EditorGridConfigHash = InEditorGridConfigHash;
+			PerWorldEditorGridSettings.Empty();
+			SaveConfig();
+		}
+	}
+
+	const TArray<FVector>& GetEditorGridLastLoadedCells(UWorld* InWorld)
+	{ 
+		return PerWorldEditorGridSettings.FindOrAdd(TSoftObjectPtr<UWorld>(InWorld)).LastLoadedCells;
+	}
+
+	void SetEditorGridLastLoadedCells(UWorld* InWorld, TArray<FVector>& InEditorGridLastLoadedCells)
+	{
+		PerWorldEditorGridSettings.Add(TSoftObjectPtr<UWorld>(InWorld), InEditorGridLastLoadedCells);
+		SaveConfig();
+	}
+#endif
+
+private:
+#if WITH_EDITORONLY_DATA
+	UPROPERTY(config)
+	uint32 EditorGridConfigHash;
+
+	UPROPERTY(config)
+	TMap<TSoftObjectPtr<UWorld>, FWorldPartitionPerWorldEditorGridSettings> PerWorldEditorGridSettings;
+#endif
+};
\ No newline at end of file
