Added per-world user settings in WorldPartitionEditorPerProjectUserSettings.
Merged DataLayerEditorPerProjectUserSettings into WorldPartitionEditorPerProjectUserSettings.

diff --git a/Engine/Source/Editor/DataLayerEditor/Private/DataLayer/DataLayerEditorSubsystem.cpp b/Engine/Source/Editor/DataLayerEditor/Private/DataLayer/DataLayerEditorSubsystem.cpp
--- a/Engine/Source/Editor/DataLayerEditor/Private/DataLayer/DataLayerEditorSubsystem.cpp
+++ b/Engine/Source/Editor/DataLayerEditor/Private/DataLayer/DataLayerEditorSubsystem.cpp
@@ -6,7 +6,7 @@
 #include "WorldPartition/DataLayer/WorldDataLayers.h"
 #include "WorldPartition/WorldPartitionSubsystem.h"
 #include "WorldPartition/WorldPartition.h"
-#include "WorldPartition/DataLayer/DataLayerEditorPerProjectUserSettings.h"
+#include "WorldPartition/WorldPartitionEditorPerProjectUserSettings.h"
 #include "Modules/ModuleManager.h"
 #include "DataLayerEditorModule.h"
 #include "ActorEditorUtils.h"
@@ -195,7 +195,7 @@ void UDataLayerEditorSubsystem::UpdateDataLayerEditorPerProjectUserSettings()
 			}
 			return true;
 		});
-		GetMutableDefault<UDataLayerEditorPerProjectUserSettings>()->SetWorldDataLayersNotLoadedInEditor(GetWorld(), DataLayersNotLoadedInEditor);
+		GetMutableDefault<UWorldPartitionEditorPerProjectUserSettings>()->SetWorldDataLayersNotLoadedInEditor(GetWorld(), DataLayersNotLoadedInEditor);
 	}
 }
 
diff --git a/Engine/Source/Editor/DataLayerEditor/Private/DataLayer/SDataLayerBrowser.cpp b/Engine/Source/Editor/DataLayerEditor/Private/DataLayer/SDataLayerBrowser.cpp
--- a/Engine/Source/Editor/DataLayerEditor/Private/DataLayer/SDataLayerBrowser.cpp
+++ b/Engine/Source/Editor/DataLayerEditor/Private/DataLayer/SDataLayerBrowser.cpp
@@ -7,7 +7,7 @@
 #include "DataLayerActorTreeItem.h"
 #include "DataLayerTreeItem.h"
 #include "DataLayerMode.h"
-#include "WorldPartition/DataLayer/DataLayerEditorPerProjectUserSettings.h"
+#include "WorldPartition/WorldPartitionEditorPerProjectUserSettings.h"
 #include "WorldPartition/DataLayer/DataLayer.h"
 #include "DataLayerOutlinerIsDynamicallyLoadedColumn.h"
 #include "DataLayerOutlinerDeleteButtonColumn.h"
@@ -22,7 +22,7 @@
 
 void SDataLayerBrowser::Construct(const FArguments& InArgs)
 {
-	Mode = GetDefault<UDataLayerEditorPerProjectUserSettings>()->GetShowDataLayerContent() ? EDataLayerBrowserMode::DataLayerContents : EDataLayerBrowserMode::DataLayers;
+	Mode = GetDefault<UWorldPartitionEditorPerProjectUserSettings>()->GetShowDataLayerContent() ? EDataLayerBrowserMode::DataLayerContents : EDataLayerBrowserMode::DataLayers;
 
 	auto ToggleDataLayerContents = [this]()
 	{
@@ -195,7 +195,7 @@ void SDataLayerBrowser::SetupDataLayerMode(EDataLayerBrowserMode InNewMode)
 
 	Mode = InNewMode;
 	
-	GetMutableDefault<UDataLayerEditorPerProjectUserSettings>()->SetShowDataLayerContent(Mode == EDataLayerBrowserMode::DataLayerContents);
+	GetMutableDefault<UWorldPartitionEditorPerProjectUserSettings>()->SetShowDataLayerContent(Mode == EDataLayerBrowserMode::DataLayerContents);
 
 	ModeChanged.Broadcast(Mode);
 }
diff --git a/Engine/Source/Runtime/Engine/Private/WorldPartition/DataLayer/WorldDataLayers.cpp b/Engine/Source/Runtime/Engine/Private/WorldPartition/DataLayer/WorldDataLayers.cpp
--- a/Engine/Source/Runtime/Engine/Private/WorldPartition/DataLayer/WorldDataLayers.cpp
+++ b/Engine/Source/Runtime/Engine/Private/WorldPartition/DataLayer/WorldDataLayers.cpp
@@ -9,7 +9,7 @@
 #include "WorldPartition/DataLayer/DataLayerSubsystem.h"
 #include "EngineUtils.h"
 #if WITH_EDITOR
-#include "WorldPartition/DataLayer/DataLayerEditorPerProjectUserSettings.h"
+#include "WorldPartition/WorldPartitionEditorPerProjectUserSettings.h"
 #endif
 
 #define LOCTEXT_NAMESPACE "WorldDataLayers"
@@ -232,14 +232,12 @@ void AWorldDataLayers::PostLoad()
 
 #if WITH_EDITOR
 	// Initialize DataLayer's IsDynamicallyLoadedInEditor based on DataLayerEditorPerProjectUserSettings
-	if (const FDataLayerNames* SettingsDataLayersNotLoadedInEditor = GetDefault<UDataLayerEditorPerProjectUserSettings>()->GetWorldDataLayersNotLoadedInEditor(GetWorld()))
+	const TArray<FName>& SettingsDataLayersNotLoadedInEditor = GetMutableDefault<UWorldPartitionEditorPerProjectUserSettings>()->GetWorldDataLayersNotLoadedInEditor(GetWorld());
+	for (const FName& DataLayerName : SettingsDataLayersNotLoadedInEditor)
 	{
-		for (const FName& DataLayerName : SettingsDataLayersNotLoadedInEditor->DataLayers)
+		if (UDataLayer* DataLayer = const_cast<UDataLayer*>(GetDataLayerFromName(DataLayerName)))
 		{
-			if (UDataLayer* DataLayer = const_cast<UDataLayer*>(GetDataLayerFromName(DataLayerName)))
-			{
-				DataLayer->SetIsDynamicallyLoadedInEditor(false);
-			}
+			DataLayer->SetIsDynamicallyLoadedInEditor(false);
 		}
 	}
 #else
diff --git a/Engine/Source/Runtime/Engine/Private/WorldPartition/WorldPartition.cpp b/Engine/Source/Runtime/Engine/Private/WorldPartition/WorldPartition.cpp
--- a/Engine/Source/Runtime/Engine/Private/WorldPartition/WorldPartition.cpp
+++ b/Engine/Source/Runtime/Engine/Private/WorldPartition/WorldPartition.cpp
@@ -95,7 +95,7 @@ struct FWorldPartionCellUpdateContext
 					}
 				});
 
-				GetMutableDefault<UWorldPartitionEditorPerProjectUserSettings>()->SetEditorGridLastLoadedCells(WorldPartition->GetWorld(), EditorGridLastLoadedCells);
+				GetMutableDefault<UWorldPartitionEditorPerProjectUserSettings>()->SetEditorGridLoadedCells(WorldPartition->GetWorld(), EditorGridLastLoadedCells);
 			}
 		}
 	}
@@ -274,7 +274,7 @@ void UWorldPartition::Initialize(UWorld* InWorld, const FTransform& InTransform)
 			// Load last loaded cells
 			if (WorldPartitionEditorModule.GetEnableLoadingOfLastLoadedCells())
 			{
-				const TArray<FVector>& EditorGridLastLoadedCells = GetMutableDefault<UWorldPartitionEditorPerProjectUserSettings>()->GetEditorGridLastLoadedCells(InWorld);
+				const TArray<FVector>& EditorGridLastLoadedCells = GetMutableDefault<UWorldPartitionEditorPerProjectUserSettings>()->GetEditorGridLoadedCells(InWorld);
 
 				for (const FVector& EditorGridLastLoadedCellCenter : EditorGridLastLoadedCells)
 				{
diff --git a/Engine/Source/Runtime/Engine/Public/WorldPartition/DataLayer/DataLayerEditorPerProjectUserSettings.h b/Engine/Source/Runtime/Engine/Public/WorldPartition/DataLayer/DataLayerEditorPerProjectUserSettings.h
deleted file mode 100644
diff --git a/Engine/Source/Runtime/Engine/Public/WorldPartition/WorldPartitionEditorPerProjectUserSettings.h b/Engine/Source/Runtime/Engine/Public/WorldPartition/WorldPartitionEditorPerProjectUserSettings.h
--- a/Engine/Source/Runtime/Engine/Public/WorldPartition/WorldPartitionEditorPerProjectUserSettings.h
+++ b/Engine/Source/Runtime/Engine/Public/WorldPartition/WorldPartitionEditorPerProjectUserSettings.h
@@ -3,25 +3,32 @@
 #pragma once
 
 #include "CoreMinimal.h"
+#include "UObject/ObjectMacros.h"
+#include "UObject/Object.h"
+#include "UObject/SoftObjectPtr.h"
+#include "Engine/World.h"
 #include "WorldPartitionEditorPerProjectUserSettings.generated.h"
 
 USTRUCT()
-struct FWorldPartitionPerWorldEditorGridSettings
+struct FWorldPartitionPerWorldSettings
 {
 	GENERATED_BODY();
 
 #if WITH_EDITOR
-	FWorldPartitionPerWorldEditorGridSettings()
+	FWorldPartitionPerWorldSettings()
 	{}
 
-	FWorldPartitionPerWorldEditorGridSettings(TArray<FVector>& InLastLoadedCells)
-		: LastLoadedCells(InLastLoadedCells)
+	FWorldPartitionPerWorldSettings(TArray<FVector>& InLoadedEditorGridCells)
+		: LoadedEditorGridCells(InLoadedEditorGridCells)
 	{}
 #endif
 
 #if WITH_EDITORONLY_DATA
 	UPROPERTY()
-	TArray<FVector> LastLoadedCells;
+	TArray<FVector> LoadedEditorGridCells;
+
+	UPROPERTY()
+	TArray<FName> NotLoadedDataLayers;
 #endif
 };
 
@@ -42,21 +49,55 @@ public:
 		if (EditorGridConfigHash != InEditorGridConfigHash)
 		{
 			EditorGridConfigHash = InEditorGridConfigHash;
-			PerWorldEditorGridSettings.Empty();
+
+			for (auto& PerWorldEditorSetting : PerWorldEditorSettings)
+			{
+				PerWorldEditorSetting.Value.LoadedEditorGridCells.Empty();
+			}
 			SaveConfig();
 		}
 	}
 
-	const TArray<FVector>& GetEditorGridLastLoadedCells(UWorld* InWorld)
+	const TArray<FVector>& GetEditorGridLoadedCells(UWorld* InWorld)
 	{ 
-		return PerWorldEditorGridSettings.FindOrAdd(TSoftObjectPtr<UWorld>(InWorld)).LastLoadedCells;
+		return PerWorldEditorSettings.FindOrAdd(TSoftObjectPtr<UWorld>(InWorld)).LoadedEditorGridCells;
 	}
 
-	void SetEditorGridLastLoadedCells(UWorld* InWorld, TArray<FVector>& InEditorGridLastLoadedCells)
+	void SetEditorGridLoadedCells(UWorld* InWorld, TArray<FVector>& InEditorGridLoadedCells)
 	{
-		PerWorldEditorGridSettings.Add(TSoftObjectPtr<UWorld>(InWorld), InEditorGridLastLoadedCells);
+		TArray<FVector>& GridLoadedCells = PerWorldEditorSettings.FindOrAdd(TSoftObjectPtr<UWorld>(InWorld)).LoadedEditorGridCells;
+		GridLoadedCells = InEditorGridLoadedCells;
 		SaveConfig();
 	}
+
+	bool GetShowDataLayerContent() const
+	{
+		return bShowDataLayerContent;
+	}
+
+	void SetShowDataLayerContent(bool bInShowDataLayerContent)
+	{
+		if (bShowDataLayerContent != bInShowDataLayerContent)
+		{
+			bShowDataLayerContent = bInShowDataLayerContent;
+			SaveConfig();
+		}
+	}
+
+	const TArray<FName>& GetWorldDataLayersNotLoadedInEditor(UWorld* InWorld)
+	{
+		return PerWorldEditorSettings.FindOrAdd(TSoftObjectPtr<UWorld>(InWorld)).NotLoadedDataLayers;
+	}
+
+	void SetWorldDataLayersNotLoadedInEditor(UWorld* InWorld, const TArray<FName>& InDataLayersLoadedInEditor)
+	{
+		if (InWorld)
+		{
+			TArray<FName>& DataLayersNotLoadedInEditor = PerWorldEditorSettings.FindOrAdd(TSoftObjectPtr<UWorld>(InWorld)).NotLoadedDataLayers;
+			DataLayersNotLoadedInEditor = InDataLayersLoadedInEditor;
+			SaveConfig();
+		}
+	}
 #endif
 
 private:
@@ -65,6 +106,9 @@ private:
 	uint32 EditorGridConfigHash;
 
 	UPROPERTY(config)
-	TMap<TSoftObjectPtr<UWorld>, FWorldPartitionPerWorldEditorGridSettings> PerWorldEditorGridSettings;
+	uint32 bShowDataLayerContent : 1;
+
+	UPROPERTY(config)
+	TMap<TSoftObjectPtr<UWorld>, FWorldPartitionPerWorldSettings> PerWorldEditorSettings;
 #endif
 };
\ No newline at end of file
