Reworks BlendSpace to make UBlendSpace the main runtime class, which is able to handle 1 and 2 dimensions.

diff --git a/Engine/Config/BaseEditor.ini b/Engine/Config/BaseEditor.ini
--- a/Engine/Config/BaseEditor.ini
+++ b/Engine/Config/BaseEditor.ini
@@ -46,7 +46,7 @@ ColorFontAtlasSize=512
 +RenderableThumbnailTypes=(ClassNeedingThumbnailName="/Script/Engine.StaticMesh",RendererClassName="/Script/UnrealEd.StaticMeshThumbnailRenderer")
 +RenderableThumbnailTypes=(ClassNeedingThumbnailName="/Script/Engine.SkeletalMesh",RendererClassName="/Script/UnrealEd.SkeletalMeshThumbnailRenderer")
 +RenderableThumbnailTypes=(ClassNeedingThumbnailName="/Script/Engine.AnimSequenceBase",RendererClassName="/Script/UnrealEd.AnimSequenceThumbnailRenderer")
-+RenderableThumbnailTypes=(ClassNeedingThumbnailName="/Script/Engine.BlendSpaceBase",RendererClassName="/Script/UnrealEd.BlendSpaceThumbnailRenderer")
++RenderableThumbnailTypes=(ClassNeedingThumbnailName="/Script/Engine.BlendSpace",RendererClassName="/Script/UnrealEd.BlendSpaceThumbnailRenderer")
 +RenderableThumbnailTypes=(ClassNeedingThumbnailName="/Script/Engine.AnimBlueprint",RendererClassName="/Script/UnrealEd.AnimBlueprintThumbnailRenderer")
 +RenderableThumbnailTypes=(ClassNeedingThumbnailName="/Script/Engine.Level",RendererClassName="/Script/UnrealEd.LevelThumbnailRenderer")
 +RenderableThumbnailTypes=(ClassNeedingThumbnailName="/Script/Engine.World",RendererClassName="/Script/UnrealEd.WorldThumbnailRenderer")
diff --git a/Engine/Plugins/Animation/GameplayInsights/Source/GameplayInsights/Private/AnimNodesTrack.cpp b/Engine/Plugins/Animation/GameplayInsights/Source/GameplayInsights/Private/AnimNodesTrack.cpp
--- a/Engine/Plugins/Animation/GameplayInsights/Source/GameplayInsights/Private/AnimNodesTrack.cpp
+++ b/Engine/Plugins/Animation/GameplayInsights/Source/GameplayInsights/Private/AnimNodesTrack.cpp
@@ -27,7 +27,7 @@
 #include "Animation/AnimBlueprint.h"
 #include "Subsystems/AssetEditorSubsystem.h"
 #include "Animation/AnimInstance.h"
-#include "Animation/BlendSpaceBase.h"
+#include "Animation/BlendSpace.h"
 #endif
 
 #define LOCTEXT_NAMESPACE "AnimNodesTrack"
@@ -362,14 +362,14 @@ void FAnimNodesTrack::UpdateDebugData(const TraceServices::FFrame& InFrame)
 									{
 										InBlendSpacePlayersTimeline.EnumerateEvents(InGraphStartTime, InGraphEndTime, [GameplayProvider, &DebugData](double InStartTime, double InEndTime, uint32 InDepth, const FBlendSpacePlayerMessage& InMessage)
 										{
-											UBlendSpaceBase* BlendSpaceBase = nullptr;
+											UBlendSpace* BlendSpace = nullptr;
 											const FObjectInfo* BlendSpaceInfo = GameplayProvider->FindObjectInfo(InMessage.BlendSpaceId);
 											if(BlendSpaceInfo)
 											{
-												BlendSpaceBase = TSoftObjectPtr<UBlendSpaceBase>(FSoftObjectPath(BlendSpaceInfo->PathName)).LoadSynchronous();
+												BlendSpace = TSoftObjectPtr<UBlendSpace>(FSoftObjectPath(BlendSpaceInfo->PathName)).LoadSynchronous();
 											}
 
-											DebugData.RecordBlendSpacePlayer(InMessage.NodeId, BlendSpaceBase, FVector(InMessage.PositionX, InMessage.PositionY, InMessage.PositionZ), FVector(InMessage.FilteredPositionX, InMessage.FilteredPositionY, InMessage.FilteredPositionZ));
+											DebugData.RecordBlendSpacePlayer(InMessage.NodeId, BlendSpace, FVector(InMessage.PositionX, InMessage.PositionY, InMessage.PositionZ), FVector(InMessage.FilteredPositionX, InMessage.FilteredPositionY, InMessage.FilteredPositionZ));
 											return TraceServices::EEventEnumerate::Continue;
 										});
 									});
diff --git a/Engine/Plugins/Animation/GameplayInsights/Source/GameplayInsights/Private/SAnimGraphSchematicView.cpp b/Engine/Plugins/Animation/GameplayInsights/Source/GameplayInsights/Private/SAnimGraphSchematicView.cpp
--- a/Engine/Plugins/Animation/GameplayInsights/Source/GameplayInsights/Private/SAnimGraphSchematicView.cpp
+++ b/Engine/Plugins/Animation/GameplayInsights/Source/GameplayInsights/Private/SAnimGraphSchematicView.cpp
@@ -688,14 +688,14 @@ void SAnimGraphSchematicView::RefreshNodes()
 // 						{
 // 							InBlendSpacePlayersTimeline.EnumerateEvents(InGraphStartTime, InGraphEndTime, [GameplayProvider](double InStartTime, double InEndTime, uint32 InDepth, const FBlendSpacePlayerMessage& InMessage)
 // 							{
-// // 								UBlendSpaceBase* BlendSpaceBase = nullptr;
+// // 								UBlendSpace* BlendSpace = nullptr;
 // // 								const FObjectInfo* BlendSpaceInfo = GameplayProvider->FindObjectInfo(InMessage.BlendSpaceId);
 // // 								if(BlendSpaceInfo)
 // // 								{
-// // 									BlendSpaceBase = TSoftObjectPtr<UBlendSpaceBase>(FSoftObjectPath(BlendSpaceInfo->PathName)).LoadSynchronous();
+// // 									BlendSpace = TSoftObjectPtr<UBlendSpace>(FSoftObjectPath(BlendSpaceInfo->PathName)).LoadSynchronous();
 // // 								}
 // // 
-// // 								DebugData.RecordBlendSpacePlayer(InMessage.NodeId, BlendSpaceBase, InMessage.PositionX, InMessage.PositionY, InMessage.PositionZ);
+// // 								DebugData.RecordBlendSpacePlayer(InMessage.NodeId, BlendSpace, InMessage.PositionX, InMessage.PositionY, InMessage.PositionZ);
 // 							});
 // 						});
 
diff --git a/Engine/Source/Editor/AnimGraph/Private/AnimBlueprintGeneratedClassCompiledData.cpp b/Engine/Source/Editor/AnimGraph/Private/AnimBlueprintGeneratedClassCompiledData.cpp
--- a/Engine/Source/Editor/AnimGraph/Private/AnimBlueprintGeneratedClassCompiledData.cpp
+++ b/Engine/Source/Editor/AnimGraph/Private/AnimBlueprintGeneratedClassCompiledData.cpp
@@ -71,9 +71,9 @@ TMap<FName, FGraphAssetPlayerInformation>& FAnimBlueprintGeneratedClassCompiledD
 	return Class->GraphAssetPlayerInformation;
 }
 
-UBlendSpaceBase* FAnimBlueprintGeneratedClassCompiledData::AddBlendSpace(UBlendSpaceBase* InSourceBlendSpace)
+UBlendSpace* FAnimBlueprintGeneratedClassCompiledData::AddBlendSpace(UBlendSpace* InSourceBlendSpace)
 {
-	UBlendSpaceBase* CopiedBlendSpace = DuplicateObject(InSourceBlendSpace, Class);
+	UBlendSpace* CopiedBlendSpace = DuplicateObject(InSourceBlendSpace, Class);
 	Class->BlendSpaces.Add(CopiedBlendSpace);
 	return CopiedBlendSpace;
 }
\ No newline at end of file
diff --git a/Engine/Source/Editor/AnimGraph/Private/AnimBlueprintGeneratedClassCompiledData.h b/Engine/Source/Editor/AnimGraph/Private/AnimBlueprintGeneratedClassCompiledData.h
--- a/Engine/Source/Editor/AnimGraph/Private/AnimBlueprintGeneratedClassCompiledData.h
+++ b/Engine/Source/Editor/AnimGraph/Private/AnimBlueprintGeneratedClassCompiledData.h
@@ -24,7 +24,7 @@ private:
 	virtual FPropertyAccessLibrary& GetPropertyAccessLibrary() const override;
 	virtual FAnimBlueprintDebugData& GetAnimBlueprintDebugData() const override;
 	virtual TMap<FName, FGraphAssetPlayerInformation>& GetGraphAssetPlayerInformation() const override;
-	virtual UBlendSpaceBase* AddBlendSpace(UBlendSpaceBase* InSourceBlendSpace) override;
+	virtual UBlendSpace* AddBlendSpace(UBlendSpace* InSourceBlendSpace) override;
 
 private:
 	// The class we are wrapping
diff --git a/Engine/Source/Editor/AnimGraph/Private/AnimBlueprintPostCompileValidation.cpp b/Engine/Source/Editor/AnimGraph/Private/AnimBlueprintPostCompileValidation.cpp
--- a/Engine/Source/Editor/AnimGraph/Private/AnimBlueprintPostCompileValidation.cpp
+++ b/Engine/Source/Editor/AnimGraph/Private/AnimBlueprintPostCompileValidation.cpp
@@ -2,7 +2,7 @@
 
 #include "AnimBlueprintPostCompileValidation.h" 
 #include "AnimNodes/AnimNode_BlendSpacePlayer.h"
-#include "Animation/BlendSpaceBase.h"
+#include "Animation/BlendSpace.h"
 #include "Animation/AnimNode_SequencePlayer.h"
 #include "Animation/AnimInstance.h"
 #include "AnimGraphNode_Base.h"
@@ -95,7 +95,7 @@ void UAnimBlueprintPostCompileValidation::PCV_GatherAnimSequencesFromProperty(TA
 				}
 				else if (const FObjectProperty* ParentObjectProperty = CastField<FObjectProperty>(Parent.Property))
 				{
-					if (ParentObjectProperty->PropertyClass && ParentObjectProperty->PropertyClass->IsChildOf(UBlendSpaceBase::StaticClass()))
+					if (ParentObjectProperty->PropertyClass && ParentObjectProperty->PropertyClass->IsChildOf(UBlendSpace::StaticClass()))
 					{
 						Referencer = ParentObjectProperty->GetObjectPropertyValue(Parent.Value);
 						break;
@@ -105,7 +105,7 @@ void UAnimBlueprintPostCompileValidation::PCV_GatherAnimSequencesFromProperty(TA
 
 			OutRefAnimSequences.Add(FPCV_ReferencedAnimSequence(AnimSequence, Referencer));
 		}
-		else if (const UBlendSpaceBase* const BlendSpace = Cast<UBlendSpaceBase>(ObjectPropertyValue))
+		else if (const UBlendSpace* const BlendSpace = Cast<UBlendSpace>(ObjectPropertyValue))
 		{
 			PCV_PreloadObject(BlendSpace);
 
@@ -123,7 +123,7 @@ void UAnimBlueprintPostCompileValidation::PCV_GatherAnimSequences(TArray<const U
 	}
 }
 
-void UAnimBlueprintPostCompileValidation::PCV_GatherAnimSequences(TArray<const UAnimSequence*>& OutAnimSequences, const class UBlendSpaceBase* const InBlendSpace)
+void UAnimBlueprintPostCompileValidation::PCV_GatherAnimSequences(TArray<const UAnimSequence*>& OutAnimSequences, const class UBlendSpace* const InBlendSpace)
 {
 	// Make sure BlendSpace is loaded, so we can access referenced AnimSequences.
 	PCV_PreloadObject(InBlendSpace);
@@ -168,7 +168,7 @@ void UAnimBlueprintPostCompileValidation::PCV_GatherAnimSequencesFromGraph(TArra
 	}
 }
 
-void UAnimBlueprintPostCompileValidation::PCV_GatherBlendSpacesFromGraph(TArray<const class UBlendSpaceBase*>& OutBlendSpaces, FAnimBPCompileValidationParams& PCV_Params)
+void UAnimBlueprintPostCompileValidation::PCV_GatherBlendSpacesFromGraph(TArray<const class UBlendSpace*>& OutBlendSpaces, FAnimBPCompileValidationParams& PCV_Params)
 {
 	for (FStructProperty* Property : TFieldRange<FStructProperty>(PCV_Params.NewAnimBlueprintClass, EFieldIteratorFlags::IncludeSuper))
 	{
diff --git a/Engine/Source/Editor/AnimGraph/Private/AnimGraphNode_AimOffsetLookAt.cpp b/Engine/Source/Editor/AnimGraph/Private/AnimGraphNode_AimOffsetLookAt.cpp
--- a/Engine/Source/Editor/AnimGraph/Private/AnimGraphNode_AimOffsetLookAt.cpp
+++ b/Engine/Source/Editor/AnimGraph/Private/AnimGraphNode_AimOffsetLookAt.cpp
@@ -31,11 +31,11 @@ FText UAnimGraphNode_AimOffsetLookAt::GetTooltipText() const
 
 FText UAnimGraphNode_AimOffsetLookAt::GetNodeTitle(ENodeTitleType::Type TitleType) const
 {
-	UBlendSpaceBase* BlendSpaceToCheck = Node.BlendSpace;
+	UBlendSpace* BlendSpaceToCheck = Node.BlendSpace;
 	UEdGraphPin* BlendSpacePin = FindPin(GET_MEMBER_NAME_STRING_CHECKED(FAnimNode_AimOffsetLookAt, BlendSpace));
 	if (BlendSpacePin != nullptr && BlendSpaceToCheck == nullptr)
 	{
-		BlendSpaceToCheck = Cast<UBlendSpaceBase>(BlendSpacePin->DefaultObject);
+		BlendSpaceToCheck = Cast<UBlendSpace>(BlendSpacePin->DefaultObject);
 	}
 
 	if (BlendSpaceToCheck == nullptr)
@@ -75,13 +75,13 @@ void UAnimGraphNode_AimOffsetLookAt::GetMenuActions(FBlueprintActionDatabaseRegi
 {
 	struct GetMenuActions_Utils
 	{
-		static void SetNodeBlendSpace(UEdGraphNode* NewNode, bool /*bIsTemplateNode*/, TWeakObjectPtr<UBlendSpaceBase> BlendSpace)
+		static void SetNodeBlendSpace(UEdGraphNode* NewNode, bool /*bIsTemplateNode*/, TWeakObjectPtr<UBlendSpace> BlendSpace)
 		{
 			UAnimGraphNode_AimOffsetLookAt* BlendSpaceNode = CastChecked<UAnimGraphNode_AimOffsetLookAt>(NewNode);
 			BlendSpaceNode->Node.BlendSpace = BlendSpace.Get();
 		}
 
-		static UBlueprintNodeSpawner* MakeBlendSpaceAction(TSubclassOf<UEdGraphNode> const NodeClass, const UBlendSpaceBase* BlendSpace)
+		static UBlueprintNodeSpawner* MakeBlendSpaceAction(TSubclassOf<UEdGraphNode> const NodeClass, const UBlendSpace* BlendSpace)
 		{
 			UBlueprintNodeSpawner* NodeSpawner = nullptr;
 
@@ -92,7 +92,7 @@ void UAnimGraphNode_AimOffsetLookAt::GetMenuActions(FBlueprintActionDatabaseRegi
 				NodeSpawner = UBlueprintNodeSpawner::Create(NodeClass);
 				check(NodeSpawner != nullptr);
 
-				TWeakObjectPtr<UBlendSpaceBase> BlendSpacePtr = MakeWeakObjectPtr(const_cast<UBlendSpaceBase*>(BlendSpace));
+				TWeakObjectPtr<UBlendSpace> BlendSpacePtr = MakeWeakObjectPtr(const_cast<UBlendSpace*>(BlendSpace));
 				NodeSpawner->CustomizeNodeDelegate = UBlueprintNodeSpawner::FCustomizeNodeDelegate::CreateStatic(GetMenuActions_Utils::SetNodeBlendSpace, BlendSpacePtr);
 			}
 			return NodeSpawner;
@@ -101,7 +101,7 @@ void UAnimGraphNode_AimOffsetLookAt::GetMenuActions(FBlueprintActionDatabaseRegi
 
 	if (const UObject* RegistrarTarget = ActionRegistrar.GetActionKeyFilter())
 	{
-		if (const UBlendSpaceBase* TargetBlendSpace = Cast<UBlendSpaceBase>(RegistrarTarget))
+		if (const UBlendSpace* TargetBlendSpace = Cast<UBlendSpace>(RegistrarTarget))
 		{
 			if(TargetBlendSpace->IsAsset())
 			{
@@ -116,9 +116,9 @@ void UAnimGraphNode_AimOffsetLookAt::GetMenuActions(FBlueprintActionDatabaseRegi
 	else
 	{
 		UClass* NodeClass = GetClass();
-		for (TObjectIterator<UBlendSpaceBase> BlendSpaceIt; BlendSpaceIt; ++BlendSpaceIt)
+		for (TObjectIterator<UBlendSpace> BlendSpaceIt; BlendSpaceIt; ++BlendSpaceIt)
 		{
-			UBlendSpaceBase* BlendSpace = *BlendSpaceIt;
+			UBlendSpace* BlendSpace = *BlendSpaceIt;
 			if(BlendSpace->IsAsset())
 			{
 				if (UBlueprintNodeSpawner* NodeSpawner = GetMenuActions_Utils::MakeBlendSpaceAction(NodeClass, BlendSpace))
@@ -140,7 +140,7 @@ FBlueprintNodeSignature UAnimGraphNode_AimOffsetLookAt::GetSignature() const
 
 void UAnimGraphNode_AimOffsetLookAt::SetAnimationAsset(UAnimationAsset* Asset)
 {
-	if (UBlendSpaceBase* BlendSpace = Cast<UBlendSpaceBase>(Asset))
+	if (UBlendSpace* BlendSpace = Cast<UBlendSpace>(Asset))
 	{
 		Node.BlendSpace = BlendSpace;
 	}
@@ -150,11 +150,11 @@ void UAnimGraphNode_AimOffsetLookAt::ValidateAnimNodeDuringCompilation(class USk
 {
 	Super::ValidateAnimNodeDuringCompilation(ForSkeleton, MessageLog);
 
-	UBlendSpaceBase* BlendSpaceToCheck = Node.BlendSpace;
+	UBlendSpace* BlendSpaceToCheck = Node.BlendSpace;
 	UEdGraphPin* BlendSpacePin = FindPin(GET_MEMBER_NAME_STRING_CHECKED(FAnimNode_AimOffsetLookAt, BlendSpace));
 	if (BlendSpacePin != nullptr && BlendSpaceToCheck == nullptr)
 	{
-		BlendSpaceToCheck = Cast<UBlendSpaceBase>(BlendSpacePin->DefaultObject);
+		BlendSpaceToCheck = Cast<UBlendSpace>(BlendSpacePin->DefaultObject);
 	}
 
 	if (!BlendSpaceToCheck)
@@ -246,7 +246,7 @@ void UAnimGraphNode_AimOffsetLookAt::CustomizePinData(UEdGraphPin* Pin, FName So
 	Super::CustomizePinData(Pin, SourcePropertyName, ArrayIndex);
 
 	// Hide input pins that are not relevant for this child class.
-	UBlendSpaceBase * BlendSpace = GetBlendSpace();
+	UBlendSpace * BlendSpace = GetBlendSpace();
 	if (BlendSpace != NULL)
 	{
 		if ((SourcePropertyName == TEXT("X")) || (SourcePropertyName == FName(*BlendSpace->GetBlendParameter(0).DisplayName)))
diff --git a/Engine/Source/Editor/AnimGraph/Private/AnimGraphNode_AssetPlayerBase.cpp b/Engine/Source/Editor/AnimGraph/Private/AnimGraphNode_AssetPlayerBase.cpp
--- a/Engine/Source/Editor/AnimGraph/Private/AnimGraphNode_AssetPlayerBase.cpp
+++ b/Engine/Source/Editor/AnimGraph/Private/AnimGraphNode_AssetPlayerBase.cpp
@@ -3,7 +3,7 @@
 #include "AnimGraphNode_AssetPlayerBase.h"
 #include "EdGraphSchema_K2.h"
 #include "Animation/AnimComposite.h"
-#include "Animation/BlendSpaceBase.h"
+#include "Animation/BlendSpace.h"
 #include "Animation/AimOffsetBlendSpace.h"
 #include "Animation/AimOffsetBlendSpace1D.h"
 #include "AnimGraphNode_SequencePlayer.h"
diff --git a/Engine/Source/Editor/AnimGraph/Private/AnimGraphNode_BlendSpaceBase.cpp b/Engine/Source/Editor/AnimGraph/Private/AnimGraphNode_BlendSpaceBase.cpp
--- a/Engine/Source/Editor/AnimGraph/Private/AnimGraphNode_BlendSpaceBase.cpp
+++ b/Engine/Source/Editor/AnimGraph/Private/AnimGraphNode_BlendSpaceBase.cpp
@@ -25,7 +25,7 @@ FLinearColor UAnimGraphNode_BlendSpaceBase::GetNodeTitleColor() const
 
 void UAnimGraphNode_BlendSpaceBase::CustomizePinData(UEdGraphPin* Pin, FName SourcePropertyName, int32 ArrayIndex) const
 {
-	UBlendSpaceBase * BlendSpace = GetBlendSpace();
+	UBlendSpace * BlendSpace = GetBlendSpace();
 
 	if (BlendSpace != NULL)
 	{
@@ -57,7 +57,7 @@ void UAnimGraphNode_BlendSpaceBase::PostProcessPinName(const UEdGraphPin* Pin, F
 {
 	if(Pin->Direction == EGPD_Input)
 	{
-		UBlendSpaceBase * BlendSpace = GetBlendSpace();
+		UBlendSpace * BlendSpace = GetBlendSpace();
 
 		if(BlendSpace != NULL)
 		{
diff --git a/Engine/Source/Editor/AnimGraph/Private/AnimGraphNode_BlendSpaceEvaluator.cpp b/Engine/Source/Editor/AnimGraph/Private/AnimGraphNode_BlendSpaceEvaluator.cpp
--- a/Engine/Source/Editor/AnimGraph/Private/AnimGraphNode_BlendSpaceEvaluator.cpp
+++ b/Engine/Source/Editor/AnimGraph/Private/AnimGraphNode_BlendSpaceEvaluator.cpp
@@ -21,7 +21,7 @@ FText UAnimGraphNode_BlendSpaceEvaluator::GetTooltipText() const
 	return GetNodeTitle(ENodeTitleType::ListView);
 }
 
-FText UAnimGraphNode_BlendSpaceEvaluator::GetNodeTitleForBlendSpace(ENodeTitleType::Type TitleType, UBlendSpaceBase* InBlendSpace) const
+FText UAnimGraphNode_BlendSpaceEvaluator::GetNodeTitleForBlendSpace(ENodeTitleType::Type TitleType, UBlendSpace* InBlendSpace) const
 {
 	const FText BlendSpaceName = FText::FromString(InBlendSpace->GetName());
 
@@ -65,7 +65,7 @@ FText UAnimGraphNode_BlendSpaceEvaluator::GetNodeTitle(ENodeTitleType::Type Titl
 		}
 		else if (BlendSpacePin && BlendSpacePin->DefaultObject != nullptr)
 		{
-			return GetNodeTitleForBlendSpace(TitleType, CastChecked<UBlendSpaceBase>(BlendSpacePin->DefaultObject));
+			return GetNodeTitleForBlendSpace(TitleType, CastChecked<UBlendSpace>(BlendSpacePin->DefaultObject));
 		}
 		else
 		{
@@ -97,11 +97,11 @@ void UAnimGraphNode_BlendSpaceEvaluator::ValidateAnimNodeDuringCompilation(class
 {
 	Super::ValidateAnimNodeDuringCompilation(ForSkeleton, MessageLog);
 
-	UBlendSpaceBase* BlendSpaceToCheck = Node.BlendSpace;
+	UBlendSpace* BlendSpaceToCheck = Node.BlendSpace;
 	UEdGraphPin* BlendSpacePin = FindPin(GET_MEMBER_NAME_STRING_CHECKED(FAnimNode_BlendSpaceEvaluator, BlendSpace));
 	if (BlendSpacePin != nullptr && BlendSpaceToCheck == nullptr)
 	{
-		BlendSpaceToCheck = Cast<UBlendSpaceBase>(BlendSpacePin->DefaultObject);
+		BlendSpaceToCheck = Cast<UBlendSpace>(BlendSpacePin->DefaultObject);
 	}
 
 	if (BlendSpaceToCheck == nullptr)
@@ -157,7 +157,7 @@ void UAnimGraphNode_BlendSpaceEvaluator::GetNodeContextMenuActions(UToolMenu* Me
 
 void UAnimGraphNode_BlendSpaceEvaluator::SetAnimationAsset(UAnimationAsset* Asset)
 {
-	if (UBlendSpaceBase* BlendSpace = Cast<UBlendSpaceBase>(Asset))
+	if (UBlendSpace* BlendSpace = Cast<UBlendSpace>(Asset))
 	{
 		Node.BlendSpace = BlendSpace;
 	}
@@ -170,11 +170,11 @@ bool UAnimGraphNode_BlendSpaceEvaluator::DoesSupportTimeForTransitionGetter() co
 
 UAnimationAsset* UAnimGraphNode_BlendSpaceEvaluator::GetAnimationAsset() const 
 {
-	UBlendSpaceBase* BlendSpace = Node.BlendSpace;
+	UBlendSpace* BlendSpace = Node.BlendSpace;
 	UEdGraphPin* BlendSpacePin = FindPin(GET_MEMBER_NAME_STRING_CHECKED(FAnimNode_BlendSpaceEvaluator, BlendSpace));
 	if (BlendSpacePin != nullptr && BlendSpace == nullptr)
 	{
-		BlendSpace = Cast<UBlendSpaceBase>(BlendSpacePin->DefaultObject);
+		BlendSpace = Cast<UBlendSpace>(BlendSpacePin->DefaultObject);
 	}
 
 	return BlendSpace;
@@ -192,7 +192,7 @@ UScriptStruct* UAnimGraphNode_BlendSpaceEvaluator::GetTimePropertyStruct() const
 
 EAnimAssetHandlerType UAnimGraphNode_BlendSpaceEvaluator::SupportsAssetClass(const UClass* AssetClass) const
 {
-	if (AssetClass->IsChildOf(UBlendSpaceBase::StaticClass()) && !IsAimOffsetBlendSpace(AssetClass))
+	if (AssetClass->IsChildOf(UBlendSpace::StaticClass()) && !IsAimOffsetBlendSpace(AssetClass))
 	{
 		return EAnimAssetHandlerType::Supported;
 	}
diff --git a/Engine/Source/Editor/AnimGraph/Private/AnimGraphNode_BlendSpaceGraph.cpp b/Engine/Source/Editor/AnimGraph/Private/AnimGraphNode_BlendSpaceGraph.cpp
--- a/Engine/Source/Editor/AnimGraph/Private/AnimGraphNode_BlendSpaceGraph.cpp
+++ b/Engine/Source/Editor/AnimGraph/Private/AnimGraphNode_BlendSpaceGraph.cpp
@@ -55,13 +55,13 @@ void UAnimGraphNode_BlendSpaceGraph::GetMenuActions(FBlueprintActionDatabaseRegi
 {
 	struct GetMenuActions_Utils
 	{
-		static void SetNodeBlendSpace(UEdGraphNode* NewNode, bool bIsTemplateNode, TWeakObjectPtr<UBlendSpaceBase> BlendSpace)
+		static void SetNodeBlendSpace(UEdGraphNode* NewNode, bool bIsTemplateNode, TWeakObjectPtr<UBlendSpace> BlendSpace)
 		{
 			UAnimGraphNode_BlendSpaceGraph* BlendSpaceNode = CastChecked<UAnimGraphNode_BlendSpaceGraph>(NewNode);
 			BlendSpaceNode->SetupFromAsset(BlendSpace.Get(), bIsTemplateNode);
 		}
 
-		static UBlueprintNodeSpawner* MakeBlendSpaceAction(TSubclassOf<UEdGraphNode> const NodeClass, const UBlendSpaceBase* InBlendSpace)
+		static UBlueprintNodeSpawner* MakeBlendSpaceAction(TSubclassOf<UEdGraphNode> const NodeClass, const UBlendSpace* InBlendSpace)
 		{
 			UBlueprintNodeSpawner* NodeSpawner = nullptr;
 
@@ -72,19 +72,19 @@ void UAnimGraphNode_BlendSpaceGraph::GetMenuActions(FBlueprintActionDatabaseRegi
 				NodeSpawner = UBlueprintNodeSpawner::Create(NodeClass);
 				check(NodeSpawner != nullptr);
 
-				TWeakObjectPtr<UBlendSpaceBase> BlendSpacePtr = MakeWeakObjectPtr(const_cast<UBlendSpaceBase*>(InBlendSpace));
+				TWeakObjectPtr<UBlendSpace> BlendSpacePtr = MakeWeakObjectPtr(const_cast<UBlendSpace*>(InBlendSpace));
 				NodeSpawner->CustomizeNodeDelegate = UBlueprintNodeSpawner::FCustomizeNodeDelegate::CreateStatic(GetMenuActions_Utils::SetNodeBlendSpace, BlendSpacePtr);
 			}	
 			return NodeSpawner;
 		}
 
-		static void SetNodeBlendSpaceClass(UEdGraphNode* NewNode, bool bIsTemplateNode, TSubclassOf<UBlendSpaceBase> InBlendSpaceClass)
+		static void SetNodeBlendSpaceClass(UEdGraphNode* NewNode, bool bIsTemplateNode, TSubclassOf<UBlendSpace> InBlendSpaceClass)
 		{
 			UAnimGraphNode_BlendSpaceGraph* BlendSpaceNode = CastChecked<UAnimGraphNode_BlendSpaceGraph>(NewNode);
 			BlendSpaceNode->SetupFromClass(InBlendSpaceClass, bIsTemplateNode);
 		}
 
-		static UBlueprintNodeSpawner* MakeBlendSpaceAction(TSubclassOf<UEdGraphNode> const NodeClass, TSubclassOf<UBlendSpaceBase> InBlendSpaceClass)
+		static UBlueprintNodeSpawner* MakeBlendSpaceAction(TSubclassOf<UEdGraphNode> const NodeClass, TSubclassOf<UBlendSpace> InBlendSpaceClass)
 		{
 			UBlueprintNodeSpawner* NodeSpawner = UBlueprintNodeSpawner::Create(NodeClass);
 			check(NodeSpawner != nullptr);
@@ -101,7 +101,7 @@ void UAnimGraphNode_BlendSpaceGraph::GetMenuActions(FBlueprintActionDatabaseRegi
 
 	if (const UObject* RegistrarTarget = ActionRegistrar.GetActionKeyFilter())
 	{
-		if (const UBlendSpaceBase* TargetBlendSpace = Cast<UBlendSpaceBase>(RegistrarTarget))
+		if (const UBlendSpace* TargetBlendSpace = Cast<UBlendSpace>(RegistrarTarget))
 		{
 			if(TargetBlendSpace->IsAsset())
 			{
@@ -116,7 +116,7 @@ void UAnimGraphNode_BlendSpaceGraph::GetMenuActions(FBlueprintActionDatabaseRegi
 	else
 	{
 		UClass* NodeClass = GetClass();
-		for (TObjectIterator<UBlendSpaceBase> BlendSpaceIt; BlendSpaceIt; ++BlendSpaceIt)
+		for (TObjectIterator<UBlendSpace> BlendSpaceIt; BlendSpaceIt; ++BlendSpaceIt)
 		{
 			if(BlendSpaceIt->IsAsset())
 			{
diff --git a/Engine/Source/Editor/AnimGraph/Private/AnimGraphNode_BlendSpaceGraphBase.cpp b/Engine/Source/Editor/AnimGraph/Private/AnimGraphNode_BlendSpaceGraphBase.cpp
--- a/Engine/Source/Editor/AnimGraph/Private/AnimGraphNode_BlendSpaceGraphBase.cpp
+++ b/Engine/Source/Editor/AnimGraph/Private/AnimGraphNode_BlendSpaceGraphBase.cpp
@@ -128,7 +128,7 @@ void UAnimGraphNode_BlendSpaceGraphBase::OnCopyTermDefaultsToDefaultObject(IAnim
 	DestinationNode->BlendSpace = OutCompiledData.AddBlendSpace(BlendSpace);
 }
 
-void UAnimGraphNode_BlendSpaceGraphBase::SetupFromAsset(UBlendSpaceBase* InBlendSpace, bool bInIsTemplateNode)
+void UAnimGraphNode_BlendSpaceGraphBase::SetupFromAsset(UBlendSpace* InBlendSpace, bool bInIsTemplateNode)
 {
 	if(bInIsTemplateNode)
 	{
@@ -171,7 +171,7 @@ void UAnimGraphNode_BlendSpaceGraphBase::SetupFromAsset(UBlendSpaceBase* InBlend
 	}
 }
 
-void UAnimGraphNode_BlendSpaceGraphBase::SetupFromClass(TSubclassOf<UBlendSpaceBase> InBlendSpaceClass, bool bInIsTemplateNode)
+void UAnimGraphNode_BlendSpaceGraphBase::SetupFromClass(TSubclassOf<UBlendSpace> InBlendSpaceClass, bool bInIsTemplateNode)
 {
 	if(bInIsTemplateNode)
 	{
@@ -180,7 +180,7 @@ void UAnimGraphNode_BlendSpaceGraphBase::SetupFromClass(TSubclassOf<UBlendSpaceB
 	else
 	{
 		BlendSpaceGraph = CastChecked<UBlendSpaceGraph>(FBlueprintEditorUtils::CreateNewGraph(this, InBlendSpaceClass.Get()->GetFName(), UBlendSpaceGraph::StaticClass(), UAnimationGraphSchema::StaticClass()));
-		BlendSpaceGraph->BlendSpace = BlendSpace = NewObject<UBlendSpaceBase>(BlendSpaceGraph, InBlendSpaceClass.Get());
+		BlendSpaceGraph->BlendSpace = BlendSpace = NewObject<UBlendSpace>(BlendSpaceGraph, InBlendSpaceClass.Get());
 		BlendSpaceGraph->BlendSpace->SetFlags(RF_Transactional);
 		BlendSpaceGraph->bAllowDeletion = false;
 		BlendSpaceGraph->bAllowRenaming = true;
@@ -328,7 +328,7 @@ void UAnimGraphNode_BlendSpaceGraphBase::PostPlacedNewNode()
 		check(BlendSpace == nullptr);
 
 		BlendSpaceGraph = CastChecked<UBlendSpaceGraph>(FBlueprintEditorUtils::CreateNewGraph(this, NAME_None, UBlendSpaceGraph::StaticClass(), UAnimationGraphSchema::StaticClass()));
-		BlendSpaceGraph->BlendSpace = BlendSpace = NewObject<UBlendSpaceBase>();
+		BlendSpaceGraph->BlendSpace = BlendSpace = NewObject<UBlendSpace>();
 		
 		// Find an interesting name
 		TSharedPtr<INameValidatorInterface> NameValidator = FNameValidatorFactory::MakeValidator(this);
diff --git a/Engine/Source/Editor/AnimGraph/Private/AnimGraphNode_BlendSpacePlayer.cpp b/Engine/Source/Editor/AnimGraph/Private/AnimGraphNode_BlendSpacePlayer.cpp
--- a/Engine/Source/Editor/AnimGraph/Private/AnimGraphNode_BlendSpacePlayer.cpp
+++ b/Engine/Source/Editor/AnimGraph/Private/AnimGraphNode_BlendSpacePlayer.cpp
@@ -28,7 +28,7 @@ FText UAnimGraphNode_BlendSpacePlayer::GetTooltipText() const
 	return GetNodeTitle(ENodeTitleType::ListView);
 }
 
-FText UAnimGraphNode_BlendSpacePlayer::GetNodeTitleForBlendSpace(ENodeTitleType::Type TitleType, UBlendSpaceBase* InBlendSpace) const
+FText UAnimGraphNode_BlendSpacePlayer::GetNodeTitleForBlendSpace(ENodeTitleType::Type TitleType, UBlendSpace* InBlendSpace) const
 {
 	const FText BlendSpaceName = FText::FromString(InBlendSpace->GetName());
 
@@ -94,7 +94,7 @@ FText UAnimGraphNode_BlendSpacePlayer::GetNodeTitle(ENodeTitleType::Type TitleTy
 		}
 		else if (BlendSpacePin && BlendSpacePin->DefaultObject != nullptr)
 		{
-			return GetNodeTitleForBlendSpace(TitleType, CastChecked<UBlendSpaceBase>(BlendSpacePin->DefaultObject));
+			return GetNodeTitleForBlendSpace(TitleType, CastChecked<UBlendSpace>(BlendSpacePin->DefaultObject));
 		}
 		else
 		{
@@ -120,11 +120,11 @@ void UAnimGraphNode_BlendSpacePlayer::ValidateAnimNodeDuringCompilation(class US
 {
 	Super::ValidateAnimNodeDuringCompilation(ForSkeleton, MessageLog);
 
-	UBlendSpaceBase* BlendSpaceToCheck = Node.BlendSpace;
+	UBlendSpace* BlendSpaceToCheck = Node.BlendSpace;
 	UEdGraphPin* BlendSpacePin = FindPin(GET_MEMBER_NAME_STRING_CHECKED(FAnimNode_BlendSpacePlayer, BlendSpace));
 	if (BlendSpacePin != nullptr && BlendSpaceToCheck == nullptr)
 	{
-		BlendSpaceToCheck = Cast<UBlendSpaceBase>(BlendSpacePin->DefaultObject);
+		BlendSpaceToCheck = Cast<UBlendSpace>(BlendSpacePin->DefaultObject);
 	}
 
 	if (BlendSpaceToCheck == nullptr)
@@ -183,13 +183,13 @@ void UAnimGraphNode_BlendSpacePlayer::GetMenuActions(FBlueprintActionDatabaseReg
 {
 	struct GetMenuActions_Utils
 	{
-		static void SetNodeBlendSpace(UEdGraphNode* NewNode, bool /*bIsTemplateNode*/, TWeakObjectPtr<UBlendSpaceBase> BlendSpace)
+		static void SetNodeBlendSpace(UEdGraphNode* NewNode, bool /*bIsTemplateNode*/, TWeakObjectPtr<UBlendSpace> BlendSpace)
 		{
 			UAnimGraphNode_BlendSpacePlayer* BlendSpaceNode = CastChecked<UAnimGraphNode_BlendSpacePlayer>(NewNode);
 			BlendSpaceNode->Node.BlendSpace = BlendSpace.Get();
 		}
 
-		static UBlueprintNodeSpawner* MakeBlendSpaceAction(TSubclassOf<UEdGraphNode> const NodeClass, const UBlendSpaceBase* BlendSpace)
+		static UBlueprintNodeSpawner* MakeBlendSpaceAction(TSubclassOf<UEdGraphNode> const NodeClass, const UBlendSpace* BlendSpace)
 		{
 			UBlueprintNodeSpawner* NodeSpawner = nullptr;
 
@@ -200,7 +200,7 @@ void UAnimGraphNode_BlendSpacePlayer::GetMenuActions(FBlueprintActionDatabaseReg
 				NodeSpawner = UBlueprintNodeSpawner::Create(NodeClass);
 				check(NodeSpawner != nullptr);
 
-				TWeakObjectPtr<UBlendSpaceBase> BlendSpacePtr = MakeWeakObjectPtr(const_cast<UBlendSpaceBase*>(BlendSpace));
+				TWeakObjectPtr<UBlendSpace> BlendSpacePtr = MakeWeakObjectPtr(const_cast<UBlendSpace*>(BlendSpace));
 				NodeSpawner->CustomizeNodeDelegate = UBlueprintNodeSpawner::FCustomizeNodeDelegate::CreateStatic(GetMenuActions_Utils::SetNodeBlendSpace, BlendSpacePtr);
 			}	
 			return NodeSpawner;
@@ -209,7 +209,7 @@ void UAnimGraphNode_BlendSpacePlayer::GetMenuActions(FBlueprintActionDatabaseReg
 
 	if (const UObject* RegistrarTarget = ActionRegistrar.GetActionKeyFilter())
 	{
-		if (const UBlendSpaceBase* TargetBlendSpace = Cast<UBlendSpaceBase>(RegistrarTarget))
+		if (const UBlendSpace* TargetBlendSpace = Cast<UBlendSpace>(RegistrarTarget))
 		{
 			if(TargetBlendSpace->IsAsset())
 			{
@@ -224,9 +224,9 @@ void UAnimGraphNode_BlendSpacePlayer::GetMenuActions(FBlueprintActionDatabaseReg
 	else
 	{
 		UClass* NodeClass = GetClass();
-		for (TObjectIterator<UBlendSpaceBase> BlendSpaceIt; BlendSpaceIt; ++BlendSpaceIt)
+		for (TObjectIterator<UBlendSpace> BlendSpaceIt; BlendSpaceIt; ++BlendSpaceIt)
 		{
-			UBlendSpaceBase* BlendSpace = *BlendSpaceIt;
+			UBlendSpace* BlendSpace = *BlendSpaceIt;
 			if(BlendSpace->IsAsset())
 			{
 				if (UBlueprintNodeSpawner* NodeSpawner = GetMenuActions_Utils::MakeBlendSpaceAction(NodeClass, BlendSpace))
@@ -248,7 +248,7 @@ FBlueprintNodeSignature UAnimGraphNode_BlendSpacePlayer::GetSignature() const
 
 void UAnimGraphNode_BlendSpacePlayer::SetAnimationAsset(UAnimationAsset* Asset)
 {
-	if (UBlendSpaceBase* BlendSpace = Cast<UBlendSpaceBase>(Asset))
+	if (UBlendSpace* BlendSpace = Cast<UBlendSpace>(Asset))
 	{
 		Node.BlendSpace = BlendSpace;
 	}
@@ -274,11 +274,11 @@ bool UAnimGraphNode_BlendSpacePlayer::DoesSupportTimeForTransitionGetter() const
 
 UAnimationAsset* UAnimGraphNode_BlendSpacePlayer::GetAnimationAsset() const 
 {
-	UBlendSpaceBase* BlendSpace = Node.BlendSpace;
+	UBlendSpace* BlendSpace = Node.BlendSpace;
 	UEdGraphPin* BlendSpacePin = FindPin(GET_MEMBER_NAME_STRING_CHECKED(FAnimNode_BlendSpacePlayer, BlendSpace));
 	if (BlendSpacePin != nullptr && BlendSpace == nullptr)
 	{
-		BlendSpace = Cast<UBlendSpaceBase>(BlendSpacePin->DefaultObject);
+		BlendSpace = Cast<UBlendSpace>(BlendSpacePin->DefaultObject);
 	}
 
 	return BlendSpace;
@@ -296,7 +296,7 @@ UScriptStruct* UAnimGraphNode_BlendSpacePlayer::GetTimePropertyStruct() const
 
 EAnimAssetHandlerType UAnimGraphNode_BlendSpacePlayer::SupportsAssetClass(const UClass* AssetClass) const
 {
-	if (AssetClass->IsChildOf(UBlendSpaceBase::StaticClass()) && !IsAimOffsetBlendSpace(AssetClass))
+	if (AssetClass->IsChildOf(UBlendSpace::StaticClass()) && !IsAimOffsetBlendSpace(AssetClass))
 	{
 		return EAnimAssetHandlerType::PrimaryHandler;
 	}
diff --git a/Engine/Source/Editor/AnimGraph/Private/AnimGraphNode_RotationOffsetBlendSpace.cpp b/Engine/Source/Editor/AnimGraph/Private/AnimGraphNode_RotationOffsetBlendSpace.cpp
--- a/Engine/Source/Editor/AnimGraph/Private/AnimGraphNode_RotationOffsetBlendSpace.cpp
+++ b/Engine/Source/Editor/AnimGraph/Private/AnimGraphNode_RotationOffsetBlendSpace.cpp
@@ -34,11 +34,11 @@ FText UAnimGraphNode_RotationOffsetBlendSpace::GetTooltipText() const
 
 FText UAnimGraphNode_RotationOffsetBlendSpace::GetNodeTitle(ENodeTitleType::Type TitleType) const
 {
-	UBlendSpaceBase* BlendSpaceToCheck = Node.BlendSpace;
+	UBlendSpace* BlendSpaceToCheck = Node.BlendSpace;
 	UEdGraphPin* BlendSpacePin = FindPin(GET_MEMBER_NAME_STRING_CHECKED(FAnimNode_RotationOffsetBlendSpace, BlendSpace));
 	if (BlendSpacePin != nullptr && BlendSpaceToCheck == nullptr)
 	{
-		BlendSpaceToCheck = Cast<UBlendSpaceBase>(BlendSpacePin->DefaultObject);
+		BlendSpaceToCheck = Cast<UBlendSpace>(BlendSpacePin->DefaultObject);
 	}
 
 	if (BlendSpaceToCheck == nullptr)
@@ -78,13 +78,13 @@ void UAnimGraphNode_RotationOffsetBlendSpace::GetMenuActions(FBlueprintActionDat
 {
 	struct GetMenuActions_Utils
 	{
-		static void SetNodeBlendSpace(UEdGraphNode* NewNode, bool /*bIsTemplateNode*/, TWeakObjectPtr<UBlendSpaceBase> BlendSpace)
+		static void SetNodeBlendSpace(UEdGraphNode* NewNode, bool /*bIsTemplateNode*/, TWeakObjectPtr<UBlendSpace> BlendSpace)
 		{
 			UAnimGraphNode_RotationOffsetBlendSpace* BlendSpaceNode = CastChecked<UAnimGraphNode_RotationOffsetBlendSpace>(NewNode);
 			BlendSpaceNode->Node.BlendSpace = BlendSpace.Get();
 		}
 
-		static UBlueprintNodeSpawner* MakeBlendSpaceAction(TSubclassOf<UEdGraphNode> const NodeClass, const UBlendSpaceBase* BlendSpace)
+		static UBlueprintNodeSpawner* MakeBlendSpaceAction(TSubclassOf<UEdGraphNode> const NodeClass, const UBlendSpace* BlendSpace)
 		{
 			UBlueprintNodeSpawner* NodeSpawner = nullptr;
 
@@ -95,7 +95,7 @@ void UAnimGraphNode_RotationOffsetBlendSpace::GetMenuActions(FBlueprintActionDat
 				NodeSpawner = UBlueprintNodeSpawner::Create(NodeClass);
 				check(NodeSpawner != nullptr);
 
-				TWeakObjectPtr<UBlendSpaceBase> BlendSpacePtr = MakeWeakObjectPtr(const_cast<UBlendSpaceBase*>(BlendSpace));
+				TWeakObjectPtr<UBlendSpace> BlendSpacePtr = MakeWeakObjectPtr(const_cast<UBlendSpace*>(BlendSpace));
 				NodeSpawner->CustomizeNodeDelegate = UBlueprintNodeSpawner::FCustomizeNodeDelegate::CreateStatic(GetMenuActions_Utils::SetNodeBlendSpace, BlendSpacePtr);
 			}
 			return NodeSpawner;
@@ -104,7 +104,7 @@ void UAnimGraphNode_RotationOffsetBlendSpace::GetMenuActions(FBlueprintActionDat
 
 	if (const UObject* RegistrarTarget = ActionRegistrar.GetActionKeyFilter())
 	{
-		if (const UBlendSpaceBase* TargetBlendSpace = Cast<UBlendSpaceBase>(RegistrarTarget))
+		if (const UBlendSpace* TargetBlendSpace = Cast<UBlendSpace>(RegistrarTarget))
 		{
 			if(TargetBlendSpace->IsAsset())
 			{
@@ -119,9 +119,9 @@ void UAnimGraphNode_RotationOffsetBlendSpace::GetMenuActions(FBlueprintActionDat
 	else
 	{
 		UClass* NodeClass = GetClass();
-		for (TObjectIterator<UBlendSpaceBase> BlendSpaceIt; BlendSpaceIt; ++BlendSpaceIt)
+		for (TObjectIterator<UBlendSpace> BlendSpaceIt; BlendSpaceIt; ++BlendSpaceIt)
 		{
-			UBlendSpaceBase* BlendSpace = *BlendSpaceIt;
+			UBlendSpace* BlendSpace = *BlendSpaceIt;
 			if(BlendSpace->IsAsset())
 			{
 				if (UBlueprintNodeSpawner* NodeSpawner = GetMenuActions_Utils::MakeBlendSpaceAction(NodeClass, BlendSpace))
@@ -143,7 +143,7 @@ FBlueprintNodeSignature UAnimGraphNode_RotationOffsetBlendSpace::GetSignature()
 
 void UAnimGraphNode_RotationOffsetBlendSpace::SetAnimationAsset(UAnimationAsset* Asset)
 {
-	if (UBlendSpaceBase* BlendSpace = Cast<UBlendSpaceBase>(Asset))
+	if (UBlendSpace* BlendSpace = Cast<UBlendSpace>(Asset))
 	{
 		Node.BlendSpace = BlendSpace;
 	}
@@ -153,11 +153,11 @@ void UAnimGraphNode_RotationOffsetBlendSpace::ValidateAnimNodeDuringCompilation(
 {
 	Super::ValidateAnimNodeDuringCompilation(ForSkeleton, MessageLog);
 
-	UBlendSpaceBase* BlendSpaceToCheck = Node.BlendSpace;
+	UBlendSpace* BlendSpaceToCheck = Node.BlendSpace;
 	UEdGraphPin* BlendSpacePin = FindPin(GET_MEMBER_NAME_STRING_CHECKED(FAnimNode_RotationOffsetBlendSpace, BlendSpace));
 	if (BlendSpacePin != nullptr && BlendSpaceToCheck == nullptr)
 	{
-		BlendSpaceToCheck = Cast<UBlendSpaceBase>(BlendSpacePin->DefaultObject);
+		BlendSpaceToCheck = Cast<UBlendSpace>(BlendSpacePin->DefaultObject);
 	}
 
 	if (BlendSpaceToCheck == NULL)
@@ -223,7 +223,7 @@ void UAnimGraphNode_RotationOffsetBlendSpace::ReplaceReferredAnimations(const TM
 
 EAnimAssetHandlerType UAnimGraphNode_RotationOffsetBlendSpace::SupportsAssetClass(const UClass* AssetClass) const
 {
-	if (AssetClass->IsChildOf(UBlendSpaceBase::StaticClass()) && IsAimOffsetBlendSpace(AssetClass))
+	if (AssetClass->IsChildOf(UBlendSpace::StaticClass()) && IsAimOffsetBlendSpace(AssetClass))
 	{
 		return EAnimAssetHandlerType::PrimaryHandler;
 	}
diff --git a/Engine/Source/Editor/AnimGraph/Private/AnimGraphNode_RotationOffsetBlendSpaceGraph.cpp b/Engine/Source/Editor/AnimGraph/Private/AnimGraphNode_RotationOffsetBlendSpaceGraph.cpp
--- a/Engine/Source/Editor/AnimGraph/Private/AnimGraphNode_RotationOffsetBlendSpaceGraph.cpp
+++ b/Engine/Source/Editor/AnimGraph/Private/AnimGraphNode_RotationOffsetBlendSpaceGraph.cpp
@@ -57,13 +57,13 @@ void UAnimGraphNode_RotationOffsetBlendSpaceGraph::GetMenuActions(FBlueprintActi
 {
 	struct GetMenuActions_Utils
 	{
-		static void SetNodeBlendSpace(UEdGraphNode* NewNode, bool bIsTemplateNode, TWeakObjectPtr<UBlendSpaceBase> BlendSpace)
+		static void SetNodeBlendSpace(UEdGraphNode* NewNode, bool bIsTemplateNode, TWeakObjectPtr<UBlendSpace> BlendSpace)
 		{
 			UAnimGraphNode_RotationOffsetBlendSpaceGraph* BlendSpaceNode = CastChecked<UAnimGraphNode_RotationOffsetBlendSpaceGraph>(NewNode);
 			BlendSpaceNode->SetupFromAsset(BlendSpace.Get(), bIsTemplateNode);
 		}
 
-		static UBlueprintNodeSpawner* MakeBlendSpaceAction(TSubclassOf<UEdGraphNode> const NodeClass, const UBlendSpaceBase* InBlendSpace)
+		static UBlueprintNodeSpawner* MakeBlendSpaceAction(TSubclassOf<UEdGraphNode> const NodeClass, const UBlendSpace* InBlendSpace)
 		{
 			UBlueprintNodeSpawner* NodeSpawner = nullptr;
 
@@ -74,19 +74,19 @@ void UAnimGraphNode_RotationOffsetBlendSpaceGraph::GetMenuActions(FBlueprintActi
 				NodeSpawner = UBlueprintNodeSpawner::Create(NodeClass);
 				check(NodeSpawner != nullptr);
 
-				TWeakObjectPtr<UBlendSpaceBase> BlendSpacePtr = MakeWeakObjectPtr(const_cast<UBlendSpaceBase*>(InBlendSpace));
+				TWeakObjectPtr<UBlendSpace> BlendSpacePtr = MakeWeakObjectPtr(const_cast<UBlendSpace*>(InBlendSpace));
 				NodeSpawner->CustomizeNodeDelegate = UBlueprintNodeSpawner::FCustomizeNodeDelegate::CreateStatic(GetMenuActions_Utils::SetNodeBlendSpace, BlendSpacePtr);
 			}	
 			return NodeSpawner;
 		}
 
-		static void SetNodeBlendSpaceClass(UEdGraphNode* NewNode, bool bIsTemplateNode, TSubclassOf<UBlendSpaceBase> InBlendSpaceClass)
+		static void SetNodeBlendSpaceClass(UEdGraphNode* NewNode, bool bIsTemplateNode, TSubclassOf<UBlendSpace> InBlendSpaceClass)
 		{
 			UAnimGraphNode_RotationOffsetBlendSpaceGraph* BlendSpaceNode = CastChecked<UAnimGraphNode_RotationOffsetBlendSpaceGraph>(NewNode);
 			BlendSpaceNode->SetupFromClass(InBlendSpaceClass, bIsTemplateNode);
 		}
 
-		static UBlueprintNodeSpawner* MakeBlendSpaceAction(TSubclassOf<UEdGraphNode> const NodeClass, TSubclassOf<UBlendSpaceBase> InBlendSpaceClass)
+		static UBlueprintNodeSpawner* MakeBlendSpaceAction(TSubclassOf<UEdGraphNode> const NodeClass, TSubclassOf<UBlendSpace> InBlendSpaceClass)
 		{
 			UBlueprintNodeSpawner* NodeSpawner = UBlueprintNodeSpawner::Create(NodeClass);
 			check(NodeSpawner != nullptr);
@@ -103,7 +103,7 @@ void UAnimGraphNode_RotationOffsetBlendSpaceGraph::GetMenuActions(FBlueprintActi
 
 	if (const UObject* RegistrarTarget = ActionRegistrar.GetActionKeyFilter())
 	{
-		if (const UBlendSpaceBase* TargetBlendSpace = Cast<UBlendSpaceBase>(RegistrarTarget))
+		if (const UBlendSpace* TargetBlendSpace = Cast<UBlendSpace>(RegistrarTarget))
 		{
 			if(TargetBlendSpace->IsAsset())
 			{
@@ -118,7 +118,7 @@ void UAnimGraphNode_RotationOffsetBlendSpaceGraph::GetMenuActions(FBlueprintActi
 	else
 	{
 		UClass* NodeClass = GetClass();
-		for (TObjectIterator<UBlendSpaceBase> BlendSpaceIt; BlendSpaceIt; ++BlendSpaceIt)
+		for (TObjectIterator<UBlendSpace> BlendSpaceIt; BlendSpaceIt; ++BlendSpaceIt)
 		{
 			if(BlendSpaceIt->IsAsset())
 			{
diff --git a/Engine/Source/Editor/AnimGraph/Private/AnimationGraphSchema.cpp b/Engine/Source/Editor/AnimGraph/Private/AnimationGraphSchema.cpp
--- a/Engine/Source/Editor/AnimGraph/Private/AnimationGraphSchema.cpp
+++ b/Engine/Source/Editor/AnimGraph/Private/AnimationGraphSchema.cpp
@@ -392,7 +392,7 @@ bool UAnimationGraphSchema::CreateAutomaticConversionNodeAndConnections(UEdGraph
 	}
 }
 
-bool IsAimOffsetBlendSpace(UBlendSpaceBase* BlendSpace)
+bool IsAimOffsetBlendSpace(UBlendSpace* BlendSpace)
 {
 	return	BlendSpace->IsA(UAimOffsetBlendSpace::StaticClass()) ||
 			BlendSpace->IsA(UAimOffsetBlendSpace1D::StaticClass());
diff --git a/Engine/Source/Editor/AnimGraph/Public/AnimBlueprintPostCompileValidation.h b/Engine/Source/Editor/AnimGraph/Public/AnimBlueprintPostCompileValidation.h
--- a/Engine/Source/Editor/AnimGraph/Public/AnimBlueprintPostCompileValidation.h
+++ b/Engine/Source/Editor/AnimGraph/Public/AnimBlueprintPostCompileValidation.h
@@ -47,7 +47,7 @@ protected:
 	static void PCV_PreloadObject(const UObject* const ReferencedObject);
 
 	static void PCV_GatherAnimSequences(TArray<const UAnimSequence*>& OutAnimSequences, const UAnimSequenceBase* const InAnimSequenceBase);
-	static void PCV_GatherAnimSequences(TArray<const UAnimSequence*>& OutAnimSequences, const class UBlendSpaceBase* const InBlendSpace);
+	static void PCV_GatherAnimSequences(TArray<const UAnimSequence*>& OutAnimSequences, const class UBlendSpace* const InBlendSpace);
 
 	struct FPCV_GatherParams
 	{
@@ -85,7 +85,7 @@ protected:
 	};
 
 	static void PCV_GatherAnimSequencesFromGraph(TArray<const UAnimSequence*>& OutAnimSequences, FAnimBPCompileValidationParams& PCV_Params, const FPCV_GatherParams& GatherParams);
-	static void PCV_GatherBlendSpacesFromGraph(TArray<const class UBlendSpaceBase*>& OutBlendSpaces, FAnimBPCompileValidationParams& PCV_Params);
+	static void PCV_GatherBlendSpacesFromGraph(TArray<const class UBlendSpace*>& OutBlendSpaces, FAnimBPCompileValidationParams& PCV_Params);
 
 	struct FPCV_ReferencedAnimSequence
 	{
diff --git a/Engine/Source/Editor/AnimGraph/Public/AnimGraphNode_BlendSpaceBase.h b/Engine/Source/Editor/AnimGraph/Public/AnimGraphNode_BlendSpaceBase.h
--- a/Engine/Source/Editor/AnimGraph/Public/AnimGraphNode_BlendSpaceBase.h
+++ b/Engine/Source/Editor/AnimGraph/Public/AnimGraphNode_BlendSpaceBase.h
@@ -5,7 +5,7 @@
 #include "CoreMinimal.h"
 #include "UObject/ObjectMacros.h"
 #include "AnimGraphNode_AssetPlayerBase.h"
-#include "Animation/BlendSpaceBase.h"
+#include "Animation/BlendSpace.h"
 #include "AnimGraphNode_BlendSpaceBase.generated.h"
 
 UCLASS(Abstract, MinimalAPI)
@@ -25,7 +25,7 @@ class UAnimGraphNode_BlendSpaceBase : public UAnimGraphNode_AssetPlayerBase
 	// End of UAnimGraphNode_Base interface
 
 protected:
-	UBlendSpaceBase* GetBlendSpace() const { return Cast<UBlendSpaceBase>(GetAnimationAsset()); }
+	UBlendSpace* GetBlendSpace() const { return Cast<UBlendSpace>(GetAnimationAsset()); }
 
 	/** Util to determine is an asset class is an aim offset */
 	static bool IsAimOffsetBlendSpace(const UClass* BlendSpaceClass);
diff --git a/Engine/Source/Editor/AnimGraph/Public/AnimGraphNode_BlendSpaceEvaluator.h b/Engine/Source/Editor/AnimGraph/Public/AnimGraphNode_BlendSpaceEvaluator.h
--- a/Engine/Source/Editor/AnimGraph/Public/AnimGraphNode_BlendSpaceEvaluator.h
+++ b/Engine/Source/Editor/AnimGraph/Public/AnimGraphNode_BlendSpaceEvaluator.h
@@ -45,7 +45,7 @@ class UAnimGraphNode_BlendSpaceEvaluator : public UAnimGraphNode_BlendSpaceBase
 	// End of UAnimGraphNode_AssetPlayerBase interface
 
 private:
-	FText GetNodeTitleForBlendSpace(ENodeTitleType::Type TitleType, UBlendSpaceBase* InBlendSpace) const;
+	FText GetNodeTitleForBlendSpace(ENodeTitleType::Type TitleType, UBlendSpace* InBlendSpace) const;
 
 private:
 	/** Constructing FText strings can be costly, so we cache the node's title */
diff --git a/Engine/Source/Editor/AnimGraph/Public/AnimGraphNode_BlendSpaceGraphBase.h b/Engine/Source/Editor/AnimGraph/Public/AnimGraphNode_BlendSpaceGraphBase.h
--- a/Engine/Source/Editor/AnimGraph/Public/AnimGraphNode_BlendSpaceGraphBase.h
+++ b/Engine/Source/Editor/AnimGraph/Public/AnimGraphNode_BlendSpaceGraphBase.h
@@ -39,7 +39,7 @@ public:
 	void ReplaceGraph(int32 InSampleIndex, UAnimSequence* InSequence);
 
 	// Setup this node from the specified asset
-	void SetupFromAsset(UBlendSpaceBase* InBlendSpace, bool bInIsTemplateNode);
+	void SetupFromAsset(UBlendSpace* InBlendSpace, bool bInIsTemplateNode);
 
 	// UEdGraphNode interface
 	virtual void PostPlacedNewNode() override;
@@ -58,15 +58,15 @@ protected:
 	FString GetBlendSpaceName() const;
 
 	// Setup this node from the specified class
-	void SetupFromClass(TSubclassOf<UBlendSpaceBase> InBlendSpaceClass, bool bInIsTemplateNode);
+	void SetupFromClass(TSubclassOf<UBlendSpace> InBlendSpaceClass, bool bInIsTemplateNode);
 
 	// Internal blendspace
 	UPROPERTY()
-	TObjectPtr<UBlendSpaceBase> BlendSpace;
+	TObjectPtr<UBlendSpace> BlendSpace;
 
 	// Blendspace class, for template nodes
 	UPROPERTY()
-	TSubclassOf<UBlendSpaceBase> BlendSpaceClass;
+	TSubclassOf<UBlendSpace> BlendSpaceClass;
 
 	// Dummy blendspace graph (used for navigation only)
 	UPROPERTY()
diff --git a/Engine/Source/Editor/AnimGraph/Public/AnimGraphNode_BlendSpacePlayer.h b/Engine/Source/Editor/AnimGraph/Public/AnimGraphNode_BlendSpacePlayer.h
--- a/Engine/Source/Editor/AnimGraph/Public/AnimGraphNode_BlendSpacePlayer.h
+++ b/Engine/Source/Editor/AnimGraph/Public/AnimGraphNode_BlendSpacePlayer.h
@@ -49,7 +49,7 @@ class UAnimGraphNode_BlendSpacePlayer : public UAnimGraphNode_BlendSpaceBase
 
 private:
 	/** Helper function for GetNodeTitle */
-	FText GetNodeTitleForBlendSpace(ENodeTitleType::Type TitleType, UBlendSpaceBase* InBlendSpace) const;
+	FText GetNodeTitleForBlendSpace(ENodeTitleType::Type TitleType, UBlendSpace* InBlendSpace) const;
 
 private:
 	/** Constructing FText strings can be costly, so we cache the node's title */
diff --git a/Engine/Source/Editor/AnimGraph/Public/BlendSpaceGraph.h b/Engine/Source/Editor/AnimGraph/Public/BlendSpaceGraph.h
--- a/Engine/Source/Editor/AnimGraph/Public/BlendSpaceGraph.h
+++ b/Engine/Source/Editor/AnimGraph/Public/BlendSpaceGraph.h
@@ -6,7 +6,7 @@
 #include "EdGraph/EdGraph.h"
 #include "BlendSpaceGraph.generated.h"
 
-class UBlendSpaceBase;
+class UBlendSpace;
 
 // Dummy graph to hold sub-graphs for blendspaces. Not edited directly.
 UCLASS(MinimalAPI)
@@ -17,5 +17,5 @@ class UBlendSpaceGraph : public UEdGraph
 public:
 	// Blendspace that we wrap
 	UPROPERTY()
-	TObjectPtr<UBlendSpaceBase> BlendSpace;
+	TObjectPtr<UBlendSpace> BlendSpace;
 };
\ No newline at end of file
diff --git a/Engine/Source/Editor/AnimGraph/Public/IAnimBlueprintGeneratedClassCompiledData.h b/Engine/Source/Editor/AnimGraph/Public/IAnimBlueprintGeneratedClassCompiledData.h
--- a/Engine/Source/Editor/AnimGraph/Public/IAnimBlueprintGeneratedClassCompiledData.h
+++ b/Engine/Source/Editor/AnimGraph/Public/IAnimBlueprintGeneratedClassCompiledData.h
@@ -46,5 +46,5 @@ public:
 	// Add a class-internal blendspace
 	// @param	InSourceBlendSpace	The blendspace to duplicate
 	// @return a duplicate of the blendspace outered to the class
-	virtual UBlendSpaceBase* AddBlendSpace(UBlendSpaceBase* InSourceBlendSpace) = 0;
+	virtual UBlendSpace* AddBlendSpace(UBlendSpace* InSourceBlendSpace) = 0;
 };
\ No newline at end of file
diff --git a/Engine/Source/Editor/AnimationBlueprintEditor/Private/AnimationNodes/SBlendSpacePreview.cpp b/Engine/Source/Editor/AnimationBlueprintEditor/Private/AnimationNodes/SBlendSpacePreview.cpp
--- a/Engine/Source/Editor/AnimationBlueprintEditor/Private/AnimationNodes/SBlendSpacePreview.cpp
+++ b/Engine/Source/Editor/AnimationBlueprintEditor/Private/AnimationNodes/SBlendSpacePreview.cpp
@@ -63,7 +63,7 @@ EVisibility SBlendSpacePreview::GetBlendSpaceVisibility() const
 	return EVisibility::Collapsed;
 }
 
-bool SBlendSpacePreview::GetBlendSpaceInfo(TWeakObjectPtr<const UBlendSpaceBase>& OutBlendSpace, FVector& OutPosition, FVector& OutFilteredPosition) const
+bool SBlendSpacePreview::GetBlendSpaceInfo(TWeakObjectPtr<const UBlendSpace>& OutBlendSpace, FVector& OutPosition, FVector& OutFilteredPosition) const
 {
 	if (Node.Get() != nullptr)
 	{
diff --git a/Engine/Source/Editor/AnimationBlueprintEditor/Private/AnimationNodes/SBlendSpacePreview.h b/Engine/Source/Editor/AnimationBlueprintEditor/Private/AnimationNodes/SBlendSpacePreview.h
--- a/Engine/Source/Editor/AnimationBlueprintEditor/Private/AnimationNodes/SBlendSpacePreview.h
+++ b/Engine/Source/Editor/AnimationBlueprintEditor/Private/AnimationNodes/SBlendSpacePreview.h
@@ -7,7 +7,7 @@
 #include "Widgets/DeclarativeSyntaxSupport.h"
 #include "PersonaDelegates.h"
 
-class UBlendSpaceBase;
+class UBlendSpace;
 class UAnimGraphNode_Base;
 
 class SBlendSpacePreview : public SCompoundWidget
@@ -23,10 +23,10 @@ public:
 
 protected:
 	EVisibility GetBlendSpaceVisibility() const;
-	bool GetBlendSpaceInfo(TWeakObjectPtr<const UBlendSpaceBase>& OutBlendSpace, FVector& OutPosition, FVector& OutFilteredPosition) const;
+	bool GetBlendSpaceInfo(TWeakObjectPtr<const UBlendSpace>& OutBlendSpace, FVector& OutPosition, FVector& OutFilteredPosition) const;
 
 	TWeakObjectPtr<const UAnimGraphNode_Base> Node;
-	TWeakObjectPtr<const UBlendSpaceBase> CachedBlendSpace;
+	TWeakObjectPtr<const UBlendSpace> CachedBlendSpace;
 	FVector CachedPosition;
 	FVector CachedFilteredPosition;
 };
diff --git a/Engine/Source/Editor/AnimationBlueprintEditor/Private/AnimationNodes/SGraphNodeBlendSpacePlayer.h b/Engine/Source/Editor/AnimationBlueprintEditor/Private/AnimationNodes/SGraphNodeBlendSpacePlayer.h
--- a/Engine/Source/Editor/AnimationBlueprintEditor/Private/AnimationNodes/SGraphNodeBlendSpacePlayer.h
+++ b/Engine/Source/Editor/AnimationBlueprintEditor/Private/AnimationNodes/SGraphNodeBlendSpacePlayer.h
@@ -7,7 +7,7 @@
 #include "Widgets/DeclarativeSyntaxSupport.h"
 #include "SNodePanel.h"
 #include "AnimationNodes/SAnimationGraphNode.h"
-#include "Animation/BlendSpaceBase.h"
+#include "Animation/BlendSpace.h"
 
 class SVerticalBox;
 class UAnimGraphNode_Base;
diff --git a/Engine/Source/Editor/AnimationBlueprintEditor/Private/BlendSpaceDocumentTabFactory.cpp b/Engine/Source/Editor/AnimationBlueprintEditor/Private/BlendSpaceDocumentTabFactory.cpp
--- a/Engine/Source/Editor/AnimationBlueprintEditor/Private/BlendSpaceDocumentTabFactory.cpp
+++ b/Engine/Source/Editor/AnimationBlueprintEditor/Private/BlendSpaceDocumentTabFactory.cpp
@@ -236,7 +236,7 @@ TAttribute<FText> FBlendSpaceDocumentTabFactory::ConstructTabName(const FWorkflo
 
 	UBlendSpaceGraph* DocumentID = FTabPayload_BlendSpaceGraph::GetBlendSpaceGraph(Info.Payload);
 
-	return MakeAttributeLambda([WeakBlendSpace = TWeakObjectPtr<UBlendSpaceBase>(DocumentID->BlendSpace)]()
+	return MakeAttributeLambda([WeakBlendSpace = TWeakObjectPtr<UBlendSpace>(DocumentID->BlendSpace)]()
 		{ 
 			return WeakBlendSpace.Get() ? FText::FromName(WeakBlendSpace->GetFName()) : FText::GetEmpty();
 		});
diff --git a/Engine/Source/Editor/AnimationBlueprintEditor/Private/BlendSpaceDocumentTabFactory.h b/Engine/Source/Editor/AnimationBlueprintEditor/Private/BlendSpaceDocumentTabFactory.h
--- a/Engine/Source/Editor/AnimationBlueprintEditor/Private/BlendSpaceDocumentTabFactory.h
+++ b/Engine/Source/Editor/AnimationBlueprintEditor/Private/BlendSpaceDocumentTabFactory.h
@@ -2,7 +2,7 @@
 
 #pragma once
 
-#include "Animation/BlendSpaceBase.h"
+#include "Animation/BlendSpace.h"
 #include "WorkflowOrientedApp/WorkflowUObjectDocuments.h"
 
 class FAnimationBlueprintEditor;
diff --git a/Engine/Source/Editor/Persona/Private/AnimationBlendSpaceHelpers.h b/Engine/Source/Editor/Persona/Private/AnimationBlendSpaceHelpers.h
--- a/Engine/Source/Editor/Persona/Private/AnimationBlendSpaceHelpers.h
+++ b/Engine/Source/Editor/Persona/Private/AnimationBlendSpaceHelpers.h
@@ -3,7 +3,7 @@
 #pragma once
 
 #include "CoreMinimal.h"
-#include "Animation/BlendSpaceBase.h"
+#include "Animation/BlendSpace.h"
 #include "Animation/BlendSpace.h"
 
 struct FTriangle;
diff --git a/Engine/Source/Editor/Persona/Private/Customization/BlendParameterDetails.cpp b/Engine/Source/Editor/Persona/Private/Customization/BlendParameterDetails.cpp
--- a/Engine/Source/Editor/Persona/Private/Customization/BlendParameterDetails.cpp
+++ b/Engine/Source/Editor/Persona/Private/Customization/BlendParameterDetails.cpp
@@ -6,7 +6,7 @@
 #include "PropertyHandle.h"
 #include "IDetailChildrenBuilder.h"
 #include "DetailWidgetRow.h"
-#include "Animation/BlendSpaceBase.h"
+#include "Animation/BlendSpace.h"
 
 
 #define LOCTEXT_NAMESPACE "BlendParameterDetails"
diff --git a/Engine/Source/Editor/Persona/Private/Customization/BlendSampleDetails.cpp b/Engine/Source/Editor/Persona/Private/Customization/BlendSampleDetails.cpp
--- a/Engine/Source/Editor/Persona/Private/Customization/BlendSampleDetails.cpp
+++ b/Engine/Source/Editor/Persona/Private/Customization/BlendSampleDetails.cpp
@@ -16,7 +16,7 @@
 #include "Widgets/Text/STextBlock.h"
 
 #include "Animation/AnimSequence.h"
-#include "Animation/BlendSpaceBase.h"
+#include "Animation/BlendSpace.h"
 #include "Animation/BlendSpace1D.h"
 #include "SAnimationBlendSpaceGridWidget.h"
 #include "PropertyCustomizationHelpers.h"
@@ -33,7 +33,7 @@
 
 #define LOCTEXT_NAMESPACE "BlendSampleDetails"
 
-FBlendSampleDetails::FBlendSampleDetails(const UBlendSpaceBase* InBlendSpace, class SBlendSpaceGridWidget* InGridWidget, int32 InSampleIndex)
+FBlendSampleDetails::FBlendSampleDetails(const UBlendSpace* InBlendSpace, class SBlendSpaceGridWidget* InGridWidget, int32 InSampleIndex)
 	: BlendSpace(InBlendSpace)
 	, GridWidget(InGridWidget)
 	, SampleIndex(InSampleIndex)
@@ -77,18 +77,18 @@ void FBlendSampleDetails::CustomizeDetails(class IDetailLayoutBuilder& DetailBui
 	TArray<UObject*> Objects;
 	UPackageTools::GetObjectsInPackages(&Packages, Objects);
 
-	UBlendSpaceBase* BlendSpaceBase = nullptr;
+	UBlendSpace* BlendSpaceBase = nullptr;
 	// Find blendspace in found objects
 	for ( UObject* Object : Objects )
 	{
-		BlendSpaceBase = Cast<UBlendSpaceBase>(Object);
+		BlendSpaceBase = Cast<UBlendSpace>(Object);
 		if (BlendSpaceBase)
 		{
 			break;
 		}
 	}
 
-	const UBlendSpaceBase* BlendspaceToUse = BlendSpaceBase != nullptr ? BlendSpaceBase : BlendSpace;
+	const UBlendSpace* BlendspaceToUse = BlendSpaceBase != nullptr ? BlendSpaceBase : BlendSpace;
 	UAnimGraphNode_BlendSpaceGraphBase* BlendSpaceNode = nullptr;
 	if(UBlendSpaceGraph* BlendSpaceGraph = Cast<UBlendSpaceGraph>(BlendspaceToUse->GetOuter()))
 	{
@@ -119,7 +119,7 @@ void FBlendSampleDetails::CustomizeDetails(class IDetailLayoutBuilder& DetailBui
 	CategoryBuilder.AddProperty(SnappedProperty);
 }
 
-void FBlendSampleDetails::GenerateBlendSampleWidget(TFunction<FDetailWidgetRow& (void)> InFunctor, FOnSampleMoved OnSampleMoved, const UBlendSpaceBase* BlendSpace, const int32 SampleIndex, bool bShowLabel)
+void FBlendSampleDetails::GenerateBlendSampleWidget(TFunction<FDetailWidgetRow& (void)> InFunctor, FOnSampleMoved OnSampleMoved, const UBlendSpace* BlendSpace, const int32 SampleIndex, bool bShowLabel)
 {
 	const int32 NumParameters = BlendSpace->IsA<UBlendSpace1D>() ? 1 : 2;
 	for (int32 ParameterIndex = 0; ParameterIndex < NumParameters; ++ParameterIndex)
@@ -209,7 +209,7 @@ void FBlendSampleDetails::GenerateBlendSampleWidget(TFunction<FDetailWidgetRow&
 	}
 }
 
-void FBlendSampleDetails::GenerateAnimationWidget(FDetailWidgetRow& Row, const UBlendSpaceBase* BlendSpace, TSharedPtr<IPropertyHandle> AnimationProperty)
+void FBlendSampleDetails::GenerateAnimationWidget(FDetailWidgetRow& Row, const UBlendSpace* BlendSpace, TSharedPtr<IPropertyHandle> AnimationProperty)
 {
 	Row.NameContent()
 		[
@@ -278,7 +278,7 @@ void FBlendSampleDetails::GenerateSampleGraphWidget(FDetailWidgetRow& InRow, UAn
 		];
 }
 
-bool FBlendSampleDetails::ShouldFilterAssetStatic(const FAssetData& AssetData, const UBlendSpaceBase* BlendSpaceBase)
+bool FBlendSampleDetails::ShouldFilterAssetStatic(const FAssetData& AssetData, const UBlendSpace* BlendSpaceBase)
 {
 	/** Cached flags to check whether or not an additive animation type is compatible with the blend space*/
 	TMap<FString, bool> bValidAdditiveTypes;
diff --git a/Engine/Source/Editor/Persona/Private/Customization/BlendSampleDetails.h b/Engine/Source/Editor/Persona/Private/Customization/BlendSampleDetails.h
--- a/Engine/Source/Editor/Persona/Private/Customization/BlendSampleDetails.h
+++ b/Engine/Source/Editor/Persona/Private/Customization/BlendSampleDetails.h
@@ -11,15 +11,15 @@
 
 struct FAssetData;
 class FDetailWidgetRow;
-class UBlendSpaceBase;
+class UBlendSpace;
 class UAnimGraphNode_BlendSpaceGraphBase;
 
 class FBlendSampleDetails : public IDetailCustomization
 {
 public:
-	FBlendSampleDetails(const class UBlendSpaceBase* InBlendSpace, class SBlendSpaceGridWidget* InGridWidget, int32 InSampleIndex);
+	FBlendSampleDetails(const class UBlendSpace* InBlendSpace, class SBlendSpaceGridWidget* InGridWidget, int32 InSampleIndex);
 
-	static TSharedRef<IDetailCustomization> MakeInstance(const class UBlendSpaceBase* InBlendSpace, class SBlendSpaceGridWidget* InGridWidget, int32 InSampleIndex)
+	static TSharedRef<IDetailCustomization> MakeInstance(const class UBlendSpace* InBlendSpace, class SBlendSpaceGridWidget* InGridWidget, int32 InSampleIndex)
 	{
 		return MakeShareable( new FBlendSampleDetails(InBlendSpace, InGridWidget, InSampleIndex) );
 	}
@@ -28,20 +28,20 @@ public:
 	virtual void CustomizeDetails(class IDetailLayoutBuilder& DetailBuilder) override;
 	// End IDetailCustomization interface
 	
-	static void GenerateBlendSampleWidget(TFunction<FDetailWidgetRow& (void)>InFunctor, FOnSampleMoved OnSampleMoved, const class UBlendSpaceBase* BlendSpace, const int32 SampleIndex, bool bShowLabel);
+	static void GenerateBlendSampleWidget(TFunction<FDetailWidgetRow& (void)>InFunctor, FOnSampleMoved OnSampleMoved, const class UBlendSpace* BlendSpace, const int32 SampleIndex, bool bShowLabel);
 
-	static void GenerateAnimationWidget(FDetailWidgetRow& Row, const UBlendSpaceBase* BlendSpace, TSharedPtr<IPropertyHandle> AnimationProperty);
+	static void GenerateAnimationWidget(FDetailWidgetRow& Row, const UBlendSpace* BlendSpace, TSharedPtr<IPropertyHandle> AnimationProperty);
 
 	static void GenerateSampleGraphWidget(FDetailWidgetRow& Row, UAnimGraphNode_BlendSpaceGraphBase* BlendSpaceNode, int32 SampleIndex);
 
-	static bool ShouldFilterAssetStatic(const FAssetData& AssetData, const UBlendSpaceBase* BlendSpaceBase);
+	static bool ShouldFilterAssetStatic(const FAssetData& AssetData, const UBlendSpace* BlendSpaceBase);
 
 protected:
 	/** Checks whether or not the specified asset should not be shown in the mini content browser when changing the animation */
 	bool ShouldFilterAsset(const FAssetData& AssetData) const;
 private:
 	/** Pointer to the current parent blend space for the customized blend sample*/
-	const class UBlendSpaceBase* BlendSpace;
+	const class UBlendSpace* BlendSpace;
 	/** Parent grid widget object */
 	SBlendSpaceGridWidget* GridWidget;
 	/** Current sample index */
diff --git a/Engine/Source/Editor/Persona/Private/Customization/BlendSpaceDetails.cpp b/Engine/Source/Editor/Persona/Private/Customization/BlendSpaceDetails.cpp
--- a/Engine/Source/Editor/Persona/Private/Customization/BlendSpaceDetails.cpp
+++ b/Engine/Source/Editor/Persona/Private/Customization/BlendSpaceDetails.cpp
@@ -8,7 +8,7 @@
 #include "DetailLayoutBuilder.h"
 #include "DetailWidgetRow.h"
 
-#include "Animation/BlendSpaceBase.h"
+#include "Animation/BlendSpace.h"
 #include "Animation/BlendSpace1D.h"
 
 #include "BlendSampleDetails.h"
@@ -38,10 +38,10 @@ void FBlendSpaceDetails::CustomizeDetails(class IDetailLayoutBuilder& DetailBuil
 	DetailBuilder.GetObjectsBeingCustomized(Objects);
 
 	Builder = &DetailBuilder;
-	TWeakObjectPtr<UObject>* WeakPtr = Objects.FindByPredicate([](const TWeakObjectPtr<UObject>& ObjectPtr) { return ObjectPtr->IsA<UBlendSpaceBase>(); });
+	TWeakObjectPtr<UObject>* WeakPtr = Objects.FindByPredicate([](const TWeakObjectPtr<UObject>& ObjectPtr) { return ObjectPtr->IsA<UBlendSpace>(); });
 	if (WeakPtr)
 	{
-		BlendSpaceBase = Cast<UBlendSpaceBase>(WeakPtr->Get());
+		BlendSpaceBase = Cast<UBlendSpace>(WeakPtr->Get());
 
 		if(!BlendSpaceBase->IsAsset())
 		{
@@ -60,6 +60,12 @@ void FBlendSpaceDetails::CustomizeDetails(class IDetailLayoutBuilder& DetailBuil
 		}
 		const bool b1DBlendSpace = BlendSpaceBase->IsA<UBlendSpace1D>();
 
+		if (b1DBlendSpace)
+		{
+			TSharedPtr<IPropertyHandle> AxisToScaleAnimation = DetailBuilder.GetProperty(GET_MEMBER_NAME_CHECKED(UBlendSpace, AxisToScaleAnimation), UBlendSpace::StaticClass());
+			DetailBuilder.HideProperty(AxisToScaleAnimation);
+		}
+
 		IDetailCategoryBuilder& CategoryBuilder = DetailBuilder.EditCategory(FName("Axis Settings"));
 		IDetailGroup* Groups[2] =
 		{
@@ -68,8 +74,8 @@ void FBlendSpaceDetails::CustomizeDetails(class IDetailLayoutBuilder& DetailBuil
 		};
 
 		// Hide the default blend and interpolation parameters
-		TSharedPtr<IPropertyHandle> BlendParameters = DetailBuilder.GetProperty(GET_MEMBER_NAME_CHECKED(UBlendSpaceBase, BlendParameters), UBlendSpaceBase::StaticClass());
-		TSharedPtr<IPropertyHandle> InterpolationParameters = DetailBuilder.GetProperty(GET_MEMBER_NAME_CHECKED(UBlendSpaceBase, InterpolationParam), UBlendSpaceBase::StaticClass());
+		TSharedPtr<IPropertyHandle> BlendParameters = DetailBuilder.GetProperty(GET_MEMBER_NAME_CHECKED(UBlendSpace, BlendParameters), UBlendSpace::StaticClass());
+		TSharedPtr<IPropertyHandle> InterpolationParameters = DetailBuilder.GetProperty(GET_MEMBER_NAME_CHECKED(UBlendSpace, InterpolationParam), UBlendSpace::StaticClass());
 		DetailBuilder.HideProperty(BlendParameters);
 		DetailBuilder.HideProperty(InterpolationParameters);
 
@@ -103,7 +109,7 @@ void FBlendSpaceDetails::CustomizeDetails(class IDetailLayoutBuilder& DetailBuil
 		FSimpleDelegate RefreshDelegate = FSimpleDelegate::CreateLambda([this]() { Builder->ForceRefreshDetails(); });
 
 		// Retrieve blend samples array
-		TSharedPtr<IPropertyHandleArray> BlendSamplesArrayProperty = DetailBuilder.GetProperty(GET_MEMBER_NAME_CHECKED(UBlendSpaceBase, SampleData), UBlendSpaceBase::StaticClass())->AsArray();
+		TSharedPtr<IPropertyHandleArray> BlendSamplesArrayProperty = DetailBuilder.GetProperty(GET_MEMBER_NAME_CHECKED(UBlendSpace, SampleData), UBlendSpace::StaticClass())->AsArray();
 		BlendSamplesArrayProperty->SetOnNumElementsChanged(RefreshDelegate);
 		
 		uint32 NumBlendSampleEntries = 0;
diff --git a/Engine/Source/Editor/Persona/Private/Customization/BlendSpaceDetails.h b/Engine/Source/Editor/Persona/Private/Customization/BlendSpaceDetails.h
--- a/Engine/Source/Editor/Persona/Private/Customization/BlendSpaceDetails.h
+++ b/Engine/Source/Editor/Persona/Private/Customization/BlendSpaceDetails.h
@@ -8,7 +8,7 @@
 #include "UObject/WeakObjectPtrTemplates.h"
 
 class IDetailLayoutBuilder;
-class UBlendSpaceBase;
+class UBlendSpace;
 class UAnimGraphNode_BlendSpaceGraphBase;
 
 class FBlendSpaceDetails : public IDetailCustomization
@@ -26,6 +26,6 @@ public:
 	virtual void CustomizeDetails(class IDetailLayoutBuilder& DetailBuilder) override;
 private:
 	IDetailLayoutBuilder* Builder;
-	UBlendSpaceBase* BlendSpaceBase;
+	UBlendSpace* BlendSpaceBase;
 	TWeakObjectPtr<UAnimGraphNode_BlendSpaceGraphBase> BlendSpaceNode;
 };
diff --git a/Engine/Source/Editor/Persona/Private/Customization/SkeletalMeshRegionCustomization.cpp b/Engine/Source/Editor/Persona/Private/Customization/SkeletalMeshRegionCustomization.cpp
--- a/Engine/Source/Editor/Persona/Private/Customization/SkeletalMeshRegionCustomization.cpp
+++ b/Engine/Source/Editor/Persona/Private/Customization/SkeletalMeshRegionCustomization.cpp
@@ -15,7 +15,7 @@
 #include "Widgets/Input/SNumericEntryBox.h"
 
 #include "Animation/AnimSequence.h"
-#include "Animation/BlendSpaceBase.h"
+#include "Animation/BlendSpace.h"
 #include "Animation/BlendSpace1D.h"
 #include "SAnimationBlendSpaceGridWidget.h"
 #include "PropertyCustomizationHelpers.h"
diff --git a/Engine/Source/Editor/Persona/Private/PersonaModule.cpp b/Engine/Source/Editor/Persona/Private/PersonaModule.cpp
--- a/Engine/Source/Editor/Persona/Private/PersonaModule.cpp
+++ b/Engine/Source/Editor/Persona/Private/PersonaModule.cpp
@@ -40,7 +40,6 @@
 #include "Animation/BlendSpace.h"
 #include "SAnimationBlendSpace.h"
 #include "Animation/BlendSpace1D.h"
-#include "SAnimationBlendSpace1D.h"
 #include "Animation/AimOffsetBlendSpace.h"
 #include "Animation/AimOffsetBlendSpace1D.h"
 #include "SAnimationDlgs.h"
@@ -119,7 +118,7 @@ void FPersonaModule::StartupModule()
 		PropertyModule.RegisterCustomClassLayout( "EditorNotifyObject", FOnGetDetailCustomizationInstance::CreateStatic(&FAnimNotifyDetails::MakeInstance));
 		PropertyModule.RegisterCustomClassLayout( "AnimGraphNode_Base", FOnGetDetailCustomizationInstance::CreateStatic( &FAnimGraphNodeDetails::MakeInstance ) );
 		PropertyModule.RegisterCustomClassLayout( "AnimInstance", FOnGetDetailCustomizationInstance::CreateStatic(&FAnimInstanceDetails::MakeInstance));
-		PropertyModule.RegisterCustomClassLayout("BlendSpaceBase", FOnGetDetailCustomizationInstance::CreateStatic(&FBlendSpaceDetails::MakeInstance));	
+		PropertyModule.RegisterCustomClassLayout("BlendSpace", FOnGetDetailCustomizationInstance::CreateStatic(&FBlendSpaceDetails::MakeInstance));	
 
 		PropertyModule.RegisterCustomPropertyTypeLayout("BlendProfile", FOnGetPropertyTypeCustomizationInstance::CreateStatic(&FBlendProfileCustomization::MakeInstance));
 		PropertyModule.RegisterCustomPropertyTypeLayout( "InputScaleBias", FOnGetPropertyTypeCustomizationInstance::CreateStatic( &FInputScaleBiasCustomization::MakeInstance ) );
@@ -164,7 +163,7 @@ void FPersonaModule::ShutdownModule()
 		PropertyModule.UnregisterCustomClassLayout("EditorNotifyObject");
 		PropertyModule.UnregisterCustomClassLayout("AnimGraphNode_Base");
 		PropertyModule.UnregisterCustomClassLayout("AnimInstance");
-		PropertyModule.UnregisterCustomClassLayout("BlendSpaceBase");
+		PropertyModule.UnregisterCustomClassLayout("BlendSpace");
 
 		PropertyModule.UnregisterCustomPropertyTypeLayout("InputScaleBias");
 		PropertyModule.UnregisterCustomPropertyTypeLayout("BoneReference");
@@ -353,7 +352,7 @@ TSharedRef<SWidget> FPersonaModule::CreateBlendSpacePreviewWidget(const FBlendSp
 		.OnGetBlendSpaceSampleName(InArgs.OnGetBlendSpaceSampleName);
 }
 
-TSharedRef<SWidget> FPersonaModule::CreateBlendSpacePreviewWidget(TAttribute<const UBlendSpaceBase*> InBlendSpace, TAttribute<FVector> InPosition, TAttribute<FVector> InFilteredPosition) const
+TSharedRef<SWidget> FPersonaModule::CreateBlendSpacePreviewWidget(TAttribute<const UBlendSpace*> InBlendSpace, TAttribute<FVector> InPosition, TAttribute<FVector> InFilteredPosition) const
 {
 PRAGMA_DISABLE_DEPRECATION_WARNINGS
 	FBlendSpacePreviewArgs Args;
@@ -456,8 +455,8 @@ TSharedRef<SWidget> FPersonaModule::CreateEditorWidgetForAnimDocument(const TSha
 		}
 		else if (UBlendSpace1D* BlendSpace1D = Cast<UBlendSpace1D>(InAnimAsset))
 		{
-			Result = SNew(SBlendSpaceEditor1D, InArgs.PreviewScene.Pin().ToSharedRef())
-				.BlendSpace1D(BlendSpace1D);
+			Result = SNew(SBlendSpaceEditor, InArgs.PreviewScene.Pin().ToSharedRef())
+				.BlendSpace(BlendSpace1D);
 
 			if (Cast<UAimOffsetBlendSpace1D>(InAnimAsset))
 			{
@@ -1494,42 +1493,21 @@ void FPersonaModule::HandleNewAnimNotifyStateBlueprintCreated(UBlueprint* InBlue
 	}
 }
 
-TSharedRef<SWidget> FPersonaModule::CreateBlendSpaceEditWidget(UBlendSpaceBase* InBlendSpace, const FBlendSpaceEditorArgs& InArgs) const
+TSharedRef<SWidget> FPersonaModule::CreateBlendSpaceEditWidget(UBlendSpace* InBlendSpace, const FBlendSpaceEditorArgs& InArgs) const
 {
-	if (UBlendSpace* BlendSpace = Cast<UBlendSpace>(InBlendSpace))
-	{
-		return SNew(SBlendSpaceEditor)
-			.BlendSpace(BlendSpace)
-			.DisplayScrubBar(false)
-			.OnBlendSpaceSampleDoubleClicked(InArgs.OnBlendSpaceSampleDoubleClicked)
-			.OnBlendSpaceSampleAdded(InArgs.OnBlendSpaceSampleAdded)
-			.OnBlendSpaceSampleRemoved(InArgs.OnBlendSpaceSampleRemoved)
-			.OnBlendSpaceSampleReplaced(InArgs.OnBlendSpaceSampleReplaced)
-			.OnGetBlendSpaceSampleName(InArgs.OnGetBlendSpaceSampleName)
-			.OnExtendSampleTooltip(InArgs.OnExtendSampleTooltip)
-			.OnSetPreviewPosition(InArgs.OnSetPreviewPosition)
-			.PreviewPosition(InArgs.PreviewPosition)
-			.PreviewFilteredPosition(InArgs.PreviewFilteredPosition)
-			.StatusBarName(InArgs.StatusBarName);
-	}
-	else if (UBlendSpace1D* BlendSpace1D = Cast<UBlendSpace1D>(InBlendSpace))
-	{
-		return SNew(SBlendSpaceEditor1D)
-			.BlendSpace1D(BlendSpace1D)
-			.DisplayScrubBar(false)
-			.OnBlendSpaceSampleDoubleClicked(InArgs.OnBlendSpaceSampleDoubleClicked)
-			.OnBlendSpaceSampleAdded(InArgs.OnBlendSpaceSampleAdded)
-			.OnBlendSpaceSampleRemoved(InArgs.OnBlendSpaceSampleRemoved)
-			.OnBlendSpaceSampleReplaced(InArgs.OnBlendSpaceSampleReplaced)
-			.OnGetBlendSpaceSampleName(InArgs.OnGetBlendSpaceSampleName)
-			.OnExtendSampleTooltip(InArgs.OnExtendSampleTooltip)
-			.OnSetPreviewPosition(InArgs.OnSetPreviewPosition)
-			.PreviewPosition(InArgs.PreviewPosition)
-			.PreviewFilteredPosition(InArgs.PreviewFilteredPosition)
-			.StatusBarName(InArgs.StatusBarName);
-	}
-
-	return SNullWidget::NullWidget;
+	return SNew(SBlendSpaceEditor)
+		.BlendSpace(InBlendSpace)
+		.DisplayScrubBar(false)
+		.OnBlendSpaceSampleDoubleClicked(InArgs.OnBlendSpaceSampleDoubleClicked)
+		.OnBlendSpaceSampleAdded(InArgs.OnBlendSpaceSampleAdded)
+		.OnBlendSpaceSampleRemoved(InArgs.OnBlendSpaceSampleRemoved)
+		.OnBlendSpaceSampleReplaced(InArgs.OnBlendSpaceSampleReplaced)
+		.OnGetBlendSpaceSampleName(InArgs.OnGetBlendSpaceSampleName)
+		.OnExtendSampleTooltip(InArgs.OnExtendSampleTooltip)
+		.OnSetPreviewPosition(InArgs.OnSetPreviewPosition)
+		.PreviewPosition(InArgs.PreviewPosition)
+		.PreviewFilteredPosition(InArgs.PreviewFilteredPosition)
+		.StatusBarName(InArgs.StatusBarName);
 }
 
 #undef LOCTEXT_NAMESPACE
diff --git a/Engine/Source/Editor/Persona/Private/SAnimEditorBase.cpp b/Engine/Source/Editor/Persona/Private/SAnimEditorBase.cpp
--- a/Engine/Source/Editor/Persona/Private/SAnimEditorBase.cpp
+++ b/Engine/Source/Editor/Persona/Private/SAnimEditorBase.cpp
@@ -11,7 +11,7 @@
 #include "IDocumentation.h"
 
 #include "AnimPreviewInstance.h"
-#include "Animation/BlendSpaceBase.h"
+#include "Animation/BlendSpace.h"
 #include "AnimModel.h"
 #include "SAnimTimeline.h"
 
@@ -175,7 +175,7 @@ float SAnimEditorBase::GetSequenceLength() const
 	{
 		return AnimSeqBase->GetPlayLength();
 	}
-	else if (UBlendSpaceBase* BlendSpaceBase = Cast<UBlendSpaceBase>(GetEditorObject()))
+	else if (UBlendSpace* BlendSpaceBase = Cast<UBlendSpace>(GetEditorObject()))
 	{
 		// Blendspaces use normalized time, so we just return 1 here
 		return 1.0f;
diff --git a/Engine/Source/Editor/Persona/Private/SAnimNotifyPanel.cpp b/Engine/Source/Editor/Persona/Private/SAnimNotifyPanel.cpp
--- a/Engine/Source/Editor/Persona/Private/SAnimNotifyPanel.cpp
+++ b/Engine/Source/Editor/Persona/Private/SAnimNotifyPanel.cpp
@@ -32,7 +32,7 @@
 #include "BlueprintActionDatabase.h"
 #include "Animation/AnimNotifies/AnimNotifyState.h"
 #include "Animation/AnimNotifies/AnimNotify.h"
-#include "Animation/BlendSpaceBase.h"
+#include "Animation/BlendSpace.h"
 #include "TabSpawners.h"
 #include "Widgets/Text/SInlineEditableTextBlock.h"
 #include "HAL/PlatformApplicationMisc.h"
@@ -2591,7 +2591,7 @@ void SAnimNotifyTrack::CreateNewSyncMarkerAtCursor(FString NewSyncMarkerName, UC
 	Seq->MarkPackageDirty();
 	OnUpdatePanel.ExecuteIfBound();
 
-	UBlendSpaceBase::UpdateBlendSpacesUsingAnimSequence(Seq);
+	UBlendSpace::UpdateBlendSpacesUsingAnimSequence(Seq);
 }
 
 void SAnimNotifyTrack::ReplaceSelectedWithBlueprintNotify(FString NewNotifyName, FString BlueprintPath)
@@ -3522,7 +3522,7 @@ void SAnimNotifyTrack::HandleNodeDrop(TSharedPtr<SAnimNotifyNode> Node, float Of
 
 	if (Node->NodeObjectInterface->GetType() == ENodeObjectTypes::SYNC_MARKER)
 	{
-		UBlendSpaceBase::UpdateBlendSpacesUsingAnimSequence(Sequence);
+		UBlendSpace::UpdateBlendSpacesUsingAnimSequence(Sequence);
 	}
 
 	float LocalX = GetCachedGeometry().AbsoluteToLocal(Node->GetScreenPosition() + Offset).X;
@@ -3689,7 +3689,7 @@ void SAnimNotifyTrack::PasteSingleSyncMarker(FString& MarkerString, float PasteT
 			AnimSeq->AuthoredSyncMarkers.RemoveAt(NewIdx);
 		}
 
-		UBlendSpaceBase::UpdateBlendSpacesUsingAnimSequence(Sequence);
+		UBlendSpace::UpdateBlendSpacesUsingAnimSequence(Sequence);
 
 		OnDeselectAllNotifies.ExecuteIfBound();
 		Sequence->PostEditChange();
@@ -4300,7 +4300,7 @@ void SAnimNotifyPanel::DeleteSelectedNodeObjects()
 
 		if (bContainsSyncMarkers)
 		{
-			UBlendSpaceBase::UpdateBlendSpacesUsingAnimSequence(Sequence);
+			UBlendSpace::UpdateBlendSpacesUsingAnimSequence(Sequence);
 		}
 	}
 
diff --git a/Engine/Source/Editor/Persona/Private/SAnimationBlendSpace.cpp b/Engine/Source/Editor/Persona/Private/SAnimationBlendSpace.cpp
--- a/Engine/Source/Editor/Persona/Private/SAnimationBlendSpace.cpp
+++ b/Engine/Source/Editor/Persona/Private/SAnimationBlendSpace.cpp
@@ -2,62 +2,405 @@
 
 #include "SAnimationBlendSpace.h"
 #include "Widgets/Notifications/SNotificationList.h"
-#include "Widgets/Text/STextBlock.h"
+#include "Editor.h"
+
 #include "SlateOptMacros.h"
-#include "Widgets/Layout/SScrollBox.h"
+#include "Fonts/FontMeasure.h"
+#include "Framework/Application/SlateApplication.h"
+#include "Widgets/Input/SEditableTextBox.h"
+#include "Widgets/Input/SButton.h"
+#include "Animation/DebugSkelMeshComponent.h"
+#include "IContentBrowserSingleton.h"
+#include "ContentBrowserModule.h"
+#include "DragAndDrop/AssetDragDropOp.h"
+#include "ScopedTransaction.h"
+#include "AnimPreviewInstance.h"
+#include "Widgets/Layout/SExpandableArea.h"
+#include "Animation/AimOffsetBlendSpace.h"
+#include "Animation/AimOffsetBlendSpace1D.h"
+#include "AnimationBlendSpaceHelpers.h"
+#include "AnimationBlendSpace1DHelpers.h"
+#include "SAnimationBlendSpaceGridWidget.h"
 
 #define LOCTEXT_NAMESPACE "BlendSpaceEditor"
 
+SBlendSpaceEditor::~SBlendSpaceEditor()
+{
+	FCoreUObjectDelegates::OnObjectPropertyChanged.Remove(OnPropertyChangedHandleDelegateHandle);
+}
+
 void SBlendSpaceEditor::Construct(const FArguments& InArgs)
 {
-	SBlendSpaceEditorBase::Construct(SBlendSpaceEditorBase::FArguments()
-		.BlendSpace(InArgs._BlendSpace)
-		.DisplayScrubBar(InArgs._DisplayScrubBar)
-		.OnBlendSpaceSampleDoubleClicked(InArgs._OnBlendSpaceSampleDoubleClicked)
-		.OnBlendSpaceSampleAdded(InArgs._OnBlendSpaceSampleAdded)
-		.OnBlendSpaceSampleRemoved(InArgs._OnBlendSpaceSampleRemoved)
-		.OnBlendSpaceSampleReplaced(InArgs._OnBlendSpaceSampleReplaced)
-		.OnGetBlendSpaceSampleName(InArgs._OnGetBlendSpaceSampleName)
-		.OnExtendSampleTooltip(InArgs._OnExtendSampleTooltip)
-		.OnSetPreviewPosition(InArgs._OnSetPreviewPosition)
-		.PreviewPosition(InArgs._PreviewPosition)
-		.PreviewFilteredPosition(InArgs._PreviewFilteredPosition)
-		.StatusBarName(InArgs._StatusBarName));
+	BlendSpace = InArgs._BlendSpace;
+
+	OnBlendSpaceSampleAdded = InArgs._OnBlendSpaceSampleAdded;
+	OnBlendSpaceSampleRemoved = InArgs._OnBlendSpaceSampleRemoved;
+	OnBlendSpaceSampleReplaced = InArgs._OnBlendSpaceSampleReplaced;
+	OnSetPreviewPosition = InArgs._OnSetPreviewPosition;
+
+	bShouldSetPreviewPosition = false;
+
+	SAnimEditorBase::Construct(SAnimEditorBase::FArguments()
+		.DisplayAnimTimeline(false)
+		.DisplayAnimScrubBar(InArgs._DisplayScrubBar),
+		PreviewScenePtr.Pin());
+
+	NonScrollEditorPanels->AddSlot()
+	[
+		SNew(SVerticalBox)
+		+SVerticalBox::Slot()
+		.FillHeight(1.0f)
+		[
+			SNew(SOverlay)
+			+ SOverlay::Slot()
+			[
+				SNew(SHorizontalBox)
+				+SHorizontalBox::Slot()
+				.FillWidth(1)
+				.Padding(4.0f)
+				[
+					SNew(SHorizontalBox)
+					+SHorizontalBox::Slot()
+					.FillWidth(1)
+					.Padding(2) 
+					[
+						SNew(SVerticalBox)
+						// Grid area
+						+SVerticalBox::Slot()
+						.FillHeight(1)
+						[
+							SAssignNew(NewBlendSpaceGridWidget, SBlendSpaceGridWidget)
+							.Cursor(EMouseCursor::Crosshairs)
+							.BlendSpaceBase(BlendSpace)
+							.NotifyHook(this)
+							.Position(InArgs._PreviewPosition)
+							.FilteredPosition(InArgs._PreviewFilteredPosition)
+							.OnSampleMoved(this, &SBlendSpaceEditor::OnSampleMoved)
+							.OnSampleRemoved(this, &SBlendSpaceEditor::OnSampleRemoved)
+							.OnSampleAdded(this, &SBlendSpaceEditor::OnSampleAdded)
+							.OnSampleReplaced(this, &SBlendSpaceEditor::OnSampleReplaced)
+							.OnSampleDoubleClicked(InArgs._OnBlendSpaceSampleDoubleClicked)
+							.OnExtendSampleTooltip(InArgs._OnExtendSampleTooltip)
+							.OnGetBlendSpaceSampleName(InArgs._OnGetBlendSpaceSampleName)
+							.StatusBarName(InArgs._StatusBarName)
+						]
+					]
+				]
+			]
+		]
+	];
+
+	OnPropertyChangedHandle = FCoreUObjectDelegates::FOnObjectPropertyChanged::FDelegate::CreateRaw(this, &SBlendSpaceEditor::OnPropertyChanged);
+	OnPropertyChangedHandleDelegateHandle = FCoreUObjectDelegates::OnObjectPropertyChanged.Add(OnPropertyChangedHandle);
 }
 
 void SBlendSpaceEditor::Construct(const FArguments& InArgs, const TSharedRef<class IPersonaPreviewScene>& InPreviewScene)
 {
-	SBlendSpaceEditorBase::Construct(SBlendSpaceEditorBase::FArguments()
-		.BlendSpace(InArgs._BlendSpace)
-		.DisplayScrubBar(InArgs._DisplayScrubBar)
-		.OnBlendSpaceSampleDoubleClicked(InArgs._OnBlendSpaceSampleDoubleClicked)
-		.OnBlendSpaceSampleAdded(InArgs._OnBlendSpaceSampleAdded)
-		.OnBlendSpaceSampleRemoved(InArgs._OnBlendSpaceSampleRemoved)
-		.OnBlendSpaceSampleReplaced(InArgs._OnBlendSpaceSampleReplaced)
-		.OnGetBlendSpaceSampleName(InArgs._OnGetBlendSpaceSampleName)
-		.OnExtendSampleTooltip(InArgs._OnExtendSampleTooltip)
-		.OnSetPreviewPosition(InArgs._OnSetPreviewPosition)
-		.PreviewPosition(InArgs._PreviewPosition)
-		.PreviewFilteredPosition(InArgs._PreviewFilteredPosition)
-		.StatusBarName(InArgs._StatusBarName),
-		InPreviewScene);
+	PreviewScenePtr = InPreviewScene;
+
+	Construct(InArgs);
+}
+
+void SBlendSpaceEditor::OnSampleMoved(const int32 SampleIndex, const FVector& NewValue, bool bIsInteractive, bool bSnap)
+{
+	bool bMoveSuccesful = true;
+	if (BlendSpace->IsValidBlendSampleIndex(SampleIndex) && BlendSpace->GetBlendSample(SampleIndex).SampleValue != NewValue && !BlendSpace->IsTooCloseToExistingSamplePoint(NewValue, SampleIndex))
+	{
+		FScopedTransaction ScopedTransaction(LOCTEXT("MoveSample", "Moving Blend Grid Sample"));
+		BlendSpace->Modify();
+
+		bMoveSuccesful = BlendSpace->EditSampleValue(SampleIndex, NewValue, bSnap);
+		if (bMoveSuccesful)
+		{
+			BlendSpace->ValidateSampleData();
+			ResampleData();
+		}
+	}
+}
+
+void SBlendSpaceEditor::OnSampleRemoved(const int32 SampleIndex)
+{
+	FScopedTransaction ScopedTransaction(LOCTEXT("RemoveSample", "Removing Blend Grid Sample"));
+	BlendSpace->Modify();
+
+	const bool bRemoveSuccesful = BlendSpace->DeleteSample(SampleIndex);
+	if (bRemoveSuccesful)
+	{
+		ResampleData();
+		BlendSpace->ValidateSampleData();
+
+		OnBlendSpaceSampleRemoved.ExecuteIfBound(SampleIndex);
+	}
+	BlendSpace->PostEditChange();
+}
+
+void SBlendSpaceEditor::OnSampleAdded(UAnimSequence* Animation, const FVector& Value)
+{
+	FScopedTransaction ScopedTransaction(LOCTEXT("AddSample", "Adding Blend Grid Sample"));
+	BlendSpace->Modify();
+
+	bool bAddSuccesful = false;
+
+	if(BlendSpace->IsAsset())
+	{
+		bAddSuccesful = BlendSpace->AddSample(Animation, Value);
+	}
+	else
+	{
+		bAddSuccesful = BlendSpace->AddSample(Value);
+	}
+
+	if (bAddSuccesful)
+	{
+		ResampleData();
+		BlendSpace->ValidateSampleData();
+
+		OnBlendSpaceSampleAdded.ExecuteIfBound(Animation, Value);
+	}
+	BlendSpace->PostEditChange();
+}
+
+void SBlendSpaceEditor::OnSampleReplaced(int32 InSampleIndex, UAnimSequence* Animation)
+{
+	FScopedTransaction ScopedTransaction(LOCTEXT("UpdateAnimation", "Changing Animation Sequence"));
+	BlendSpace->Modify();
+
+	bool bUpdateSuccesful = false;
+	if(BlendSpace->IsAsset())
+	{
+		bUpdateSuccesful = BlendSpace->ReplaceSampleAnimation(InSampleIndex, Animation);
+	}
+	else
+	{
+		bUpdateSuccesful = true;
+	}
+
+	if (bUpdateSuccesful)
+	{
+		ResampleData();
+		BlendSpace->ValidateSampleData();
+
+		OnBlendSpaceSampleReplaced.ExecuteIfBound(InSampleIndex, Animation);
+	}
+}
+
+void SBlendSpaceEditor::PostUndoRedo()
+{
+	// Validate and resample blend space data
+	BlendSpace->ValidateSampleData();
+	ResampleData();
+
+	// Invalidate widget data
+	NewBlendSpaceGridWidget->InvalidateCachedData();
+
+	// Invalidate sample indices used for UI info
+	NewBlendSpaceGridWidget->InvalidateState();
+
+	// Set flag which will update the preview value in the next tick (this due the recreation of data after Undo)
+	bShouldSetPreviewPosition = true;
+}
+
+TSharedPtr<class IPersonaPreviewScene> SBlendSpaceEditor::GetPreviewScene() const
+{
+	return PreviewScenePtr.Pin();
+}
+
+void SBlendSpaceEditor::UpdatePreviewParameter() const
+{
+	if(GetPreviewScene().IsValid())
+	{
+		class UDebugSkelMeshComponent* Component = GetPreviewScene()->GetPreviewMeshComponent();
+
+		if (Component != nullptr && Component->IsPreviewOn())
+		{
+			if (Component->PreviewInstance->GetCurrentAsset() == BlendSpace)
+			{
+				const FVector PreviewPosition = NewBlendSpaceGridWidget->GetPreviewPosition();
+				Component->PreviewInstance->SetBlendSpacePosition(PreviewPosition);
+				GetPreviewScene()->InvalidateViews();			
+			}
+		}
+	}
+	else if(OnSetPreviewPosition.IsBound())
+	{
+		const FVector PreviewPosition = NewBlendSpaceGridWidget->GetPreviewPosition();
+		OnSetPreviewPosition.Execute(PreviewPosition);
+	}
+}
+
+void SBlendSpaceEditor::UpdateFromBlendSpaceState() const
+{
+	if (GetPreviewScene().IsValid())
+	{
+		class UDebugSkelMeshComponent* Component = GetPreviewScene()->GetPreviewMeshComponent();
+
+		if (Component != nullptr && Component->IsPreviewOn())
+		{
+			if (Component->PreviewInstance->GetCurrentAsset() == BlendSpace)
+			{
+				FVector FilteredPosition;
+				FVector Position;
+				Component->PreviewInstance->GetBlendSpaceState(Position, FilteredPosition);
+				NewBlendSpaceGridWidget->SetPreviewingState(Position, FilteredPosition);
+			}
+		}
+	}
+}
+
+void SBlendSpaceEditor::Tick(const FGeometry& AllottedGeometry, const double InCurrentTime, const float InDeltaTime)
+{
+	// Update the preview as long as its enabled
+	if (NewBlendSpaceGridWidget->IsPreviewing() || bShouldSetPreviewPosition)
+	{
+		UpdatePreviewParameter();
+		bShouldSetPreviewPosition = false;
+	}
+
+	UpdateFromBlendSpaceState();
+}
+
+void SBlendSpaceEditor::OnPropertyChanged(UObject* ObjectBeingModified, FPropertyChangedEvent& PropertyChangedEvent)
+{
+	if (ObjectBeingModified == BlendSpace)
+	{
+		BlendSpace->ValidateSampleData();
+		ResampleData();
+		NewBlendSpaceGridWidget->InvalidateCachedData();
+	}
+}
+
+void SBlendSpaceEditor::NotifyPreChange(FProperty* PropertyAboutToChange)
+{
+	if (BlendSpace)
+	{
+		BlendSpace->Modify();
+	}	
+}
+
+void SBlendSpaceEditor::NotifyPostChange(const FPropertyChangedEvent& PropertyChangedEvent, FProperty* PropertyThatChanged)
+{
+	if (BlendSpace)
+	{
+		BlendSpace->ValidateSampleData();
+		ResampleData();
+		BlendSpace->MarkPackageDirty();
+	}
 }
 
 void SBlendSpaceEditor::ResampleData()
 {
+	const TArray<FBlendSample>& BlendSamples = BlendSpace->GetBlendSamples();
+	FBox AABB(ForceInit);
+	for (const FBlendSample& Sample : BlendSamples)
+	{
+		// Add X value from sample (this is the only valid value to be set for 1D blend spaces / aim offsets
+		if (Sample.bIsValid)
+		{
+			AABB += Sample.SampleValue;
+		}
+	}
+
+	TArray<int32> DimensionIndices;
+	DimensionIndices.Reset(3);
+	if (AABB.GetExtent().X > 0.0f)
+	{
+		DimensionIndices.Push(0);
+	}
+	if (AABB.GetExtent().Y > 0.0f)
+	{
+		DimensionIndices.Push(1);
+	}
+	if (AABB.GetExtent().Z > 0.0f)
+	{
+		DimensionIndices.Push(2);
+	}
+
+	if (DimensionIndices.Num() == 1)
+	{
+		ResampleData1D(DimensionIndices);
+	}
+	else if (DimensionIndices.Num() == 2)
+	{
+		ResampleData2D(DimensionIndices);
+	}
+	else
+	{
+		UE_LOG(LogAnimation, Warning, TEXT("Found %d dimensions from the samples"), DimensionIndices.Num());
+	}
+}
+
+
+void SBlendSpaceEditor::ResampleData1D(const TArray<int32>& DimensionIndices)
+{
+	check(DimensionIndices.Num() == 1);
+	FLineElementGenerator LineElementGenerator;
+
+	int32 Index0 = DimensionIndices[0];
+
+	const FBlendParameter& BlendParameter = BlendSpace->GetBlendParameter(Index0);
+	LineElementGenerator.Init(BlendParameter);
+
+	const TArray<FBlendSample>& BlendSamples = BlendSpace->GetBlendSamples();
+	UE_LOG(LogAnimation, Log, TEXT("Resampling data in 1D - %d samples"), BlendSamples.Num());
+
+	if (BlendSamples.Num())
+	{
+		for (const FBlendSample& Sample : BlendSamples)
+		{
+			// Add X value from sample (this is the only valid value to be set for 1D blend spaces / aim offsets
+			if (Sample.bIsValid)
+			{
+				LineElementGenerator.SamplePointList.Add(Sample.SampleValue[Index0]);
+			}
+		}
+		LineElementGenerator.CalculateEditorElements();
+
+		// Create point to sample index list
+		TArray<int32> PointListToSampleIndices;
+		PointListToSampleIndices.Init(INDEX_NONE, LineElementGenerator.SamplePointList.Num());
+		for (int32 PointIndex = 0; PointIndex < LineElementGenerator.SamplePointList.Num(); ++PointIndex)
+		{
+			const float Point = LineElementGenerator.SamplePointList[PointIndex];
+			for (int32 SampleIndex = 0; SampleIndex < BlendSamples.Num(); ++SampleIndex)
+			{
+				if (BlendSamples[SampleIndex].SampleValue[Index0] == Point)
+				{
+					PointListToSampleIndices[PointIndex] = SampleIndex;
+					break;
+				}
+			}
+		}
+
+		BlendSpace->FillupGridElements(PointListToSampleIndices, LineElementGenerator.EditorElements, DimensionIndices);
+	}
+}
+
+
+void SBlendSpaceEditor::ResampleData2D(const TArray<int32>& DimensionIndices)
+{
+	check(DimensionIndices.Num() == 2);
+
+	// TODO for now, Index0 will always be 0 and Index1 will always be 1, since we don't support 3D.
+	// However, if/when we support authoring using a third dimension, we could get here with any 2D
+	// combination - e.g. XY, XZ or YZ. Then we can either make the triangulation code handle the
+	// indexing, or tweak the triangles going into and out of it.
+	int32 Index0 = DimensionIndices[0];
+	int32 Index1 = DimensionIndices[1];
+
 	// clear first
-	BlendSpaceGrid.Reset();
-	Generator.Reset();
+	FDelaunayTriangleGenerator DelaunayTriangleGenerator;
 
 	// you don't like to overwrite the link here (between visible points vs sample points, 
 	// so allow this if no triangle is generated
-	const FBlendParameter& BlendParamX = BlendSpace->GetBlendParameter(0);
-	const FBlendParameter& BlendParamY = BlendSpace->GetBlendParameter(1);
+	const FBlendParameter& BlendParamX = BlendSpace->GetBlendParameter(Index0);
+	const FBlendParameter& BlendParamY = BlendSpace->GetBlendParameter(Index1);
+	FBlendSpaceGrid	BlendSpaceGrid;
 	BlendSpaceGrid.SetGridInfo(BlendParamX, BlendParamY);
-	Generator.SetGridBox(BlendParamX, BlendParamY);
+	DelaunayTriangleGenerator.SetGridBox(BlendParamX, BlendParamY);
 
 	BlendSpace->EmptyGridElements();
 
+	UE_LOG(LogAnimation, Log, TEXT("Resampling data in 2D - %d samples"), BlendSpace->GetNumberOfBlendSamples());
 	if (BlendSpace->GetNumberOfBlendSamples())
 	{
 		bool bAllSamplesValid = true;
@@ -67,24 +410,24 @@ void SBlendSpaceEditor::ResampleData()
 
 			// Do not add invalid sample points (user will need to correct them to be incorporated into the blendspace)
 			if (Sample.bIsValid)
-			{				
-				Generator.AddSamplePoint(Sample.SampleValue, SampleIndex);
+			{
+				DelaunayTriangleGenerator.AddSamplePoint(Sample.SampleValue, SampleIndex);
 			}
-		}		
-		
+		}
+
 		// triangulate
-		Generator.Triangulate();
+		DelaunayTriangleGenerator.Triangulate();
 
 		// once triangulated, generate grid
-		const TArray<FPoint>& Points = Generator.GetSamplePointList();
-		const TArray<FTriangle*>& Triangles = Generator.GetTriangleList();
+		const TArray<FPoint>& Points = DelaunayTriangleGenerator.GetSamplePointList();
+		const TArray<FTriangle*>& Triangles = DelaunayTriangleGenerator.GetTriangleList();
 		BlendSpaceGrid.GenerateGridElements(Points, Triangles);
 
 		// now fill up grid elements in BlendSpace using this Element information
 		if (Triangles.Num() > 0)
 		{
 			const TArray<FEditorElement>& GridElements = BlendSpaceGrid.GetElements();
-			BlendSpace->FillupGridElements(Generator.GetIndiceMapping(), GridElements);
+			BlendSpace->FillupGridElements(DelaunayTriangleGenerator.GetIndiceMapping(), GridElements, DimensionIndices);
 		}
 	}
 }
diff --git a/Engine/Source/Editor/Persona/Private/SAnimationBlendSpace.h b/Engine/Source/Editor/Persona/Private/SAnimationBlendSpace.h
--- a/Engine/Source/Editor/Persona/Private/SAnimationBlendSpace.h
+++ b/Engine/Source/Editor/Persona/Private/SAnimationBlendSpace.h
@@ -3,19 +3,27 @@
 #pragma once
 
 #include "CoreMinimal.h"
+#include "SlateFwd.h"
+#include "UObject/UObjectGlobals.h"
 #include "Widgets/DeclarativeSyntaxSupport.h"
-#include "SAnimationBlendSpaceBase.h"
-#include "AnimationBlendSpaceHelpers.h"
+#include "IPersonaPreviewScene.h"
+#include "Rendering/DrawElements.h"
+#include "SAnimEditorBase.h"
+#include "Animation/BlendSpace.h"
+#include "Widgets/Input/SCheckBox.h"
+#include "Misc/NotifyHook.h"
+#include "EditorUndoClient.h"
 
-class SBlendSpaceEditor : public SBlendSpaceEditorBase
+class SBlendSpaceGridWidget;
+
+class SBlendSpaceEditor : public SAnimEditorBase, public FNotifyHook, public FSelfRegisteringEditorUndoClient
 {
 public:
 	SLATE_BEGIN_ARGS(SBlendSpaceEditor)
-		: _BlendSpace(nullptr)		
-		, _DisplayScrubBar(true)
+		: _BlendSpace(NULL)
 		, _StatusBarName(TEXT("AssetEditor.AnimationEditor.MainMenu"))
 		{}
-
+		
 	SLATE_ARGUMENT(UBlendSpace*, BlendSpace)
 
 	SLATE_ARGUMENT(bool, DisplayScrubBar)
@@ -28,7 +36,7 @@ public:
 
 	SLATE_EVENT(FOnBlendSpaceSampleReplaced, OnBlendSpaceSampleReplaced)
 
-	SLATE_EVENT(FOnGetBlendSpaceSampleName, OnGetBlendSpaceSampleName)
+	SLATE_EVENT(FOnGetBlendSpaceSampleName, OnGetBlendSpaceSampleName) 
 
 	SLATE_EVENT(FOnExtendBlendSpaceSampleTooltip, OnExtendSampleTooltip)
 
@@ -42,19 +50,82 @@ public:
 
 	SLATE_END_ARGS()
 
+	~SBlendSpaceEditor();
+
 	void Construct(const FArguments& InArgs);
 
 	void Construct(const FArguments& InArgs, const TSharedRef<class IPersonaPreviewScene>& InPreviewScene);
+	
+	// Begin SWidget overrides
+	virtual void Tick(const FGeometry& AllottedGeometry, const double InCurrentTime, const float InDeltaTime) override;
+	// End SWidget overrides
+
+	// Begin FNotifyHook overrides
+	virtual void NotifyPreChange(FProperty* PropertyAboutToChange) override;
+	virtual void NotifyPostChange(const FPropertyChangedEvent& PropertyChangedEvent, FProperty* PropertyThatChanged) override;
+	// End FNotifyHook overrides
 
 protected:
-	virtual void ResampleData() override;
-
-	/**
-	* Triangle Generator
-	 */
-	FDelaunayTriangleGenerator Generator;
-	/**
-	 * Blend Space Grid to represent data
-	 */
-	FBlendSpaceGrid	BlendSpaceGrid;
+	void OnSampleMoved(const int32 SampleIndex, const FVector& NewValue, bool bIsInteractive, bool bSnap);
+	void OnSampleRemoved(const int32 SampleIndex);
+	void OnSampleAdded(UAnimSequence* Animation, const FVector& Value);
+	void OnSampleReplaced(const int32 SampleIndex, UAnimSequence* Animation);
+	
+	// Begin SAnimEditorBase overrides
+	virtual UAnimationAsset* GetEditorObject() const override { return BlendSpace; }
+	// End SAnimEditorBase overrides
+
+	// FSelfRegisteringEditorUndoClient interface
+	virtual void PostUndo(bool bSuccess) override { PostUndoRedo(); }
+	virtual void PostRedo(bool bSuccess) override { PostUndoRedo(); }
+
+	/** Delegate which is called when the Editor has performed an undo operation*/
+	void PostUndoRedo();
+
+	/** Updates Persona's preview window */
+	void UpdatePreviewParameter() const;
+	
+	/** Updates Persona's preview window with additional data that should always be shown. */
+	void UpdateFromBlendSpaceState() const;
+
+	/** Retrieves the preview scene shown by Persona */
+	TSharedPtr<class IPersonaPreviewScene> GetPreviewScene() const;
+
+	/** Global callback to anticipate on changes to the blend space */
+	void OnPropertyChanged(UObject* ObjectBeingModified, FPropertyChangedEvent& PropertyChangedEvent);
+
+	void ResampleData();
+	void ResampleData1D(const TArray<int32>& DimensionIndices);
+	void ResampleData2D(const TArray<int32>& DimensionIndices);
+
+protected:
+	/** The blend space being edited */
+	UBlendSpace* BlendSpace;
+
+	/** The preview scene we are viewing */
+	TWeakPtr<class IPersonaPreviewScene> PreviewScenePtr;
+	
+	/** Pointer to the grid widget which displays the blendspace visualization */
+	TSharedPtr<class SBlendSpaceGridWidget> NewBlendSpaceGridWidget;	
+
+	// Property changed delegate
+	FCoreUObjectDelegates::FOnObjectPropertyChanged::FDelegate OnPropertyChangedHandle;	
+
+	/** Handle to the registered OnPropertyChangedHandle delegate */
+	FDelegateHandle OnPropertyChangedHandleDelegateHandle;
+
+	/** Delegate called when a sample is added */
+	FOnBlendSpaceSampleAdded OnBlendSpaceSampleAdded;
+
+	/** Delegate called when a sample is removed */
+	FOnBlendSpaceSampleRemoved OnBlendSpaceSampleRemoved;
+
+	/** Delegate called when a sample is replaced */
+	FOnBlendSpaceSampleReplaced OnBlendSpaceSampleReplaced;
+
+	/** Delegate called to externally control the preview position */
+	FOnSetBlendSpacePreviewPosition OnSetPreviewPosition;
+
+	/** Flag to check whether or not the preview value should be (re-)set on the next tick */
+	bool bShouldSetPreviewPosition;
 };
diff --git a/Engine/Source/Editor/Persona/Private/SAnimationBlendSpace1D.cpp b/Engine/Source/Editor/Persona/Private/SAnimationBlendSpace1D.cpp
deleted file mode 100644
diff --git a/Engine/Source/Editor/Persona/Private/SAnimationBlendSpace1D.h b/Engine/Source/Editor/Persona/Private/SAnimationBlendSpace1D.h
deleted file mode 100644
diff --git a/Engine/Source/Editor/Persona/Private/SAnimationBlendSpaceBase.cpp b/Engine/Source/Editor/Persona/Private/SAnimationBlendSpaceBase.cpp
deleted file mode 100644
diff --git a/Engine/Source/Editor/Persona/Private/SAnimationBlendSpaceBase.h b/Engine/Source/Editor/Persona/Private/SAnimationBlendSpaceBase.h
deleted file mode 100644
diff --git a/Engine/Source/Editor/Persona/Private/SAnimationBlendSpaceGridWidget.cpp b/Engine/Source/Editor/Persona/Private/SAnimationBlendSpaceGridWidget.cpp
--- a/Engine/Source/Editor/Persona/Private/SAnimationBlendSpaceGridWidget.cpp
+++ b/Engine/Source/Editor/Persona/Private/SAnimationBlendSpaceGridWidget.cpp
@@ -3,7 +3,7 @@
 #include "SAnimationBlendSpaceGridWidget.h"
 
 #include "Animation/AnimSequence.h"
-#include "Animation/BlendSpaceBase.h"
+#include "Animation/BlendSpace.h"
 #include "Animation/BlendSpace1D.h"
 
 #include "Widgets/SBoxPanel.h"
@@ -322,7 +322,7 @@ int32 SBlendSpaceGridWidget::OnPaint(const FPaintArgs& Args, const FGeometry& Al
 	
 	PaintBackgroundAndGrid(AllottedGeometry, MyCullingRect, OutDrawElements, LayerId);
 
-	if(const UBlendSpaceBase* BlendSpace = BlendSpaceBase.Get())
+	if(const UBlendSpace* BlendSpace = BlendSpaceBase.Get())
 	{
 		if (bShowTriangulation)
 		{
@@ -349,7 +349,7 @@ void SBlendSpaceGridWidget::PaintBackgroundAndGrid(const FGeometry& AllottedGeom
 	// Fill the background
 	FSlateDrawElement::MakeBox( OutDrawElements, DrawLayerId + 1, AllottedGeometry.ToPaintGeometry(), BackgroundImage );
 
-	if(const UBlendSpaceBase* BlendSpace = BlendSpaceBase.Get())
+	if(const UBlendSpace* BlendSpace = BlendSpaceBase.Get())
 	{
 		// Create the grid
 		const FVector2D GridSize = CachedGridRectangle.GetSize();
@@ -402,7 +402,7 @@ void SBlendSpaceGridWidget::PaintBackgroundAndGrid(const FGeometry& AllottedGeom
 
 void SBlendSpaceGridWidget::PaintSampleKeys(const FGeometry& AllottedGeometry, const FSlateRect& MyCullingRect, FSlateWindowElementList& OutDrawElements, int32& DrawLayerId) const
 {
-	if(const UBlendSpaceBase* BlendSpace = BlendSpaceBase.Get())
+	if(const UBlendSpace* BlendSpace = BlendSpaceBase.Get())
 	{
 		// Draw keys
 		const TArray<FBlendSample>& Samples = BlendSpace->GetBlendSamples();
@@ -533,7 +533,7 @@ void SBlendSpaceGridWidget::PaintAxisText(const FGeometry& AllottedGeometry, con
 
 void SBlendSpaceGridWidget::PaintTriangulation(const FGeometry& AllottedGeometry, const FSlateRect& MyCullingRect, FSlateWindowElementList& OutDrawElements, int32& DrawLayerId) const
 {
-	if(const UBlendSpaceBase* BlendSpace = BlendSpaceBase.Get())
+	if(const UBlendSpace* BlendSpace = BlendSpaceBase.Get())
 	{
 		const TArray<FBlendSample>& Samples = BlendSpace->GetBlendSamples();
 		const TArray<FEditorElement>& EditorElements = BlendSpace->GetGridSamples();
@@ -589,7 +589,7 @@ FText SBlendSpaceGridWidget::GetSampleName(const FBlendSample& InBlendSample, in
 
 void SBlendSpaceGridWidget::PaintAnimationNames(const FGeometry& AllottedGeometry, const FSlateRect& MyCullingRect, FSlateWindowElementList& OutDrawElements, int32& DrawLayerId) const
 {
-	if(const UBlendSpaceBase* BlendSpace = BlendSpaceBase.Get())
+	if(const UBlendSpace* BlendSpace = BlendSpaceBase.Get())
 	{
 		const TSharedRef< FSlateFontMeasure > FontMeasure = FSlateApplication::Get().GetRenderer()->GetFontMeasureService();
 		const FVector2D GridCenter = CachedGridRectangle.GetCenter();
@@ -623,7 +623,7 @@ FReply SBlendSpaceGridWidget::OnDrop(const FGeometry& MyGeometry, const FDragDro
 {
 	if(!bReadOnly)
 	{
-		if(const UBlendSpaceBase* BlendSpace = BlendSpaceBase.Get())
+		if(const UBlendSpace* BlendSpace = BlendSpaceBase.Get())
 		{
 			// Check if we are in dropping state and if so snap to the grid and try to add the sample
 			if (DragState == EDragState::DragDrop || DragState == EDragState::InvalidDragDrop || DragState == EDragState::DragDropOverride)
@@ -665,7 +665,7 @@ void SBlendSpaceGridWidget::OnDragEnter(const FGeometry& MyGeometry, const FDrag
 {
 	if(!bReadOnly)
 	{
-		if(const UBlendSpaceBase* BlendSpace = BlendSpaceBase.Get())
+		if(const UBlendSpace* BlendSpace = BlendSpaceBase.Get())
 		{
 			if (DragDropEvent.GetOperationAs<FAssetDragDropOp>().IsValid())
 			{
@@ -679,7 +679,7 @@ FReply SBlendSpaceGridWidget::OnDragOver(const FGeometry& MyGeometry, const FDra
 {
 	if(!bReadOnly)
 	{
-		if(const UBlendSpaceBase* BlendSpace = BlendSpaceBase.Get())
+		if(const UBlendSpace* BlendSpace = BlendSpaceBase.Get())
 		{
 			if (DragState == EDragState::DragDrop || DragState == EDragState::InvalidDragDrop || DragState == EDragState::DragDropOverride)
 			{		
@@ -703,7 +703,7 @@ void SBlendSpaceGridWidget::OnDragLeave(const FDragDropEvent& DragDropEvent)
 {
 	if(!bReadOnly)
 	{
-		if(const UBlendSpaceBase* BlendSpace = BlendSpaceBase.Get())
+		if(const UBlendSpace* BlendSpace = BlendSpaceBase.Get())
 		{
 			if (DragState == EDragState::DragDrop || DragState == EDragState::InvalidDragDrop || DragState == EDragState::DragDropOverride)
 			{
@@ -720,7 +720,7 @@ FReply SBlendSpaceGridWidget::OnMouseButtonUp(const FGeometry& MyGeometry, const
 {
 	if(!bReadOnly)
 	{
-		if(const UBlendSpaceBase* BlendSpace = BlendSpaceBase.Get())
+		if(const UBlendSpace* BlendSpace = BlendSpaceBase.Get())
 		{
 			if (this->HasMouseCapture())
 			{
@@ -760,7 +760,7 @@ FReply SBlendSpaceGridWidget::OnMouseButtonDown(const FGeometry& MyGeometry, con
 {
 	if(!bReadOnly)
 	{
-		if(const UBlendSpaceBase* BlendSpace = BlendSpaceBase.Get())
+		if(const UBlendSpace* BlendSpace = BlendSpaceBase.Get())
 		{
 			if (MouseEvent.GetEffectingButton() == EKeys::LeftMouseButton)
 			{
@@ -800,7 +800,7 @@ FReply SBlendSpaceGridWidget::OnMouseButtonDoubleClick(const FGeometry& InMyGeom
 {
 	if(!bReadOnly)
 	{
-		if(const UBlendSpaceBase* BlendSpace = BlendSpaceBase.Get())
+		if(const UBlendSpace* BlendSpace = BlendSpaceBase.Get())
 		{
 			if (InMouseEvent.GetEffectingButton() == EKeys::LeftMouseButton)
 			{
@@ -819,7 +819,7 @@ FReply SBlendSpaceGridWidget::OnMouseButtonDoubleClick(const FGeometry& InMyGeom
 
 FReply SBlendSpaceGridWidget::OnMouseMove(const FGeometry& MyGeometry, const FPointerEvent& MouseEvent)
 {
-	if(const UBlendSpaceBase* BlendSpace = BlendSpaceBase.Get())
+	if(const UBlendSpace* BlendSpace = BlendSpaceBase.Get())
 	{
 		if(!bReadOnly)
 		{
@@ -891,7 +891,7 @@ FReply SBlendSpaceGridWidget::ProcessClick(const FGeometry& MyGeometry, const FP
 {
 	if(!bReadOnly)
 	{
-		if(const UBlendSpaceBase* BlendSpace = BlendSpaceBase.Get())
+		if(const UBlendSpace* BlendSpace = BlendSpaceBase.Get())
 		{
 			if (MouseEvent.GetEffectingButton() == EKeys::LeftMouseButton)
 			{
@@ -965,7 +965,7 @@ FReply SBlendSpaceGridWidget::OnKeyDown(const FGeometry& MyGeometry, const FKeyE
 {
 	if(!bReadOnly)
 	{
-		if(const UBlendSpaceBase* BlendSpace = BlendSpaceBase.Get())
+		if(const UBlendSpace* BlendSpace = BlendSpaceBase.Get())
 		{
 			// Start previewing when either one of the shift keys is pressed
 			if (bMouseIsOverGeometry)
@@ -996,7 +996,7 @@ FReply SBlendSpaceGridWidget::OnKeyUp(const FGeometry& MyGeometry, const FKeyEve
 {	
 	if(!bReadOnly)
 	{
-		if(const UBlendSpaceBase* BlendSpace = BlendSpaceBase.Get())
+		if(const UBlendSpace* BlendSpace = BlendSpaceBase.Get())
 		{
 			// Stop previewing when shift keys are released 
 			if (DragState != EDragState::DragPreview && ((InKeyEvent.GetKey() == EKeys::LeftShift) || (InKeyEvent.GetKey() == EKeys::RightShift)))
@@ -1117,7 +1117,7 @@ TSharedPtr<SWidget> SBlendSpaceGridWidget::CreateBlendSampleContextMenu()
 	StructureDetailsView = FModuleManager::GetModuleChecked<FPropertyEditorModule>("PropertyEditor")
 		.CreateStructureDetailView(DetailsViewArgs, StructureViewArgs, nullptr, LOCTEXT("SampleData", "Blend Sample"));
 	{
-		if(const UBlendSpaceBase* BlendSpace = BlendSpaceBase.Get())
+		if(const UBlendSpace* BlendSpace = BlendSpaceBase.Get())
 		{
 			const FBlendSample& Sample = BlendSpace->GetBlendSample(HighlightedSampleIndex);		
 			StructureDetailsView->GetDetailsView()->SetGenericLayoutDetailsDelegate(FOnGetDetailCustomizationInstance::CreateStatic(&FBlendSampleDetails::MakeInstance, BlendSpace, this, HighlightedSampleIndex));
@@ -1156,7 +1156,7 @@ TSharedPtr<SWidget> SBlendSpaceGridWidget::CreateNewBlendSampleContextMenu(const
 		NewSampleValue = GridPositionToSampleValue(SnapToClosestGridPoint(InMousePosition), false);
 	}
 
-	if(const UBlendSpaceBase* BlendSpace = BlendSpaceBase.Get())
+	if(const UBlendSpace* BlendSpace = BlendSpaceBase.Get())
 	{
 		MenuBuilder.BeginSection("Sample", LOCTEXT("SampleMenuHeader", "Sample"));
 		{
@@ -1312,7 +1312,7 @@ int32 SBlendSpaceGridWidget::GetClosestSamplePointIndexToMouse() const
 	float BestDistance = FLT_MAX;
 	int32 BestIndex = INDEX_NONE;
 
-	if(const UBlendSpaceBase* BlendSpace = BlendSpaceBase.Get())
+	if(const UBlendSpace* BlendSpace = BlendSpaceBase.Get())
 	{
 		const TArray<FBlendSample>& Samples = BlendSpace->GetBlendSamples();
 		for (int32 SampleIndex = 0; SampleIndex < Samples.Num(); ++SampleIndex)
@@ -1353,7 +1353,7 @@ void SBlendSpaceGridWidget::StopPreviewing()
 FText SBlendSpaceGridWidget::GetToolTipAnimationName() const
 {
 	FText ToolTipText = FText::GetEmpty();
-	if(const UBlendSpaceBase* BlendSpace = BlendSpaceBase.Get())
+	if(const UBlendSpace* BlendSpace = BlendSpaceBase.Get())
 	{
 		const FText PreviewValue = LOCTEXT("PreviewValueTooltip", "Preview Value");
 
@@ -1439,7 +1439,7 @@ FText SBlendSpaceGridWidget::GetToolTipSampleValue() const
 {
 	FText ToolTipText = FText::GetEmpty();
 
-	if(const UBlendSpaceBase* BlendSpace = BlendSpaceBase.Get())
+	if(const UBlendSpace* BlendSpace = BlendSpaceBase.Get())
 	{
 		static const FTextFormat OneAxisFormat = LOCTEXT("OneAxisFormat", "{0}: {1}");
 		static const FTextFormat TwoAxisFormat = LOCTEXT("TwoAxisFormat", "{0}: {1} - {2}: {3}");
@@ -1662,7 +1662,7 @@ TOptional<float> SBlendSpaceGridWidget::GetInputBoxValue(const int32 ParameterIn
 {
 	checkf(ParameterIndex < 3, TEXT("Invalid parameter index, suppose to be within FVector array range"));
 	float ReturnValue = 0.0f;
-	if(const UBlendSpaceBase* BlendSpace = BlendSpaceBase.Get())
+	if(const UBlendSpace* BlendSpace = BlendSpaceBase.Get())
 	{
 		if (SelectedSampleIndex != INDEX_NONE && SelectedSampleIndex < BlendSpace->GetNumberOfBlendSamples())
 		{
@@ -1808,7 +1808,7 @@ EVisibility SBlendSpaceGridWidget::GetFittingButtonVisibility() const
 
 void SBlendSpaceGridWidget::UpdateCachedBlendParameterData()
 {
-	if(const UBlendSpaceBase* BlendSpace = BlendSpaceBase.Get())
+	if(const UBlendSpace* BlendSpace = BlendSpaceBase.Get())
 	{
 		const FBlendParameter& BlendParameterX = BlendSpace->GetBlendParameter(0);
 		const FBlendParameter& BlendParameterY = BlendSpace->GetBlendParameter(1);
@@ -1889,7 +1889,7 @@ void SBlendSpaceGridWidget::Tick(const FGeometry& AllottedGeometry, const double
 	HighlightedSampleIndex = INDEX_NONE;
 	const bool bPreviousHighlightPreviewPin = bHighlightPreviewPin;
 
-	if(const UBlendSpaceBase* BlendSpace = BlendSpaceBase.Get())
+	if(const UBlendSpace* BlendSpace = BlendSpaceBase.Get())
 	{
 		if(PreviousBlendSpaceBase.Get() != BlendSpace)
 		{
@@ -2047,7 +2047,7 @@ const FVector SBlendSpaceGridWidget::GetPreviewPosition() const
 
 void SBlendSpaceGridWidget::SetPreviewingState(const FVector& InPosition, const FVector& InFilteredPosition)
 {
-	if (const UBlendSpaceBase* BlendSpace = BlendSpaceBase.Get())
+	if (const UBlendSpace* BlendSpace = BlendSpaceBase.Get())
 	{
 		PreviewFilteredPosition = BlendSpace->GetClampedAndWrappedBlendInput(InFilteredPosition);
 		PreviewPosition = BlendSpace->GetClampedAndWrappedBlendInput(InPosition);
@@ -2118,7 +2118,7 @@ bool SBlendSpaceGridWidget::ValidateAnimationSequence(const UAnimSequence* Anima
 {	
 	if (AnimationSequence != nullptr)
 	{
-		if(const UBlendSpaceBase* BlendSpace = BlendSpaceBase.Get())
+		if(const UBlendSpace* BlendSpace = BlendSpaceBase.Get())
 		{
 			if(BlendSpace->IsAsset())
 			{
diff --git a/Engine/Source/Editor/Persona/Private/SAnimationBlendSpaceGridWidget.h b/Engine/Source/Editor/Persona/Private/SAnimationBlendSpaceGridWidget.h
--- a/Engine/Source/Editor/Persona/Private/SAnimationBlendSpaceGridWidget.h
+++ b/Engine/Source/Editor/Persona/Private/SAnimationBlendSpaceGridWidget.h
@@ -20,7 +20,7 @@
 class FPaintArgs;
 class FSlateWindowElementList;
 class UAnimSequence;
-class UBlendSpaceBase;
+class UBlendSpace;
 
 DECLARE_DELEGATE_FourParams(FOnSampleMoved, const int32 /*SampleIndex*/, const FVector& /*SampleValue*/, bool /*bIsInteractive*/, bool /*bSnap*/);
 DECLARE_DELEGATE_TwoParams(FOnSampleAnimationChanged, UAnimSequence*, const FVector&);
@@ -36,7 +36,7 @@ public:
 		, _ShowSettingsButtons(true)
 		, _StatusBarName(TEXT("AssetEditor.AnimationEditor.MainMenu"))
 	{}
-		SLATE_ATTRIBUTE(const UBlendSpaceBase*, BlendSpaceBase)
+		SLATE_ATTRIBUTE(const UBlendSpace*, BlendSpaceBase)
 		SLATE_ATTRIBUTE(FVector, Position)
 		SLATE_ATTRIBUTE(FVector, FilteredPosition)
 		SLATE_ARGUMENT(FNotifyHook*, NotifyHook)
@@ -223,9 +223,9 @@ protected:
 	void EnableStatusBarMessage(bool bEnable);
 private:
 	/** Currently visualized blendspace (const to ensure changes to it are only made within SAnimationBlendSpace */
-	TAttribute<const UBlendSpaceBase*> BlendSpaceBase;
+	TAttribute<const UBlendSpace*> BlendSpaceBase;
 	/** Previous blendspace, cached here to allow us to update cached data when it changes. */
-	TWeakObjectPtr<const UBlendSpaceBase> PreviousBlendSpaceBase;
+	TWeakObjectPtr<const UBlendSpace> PreviousBlendSpaceBase;
 	/** Current BlendSpace params (driven from outside the widget) */
 	TAttribute<FVector> Position;
 	/** Current BlendSpace state - filtered version of Position (driven from outside the widget) */
diff --git a/Engine/Source/Editor/Persona/Private/SAnimationScrubPanel.cpp b/Engine/Source/Editor/Persona/Private/SAnimationScrubPanel.cpp
--- a/Engine/Source/Editor/Persona/Private/SAnimationScrubPanel.cpp
+++ b/Engine/Source/Editor/Persona/Private/SAnimationScrubPanel.cpp
@@ -8,7 +8,7 @@
 #include "AnimPreviewInstance.h"
 #include "SScrubControlPanel.h"
 #include "ScopedTransaction.h"
-#include "Animation/BlendSpaceBase.h"
+#include "Animation/BlendSpace.h"
 #include "AnimationEditorPreviewScene.h"
 #include "Animation/AnimData/AnimDataModel.h"
 #include "Animation/AnimData/AnimDataController.h"
@@ -322,7 +322,7 @@ uint32 SAnimationScrubPanel::GetNumberOfKeys() const
 			{
 				NumKeys = CastChecked<UAnimSequenceBase>(PreviewInstance->GetCurrentAsset())->GetNumberOfSampledKeys();
 			}
-			else if(PreviewInstance->GetCurrentAsset()->IsA(UBlendSpaceBase::StaticClass()))
+			else if(PreviewInstance->GetCurrentAsset()->IsA(UBlendSpace::StaticClass()))
 			{
 				// Blendspaces dont display frame notches, so just return 0 here
 				NumKeys = 0;
diff --git a/Engine/Source/Editor/Persona/Public/PersonaModule.h b/Engine/Source/Editor/Persona/Public/PersonaModule.h
--- a/Engine/Source/Editor/Persona/Public/PersonaModule.h
+++ b/Engine/Source/Editor/Persona/Public/PersonaModule.h
@@ -23,7 +23,7 @@ class IPinnedCommandList;
 class FWorkflowAllowedTabSet;
 class IAssetFamily;
 class FWorkflowTabFactory;
-class UBlendSpaceBase;
+class UBlendSpace;
 class IAnimSequenceCurveEditor;
 class IAnimationEditor;
 
@@ -135,7 +135,7 @@ struct FBlendSpaceEditorArgs
 
 struct FBlendSpacePreviewArgs
 {
-	TAttribute<const UBlendSpaceBase*> PreviewBlendSpace;
+	TAttribute<const UBlendSpace*> PreviewBlendSpace;
 
 	// Allows the target preview position to be programatically driven
 	TAttribute<FVector> PreviewPosition;
@@ -300,13 +300,13 @@ public:
 
 	/** Deprecated */
 	UE_DEPRECATED(5.0, "Please use the overload that takes a FBlendSpacePreviewArgs struct")
-	virtual TSharedRef<SWidget> CreateBlendSpacePreviewWidget(TAttribute<const UBlendSpaceBase*> InBlendSpace, TAttribute<FVector> InBlendPosition, TAttribute<FVector> InFilteredBlendPosition) const;
+	virtual TSharedRef<SWidget> CreateBlendSpacePreviewWidget(TAttribute<const UBlendSpace*> InBlendSpace, TAttribute<FVector> InBlendPosition, TAttribute<FVector> InFilteredBlendPosition) const;
 
 	/** Create a widget to preview a blendspace */
 	virtual TSharedRef<SWidget> CreateBlendSpacePreviewWidget(const FBlendSpacePreviewArgs& InArgs) const;
 
 	/** Create a widget to edit a blendspace */
-	virtual TSharedRef<SWidget> CreateBlendSpaceEditWidget(UBlendSpaceBase* InBlendSpace, const FBlendSpaceEditorArgs& InArgs) const;
+	virtual TSharedRef<SWidget> CreateBlendSpaceEditWidget(UBlendSpace* InBlendSpace, const FBlendSpaceEditorArgs& InArgs) const;
 
 	/** Create a tab factory for editing montage sections */
 	virtual TSharedRef<FWorkflowTabFactory> CreateAnimMontageSectionsTabFactory(const TSharedRef<class FWorkflowCentricApplication>& InHostingApp, const TSharedRef<IPersonaToolkit>& InPersonaToolkit, FSimpleMulticastDelegate& InOnSectionsChanged) const;
diff --git a/Engine/Source/Editor/UnrealEd/Private/Animation/DebugSkelMeshComponent.cpp b/Engine/Source/Editor/UnrealEd/Private/Animation/DebugSkelMeshComponent.cpp
--- a/Engine/Source/Editor/UnrealEd/Private/Animation/DebugSkelMeshComponent.cpp
+++ b/Engine/Source/Editor/UnrealEd/Private/Animation/DebugSkelMeshComponent.cpp
@@ -13,7 +13,7 @@
 #include "SkeletalRenderPublic.h"
 #include "AnimPreviewInstance.h"
 #include "Animation/AnimComposite.h"
-#include "Animation/BlendSpaceBase.h"
+#include "Animation/BlendSpace.h"
 #include "Rendering/SkeletalMeshRenderData.h"
 #include "Rendering/SkeletalMeshModel.h"
 
@@ -237,7 +237,7 @@ FString UDebugSkelMeshComponent::GetPreviewText() const
 			FText Label = SkeletalMeshComponent->GetOwner() ? FText::FromString(SkeletalMeshComponent->GetOwner()->GetActorLabel()) : LOCTEXT("NoActor", "None");
 			return FText::Format(LOCTEXT("ExternalComponent", "External Instance on {0}"), Label).ToString();
 		}
-		else if (UBlendSpaceBase* BlendSpace = Cast<UBlendSpaceBase>(CurrentAsset))
+		else if (UBlendSpace* BlendSpace = Cast<UBlendSpace>(CurrentAsset))
 		{
 			return FText::Format( LOCTEXT("BlendSpace", "Blend Space {0}"), FText::FromString(BlendSpace->GetName()) ).ToString();
 		}
diff --git a/Engine/Source/Editor/UnrealEd/Private/ThumbnailHelpers.cpp b/Engine/Source/Editor/UnrealEd/Private/ThumbnailHelpers.cpp
--- a/Engine/Source/Editor/UnrealEd/Private/ThumbnailHelpers.cpp
+++ b/Engine/Source/Editor/UnrealEd/Private/ThumbnailHelpers.cpp
@@ -709,7 +709,7 @@ FBlendSpaceThumbnailScene::FBlendSpaceThumbnailScene()
 	PreviewActor->SetActorEnableCollision(false);
 }
 
-bool FBlendSpaceThumbnailScene::SetBlendSpace(class UBlendSpaceBase* InBlendSpace)
+bool FBlendSpaceThumbnailScene::SetBlendSpace(class UBlendSpace* InBlendSpace)
 {
 	PreviewActor->GetSkeletalMeshComponent()->OverrideMaterials.Empty();
 
diff --git a/Engine/Source/Editor/UnrealEd/Private/ThumbnailRendering/BlendSpaceThumbnailRenderer.cpp b/Engine/Source/Editor/UnrealEd/Private/ThumbnailRendering/BlendSpaceThumbnailRenderer.cpp
--- a/Engine/Source/Editor/UnrealEd/Private/ThumbnailRendering/BlendSpaceThumbnailRenderer.cpp
+++ b/Engine/Source/Editor/UnrealEd/Private/ThumbnailRendering/BlendSpaceThumbnailRenderer.cpp
@@ -6,7 +6,7 @@
 #include "SceneView.h"
 #include "ThumbnailHelpers.h"
 
-#include "Animation/BlendSpaceBase.h"
+#include "Animation/BlendSpace.h"
 
 UBlendSpaceThumbnailRenderer::UBlendSpaceThumbnailRenderer(const FObjectInitializer& ObjectInitializer)
 	: Super(ObjectInitializer)
@@ -16,7 +16,7 @@ UBlendSpaceThumbnailRenderer::UBlendSpaceThumbnailRenderer(const FObjectInitiali
 
 void UBlendSpaceThumbnailRenderer::Draw(UObject* Object, int32 X, int32 Y, uint32 Width, uint32 Height, FRenderTarget* RenderTarget, FCanvas* Canvas, bool bAdditionalViewFamily)
 {
-	UBlendSpaceBase* BlendSpace = Cast<UBlendSpaceBase>(Object);
+	UBlendSpace* BlendSpace = Cast<UBlendSpace>(Object);
 	if (BlendSpace != nullptr)
 	{
 		if (ThumbnailScene == nullptr)
diff --git a/Engine/Source/Editor/UnrealEd/Public/ThumbnailHelpers.h b/Engine/Source/Editor/UnrealEd/Public/ThumbnailHelpers.h
--- a/Engine/Source/Editor/UnrealEd/Public/ThumbnailHelpers.h
+++ b/Engine/Source/Editor/UnrealEd/Public/ThumbnailHelpers.h
@@ -168,7 +168,7 @@ public:
 	FBlendSpaceThumbnailScene();
 
 	/** Sets the animation to use in the next GetView() */
-	bool SetBlendSpace(class UBlendSpaceBase* InBlendSpace);
+	bool SetBlendSpace(class UBlendSpace* InBlendSpace);
 
 protected:
 	// FThumbnailPreviewScene implementation
@@ -182,7 +182,7 @@ private:
 	class AAnimationThumbnailSkeletalMeshActor* PreviewActor;
 
 	/** Animation we are generating the thumbnail for */
-	class UBlendSpaceBase* PreviewAnimation;
+	class UBlendSpace* PreviewAnimation;
 };
 
 class UNREALED_API FAnimBlueprintThumbnailScene : public FThumbnailPreviewScene
diff --git a/Engine/Source/Runtime/AnimGraphRuntime/Private/AnimGraphRuntimeTrace.cpp b/Engine/Source/Runtime/AnimGraphRuntime/Private/AnimGraphRuntimeTrace.cpp
--- a/Engine/Source/Runtime/AnimGraphRuntime/Private/AnimGraphRuntimeTrace.cpp
+++ b/Engine/Source/Runtime/AnimGraphRuntime/Private/AnimGraphRuntimeTrace.cpp
@@ -5,7 +5,7 @@
 #if ANIM_TRACE_ENABLED
 
 #include "Trace/Trace.inl"
-#include "Animation/BlendSpaceBase.h"
+#include "Animation/BlendSpace.h"
 #include "AnimNodes/AnimNode_BlendSpacePlayer.h"
 #include "AnimNodes/AnimNode_BlendSpaceGraphBase.h"
 #include "Animation/AnimInstanceProxy.h"
diff --git a/Engine/Source/Runtime/AnimGraphRuntime/Private/AnimNodes/AnimNode_AimOffsetLookAt.cpp b/Engine/Source/Runtime/AnimGraphRuntime/Private/AnimNodes/AnimNode_AimOffsetLookAt.cpp
--- a/Engine/Source/Runtime/AnimGraphRuntime/Private/AnimNodes/AnimNode_AimOffsetLookAt.cpp
+++ b/Engine/Source/Runtime/AnimGraphRuntime/Private/AnimNodes/AnimNode_AimOffsetLookAt.cpp
@@ -3,7 +3,7 @@
 #include "AnimNodes/AnimNode_AimOffsetLookAt.h"
 #include "Animation/AnimInstanceProxy.h"
 #include "AnimationRuntime.h"
-#include "Animation/BlendSpaceBase.h"
+#include "Animation/BlendSpace.h"
 #include "Animation/BlendSpace1D.h"
 #include "Engine/SkeletalMeshSocket.h"
 #include "DrawDebugHelpers.h"
diff --git a/Engine/Source/Runtime/AnimGraphRuntime/Private/AnimNodes/AnimNode_BlendSpaceEvaluator.cpp b/Engine/Source/Runtime/AnimGraphRuntime/Private/AnimNodes/AnimNode_BlendSpaceEvaluator.cpp
--- a/Engine/Source/Runtime/AnimGraphRuntime/Private/AnimNodes/AnimNode_BlendSpaceEvaluator.cpp
+++ b/Engine/Source/Runtime/AnimGraphRuntime/Private/AnimNodes/AnimNode_BlendSpaceEvaluator.cpp
@@ -1,7 +1,7 @@
 // Copyright Epic Games, Inc. All Rights Reserved.
 
 #include "AnimNodes/AnimNode_BlendSpaceEvaluator.h"
-#include "Animation/BlendSpaceBase.h"
+#include "Animation/BlendSpace.h"
 #include "Animation/AnimTrace.h"
 
 /////////////////////////////////////////////////////
diff --git a/Engine/Source/Runtime/AnimGraphRuntime/Private/AnimNodes/AnimNode_BlendSpaceGraphBase.cpp b/Engine/Source/Runtime/AnimGraphRuntime/Private/AnimNodes/AnimNode_BlendSpaceGraphBase.cpp
--- a/Engine/Source/Runtime/AnimGraphRuntime/Private/AnimNodes/AnimNode_BlendSpaceGraphBase.cpp
+++ b/Engine/Source/Runtime/AnimGraphRuntime/Private/AnimNodes/AnimNode_BlendSpaceGraphBase.cpp
@@ -1,7 +1,7 @@
 // Copyright Epic Games, Inc. All Rights Reserved.
 
 #include "AnimNodes/AnimNode_BlendSpaceGraphBase.h"
-#include "Animation/BlendSpaceBase.h"
+#include "Animation/BlendSpace.h"
 #include "Animation/AnimInstanceProxy.h"
 #include "AnimGraphRuntimeTrace.h"
 #include "Animation/AnimSyncScope.h"
diff --git a/Engine/Source/Runtime/AnimGraphRuntime/Private/AnimNodes/AnimNode_BlendSpacePlayer.cpp b/Engine/Source/Runtime/AnimGraphRuntime/Private/AnimNodes/AnimNode_BlendSpacePlayer.cpp
--- a/Engine/Source/Runtime/AnimGraphRuntime/Private/AnimNodes/AnimNode_BlendSpacePlayer.cpp
+++ b/Engine/Source/Runtime/AnimGraphRuntime/Private/AnimNodes/AnimNode_BlendSpacePlayer.cpp
@@ -1,7 +1,7 @@
 // Copyright Epic Games, Inc. All Rights Reserved.
 
 #include "AnimNodes/AnimNode_BlendSpacePlayer.h"
-#include "Animation/BlendSpaceBase.h"
+#include "Animation/BlendSpace.h"
 #include "Animation/AnimSequence.h"
 #include "Animation/AnimInstanceProxy.h"
 #include "AnimGraphRuntimeTrace.h"
@@ -135,7 +135,7 @@ void FAnimNode_BlendSpacePlayer::Evaluate_AnyThread(FPoseContext& Output)
 
 void FAnimNode_BlendSpacePlayer::OverrideAsset(UAnimationAsset* NewAsset)
 {
-	if(UBlendSpaceBase* NewBlendSpace = Cast<UBlendSpaceBase>(NewAsset))
+	if(UBlendSpace* NewBlendSpace = Cast<UBlendSpace>(NewAsset))
 	{
 		BlendSpace = NewBlendSpace;
 	}
diff --git a/Engine/Source/Runtime/AnimGraphRuntime/Public/AnimNodes/AnimNode_BlendSpaceGraphBase.h b/Engine/Source/Runtime/AnimGraphRuntime/Public/AnimNodes/AnimNode_BlendSpaceGraphBase.h
--- a/Engine/Source/Runtime/AnimGraphRuntime/Public/AnimNodes/AnimNode_BlendSpaceGraphBase.h
+++ b/Engine/Source/Runtime/AnimGraphRuntime/Public/AnimNodes/AnimNode_BlendSpaceGraphBase.h
@@ -8,7 +8,7 @@
 #include "Animation/AnimNodeBase.h"
 #include "AnimNode_BlendSpaceGraphBase.generated.h"
 
-class UBlendSpaceBase;
+class UBlendSpace;
 
 // Allows multiple animations to be blended between based on input parameters
 USTRUCT(BlueprintInternalUseOnly)
@@ -19,7 +19,7 @@ struct ANIMGRAPHRUNTIME_API FAnimNode_BlendSpaceGraphBase : public FAnimNode_Bas
 	friend class UAnimGraphNode_BlendSpaceGraphBase;
 
 	// @return the blendspace that this node uses
-	const UBlendSpaceBase* GetBlendSpace() const { return BlendSpace; }
+	const UBlendSpace* GetBlendSpace() const { return BlendSpace; }
 
 	// @return the current sample coordinates that this node is using to sample the blendspace
 	FVector GetPosition() const { return FVector(X, Y, Z); }
@@ -55,7 +55,7 @@ protected:
 
 	// The internal blendspace asset to play
 	UPROPERTY()
-	TObjectPtr<const UBlendSpaceBase> BlendSpace = nullptr;
+	TObjectPtr<const UBlendSpace> BlendSpace = nullptr;
 
 	// Pose links for each sample in the blendspace
 	UPROPERTY()
diff --git a/Engine/Source/Runtime/AnimGraphRuntime/Public/AnimNodes/AnimNode_BlendSpacePlayer.h b/Engine/Source/Runtime/AnimGraphRuntime/Public/AnimNodes/AnimNode_BlendSpacePlayer.h
--- a/Engine/Source/Runtime/AnimGraphRuntime/Public/AnimNodes/AnimNode_BlendSpacePlayer.h
+++ b/Engine/Source/Runtime/AnimGraphRuntime/Public/AnimNodes/AnimNode_BlendSpacePlayer.h
@@ -8,7 +8,7 @@
 #include "Animation/AnimNode_AssetPlayerBase.h"
 #include "AnimNode_BlendSpacePlayer.generated.h"
 
-class UBlendSpaceBase;
+class UBlendSpace;
 
 //@TODO: Comment
 USTRUCT(BlueprintInternalUseOnly)
@@ -54,7 +54,7 @@ public:
 
 	// The blendspace asset to play
 	UPROPERTY(EditAnywhere, BlueprintReadWrite, Category = Settings, meta = (PinHiddenByDefault))
-	TObjectPtr<UBlendSpaceBase> BlendSpace;
+	TObjectPtr<UBlendSpace> BlendSpace;
 
 protected:
 
@@ -63,7 +63,7 @@ protected:
 	TArray<FBlendSampleData> BlendSampleDataCache;
 
 	UPROPERTY(Transient)
-	TObjectPtr<UBlendSpaceBase> PreviousBlendSpace;
+	TObjectPtr<UBlendSpace> PreviousBlendSpace;
 
 public:	
 	FAnimNode_BlendSpacePlayer();
diff --git a/Engine/Source/Runtime/Engine/Classes/Animation/AnimBlueprintGeneratedClass.h b/Engine/Source/Runtime/Engine/Classes/Animation/AnimBlueprintGeneratedClass.h
--- a/Engine/Source/Runtime/Engine/Classes/Animation/AnimBlueprintGeneratedClass.h
+++ b/Engine/Source/Runtime/Engine/Classes/Animation/AnimBlueprintGeneratedClass.h
@@ -11,7 +11,7 @@
 #include "Animation/AnimStateMachineTypes.h"
 #include "Animation/AnimClassInterface.h"
 #include "Animation/AnimNodeBase.h"
-#include "Animation/BlendSpaceBase.h"
+#include "Animation/BlendSpace.h"
 #include "PropertyAccess.h"
 
 #include "AnimBlueprintGeneratedClass.generated.h"
@@ -232,7 +232,7 @@ public:
 	// Record of a blend space player's state
 	struct FBlendSpacePlayerRecord
 	{
-		FBlendSpacePlayerRecord(int32 InNodeID, const UBlendSpaceBase* InBlendSpace, const FVector& InPosition, const FVector& InFilteredPosition)
+		FBlendSpacePlayerRecord(int32 InNodeID, const UBlendSpace* InBlendSpace, const FVector& InPosition, const FVector& InFilteredPosition)
 			: NodeID(InNodeID)
 			, BlendSpace(InBlendSpace)
 			, Position(InPosition)
@@ -240,7 +240,7 @@ public:
 		{}
 
 		int32 NodeID;
-		TWeakObjectPtr<const UBlendSpaceBase> BlendSpace;
+		TWeakObjectPtr<const UBlendSpace> BlendSpace;
 		FVector Position;
 		FVector FilteredPosition;
 	};
@@ -284,7 +284,7 @@ public:
 	void RecordStateData(int32 StateMachineIndex, int32 StateIndex, float Weight, float ElapsedTime);
 	void RecordNodeValue(int32 InNodeID, const FString& InText);
 	void RecordSequencePlayer(int32 InNodeID, float InPosition, float InLength, int32 InFrameCount);
-	void RecordBlendSpacePlayer(int32 InNodeID, const UBlendSpaceBase* InBlendSpace, const FVector& InPosition, const FVector& InFilteredPosition);
+	void RecordBlendSpacePlayer(int32 InNodeID, const UBlendSpace* InBlendSpace, const FVector& InPosition, const FVector& InFilteredPosition);
 
 	void AddPoseWatch(int32 NodeID, FColor Color);
 	void RemovePoseWatch(int32 NodeID);
@@ -364,7 +364,7 @@ private:
 
 	// Any internal blendspaces we host
 	UPROPERTY()
-	TArray<TObjectPtr<UBlendSpaceBase>> BlendSpaces;
+	TArray<TObjectPtr<UBlendSpace>> BlendSpaces;
 
 public:
 	// IAnimClassInterface interface
diff --git a/Engine/Source/Runtime/Engine/Classes/Animation/AnimationAsset.h b/Engine/Source/Runtime/Engine/Classes/Animation/AnimationAsset.h
--- a/Engine/Source/Runtime/Engine/Classes/Animation/AnimationAsset.h
+++ b/Engine/Source/Runtime/Engine/Classes/Animation/AnimationAsset.h
@@ -24,7 +24,7 @@ class UAssetMappingTable;
 class UAssetUserData;
 class USkeleton;
 class UAnimSequenceBase;
-class UBlendSpaceBase;
+class UBlendSpace;
 class UPoseAsset;
 
 namespace MarkerIndexSpecialValues
@@ -350,7 +350,7 @@ public:
 	ENGINE_API FAnimTickRecord(UAnimSequenceBase* InSequence, bool bInLooping, float InPlayRate, float InFinalBlendWeight, float& InCurrentTime, FMarkerTickRecord& InMarkerTickRecord);
 
 	// Create a tick record for a blendspace
-	ENGINE_API FAnimTickRecord(UBlendSpaceBase* InBlendSpace, const FVector& InBlendInput, TArray<FBlendSampleData>& InBlendSampleDataCache, FBlendFilter& InBlendFilter, bool bInLooping, float InPlayRate, float InFinalBlendWeight, float& InCurrentTime, FMarkerTickRecord& InMarkerTickRecord);
+	ENGINE_API FAnimTickRecord(UBlendSpace* InBlendSpace, const FVector& InBlendInput, TArray<FBlendSampleData>& InBlendSampleDataCache, FBlendFilter& InBlendFilter, bool bInLooping, float InPlayRate, float InFinalBlendWeight, float& InCurrentTime, FMarkerTickRecord& InMarkerTickRecord);
 
 	// Create a tick record for a montage
 	ENGINE_API FAnimTickRecord(UAnimMontage* InMontage, float InCurrentPosition, float InPreviousPosition, float InMoveDelta, float InWeight, TArray<FPassedMarker>& InMarkersPassedThisTick, FMarkerTickRecord& InMarkerTickRecord);
diff --git a/Engine/Source/Runtime/Engine/Classes/Animation/BlendSpace.h b/Engine/Source/Runtime/Engine/Classes/Animation/BlendSpace.h
--- a/Engine/Source/Runtime/Engine/Classes/Animation/BlendSpace.h
+++ b/Engine/Source/Runtime/Engine/Classes/Animation/BlendSpace.h
@@ -1,46 +1,588 @@
 // Copyright Epic Games, Inc. All Rights Reserved.
 
+/**
+ * Blend Space. Contains functionality shared across all blend space objects
+ *
+ */
+
 #pragma once
 
 #include "CoreMinimal.h"
 #include "UObject/ObjectMacros.h"
-#include "Animation/BlendSpaceBase.h"
-
+#include "BoneContainer.h"
+#include "Animation/AnimationAsset.h"
+#include "AnimationRuntime.h"
+#include "AnimNodeBase.h"
+#include "Containers/ArrayView.h"
 #include "BlendSpace.generated.h"
 
+/** Interpolation data types. */
+UENUM()
+enum EBlendSpaceAxis
+{
+	BSA_None UMETA(DisplayName = "None"),
+	BSA_X UMETA(DisplayName = "Horizontal (X) Axis"),
+	BSA_Y UMETA(DisplayName = "Vertical (Y) Axis"),
+	BSA_Max
+};
+
+USTRUCT()
+struct FInterpolationParameter
+{
+	GENERATED_USTRUCT_BODY()
+
+	/** Interpolation Time for input, when it gets input, it will use this time to interpolate to target, used for smoother interpolation. */
+	UPROPERTY(EditAnywhere, Category=Parameter)
+	float InterpolationTime = 0.f;
+
+	/** Type of interpolation used for filtering the input value to decide how to get to target. */
+	UPROPERTY(EditAnywhere, Category=Parameter)
+	TEnumAsByte<EFilterInterpolationType> InterpolationType = EFilterInterpolationType::BSIT_Average;
+};
+
+USTRUCT()
+struct FBlendParameter
+{
+	GENERATED_USTRUCT_BODY()
+
+	UPROPERTY(EditAnywhere, DisplayName = "Name", Category=BlendParameter)
+	FString DisplayName;
+
+	/** Min value for this parameter. */
+	UPROPERTY(EditAnywhere, DisplayName = "Minimum Axis Value", Category=BlendParameter)
+	float Min;
+
+	/** Max value for this parameter. */
+	UPROPERTY(EditAnywhere, DisplayName = "Maximum Axis Value", Category=BlendParameter)
+	float Max;
+
+	/** The number of grid divisions for this parameter (axis). */
+	UPROPERTY(EditAnywhere, DisplayName = "Number of Grid Divisions", Category=BlendParameter, meta=(UIMin="1", ClampMin="1"))
+	int32 GridNum;
+
+	/** If false then input parameters are clamped to the min/max values on this axis. If true then the input can go outside the min/max range and the blend space is treated as being cyclic on this axis. */
+	UPROPERTY(EditAnywhere, DisplayName = "Wrap Input", Category = BlendParameter)
+	bool bWrapInput;
+
+	FBlendParameter()
+		: DisplayName(TEXT("None"))
+		, Min(0.f)
+		, Max(100.f)
+		, GridNum(4) // TODO when changing GridNum's default value, it breaks all grid samples ATM - provide way to rebuild grid samples during loading
+		, bWrapInput(false)
+	{
+	}
+
+	float GetRange() const
+	{
+		return Max-Min;
+	}
+	/** Return size of each grid. */
+	float GetGridSize() const
+	{
+		return GetRange()/(float)GridNum;
+	}
+	
+};
+
+/** Sample data */
+USTRUCT()
+struct FBlendSample
+{
+	GENERATED_USTRUCT_BODY()
+
+	// For linked animations
+	UPROPERTY(EditAnywhere, Category=BlendSample)
+	TObjectPtr<class UAnimSequence> Animation;
+
+	//blend 0->x, blend 1->y, blend 2->z
+
+	UPROPERTY(EditAnywhere, Category=BlendSample)
+	FVector SampleValue;
+	
+	UPROPERTY(EditAnywhere, Category = BlendSample, meta=(UIMin="0.01", UIMax="2.0", ClampMin="0.01", ClampMax="64.0"))
+	float RateScale;
+
+#if WITH_EDITORONLY_DATA
+	UPROPERTY(EditAnywhere, Category=BlendSample)
+	uint8 bSnapToGrid : 1;
+
+	UPROPERTY(transient)
+	uint8 bIsValid : 1;
+
+	// Cache the samples marker data counter so that we can track if it changes and revalidate the blendspace
+	int32 CachedMarkerDataUpdateCounter;
+
+#endif // WITH_EDITORONLY_DATA
+
+	FBlendSample()
+		: Animation(nullptr)
+		, SampleValue(0.f)
+		, RateScale(1.0f)
+#if WITH_EDITORONLY_DATA
+		, bSnapToGrid(true)
+		, bIsValid(false)
+		, CachedMarkerDataUpdateCounter(INDEX_NONE)
+#endif // WITH_EDITORONLY_DATA
+	{		
+	}
+	
+	FBlendSample(class UAnimSequence* InAnim, FVector InValue, bool bInIsSnapped, bool bInIsValid) 
+		: Animation(InAnim)
+		, SampleValue(InValue)
+		, RateScale(1.0f)
+#if WITH_EDITORONLY_DATA
+		, bSnapToGrid(bInIsSnapped)
+		, bIsValid(bInIsValid)
+		, CachedMarkerDataUpdateCounter(INDEX_NONE)
+#endif // WITH_EDITORONLY_DATA
+	{		
+	}
+	
+	bool operator==( const FBlendSample& Other ) const 
+	{
+		return (Other.Animation == Animation && Other.SampleValue == SampleValue && FMath::IsNearlyEqual(Other.RateScale, RateScale));
+	}
+};
+
+/**
+ * Each elements in the grid
+ */
+USTRUCT()
+struct FEditorElement
+{
+	GENERATED_USTRUCT_BODY()
+
+	// for now we only support triangles
+	static const int32 MAX_VERTICES = 3;
+
+	UPROPERTY(EditAnywhere, Category=EditorElement)
+	int32 Indices[MAX_VERTICES];
+
+	UPROPERTY(EditAnywhere, Category=EditorElement)
+	float Weights[MAX_VERTICES];
+
+	FEditorElement()
+	{
+		for (int32 ElementIndex = 0; ElementIndex < MAX_VERTICES; ElementIndex++)
+		{
+			Indices[ElementIndex] = INDEX_NONE;
+		}
+		for (int32 ElementIndex = 0; ElementIndex < MAX_VERTICES; ElementIndex++)
+		{
+			Weights[ElementIndex] = 0;
+		}
+	}
+	
+};
+
+/** result of how much weight of the grid element **/
+USTRUCT()
+struct FGridBlendSample
+{
+	GENERATED_USTRUCT_BODY()
+
+	UPROPERTY()
+	struct FEditorElement GridElement;
+
+	UPROPERTY()
+	float BlendWeight;
+
+	FGridBlendSample()
+		: BlendWeight(0)
+	{
+	}
+
+};
+
+USTRUCT()
+struct FPerBoneInterpolation
+{
+	GENERATED_USTRUCT_BODY()
+
+	UPROPERTY(EditAnywhere, Category=FPerBoneInterpolation)
+	FBoneReference BoneReference;
+
+	/**
+	* If greater than zero, this is the speed at which the sample weights are allowed to change for this specific bone.
+	*
+	* A speed of 1 means a sample weight can change from zero to one (or one to zero) in one second.
+	* A speed of 2 means that this would take half a second.
+	*
+	* Smaller values mean slower adjustments of the sample weights, and thus more smoothing. However, a
+	* value of zero disables this smoothing entirely.
+	* 
+	* If set, the value overrides the overall Sample Weight Speed which will no longer affect this bone.
+	*/
+	UPROPERTY(EditAnywhere, Category=FPerBoneInterpolation, meta=(DisplayName="Sample Weight Speed"))
+	float InterpolationSpeedPerSec;
+
+	FPerBoneInterpolation()
+		: InterpolationSpeedPerSec(6.f)
+	{}
+
+	void Initialize(const USkeleton* Skeleton)
+	{
+		BoneReference.Initialize(Skeleton);
+	}
+};
+
+UENUM()
+namespace ENotifyTriggerMode
+{
+	enum Type
+	{
+		AllAnimations UMETA(DisplayName="All Animations"),
+		HighestWeightedAnimation UMETA(DisplayName="Highest Weighted Animation"),
+		None,
+	};
+}
+
 /**
- * Contains a grid of data points with weights from sample points in the space
+ * Allows multiple animations to be blended between based on input parameters
  */
 UCLASS(config=Engine, hidecategories=Object, MinimalAPI, BlueprintType)
-class UBlendSpace : public UBlendSpaceBase
+class UBlendSpace : public UAnimationAsset, public IInterpolationIndexProvider
 {
 	GENERATED_UCLASS_BODY()
 
-public:
-	virtual bool IsValidAdditiveType(EAdditiveAnimationType AdditiveType) const override;
-	virtual bool IsValidAdditive() const override;
-protected:
-	//~ Begin UBlendSpaceBase Interface
-	virtual EBlendSpaceAxis GetAxisToScale() const override;
-	virtual bool IsSameSamplePoint(const FVector& SamplePointA, const FVector& SamplePointB) const;
-	virtual void GetRawSamplesFromBlendInput(const FVector &BlendInput, TArray<FGridBlendSample, TInlineAllocator<4> > & OutBlendSamples) const override;
+	/** Required for accessing protected variable names */
+	friend class FBlendSpaceDetails;
+	friend class UAnimGraphNode_BlendSpaceGraphBase;
+
+	//~ Begin UObject Interface
+	virtual void PostLoad() override;
+	virtual void Serialize(FArchive& Ar) override;
 #if WITH_EDITOR
-	virtual void SnapSamplesToClosestGridPoint();
-	virtual void RemapSamplesToNewAxisRange() override;
+	virtual void PreEditChange(FProperty* PropertyAboutToChange) override;
+	virtual void PostEditChangeProperty( struct FPropertyChangedEvent& PropertyChangedEvent) override;
 #endif // WITH_EDITOR
-	//~ End UBlendSpaceBase Interface
+	//~ End UObject Interface
 
+	//~ Begin UAnimationAsset Interface
+	virtual void TickAssetPlayer(FAnimTickRecord& Instance, struct FAnimNotifyQueue& NotifyQueue, FAnimAssetTickContext& Context) const override;
+	// this is used in editor only when used for transition getter
+	// this doesn't mean max time. In Sequence, this is SequenceLength,
+	// but for BlendSpace CurrentTime is normalized [0,1], so this is 1
+	virtual float GetPlayLength() const override { return 1.f; }
+	virtual TArray<FName>* GetUniqueMarkerNames() override;
+	virtual bool IsValidAdditive() const override;
+#if WITH_EDITOR
+	virtual bool GetAllAnimationSequencesReferred(TArray<UAnimationAsset*>& AnimationAssets, bool bRecursive = true) override;
+	virtual void ReplaceReferredAnimations(const TMap<UAnimationAsset*, UAnimationAsset*>& ReplacementMap) override;
+	virtual int32 GetMarkerUpdateCounter() const;
+	void    RuntimeValidateMarkerData();
+#endif
+	//~ End UAnimationAsset Interface
+	
+	// Begin IInterpolationIndexProvider Overrides
 	/**
-	* Get Grid Samples from BlendInput, From Input, it will return the 4 points of the grid that this input belongs to.
+	* Get PerBoneInterpolationIndex for the input BoneIndex
+	* If nothing found, return INDEX_NONE
+	*/
+	virtual int32 GetPerBoneInterpolationIndex(int32 BoneIndex, const FBoneContainer& RequiredBones) const override;	
+	// End UBlendSpace Overrides
+
+	/** Returns whether or not the given additive animation type is compatible with the blendspace type */
+	ENGINE_API virtual bool IsValidAdditiveType(EAdditiveAnimationType AdditiveType) const;
+
+	/**
+	 * BlendSpace Get Animation Pose function
+	 */
+	UE_DEPRECATED(4.26, "Use GetAnimationPose with other signature")
+	ENGINE_API void GetAnimationPose(TArray<FBlendSampleData>& BlendSampleDataCache, /*out*/ FCompactPose& OutPose, /*out*/ FBlendedCurve& OutCurve) const;
+	
+	ENGINE_API void GetAnimationPose(TArray<FBlendSampleData>& BlendSampleDataCache, /*out*/ FAnimationPoseData& OutAnimationPoseData) const;
+
+	ENGINE_API void GetAnimationPose(TArray<FBlendSampleData>& BlendSampleDataCache, TArrayView<FPoseLink> InPoseLinks, /*out*/ FPoseContext& Output) const;
+
+	/** Accessor for blend parameter **/
+	ENGINE_API const FBlendParameter& GetBlendParameter(const int32 Index) const;
+
+	/** Get this blend spaces sample data */
+	const TArray<struct FBlendSample>& GetBlendSamples() const { return SampleData; }
+
+	/** Returns the Blend Sample at the given index, will assert on invalid indices */
+	ENGINE_API const struct FBlendSample& GetBlendSample(const int32 SampleIndex) const;
+
+	/**
+	* Get Grid Samples from BlendInput
+	* It will return all samples that has weight > KINDA_SMALL_NUMBER
 	*
 	* @param	BlendInput	BlendInput X, Y, Z corresponds to BlendParameters[0], [1], [2]
 	*
-	* @return	LeftTop, RightTop, LeftBottom, RightBottom	4 corner of the grid this BlendInput belongs to
-	*			It's possible they return INDEX_NONE, but that case the weight also should be 0.f
+	* @return	true if it has valid OutSampleDataList, false otherwise
+	*/
+	ENGINE_API bool GetSamplesFromBlendInput(const FVector &BlendInput, TArray<FBlendSampleData> & OutSampleDataList) const;
+
+	/** Initialize BlendSpace for runtime. It needs certain data to be reinitialized per instsance **/
+	ENGINE_API void InitializeFilter(FBlendFilter* Filter) const;
+
+	/** Returns the blend input after clamping and/or wrapping */
+	ENGINE_API FVector GetClampedAndWrappedBlendInput(const FVector& BlendInput) const;
+
+	/** 
+	 * Updates a cached set of blend samples according to internal parameters, blendspace position and a delta time. Used internally to GetAnimationPose().
+	 * Note that this function does not perform any filtering internally.
+ 	 * @param	InBlendSpacePosition	The current position parameter of the blendspace
+	 * @param	InOutSampleDataCache	The sample data cache. Previous frames samples are re-used in the case of target weight interpolation
+	 * @param	InDeltaTime				The tick time for this update
+	 */
+	ENGINE_API bool UpdateBlendSamples(const FVector& InBlendSpacePosition, float InDeltaTime, TArray<FBlendSampleData>& InOutSampleDataCache) const;
+
+	/** Interpolate BlendInput based on Filter data **/
+	ENGINE_API FVector FilterInput(FBlendFilter* Filter, const FVector& BlendInput, float DeltaTime) const;
+
+#if WITH_EDITOR	
+	/** Validates sample data for blendspaces using the given animation sequence */
+	ENGINE_API static void UpdateBlendSpacesUsingAnimSequence(UAnimSequenceBase* Sequence);
+
+	/** Validates the contained data */
+	ENGINE_API void ValidateSampleData();
+
+	/** Add samples */
+	ENGINE_API bool AddSample(const FVector& SampleValue);
+	ENGINE_API bool	AddSample(UAnimSequence* AnimationSequence, const FVector& SampleValue);
+
+	/** edit samples */
+	ENGINE_API bool	EditSampleValue(const int32 BlendSampleIndex, const FVector& NewValue, bool bSnap = true);
+
+	UE_DEPRECATED(5.0, "Please use ReplaceSampleAnimation instead")
+	ENGINE_API bool	UpdateSampleAnimation(UAnimSequence* AnimationSequence, const FVector& SampleValue);
+
+	/** update animation on grid sample */
+	ENGINE_API bool	ReplaceSampleAnimation(const int32 BlendSampleIndex, UAnimSequence* AnimationSequence);
+
+	/** delete samples */
+	ENGINE_API bool	DeleteSample(const int32 BlendSampleIndex);
+	
+	/** Get the number of sample points for this blend space */
+	ENGINE_API int32 GetNumberOfBlendSamples()  const { return SampleData.Num(); }
+
+	/** Check whether or not the sample index is valid in combination with the stored sample data */
+	ENGINE_API bool IsValidBlendSampleIndex(const int32 SampleIndex) const;
+
+	/**
+	* return GridSamples from this BlendSpace
+	*
+	* @param	OutGridElements
+	*
+	* @return	Number of OutGridElements
+	*/
+	ENGINE_API const TArray<FEditorElement>& GetGridSamples() const;
+
+	/** Fill up local GridElements from the grid elements that are created using the sorted points
+	*	This will map back to original index for result
+	*
+	*  @param	SortedPointList		This is the pointlist that are used to create the given GridElements
+	*								This list contains subsets of the points it originally requested for visualization and sorted
+	*
+	*/
+	ENGINE_API void FillupGridElements(const TArray<int32> & PointListToSampleIndices, const TArray<FEditorElement> & GridElements, const TArray<int32>& InDimensionIndices);
+		
+	ENGINE_API void EmptyGridElements();
+	
+	/** Validate that the given animation sequence and contained blendspace data */
+	ENGINE_API bool ValidateAnimationSequence(const UAnimSequence* AnimationSequence) const;
+
+	/** Check if the blend spaces contains samples whos additive type match that of the animation sequence */
+	ENGINE_API bool DoesAnimationMatchExistingSamples(const UAnimSequence* AnimationSequence) const;
+	
+	/** Check if the the blendspace contains additive samples only */	
+	ENGINE_API bool ShouldAnimationBeAdditive() const;
+
+	/** Check if the animation sequence's skeleton is compatible with this blendspace */
+	ENGINE_API bool IsAnimationCompatibleWithSkeleton(const UAnimSequence* AnimationSequence) const;
+
+	/** Check if the animation sequence additive type is compatible with this blend space */
+	ENGINE_API bool IsAnimationCompatible(const UAnimSequence* AnimationSequence) const;
+
+	/** Validates supplied blend sample against current contents of blendspace */
+	ENGINE_API bool ValidateSampleValue(const FVector& SampleValue, int32 OriginalIndex = INDEX_NONE) const;
+
+	ENGINE_API bool IsSampleWithinBounds(const FVector &SampleValue) const;
+
+	/** Check if given sample value isn't too close to existing sample point **/
+	ENGINE_API bool IsTooCloseToExistingSamplePoint(const FVector& SampleValue, int32 OriginalIndex) const;
+#endif
+
+protected:
+	/**
+	* Get Grid Samples from BlendInput, From Input, it will populate OutGridSamples with the closest grid points.
+	*
+	* @param	BlendInput			BlendInput X, Y, Z corresponds to BlendParameters[0], [1], [2]
+	* @param	OutBlendSamples		Populated with the samples nearest the BlendInput
 	*
 	*/
-	void GetGridSamplesFromBlendInput(const FVector &BlendInput, FGridBlendSample & LeftBottom, FGridBlendSample & RightBottom, FGridBlendSample & LeftTop, FGridBlendSample& RightTop) const;
+	void GetRawSamplesFromBlendInput(const FVector& BlendInput, TArray<FGridBlendSample, TInlineAllocator<4> >& OutBlendSamples) const;
+
+	/** Returns the axis which can be used to scale animation speed. **/
+	virtual EBlendSpaceAxis GetAxisToScale() const { return AxisToScaleAnimation; }
+
+	/** Initialize Per Bone Blend **/
+	void InitializePerBoneBlend();
+
+	void TickFollowerSamples(TArray<FBlendSampleData> &SampleDataList, const int32 HighestWeightIndex, FAnimAssetTickContext &Context, bool bResetMarkerDataOnFollowers) const;
+
+	/** Utility function to calculate animation length from sample data list **/
+	float GetAnimationLengthFromSampleData(const TArray<FBlendSampleData> & SampleDataList) const;
+
+	/** Returns the blend input clamped to the valid range, unless that axis has been set to wrap in which case no clamping is done **/
+	FVector GetClampedBlendInput(const FVector& BlendInput) const;
+	
+	/** Translates BlendInput to grid space */
+	FVector GetNormalizedBlendInput(const FVector& BlendInput) const;
+
+	/** Returns the grid element at Index or NULL if Index is not valid */
+	const FEditorElement* GetGridSampleInternal(int32 Index) const;
+	
+	/** Utility function to interpolate weight of samples from OldSampleDataList to NewSampleDataList and copy back the interpolated result to FinalSampleDataList **/
+	bool InterpolateWeightOfSampleData(float DeltaTime, const TArray<FBlendSampleData> & OldSampleDataList, const TArray<FBlendSampleData> & NewSampleDataList, TArray<FBlendSampleData> & FinalSampleDataList) const;
 
+	/** Returns whether or not all animation set on the blend space samples match the given additive type */
+	bool ContainsMatchingSamples(EAdditiveAnimationType AdditiveType) const;
+
+	/** Checks if the given samples points overlap */
+	bool IsSameSamplePoint(const FVector& SamplePointA, const FVector& SamplePointB) const;	
+
+#if WITH_EDITOR
+	bool ContainsNonAdditiveSamples() const;
+	void UpdatePreviewBasePose();
+	/** If around border, snap to the border to avoid empty hole of data that is not valid **/
+	void SnapSamplesToClosestGridPoint();
+
+	void RemapSamplesToNewAxisRange();
+#endif // WITH_EDITOR
+	
+private:
+	// Internal helper function for GetAnimationPose variants
+	void GetAnimationPose_Internal(TArray<FBlendSampleData>& BlendSampleDataCache, TArrayView<FPoseLink> InPoseLinks, FAnimInstanceProxy* InProxy, bool bInExpectsAdditivePose, /*out*/ FAnimationPoseData& OutAnimationPoseData) const;
+
+	// Internal helper function for UpdateBlendSamples and TickAssetPlayer
+	bool UpdateBlendSamples_Internal(const FVector& InBlendSpacePosition, float InDeltaTime, TArray<FBlendSampleData>& InOutOldSampleDataList, TArray<FBlendSampleData>& InOutSampleDataCache) const;
+
+public:
+	/**
+	* When you use blend per bone, allows rotation to blend in mesh space. This only works if this does not contain additive animation samples
+	* This is more performance intensive
+	*/
+	UPROPERTY()
+	bool bRotationBlendInMeshSpace;
+
+#if WITH_EDITORONLY_DATA
+	/** Preview Base pose for additive BlendSpace **/
+	UPROPERTY(EditAnywhere, Category = AdditiveSettings)
+	TObjectPtr<UAnimSequence> PreviewBasePose;
+#endif // WITH_EDITORONLY_DATA
+
+	/** This animation length changes based on current input (resulting in different blend time)**/
+	UPROPERTY(transient)
+	float AnimLength;
+
+	/** Input interpolation parameter for all 3 axis, for each axis input, decide how you'd like to interpolate input to*/
+	UPROPERTY(EditAnywhere, Category = InputInterpolation)
+	FInterpolationParameter	InterpolationParam[3];
+
+	/**
+	* If greater than zero, this is the speed at which the sample weights are allowed to change.
+	* 
+	* A speed of 1 means a sample weight can change from zero to one (or one to zero) in one second.
+	* A speed of 2 means that this would take half a second.
+	* 
+	* This allows the Blend Space to switch to new parameters without going through intermediate states, 
+	* effectively blending between where it was and where the new target is. For example, imagine we have 
+	* a blend space for locomotion, moving left, forward and right. Now if you interpolate the inputs of 
+	* the blend space itself, from one extreme to the other, you will go from left, to forward, to right. 
+	* As an alternative, by setting this Sample Weight Speed to a value higher than zero, it will go 
+	* directly from left to right, without going through moving forward first.
+	* 
+	* Smaller values mean slower adjustments of the sample weights, and thus more smoothing. However, a 
+	* value of zero disables this smoothing entirely.
+	*/
+	UPROPERTY(EditAnywhere, Category = SampleInterpolation, meta = (DisplayName = "Sample Weight Speed"))
+	float TargetWeightInterpolationSpeedPerSec;
+
+	/**
+	 * If set then this eases in/out the sample weight adjustments, using the speed to determine how much smoothing to apply.
+	 */
+	UPROPERTY(EditAnywhere, Category = SampleInterpolation, meta = (DisplayName = "Sample Weight Speed Smoothing "))
+	bool bTargetWeightInterpolationEaseInOut;
+
+	/** The current mode used by the BlendSpace to decide which animation notifies to fire. Valid options are:
+	- AllAnimations - All notify events will fire
+	- HighestWeightedAnimation - Notify events will only fire from the highest weighted animation
+	- None - No notify events will fire from any animations
+	*/
+	UPROPERTY(EditAnywhere, Category = AnimationNotifies)
+	TEnumAsByte<ENotifyTriggerMode::Type> NotifyTriggerMode;
+
+protected:
+
+	/**
+	* Per bone interpolation speed settings. 
+	* These act as overrides to the global interpolation speed. 
+	* This means the global interpolation speed does not impact these bones.
+	*/
+	UPROPERTY(EditAnywhere, Category = SampleInterpolation, meta = (DisplayName="Per Bone Overrides"))
+	TArray<FPerBoneInterpolation> PerBoneBlend;
+
+	/** Track index to get marker data from. Samples are tested for the suitability of marker based sync
+	    during load and if we can use marker based sync we cache an index to a representative sample here */
+	UPROPERTY()
+	int32 SampleIndexWithMarkers;
+
+	/** Sample animation data **/
+	UPROPERTY(EditAnywhere, Category=BlendSamples)
+	TArray<struct FBlendSample> SampleData;
+
+	/** Grid samples, indexing scheme imposed by subclass **/
+	UPROPERTY()
+	TArray<struct FEditorElement> GridSamples;
+	
+	/** Blend Parameters for each axis. **/
+	UPROPERTY(EditAnywhere, Category = BlendParametersTest)
+	struct FBlendParameter BlendParameters[3];
+
+	/** If you have input interpolation, which axis to drive animation speed (scale) - i.e. for locomotion animation, speed axis will drive animation speed (thus scale)**/
+	UPROPERTY(EditAnywhere, Category = InputInterpolation)
+	TEnumAsByte<EBlendSpaceAxis> AxisToScaleAnimation;
+
+	/** Reset to reference pose. It does apply different refpose based on additive or not*/
+	void ResetToRefPose(FCompactPose& OutPose) const;
+
+	/** The order in which to use the dimensions in the data - e.g. [1, 2] means a 2D blend using Y and Z */
+	UPROPERTY()
+	TArray<int32> DimensionIndices;
+
+#if WITH_EDITOR
+private:
+	// Track whether we have updated markers so cached data can be updated
+	int32 MarkerDataUpdateCounter;
+protected:
+	FVector PreviousAxisMinMaxValues[3];
+#endif	
+
+protected:
+
+	//================================================================================================
+	// 1D specific functions 
+	//================================================================================================
+	void GetRawSamplesFromBlendInput1D(const FVector& BlendInput, TArray<FGridBlendSample, TInlineAllocator<4> >& OutBlendSamples) const;
+#if WITH_EDITOR
+	void SnapSamplesToClosestGridPoint1D();
+	void RemapSamplesToNewAxisRange1D();
+#endif // WITH_EDITOR
+
+	//================================================================================================
+	// 2D specific functions 
+	//================================================================================================
+	void GetRawSamplesFromBlendInput2D(const FVector& BlendInput, TArray<FGridBlendSample, TInlineAllocator<4> >& OutBlendSamples) const;
+#if WITH_EDITOR
+	void SnapSamplesToClosestGridPoint2D();
+	void RemapSamplesToNewAxisRange2D();
+#endif // WITH_EDITOR
+
+private:
 	/** Get the Editor Element from Index
 	*
 	* @param	XIndex	Index of X
@@ -49,8 +591,4 @@ protected:
 	* @return	FEditorElement * return the grid data
 	*/
 	const FEditorElement* GetEditorElement(int32 XIndex, int32 YIndex) const;
-protected:
-	/** If you have input interpolation, which axis to drive animation speed (scale) - i.e. for locomotion animation, speed axis will drive animation speed (thus scale)**/
-	UPROPERTY(EditAnywhere, Category = InputInterpolation)
-	TEnumAsByte<EBlendSpaceAxis> AxisToScaleAnimation;
 };
diff --git a/Engine/Source/Runtime/Engine/Classes/Animation/BlendSpace1D.h b/Engine/Source/Runtime/Engine/Classes/Animation/BlendSpace1D.h
--- a/Engine/Source/Runtime/Engine/Classes/Animation/BlendSpace1D.h
+++ b/Engine/Source/Runtime/Engine/Classes/Animation/BlendSpace1D.h
@@ -9,11 +9,11 @@
 
 #include "CoreMinimal.h"
 #include "UObject/ObjectMacros.h"
-#include "Animation/BlendSpaceBase.h"
+#include "Animation/BlendSpace.h"
 #include "BlendSpace1D.generated.h"
 
 UCLASS(config=Engine, hidecategories=Object, MinimalAPI, BlueprintType)
-class UBlendSpace1D : public UBlendSpaceBase
+class UBlendSpace1D : public UBlendSpace
 {
 	GENERATED_UCLASS_BODY()
 
@@ -28,16 +28,8 @@ public:
 	UPROPERTY(EditAnywhere, Category = InputInterpolation)
 	bool bScaleAnimation;
 
-	virtual bool IsValidAdditive() const override;
-	virtual bool IsValidAdditiveType(EAdditiveAnimationType AdditiveType) const override;
 protected:
-	//~ Begin UBlendSpaceBase Interface
-	virtual bool IsSameSamplePoint(const FVector& SamplePointA, const FVector& SamplePointB) const;	
+	//~ Begin UBlendSpace Interface
 	virtual EBlendSpaceAxis GetAxisToScale() const override;
-	virtual void GetRawSamplesFromBlendInput(const FVector &BlendInput, TArray<FGridBlendSample, TInlineAllocator<4> > & OutBlendSamples) const override;
-#if WITH_EDITOR
-	virtual void SnapSamplesToClosestGridPoint() override;
-	virtual void RemapSamplesToNewAxisRange() override;
-#endif // WITH_EDITOR
-	//~ End UBlendSpaceBase Interface
+	//~ End UBlendSpace Interface
 };
diff --git a/Engine/Source/Runtime/Engine/Classes/Animation/BlendSpaceBase.h b/Engine/Source/Runtime/Engine/Classes/Animation/BlendSpaceBase.h
--- a/Engine/Source/Runtime/Engine/Classes/Animation/BlendSpaceBase.h
+++ b/Engine/Source/Runtime/Engine/Classes/Animation/BlendSpaceBase.h
@@ -1,555 +1,15 @@
 // Copyright Epic Games, Inc. All Rights Reserved.
 
 /**
- * Blend Space Base. Contains base functionality shared across all blend space objects
+ * Note that this file is now deprecated, and all functionality has moved to BlendSpace.
  *
  */
 
 #pragma once
 
-#include "CoreMinimal.h"
-#include "UObject/ObjectMacros.h"
-#include "BoneContainer.h"
-#include "Animation/AnimationAsset.h"
-#include "AnimationRuntime.h"
-#include "AnimNodeBase.h"
-#include "Containers/ArrayView.h"
-#include "BlendSpaceBase.generated.h"
+#include "Animation/BlendSpace.h"
 
-/** Interpolation data types. */
-UENUM()
-enum EBlendSpaceAxis
-{
-	BSA_None UMETA(DisplayName = "None"),
-	BSA_X UMETA(DisplayName = "Horizontal (X) Axis"),
-	BSA_Y UMETA(DisplayName = "Vertical (Y) Axis"),
-	BSA_Max
-};
+UE_STATIC_DEPRECATE(5.0, true, "BlendSpaceBase.h has been deprecated - please include and use BlendSpace.h instead.");
 
-USTRUCT()
-struct FInterpolationParameter
-{
-	GENERATED_USTRUCT_BODY()
-
-	/** Interpolation Time for input, when it gets input, it will use this time to interpolate to target, used for smoother interpolation. */
-	UPROPERTY(EditAnywhere, Category=Parameter)
-	float InterpolationTime = 0.f;
-
-	/** Type of interpolation used for filtering the input value to decide how to get to target. */
-	UPROPERTY(EditAnywhere, Category=Parameter)
-	TEnumAsByte<EFilterInterpolationType> InterpolationType = EFilterInterpolationType::BSIT_Average;
-};
-
-USTRUCT()
-struct FBlendParameter
-{
-	GENERATED_USTRUCT_BODY()
-
-	UPROPERTY(EditAnywhere, DisplayName = "Name", Category=BlendParameter)
-	FString DisplayName;
-
-	/** Min value for this parameter. */
-	UPROPERTY(EditAnywhere, DisplayName = "Minimum Axis Value", Category=BlendParameter)
-	float Min;
-
-	/** Max value for this parameter. */
-	UPROPERTY(EditAnywhere, DisplayName = "Maximum Axis Value", Category=BlendParameter)
-	float Max;
-
-	/** The number of grid divisions for this parameter (axis). */
-	UPROPERTY(EditAnywhere, DisplayName = "Number of Grid Divisions", Category=BlendParameter, meta=(UIMin="1", ClampMin="1"))
-	int32 GridNum;
-
-	/** If false then input parameters are clamped to the min/max values on this axis. If true then the input can go outside the min/max range and the blend space is treated as being cyclic on this axis. */
-	UPROPERTY(EditAnywhere, DisplayName = "Wrap Input", Category = BlendParameter)
-	bool bWrapInput;
-
-	FBlendParameter()
-		: DisplayName(TEXT("None"))
-		, Min(0.f)
-		, Max(100.f)
-		, GridNum(4) // TODO when changing GridNum's default value, it breaks all grid samples ATM - provide way to rebuild grid samples during loading
-		, bWrapInput(false)
-	{
-	}
-
-	float GetRange() const
-	{
-		return Max-Min;
-	}
-	/** Return size of each grid. */
-	float GetGridSize() const
-	{
-		return GetRange()/(float)GridNum;
-	}
-	
-};
-
-/** Sample data */
-USTRUCT()
-struct FBlendSample
-{
-	GENERATED_USTRUCT_BODY()
-
-	// For linked animations
-	UPROPERTY(EditAnywhere, Category=BlendSample)
-	TObjectPtr<class UAnimSequence> Animation;
-
-	//blend 0->x, blend 1->y, blend 2->z
-
-	UPROPERTY(EditAnywhere, Category=BlendSample)
-	FVector SampleValue;
-	
-	UPROPERTY(EditAnywhere, Category = BlendSample, meta=(UIMin="0.01", UIMax="2.0", ClampMin="0.01", ClampMax="64.0"))
-	float RateScale;
-
-#if WITH_EDITORONLY_DATA
-	UPROPERTY(EditAnywhere, Category=BlendSample)
-	uint8 bSnapToGrid : 1;
-
-	UPROPERTY(transient)
-	uint8 bIsValid : 1;
-
-	// Cache the samples marker data counter so that we can track if it changes and revalidate the blendspace
-	int32 CachedMarkerDataUpdateCounter;
-
-#endif // WITH_EDITORONLY_DATA
-
-	FBlendSample()
-		: Animation(nullptr)
-		, SampleValue(0.f)
-		, RateScale(1.0f)
-#if WITH_EDITORONLY_DATA
-		, bSnapToGrid(true)
-		, bIsValid(false)
-		, CachedMarkerDataUpdateCounter(INDEX_NONE)
-#endif // WITH_EDITORONLY_DATA
-	{		
-	}
-	
-	FBlendSample(class UAnimSequence* InAnim, FVector InValue, bool bInIsSnapped, bool bInIsValid) 
-		: Animation(InAnim)
-		, SampleValue(InValue)
-		, RateScale(1.0f)
-#if WITH_EDITORONLY_DATA
-		, bSnapToGrid(bInIsSnapped)
-		, bIsValid(bInIsValid)
-		, CachedMarkerDataUpdateCounter(INDEX_NONE)
-#endif // WITH_EDITORONLY_DATA
-	{		
-	}
-	
-	bool operator==( const FBlendSample& Other ) const 
-	{
-		return (Other.Animation == Animation && Other.SampleValue == SampleValue && FMath::IsNearlyEqual(Other.RateScale, RateScale));
-	}
-};
-
-/**
- * Each elements in the grid
- */
-USTRUCT()
-struct FEditorElement
-{
-	GENERATED_USTRUCT_BODY()
-
-	// for now we only support triangles
-	static const int32 MAX_VERTICES = 3;
-
-	UPROPERTY(EditAnywhere, Category=EditorElement)
-	int32 Indices[MAX_VERTICES];
-
-	UPROPERTY(EditAnywhere, Category=EditorElement)
-	float Weights[MAX_VERTICES];
-
-	FEditorElement()
-	{
-		for (int32 ElementIndex = 0; ElementIndex < MAX_VERTICES; ElementIndex++)
-		{
-			Indices[ElementIndex] = INDEX_NONE;
-		}
-		for (int32 ElementIndex = 0; ElementIndex < MAX_VERTICES; ElementIndex++)
-		{
-			Weights[ElementIndex] = 0;
-		}
-	}
-	
-};
-
-/** result of how much weight of the grid element **/
-USTRUCT()
-struct FGridBlendSample
-{
-	GENERATED_USTRUCT_BODY()
-
-	UPROPERTY()
-	struct FEditorElement GridElement;
-
-	UPROPERTY()
-	float BlendWeight;
-
-	FGridBlendSample()
-		: BlendWeight(0)
-	{
-	}
-
-};
-
-USTRUCT()
-struct FPerBoneInterpolation
-{
-	GENERATED_USTRUCT_BODY()
-
-	UPROPERTY(EditAnywhere, Category=FPerBoneInterpolation)
-	FBoneReference BoneReference;
-
-	/**
-	* If greater than zero, this is the speed at which the sample weights are allowed to change for this specific bone.
-	*
-	* A speed of 1 means a sample weight can change from zero to one (or one to zero) in one second.
-	* A speed of 2 means that this would take half a second.
-	*
-	* Smaller values mean slower adjustments of the sample weights, and thus more smoothing. However, a
-	* value of zero disables this smoothing entirely.
-	* 
-	* If set, the value overrides the overall Sample Weight Speed which will no longer affect this bone.
-	*/
-	UPROPERTY(EditAnywhere, Category=FPerBoneInterpolation, meta=(DisplayName="Sample Weight Speed"))
-	float InterpolationSpeedPerSec;
-
-	FPerBoneInterpolation()
-		: InterpolationSpeedPerSec(6.f)
-	{}
-
-	void Initialize(const USkeleton* Skeleton)
-	{
-		BoneReference.Initialize(Skeleton);
-	}
-};
-
-UENUM()
-namespace ENotifyTriggerMode
-{
-	enum Type
-	{
-		AllAnimations UMETA(DisplayName="All Animations"),
-		HighestWeightedAnimation UMETA(DisplayName="Highest Weighted Animation"),
-		None,
-	};
-}
-
-/**
- * Allows multiple animations to be blended between based on input parameters
- */
-UCLASS(abstract, config=Engine, hidecategories=Object, MinimalAPI, BlueprintType)
-class UBlendSpaceBase : public UAnimationAsset, public IInterpolationIndexProvider
-{
-	GENERATED_UCLASS_BODY()
-
-	/** Required for accessing protected variable names */
-	friend class FBlendSpaceDetails;
-	friend class UAnimGraphNode_BlendSpaceGraphBase;
-
-	//~ Begin UObject Interface
-	virtual void PostLoad() override;
-	virtual void Serialize(FArchive& Ar) override;
-#if WITH_EDITOR
-	virtual void PreEditChange(FProperty* PropertyAboutToChange) override;
-	virtual void PostEditChangeProperty( struct FPropertyChangedEvent& PropertyChangedEvent) override;
-#endif // WITH_EDITOR
-	//~ End UObject Interface
-
-	//~ Begin UAnimationAsset Interface
-	virtual void TickAssetPlayer(FAnimTickRecord& Instance, struct FAnimNotifyQueue& NotifyQueue, FAnimAssetTickContext& Context) const override;
-	// this is used in editor only when used for transition getter
-	// this doesn't mean max time. In Sequence, this is SequenceLength,
-	// but for BlendSpace CurrentTime is normalized [0,1], so this is 1
-	virtual float GetPlayLength() const override { return 1.f; }
-	virtual TArray<FName>* GetUniqueMarkerNames() override;
-	virtual bool IsValidAdditive() const override;
-#if WITH_EDITOR
-	virtual bool GetAllAnimationSequencesReferred(TArray<UAnimationAsset*>& AnimationAssets, bool bRecursive = true) override;
-	virtual void ReplaceReferredAnimations(const TMap<UAnimationAsset*, UAnimationAsset*>& ReplacementMap) override;
-	virtual int32 GetMarkerUpdateCounter() const;
-	void    RuntimeValidateMarkerData();
-#endif
-	//~ End UAnimationAsset Interface
-	
-	// Begin IInterpolationIndexProvider Overrides
-	/**
-	* Get PerBoneInterpolationIndex for the input BoneIndex
-	* If nothing found, return INDEX_NONE
-	*/
-	virtual int32 GetPerBoneInterpolationIndex(int32 BoneIndex, const FBoneContainer& RequiredBones) const override;	
-	// End UBlendSpaceBase Overrides
-
-	/** Returns whether or not the given additive animation type is compatible with the blendspace type */
-	ENGINE_API virtual bool IsValidAdditiveType(EAdditiveAnimationType AdditiveType) const;
-
-	/**
-	 * BlendSpace Get Animation Pose function
-	 */
-	UE_DEPRECATED(4.26, "Use GetAnimationPose with other signature")
-	ENGINE_API void GetAnimationPose(TArray<FBlendSampleData>& BlendSampleDataCache, /*out*/ FCompactPose& OutPose, /*out*/ FBlendedCurve& OutCurve) const;
-	
-	ENGINE_API void GetAnimationPose(TArray<FBlendSampleData>& BlendSampleDataCache, /*out*/ FAnimationPoseData& OutAnimationPoseData) const;
-
-	ENGINE_API void GetAnimationPose(TArray<FBlendSampleData>& BlendSampleDataCache, TArrayView<FPoseLink> InPoseLinks, /*out*/ FPoseContext& Output) const;
-
-	/** Accessor for blend parameter **/
-	ENGINE_API const FBlendParameter& GetBlendParameter(const int32 Index) const;
-
-	/** Get this blend spaces sample data */
-	const TArray<struct FBlendSample>& GetBlendSamples() const { return SampleData; }
-
-	/** Returns the Blend Sample at the given index, will assert on invalid indices */
-	ENGINE_API const struct FBlendSample& GetBlendSample(const int32 SampleIndex) const;
-
-	/**
-	* Get Grid Samples from BlendInput
-	* It will return all samples that has weight > KINDA_SMALL_NUMBER
-	*
-	* @param	BlendInput	BlendInput X, Y, Z corresponds to BlendParameters[0], [1], [2]
-	*
-	* @return	true if it has valid OutSampleDataList, false otherwise
-	*/
-	ENGINE_API bool GetSamplesFromBlendInput(const FVector &BlendInput, TArray<FBlendSampleData> & OutSampleDataList) const;
-
-	/** Initialize BlendSpace for runtime. It needs certain data to be reinitialized per instsance **/
-	ENGINE_API void InitializeFilter(FBlendFilter* Filter) const;
-
-	/** Returns the blend input after clamping and/or wrapping */
-	ENGINE_API FVector GetClampedAndWrappedBlendInput(const FVector& BlendInput) const;
-
-	/** 
-	 * Updates a cached set of blend samples according to internal parameters, blendspace position and a delta time. Used internally to GetAnimationPose().
-	 * Note that this function does not perform any filtering internally.
- 	 * @param	InBlendSpacePosition	The current position parameter of the blendspace
-	 * @param	InOutSampleDataCache	The sample data cache. Previous frames samples are re-used in the case of target weight interpolation
-	 * @param	InDeltaTime				The tick time for this update
-	 */
-	ENGINE_API bool UpdateBlendSamples(const FVector& InBlendSpacePosition, float InDeltaTime, TArray<FBlendSampleData>& InOutSampleDataCache) const;
-
-	/** Interpolate BlendInput based on Filter data **/
-	ENGINE_API FVector FilterInput(FBlendFilter* Filter, const FVector& BlendInput, float DeltaTime) const;
-
-#if WITH_EDITOR	
-	/** Validates sample data for blendspaces using the given animation sequence */
-	ENGINE_API static void UpdateBlendSpacesUsingAnimSequence(UAnimSequenceBase* Sequence);
-
-	/** Validates the contained data */
-	ENGINE_API void ValidateSampleData();
-
-	/** Add samples */
-	ENGINE_API bool AddSample(const FVector& SampleValue);
-	ENGINE_API bool	AddSample(UAnimSequence* AnimationSequence, const FVector& SampleValue);
-
-	/** edit samples */
-	ENGINE_API bool	EditSampleValue(const int32 BlendSampleIndex, const FVector& NewValue, bool bSnap = true);
-
-	UE_DEPRECATED(5.0, "Please use ReplaceSampleAnimation instead")
-	ENGINE_API bool	UpdateSampleAnimation(UAnimSequence* AnimationSequence, const FVector& SampleValue);
-
-	/** update animation on grid sample */
-	ENGINE_API bool	ReplaceSampleAnimation(const int32 BlendSampleIndex, UAnimSequence* AnimationSequence);
-
-	/** delete samples */
-	ENGINE_API bool	DeleteSample(const int32 BlendSampleIndex);
-	
-	/** Get the number of sample points for this blend space */
-	ENGINE_API int32 GetNumberOfBlendSamples()  const { return SampleData.Num(); }
-
-	/** Check whether or not the sample index is valid in combination with the stored sample data */
-	ENGINE_API bool IsValidBlendSampleIndex(const int32 SampleIndex) const;
-
-	/**
-	* return GridSamples from this BlendSpace
-	*
-	* @param	OutGridElements
-	*
-	* @return	Number of OutGridElements
-	*/
-	ENGINE_API const TArray<FEditorElement>& GetGridSamples() const;
-
-	/** Fill up local GridElements from the grid elements that are created using the sorted points
-	*	This will map back to original index for result
-	*
-	*  @param	SortedPointList		This is the pointlist that are used to create the given GridElements
-	*								This list contains subsets of the points it originally requested for visualization and sorted
-	*
-	*/
-	ENGINE_API void FillupGridElements(const TArray<int32> & PointListToSampleIndices, const TArray<FEditorElement> & GridElements);
-		
-	ENGINE_API void EmptyGridElements();
-	
-	/** Validate that the given animation sequence and contained blendspace data */
-	ENGINE_API bool ValidateAnimationSequence(const UAnimSequence* AnimationSequence) const;
-
-	/** Check if the blend spaces contains samples whos additive type match that of the animation sequence */
-	ENGINE_API bool DoesAnimationMatchExistingSamples(const UAnimSequence* AnimationSequence) const;
-	
-	/** Check if the the blendspace contains additive samples only */	
-	ENGINE_API bool ShouldAnimationBeAdditive() const;
-
-	/** Check if the animation sequence's skeleton is compatible with this blendspace */
-	ENGINE_API bool IsAnimationCompatibleWithSkeleton(const UAnimSequence* AnimationSequence) const;
-
-	/** Check if the animation sequence additive type is compatible with this blend space */
-	ENGINE_API bool IsAnimationCompatible(const UAnimSequence* AnimationSequence) const;
-
-	/** Validates supplied blend sample against current contents of blendspace */
-	ENGINE_API bool ValidateSampleValue(const FVector& SampleValue, int32 OriginalIndex = INDEX_NONE) const;
-
-	ENGINE_API bool IsSampleWithinBounds(const FVector &SampleValue) const;
-
-	/** Check if given sample value isn't too close to existing sample point **/
-	ENGINE_API bool IsTooCloseToExistingSamplePoint(const FVector& SampleValue, int32 OriginalIndex) const;
-#endif
-
-protected:
-	/**
-	* Get Grid Samples from BlendInput, From Input, it will populate OutGridSamples with the closest grid points.
-	*
-	* @param	BlendInput			BlendInput X, Y, Z corresponds to BlendParameters[0], [1], [2]
-	* @param	OutBlendSamples		Populated with the samples nearest the BlendInput
-	*
-	*/
-	virtual void GetRawSamplesFromBlendInput(const FVector &BlendInput, TArray<FGridBlendSample, TInlineAllocator<4> > & OutBlendSamples) const {}
-	/** Let derived blend space decided how to handle scaling */
-	virtual EBlendSpaceAxis GetAxisToScale() const PURE_VIRTUAL(UBlendSpaceBase::GetAxisToScale, return BSA_None;);
-
-	/** Initialize Per Bone Blend **/
-	void InitializePerBoneBlend();
-
-	void TickFollowerSamples(TArray<FBlendSampleData> &SampleDataList, const int32 HighestWeightIndex, FAnimAssetTickContext &Context, bool bResetMarkerDataOnFollowers) const;
-
-	/** Utility function to calculate animation length from sample data list **/
-	float GetAnimationLengthFromSampleData(const TArray<FBlendSampleData> & SampleDataList) const;
-
-	/** Returns the blend input clamped to the valid range, unless that axis has been set to wrap in which case no clamping is done **/
-	FVector GetClampedBlendInput(const FVector& BlendInput) const;
-	
-	/** Translates BlendInput to grid space */
-	FVector GetNormalizedBlendInput(const FVector& BlendInput) const;
-
-	/** Returns the grid element at Index or NULL if Index is not valid */
-	const FEditorElement* GetGridSampleInternal(int32 Index) const;
-	
-	/** Utility function to interpolate weight of samples from OldSampleDataList to NewSampleDataList and copy back the interpolated result to FinalSampleDataList **/
-	bool InterpolateWeightOfSampleData(float DeltaTime, const TArray<FBlendSampleData> & OldSampleDataList, const TArray<FBlendSampleData> & NewSampleDataList, TArray<FBlendSampleData> & FinalSampleDataList) const;
-
-	/** Returns whether or not all animation set on the blend space samples match the given additive type */
-	bool ContainsMatchingSamples(EAdditiveAnimationType AdditiveType) const;
-
-	/** Checks if the given samples points overlap */
-	virtual bool IsSameSamplePoint(const FVector& SamplePointA, const FVector& SamplePointB) const PURE_VIRTUAL(UBlendSpaceBase::IsSameSamplePoint, return false;);	
-
-#if WITH_EDITOR
-	bool ContainsNonAdditiveSamples() const;
-	void UpdatePreviewBasePose();
-	/** If around border, snap to the border to avoid empty hole of data that is not valid **/
-	virtual void SnapSamplesToClosestGridPoint() PURE_VIRTUAL(UBlendSpaceBase::SnapSamplesToClosestGridPoint, return;);
-
-	virtual void RemapSamplesToNewAxisRange() PURE_VIRTUAL(UBlendSpaceBase::RemapSamplesToNewAxisRange, return;);
-#endif // WITH_EDITOR
-	
-private:
-	// Internal helper function for GetAnimationPose variants
-	void GetAnimationPose_Internal(TArray<FBlendSampleData>& BlendSampleDataCache, TArrayView<FPoseLink> InPoseLinks, FAnimInstanceProxy* InProxy, bool bInExpectsAdditivePose, /*out*/ FAnimationPoseData& OutAnimationPoseData) const;
-
-	// Internal helper function for UpdateBlendSamples and TickAssetPlayer
-	bool UpdateBlendSamples_Internal(const FVector& InBlendSpacePosition, float InDeltaTime, TArray<FBlendSampleData>& InOutOldSampleDataList, TArray<FBlendSampleData>& InOutSampleDataCache) const;
-
-public:
-	/**
-	* When you use blend per bone, allows rotation to blend in mesh space. This only works if this does not contain additive animation samples
-	* This is more performance intensive
-	*/
-	UPROPERTY()
-	bool bRotationBlendInMeshSpace;
-
-#if WITH_EDITORONLY_DATA
-	/** Preview Base pose for additive BlendSpace **/
-	UPROPERTY(EditAnywhere, Category = AdditiveSettings)
-	TObjectPtr<UAnimSequence> PreviewBasePose;
-#endif // WITH_EDITORONLY_DATA
-
-	/** This animation length changes based on current input (resulting in different blend time)**/
-	UPROPERTY(transient)
-	float AnimLength;
-
-	/** Input interpolation parameter for all 3 axis, for each axis input, decide how you'd like to interpolate input to*/
-	UPROPERTY(EditAnywhere, Category = InputInterpolation)
-	FInterpolationParameter	InterpolationParam[3];
-
-	/**
-	* If greater than zero, this is the speed at which the sample weights are allowed to change.
-	* 
-	* A speed of 1 means a sample weight can change from zero to one (or one to zero) in one second.
-	* A speed of 2 means that this would take half a second.
-	* 
-	* This allows the Blend Space to switch to new parameters without going through intermediate states, 
-	* effectively blending between where it was and where the new target is. For example, imagine we have 
-	* a blend space for locomotion, moving left, forward and right. Now if you interpolate the inputs of 
-	* the blend space itself, from one extreme to the other, you will go from left, to forward, to right. 
-	* As an alternative, by setting this Sample Weight Speed to a value higher than zero, it will go 
-	* directly from left to right, without going through moving forward first.
-	* 
-	* Smaller values mean slower adjustments of the sample weights, and thus more smoothing. However, a 
-	* value of zero disables this smoothing entirely.
-	*/
-	UPROPERTY(EditAnywhere, Category = SampleInterpolation, meta = (DisplayName = "Sample Weight Speed"))
-	float TargetWeightInterpolationSpeedPerSec;
-
-	/**
-	 * If set then this eases in/out the sample weight adjustments, using the speed to determine how much smoothing to apply.
-	 */
-	UPROPERTY(EditAnywhere, Category = SampleInterpolation, meta = (DisplayName = "Sample Weight Speed Smoothing "))
-	bool bTargetWeightInterpolationEaseInOut;
-
-	/** The current mode used by the BlendSpace to decide which animation notifies to fire. Valid options are:
-	- AllAnimations - All notify events will fire
-	- HighestWeightedAnimation - Notify events will only fire from the highest weighted animation
-	- None - No notify events will fire from any animations
-	*/
-	UPROPERTY(EditAnywhere, Category = AnimationNotifies)
-	TEnumAsByte<ENotifyTriggerMode::Type> NotifyTriggerMode;
-
-protected:
-
-	/**
-	* Per bone interpolation speed settings. 
-	* These act as overrides to the global interpolation speed. 
-	* This means the global interpolation speed does not impact these bones.
-	*/
-	UPROPERTY(EditAnywhere, Category = SampleInterpolation, meta = (DisplayName="Per Bone Overrides"))
-	TArray<FPerBoneInterpolation> PerBoneBlend;
-
-	/** Track index to get marker data from. Samples are tested for the suitability of marker based sync
-	    during load and if we can use marker based sync we cache an index to a representative sample here */
-	UPROPERTY()
-	int32 SampleIndexWithMarkers;
-
-	/** Sample animation data **/
-	UPROPERTY(EditAnywhere, Category=BlendSamples)
-	TArray<struct FBlendSample> SampleData;
-
-	/** Grid samples, indexing scheme imposed by subclass **/
-	UPROPERTY()
-	TArray<struct FEditorElement> GridSamples;
-	
-	/** Blend Parameters for each axis. **/
-	UPROPERTY(EditAnywhere, Category = BlendParametersTest)
-	struct FBlendParameter BlendParameters[3];
-
-	/** Reset to reference pose. It does apply different refpose based on additive or not*/
-	void ResetToRefPose(FCompactPose& OutPose) const;
-
-#if WITH_EDITOR
-private:
-	// Track whether we have updated markers so cached data can be updated
-	int32 MarkerDataUpdateCounter;
-protected:
-	FVector PreviousAxisMinMaxValues[3];
-#endif	
-};
+UE_DEPRECATED(5.0, "UBlendSpaceBase has been deprecated. Please use UBlendSpace instead.")
+typedef UBlendSpace UBlendSpaceBase;
diff --git a/Engine/Source/Runtime/Engine/Private/Animation/AnimBlueprintGeneratedClass.cpp b/Engine/Source/Runtime/Engine/Private/Animation/AnimBlueprintGeneratedClass.cpp
--- a/Engine/Source/Runtime/Engine/Private/Animation/AnimBlueprintGeneratedClass.cpp
+++ b/Engine/Source/Runtime/Engine/Private/Animation/AnimBlueprintGeneratedClass.cpp
@@ -224,7 +224,7 @@ void FAnimBlueprintDebugData::RecordSequencePlayer(int32 InNodeID, float InPosit
 	SequencePlayerRecordsThisFrame.Emplace(InNodeID, InPosition, InLength, InFrameCount);
 }
 
-void FAnimBlueprintDebugData::RecordBlendSpacePlayer(int32 InNodeID, const UBlendSpaceBase* InBlendSpace, const FVector& InPosition, const FVector& InFilteredPosition)
+void FAnimBlueprintDebugData::RecordBlendSpacePlayer(int32 InNodeID, const UBlendSpace* InBlendSpace, const FVector& InPosition, const FVector& InFilteredPosition)
 {
 	BlendSpacePlayerRecordsThisFrame.Emplace(InNodeID, InBlendSpace, InPosition, InFilteredPosition);
 }
diff --git a/Engine/Source/Runtime/Engine/Private/Animation/AnimInstance.cpp b/Engine/Source/Runtime/Engine/Private/Animation/AnimInstance.cpp
--- a/Engine/Source/Runtime/Engine/Private/Animation/AnimInstance.cpp
+++ b/Engine/Source/Runtime/Engine/Private/Animation/AnimInstance.cpp
@@ -15,7 +15,7 @@
 #include "Animation/AnimClassInterface.h"
 #include "Animation/AnimBlueprintGeneratedClass.h"
 #include "DisplayDebugHelpers.h"
-#include "Animation/BlendSpaceBase.h"
+#include "Animation/BlendSpace.h"
 #include "Animation/AnimNotifies/AnimNotifyState.h"
 #include "Animation/AnimNotifies/AnimNotify.h"
 #include "Engine/Canvas.h"
@@ -817,7 +817,7 @@ void OutputTickRecords(const TArray<FAnimTickRecord>& Records, UCanvas* Canvas,
 
 		DisplayDebugManager.DrawString(PlayerEntry, Indent);
 
-		if (UBlendSpaceBase* BlendSpace = Cast<UBlendSpaceBase>(Player.SourceAsset))
+		if (UBlendSpace* BlendSpace = Cast<UBlendSpace>(Player.SourceAsset))
 		{
 			if (bFullBlendspaceDisplay && Player.BlendSpace.BlendSampleDataCache && Player.BlendSpace.BlendSampleDataCache->Num() > 0)
 			{
diff --git a/Engine/Source/Runtime/Engine/Private/Animation/AnimInstanceProxy.cpp b/Engine/Source/Runtime/Engine/Private/Animation/AnimInstanceProxy.cpp
--- a/Engine/Source/Runtime/Engine/Private/Animation/AnimInstanceProxy.cpp
+++ b/Engine/Source/Runtime/Engine/Private/Animation/AnimInstanceProxy.cpp
@@ -4,7 +4,7 @@
 #include "Animation/AnimNodeBase.h"
 #include "Animation/PoseAsset.h"
 #include "AnimationRuntime.h"
-#include "Animation/BlendSpaceBase.h"
+#include "Animation/BlendSpace.h"
 #include "AnimationUtils.h"
 #include "Logging/MessageLog.h"
 #include "Animation/AnimNode_AssetPlayerBase.h"
@@ -653,7 +653,7 @@ void FAnimInstanceProxy::MakeSequenceTickRecord(FAnimTickRecord& TickRecord, cla
 	TickRecord.bLooping = bLooping;
 }
 
-void FAnimInstanceProxy::MakeBlendSpaceTickRecord(FAnimTickRecord& TickRecord, class UBlendSpaceBase* BlendSpace, const FVector& BlendInput, TArray<FBlendSampleData>& BlendSampleDataCache, FBlendFilter& BlendFilter, bool bLooping, float PlayRate, float FinalBlendWeight, float& CurrentTime, FMarkerTickRecord& MarkerTickRecord) const
+void FAnimInstanceProxy::MakeBlendSpaceTickRecord(FAnimTickRecord& TickRecord, class UBlendSpace* BlendSpace, const FVector& BlendInput, TArray<FBlendSampleData>& BlendSampleDataCache, FBlendFilter& BlendFilter, bool bLooping, float PlayRate, float FinalBlendWeight, float& CurrentTime, FMarkerTickRecord& MarkerTickRecord) const
 {
 	TickRecord.SourceAsset = BlendSpace;
 	TickRecord.BlendSpace.BlendSpacePositionX = BlendInput.X;
@@ -684,7 +684,7 @@ PRAGMA_ENABLE_DEPRECATION_WARNINGS
 	TickRecord.SourceAsset->TickAssetPlayer(TickRecord, NotifyQueue, TickContext);
 }
 
-void FAnimInstanceProxy::BlendSpaceAdvanceImmediate(class UBlendSpaceBase* BlendSpace, const FVector& BlendInput, TArray<FBlendSampleData>& BlendSampleDataCache, FBlendFilter& BlendFilter, bool bLooping, float PlayRate, float DeltaSeconds, float& CurrentTime, FMarkerTickRecord& MarkerTickRecord)
+void FAnimInstanceProxy::BlendSpaceAdvanceImmediate(class UBlendSpace* BlendSpace, const FVector& BlendInput, TArray<FBlendSampleData>& BlendSampleDataCache, FBlendFilter& BlendFilter, bool bLooping, float PlayRate, float DeltaSeconds, float& CurrentTime, FMarkerTickRecord& MarkerTickRecord)
 {
 	FAnimTickRecord TickRecord;
 PRAGMA_DISABLE_DEPRECATION_WARNINGS
diff --git a/Engine/Source/Runtime/Engine/Private/Animation/AnimSequence.cpp b/Engine/Source/Runtime/Engine/Private/Animation/AnimSequence.cpp
--- a/Engine/Source/Runtime/Engine/Private/Animation/AnimSequence.cpp
+++ b/Engine/Source/Runtime/Engine/Private/Animation/AnimSequence.cpp
@@ -22,7 +22,7 @@
 #include "Animation/AnimCompress.h"
 #include "Components/SkeletalMeshComponent.h"
 #include "Animation/AnimNotifies/AnimNotify.h"
-#include "Animation/BlendSpaceBase.h"
+#include "Animation/BlendSpace.h"
 #include "Animation/Rig.h"
 #include "Animation/AnimationSettings.h"
 #include "Animation/AnimBoneCompressionCodec.h"
@@ -4323,7 +4323,7 @@ void UAnimSequence::RefreshSyncMarkerDataFromAuthored()
 	check(IsInGameThread());
 
 	// Update blend spaces that may be referencing us
-	for(TObjectIterator<UBlendSpaceBase> It; It; ++It)
+	for(TObjectIterator<UBlendSpace> It; It; ++It)
 	{
 		if(!It->HasAnyFlags(RF_NeedLoad | RF_NeedPostLoad))
 		{
diff --git a/Engine/Source/Runtime/Engine/Private/Animation/AnimSingleNodeInstance.cpp b/Engine/Source/Runtime/Engine/Private/Animation/AnimSingleNodeInstance.cpp
--- a/Engine/Source/Runtime/Engine/Private/Animation/AnimSingleNodeInstance.cpp
+++ b/Engine/Source/Runtime/Engine/Private/Animation/AnimSingleNodeInstance.cpp
@@ -6,7 +6,7 @@
 =============================================================================*/ 
 
 #include "Animation/AnimSingleNodeInstance.h"
-#include "Animation/BlendSpaceBase.h"
+#include "Animation/BlendSpace.h"
 #include "Animation/BlendSpace.h"
 #include "Animation/AnimSingleNodeInstanceProxy.h"
 
@@ -455,7 +455,7 @@ FAnimInstanceProxy* UAnimSingleNodeInstance::CreateAnimInstanceProxy()
 
 FVector UAnimSingleNodeInstance::GetFilterLastOutput()
 {
-	if (UBlendSpaceBase* Blendspace = Cast<UBlendSpaceBase>(CurrentAsset))
+	if (UBlendSpace* Blendspace = Cast<UBlendSpace>(CurrentAsset))
 	{
 		FAnimSingleNodeInstanceProxy& Proxy = GetProxyOnGameThread<FAnimSingleNodeInstanceProxy>();
 		return Proxy.GetFilterLastOutput();
diff --git a/Engine/Source/Runtime/Engine/Private/Animation/AnimSingleNodeInstanceProxy.cpp b/Engine/Source/Runtime/Engine/Private/Animation/AnimSingleNodeInstanceProxy.cpp
--- a/Engine/Source/Runtime/Engine/Private/Animation/AnimSingleNodeInstanceProxy.cpp
+++ b/Engine/Source/Runtime/Engine/Private/Animation/AnimSingleNodeInstanceProxy.cpp
@@ -3,7 +3,7 @@
 #include "Animation/AnimSingleNodeInstanceProxy.h"
 #include "AnimationRuntime.h"
 #include "Animation/AnimComposite.h"
-#include "Animation/BlendSpaceBase.h"
+#include "Animation/BlendSpace.h"
 #include "Animation/PoseAsset.h"
 #include "Animation/AnimStreamable.h"
 #include "Animation/AnimSingleNodeInstance.h"
@@ -156,7 +156,7 @@ void FAnimSingleNodeInstanceProxy::SetMontagePreviewSlot(FName PreviewSlot)
 	SingleNode.ActiveMontageSlot = PreviewSlot;
 }
 
-void FAnimSingleNodeInstanceProxy::InternalBlendSpaceEvaluatePose(class UBlendSpaceBase* BlendSpace, TArray<FBlendSampleData>& BlendSampleDataCache, FPoseContext& OutContext)
+void FAnimSingleNodeInstanceProxy::InternalBlendSpaceEvaluatePose(class UBlendSpace* BlendSpace, TArray<FBlendSampleData>& BlendSampleDataCache, FPoseContext& OutContext)
 {
 	FAnimationPoseData AnimationPoseData = { OutContext.Pose, OutContext.Curve, OutContext.CustomAttributes };
 
@@ -209,7 +209,7 @@ void FAnimSingleNodeInstanceProxy::SetAnimationAsset(class UAnimationAsset* NewA
 #endif
 
 	
-	if (UBlendSpaceBase* BlendSpace = Cast<UBlendSpaceBase>(NewAsset))
+	if (UBlendSpace* BlendSpace = Cast<UBlendSpace>(NewAsset))
 	{
 		BlendSpace->InitializeFilter(&BlendFilter);
 	}
@@ -217,7 +217,7 @@ void FAnimSingleNodeInstanceProxy::SetAnimationAsset(class UAnimationAsset* NewA
 
 void FAnimSingleNodeInstanceProxy::UpdateBlendspaceSamples(FVector InBlendInput)
 {
-	if (UBlendSpaceBase* BlendSpace = Cast<UBlendSpaceBase>(CurrentAsset))
+	if (UBlendSpace* BlendSpace = Cast<UBlendSpace>(CurrentAsset))
 	{
 		float OutCurrentTime = 0.f;
 		FMarkerTickRecord TempMarkerTickRecord;
@@ -278,7 +278,7 @@ void FAnimNode_SingleNode::Evaluate_AnyThread(FPoseContext& Output)
 		FAnimationPoseData OutputAnimationPoseData(Output);
 
 		//@TODO: animrefactor: Seems like more code duplication than we need
-		if (UBlendSpaceBase* BlendSpace = Cast<UBlendSpaceBase>(Proxy->CurrentAsset))
+		if (UBlendSpace* BlendSpace = Cast<UBlendSpace>(Proxy->CurrentAsset))
 		{
 			Proxy->InternalBlendSpaceEvaluatePose(BlendSpace, Proxy->BlendSampleData, Output);
 		}
@@ -532,7 +532,7 @@ void FAnimNode_SingleNode::Update_AnyThread(const FAnimationUpdateContext& Conte
 	{
 		UE::Anim::FAnimSyncGroupScope& SyncScope = Context.GetMessageChecked<UE::Anim::FAnimSyncGroupScope>();
 
-		if (UBlendSpaceBase* BlendSpace = Cast<UBlendSpaceBase>(Proxy->CurrentAsset))
+		if (UBlendSpace* BlendSpace = Cast<UBlendSpace>(Proxy->CurrentAsset))
 		{
 			FAnimTickRecord TickRecord(BlendSpace, Proxy->BlendSpacePosition, Proxy->BlendSampleData, Proxy->BlendFilter, Proxy->bLooping, NewPlayRate, 1.f, /*inout*/ Proxy->CurrentTime, Proxy->MarkerTickRecord);
 			SyncScope.AddTickRecord(TickRecord);
diff --git a/Engine/Source/Runtime/Engine/Private/Animation/AnimTrace.cpp b/Engine/Source/Runtime/Engine/Private/Animation/AnimTrace.cpp
--- a/Engine/Source/Runtime/Engine/Private/Animation/AnimTrace.cpp
+++ b/Engine/Source/Runtime/Engine/Private/Animation/AnimTrace.cpp
@@ -14,7 +14,7 @@
 #include "Math/TransformNonVectorized.h"
 #include "Animation/AnimNodeBase.h"
 #include "Animation/AnimMontage.h"
-#include "Animation/BlendSpaceBase.h"
+#include "Animation/BlendSpace.h"
 #include "Animation/AnimNode_SequencePlayer.h"
 #include "Animation/AnimNotifies/AnimNotify.h"
 #include "Animation/AnimNotifies/AnimNotifyState.h"
@@ -405,7 +405,7 @@ void FAnimTrace::OutputAnimTickRecord(const FAnimationBaseContext& InContext, co
 		float BlendSpacePositionY = 0.0f;
 		float BlendSpaceFilteredPositionX = 0.0f;
 		float BlendSpaceFilteredPositionY = 0.0f;
-		const bool bIsBlendSpace = InTickRecord.SourceAsset->IsA<UBlendSpaceBase>();
+		const bool bIsBlendSpace = InTickRecord.SourceAsset->IsA<UBlendSpace>();
 		if(bIsBlendSpace)
 		{
 			BlendSpacePositionX = InTickRecord.BlendSpace.BlendSpacePositionX;
diff --git a/Engine/Source/Runtime/Engine/Private/Animation/AnimationRuntime.cpp b/Engine/Source/Runtime/Engine/Private/Animation/AnimationRuntime.cpp
--- a/Engine/Source/Runtime/Engine/Private/Animation/AnimationRuntime.cpp
+++ b/Engine/Source/Runtime/Engine/Private/Animation/AnimationRuntime.cpp
@@ -6,7 +6,7 @@
 
 #include "AnimationRuntime.h"
 #include "Animation/AnimData/BoneMaskFilter.h"
-#include "Animation/BlendSpaceBase.h"
+#include "Animation/BlendSpace.h"
 #include "Animation/AnimInstance.h"
 #include "SkeletalRender.h"
 #include "Animation/CustomAttributes.h"
@@ -663,7 +663,7 @@ void FAnimationRuntime::BlendPosesTogetherPerBone(TArrayView<const FCompactPose>
 void FAnimationRuntime::BlendPosesTogetherPerBoneInMeshSpace(
 	const TArrayView<FCompactPose> SourcePoses,
 	const TArrayView<const FBlendedCurve> SourceCurves,
-	const UBlendSpaceBase* BlendSpace,
+	const UBlendSpace* BlendSpace,
 	const TArrayView<const FBlendSampleData> BlendSampleDataCache,
 	/*out*/ FCompactPose& ResultPose,
 	/*out*/ FBlendedCurve& ResultCurve)
@@ -674,7 +674,7 @@ void FAnimationRuntime::BlendPosesTogetherPerBoneInMeshSpace(
 	BlendPosesTogetherPerBoneInMeshSpace(SourcePoses, SourceCurves, {}, BlendSpace, BlendSampleDataCache, AnimationPoseData);
 }
 
-void FAnimationRuntime::BlendPosesTogetherPerBoneInMeshSpace(TArrayView<FCompactPose> SourcePoses, TArrayView<const FBlendedCurve> SourceCurves, TArrayView<const FStackCustomAttributes> SourceAttributes, const UBlendSpaceBase* BlendSpace, TArrayView<const FBlendSampleData> BlendSampleDataCache, /*out*/ FAnimationPoseData& OutAnimationPoseData)
+void FAnimationRuntime::BlendPosesTogetherPerBoneInMeshSpace(TArrayView<FCompactPose> SourcePoses, TArrayView<const FBlendedCurve> SourceCurves, TArrayView<const FStackCustomAttributes> SourceAttributes, const UBlendSpace* BlendSpace, TArrayView<const FBlendSampleData> BlendSampleDataCache, /*out*/ FAnimationPoseData& OutAnimationPoseData)
 {
 	FQuat NewRotation;
 	USkeleton* Skeleton = BlendSpace->GetSkeleton();
diff --git a/Engine/Source/Runtime/Engine/Private/Animation/BlendSpace.cpp b/Engine/Source/Runtime/Engine/Private/Animation/BlendSpace.cpp
--- a/Engine/Source/Runtime/Engine/Private/Animation/BlendSpace.cpp
+++ b/Engine/Source/Runtime/Engine/Private/Animation/BlendSpace.cpp
@@ -5,24 +5,1809 @@
 =============================================================================*/ 
 
 #include "Animation/BlendSpace.h"
+#include "Animation/AnimNotifyQueue.h"
+#include "AnimationUtils.h"
+#include "Animation/BlendSpaceUtilities.h"
+#include "UObject/FrameworkObjectVersion.h"
+#include "UObject/UObjectIterator.h"
+#include "Logging/TokenizedMessage.h"
+#include "Logging/MessageLog.h"
+#include "Animation/AnimationPoseData.h"
+#include "Animation/CustomAttributesRuntime.h"
+
+#define LOCTEXT_NAMESPACE "BlendSpace"
+
+DECLARE_CYCLE_STAT(TEXT("BlendSpace GetAnimPose"), STAT_BlendSpace_GetAnimPose, STATGROUP_Anim);
+
+/** Scratch buffers for multi-threaded usage */
+struct FBlendSpaceScratchData : public TThreadSingleton<FBlendSpaceScratchData>
+{
+	TArray<FBlendSampleData> OldSampleDataList;
+	TArray<FBlendSampleData> NewSampleDataList;
+	TArray<FGridBlendSample, TInlineAllocator<4> > RawGridSamples;
+};
 
 UBlendSpace::UBlendSpace(const FObjectInitializer& ObjectInitializer)
 	: Super(ObjectInitializer)
 {
+	SampleIndexWithMarkers = INDEX_NONE;
+
+	/** Use highest weighted animation as default */
+	NotifyTriggerMode = ENotifyTriggerMode::HighestWeightedAnimation;
+
+	if (DimensionIndices.Num() == 0)
+	{
+		DimensionIndices = { 0, 1 };
+	}
+}
+
+void UBlendSpace::PostLoad()
+{
+	Super::PostLoad();
+
+#if WITH_EDITOR	
+	// Only do this during editor time (could alter the blendspace data during runtime otherwise) 
+	ValidateSampleData();
+#endif // WITH_EDITOR
+
+	InitializePerBoneBlend();
+}
+
+void UBlendSpace::Serialize(FArchive& Ar)
+{
+	Ar.UsingCustomVersion(FFrameworkObjectVersion::GUID);
+	Super::Serialize(Ar);
+
+#if WITH_EDITOR
+	if (Ar.IsLoading() && (Ar.CustomVer(FFrameworkObjectVersion::GUID) < FFrameworkObjectVersion::BlendSpacePostLoadSnapToGrid))
+	{
+		// This will ensure that all grid points are in valid position and the bIsSnapped flag is set
+		SnapSamplesToClosestGridPoint();
+	}
+
+	if (Ar.IsLoading() && (Ar.CustomVer(FFrameworkObjectVersion::GUID) < FFrameworkObjectVersion::SupportBlendSpaceRateScale))
+	{
+		for (FBlendSample& Sample : SampleData)
+		{
+			Sample.RateScale = 1.0f;
+		}
+	}
+#endif // WITH_EDITOR
+}
+
+#if WITH_EDITOR
+void UBlendSpace::PreEditChange(FProperty* PropertyAboutToChange)
+{
+	Super::PreEditChange(PropertyAboutToChange);
+
+	// Cache the axis ranges if it is going to change, this so the samples can be remapped correctly
+	const FName PropertyName = PropertyAboutToChange ? PropertyAboutToChange->GetFName() : NAME_None;
+	if ((PropertyName == GET_MEMBER_NAME_CHECKED(FBlendParameter, Min) || PropertyName == GET_MEMBER_NAME_CHECKED(FBlendParameter, Max)))
+	{
+		for (int32 AxisIndex = 0; AxisIndex < 3; ++AxisIndex)
+		{
+			PreviousAxisMinMaxValues[AxisIndex].X = BlendParameters[AxisIndex].Min;
+			PreviousAxisMinMaxValues[AxisIndex].Y = BlendParameters[AxisIndex].Max;
+		}
+	}
+}
+
+void UBlendSpace::PostEditChangeProperty(struct FPropertyChangedEvent& PropertyChangedEvent)
+{
+	const FName MemberPropertyName = PropertyChangedEvent.MemberProperty ? PropertyChangedEvent.MemberProperty->GetFName() : NAME_None;
+	const FName PropertyName = PropertyChangedEvent.Property ? PropertyChangedEvent.Property->GetFName() : NAME_None;
+
+	if ((MemberPropertyName == GET_MEMBER_NAME_CHECKED(UBlendSpace, PerBoneBlend) && PropertyName == GET_MEMBER_NAME_CHECKED(FBoneReference, BoneName)) || PropertyName == GET_MEMBER_NAME_CHECKED(UBlendSpace, PerBoneBlend))
+	{
+		InitializePerBoneBlend();
+	}
+
+	if (MemberPropertyName == GET_MEMBER_NAME_CHECKED(UBlendSpace, BlendParameters))
+	{
+		if (PropertyName == GET_MEMBER_NAME_CHECKED(FBlendParameter, GridNum))
+		{
+			// Tried and snap samples to points on the grid, those who do not fit or cannot be snapped are marked as invalid
+			SnapSamplesToClosestGridPoint();
+		}
+		else if (PropertyName == GET_MEMBER_NAME_CHECKED(FBlendParameter, Min) || PropertyName == GET_MEMBER_NAME_CHECKED(FBlendParameter, Max))
+		{
+			// Remap the samples to the new values by normalizing the axis and applying the new value range		
+			RemapSamplesToNewAxisRange();
+		}
+	}
+
+	Super::PostEditChangeProperty(PropertyChangedEvent);
+}
+#endif // WITH_EDITOR
+
+
+bool UBlendSpace::UpdateBlendSamples_Internal(const FVector& InBlendSpacePosition, float InDeltaTime, TArray<FBlendSampleData>& InOutOldSampleDataList, TArray<FBlendSampleData>& InOutSampleDataCache) const
+{
+	// For Target weight interpolation, we'll need to save old data, and interpolate to new data
+	TArray<FBlendSampleData>& NewSampleDataList = FBlendSpaceScratchData::Get().NewSampleDataList;
+	check(!NewSampleDataList.Num()); // this must be called non-recursively
+
+	InOutOldSampleDataList.Append(InOutSampleDataCache);
+
+	// @fixme: temporary code to clear the invalid sample data 
+	// related jira: UE-71107
+	for (int32 Index = 0; Index < InOutOldSampleDataList.Num(); ++Index)
+	{
+		if (!SampleData.IsValidIndex(InOutOldSampleDataList[Index].SampleDataIndex))
+		{
+			InOutOldSampleDataList.RemoveAt(Index);
+			--Index;
+		}
+	}
+
+	// get sample data based on new input
+	// consolidate all samples and sort them, so that we can handle from biggest weight to smallest
+	InOutSampleDataCache.Reset();
+
+	// get sample data from blendspace
+	bool bSuccessfullySampled = false;
+	if (GetSamplesFromBlendInput(InBlendSpacePosition, NewSampleDataList))
+	{
+		// if target weight interpolation is set
+		if (TargetWeightInterpolationSpeedPerSec > 0.f || PerBoneBlend.Num() > 0)
+		{
+			// target weight interpolation
+			if (InterpolateWeightOfSampleData(InDeltaTime, InOutOldSampleDataList, NewSampleDataList, InOutSampleDataCache))
+			{
+				// now I need to normalize
+				FBlendSampleData::NormalizeDataWeight(InOutSampleDataCache);
+			}
+			else
+			{
+				// if interpolation failed, just copy new sample data to sample data
+				InOutSampleDataCache = NewSampleDataList;
+			}
+		}
+		else
+		{
+			// when there is no target weight interpolation, just copy new to target
+			InOutSampleDataCache.Append(NewSampleDataList);
+		}
+
+		bSuccessfullySampled = true;
+	}
+
+	NewSampleDataList.Reset();
+
+	return bSuccessfullySampled;
+}
+
+bool UBlendSpace::UpdateBlendSamples(const FVector& InBlendSpacePosition, float InDeltaTime, TArray<FBlendSampleData>& InOutSampleDataCache) const
+{
+	TArray<FBlendSampleData>& OldSampleDataList = FBlendSpaceScratchData::Get().OldSampleDataList;
+	check(!OldSampleDataList.Num()); // this must be called non-recursively
+	const bool bResult = UpdateBlendSamples_Internal(InBlendSpacePosition, InDeltaTime, OldSampleDataList, InOutSampleDataCache);
+	OldSampleDataList.Reset();
+	return bResult;
+}
+
+void UBlendSpace::TickAssetPlayer(FAnimTickRecord& Instance, struct FAnimNotifyQueue& NotifyQueue, FAnimAssetTickContext& Context) const
+{
+	check(Instance.BlendSpace.BlendSampleDataCache);
+
+	// Scratch area for old samples
+	TArray<FBlendSampleData>& OldSampleDataList = FBlendSpaceScratchData::Get().OldSampleDataList;
+	check(!OldSampleDataList.Num()); // this must be called non-recursively
+	// new sample data that will be used for evaluation
+	TArray<FBlendSampleData>& SampleDataList = *Instance.BlendSpace.BlendSampleDataCache;
+
+	const float DeltaTime = Context.GetDeltaTime();
+	float MoveDelta = Instance.PlayRateMultiplier * DeltaTime;
+
+	// this happens even if MoveDelta == 0.f. This still should happen if it is being interpolated
+	// since we allow setting position of blendspace, we can't ignore MoveDelta == 0.f
+	// also now we don't have to worry about not following if DeltaTime = 0.f
+	{
+		// first filter input using blend filter
+		const FVector BlendSpacePosition(Instance.BlendSpace.BlendSpacePositionX, Instance.BlendSpace.BlendSpacePositionY, 0.f);
+		const FVector FilteredBlendInput = FilterInput(Instance.BlendSpace.BlendFilter, BlendSpacePosition, DeltaTime);
+
+		if (UpdateBlendSamples_Internal(FilteredBlendInput, DeltaTime, OldSampleDataList, SampleDataList))
+		{
+			float NewAnimLength = 0.f;
+			float PreInterpAnimLength = 0.f;
+
+			if (TargetWeightInterpolationSpeedPerSec > 0.f)
+			{
+				// recalculate AnimLength based on weight of target animations - this is used for scaling animation later (change speed)
+				PreInterpAnimLength = GetAnimationLengthFromSampleData(*Instance.BlendSpace.BlendSampleDataCache);
+				UE_LOG(LogAnimation, Verbose, TEXT("BlendSpace(%s) - FilteredBlendInput(%s) : PreAnimLength(%0.5f) "), *GetName(), *FilteredBlendInput.ToString(), PreInterpAnimLength);
+			}
+
+			EBlendSpaceAxis AxisToScale = GetAxisToScale();
+			if (AxisToScale != BSA_None)
+			{
+				float FilterMultiplier = 1.f;
+				// first use multiplier using new blendinput
+				// new filtered input is going to be used for sampling animation
+				// so we'll need to change playrate if you'd like to not slide foot
+				if (!BlendSpacePosition.Equals(FilteredBlendInput))
+				{
+					// apply speed change if you want, 
+					if (AxisToScale == BSA_X)
+					{
+						if (FilteredBlendInput.X != 0.f)
+						{
+							FilterMultiplier = BlendSpacePosition.X / FilteredBlendInput.X;
+						}
+					}
+					else if (AxisToScale == BSA_Y)
+					{
+						if (FilteredBlendInput.Y != 0.f)
+						{
+							FilterMultiplier = BlendSpacePosition.Y / FilteredBlendInput.Y;
+						}
+					}
+				}
+
+				// Now find if clamped input is different. If different, then apply scale to fit in. This allows
+				// "extrapolation" of the blend space outside of the range by time scaling the animation, which is
+				// appropriate when the specified axis is speed (for example).
+				FVector ClampedInput = GetClampedBlendInput(FilteredBlendInput);
+				if (!ClampedInput.Equals(FilteredBlendInput))
+				{
+					// apply speed change if you want, 
+					if (AxisToScale == BSA_X && !BlendParameters[0].bWrapInput)
+					{
+						if (ClampedInput.X != 0.f)
+						{
+							FilterMultiplier *= FilteredBlendInput.X / ClampedInput.X;
+						}
+					}
+					else if (AxisToScale == BSA_Y)
+					{
+						if (ClampedInput.Y != 0.f && !BlendParameters[1].bWrapInput)
+						{
+							FilterMultiplier *= FilteredBlendInput.Y / ClampedInput.Y;
+						}
+					}
+				}
+
+				MoveDelta *= FilterMultiplier;
+				UE_LOG(LogAnimation, Log, TEXT("BlendSpace(%s) - FilteredBlendInput(%s) : FilteredBlendInput(%s), FilterMultiplier(%0.2f)"), *GetName(), *BlendSpacePosition.ToString(), *FilteredBlendInput.ToString(), FilterMultiplier);
+			}
+
+			bool bCanDoMarkerSync = (SampleIndexWithMarkers != INDEX_NONE) && (Context.IsSingleAnimationContext() || (Instance.bCanUseMarkerSync && Context.CanUseMarkerPosition()));
+
+			if (bCanDoMarkerSync)
+			{
+				//Copy previous frame marker data to current frame
+				for (const FBlendSampleData& PrevBlendSampleItem : OldSampleDataList)
+				{
+					for (FBlendSampleData& CurrentBlendSampleItem : SampleDataList)
+					{
+						// it only can have one animation in the sample, make sure to copy Time
+						if (PrevBlendSampleItem.Animation && PrevBlendSampleItem.Animation == CurrentBlendSampleItem.Animation)
+						{
+							CurrentBlendSampleItem.Time = PrevBlendSampleItem.Time;
+							CurrentBlendSampleItem.PreviousTime = PrevBlendSampleItem.PreviousTime;
+							CurrentBlendSampleItem.MarkerTickRecord = PrevBlendSampleItem.MarkerTickRecord;
+						}
+					}
+				}
+			}
+
+			NewAnimLength = GetAnimationLengthFromSampleData(SampleDataList);
+
+			if (PreInterpAnimLength > 0.f && NewAnimLength > 0.f)
+			{
+				MoveDelta *= PreInterpAnimLength / NewAnimLength;
+			}
+
+			float& NormalizedCurrentTime = *(Instance.TimeAccumulator);
+			float NormalizedPreviousTime = NormalizedCurrentTime;
+
+			// @note for sync group vs non sync group
+			// in blendspace, it will still sync even if only one node in sync group
+			// so you're never non-sync group unless you have situation where some markers are relevant to one sync group but not all the time
+			// here we save NormalizedCurrentTime as Highest weighted samples' position in sync group
+			// if you're not in sync group, NormalizedCurrentTime is based on normalized length by sample weights
+			// if you move between sync to non sync within blendspace, you're going to see pop because we'll have to jump
+			// for now, our rule is to keep normalized time as highest weighted sample position within its own length
+			// also MoveDelta doesn't work if you're in sync group. It will move according to sync group position
+			// @todo consider using MoveDelta when  this is leader, but that can be scary because it's not matching with DeltaTime any more. 
+			// if you have interpolation delay, that value can be applied, but the output might be unpredictable. 
+			// 
+			// to fix this better in the future, we should use marker sync position from last tick
+			// but that still doesn't fix if you just join sync group, you're going to see pop since your animation doesn't fix
+
+			if (Context.IsLeader())
+			{
+				// advance current time - blend spaces hold normalized time as when dealing with changing anim length it would be possible to go backwards
+				UE_LOG(LogAnimation, Verbose, TEXT("BlendSpace(%s) - FilteredBlendInput(%s) : AnimLength(%0.5f) "), *GetName(), *FilteredBlendInput.ToString(), NewAnimLength);
+
+				Context.SetPreviousAnimationPositionRatio(NormalizedCurrentTime);
+
+				const int32 HighestMarkerSyncWeightIndex = bCanDoMarkerSync ? FBlendSpaceUtilities::GetHighestWeightMarkerSyncSample(SampleDataList, SampleData) : -1;
+
+				if (HighestMarkerSyncWeightIndex == -1)
+				{
+					bCanDoMarkerSync = false;
+				}
+
+				if (bCanDoMarkerSync)
+				{
+					FBlendSampleData& SampleDataItem = SampleDataList[HighestMarkerSyncWeightIndex];
+					const FBlendSample& Sample = SampleData[SampleDataItem.SampleDataIndex];
+
+					bool bResetMarkerDataOnFollowers = false;
+					if (!Instance.MarkerTickRecord->IsValid(Instance.bLooping))
+					{
+						SampleDataItem.MarkerTickRecord.Reset();
+						bResetMarkerDataOnFollowers = true;
+						SampleDataItem.Time = NormalizedCurrentTime * Sample.Animation->GetPlayLength();
+					}
+					else if (!SampleDataItem.MarkerTickRecord.IsValid(Instance.bLooping) && Context.MarkerTickContext.GetMarkerSyncStartPosition().IsValid())
+					{
+						Sample.Animation->GetMarkerIndicesForPosition(Context.MarkerTickContext.GetMarkerSyncStartPosition(), true, SampleDataItem.MarkerTickRecord.PreviousMarker, SampleDataItem.MarkerTickRecord.NextMarker, SampleDataItem.Time);
+					}
+
+					const float NewDeltaTime = Context.GetDeltaTime() * Instance.PlayRateMultiplier * Sample.RateScale * Sample.Animation->RateScale;
+					if (!FMath::IsNearlyZero(NewDeltaTime))
+					{
+						Context.SetLeaderDelta(NewDeltaTime);
+						Sample.Animation->TickByMarkerAsLeader(SampleDataItem.MarkerTickRecord, Context.MarkerTickContext, SampleDataItem.Time, SampleDataItem.PreviousTime, NewDeltaTime, Instance.bLooping);
+						check(!Instance.bLooping || Context.MarkerTickContext.IsMarkerSyncStartValid());
+						TickFollowerSamples(SampleDataList, HighestMarkerSyncWeightIndex, Context, bResetMarkerDataOnFollowers);
+					}
+					NormalizedCurrentTime = SampleDataItem.Time / Sample.Animation->GetPlayLength();
+					*Instance.MarkerTickRecord = SampleDataItem.MarkerTickRecord;
+				}
+				else
+				{
+					// Advance time using current/new anim length
+					float CurrentTime = NormalizedCurrentTime * NewAnimLength;
+					FAnimationRuntime::AdvanceTime(Instance.bLooping, MoveDelta, /*inout*/ CurrentTime, NewAnimLength);
+					NormalizedCurrentTime = NewAnimLength ? (CurrentTime / NewAnimLength) : 0.0f;
+					UE_LOG(LogAnimMarkerSync, Log, TEXT("Leader (%s) (bCanDoMarkerSync == false)  - PreviousTime (%0.2f), CurrentTime (%0.2f), MoveDelta (%0.2f) "), *GetName(), NormalizedPreviousTime, NormalizedCurrentTime, MoveDelta);
+				}
+
+				Context.SetAnimationPositionRatio(NormalizedCurrentTime);
+			}
+			else
+			{
+				if (!Context.MarkerTickContext.IsMarkerSyncStartValid())
+				{
+					bCanDoMarkerSync = false;
+				}
+
+				if (bCanDoMarkerSync)
+				{
+					const int32 HighestWeightIndex = FBlendSpaceUtilities::GetHighestWeightSample(SampleDataList);
+					FBlendSampleData& SampleDataItem = SampleDataList[HighestWeightIndex];
+					const FBlendSample& Sample = SampleData[SampleDataItem.SampleDataIndex];
+
+					if (Context.GetDeltaTime() != 0.f)
+					{
+						if (!Instance.MarkerTickRecord->IsValid(Instance.bLooping))
+						{
+							SampleDataItem.Time = NormalizedCurrentTime * Sample.Animation->GetPlayLength();
+						}
+
+						TickFollowerSamples(SampleDataList, -1, Context, false);
+					}
+					*Instance.MarkerTickRecord = SampleDataItem.MarkerTickRecord;
+					NormalizedCurrentTime = SampleDataItem.Time / Sample.Animation->GetPlayLength();
+				}
+				else
+				{
+					NormalizedPreviousTime = Context.GetPreviousAnimationPositionRatio();
+					NormalizedCurrentTime = Context.GetAnimationPositionRatio();
+					UE_LOG(LogAnimMarkerSync, Log, TEXT("Follower (%s) (bCanDoMarkerSync == false) - PreviousTime (%0.2f), CurrentTime (%0.2f), MoveDelta (%0.2f) "), *GetName(), NormalizedPreviousTime, NormalizedCurrentTime, MoveDelta);
+				}
+			}
+
+			// generate notifies and sets time
+			{
+				TArray<FAnimNotifyEventReference> Notifies;
+
+				const float ClampedNormalizedPreviousTime = FMath::Clamp<float>(NormalizedPreviousTime, 0.f, 1.f);
+				const float ClampedNormalizedCurrentTime = FMath::Clamp<float>(NormalizedCurrentTime, 0.f, 1.f);
+				const bool bGenerateNotifies = (NormalizedCurrentTime != NormalizedPreviousTime) && NotifyTriggerMode != ENotifyTriggerMode::None;
+
+				// Get the index of the highest weight, assuming that the first is the highest until we find otherwise
+				const bool bTriggerNotifyHighestWeightedAnim = NotifyTriggerMode == ENotifyTriggerMode::HighestWeightedAnimation && SampleDataList.Num() > 0;
+				const int32 HighestWeightIndex = (bGenerateNotifies && bTriggerNotifyHighestWeightedAnim) ? FBlendSpaceUtilities::GetHighestWeightSample(SampleDataList) : -1;
+
+				for (int32 I = 0; I < SampleDataList.Num(); ++I)
+				{
+					FBlendSampleData& SampleEntry = SampleDataList[I];
+					const int32 SampleDataIndex = SampleEntry.SampleDataIndex;
+
+					// Skip SamplesPoints that has no relevant weight
+					if (SampleData.IsValidIndex(SampleDataIndex) && (SampleEntry.TotalWeight > ZERO_ANIMWEIGHT_THRESH))
+					{
+						const FBlendSample& Sample = SampleData[SampleDataIndex];
+						if (Sample.Animation)
+						{
+							float PrevSampleDataTime;
+							float& CurrentSampleDataTime = SampleEntry.Time;
+
+							const float MultipliedSampleRateScale = Sample.Animation->RateScale * Sample.RateScale;
+
+							if (!bCanDoMarkerSync || Sample.Animation->AuthoredSyncMarkers.Num() == 0) //Have already updated time if we are doing marker sync
+							{
+								const float SampleNormalizedPreviousTime = MultipliedSampleRateScale >= 0.f ? ClampedNormalizedPreviousTime : 1.f - ClampedNormalizedPreviousTime;
+								const float SampleNormalizedCurrentTime = MultipliedSampleRateScale >= 0.f ? ClampedNormalizedCurrentTime : 1.f - ClampedNormalizedCurrentTime;
+								PrevSampleDataTime = SampleNormalizedPreviousTime * Sample.Animation->GetPlayLength();
+								CurrentSampleDataTime = SampleNormalizedCurrentTime * Sample.Animation->GetPlayLength();
+							}
+							else
+							{
+								PrevSampleDataTime = SampleEntry.PreviousTime;
+							}
+
+							// Figure out delta time 
+							float DeltaTimePosition = CurrentSampleDataTime - PrevSampleDataTime;
+							const float SampleMoveDelta = MoveDelta * MultipliedSampleRateScale;
+
+							// if we went against play rate, then loop around.
+							if ((SampleMoveDelta * DeltaTimePosition) < 0.f)
+							{
+								DeltaTimePosition += FMath::Sign<float>(SampleMoveDelta) * Sample.Animation->GetPlayLength();
+							}
+
+							if (bGenerateNotifies && (!bTriggerNotifyHighestWeightedAnim || (I == HighestWeightIndex)))
+							{
+								// Harvest and record notifies
+								Sample.Animation->GetAnimNotifies(PrevSampleDataTime, DeltaTimePosition, Instance.bLooping, Notifies);
+							}
+
+							if (Context.RootMotionMode == ERootMotionMode::RootMotionFromEverything && Sample.Animation->bEnableRootMotion)
+							{
+								Context.RootMotionMovementParams.AccumulateWithBlend(Sample.Animation->ExtractRootMotion(PrevSampleDataTime, DeltaTimePosition, Instance.bLooping), SampleEntry.GetWeight());
+							}
+
+
+							UE_LOG(LogAnimation, Verbose, TEXT("%d. Blending animation(%s) with %f weight at time %0.2f"), I + 1, *Sample.Animation->GetName(), SampleEntry.GetWeight(), CurrentSampleDataTime);
+						}
+					}
+				}
+
+				if (bGenerateNotifies && Notifies.Num() > 0)
+				{
+					NotifyQueue.AddAnimNotifies(Context.ShouldGenerateNotifies(), Notifies, Instance.EffectiveBlendWeight);
+				}
+			}
+		}
+
+		OldSampleDataList.Reset();
+	}
+}
+
+bool UBlendSpace::IsValidAdditive() const
+{
+	return ContainsMatchingSamples(AAT_LocalSpaceBase) || ContainsMatchingSamples(AAT_RotationOffsetMeshSpace);
+}
+
+#if WITH_EDITOR
+bool UBlendSpace::GetAllAnimationSequencesReferred(TArray<UAnimationAsset*>& AnimationAssets, bool bRecursive /*= true*/)
+{
+	Super::GetAllAnimationSequencesReferred(AnimationAssets, bRecursive);
+
+	for (auto Iter = SampleData.CreateConstIterator(); Iter; ++Iter)
+	{
+		// saves all samples in the AnimSequences
+		UAnimSequence* Sequence = (*Iter).Animation;
+		if (Sequence)
+		{
+			Sequence->HandleAnimReferenceCollection(AnimationAssets, bRecursive);
+		}
+	}
+
+	if (PreviewBasePose)
+	{
+		PreviewBasePose->HandleAnimReferenceCollection(AnimationAssets, bRecursive);
+	}
+
+	return (AnimationAssets.Num() > 0);
+}
+
+void UBlendSpace::ReplaceReferredAnimations(const TMap<UAnimationAsset*, UAnimationAsset*>& ReplacementMap)
+{
+	Super::ReplaceReferredAnimations(ReplacementMap);
+
+	TArray<FBlendSample> NewSamples;
+	for (auto Iter = SampleData.CreateIterator(); Iter; ++Iter)
+	{
+		FBlendSample& Sample = (*Iter);
+		UAnimSequence* Anim = Sample.Animation;
+
+		if (Anim)
+		{
+			UAnimSequence* const* ReplacementAsset = (UAnimSequence* const*)ReplacementMap.Find(Anim);
+			if (ReplacementAsset)
+			{
+				Sample.Animation = *ReplacementAsset;
+				Sample.Animation->ReplaceReferredAnimations(ReplacementMap);
+				NewSamples.Add(Sample);
+			}
+		}
+	}
+
+	if (PreviewBasePose)
+	{
+		UAnimSequence* const* ReplacementAsset = (UAnimSequence* const*)ReplacementMap.Find(PreviewBasePose);
+		if (ReplacementAsset)
+		{
+			PreviewBasePose = *ReplacementAsset;
+			PreviewBasePose->ReplaceReferredAnimations(ReplacementMap);
+		}
+	}
+
+	SampleData = NewSamples;
+}
+
+int32 UBlendSpace::GetMarkerUpdateCounter() const
+{
+	return MarkerDataUpdateCounter;
+}
+
+void UBlendSpace::RuntimeValidateMarkerData()
+{
+	check(IsInGameThread());
+
+	for (FBlendSample& Sample : SampleData)
+	{
+		if (Sample.Animation && Sample.CachedMarkerDataUpdateCounter != Sample.Animation->GetMarkerUpdateCounter())
+		{
+			// Revalidate data
+			ValidateSampleData();
+			return;
+		}
+	}
+}
+
+#endif // WITH_EDITOR
+
+// @todo fixme: slow approach. If the perbone gets popular, we should change this to array of weight 
+// we should think about changing this to 
+int32 UBlendSpace::GetPerBoneInterpolationIndex(int32 BoneIndex, const FBoneContainer& RequiredBones) const
+{
+	for (int32 Iter = 0; Iter < PerBoneBlend.Num(); ++Iter)
+	{
+		// we would like to make sure if 
+		if (PerBoneBlend[Iter].BoneReference.IsValidToEvaluate(RequiredBones) && RequiredBones.BoneIsChildOf(BoneIndex, RequiredBones.GetCompactPoseIndexFromSkeletonIndex(PerBoneBlend[Iter].BoneReference.BoneIndex).GetInt()))
+		{
+			return Iter;
+		}
+	}
+
+	return INDEX_NONE;
+}
+
+bool UBlendSpace::IsValidAdditiveType(EAdditiveAnimationType AdditiveType) const
+{
+	return (AdditiveType == AAT_LocalSpaceBase || AdditiveType == AAT_RotationOffsetMeshSpace || AdditiveType == AAT_None);
+}
+
+void UBlendSpace::ResetToRefPose(FCompactPose& OutPose) const
+{
+	if (IsValidAdditive())
+	{
+		OutPose.ResetToAdditiveIdentity();
+	}
+	else
+	{
+		OutPose.ResetToRefPose();
+	}
+}
+
+void UBlendSpace::GetAnimationPose(TArray<FBlendSampleData>& BlendSampleDataCache, /*out*/ FCompactPose& OutPose, /*out*/ FBlendedCurve& OutCurve) const
+{
+	FStackCustomAttributes TempAttributes;
+	FAnimationPoseData AnimationPoseData = { OutPose, OutCurve, TempAttributes };
+	GetAnimationPose(BlendSampleDataCache, AnimationPoseData);
+}
+
+void UBlendSpace::GetAnimationPose(TArray<FBlendSampleData>& BlendSampleDataCache, /*out*/ FAnimationPoseData& OutAnimationPoseData) const
+{
+	GetAnimationPose_Internal(BlendSampleDataCache, TArrayView<FPoseLink>(), nullptr, false, OutAnimationPoseData);
+}
+
+void UBlendSpace::GetAnimationPose(TArray<FBlendSampleData>& BlendSampleDataCache, TArrayView<FPoseLink> InPoseLinks, /*out*/ FPoseContext& Output) const
+{
+	FAnimationPoseData AnimationPoseData(Output);
+	GetAnimationPose_Internal(BlendSampleDataCache, InPoseLinks, Output.AnimInstanceProxy, Output.ExpectsAdditivePose(), AnimationPoseData);
+}
+
+void UBlendSpace::GetAnimationPose_Internal(TArray<FBlendSampleData>& BlendSampleDataCache, TArrayView<FPoseLink> InPoseLinks, FAnimInstanceProxy* InProxy, bool bInExpectsAdditivePose, /*out*/ FAnimationPoseData& OutAnimationPoseData) const
+{
+	SCOPE_CYCLE_COUNTER(STAT_BlendSpace_GetAnimPose);
+	FScopeCycleCounterUObject BlendSpaceScope(this);
+
+	FCompactPose& OutPose = OutAnimationPoseData.GetPose();
+	FBlendedCurve& OutCurve = OutAnimationPoseData.GetCurve();
+
+	if (BlendSampleDataCache.Num() == 0)
+	{
+		ResetToRefPose(OutPose);
+		return;
+	}
+
+	const bool bNested = InPoseLinks.Num() > 0;
+	const int32 NumPoses = BlendSampleDataCache.Num();
+
+	TArray<FCompactPose, TInlineAllocator<8>> ChildrenPoses;
+	ChildrenPoses.AddZeroed(NumPoses);
+
+	TArray<FBlendedCurve, TInlineAllocator<8>> ChildrenCurves;
+	ChildrenCurves.AddZeroed(NumPoses);
+
+	TArray<FStackCustomAttributes, TInlineAllocator<8>> ChildrenAttributes;
+	ChildrenAttributes.AddZeroed(NumPoses);
+
+	TArray<float, TInlineAllocator<8>> ChildrenWeights;
+	ChildrenWeights.AddZeroed(NumPoses);
+
+	for (int32 ChildrenIdx = 0; ChildrenIdx < ChildrenPoses.Num(); ++ChildrenIdx)
+	{
+		ChildrenPoses[ChildrenIdx].SetBoneContainer(&OutPose.GetBoneContainer());
+		ChildrenCurves[ChildrenIdx].InitFrom(OutCurve);
+	}
+
+	// get all child atoms we interested in
+	for (int32 I = 0; I < NumPoses; ++I)
+	{
+		FCompactPose& Pose = ChildrenPoses[I];
+
+		if (SampleData.IsValidIndex(BlendSampleDataCache[I].SampleDataIndex))
+		{
+			const FBlendSample& Sample = SampleData[BlendSampleDataCache[I].SampleDataIndex];
+			ChildrenWeights[I] = BlendSampleDataCache[I].GetWeight();
+
+			if (bNested)
+			{
+				check(InPoseLinks.IsValidIndex(BlendSampleDataCache[I].SampleDataIndex));
+
+				// Evaluate the linked graphs
+				FPoseContext ChildPoseContext(InProxy, bInExpectsAdditivePose);
+				InPoseLinks[BlendSampleDataCache[I].SampleDataIndex].Evaluate(ChildPoseContext);
+
+				// Move out poses etc. for blending
+				ChildrenPoses[I] = MoveTemp(ChildPoseContext.Pose);
+				ChildrenCurves[I] = MoveTemp(ChildPoseContext.Curve);
+				ChildrenAttributes[I] = MoveTemp(ChildPoseContext.CustomAttributes);
+			}
+			else
+			{
+				if (Sample.Animation
+#if WITH_EDITOR
+					// verify if Sample.Animation->GetSkeleton matches
+					&& ensure(Sample.Animation->GetSkeleton() == GetSkeleton())
+#endif // WITH_EDITOR
+					)
+				{
+					const float Time = FMath::Clamp<float>(BlendSampleDataCache[I].Time, 0.f, Sample.Animation->GetPlayLength());
+
+					FAnimationPoseData ChildAnimationPoseData = { Pose, ChildrenCurves[I], ChildrenAttributes[I] };
+					// first one always fills up the source one
+					Sample.Animation->GetAnimationPose(ChildAnimationPoseData, FAnimExtractContext(Time, true));
+				}
+				else
+				{
+					ResetToRefPose(Pose);
+				}
+			}
+		}
+		else
+		{
+			ResetToRefPose(Pose);
+		}
+	}
+
+	TArrayView<FCompactPose> ChildrenPosesView(ChildrenPoses);
+
+	if (PerBoneBlend.Num() > 0)
+	{
+		if (IsValidAdditive())
+		{
+			if (bRotationBlendInMeshSpace)
+			{
+				FAnimationRuntime::BlendPosesTogetherPerBoneInMeshSpace(ChildrenPosesView, ChildrenCurves, ChildrenAttributes, this, BlendSampleDataCache, OutAnimationPoseData);
+			}
+			else
+			{
+				FAnimationRuntime::BlendPosesTogetherPerBone(ChildrenPosesView, ChildrenCurves, ChildrenAttributes, this, BlendSampleDataCache, OutAnimationPoseData);
+			}
+		}
+		else
+		{
+			FAnimationRuntime::BlendPosesTogetherPerBone(ChildrenPosesView, ChildrenCurves, ChildrenAttributes, this, BlendSampleDataCache, OutAnimationPoseData);
+		}
+	}
+	else
+	{
+		FAnimationRuntime::BlendPosesTogether(ChildrenPosesView, ChildrenCurves, ChildrenAttributes, ChildrenWeights, OutAnimationPoseData);
+	}
+
+	// Once all the accumulation and blending has been done, normalize rotations.
+	OutPose.NormalizeRotations();
+}
+
+const FBlendParameter& UBlendSpace::GetBlendParameter(const int32 Index) const
+{
+	checkf(Index >= 0 && Index < 3, TEXT("Invalid Blend Parameter Index"));
+	return BlendParameters[Index];
+}
+
+const struct FBlendSample& UBlendSpace::GetBlendSample(const int32 SampleIndex) const
+{
+#if WITH_EDITOR
+	checkf(IsValidBlendSampleIndex(SampleIndex), TEXT("Invalid blend sample index"));
+#endif // WITH_EDITOR
+	return SampleData[SampleIndex];
+}
+
+bool UBlendSpace::GetSamplesFromBlendInput(const FVector& BlendInput, TArray<FBlendSampleData>& OutSampleDataList) const
+{
+	TArray<FGridBlendSample, TInlineAllocator<4> >& RawGridSamples = FBlendSpaceScratchData::Get().RawGridSamples;
+	check(!RawGridSamples.Num()); // this must be called non-recursively
+	GetRawSamplesFromBlendInput(BlendInput, RawGridSamples);
+
+	OutSampleDataList.Reset();
+	OutSampleDataList.Reserve(RawGridSamples.Num() * FEditorElement::MAX_VERTICES);
+
+	// Consolidate all samples
+	for (int32 SampleNum = 0; SampleNum < RawGridSamples.Num(); ++SampleNum)
+	{
+		FGridBlendSample& GridSample = RawGridSamples[SampleNum];
+		float GridWeight = GridSample.BlendWeight;
+		FEditorElement& GridElement = GridSample.GridElement;
+
+		for (int32 Ind = 0; Ind < GridElement.MAX_VERTICES; ++Ind)
+		{
+			const int32 SampleDataIndex = GridElement.Indices[Ind];
+			if (SampleData.IsValidIndex(SampleDataIndex))
+			{
+				int32 Index = OutSampleDataList.AddUnique(SampleDataIndex);
+				FBlendSampleData& NewSampleData = OutSampleDataList[Index];
+
+				NewSampleData.AddWeight(GridElement.Weights[Ind] * GridWeight);
+				NewSampleData.Animation = SampleData[SampleDataIndex].Animation;
+				NewSampleData.SamplePlayRate = SampleData[SampleDataIndex].RateScale;
+			}
+		}
+	}
+
+	// At this point we'll only have one of each sample, but different samples can point to the same
+	// animation. We can combine those, making sure to interpolate the parameters like play rate too.
+	for (int32 Index1 = 0; Index1 < OutSampleDataList.Num(); ++Index1)
+	{
+		FBlendSampleData& FirstSample = OutSampleDataList[Index1];
+		for (int32 Index2 = Index1 + 1; Index2 < OutSampleDataList.Num(); ++Index2)
+		{
+			FBlendSampleData& SecondSample = OutSampleDataList[Index2];
+			// if they have same sample, remove the Index2, and get out
+			if (FirstSample.SampleDataIndex == SecondSample.SampleDataIndex || // Shouldn't happen
+				(FirstSample.Animation != nullptr && FirstSample.Animation == SecondSample.Animation))
+			{
+				//Calc New Sample Playrate
+				const float TotalWeight = FirstSample.GetWeight() + SecondSample.GetWeight();
+
+				// Only combine playrates if total weight > 0
+				if (!FMath::IsNearlyZero(TotalWeight))
+				{
+					const float OriginalWeightedPlayRate = FirstSample.SamplePlayRate * (FirstSample.GetWeight() / TotalWeight);
+					const float SecondSampleWeightedPlayRate = SecondSample.SamplePlayRate * (SecondSample.GetWeight() / TotalWeight);
+					FirstSample.SamplePlayRate = OriginalWeightedPlayRate + SecondSampleWeightedPlayRate;
+
+					// add weight
+					FirstSample.AddWeight(SecondSample.GetWeight());
+				}
+
+				// as for time or previous time will be the master one(Index1)
+				OutSampleDataList.RemoveAtSwap(Index2, 1, false);
+				--Index2;
+			}
+		}
+	}
+
+	OutSampleDataList.Sort([](const FBlendSampleData& A, const FBlendSampleData& B) { return B.TotalWeight < A.TotalWeight; });
+
+	// remove noisy ones
+	int32 TotalSample = OutSampleDataList.Num();
+	float TotalWeight = 0.f;
+	for (int32 I = 0; I < TotalSample; ++I)
+	{
+		if (OutSampleDataList[I].TotalWeight < ZERO_ANIMWEIGHT_THRESH)
+		{
+			// cut anything in front of this 
+			OutSampleDataList.RemoveAt(I, TotalSample - I, false); // we won't shrink here, that might screw up alloc optimization at a higher level, if not this is temp anyway
+			break;
+		}
+
+		TotalWeight += OutSampleDataList[I].TotalWeight;
+	}
+
+	for (int32 I = 0; I < OutSampleDataList.Num(); ++I)
+	{
+		// normalize to all weights
+		OutSampleDataList[I].TotalWeight /= TotalWeight;
+	}
+	RawGridSamples.Reset();
+	return (OutSampleDataList.Num() != 0);
+}
+
+void UBlendSpace::InitializeFilter(FBlendFilter* Filter) const
+{
+	if (Filter)
+	{
+		Filter->FilterPerAxis[0].Initialize(InterpolationParam[0].InterpolationTime, InterpolationParam[0].InterpolationType);
+		Filter->FilterPerAxis[1].Initialize(InterpolationParam[1].InterpolationTime, InterpolationParam[1].InterpolationType);
+		Filter->FilterPerAxis[2].Initialize(InterpolationParam[2].InterpolationTime, InterpolationParam[2].InterpolationType);
+	}
+}
+
+#if WITH_EDITOR
+void UBlendSpace::ValidateSampleData()
+{
+	// (done here since it won't be triggered in the BlendSpace::PostEditChangeProperty, due to empty property during Undo)
+	SnapSamplesToClosestGridPoint();
+
+	bool bSampleDataChanged = false;
+	AnimLength = 0.f;
+
+	bool bAllMarkerPatternsMatch = true;
+	FSyncPattern BlendSpacePattern;
+
+	int32 SampleWithMarkers = INDEX_NONE;
+
+	for (int32 SampleIndex = 0; SampleIndex < SampleData.Num(); ++SampleIndex)
+	{
+		FBlendSample& Sample = SampleData[SampleIndex];
+
+		// see if same data exists, by same, same values
+		for (int32 ComparisonSampleIndex = SampleIndex + 1; ComparisonSampleIndex < SampleData.Num(); ++ComparisonSampleIndex)
+		{
+			if (IsSameSamplePoint(Sample.SampleValue, SampleData[ComparisonSampleIndex].SampleValue))
+			{
+				SampleData.RemoveAt(ComparisonSampleIndex);
+				--ComparisonSampleIndex;
+
+				bSampleDataChanged = true;
+			}
+		}
+
+		if (IsAsset())
+		{
+			Sample.bIsValid = ValidateSampleValue(Sample.SampleValue, SampleIndex) && (Sample.Animation != nullptr);
+
+			if (Sample.bIsValid)
+			{
+				if (Sample.Animation->GetPlayLength() > AnimLength)
+				{
+					// @todo : should apply scale? If so, we'll need to apply also when blend
+					AnimLength = Sample.Animation->GetPlayLength();
+				}
+
+				Sample.CachedMarkerDataUpdateCounter = Sample.Animation->GetMarkerUpdateCounter();
+
+				if (Sample.Animation->AuthoredSyncMarkers.Num() > 0)
+				{
+					auto PopulateMarkerNameArray = [](TArray<FName>& Pattern, TArray<struct FAnimSyncMarker>& AuthoredSyncMarkers)
+					{
+						Pattern.Reserve(AuthoredSyncMarkers.Num());
+						for (FAnimSyncMarker& Marker : AuthoredSyncMarkers)
+						{
+							Pattern.Add(Marker.MarkerName);
+						}
+					};
+
+					if (SampleWithMarkers == INDEX_NONE)
+					{
+						SampleWithMarkers = SampleIndex;
+					}
+
+					if (BlendSpacePattern.MarkerNames.Num() == 0)
+					{
+						PopulateMarkerNameArray(BlendSpacePattern.MarkerNames, Sample.Animation->AuthoredSyncMarkers);
+					}
+					else
+					{
+						TArray<FName> ThisPattern;
+						PopulateMarkerNameArray(ThisPattern, Sample.Animation->AuthoredSyncMarkers);
+						if (!BlendSpacePattern.DoesPatternMatch(ThisPattern))
+						{
+							bAllMarkerPatternsMatch = false;
+						}
+					}
+				}
+			}
+			else
+			{
+				if (IsRunningGame())
+				{
+					UE_LOG(LogAnimation, Error, TEXT("[%s : %d] - Missing Sample Animation"), *GetFullName(), SampleIndex + 1);
+				}
+				else
+				{
+					static FName NAME_LoadErrors("LoadErrors");
+					FMessageLog LoadErrors(NAME_LoadErrors);
+
+					TSharedRef<FTokenizedMessage> Message = LoadErrors.Error();
+					Message->AddToken(FTextToken::Create(LOCTEXT("EmptyAnimationData1", "The BlendSpace ")));
+					Message->AddToken(FAssetNameToken::Create(GetPathName(), FText::FromString(GetName())));
+					Message->AddToken(FTextToken::Create(LOCTEXT("EmptyAnimationData2", " has sample with no animation. Recommend to remove sample point or set new animation.")));
+					LoadErrors.Notify();
+				}
+			}
+		}
+		else
+		{
+			Sample.bIsValid = ValidateSampleValue(Sample.SampleValue, SampleIndex);
+		}
+	}
+
+	// set rotation blend in mesh space
+	bRotationBlendInMeshSpace = ContainsMatchingSamples(AAT_RotationOffsetMeshSpace);
+
+	SampleIndexWithMarkers = bAllMarkerPatternsMatch ? SampleWithMarkers : INDEX_NONE;
+
+	if (bSampleDataChanged)
+	{
+		GridSamples.Empty();
+		MarkPackageDirty();
+	}
+}
+
+bool UBlendSpace::AddSample(const FVector& SampleValue)
+{
+	// We should only be adding samples without a source animation if we are not a standalone asset
+	check(!IsAsset());
+
+	const bool bValidSampleData = ValidateSampleValue(SampleValue);
+
+	if (bValidSampleData)
+	{
+		SampleData.Add(FBlendSample(nullptr, SampleValue, true, bValidSampleData));
+		UpdatePreviewBasePose();
+	}
+
+	return bValidSampleData;
+}
+
+bool UBlendSpace::AddSample(UAnimSequence* AnimationSequence, const FVector& SampleValue)
+{
+	const bool bValidSampleData = ValidateSampleValue(SampleValue) && ValidateAnimationSequence(AnimationSequence);
+
+	if (bValidSampleData)
+	{
+		SampleData.Add(FBlendSample(AnimationSequence, SampleValue, true, bValidSampleData));
+		UpdatePreviewBasePose();
+	}
+
+	return bValidSampleData;
+}
+
+bool UBlendSpace::EditSampleValue(const int32 BlendSampleIndex, const FVector& NewValue, bool bSnap)
+{
+	const bool bValidValue = SampleData.IsValidIndex(BlendSampleIndex) && ValidateSampleValue(NewValue, BlendSampleIndex);
+
+	if (bValidValue)
+	{
+		// Set new value if it passes the tests
+		SampleData[BlendSampleIndex].SampleValue = NewValue;
+		SampleData[BlendSampleIndex].bIsValid = bValidValue;
+		SampleData[BlendSampleIndex].bSnapToGrid = bSnap;
+	}
+
+	return bValidValue;
+}
+
+bool UBlendSpace::UpdateSampleAnimation(UAnimSequence* AnimationSequence, const FVector& SampleValue)
+{
+	int32 UpdateSampleIndex = INDEX_NONE;
+	for (int32 SampleIndex = 0; SampleIndex < SampleData.Num(); ++SampleIndex)
+	{
+		if (IsSameSamplePoint(SampleValue, SampleData[SampleIndex].SampleValue))
+		{
+			UpdateSampleIndex = SampleIndex;
+			break;
+		}
+	}
+
+	if (UpdateSampleIndex != INDEX_NONE)
+	{
+		SampleData[UpdateSampleIndex].Animation = AnimationSequence;
+	}
+
+	return UpdateSampleIndex != INDEX_NONE;
+}
+
+bool UBlendSpace::ReplaceSampleAnimation(const int32 BlendSampleIndex, UAnimSequence* AnimationSequence)
+{
+	const bool bValidValue = SampleData.IsValidIndex(BlendSampleIndex);
+	if (bValidValue)
+	{
+		SampleData[BlendSampleIndex].Animation = AnimationSequence;
+	}
+
+	return bValidValue;
+}
+
+bool UBlendSpace::DeleteSample(const int32 BlendSampleIndex)
+{
+	const bool bValidRemoval = SampleData.IsValidIndex(BlendSampleIndex);
+
+	if (bValidRemoval)
+	{
+		SampleData.RemoveAtSwap(BlendSampleIndex);
+		UpdatePreviewBasePose();
+	}
+
+	return bValidRemoval;
+}
+
+bool UBlendSpace::IsValidBlendSampleIndex(const int32 SampleIndex) const
+{
+	return SampleData.IsValidIndex(SampleIndex);
+}
+
+const TArray<FEditorElement>& UBlendSpace::GetGridSamples() const
+{
+	return GridSamples;
+}
+
+void UBlendSpace::FillupGridElements(const TArray<int32>& PointListToSampleIndices, const TArray<FEditorElement>& GridElements, const TArray<int32>& InDimensionIndices)
+{
+	DimensionIndices = InDimensionIndices;
+
+	GridSamples.Empty(GridElements.Num());
+	GridSamples.AddUninitialized(GridElements.Num());
+
+	for (int32 ElementIndex = 0; ElementIndex < GridElements.Num(); ++ElementIndex)
+	{
+		const FEditorElement& ViewGrid = GridElements[ElementIndex];
+		FEditorElement NewGrid;
+		float TotalWeight = 0.f;
+		for (int32 VertexIndex = 0; VertexIndex < FEditorElement::MAX_VERTICES; ++VertexIndex)
+		{
+			const int32 SampleIndex = ViewGrid.Indices[VertexIndex];
+			if (SampleIndex != INDEX_NONE && PointListToSampleIndices.IsValidIndex(SampleIndex))
+			{
+				NewGrid.Indices[VertexIndex] = PointListToSampleIndices[SampleIndex];
+			}
+			else
+			{
+				NewGrid.Indices[VertexIndex] = INDEX_NONE;
+			}
+
+			if (NewGrid.Indices[VertexIndex] == INDEX_NONE)
+			{
+				NewGrid.Weights[VertexIndex] = 0.f;
+			}
+			else
+			{
+				NewGrid.Weights[VertexIndex] = ViewGrid.Weights[VertexIndex];
+				TotalWeight += ViewGrid.Weights[VertexIndex];
+			}
+		}
+
+		// Need to normalize the weights
+		if (TotalWeight > 0.f)
+		{
+			for (int32 J = 0; J < FEditorElement::MAX_VERTICES; ++J)
+			{
+				NewGrid.Weights[J] /= TotalWeight;
+			}
+		}
+
+		GridSamples[ElementIndex] = NewGrid;
+	}
+}
+
+void UBlendSpace::EmptyGridElements()
+{
+	GridSamples.Empty();
+}
+
+bool UBlendSpace::ValidateAnimationSequence(const UAnimSequence* AnimationSequence) const
+{
+	const bool bValidAnimationSequence = IsAnimationCompatible(AnimationSequence)
+		&& IsAnimationCompatibleWithSkeleton(AnimationSequence)
+		&& (GetNumberOfBlendSamples() == 0 || DoesAnimationMatchExistingSamples(AnimationSequence));
+
+	return bValidAnimationSequence;
+}
+
+bool UBlendSpace::DoesAnimationMatchExistingSamples(const UAnimSequence* AnimationSequence) const
+{
+	const bool bMatchesExistingAnimations = ContainsMatchingSamples(AnimationSequence->AdditiveAnimType);
+	return bMatchesExistingAnimations;
+}
+
+bool UBlendSpace::ShouldAnimationBeAdditive() const
+{
+	const bool bShouldBeAdditive = !ContainsNonAdditiveSamples();
+	return bShouldBeAdditive;
+}
+
+bool UBlendSpace::IsAnimationCompatibleWithSkeleton(const UAnimSequence* AnimationSequence) const
+{
+	// Check if the animation sequences skeleton is compatible iwth the blendspace one
+	const USkeleton* MySkeleton = GetSkeleton();
+	const bool bIsAnimationCompatible = AnimationSequence && MySkeleton && AnimationSequence->GetSkeleton() && AnimationSequence->GetSkeleton()->IsCompatible(MySkeleton);
+	return bIsAnimationCompatible;
+}
+
+bool UBlendSpace::IsAnimationCompatible(const UAnimSequence* AnimationSequence) const
+{
+	// If the supplied animation is of a different additive animation type or this blendspace support non-additive animations
+	const bool bIsCompatible = IsValidAdditiveType(AnimationSequence->AdditiveAnimType);
+	return bIsCompatible;
+}
+
+bool UBlendSpace::ValidateSampleValue(const FVector& SampleValue, int32 OriginalIndex) const
+{
+	bool bValid = true;
+	bValid &= IsSampleWithinBounds(SampleValue);
+	bValid &= !IsTooCloseToExistingSamplePoint(SampleValue, OriginalIndex);
+	return bValid;
+}
+
+bool UBlendSpace::IsSampleWithinBounds(const FVector& SampleValue) const
+{
+	return !((SampleValue.X < BlendParameters[0].Min) ||
+		(SampleValue.X > BlendParameters[0].Max) ||
+		(SampleValue.Y < BlendParameters[1].Min) ||
+		(SampleValue.Y > BlendParameters[1].Max));
+}
+
+bool UBlendSpace::IsTooCloseToExistingSamplePoint(const FVector& SampleValue, int32 OriginalIndex) const
+{
+	bool bMatchesSamplePoint = false;
+	for (int32 SampleIndex = 0; SampleIndex < SampleData.Num(); ++SampleIndex)
+	{
+		if (SampleIndex != OriginalIndex)
+		{
+			if (IsSameSamplePoint(SampleValue, SampleData[SampleIndex].SampleValue))
+			{
+				bMatchesSamplePoint = true;
+				break;
+			}
+		}
+	}
+
+	return bMatchesSamplePoint;
+}
+
+#endif // WITH_EDITOR
+
+void UBlendSpace::InitializePerBoneBlend()
+{
+	const USkeleton* MySkeleton = GetSkeleton();
+	for (FPerBoneInterpolation& BoneInterpolationData : PerBoneBlend)
+	{
+		BoneInterpolationData.Initialize(MySkeleton);
+	}
+	// Sort this by bigger to smaller, then we don't have to worry about checking the best parent
+	PerBoneBlend.Sort([](const FPerBoneInterpolation& A, const FPerBoneInterpolation& B) { return A.BoneReference.BoneIndex > B.BoneReference.BoneIndex; });
+}
+
+void UBlendSpace::TickFollowerSamples(TArray<FBlendSampleData>& SampleDataList, const int32 HighestWeightIndex, FAnimAssetTickContext& Context, bool bResetMarkerDataOnFollowers) const
+{
+	for (int32 SampleIndex = 0; SampleIndex < SampleDataList.Num(); ++SampleIndex)
+	{
+		FBlendSampleData& SampleDataItem = SampleDataList[SampleIndex];
+		const FBlendSample& Sample = SampleData[SampleDataItem.SampleDataIndex];
+		if (HighestWeightIndex != SampleIndex)
+		{
+			if (bResetMarkerDataOnFollowers)
+			{
+				SampleDataItem.MarkerTickRecord.Reset();
+			}
+
+			if (Sample.Animation->AuthoredSyncMarkers.Num() > 0) // Update followers who can do marker sync, others will be handled later in TickAssetPlayer
+			{
+				Sample.Animation->TickByMarkerAsFollower(SampleDataItem.MarkerTickRecord, Context.MarkerTickContext, SampleDataItem.Time, SampleDataItem.PreviousTime, Context.GetLeaderDelta(), true);
+			}
+		}
+	}
+}
+
+float UBlendSpace::GetAnimationLengthFromSampleData(const TArray<FBlendSampleData>& SampleDataList) const
+{
+	float BlendAnimLength = 0.f;
+	for (int32 I = 0; I < SampleDataList.Num(); ++I)
+	{
+		const int32 SampleDataIndex = SampleDataList[I].SampleDataIndex;
+		if (SampleData.IsValidIndex(SampleDataIndex))
+		{
+			const FBlendSample& Sample = SampleData[SampleDataIndex];
+			if (Sample.Animation)
+			{
+				//Use the SamplePlayRate from the SampleDataList, not the RateScale from SampleData as SamplePlayRate might contain
+				//Multiple samples contribution which we would otherwise lose
+				const float MultipliedSampleRateScale = Sample.Animation->RateScale * SampleDataList[I].SamplePlayRate;
+				// apply rate scale to get actual playback time
+				BlendAnimLength += (Sample.Animation->GetPlayLength() / ((MultipliedSampleRateScale) != 0.0f ? FMath::Abs(MultipliedSampleRateScale) : 1.0f)) * SampleDataList[I].GetWeight();
+				UE_LOG(LogAnimation, Verbose, TEXT("[%d] - Sample Animation(%s) : Weight(%0.5f) "), I + 1, *Sample.Animation->GetName(), SampleDataList[I].GetWeight());
+			}
+		}
+	}
+
+	return BlendAnimLength;
+}
+
+FVector UBlendSpace::GetClampedBlendInput(const FVector& BlendInput)  const
+{
+	FVector AdjustedInput = BlendInput;
+	for (int iAxis = 0; iAxis != 3; ++iAxis)
+	{
+		if (!BlendParameters[iAxis].bWrapInput)
+		{
+			AdjustedInput[iAxis] = FMath::Clamp(AdjustedInput[iAxis], BlendParameters[iAxis].Min, BlendParameters[iAxis].Max);
+		}
+	}
+	return AdjustedInput;
+}
+
+FVector UBlendSpace::GetClampedAndWrappedBlendInput(const FVector& BlendInput) const
+{
+	FVector AdjustedInput = BlendInput;
+	for (int iAxis = 0; iAxis != 3; ++iAxis)
+	{
+		if (BlendParameters[iAxis].bWrapInput)
+		{
+			AdjustedInput[iAxis] = FMath::Wrap(AdjustedInput[iAxis], BlendParameters[iAxis].Min, BlendParameters[iAxis].Max);
+		}
+		else
+		{
+			AdjustedInput[iAxis] = FMath::Clamp(AdjustedInput[iAxis], BlendParameters[iAxis].Min, BlendParameters[iAxis].Max);
+		}
+	}
+	return AdjustedInput;
+}
+
+FVector UBlendSpace::GetNormalizedBlendInput(const FVector& BlendInput) const
+{
+	FVector AdjustedInput = GetClampedAndWrappedBlendInput(BlendInput);
+
+	const FVector MinBlendInput = FVector(BlendParameters[0].Min, BlendParameters[1].Min, BlendParameters[2].Min);
+	const FVector GridSize = FVector(BlendParameters[0].GetGridSize(), BlendParameters[1].GetGridSize(), BlendParameters[2].GetGridSize());
+
+	FVector NormalizedBlendInput = (AdjustedInput - MinBlendInput) / GridSize;
+	return NormalizedBlendInput;
+}
+
+const FEditorElement* UBlendSpace::GetGridSampleInternal(int32 Index) const
+{
+	return GridSamples.IsValidIndex(Index) ? &GridSamples[Index] : NULL;
+}
+
+// When using CriticallyDampedSmoothing, how to go from the interpolation speed to the smooth
+// time? What would the critically damped velocity be as it goes from a starting value of 0 to a
+// target of 1 (see eq in CriticallyDampedSmoothing), starting with v = 0?
+//
+// v = W^2 t exp(-W t)
+//
+// Differentiate and set equal to zero to find maximum v is at t = 1 / W
+//
+// vMax = W / e = 2 / (SmoothingTime * e)
+//
+// Set this equal to TargetWeightInterpolationSpeedPerSec, we get
+//
+// SmoothingTime = 2 / (e * TargetWeightInterpolationSpeedPerSec)
+//
+// However - this looks significantly slower than when using a constant speed, because we're
+// easing in/out, so aim for twice this speed (i.e. smooth over half the time)
+static float SmoothingTimeFromSpeed(float Speed)
+{
+	return 1.0f / (EULERS_NUMBER * Speed);
+}
+
+static void SmoothWeight(float& Output, float& OutputRate, float Input, float InputRate, float Target, float DeltaTime, float Speed, bool bUseEaseInOut)
+{
+	if (Speed <= 0.0f)
+	{
+		Output = Target;
+		return;
+	}
+
+	if (bUseEaseInOut)
+	{
+		Output = Input;
+		OutputRate = InputRate;
+		FMath::CriticallyDampedSmoothing(Output, OutputRate, Target, DeltaTime, SmoothingTimeFromSpeed(Speed));
+	}
+	else
+	{
+		Output = FMath::FInterpConstantTo(Input, Target, DeltaTime, Speed);
+	}
+}
+
+bool UBlendSpace::InterpolateWeightOfSampleData(float DeltaTime, const TArray<FBlendSampleData>& OldSampleDataList, const TArray<FBlendSampleData>& NewSampleDataList, TArray<FBlendSampleData>& FinalSampleDataList) const
+{
+	float TotalFinalWeight = 0.f;
+	float TotalFinalPerBoneWeight = 0.0f;
+
+	// now interpolate from old to new target, this is brute-force
+	for (auto OldIt = OldSampleDataList.CreateConstIterator(); OldIt; ++OldIt)
+	{
+		// Now need to modify old sample, so copy it
+		FBlendSampleData OldSample = *OldIt;
+		bool bTargetSampleExists = false;
+
+		if (OldSample.PerBoneBlendData.Num() != PerBoneBlend.Num())
+		{
+			OldSample.PerBoneBlendData.Init(OldSample.TotalWeight, PerBoneBlend.Num());
+			OldSample.PerBoneWeightRate.Init(OldSample.WeightRate, PerBoneBlend.Num());
+		}
+
+		for (auto NewIt = NewSampleDataList.CreateConstIterator(); NewIt; ++NewIt)
+		{
+			const FBlendSampleData& NewSample = *NewIt;
+			// if same sample is found, interpolate
+			if (NewSample.SampleDataIndex == OldSample.SampleDataIndex)
+			{
+				FBlendSampleData InterpData = NewSample;
+				SmoothWeight(InterpData.TotalWeight, InterpData.WeightRate, OldSample.TotalWeight, OldSample.WeightRate, NewSample.TotalWeight, DeltaTime, TargetWeightInterpolationSpeedPerSec, bTargetWeightInterpolationEaseInOut);
+				InterpData.PerBoneBlendData = OldSample.PerBoneBlendData;
+				InterpData.PerBoneWeightRate = OldSample.PerBoneWeightRate;
+
+				// now interpolate the per bone weights
+				float TotalPerBoneWeight = 0.0f;
+				for (int32 Iter = 0; Iter < InterpData.PerBoneBlendData.Num(); ++Iter)
+				{
+					SmoothWeight(
+						InterpData.PerBoneBlendData[Iter], InterpData.PerBoneWeightRate[Iter],
+						OldSample.PerBoneBlendData[Iter], OldSample.PerBoneWeightRate[Iter], NewSample.TotalWeight,
+						DeltaTime, PerBoneBlend[Iter].InterpolationSpeedPerSec, bTargetWeightInterpolationEaseInOut);
+					TotalPerBoneWeight += InterpData.PerBoneBlendData[Iter];
+				}
+
+				if (InterpData.TotalWeight > ZERO_ANIMWEIGHT_THRESH || TotalPerBoneWeight > ZERO_ANIMWEIGHT_THRESH)
+				{
+					FinalSampleDataList.Add(InterpData);
+					TotalFinalWeight += InterpData.GetWeight();
+					TotalFinalPerBoneWeight += TotalPerBoneWeight;
+					bTargetSampleExists = true;
+					break;
+				}
+			}
+		}
+
+		// if new target isn't found, interpolate to 0.f, this is gone
+		if (bTargetSampleExists == false)
+		{
+			FBlendSampleData InterpData = OldSample;
+			SmoothWeight(InterpData.TotalWeight, InterpData.WeightRate, OldSample.TotalWeight, OldSample.WeightRate, 0.0f, DeltaTime, TargetWeightInterpolationSpeedPerSec, bTargetWeightInterpolationEaseInOut);
+			// now interpolate the per bone weights
+			float TotalPerBoneWeight = 0.0f;
+			for (int32 Iter = 0; Iter < InterpData.PerBoneBlendData.Num(); ++Iter)
+			{
+				SmoothWeight(
+					InterpData.PerBoneBlendData[Iter], InterpData.PerBoneWeightRate[Iter],
+					OldSample.PerBoneBlendData[Iter], OldSample.PerBoneWeightRate[Iter], 0.0f,
+					DeltaTime, PerBoneBlend[Iter].InterpolationSpeedPerSec, bTargetWeightInterpolationEaseInOut);
+				TotalPerBoneWeight += InterpData.PerBoneBlendData[Iter];
+			}
+
+			// add it if it's not zero
+			if (InterpData.TotalWeight > ZERO_ANIMWEIGHT_THRESH || TotalPerBoneWeight > ZERO_ANIMWEIGHT_THRESH)
+			{
+				FinalSampleDataList.Add(InterpData);
+				TotalFinalWeight += InterpData.GetWeight();
+				TotalFinalPerBoneWeight += TotalPerBoneWeight;
+			}
+		}
+	}
+
+	// now find new samples that are not found from old samples
+	for (auto OldIt = NewSampleDataList.CreateConstIterator(); OldIt; ++OldIt)
+	{
+		// Now need to modify old sample, so copy it
+		FBlendSampleData OldSample = *OldIt;
+		bool bOldSampleExists = false;
+
+		if (OldSample.PerBoneBlendData.Num() != PerBoneBlend.Num())
+		{
+			OldSample.PerBoneBlendData.Init(OldSample.TotalWeight, PerBoneBlend.Num());
+			OldSample.PerBoneWeightRate.Init(OldSample.WeightRate, PerBoneBlend.Num());
+		}
+
+		for (auto NewIt = FinalSampleDataList.CreateConstIterator(); NewIt; ++NewIt)
+		{
+			const FBlendSampleData& NewSample = *NewIt;
+			if (NewSample.SampleDataIndex == OldSample.SampleDataIndex)
+			{
+				bOldSampleExists = true;
+				break;
+			}
+		}
+
+		// add those new samples
+		if (bOldSampleExists == false)
+		{
+			FBlendSampleData InterpData = OldSample;
+			float TargetWeight = InterpData.TotalWeight;
+			OldSample.TotalWeight = 0.0f;
+			OldSample.WeightRate = 0.0f;
+			SmoothWeight(InterpData.TotalWeight, InterpData.WeightRate, OldSample.TotalWeight, OldSample.WeightRate, TargetWeight, DeltaTime, TargetWeightInterpolationSpeedPerSec, bTargetWeightInterpolationEaseInOut);
+			// now interpolate the per bone weights
+			float TotalPerBoneWeight = 0.0f;
+			for (int32 Iter = 0; Iter < InterpData.PerBoneBlendData.Num(); ++Iter)
+			{
+				float Target = OldSample.PerBoneBlendData[Iter];
+				OldSample.PerBoneBlendData[Iter] = 0.0f;
+				OldSample.PerBoneWeightRate[Iter] = 0.0f;
+				SmoothWeight(
+					InterpData.PerBoneBlendData[Iter], InterpData.PerBoneWeightRate[Iter],
+					OldSample.PerBoneBlendData[Iter], OldSample.PerBoneWeightRate[Iter], Target,
+					DeltaTime, PerBoneBlend[Iter].InterpolationSpeedPerSec, bTargetWeightInterpolationEaseInOut);
+				TotalPerBoneWeight += InterpData.PerBoneBlendData[Iter];
+			}
+			if (InterpData.TotalWeight > ZERO_ANIMWEIGHT_THRESH || TotalPerBoneWeight > ZERO_ANIMWEIGHT_THRESH)
+			{
+				FinalSampleDataList.Add(InterpData);
+				TotalFinalWeight += InterpData.GetWeight();
+				TotalFinalPerBoneWeight += TotalPerBoneWeight;
+			}
+		}
+	}
+
+	return TotalFinalWeight > ZERO_ANIMWEIGHT_THRESH || TotalFinalPerBoneWeight > ZERO_ANIMWEIGHT_THRESH;
+}
+
+FVector UBlendSpace::FilterInput(FBlendFilter* Filter, const FVector& BlendInput, float DeltaTime) const
+{
+#if WITH_EDITOR
+	// Check 
+	for (int32 AxisIndex = 0; AxisIndex < 3; ++AxisIndex)
+	{
+		if (Filter->FilterPerAxis[AxisIndex].NeedsUpdate(InterpolationParam[AxisIndex].InterpolationType, InterpolationParam[AxisIndex].InterpolationTime))
+		{
+			InitializeFilter(Filter);
+			break;
+		}
+	}
+#endif
+	FVector FilteredBlendInput;
+	FilteredBlendInput.X = Filter->FilterPerAxis[0].GetFilteredData(BlendInput.X, DeltaTime);
+	FilteredBlendInput.Y = Filter->FilterPerAxis[1].GetFilteredData(BlendInput.Y, DeltaTime);
+	FilteredBlendInput.Z = Filter->FilterPerAxis[2].GetFilteredData(BlendInput.Z, DeltaTime);
+	return FilteredBlendInput;
+}
+
+bool UBlendSpace::ContainsMatchingSamples(EAdditiveAnimationType AdditiveType) const
+{
+	bool bMatching = true;
+	for (const FBlendSample& Sample : SampleData)
+	{
+		const UAnimSequence* Animation = Sample.Animation;
+		bMatching &= (SampleData.Num() > 1 && Animation == nullptr) || (Animation && ((AdditiveType == AAT_None) ? true : Animation->IsValidAdditive()) && Animation->AdditiveAnimType == AdditiveType);
+
+		if (bMatching == false)
+		{
+			break;
+		}
+	}
+
+	return bMatching && SampleData.Num() > 0;
+}
+
+bool UBlendSpace::IsSameSamplePoint(const FVector& SamplePointA, const FVector& SamplePointB) const
+{
+#if 1
+	return FMath::IsNearlyEqual(SamplePointA.X, SamplePointB.X)
+		&& FMath::IsNearlyEqual(SamplePointA.Y, SamplePointB.Y)
+		&& FMath::IsNearlyEqual(SamplePointA.Z, SamplePointB.Z);
+#else
+	if (DimensionIndices.Num() == 0 || DimensionIndices.Num() > 3)
+	{
+		UE_LOG(LogAnimation, Warning, TEXT("Unhandled dimensionality in samples: %d"), DimensionIndices.Num());
+		return false;
+	}
+	for (int32 iDim = 0; iDim != DimensionIndices.Num(); ++iDim)
+	{
+		if (!FMath::IsNearlyEqual(SamplePointA[iDim], SamplePointB[iDim]))
+			return false;
+	}
+	return true;
+#endif
+}
+
+
+#if WITH_EDITOR
+bool UBlendSpace::ContainsNonAdditiveSamples() const
+{
+	return ContainsMatchingSamples(AAT_None);
+}
+
+void UBlendSpace::UpdatePreviewBasePose()
+{
+#if WITH_EDITORONLY_DATA
+	PreviewBasePose = nullptr;
+	// Check if blendspace is additive and try to find a ref pose 
+	if (IsValidAdditive())
+	{
+		for (const FBlendSample& BlendSample : SampleData)
+		{
+			if (BlendSample.Animation && BlendSample.Animation->RefPoseSeq)
+			{
+				PreviewBasePose = BlendSample.Animation->RefPoseSeq;
+				break;
+			}
+		}
+	}
+#endif // WITH_EDITORONLY_DATA
+}
+
+void UBlendSpace::UpdateBlendSpacesUsingAnimSequence(UAnimSequenceBase* Sequence)
+{
+	for (TObjectIterator<UBlendSpace> BlendSpaceIt; BlendSpaceIt; ++BlendSpaceIt)
+	{
+		TArray<UAnimationAsset*> ReferredAssets;
+		BlendSpaceIt->GetAllAnimationSequencesReferred(ReferredAssets, false);
+
+		if (ReferredAssets.Contains(Sequence))
+		{
+			BlendSpaceIt->Modify();
+			BlendSpaceIt->ValidateSampleData();
+		}
+	}
+}
+
+#endif // WITH_EDITOR
+
+TArray<FName>* UBlendSpace::GetUniqueMarkerNames()
+{
+	return (SampleIndexWithMarkers != INDEX_NONE && SampleData.Num() > 0) ? SampleData[SampleIndexWithMarkers].Animation->GetUniqueMarkerNames() : nullptr;
+}
+
+void UBlendSpace::GetRawSamplesFromBlendInput(const FVector& BlendInput, TArray<FGridBlendSample, TInlineAllocator<4> >& OutBlendSamples) const
+{
+	switch (DimensionIndices.Num())
+	{
+	case 1:
+		GetRawSamplesFromBlendInput1D(BlendInput, OutBlendSamples);
+		break;
+	case 2:
+		GetRawSamplesFromBlendInput2D(BlendInput, OutBlendSamples);
+		break;
+	default:
+		UE_LOG(LogAnimation, Warning, TEXT("Unhandled dimensionality in samples: %d"), DimensionIndices.Num());
+		break;
+	}
+}
+
+#if WITH_EDITOR
+void UBlendSpace::SnapSamplesToClosestGridPoint()
+{
+	switch (DimensionIndices.Num())
+	{
+	case 1:
+		SnapSamplesToClosestGridPoint1D();
+		break;
+	case 2:
+		SnapSamplesToClosestGridPoint2D();
+		break;
+	default:
+		UE_LOG(LogAnimation, Warning, TEXT("Unhandled dimensionality in samples: %d"), DimensionIndices.Num());
+		break;
+	}
+}
+
+void UBlendSpace::RemapSamplesToNewAxisRange()
+{
+	switch (DimensionIndices.Num())
+	{
+	case 1:
+		RemapSamplesToNewAxisRange1D();
+		break;
+	case 2:
+		RemapSamplesToNewAxisRange2D();
+		break;
+	default:
+		UE_LOG(LogAnimation, Warning, TEXT("Unhandled dimensionality in samples: %d"), DimensionIndices.Num());
+		break;
+	}
+}
+
+#endif // WITH_EDITOR
+
+
+
+/*-----------------------------------------------------------------------------
+	1D functionality.
+-----------------------------------------------------------------------------*/
+
+void UBlendSpace::GetRawSamplesFromBlendInput1D(const FVector& BlendInput, TArray<FGridBlendSample, TInlineAllocator<4> >& OutBlendSamples) const
+{
+	check(DimensionIndices.Num() == 1);
+	int32 Index0 = DimensionIndices[0];
+
+	FVector NormalizedBlendInput = GetNormalizedBlendInput(BlendInput);
+
+	float GridIndex = FMath::TruncToFloat(NormalizedBlendInput[Index0]);
+	float Remainder = NormalizedBlendInput[Index0] - GridIndex;
+
+	const FEditorElement* BeforeElement = GetGridSampleInternal(GridIndex);
+
+	if (BeforeElement)
+	{
+		FGridBlendSample NewSample;
+		NewSample.GridElement = *BeforeElement;
+		// now calculate weight - GridElement has weights to nearest samples, here we weight the grid element
+		NewSample.BlendWeight = (1.f - Remainder);
+		OutBlendSamples.Add(NewSample);
+	}
+	else
+	{
+		FGridBlendSample NewSample;
+		NewSample.GridElement = FEditorElement();
+		NewSample.BlendWeight = 0.f;
+		OutBlendSamples.Add(NewSample);
+	}
+
+	const FEditorElement* AfterElement = GetGridSampleInternal(GridIndex + 1);
+
+	if (AfterElement)
+	{
+		FGridBlendSample NewSample;
+		NewSample.GridElement = *AfterElement;
+		// now calculate weight - GridElement has weights to nearest samples, here we weight the grid element
+		NewSample.BlendWeight = (Remainder);
+		OutBlendSamples.Add(NewSample);
+	}
+	else
+	{
+		FGridBlendSample NewSample;
+		NewSample.GridElement = FEditorElement();
+		NewSample.BlendWeight = 0.f;
+		OutBlendSamples.Add(NewSample);
+	}
+}
+
+#if WITH_EDITOR
+
+void UBlendSpace::SnapSamplesToClosestGridPoint1D()
+{
+	check(DimensionIndices.Num() == 1);
+	int32 Index0 = DimensionIndices[0];
+
+	TArray<float> GridPoints;
+	TArray<int32> ClosestSampleToGridPoint;
+	TArray<bool> SampleDataOnPoints;
+	TArray<bool> ShouldSnap;
+
+	const float GridMin = BlendParameters[Index0].Min;
+	const float GridMax = BlendParameters[Index0].Max;
+	const float GridRange = GridMax - GridMin;
+	const int32 NumGridPoints = BlendParameters[Index0].GridNum + 1;
+	const float GridStep = GridRange / BlendParameters[Index0].GridNum;
+
+	// First mark all samples as invalid and s unsnapped, recording whether they were snapped before
+	ShouldSnap.SetNumZeroed(SampleData.Num());
+	for (int32 BlendSampleIndex = 0; BlendSampleIndex < SampleData.Num(); ++BlendSampleIndex)
+	{
+		FBlendSample& BlendSample = SampleData[BlendSampleIndex];
+		ShouldSnap[BlendSampleIndex] = BlendSample.bSnapToGrid;
+		BlendSample.bSnapToGrid = false;
+	}
+
+	for (int32 GridPointIndex = 0; GridPointIndex < NumGridPoints; ++GridPointIndex)
+	{
+		const float GridPointValue = (GridPointIndex * GridStep) + GridMin;
+		GridPoints.Add(GridPointValue);
+	}
+
+	ClosestSampleToGridPoint.Init(INDEX_NONE, GridPoints.Num());
+
+	// Find closest sample to grid point
+	for (int32 PointIndex = 0; PointIndex < GridPoints.Num(); ++PointIndex)
+	{
+		const float GridPoint = GridPoints[PointIndex];
+		float SmallestDistance = FLT_MAX;
+		int32 Index = INDEX_NONE;
+
+		for (int32 SampleIndex = 0; SampleIndex < SampleData.Num(); ++SampleIndex)
+		{
+			FBlendSample& BlendSample = SampleData[SampleIndex];
+			const float Distance = FMath::Abs(GridPoint - BlendSample.SampleValue[Index0]);
+			if (Distance < SmallestDistance)
+			{
+				Index = SampleIndex;
+				SmallestDistance = Distance;
+			}
+		}
+
+		ClosestSampleToGridPoint[PointIndex] = Index;
+	}
+
+	// Find closest grid point to sample
+	for (int32 SampleIndex = 0; SampleIndex < SampleData.Num(); ++SampleIndex)
+	{
+		if (ShouldSnap[SampleIndex])
+		{
+			FBlendSample& BlendSample = SampleData[SampleIndex];
+
+			// Find closest grid point
+			float SmallestDistance = FLT_MAX;
+			int32 Index = INDEX_NONE;
+			for (int32 PointIndex = 0; PointIndex < GridPoints.Num(); ++PointIndex)
+			{
+				const float Distance = FMath::Abs(GridPoints[PointIndex] - BlendSample.SampleValue[Index0]);
+				if (Distance < SmallestDistance)
+				{
+					Index = PointIndex;
+					SmallestDistance = Distance;
+				}
+			}
+
+			// Only move the sample if it is also closest to the grid point
+			if (Index != INDEX_NONE && ClosestSampleToGridPoint[Index] == SampleIndex)
+			{
+				BlendSample.SampleValue[Index0] = GridPoints[Index];
+				BlendSample.bSnapToGrid = true;
+			}
+		}
+	}
+}
+
+void UBlendSpace::RemapSamplesToNewAxisRange1D()
+{
+	check(DimensionIndices.Num() == 1);
+	int32 Index0 = DimensionIndices[0];
+
+	const float OldGridMin = PreviousAxisMinMaxValues[Index0].X;
+	const float OldGridMax = PreviousAxisMinMaxValues[Index0].Y;
+	const float OldGridRange = OldGridMax - OldGridMin;
+
+	const float NewGridMin = BlendParameters[Index0].Min;
+	const float NewGridMax = BlendParameters[Index0].Max;
+	const float NewGridRange = NewGridMax - NewGridMin;
+
+	for (FBlendSample& BlendSample : SampleData)
+	{
+		const float NormalizedValue = (BlendSample.SampleValue[Index0] - OldGridMin) / OldGridRange;
+		BlendSample.SampleValue[Index0] = NewGridMin + (NormalizedValue * NewGridRange);
+	}
+}
+
+
+#endif // WITH_EDITOR
+
+
+
+/*-----------------------------------------------------------------------------
+	2D functionality.
+-----------------------------------------------------------------------------*/
+
+const FEditorElement* UBlendSpace::GetEditorElement(int32 XIndex, int32 YIndex) const
+{
+	int32 Index = XIndex * (BlendParameters[1].GridNum + 1) + YIndex;
+	return GetGridSampleInternal(Index);
 }
 
-void UBlendSpace::GetGridSamplesFromBlendInput(const FVector &BlendInput, FGridBlendSample & LeftBottom, FGridBlendSample & RightBottom, FGridBlendSample & LeftTop, FGridBlendSample& RightTop) const
+void UBlendSpace::GetRawSamplesFromBlendInput2D(const FVector& BlendInput, TArray<FGridBlendSample, TInlineAllocator<4> >& OutBlendSamples) const
 {
-	const FVector NormalizedBlendInput = GetNormalizedBlendInput(BlendInput);	
+	check(DimensionIndices.Num() == 2);
+
+	OutBlendSamples.Reset();
+	OutBlendSamples.AddUninitialized(4);
+
+	FGridBlendSample& LeftBottom = OutBlendSamples[0];
+	FGridBlendSample& RightBottom = OutBlendSamples[1];
+	FGridBlendSample& LeftTop = OutBlendSamples[2];
+	FGridBlendSample& RightTop = OutBlendSamples[3];
+
+	const FVector NormalizedBlendInput = GetNormalizedBlendInput(BlendInput);
 	const FVector GridIndex(FMath::TruncToFloat(NormalizedBlendInput.X), FMath::TruncToFloat(NormalizedBlendInput.Y), 0.f);
 	const FVector Remainder = NormalizedBlendInput - GridIndex;
 	// bi-linear very simple interpolation
-	const FEditorElement* EleLT = GetEditorElement(GridIndex.X, GridIndex.Y+1);
+	const FEditorElement* EleLT = GetEditorElement(GridIndex.X, GridIndex.Y + 1);
 	if (EleLT)
 	{
 		LeftTop.GridElement = *EleLT;
 		// now calculate weight - distance to each corner since input is already normalized within grid, we can just calculate distance 
-		LeftTop.BlendWeight = (1.f-Remainder.X)*Remainder.Y;
+		LeftTop.BlendWeight = (1.f - Remainder.X) * Remainder.Y;
 	}
 	else
 	{
@@ -30,11 +1815,11 @@ void UBlendSpace::GetGridSamplesFromBlendInput(const FVector &BlendInput, FGridB
 		LeftTop.BlendWeight = 0.f;
 	}
 
-	const FEditorElement* EleRT = GetEditorElement(GridIndex.X+1, GridIndex.Y+1);
+	const FEditorElement* EleRT = GetEditorElement(GridIndex.X + 1, GridIndex.Y + 1);
 	if (EleRT)
 	{
 		RightTop.GridElement = *EleRT;
-		RightTop.BlendWeight = Remainder.X*Remainder.Y;
+		RightTop.BlendWeight = Remainder.X * Remainder.Y;
 	}
 	else
 	{
@@ -46,7 +1831,7 @@ void UBlendSpace::GetGridSamplesFromBlendInput(const FVector &BlendInput, FGridB
 	if (EleLB)
 	{
 		LeftBottom.GridElement = *EleLB;
-		LeftBottom.BlendWeight = (1.f-Remainder.X)*(1.f-Remainder.Y);
+		LeftBottom.BlendWeight = (1.f - Remainder.X) * (1.f - Remainder.Y);
 	}
 	else
 	{
@@ -54,11 +1839,11 @@ void UBlendSpace::GetGridSamplesFromBlendInput(const FVector &BlendInput, FGridB
 		LeftBottom.BlendWeight = 0.f;
 	}
 
-	const FEditorElement* EleRB = GetEditorElement(GridIndex.X+1, GridIndex.Y);
+	const FEditorElement* EleRB = GetEditorElement(GridIndex.X + 1, GridIndex.Y);
 	if (EleRB)
 	{
 		RightBottom.GridElement = *EleRB;
-		RightBottom.BlendWeight = Remainder.X*(1.f-Remainder.Y);
+		RightBottom.BlendWeight = Remainder.X * (1.f - Remainder.Y);
 	}
 	else
 	{
@@ -67,33 +1852,13 @@ void UBlendSpace::GetGridSamplesFromBlendInput(const FVector &BlendInput, FGridB
 	}
 }
 
-void UBlendSpace::GetRawSamplesFromBlendInput(const FVector &BlendInput, TArray<FGridBlendSample, TInlineAllocator<4> > & OutBlendSamples) const
-{
-	OutBlendSamples.Reset();
-	OutBlendSamples.AddUninitialized(4);
-
-	GetGridSamplesFromBlendInput(BlendInput, OutBlendSamples[0], OutBlendSamples[1], OutBlendSamples[2], OutBlendSamples[3]);
-}
-
-const FEditorElement* UBlendSpace::GetEditorElement(int32 XIndex, int32 YIndex) const
-{
-	int32 Index = XIndex*(BlendParameters[1].GridNum+1) + YIndex;
-	return GetGridSampleInternal(Index);
-}
 
-bool UBlendSpace::IsValidAdditiveType(EAdditiveAnimationType AdditiveType) const
-{
-	return (AdditiveType == AAT_LocalSpaceBase || AdditiveType == AAT_RotationOffsetMeshSpace || AdditiveType == AAT_None);
-}
+#if WITH_EDITOR
 
-bool UBlendSpace::IsValidAdditive() const
+void UBlendSpace::SnapSamplesToClosestGridPoint2D()
 {
-	return ContainsMatchingSamples(AAT_LocalSpaceBase) || ContainsMatchingSamples(AAT_RotationOffsetMeshSpace);
-}
+	check(DimensionIndices.Num() == 2);
 
-#if WITH_EDITOR
-void UBlendSpace::SnapSamplesToClosestGridPoint()
-{
 	TArray<FVector> GridPoints;
 	TArray<int32> ClosestSampleToGridPoint;
 	TArray<bool> SampleDataOnPoints;
@@ -176,8 +1941,11 @@ void UBlendSpace::SnapSamplesToClosestGridPoint()
 	}
 }
 
-void UBlendSpace::RemapSamplesToNewAxisRange()
+
+void UBlendSpace::RemapSamplesToNewAxisRange2D()
 {
+	check(DimensionIndices.Num() == 2);
+
 	const FVector OldGridMin(PreviousAxisMinMaxValues[0].X, PreviousAxisMinMaxValues[1].X, 0.0f);
 	const FVector OldGridMax(PreviousAxisMinMaxValues[0].Y, PreviousAxisMinMaxValues[1].Y, 1.0f);
 	const FVector OldGridRange = OldGridMax - OldGridMin;
@@ -194,12 +1962,5 @@ void UBlendSpace::RemapSamplesToNewAxisRange()
 }
 #endif // WITH_EDITOR
 
-EBlendSpaceAxis UBlendSpace::GetAxisToScale() const
-{
-	return AxisToScaleAnimation;
-}
+#undef LOCTEXT_NAMESPACE 
 
-bool UBlendSpace::IsSameSamplePoint(const FVector& SamplePointA, const FVector& SamplePointB) const
-{
-	return FMath::IsNearlyEqual(SamplePointA.X, SamplePointB.X) && FMath::IsNearlyEqual(SamplePointA.Y, SamplePointB.Y);
-}
diff --git a/Engine/Source/Runtime/Engine/Private/Animation/BlendSpace1D.cpp b/Engine/Source/Runtime/Engine/Private/Animation/BlendSpace1D.cpp
--- a/Engine/Source/Runtime/Engine/Private/Animation/BlendSpace1D.cpp
+++ b/Engine/Source/Runtime/Engine/Private/Animation/BlendSpace1D.cpp
@@ -9,16 +9,7 @@
 UBlendSpace1D::UBlendSpace1D(const FObjectInitializer& ObjectInitializer)
 	: Super(ObjectInitializer)
 {
-}
-
-bool UBlendSpace1D::IsValidAdditive() const
-{
-	return ContainsMatchingSamples(AAT_LocalSpaceBase) || ContainsMatchingSamples(AAT_RotationOffsetMeshSpace);
-}
-
-bool UBlendSpace1D::IsValidAdditiveType(EAdditiveAnimationType AdditiveType) const
-{
-	return (AdditiveType == AAT_LocalSpaceBase || AdditiveType == AAT_RotationOffsetMeshSpace || AdditiveType == AAT_None);
+	DimensionIndices = { 0 };
 }
 
 EBlendSpaceAxis UBlendSpace1D::GetAxisToScale() const
@@ -26,152 +17,3 @@ EBlendSpaceAxis UBlendSpace1D::GetAxisToScale() const
 	return bScaleAnimation ? BSA_X : BSA_None;
 }
 
-bool UBlendSpace1D::IsSameSamplePoint(const FVector& SamplePointA, const FVector& SamplePointB) const
-{
-	return FMath::IsNearlyEqual(SamplePointA.X, SamplePointB.X);
-}
-
-void UBlendSpace1D::GetRawSamplesFromBlendInput(const FVector &BlendInput, TArray<FGridBlendSample, TInlineAllocator<4> > & OutBlendSamples) const
-{
-
-	FVector NormalizedBlendInput = GetNormalizedBlendInput(BlendInput);
-
-	float GridIndex = FMath::TruncToFloat(NormalizedBlendInput.X);
-	float Remainder = NormalizedBlendInput.X - GridIndex;
-
-	const FEditorElement* BeforeElement = GetGridSampleInternal(GridIndex);
-
-	if (BeforeElement)
-	{
-		FGridBlendSample NewSample;
-		NewSample.GridElement = *BeforeElement;
-		// now calculate weight - GridElement has weights to nearest samples, here we weight the grid element
-		NewSample.BlendWeight = (1.f-Remainder);
-		OutBlendSamples.Add(NewSample);
-	}
-	else
-	{
-		FGridBlendSample NewSample;
-		NewSample.GridElement = FEditorElement();
-		NewSample.BlendWeight = 0.f;
-		OutBlendSamples.Add(NewSample);
-	}
-
-	const FEditorElement* AfterElement = GetGridSampleInternal(GridIndex+1);
-
-	if (AfterElement)
-	{
-		FGridBlendSample NewSample;
-		NewSample.GridElement = *AfterElement;
-		// now calculate weight - GridElement has weights to nearest samples, here we weight the grid element
-		NewSample.BlendWeight = (Remainder);
-		OutBlendSamples.Add(NewSample);
-	}
-	else
-	{
-		FGridBlendSample NewSample;
-		NewSample.GridElement = FEditorElement();
-		NewSample.BlendWeight = 0.f;
-		OutBlendSamples.Add(NewSample);
-	}
-}
-
-#if WITH_EDITOR
-void UBlendSpace1D::SnapSamplesToClosestGridPoint()
-{
-	TArray<float> GridPoints;
-	TArray<int32> ClosestSampleToGridPoint;
-	TArray<bool> SampleDataOnPoints;
-	TArray<bool> ShouldSnap;
-
-	const float GridMin = BlendParameters[0].Min;
-	const float GridMax = BlendParameters[0].Max;
-	const float GridRange = GridMax - GridMin;
-	const int32 NumGridPoints = BlendParameters[0].GridNum + 1;
-	const float GridStep = GridRange / BlendParameters[0].GridNum;
-
-	// First mark all samples as invalid and s unsnapped, recording whether they were snapped before
-	ShouldSnap.SetNumZeroed(SampleData.Num());
-	for (int32 BlendSampleIndex = 0; BlendSampleIndex < SampleData.Num(); ++BlendSampleIndex)
-	{
-		FBlendSample& BlendSample = SampleData[BlendSampleIndex];
-		ShouldSnap[BlendSampleIndex] = BlendSample.bSnapToGrid;
-		BlendSample.bSnapToGrid = false;
-	}
-
-	for (int32 GridPointIndex = 0; GridPointIndex < NumGridPoints; ++GridPointIndex)
-	{
-		const float GridPointValue = (GridPointIndex * GridStep) + GridMin;
-		GridPoints.Add(GridPointValue);
-	}
-
-	ClosestSampleToGridPoint.Init(INDEX_NONE, GridPoints.Num());
-
-	// Find closest sample to grid point
-	for (int32 PointIndex = 0; PointIndex < GridPoints.Num(); ++PointIndex)
-	{
-		const float GridPoint = GridPoints[PointIndex];
-		float SmallestDistance = FLT_MAX;
-		int32 Index = INDEX_NONE;
-
-		for (int32 SampleIndex = 0; SampleIndex < SampleData.Num(); ++SampleIndex)
-		{
-			FBlendSample& BlendSample = SampleData[SampleIndex];
-			const float Distance = FMath::Abs(GridPoint - BlendSample.SampleValue.X);
-			if (Distance < SmallestDistance)
-			{
-				Index = SampleIndex;
-				SmallestDistance = Distance;
-			}
-		}
-
-		ClosestSampleToGridPoint[PointIndex] = Index;
-	}
-
-	// Find closest grid point to sample
-	for (int32 SampleIndex = 0; SampleIndex < SampleData.Num(); ++SampleIndex)
-	{
-		if (ShouldSnap[SampleIndex])
-		{
-			FBlendSample& BlendSample = SampleData[SampleIndex];
-
-			// Find closest grid point
-			float SmallestDistance = FLT_MAX;
-			int32 Index = INDEX_NONE;
-			for (int32 PointIndex = 0; PointIndex < GridPoints.Num(); ++PointIndex)
-			{
-				const float Distance = FMath::Abs(GridPoints[PointIndex] - BlendSample.SampleValue.X);
-				if (Distance < SmallestDistance)
-				{
-					Index = PointIndex;
-					SmallestDistance = Distance;
-				}
-			}
-
-			// Only move the sample if it is also closest to the grid point
-			if (Index != INDEX_NONE && ClosestSampleToGridPoint[Index] == SampleIndex)
-			{
-				BlendSample.SampleValue.X = GridPoints[Index];
-				BlendSample.bSnapToGrid = true;
-			}
-		}
-	}
-}
-
-void UBlendSpace1D::RemapSamplesToNewAxisRange()
-{
-	const float OldGridMin = PreviousAxisMinMaxValues[0].X;
-	const float OldGridMax = PreviousAxisMinMaxValues[0].Y;
-	const float OldGridRange = OldGridMax - OldGridMin;
-
-	const float NewGridMin = BlendParameters[0].Min;
-	const float NewGridMax = BlendParameters[0].Max;
-	const float NewGridRange = NewGridMax - NewGridMin;
-
-	for (FBlendSample& BlendSample : SampleData)
-	{
-		const float NormalizedValue = (BlendSample.SampleValue.X - OldGridMin) / OldGridRange;		
-		BlendSample.SampleValue.X = NewGridMin + (NormalizedValue * NewGridRange);
-	}
-}
-#endif // WITH_EDITOR
diff --git a/Engine/Source/Runtime/Engine/Private/Animation/BlendSpaceBase.cpp b/Engine/Source/Runtime/Engine/Private/Animation/BlendSpaceBase.cpp
deleted file mode 100644
diff --git a/Engine/Source/Runtime/Engine/Private/Animation/BlendSpaceUtilities.cpp b/Engine/Source/Runtime/Engine/Private/Animation/BlendSpaceUtilities.cpp
--- a/Engine/Source/Runtime/Engine/Private/Animation/BlendSpaceUtilities.cpp
+++ b/Engine/Source/Runtime/Engine/Private/Animation/BlendSpaceUtilities.cpp
@@ -1,7 +1,7 @@
 // Copyright Epic Games, Inc. All Rights Reserved.
 
 #include "Animation/BlendSpaceUtilities.h"
-#include "Animation/BlendSpaceBase.h"
+#include "Animation/BlendSpace.h"
 #include "Animation/AnimSequence.h"
 
 int32 FSyncPattern::IndexOf(FName Name, int32 StartIndex /*= 0*/) const
diff --git a/Engine/Source/Runtime/Engine/Private/AnimationAsset.cpp b/Engine/Source/Runtime/Engine/Private/AnimationAsset.cpp
--- a/Engine/Source/Runtime/Engine/Private/AnimationAsset.cpp
+++ b/Engine/Source/Runtime/Engine/Private/AnimationAsset.cpp
@@ -8,7 +8,7 @@
 #include "AnimationUtils.h"
 #include "Animation/AnimInstance.h"
 #include "UObject/LinkerLoad.h"
-#include "Animation/BlendSpaceBase.h"
+#include "Animation/BlendSpace.h"
 #include "Animation/PoseAsset.h"
 #include "Animation/AnimSequenceHelpers.h"
 #include "Animation/AnimData/AnimDataController.h"
@@ -192,7 +192,7 @@ FAnimTickRecord::FAnimTickRecord(UAnimSequenceBase* InSequence, bool bInLooping,
 	bLooping = bInLooping;
 }
 
-FAnimTickRecord::FAnimTickRecord(UBlendSpaceBase* InBlendSpace, const FVector& InBlendInput, TArray<FBlendSampleData>& InBlendSampleDataCache, FBlendFilter& InBlendFilter, bool bInLooping, float InPlayRate, float InFinalBlendWeight, float& InCurrentTime, FMarkerTickRecord& InMarkerTickRecord)
+FAnimTickRecord::FAnimTickRecord(UBlendSpace* InBlendSpace, const FVector& InBlendInput, TArray<FBlendSampleData>& InBlendSampleDataCache, FBlendFilter& InBlendFilter, bool bInLooping, float InPlayRate, float InFinalBlendWeight, float& InCurrentTime, FMarkerTickRecord& InMarkerTickRecord)
 {
 	SourceAsset = InBlendSpace;
 	BlendSpace.BlendSpacePositionX = InBlendInput.X;
diff --git a/Engine/Source/Runtime/Engine/Public/Animation/AnimInstanceProxy.h b/Engine/Source/Runtime/Engine/Public/Animation/AnimInstanceProxy.h
--- a/Engine/Source/Runtime/Engine/Public/Animation/AnimInstanceProxy.h
+++ b/Engine/Source/Runtime/Engine/Public/Animation/AnimInstanceProxy.h
@@ -22,7 +22,7 @@
 
 #include "AnimInstanceProxy.generated.h"
 
-class UBlendSpaceBase;
+class UBlendSpace;
 struct FAnimNode_AssetPlayerBase;
 struct FAnimNode_Base;
 struct FAnimNode_SaveCachedPose;
@@ -395,8 +395,8 @@ public:
 	UE_DEPRECATED(5.0, "Please use the FAnimTickRecord constructor that takes a UAnimSequenceBase")
 	void MakeSequenceTickRecord(FAnimTickRecord& TickRecord, UAnimSequenceBase* Sequence, bool bLooping, float PlayRate, float FinalBlendWeight, float& CurrentTime, FMarkerTickRecord& MarkerTickRecord) const;
 
-	UE_DEPRECATED(5.0, "Please use the FAnimTickRecord constructor that takes a UBlendSpaceBase")
-	void MakeBlendSpaceTickRecord(FAnimTickRecord& TickRecord, UBlendSpaceBase* BlendSpace, const FVector& BlendInput, TArray<FBlendSampleData>& BlendSampleDataCache, FBlendFilter& BlendFilter, bool bLooping, float PlayRate, float FinalBlendWeight, float& CurrentTime, FMarkerTickRecord& MarkerTickRecord) const;
+	UE_DEPRECATED(5.0, "Please use the FAnimTickRecord constructor that takes a UBlendSpace")
+	void MakeBlendSpaceTickRecord(FAnimTickRecord& TickRecord, UBlendSpace* BlendSpace, const FVector& BlendInput, TArray<FBlendSampleData>& BlendSampleDataCache, FBlendFilter& BlendFilter, bool bLooping, float PlayRate, float FinalBlendWeight, float& CurrentTime, FMarkerTickRecord& MarkerTickRecord) const;
 
 	UE_DEPRECATED(5.0, "Please use the FAnimTickRecord constructor that takes a UPoseAsset")
 	void MakePoseAssetTickRecord(FAnimTickRecord& TickRecord, class UPoseAsset* PoseAsset, float FinalBlendWeight) const;
@@ -668,7 +668,7 @@ protected:
 	void SequenceAdvanceImmediate(UAnimSequenceBase* Sequence, bool bLooping, float PlayRate, float DeltaSeconds, /*inout*/ float& CurrentTime, FMarkerTickRecord& MarkerTickRecord);
 
 	// @todo document
-	void BlendSpaceAdvanceImmediate(UBlendSpaceBase* BlendSpace, const FVector& BlendInput, TArray<FBlendSampleData> & BlendSampleDataCache, FBlendFilter & BlendFilter, bool bLooping, float PlayRate, float DeltaSeconds, /*inout*/ float& CurrentTime, FMarkerTickRecord& MarkerTickRecord);
+	void BlendSpaceAdvanceImmediate(UBlendSpace* BlendSpace, const FVector& BlendInput, TArray<FBlendSampleData> & BlendSampleDataCache, FBlendFilter & BlendFilter, bool bLooping, float PlayRate, float DeltaSeconds, /*inout*/ float& CurrentTime, FMarkerTickRecord& MarkerTickRecord);
 
 	UE_DEPRECATED(5.0, "Function renamed to GetBufferReadIndex as this no longer deals with sync groups.")
 	int32 GetSyncGroupReadIndex() const 
diff --git a/Engine/Source/Runtime/Engine/Public/Animation/AnimSingleNodeInstanceProxy.h b/Engine/Source/Runtime/Engine/Public/Animation/AnimSingleNodeInstanceProxy.h
--- a/Engine/Source/Runtime/Engine/Public/Animation/AnimSingleNodeInstanceProxy.h
+++ b/Engine/Source/Runtime/Engine/Public/Animation/AnimSingleNodeInstanceProxy.h
@@ -170,7 +170,7 @@ public:
 	void SetMontagePreviewSlot(FName PreviewSlot);
 
 private:
-	void InternalBlendSpaceEvaluatePose(class UBlendSpaceBase* BlendSpace, TArray<FBlendSampleData>& BlendSampleDataCache, FPoseContext& OutContext);
+	void InternalBlendSpaceEvaluatePose(class UBlendSpace* BlendSpace, TArray<FBlendSampleData>& BlendSampleDataCache, FPoseContext& OutContext);
 
 protected:
 #if WITH_EDITOR
diff --git a/Engine/Source/Runtime/Engine/Public/AnimationRuntime.h b/Engine/Source/Runtime/Engine/Public/AnimationRuntime.h
--- a/Engine/Source/Runtime/Engine/Public/AnimationRuntime.h
+++ b/Engine/Source/Runtime/Engine/Public/AnimationRuntime.h
@@ -18,7 +18,7 @@
 #include "BonePose.h"
 #include "Containers/ArrayView.h"
 
-class UBlendSpaceBase;
+class UBlendSpace;
 class USkeletalMeshComponent;
 struct FA2CSPose;
 struct FA2Pose;
@@ -273,7 +273,7 @@ public:
 	static void BlendPosesTogetherPerBoneInMeshSpace(
 		TArrayView<FCompactPose> SourcePoses,
 		TArrayView<const FBlendedCurve> SourceCurves,
-		const UBlendSpaceBase* BlendSpace,
+		const UBlendSpace* BlendSpace,
 		TArrayView<const FBlendSampleData> BlendSampleDataCache,
 		/*out*/ FCompactPose& ResultPose,
 		/*out*/ FBlendedCurve& ResultCurve);
@@ -282,7 +282,7 @@ public:
 		TArrayView< FCompactPose> SourcePoses,
 		TArrayView<const FBlendedCurve> SourceCurves,
 		TArrayView<const FStackCustomAttributes> SourceAttributes,	
-		const UBlendSpaceBase* BlendSpace,
+		const UBlendSpace* BlendSpace,
 		TArrayView<const FBlendSampleData> BlendSampleDataCache,
 		/*out*/ FAnimationPoseData& OutAnimationPoseData);
 
diff --git a/Engine/Source/Runtime/Engine/Public/EngineFwd.h b/Engine/Source/Runtime/Engine/Public/EngineFwd.h
--- a/Engine/Source/Runtime/Engine/Public/EngineFwd.h
+++ b/Engine/Source/Runtime/Engine/Public/EngineFwd.h
@@ -36,7 +36,7 @@ class UAnimationAsset;
 class UAudioComponent;
 class UBillboardComponent;
 class UBlendProfile;
-class UBlendSpaceBase;
+class UBlendSpace;
 class UBlueprint;
 class UBlueprintGeneratedClass;
 class UBodySetup;
