Removed ordering assumptions between PostLoad and PostLoadDefaultObject in UAnimBlueprintGeneratedClass

Functions are now generated just before we use the function list

diff --git a/Engine/Source/Runtime/Engine/Private/Animation/AnimBlueprintGeneratedClass.cpp b/Engine/Source/Runtime/Engine/Private/Animation/AnimBlueprintGeneratedClass.cpp
--- a/Engine/Source/Runtime/Engine/Private/Animation/AnimBlueprintGeneratedClass.cpp
+++ b/Engine/Source/Runtime/Engine/Private/Animation/AnimBlueprintGeneratedClass.cpp
@@ -425,6 +425,12 @@ void UAnimBlueprintGeneratedClass::PostLoadDefaultObject(UObject* Object)
 {
 	Super::PostLoadDefaultObject(Object);
 
+	// Make sure that 'functions' are generated before we use them in LinkFunctionsToDefaultObjectNodes()
+	if (AnimBlueprintFunctions.IsEmpty()) 
+	{ 
+		GenerateAnimationBlueprintFunctions();
+	}
+
 	UAnimBlueprintGeneratedClass* Iter = this;
 	while(Iter)
 	{
@@ -438,7 +444,6 @@ void UAnimBlueprintGeneratedClass::PostLoadDefaultObject(UObject* Object)
 void UAnimBlueprintGeneratedClass::PostLoad()
 {
 	Super::PostLoad();
-	GenerateAnimationBlueprintFunctions();
 
 	// Post-load property access library
 	PropertyAccess::PostLoadLibrary(PropertyAccessLibrary);
