LevelInstance: ISM Packer uses ISMComponentDescriptor

Add bHiddenInGame, bIsEditorOnly to hash

diff --git a/Engine/Source/Runtime/Engine/Private/ISMPartition/ISMComponentDescriptor.cpp b/Engine/Source/Runtime/Engine/Private/ISMPartition/ISMComponentDescriptor.cpp
--- a/Engine/Source/Runtime/Engine/Private/ISMPartition/ISMComponentDescriptor.cpp
+++ b/Engine/Source/Runtime/Engine/Private/ISMPartition/ISMComponentDescriptor.cpp
@@ -10,6 +10,14 @@
 
 #endif
 
+FISMComponentDescriptor::FISMComponentDescriptor()
+{
+#if WITH_EDITORONLY_DATA
+	// Make sure we have proper defaults
+	InitFrom(UHierarchicalInstancedStaticMeshComponent::StaticClass()->GetDefaultObject<UHierarchicalInstancedStaticMeshComponent>());
+#endif
+}
+
 #if WITH_EDITOR
 
 FISMComponentDescriptor FISMComponentDescriptor::CreateFrom(const TSubclassOf<UStaticMeshComponent>& From)
@@ -51,6 +59,9 @@ void FISMComponentDescriptor::InitFrom(UStaticMeshComponent* Template, bool bIni
 	bUseAsOccluder = Template->bUseAsOccluder;
 	bRenderCustomDepth = Template->bRenderCustomDepth;
 	bIncludeInHLOD = Template->bEnableAutoLODGeneration;
+	bHiddenInGame = Template->bHiddenInGame;
+	bIsEditorOnly = Template->bIsEditorOnly;
+	bVisibleInRayTracing = Template->bVisibleInRayTracing;
 
 	if (UInstancedStaticMeshComponent* ISMTemplate = Cast<UInstancedStaticMeshComponent>(Template))
 	{
@@ -107,6 +118,9 @@ bool FISMComponentDescriptor::operator==(const FISMComponentDescriptor& Other) c
 	bRenderCustomDepth == Other.bRenderCustomDepth &&
 	bIncludeInHLOD == Other.bIncludeInHLOD &&
 	bEnableDiscardOnLoad == Other.bEnableDiscardOnLoad &&
+	bHiddenInGame == Other.bHiddenInGame &&
+	bIsEditorOnly == Other.bIsEditorOnly &&
+	bVisibleInRayTracing == Other.bVisibleInRayTracing &&
 	BodyInstance.GetCollisionEnabled() == Other.BodyInstance.GetCollisionEnabled() && 
 	BodyInstance.GetCollisionResponse() == Other.BodyInstance.GetCollisionResponse() &&
 	BodyInstance.DoesUseCollisionProfile() == Other.BodyInstance.DoesUseCollisionProfile() &&
@@ -129,6 +143,13 @@ UInstancedStaticMeshComponent* FISMComponentDescriptor::CreateComponent(UObject*
 {
 	UInstancedStaticMeshComponent* ISMComponent = NewObject<UInstancedStaticMeshComponent>(Outer, ComponentClass, Name, ObjectFlags);
 	
+	InitComponent(ISMComponent);
+
+	return ISMComponent;
+}
+
+void FISMComponentDescriptor::InitComponent(UInstancedStaticMeshComponent* ISMComponent) const
+{
 	ISMComponent->SetStaticMesh(StaticMesh);
 	ISMComponent->OverrideMaterials = OverrideMaterials;
 	ISMComponent->Mobility = Mobility;
@@ -157,14 +178,15 @@ UInstancedStaticMeshComponent* FISMComponentDescriptor::CreateComponent(UObject*
 	ISMComponent->bUseAsOccluder = bUseAsOccluder;
 	ISMComponent->bRenderCustomDepth = bRenderCustomDepth;
 	ISMComponent->bEnableAutoLODGeneration = bIncludeInHLOD;;
+	ISMComponent->bHiddenInGame = bHiddenInGame;
+	ISMComponent->bIsEditorOnly = bIsEditorOnly;
+	ISMComponent->bVisibleInRayTracing = bVisibleInRayTracing;
 
 	// HISM Specific
 	if (UHierarchicalInstancedStaticMeshComponent* HISMComponent = Cast<UHierarchicalInstancedStaticMeshComponent>(ISMComponent))
 	{
 		HISMComponent->bEnableDensityScaling = bEnableDensityScaling;
 	}
-
-	return ISMComponent;
 }
 
 #endif
diff --git a/Engine/Source/Runtime/Engine/Private/LevelInstance/Packed/LevelInstanceISMPacker.cpp b/Engine/Source/Runtime/Engine/Private/LevelInstance/Packed/LevelInstanceISMPacker.cpp
--- a/Engine/Source/Runtime/Engine/Private/LevelInstance/Packed/LevelInstanceISMPacker.cpp
+++ b/Engine/Source/Runtime/Engine/Private/LevelInstance/Packed/LevelInstanceISMPacker.cpp
@@ -7,10 +7,8 @@
 #include "LevelInstance/Packed/PackedLevelInstanceBuilder.h"
 #include "LevelInstance/Packed/PackedLevelInstanceActor.h"
 
-#include "Misc/Crc.h"
+#include "Templates/TypeHash.h"
 
-#include "Engine/StaticMeshActor.h"
-#include "Engine/StaticMesh.h"
 #include "Components/StaticMeshComponent.h"
 #include "Components/InstancedStaticMeshComponent.h"
 
@@ -48,19 +46,7 @@ void FLevelInstanceISMPacker::PackActors(FPackedLevelInstanceBuilderContext& InC
 	FLevelInstanceISMPackerCluster* ISMCluster = (FLevelInstanceISMPackerCluster*)InClusterID.GetData();
 	check(ISMCluster);
 
-	PackComponent->bReceivesDecals = ISMCluster->bReceivesDecals;
-	PackComponent->CastShadow = ISMCluster->bCastShadow;
-	PackComponent->bVisibleInRayTracing = ISMCluster->bVisibleInRayTracing;
-	PackComponent->SetStaticMesh(ISMCluster->StaticMesh);
-	for (int32 MaterialIndex = 0; MaterialIndex < ISMCluster->Materials.Num(); ++MaterialIndex)
-	{
-		PackComponent->SetMaterial(MaterialIndex, ISMCluster->Materials[MaterialIndex]);
-	}
-
-	if (!ISMCluster->bCollisionEnabled)
-	{
-		PackComponent->SetCollisionEnabled(ECollisionEnabled::NoCollision);
-	}
+	ISMCluster->ISMDescriptor.InitComponent(PackComponent);
 
 	for (UActorComponent* Component : InComponents)
 	{
@@ -92,43 +78,23 @@ void FLevelInstanceISMPacker::PackActors(FPackedLevelInstanceBuilderContext& InC
 
 FLevelInstanceISMPackerCluster::FLevelInstanceISMPackerCluster(FLevelInstancePackerID InPackerID, UStaticMeshComponent* InComponent)
 	: FLevelInstancePackerCluster(InPackerID)
-	, StaticMesh(InComponent->GetStaticMesh())
-	, Materials(InComponent->GetMaterials())
-	, bReceivesDecals(InComponent->bReceivesDecals)
-	, bCastShadow(InComponent->CastShadow)
-	, bVisibleInRayTracing(InComponent->bVisibleInRayTracing)
-	, bCollisionEnabled(InComponent->IsCollisionEnabled())
 {
-	
+	ISMDescriptor.InitFrom(InComponent, /** bInitBodyInstance= */ false);
+	// Component descriptor should be considered hidden if original actor owner was.
+	ISMDescriptor.bHiddenInGame |= InComponent->GetOwner()->IsHidden();
+	ISMDescriptor.BodyInstance.CopyRuntimeBodyInstancePropertiesFrom(&InComponent->BodyInstance);
+	ISMDescriptor.ComputeHash();
 }
 
 uint32 FLevelInstanceISMPackerCluster::ComputeHash() const
 {
-	uint32 ParentHash = FLevelInstancePackerCluster::ComputeHash();
-	// Hash pointers for now since we are not persisting the clusters
-	uint32 Hash = FCrc::TypeCrc32(StaticMesh, ParentHash);
-	for (UMaterialInterface* Material : Materials)
-	{
-		Hash = FCrc::TypeCrc32(Material, Hash);
-	}
-	Hash = FCrc::TypeCrc32(bReceivesDecals, Hash);
-	Hash = FCrc::TypeCrc32(bCastShadow, Hash);
-	Hash = FCrc::TypeCrc32(bVisibleInRayTracing, Hash);
-	Hash = FCrc::TypeCrc32(bCollisionEnabled, Hash);
-
-	return Hash;
+	return HashCombine(FLevelInstancePackerCluster::ComputeHash(), ISMDescriptor.Hash);
 }
 
 bool FLevelInstanceISMPackerCluster::operator==(const FLevelInstancePackerCluster& InOther) const
 {
 	const FLevelInstanceISMPackerCluster& ISMOther = (const FLevelInstanceISMPackerCluster&)InOther;
-	return FLevelInstancePackerCluster::operator==(InOther) &&
-		StaticMesh == ISMOther.StaticMesh &&
-		Materials == ISMOther.Materials &&
-		bReceivesDecals == ISMOther.bReceivesDecals &&
-		bCastShadow == ISMOther.bCastShadow &&
-		bVisibleInRayTracing == ISMOther.bVisibleInRayTracing &&
-		bCollisionEnabled == ISMOther.bCollisionEnabled;
+	return FLevelInstancePackerCluster::operator==(InOther) && ISMDescriptor == ISMOther.ISMDescriptor;
 }
 
 #endif
\ No newline at end of file
diff --git a/Engine/Source/Runtime/Engine/Public/ISMPartition/ISMComponentDescriptor.h b/Engine/Source/Runtime/Engine/Public/ISMPartition/ISMComponentDescriptor.h
--- a/Engine/Source/Runtime/Engine/Public/ISMPartition/ISMComponentDescriptor.h
+++ b/Engine/Source/Runtime/Engine/Public/ISMPartition/ISMComponentDescriptor.h
@@ -19,12 +19,14 @@ struct ENGINE_API FISMComponentDescriptor
 {
 	GENERATED_USTRUCT_BODY()
 
+	FISMComponentDescriptor();
 #if WITH_EDITOR
 	static FISMComponentDescriptor CreateFrom(const TSubclassOf<UStaticMeshComponent>& ComponentClass);
 	void InitFrom(UStaticMeshComponent* Component, bool bInitBodyInstance = true);
 
 	uint32 ComputeHash() const;
 	UInstancedStaticMeshComponent* CreateComponent(UObject* Outer, FName Name = NAME_None, EObjectFlags ObjectFlags = EObjectFlags::RF_NoFlags) const;
+	void InitComponent(UInstancedStaticMeshComponent* ISMComponent) const;
 
 	friend inline uint32 GetTypeHash(const FISMComponentDescriptor& Key)
 	{
@@ -142,6 +144,15 @@ public:
 
 	UPROPERTY()
 	uint8 bIncludeInHLOD : 1;
+
+	UPROPERTY()
+	uint8 bVisibleInRayTracing : 1;
+
+	UPROPERTY()
+	uint8 bHiddenInGame : 1;
+
+	UPROPERTY()
+	uint8 bIsEditorOnly : 1;
 #endif
 };
 
diff --git a/Engine/Source/Runtime/Engine/Public/LevelInstance/Packed/LevelInstanceISMPacker.h b/Engine/Source/Runtime/Engine/Public/LevelInstance/Packed/LevelInstanceISMPacker.h
--- a/Engine/Source/Runtime/Engine/Public/LevelInstance/Packed/LevelInstanceISMPacker.h
+++ b/Engine/Source/Runtime/Engine/Public/LevelInstance/Packed/LevelInstanceISMPacker.h
@@ -6,6 +6,7 @@
 
 #include "LevelInstance/Packed/PackedLevelInstanceTypes.h"
 #include "LevelInstance/Packed/ILevelInstancePacker.h"
+#include "ISMPartition/ISMComponentDescriptor.h"
 
 class AActor;
 class APackedLevelInstance;
@@ -23,9 +24,6 @@ public:
 	virtual void PackActors(FPackedLevelInstanceBuilderContext& InContext, APackedLevelInstance* InPackingActor, const FLevelInstancePackerClusterID& InClusterID, const TArray<UActorComponent*>& InComponents) const override;
 };
 
-class UStaticMesh;
-class UMaterialInterface;
-
 class FLevelInstanceISMPackerCluster : public FLevelInstancePackerCluster
 {
 public:
@@ -34,13 +32,7 @@ public:
 	virtual bool operator==(const FLevelInstancePackerCluster& InOther) const override;
 	virtual uint32 ComputeHash() const override;
 
-	// UStaticMeshComponent that contribute to the hash. More can be added if needed.
-	UStaticMesh* StaticMesh = nullptr;
-	TArray<UMaterialInterface*> Materials;
-	bool bReceivesDecals;
-	bool bCastShadow;
-	bool bVisibleInRayTracing;
-	bool bCollisionEnabled;
+	FISMComponentDescriptor ISMDescriptor;
 };
 
 #endif
\ No newline at end of file
