Fixed use-after-free of UvChunk in FHLSLMaterialTranslator::TextureSample().

AtParameterCodeChunk() returns reference to an array element that may be reallocated by AccessUniformExpression().

diff --git a/Engine/Source/Runtime/Engine/Private/Materials/HLSLMaterialTranslator.cpp b/Engine/Source/Runtime/Engine/Private/Materials/HLSLMaterialTranslator.cpp
--- a/Engine/Source/Runtime/Engine/Private/Materials/HLSLMaterialTranslator.cpp
+++ b/Engine/Source/Runtime/Engine/Private/Materials/HLSLMaterialTranslator.cpp
@@ -5736,7 +5736,7 @@ int32 FHLSLMaterialTranslator::TextureSample(
 		FString TextureName = CoerceParameter(TextureIndex, MCT_Texture2D);
 
 		EMaterialValueType UVsType = (TextureType == MCT_TextureCube || TextureType == MCT_Texture2DArray || TextureType == MCT_VolumeTexture) ? MCT_Float3 : MCT_Float2;
-		const FShaderCodeChunk& UvChunk = AtParameterCodeChunk(CoordinateIndex);
+		EMaterialExpressionAnalyticDerivativeStatus UvChunkDerivativeStatus = AtParameterCodeChunk(CoordinateIndex).DerivativeStatus;
 
 		const FString TokenUVsFinite = GetParameterCodeDeriv(CoordinateIndex,CompiledPDV_FiniteDifferences);
 		const FString TokenUVsAnalytic = GetParameterCodeDeriv(CoordinateIndex,CompiledPDV_Analytic);
@@ -5813,7 +5813,7 @@ int32 FHLSLMaterialTranslator::TextureSample(
 
 			FString ZeroFloat = (UVsType == MCT_Float2) ? TEXT("float2(0,0)") : TEXT("float3(0,0,0)");
 
-			switch(UvChunk.DerivativeStatus)
+			switch(UvChunkDerivativeStatus)
 			{
 			case MEADS_NotAware:
 			case MEADS_NotValid:
