Build fix

diff --git a/Engine/Source/Runtime/Engine/Private/PhysicsField/PhysicsFieldComponent.cpp b/Engine/Source/Runtime/Engine/Private/PhysicsField/PhysicsFieldComponent.cpp
--- a/Engine/Source/Runtime/Engine/Private/PhysicsField/PhysicsFieldComponent.cpp
+++ b/Engine/Source/Runtime/Engine/Private/PhysicsField/PhysicsFieldComponent.cpp
@@ -147,7 +147,7 @@ void InitInternalTexture(const uint32 SizeX, const uint32 SizeY, const uint32 Si
 	{
 		const uint32 BlockBytes = sizeof(BufferType) * ElementSize;
 
-		OutputBuffer.Initialize(BlockBytes, SizeX, SizeY, SizeZ, PixelFormat);
+		OutputBuffer.Initialize(TEXT("FPhysicsFieldResource"), BlockBytes, SizeX, SizeY, SizeZ, PixelFormat);
 
 		if (OutputBuffer.UAV)
 		{
