World Partition DataLayers: Override SupportsDataLayer on 5 Child AInfo classes

diff --git a/Engine/Source/Runtime/Engine/Classes/Atmosphere/AtmosphericFog.h b/Engine/Source/Runtime/Engine/Classes/Atmosphere/AtmosphericFog.h
--- a/Engine/Source/Runtime/Engine/Classes/Atmosphere/AtmosphericFog.h
+++ b/Engine/Source/Runtime/Engine/Classes/Atmosphere/AtmosphericFog.h
@@ -32,6 +32,7 @@ public:
 
 #if WITH_EDITOR
 	virtual void PostActorCreated() override;
+	virtual bool SupportsDataLayer() const override { return true; }
 #endif
 
 	/** Returns AtmosphericFogComponent subobject **/
diff --git a/Engine/Source/Runtime/Engine/Classes/Components/SkyAtmosphereComponent.h b/Engine/Source/Runtime/Engine/Classes/Components/SkyAtmosphereComponent.h
--- a/Engine/Source/Runtime/Engine/Classes/Components/SkyAtmosphereComponent.h
+++ b/Engine/Source/Runtime/Engine/Classes/Components/SkyAtmosphereComponent.h
@@ -276,6 +276,9 @@ private:
 #endif
 
 public:
+#if WITH_EDITOR
+	virtual bool SupportsDataLayer() const override { return true; }
+#endif
 
 	/** Returns SkyAtmosphereComponent subobject */
 	ENGINE_API USkyAtmosphereComponent* GetComponent() const { return SkyAtmosphereComponent; }
diff --git a/Engine/Source/Runtime/Engine/Classes/Components/VolumetricCloudComponent.h b/Engine/Source/Runtime/Engine/Classes/Components/VolumetricCloudComponent.h
--- a/Engine/Source/Runtime/Engine/Classes/Components/VolumetricCloudComponent.h
+++ b/Engine/Source/Runtime/Engine/Classes/Components/VolumetricCloudComponent.h
@@ -181,5 +181,8 @@ private:
 	TObjectPtr<class UVolumetricCloudComponent> VolumetricCloudComponent;
 
 public:
+#if WITH_EDITOR
+	virtual bool SupportsDataLayer() const override { return true; }
+#endif
 
 };
diff --git a/Engine/Source/Runtime/Engine/Classes/Engine/ExponentialHeightFog.h b/Engine/Source/Runtime/Engine/Classes/Engine/ExponentialHeightFog.h
--- a/Engine/Source/Runtime/Engine/Classes/Engine/ExponentialHeightFog.h
+++ b/Engine/Source/Runtime/Engine/Classes/Engine/ExponentialHeightFog.h
@@ -24,6 +24,9 @@ private:
 	TObjectPtr<class UExponentialHeightFogComponent> Component;
 
 public:
+#if WITH_EDITOR
+	virtual bool SupportsDataLayer() const override { return true; }
+#endif
 
 	/** replicated copy of ExponentialHeightFogComponent's bEnabled property */
 	UPROPERTY(replicatedUsing=OnRep_bEnabled)
diff --git a/Engine/Source/Runtime/Engine/Classes/Engine/SkyLight.h b/Engine/Source/Runtime/Engine/Classes/Engine/SkyLight.h
--- a/Engine/Source/Runtime/Engine/Classes/Engine/SkyLight.h
+++ b/Engine/Source/Runtime/Engine/Classes/Engine/SkyLight.h
@@ -19,6 +19,9 @@ private:
 	UPROPERTY(Category = Light, VisibleAnywhere, BlueprintReadOnly, meta = (ExposeFunctionCategories = "Light,Rendering,Rendering|Components|SkyLight", AllowPrivateAccess = "true"))
 	TObjectPtr<class USkyLightComponent> LightComponent;
 public:
+#if WITH_EDITOR
+	virtual bool SupportsDataLayer() const override { return true; }
+#endif
 
 	/** replicated copy of LightComponent's bEnabled property */
 	UPROPERTY(replicatedUsing=OnRep_bEnabled)
diff --git a/Engine/Source/Runtime/Engine/Classes/Engine/WindDirectionalSource.h b/Engine/Source/Runtime/Engine/Classes/Engine/WindDirectionalSource.h
--- a/Engine/Source/Runtime/Engine/Classes/Engine/WindDirectionalSource.h
+++ b/Engine/Source/Runtime/Engine/Classes/Engine/WindDirectionalSource.h
@@ -25,6 +25,10 @@ public:
 	/** Returns Component subobject **/
 	class UWindDirectionalSourceComponent* GetComponent() const { return Component; }
 
+#if WITH_EDITOR
+	virtual bool SupportsDataLayer() const override { return true; }
+#endif
+
 #if WITH_EDITORONLY_DATA
 	/** Returns ArrowComponent subobject **/
 	class UArrowComponent* GetArrowComponent() const { return ArrowComponent; }
