FISMComponentDescriptor::InitFrom() now takes a const UStaticMeshComponent* as input

diff --git a/Engine/Source/Runtime/Engine/Private/ISMPartition/ISMComponentDescriptor.cpp b/Engine/Source/Runtime/Engine/Private/ISMPartition/ISMComponentDescriptor.cpp
--- a/Engine/Source/Runtime/Engine/Private/ISMPartition/ISMComponentDescriptor.cpp
+++ b/Engine/Source/Runtime/Engine/Private/ISMPartition/ISMComponentDescriptor.cpp
@@ -30,7 +30,7 @@ FISMComponentDescriptor FISMComponentDescriptor::CreateFrom(const TSubclassOf<US
 	return ComponentDescriptor;
 }
 
-void FISMComponentDescriptor::InitFrom(UStaticMeshComponent* Template, bool bInitBodyInstance)
+void FISMComponentDescriptor::InitFrom(const UStaticMeshComponent* Template, bool bInitBodyInstance)
 {
 	bEnableDiscardOnLoad = false;
 	ComponentClass = Template->GetClass();
@@ -63,13 +63,13 @@ void FISMComponentDescriptor::InitFrom(UStaticMeshComponent* Template, bool bIni
 	bIsEditorOnly = Template->bIsEditorOnly;
 	bVisibleInRayTracing = Template->bVisibleInRayTracing;
 
-	if (UInstancedStaticMeshComponent* ISMTemplate = Cast<UInstancedStaticMeshComponent>(Template))
+	if (const UInstancedStaticMeshComponent* ISMTemplate = Cast<UInstancedStaticMeshComponent>(Template))
 	{
 		InstanceStartCullDistance = ISMTemplate->InstanceStartCullDistance;
 		InstanceEndCullDistance = ISMTemplate->InstanceEndCullDistance;
 
 		// HISM Specific
-		if (UHierarchicalInstancedStaticMeshComponent* HISMTemplate = Cast<UHierarchicalInstancedStaticMeshComponent>(Template))
+		if (const UHierarchicalInstancedStaticMeshComponent* HISMTemplate = Cast<UHierarchicalInstancedStaticMeshComponent>(Template))
 		{
 			bEnableDensityScaling = HISMTemplate->bEnableDensityScaling;
 		}
diff --git a/Engine/Source/Runtime/Engine/Public/ISMPartition/ISMComponentDescriptor.h b/Engine/Source/Runtime/Engine/Public/ISMPartition/ISMComponentDescriptor.h
--- a/Engine/Source/Runtime/Engine/Public/ISMPartition/ISMComponentDescriptor.h
+++ b/Engine/Source/Runtime/Engine/Public/ISMPartition/ISMComponentDescriptor.h
@@ -22,7 +22,7 @@ struct ENGINE_API FISMComponentDescriptor
 	FISMComponentDescriptor();
 #if WITH_EDITOR
 	static FISMComponentDescriptor CreateFrom(const TSubclassOf<UStaticMeshComponent>& ComponentClass);
-	void InitFrom(UStaticMeshComponent* Component, bool bInitBodyInstance = true);
+	void InitFrom(const UStaticMeshComponent* Component, bool bInitBodyInstance = true);
 
 	uint32 ComputeHash() const;
 	UInstancedStaticMeshComponent* CreateComponent(UObject* Outer, FName Name = NAME_None, EObjectFlags ObjectFlags = EObjectFlags::RF_NoFlags) const;
