Fix FSoundWaveProxy reimport behavior

diff --git a/Engine/Source/Runtime/Engine/Classes/Sound/SoundWave.h b/Engine/Source/Runtime/Engine/Classes/Sound/SoundWave.h
--- a/Engine/Source/Runtime/Engine/Classes/Sound/SoundWave.h
+++ b/Engine/Source/Runtime/Engine/Classes/Sound/SoundWave.h
@@ -777,6 +777,7 @@ public:
 	virtual bool HasCookedAmplitudeEnvelopeData() const override;
 	//~ End USoundBase Interface.
 
+	TUniquePtr<FSoundWaveProxy> CreateSoundWaveProxy();
 	//~Begin IAudioProxyDataFactory Interface.
 	virtual TUniquePtr<Audio::IProxyData> CreateNewProxyData(const Audio::FProxyDataInitParams& InitParams) override;
 	//~ End IAudioProxyDataFactory Interface.
diff --git a/Engine/Source/Runtime/Engine/Private/AudioDecompress.cpp b/Engine/Source/Runtime/Engine/Private/AudioDecompress.cpp
--- a/Engine/Source/Runtime/Engine/Private/AudioDecompress.cpp
+++ b/Engine/Source/Runtime/Engine/Private/AudioDecompress.cpp
@@ -759,11 +759,11 @@ bool ICompressedAudioInfo::StreamCompressedInfo(USoundWave* Wave, FSoundQualityI
 	}
 
 	// Create and cache proxy object
-	Audio::FProxyDataInitParams Params;
-	Params.NameOfFeatureRequestingProxy = "ICompressedAudioInfo";
-
-	TUniquePtr<Audio::IProxyData> ProxyData = Wave->CreateNewProxyData(Params);
-	StreamingSoundWave = MakeUnique<FSoundWaveProxy>(MoveTemp(ProxyData->GetAs<FSoundWaveProxy>()));
+	StreamingSoundWave = Wave->CreateSoundWaveProxy();
+	if (!StreamingSoundWave.IsValid())
+	{
+		return false;
+	}
 
 	return StreamCompressedInfoInternal(*StreamingSoundWave, QualityInfo);
 }
diff --git a/Engine/Source/Runtime/Engine/Private/SoundWave.cpp b/Engine/Source/Runtime/Engine/Private/SoundWave.cpp
--- a/Engine/Source/Runtime/Engine/Private/SoundWave.cpp
+++ b/Engine/Source/Runtime/Engine/Private/SoundWave.cpp
@@ -1138,9 +1138,7 @@ void USoundWave::PostLoad()
 
 
 	// overwrite cached proxy now that post-load is complete.
-	Audio::FProxyDataInitParams Params;
-	Params.NameOfFeatureRequestingProxy = "USoundWave";
-	InternalProxy = MakeUnique<FSoundWaveProxy>(MoveTemp(CreateNewProxyData(Params)->GetAs<FSoundWaveProxy>()));
+	InternalProxy = CreateSoundWaveProxy();
 	InternalProxy->ReleaseCompressedAudio(); // release handles if there are any
 }
 
@@ -2391,9 +2389,14 @@ bool USoundWave::HasCookedAmplitudeEnvelopeData() const
 	return CookedEnvelopeTimeData.Num() > 0;
 }
 
+TUniquePtr<FSoundWaveProxy> USoundWave::CreateSoundWaveProxy()
+{
+	return MakeUnique<FSoundWaveProxy>(this);
+}
+
 TUniquePtr<Audio::IProxyData> USoundWave::CreateNewProxyData(const Audio::FProxyDataInitParams& InitParams)
 {
-	return TUniquePtr<Audio::IProxyData>(new FSoundWaveProxy(this));
+	return CreateSoundWaveProxy();
 }
 
 void USoundWave::AddPlayingSource(const FSoundWaveClientPtr& Source)
@@ -2422,25 +2425,39 @@ void USoundWave::UpdatePlatformData()
 {
 	if (IsStreaming(nullptr))
 	{
-		// Make sure there are no pending requests in flight.
-		while (IStreamingManager::Get().GetAudioStreamingManager().IsStreamingInProgress(GetThisAsProxy()))
+		if (InternalProxy.IsValid())
 		{
-			// Give up timeslice.
-			FPlatformProcess::Sleep(0);
-		}
+			// Make sure there are no pending requests in flight.
+			while (IStreamingManager::Get().GetAudioStreamingManager().IsStreamingInProgress(*InternalProxy))
+			{
+				// Give up timeslice.
+				FPlatformProcess::Sleep(0);
+			}
 
 #if WITH_EDITORONLY_DATA
-		// Temporarily remove from streaming manager to release currently used data chunks
-		IStreamingManager::Get().GetAudioStreamingManager().RemoveStreamingSoundWave(GetThisAsProxy());
-		// Recache platform data if the source has changed.
-		CachePlatformData();
-		// Add back to the streaming manager to reload first chunk
-		IStreamingManager::Get().GetAudioStreamingManager().AddStreamingSoundWave(GetThisAsProxy());
+			// Temporarily remove from streaming manager to release currently used data chunks
+			IStreamingManager::Get().GetAudioStreamingManager().RemoveStreamingSoundWave(*InternalProxy);
+
+			// Recache platform data if the source has changed.
+			CachePlatformData();
+
+			// Release proxy holding reference to old data and replace with a new one
+			InternalProxy = CreateSoundWaveProxy();
+
+			if (InternalProxy.IsValid())
+			{
+				// Add back to the streaming manager to reload first chunk
+				IStreamingManager::Get().GetAudioStreamingManager().AddStreamingSoundWave(*InternalProxy);
+			}
 #endif
+		}
 	}
 	else
 	{
-		IStreamingManager::Get().GetAudioStreamingManager().RemoveStreamingSoundWave(GetThisAsProxy());
+		if (InternalProxy.IsValid())
+		{
+			IStreamingManager::Get().GetAudioStreamingManager().RemoveStreamingSoundWave(*InternalProxy);
+		}
 	}
 }
 
@@ -2528,9 +2545,7 @@ const FSoundWaveProxy& USoundWave::GetThisAsProxy()
 	if (!InternalProxy.IsValid())
 	{
 		// PostLoad hasn't been called. The InternalProxy created will be overritten by PostLoad().
-		Audio::FProxyDataInitParams Params;
-		Params.NameOfFeatureRequestingProxy = "USoundWave";
-		InternalProxy = MakeUnique<FSoundWaveProxy>(MoveTemp(CreateNewProxyData(Params)->GetAs<FSoundWaveProxy>()));
+		InternalProxy = CreateSoundWaveProxy();
 	}
 
 	check(InternalProxy.IsValid());
@@ -3020,6 +3035,7 @@ FSoundWaveProxy::FSoundWaveProxy(USoundWave* InWave)
 		check(LoadingBehavior != ESoundWaveLoadingBehavior::ForceInline);
 		RunningPlatformData = MakeUnique<FStreamedAudioPlatformData>(*InWave->RunningPlatformData);
 		NumChunks = InWave->GetNumChunks();
+		EnsureZerothChunkIsLoaded();
 	}
 
 //	UE_LOG(LogAudio, Display, TEXT("MAXDEBUG: FSoundWaveProxy(USoundWave* InWave)"));
@@ -3200,4 +3216,4 @@ FString FSoundWaveProxy::GetDerivedDataKey() const
 
 	return RunningPlatformData->DerivedDataKey;
 }
-#endif // #if WITH_EDITORONLY_DATA
\ No newline at end of file
+#endif // #if WITH_EDITORONLY_DATA
