FoliageISMActor: Uninitialize when removing last instance to cleanup ISMPartitionActor

diff --git a/Engine/Source/Runtime/Engine/Private/ISMPartition/ISMPartitionActor.cpp b/Engine/Source/Runtime/Engine/Private/ISMPartition/ISMPartitionActor.cpp
--- a/Engine/Source/Runtime/Engine/Private/ISMPartition/ISMPartitionActor.cpp
+++ b/Engine/Source/Runtime/Engine/Private/ISMPartition/ISMPartitionActor.cpp
@@ -115,12 +115,19 @@ FISMClientHandle AISMPartitionActor::RegisterClient(const FGuid& ClientGuid)
 
 void AISMPartitionActor::UnregisterClient(FISMClientHandle& Handle)
 {
+	Modify();
 	check(Handle.Guid == Clients[Handle.Index]);
 	RemoveISMInstances(Handle);
 			
+	int32 ClientIndex = Handle.Index;
 	Clients[Handle.Index] = FGuid();
 	Handle.Index = -1;
 	Handle.Guid = FGuid();
+
+	if (ClientIndex == Clients.Num() - 1)
+	{
+		Clients.RemoveAt(ClientIndex);
+	}
 }
 
 int32 AISMPartitionActor::RegisterISMComponentDescriptor(const FISMComponentDescriptor& Descriptor)
@@ -415,18 +422,36 @@ void AISMPartitionActor::RemoveISMInstancesInternal(FISMComponentData& Component
 	EndUpdate();
 }
 
-void AISMPartitionActor::RemoveISMInstance(const FISMClientHandle& Handle, int32 InstanceIndex)
+void AISMPartitionActor::RemoveISMInstance(const FISMClientHandle& Handle, int32 InstanceIndex, bool* bOutIsEmpty)
 {
 	check(Handle.Guid == Clients[Handle.Index]);
+	TOptional<bool> bIsEmpty;
 	for (FISMComponentData& ComponentData : DescriptorComponents)
 	{
 		if (ComponentData.ClientInstances.IsValidIndex(Handle.Index))
 		{
 			FISMClientData& ClientData = ComponentData.ClientInstances[Handle.Index];
-			
-			RemoveISMInstancesInternal(ComponentData, ClientData, InstanceIndex);
+			if (ClientData.Instances.Num())
+			{
+				RemoveISMInstancesInternal(ComponentData, ClientData, InstanceIndex);
+
+				if (!bIsEmpty.IsSet())
+				{
+					bIsEmpty = ClientData.Instances.IsEmpty();
+				}
+				else
+				{
+					// Client instances should match across components
+					check(bIsEmpty == ClientData.Instances.IsEmpty());
+				}
+			}
 		}
 	}
+
+	if (bOutIsEmpty)
+	{
+		*bOutIsEmpty = !bIsEmpty.IsSet() || bIsEmpty.GetValue();
+	}
 }
 
 void AISMPartitionActor::SelectISMInstances(const FISMClientHandle& Handle, bool bSelect, const TSet<int32>& Indices)
diff --git a/Engine/Source/Runtime/Engine/Public/ISMPartition/ISMPartitionActor.h b/Engine/Source/Runtime/Engine/Public/ISMPartition/ISMPartitionActor.h
--- a/Engine/Source/Runtime/Engine/Public/ISMPartition/ISMPartitionActor.h
+++ b/Engine/Source/Runtime/Engine/Public/ISMPartition/ISMPartitionActor.h
@@ -35,7 +35,7 @@ public:
 	const FISMComponentDescriptor& GetISMComponentDescriptor(int32 DescriptorIndex) const { return Descriptors[DescriptorIndex]; }
 
 	void AddISMInstance(const FISMClientHandle& Handle, const FTransform& InstanceTransform, const TSortedMap<int32, TArray<FTransform>>& InstanceDefinition);
-	void RemoveISMInstance(const FISMClientHandle& Handle, int32 InstanceIndex);
+	void RemoveISMInstance(const FISMClientHandle& Handle, int32 InstanceIndex, bool* bOutIsEmpty = nullptr);
 	void RemoveISMInstances(const FISMClientHandle& Handle);
 	void SelectISMInstances(const FISMClientHandle& Handle, bool bSelect, const TSet<int32>& Indices);
 	void SetISMInstanceTransform(const FISMClientHandle& Handle, int32 InstanceIndex, const FTransform& NewTransform, bool bTeleport, const TSortedMap<int32, TArray<FTransform>>& InstanceDefinition);
diff --git a/Engine/Source/Runtime/Foliage/Private/FoliageISMActor.cpp b/Engine/Source/Runtime/Foliage/Private/FoliageISMActor.cpp
--- a/Engine/Source/Runtime/Foliage/Private/FoliageISMActor.cpp
+++ b/Engine/Source/Runtime/Foliage/Private/FoliageISMActor.cpp
@@ -202,7 +202,13 @@ void FFoliageISMActor::AddInstance(const FFoliageInstance& NewInstance)
 
 void FFoliageISMActor::RemoveInstance(int32 InstanceIndex)
 {
-	GetIFA()->RemoveISMInstance(ClientHandle, InstanceIndex);
+	bool bOutIsEmpty = false;
+	GetIFA()->RemoveISMInstance(ClientHandle, InstanceIndex, &bOutIsEmpty);
+	
+	if(bOutIsEmpty)
+	{
+		Uninitialize();
+	}
 }
 
 void FFoliageISMActor::BeginUpdate()
