Refactor of source model representation for static meshes.
- Moved source models related classes and implementation into separate files.
- Created a new UObject wrapper for holding mesh description bulk data; this is created as a transactable object inside a UStaticMesh.
- The cached mesh description is now an inner object of the bulk data wrapper; this allows transactions on the bulk data to affect the cached mesh.
- The static mesh SourceModels array can no longer be accessed directly: FStaticMeshSourceModel requires extra initialization before it should be used, so new LODs should be added through the static mesh API.
- Undo/redo is now fixed within the geometry tools
- Fixed regression where mesh description was being unpacked unnecessarily during transactions
- MeshDescription bulk data and its cached MeshDescriptions are now emptied prior to removing the UObject reference, in order to immediately free memory instead of needing to wait for garbage collection
- No more deadlocks when clearing or committing mesh descriptions during garbage collection
- Bulk data UObjects are constructed, where possible, at startup, in order to prevent it from happening during critical moments, e.g. during garbage collection or package saving.

diff --git a/Engine/Config/BaseEngine.ini b/Engine/Config/BaseEngine.ini
--- a/Engine/Config/BaseEngine.ini
+++ b/Engine/Config/BaseEngine.ini
@@ -861,6 +861,7 @@ NearClipPlane=10.0
 +PropertyRedirects=(OldName="RigVMInjectionInfo.StructNode", NewName="UnitNode")
 +ClassRedirects=(OldName="/Script/UnrealEd.ProjectPackagingSettings", NewName="/Script/DeveloperToolSettings.ProjectPackagingSettings")
 +ClassRedirects=(OldName="/Script/UnrealEd.CookerSettings", NewName="/Script/DeveloperToolSettings.CookerSettings")
++ClassRedirects=(OldName="/Script/MeshDescription.MeshDescriptionBulkDataWrapper",NewName="/Script/Engine.StaticMeshDescriptionBulkData")
 +StructRedirects=(OldName="/Script/Engine.FilePath",NewName="/Script/CoreUObject.FilePath")
 +StructRedirects=(OldName="/Script/Engine.DirectoryPath",NewName="/Script/CoreUObject.DirectoryPath")
 
diff --git a/Engine/Plugins/Enterprise/DataprepEditor/Source/DataprepCore/Private/DataprepCorePrivateUtils.cpp b/Engine/Plugins/Enterprise/DataprepEditor/Source/DataprepCore/Private/DataprepCorePrivateUtils.cpp
--- a/Engine/Plugins/Enterprise/DataprepEditor/Source/DataprepCore/Private/DataprepCorePrivateUtils.cpp
+++ b/Engine/Plugins/Enterprise/DataprepEditor/Source/DataprepCore/Private/DataprepCorePrivateUtils.cpp
@@ -102,13 +102,13 @@ void DataprepCorePrivateUtils::BuildStaticMeshes(TSet<UStaticMesh*>& StaticMeshe
 		//Cache the BuildSettings and update them before building the meshes.
 		for (UStaticMesh* StaticMesh : BuiltMeshes)
 		{
-			TArray<FStaticMeshSourceModel>& SourceModels = StaticMesh->GetSourceModels();
+			int32 NumSourceModels = StaticMesh->GetNumSourceModels();
 			TArray<FMeshBuildSettings> BuildSettings;
-			BuildSettings.Reserve(SourceModels.Num());
+			BuildSettings.Reserve(NumSourceModels);
 
-			for(int32 Index = 0; Index < SourceModels.Num(); ++Index)
+			for(int32 Index = 0; Index < NumSourceModels; ++Index)
 			{
-				FStaticMeshSourceModel& SourceModel = SourceModels[Index];
+				FStaticMeshSourceModel& SourceModel = StaticMesh->GetSourceModel(Index);
 
 				BuildSettings.Add( SourceModel.BuildSettings );
 
@@ -150,11 +150,10 @@ void DataprepCorePrivateUtils::BuildStaticMeshes(TSet<UStaticMesh*>& StaticMeshe
 			UStaticMesh* StaticMesh = BuiltMeshes[Index];
 			TArray<FMeshBuildSettings>& PrevBuildSettings = StaticMeshesSettings[Index];
 
-			TArray<FStaticMeshSourceModel>& SourceModels = StaticMesh->GetSourceModels();
-
-			for(int32 SourceModelIndex = 0; SourceModelIndex < SourceModels.Num(); ++SourceModelIndex)
+			int32 NumSourceModels = StaticMesh->GetNumSourceModels();
+			for(int32 SourceModelIndex = 0; SourceModelIndex < NumSourceModels; ++SourceModelIndex)
 			{
-				SourceModels[SourceModelIndex].BuildSettings = PrevBuildSettings[SourceModelIndex];
+				StaticMesh->GetSourceModel(SourceModelIndex).BuildSettings = PrevBuildSettings[SourceModelIndex];
 			}
 
 			if(FStaticMeshRenderData* RenderData = StaticMesh->GetRenderData())
diff --git a/Engine/Plugins/Enterprise/DataprepEditor/Source/DataprepEditor/Private/Widgets/SDataprepEditorViewport.cpp b/Engine/Plugins/Enterprise/DataprepEditor/Source/DataprepEditor/Private/Widgets/SDataprepEditorViewport.cpp
--- a/Engine/Plugins/Enterprise/DataprepEditor/Source/DataprepEditor/Private/Widgets/SDataprepEditorViewport.cpp
+++ b/Engine/Plugins/Enterprise/DataprepEditor/Source/DataprepEditor/Private/Widgets/SDataprepEditorViewport.cpp
@@ -1617,12 +1617,14 @@ namespace DataprepEditor3DPreviewUtils
 			//Cache the BuildSettings and update them before building the meshes.
 			for (UStaticMesh* StaticMesh : BuiltMeshes)
 			{
-				TArray<FStaticMeshSourceModel>& SourceModels = StaticMesh->GetSourceModels();
+				int32 NumSourceModels = StaticMesh->GetNumSourceModels();
 				TArray<FMeshBuildSettings> BuildSettings;
-				BuildSettings.Reserve(SourceModels.Num());
+				BuildSettings.Reserve(NumSourceModels);
 
-				for (FStaticMeshSourceModel& SourceModel : SourceModels)
+				for (int32 LodIndex = 0; LodIndex < NumSourceModels; LodIndex++)
 				{
+					FStaticMeshSourceModel& SourceModel = StaticMesh->GetSourceModel(LodIndex);
+
 					BuildSettings.Add(SourceModel.BuildSettings);
 
 					SourceModel.BuildSettings.bGenerateLightmapUVs = false;
@@ -1649,11 +1651,11 @@ namespace DataprepEditor3DPreviewUtils
 				UStaticMesh* StaticMesh = BuiltMeshes[Index];
 				TArray<FMeshBuildSettings>& PrevBuildSettings = StaticMeshesSettings[Index];
 
-				TArray<FStaticMeshSourceModel>& SourceModels = StaticMesh->GetSourceModels();
+				int32 NumSourceModels = StaticMesh->GetNumSourceModels();
 
-				for(int32 SourceModelIndex = 0; SourceModelIndex < SourceModels.Num(); ++SourceModelIndex)
+				for(int32 SourceModelIndex = 0; SourceModelIndex < NumSourceModels; ++SourceModelIndex)
 				{
-					SourceModels[SourceModelIndex].BuildSettings = PrevBuildSettings[SourceModelIndex];
+					StaticMesh->GetSourceModel(SourceModelIndex).BuildSettings = PrevBuildSettings[SourceModelIndex];
 				}
 
 				for ( FStaticMeshLODResources& LODResources : StaticMesh->GetRenderData()->LODResources )
diff --git a/Engine/Plugins/Enterprise/DataprepEditor/Source/DataprepLibraries/Private/DataprepEditingOperations.cpp b/Engine/Plugins/Enterprise/DataprepEditor/Source/DataprepLibraries/Private/DataprepEditingOperations.cpp
--- a/Engine/Plugins/Enterprise/DataprepEditor/Source/DataprepLibraries/Private/DataprepEditingOperations.cpp
+++ b/Engine/Plugins/Enterprise/DataprepEditor/Source/DataprepLibraries/Private/DataprepEditingOperations.cpp
@@ -1052,7 +1052,7 @@ namespace DatasmithEditingOperationsUtils
 						// Skip components which are either editor only or for visualization
 						if(!MeshComponent->IsEditorOnly() && !MeshComponent->IsVisualizationComponent())
 						{
-							if(MeshComponent->GetStaticMesh() && MeshComponent->GetStaticMesh()->GetSourceModels().Num() > 0)
+							if(MeshComponent->GetStaticMesh() && MeshComponent->GetStaticMesh()->GetNumSourceModels() > 0)
 							{
 								bMeshActorIsValid = true;
 								ComponentsToMerge.Add(MeshComponent);
diff --git a/Engine/Plugins/Enterprise/DataprepEditor/Source/DataprepLibraries/Private/DataprepOperationsLibraryUtil.cpp b/Engine/Plugins/Enterprise/DataprepEditor/Source/DataprepLibraries/Private/DataprepOperationsLibraryUtil.cpp
--- a/Engine/Plugins/Enterprise/DataprepEditor/Source/DataprepLibraries/Private/DataprepOperationsLibraryUtil.cpp
+++ b/Engine/Plugins/Enterprise/DataprepEditor/Source/DataprepLibraries/Private/DataprepOperationsLibraryUtil.cpp
@@ -143,8 +143,10 @@ namespace DataprepOperationsLibraryUtil
 
 		TArray< FMeshBuildSettings > BuildSettingsBackup;
 
-		for ( FStaticMeshSourceModel& SourceModel : StaticMesh->GetSourceModels() )
+		int32 NumSourceModels = StaticMesh->GetNumSourceModels();
+		for (int32 LodIndex = 0; LodIndex < NumSourceModels; LodIndex++)
 		{
+			FStaticMeshSourceModel& SourceModel = StaticMesh->GetSourceModel(LodIndex);
 			BuildSettingsBackup.Add( SourceModel.BuildSettings );
 
 			// These were done in the PreBuild step
@@ -203,9 +205,10 @@ namespace DataprepOperationsLibraryUtil
 				// Make sure adjacency information fit new material change
 				if( RequiresAdjacencyInformation( NewMaterial, nullptr, GWorld->FeatureLevel ) )
 				{
-					for( FStaticMeshSourceModel& SourceModel : StaticMesh->GetSourceModels() )
+					int32 NumSourceModels = StaticMesh->GetNumSourceModels();
+					for (int32 LodIndex = 0; LodIndex < NumSourceModels; LodIndex++)
 					{
-						SourceModel.BuildSettings.bBuildAdjacencyBuffer = true;
+						StaticMesh->GetSourceModel(LodIndex).BuildSettings.bBuildAdjacencyBuffer = true;
 					}
 				}
 			}
@@ -270,13 +273,13 @@ namespace DataprepOperationsLibraryUtil
 
 			for (UStaticMesh* StaticMesh : BuiltMeshes)
 			{
-				TArray<FStaticMeshSourceModel>& SourceModels = StaticMesh->GetSourceModels();
+				int32 NumSourceModels = StaticMesh->GetNumSourceModels();
 				TArray<FMeshBuildSettings> BuildSettings;
-				BuildSettings.Reserve(SourceModels.Num());
+				BuildSettings.Reserve(NumSourceModels);
 
-				for(int32 Index = 0; Index < SourceModels.Num(); ++Index)
+				for(int32 Index = 0; Index < NumSourceModels; ++Index)
 				{
-					FStaticMeshSourceModel& SourceModel = SourceModels[Index];
+					FStaticMeshSourceModel& SourceModel = StaticMesh->GetSourceModel(Index);
 
 					BuildSettings.Add( SourceModel.BuildSettings );
 
@@ -317,11 +320,10 @@ namespace DataprepOperationsLibraryUtil
 				UStaticMesh* StaticMesh = BuiltMeshes[Index];
 				TArray<FMeshBuildSettings>& PrevBuildSettings = StaticMeshesSettings[Index];
 
-				TArray<FStaticMeshSourceModel>& SourceModels = StaticMesh->GetSourceModels();
-
-				for(int32 SourceModelIndex = 0; SourceModelIndex < SourceModels.Num(); ++SourceModelIndex)
+				int32 NumSourceModels = StaticMesh->GetNumSourceModels();
+				for(int32 SourceModelIndex = 0; SourceModelIndex < NumSourceModels; ++SourceModelIndex)
 				{
-					SourceModels[SourceModelIndex].BuildSettings = PrevBuildSettings[SourceModelIndex];
+					StaticMesh->GetSourceModel(SourceModelIndex).BuildSettings = PrevBuildSettings[SourceModelIndex];
 				}
 
 				for ( FStaticMeshLODResources& LODResources : StaticMesh->GetRenderData()->LODResources )
diff --git a/Engine/Plugins/Enterprise/DatasmithImporter/Source/DatasmithImporter/Private/Utility/DatasmithImporterUtils.cpp b/Engine/Plugins/Enterprise/DatasmithImporter/Source/DatasmithImporter/Private/Utility/DatasmithImporterUtils.cpp
--- a/Engine/Plugins/Enterprise/DatasmithImporter/Source/DatasmithImporter/Private/Utility/DatasmithImporterUtils.cpp
+++ b/Engine/Plugins/Enterprise/DatasmithImporter/Source/DatasmithImporter/Private/Utility/DatasmithImporterUtils.cpp
@@ -1218,11 +1218,10 @@ UObject* FDatasmithImporterUtils::StaticDuplicateObject(UObject* SourceObject, U
 
 UStaticMesh* FDatasmithImporterUtils::DuplicateStaticMesh(UStaticMesh* SourceStaticMesh, UObject* Outer, const FName Name, bool bIgnoreBulkData)
 {
-	TArray<FStaticMeshSourceModel> SourceModels;
-
 	// Since static mesh can be quite heavy, remove source models for cloning to reduce useless work.
 	// Will be reinserted on the new duplicated asset or restored on the SourceStaticMesh if bIgnoreBulkData is true.
-	SourceModels = MoveTemp(SourceStaticMesh->GetSourceModels());
+	TArray<FStaticMeshSourceModel> SourceModels = MoveTemp(SourceStaticMesh->MoveSourceModels());
+	FStaticMeshSourceModel HiResSourceModel = MoveTemp(SourceStaticMesh->MoveHiResSourceModel());
 
 	// Temporary flag to skip Postload during DuplicateObject
 	SourceStaticMesh->SetFlags(RF_ArchetypeObject);
@@ -1245,7 +1244,6 @@ UStaticMesh* FDatasmithImporterUtils::DuplicateStaticMesh(UStaticMesh* SourceSta
 		for (FStaticMeshSourceModel& SourceModel : SourceModels)
 		{
 			FStaticMeshSourceModel& DuplicateSourceModel = DuplicateMesh->AddSourceModel();
-			DuplicateSourceModel.StaticMeshOwner = DuplicateMesh;
 
 			// Apply the SourceMesh settings to the duplicated SourceModels
 			DuplicateSourceModel.BuildSettings = SourceModel.BuildSettings;
@@ -1256,7 +1254,8 @@ UStaticMesh* FDatasmithImporterUtils::DuplicateStaticMesh(UStaticMesh* SourceSta
 		}
 
 		// Move back the source models to the original mesh
-		SourceStaticMesh->GetSourceModels() = MoveTemp(SourceModels);
+		SourceStaticMesh->SetSourceModels(MoveTemp(SourceModels));
+		SourceStaticMesh->SetHiResSourceModel(MoveTemp(HiResSourceModel));
 	}
 	else
 	{
@@ -1264,14 +1263,9 @@ UStaticMesh* FDatasmithImporterUtils::DuplicateStaticMesh(UStaticMesh* SourceSta
 		// -> MarkPendingKill to avoid use-after-move crash in the StaticMesh::Build()
 		SourceStaticMesh->MarkPendingKill();
 
-		for (FStaticMeshSourceModel& SourceModel : SourceModels)
-		{
-			// Fixup the new SourceModels owner
-			SourceModel.StaticMeshOwner = DuplicateMesh;
-		}
-
 		// Apply source models to the duplicated mesh
-		DuplicateMesh->GetSourceModels() = MoveTemp(SourceModels);
+		DuplicateMesh->SetSourceModels(MoveTemp(SourceModels));
+		DuplicateMesh->SetHiResSourceModel(MoveTemp(HiResSourceModel));
 	}
 
 	return DuplicateMesh;
diff --git a/Engine/Source/Developer/HierarchicalLODUtilities/Private/HierarchicalLODUtilities.cpp b/Engine/Source/Developer/HierarchicalLODUtilities/Private/HierarchicalLODUtilities.cpp
--- a/Engine/Source/Developer/HierarchicalLODUtilities/Private/HierarchicalLODUtilities.cpp
+++ b/Engine/Source/Developer/HierarchicalLODUtilities/Private/HierarchicalLODUtilities.cpp
@@ -1094,7 +1094,7 @@ int32 FHierarchicalLODUtilities::GetLODLevelForScreenSize(const UStaticMeshCompo
 	checkf(RenderData != nullptr, TEXT("StaticMesh in StaticMeshComponent %s contains invalid render data"), *StaticMeshComponent->GetName());
 	checkf(StaticMeshComponent->GetStaticMesh()->GetNumSourceModels() > 0, TEXT("StaticMesh in StaticMeshComponent %s contains no SourceModels"), *StaticMeshComponent->GetName());
 
-	return ComputeStaticMeshLODLevel(StaticMeshComponent->GetStaticMesh()->GetSourceModels(), RenderData, ScreenSize);
+	return ComputeStaticMeshLODLevel(const_cast<const UStaticMesh*>(StaticMeshComponent->GetStaticMesh())->GetSourceModels(), RenderData, ScreenSize);
 }
 
 AHierarchicalLODVolume* FHierarchicalLODUtilities::CreateVolumeForLODActor(ALODActor* InLODActor, UWorld* InWorld)
diff --git a/Engine/Source/Developer/MeshUtilities/Private/MeshUtilities.cpp b/Engine/Source/Developer/MeshUtilities/Private/MeshUtilities.cpp
--- a/Engine/Source/Developer/MeshUtilities/Private/MeshUtilities.cpp
+++ b/Engine/Source/Developer/MeshUtilities/Private/MeshUtilities.cpp
@@ -2455,15 +2455,15 @@ public:
 	{
 		check(Stage == EStage::Uninit);
 		check(StaticMesh != nullptr);
-		TArray<FStaticMeshSourceModel>& SourceModels = StaticMesh->GetSourceModels();
 		ELightmapUVVersion LightmapUVVersion = (ELightmapUVVersion)StaticMesh->GetLightmapUVVersion();
 
 		FMeshUtilities& MeshUtilities = FModuleManager::Get().LoadModuleChecked<FMeshUtilities>("MeshUtilities");
 
 		// Gather source meshes for each LOD.
-		for (int32 LODIndex = 0; LODIndex < SourceModels.Num(); ++LODIndex)
+		int32 NumSourceModels = StaticMesh->GetNumSourceModels();
+		for (int32 LODIndex = 0; LODIndex < NumSourceModels; ++LODIndex)
 		{
-			FStaticMeshSourceModel& SrcModel = SourceModels[LODIndex];
+			FStaticMeshSourceModel& SrcModel = StaticMesh->GetSourceModel(LODIndex);
 			FRawMesh& RawMesh = *new FRawMesh;
 			LODMeshes.Add(&RawMesh);
 			FOverlappingCorners& OverlappingCorners = *new FOverlappingCorners;
@@ -2580,8 +2580,8 @@ public:
 				}
 			}
 		}
-		check(LODMeshes.Num() == SourceModels.Num());
-		check(LODOverlappingCorners.Num() == SourceModels.Num());
+		check(LODMeshes.Num() == NumSourceModels);
+		check(LODOverlappingCorners.Num() == NumSourceModels);
 
 		// Bail if there is no raw mesh data from which to build a renderable mesh.
 		if (LODMeshes.Num() == 0)
@@ -2603,8 +2603,8 @@ public:
 	{
 		check(Stage == EStage::Gathered);
 		check(StaticMesh != nullptr);
-		TArray<FStaticMeshSourceModel>& SourceModels = StaticMesh->GetSourceModels();
-		if (SourceModels.Num() == 0)
+		int32 NumSourceModels = StaticMesh->GetNumSourceModels();
+		if (NumSourceModels == 0)
 		{
 			UE_LOG(LogMeshUtilities, Error, TEXT("Mesh contains zero source models."));
 			return false;
@@ -2613,9 +2613,9 @@ public:
 		FMeshUtilities& MeshUtilities = FModuleManager::Get().LoadModuleChecked<FMeshUtilities>("MeshUtilities");
 
 		// Reduce each LOD mesh according to its reduction settings.
-		for (int32 LODIndex = 0; LODIndex < SourceModels.Num(); ++LODIndex)
+		for (int32 LODIndex = 0; LODIndex < NumSourceModels; ++LODIndex)
 		{
-			const FStaticMeshSourceModel& SrcModel = SourceModels[LODIndex];
+			const FStaticMeshSourceModel& SrcModel = StaticMesh->GetSourceModel(LODIndex);
 			FMeshReductionSettings ReductionSettings = LODGroup.GetSettings(SrcModel.ReductionSettings, LODIndex);
 			LODMaxDeviation[NumValidLODs] = 0.0f;
 			if (LODIndex != NumValidLODs)
@@ -2717,13 +2717,13 @@ public:
 		check(Stage == EStage::Reduce);
 		check(StaticMesh != nullptr);
 
-		TArray<FStaticMeshSourceModel>& InOutModels = StaticMesh->GetSourceModels();
 		int32 ImportVersion = StaticMesh->ImportVersion;
 
 		// Generate per-LOD rendering data.
 		OutRenderData.AllocateLODResources(NumValidLODs);
 		for (int32 LODIndex = 0; LODIndex < NumValidLODs; ++LODIndex)
 		{
+			FStaticMeshSourceModel& SourceModel = StaticMesh->GetSourceModel(LODIndex);
 			FStaticMeshLODResources& LODModel = OutRenderData.LODResources[LODIndex];
 			FRawMesh& RawMesh = LODMeshes[LODIndex];
 			LODModel.MaxDeviation = LODMaxDeviation[LODIndex];
@@ -2765,7 +2765,7 @@ public:
 				// TODO - write directly to TMemoryImageArray
 				// We can compute an approximate one instead for other LODs.
 				TArray<int32> TempWedgeMap;
-				TArray<int32>& WedgeMap = InOutModels[LODIndex].ReductionSettings.PercentTriangles >= 1.0f ? LODModel.WedgeMap : TempWedgeMap;
+				TArray<int32>& WedgeMap = SourceModel.ReductionSettings.PercentTriangles >= 1.0f ? LODModel.WedgeMap : TempWedgeMap;
 				WedgeMap.Reset();
 				float ComparisonThreshold = GetComparisonThreshold(LODBuildSettings[LODIndex]);
 				MeshUtilities.BuildStaticMeshVertexAndIndexBuffers(Vertices, PerSectionIndices, WedgeMap, RawMesh, LODOverlappingCorners[LODIndex], MaterialToSectionMapping, ComparisonThreshold, LODBuildSettings[LODIndex].BuildScale3D, ImportVersion);
@@ -2826,7 +2826,7 @@ public:
 			LODModel.IndexBuffer.SetIndices(CombinedIndices, bNeeds32BitIndices ? EIndexBufferStride::Force32Bit : EIndexBufferStride::Force16Bit);
 			
 			// Build the reversed index buffer.
-			if (LODModel.AdditionalIndexBuffers && InOutModels[0].BuildSettings.bBuildReversedIndexBuffer)
+			if (LODModel.AdditionalIndexBuffers && StaticMesh->GetSourceModel(0).BuildSettings.bBuildReversedIndexBuffer)
 			{
 				TArray<uint32> InversedIndices;
 				const int32 IndexCount = CombinedIndices.Num();
@@ -2864,7 +2864,7 @@ public:
 			}
 
 			// Build the inversed depth only index buffer.
-			if (LODModel.AdditionalIndexBuffers && InOutModels[0].BuildSettings.bBuildReversedIndexBuffer)
+			if (LODModel.AdditionalIndexBuffers && StaticMesh->GetSourceModel(0).BuildSettings.bBuildReversedIndexBuffer)
 			{
 				TArray<uint32> ReversedDepthOnlyIndices;
 				const int32 IndexCount = DepthOnlyIndices.Num();
@@ -2894,7 +2894,7 @@ public:
 			}
 
 			// Build the adjacency index buffer used for tessellation.
-			if (LODModel.AdditionalIndexBuffers && InOutModels[0].BuildSettings.bBuildAdjacencyBuffer)
+			if (LODModel.AdditionalIndexBuffers && StaticMesh->GetSourceModel(0).BuildSettings.bBuildAdjacencyBuffer)
 			{
 				TArray<uint32> AdjacencyIndices;
 
@@ -2942,16 +2942,14 @@ public:
 		check(Stage == EStage::Reduce);
 		check(StaticMesh != nullptr);
 
-		TArray<FStaticMeshSourceModel>& SourceModels = StaticMesh->GetSourceModels();
-
 		check(HasRawMesh[0]);
-		check(SourceModels.Num() >= NumValidLODs);
+		check(StaticMesh->GetNumSourceModels() >= NumValidLODs);
 		bool bDirty = false;
 		for (int32 Index = 1; Index < NumValidLODs; ++Index)
 		{
 			if (!HasRawMesh[Index])
 			{
-				SourceModels[Index].SaveRawMesh(LODMeshes[Index]);
+				StaticMesh->GetSourceModel(Index).SaveRawMesh(LODMeshes[Index]);
 				bDirty = true;
 			}
 		}
@@ -3035,7 +3033,6 @@ void FMeshUtilities::FixupMaterialSlotNames(USkeletalMesh* SkeletalMesh) const
 
 bool FMeshUtilities::BuildStaticMesh(FStaticMeshRenderData& OutRenderData, UStaticMesh* StaticMesh, const FStaticMeshLODGroup& LODGroup)
 {
-	TArray<FStaticMeshSourceModel>& SourceModels = StaticMesh->GetSourceModels();
 	int32 LightmapUVVersion = StaticMesh->GetLightmapUVVersion();
 	int32 ImportVersion = StaticMesh->ImportVersion;
 
@@ -3047,7 +3044,7 @@ bool FMeshUtilities::BuildStaticMesh(FStaticMeshRenderData& OutRenderData, UStat
 	}
 
 	TArray<bool> WasReduced;
-	WasReduced.AddZeroed(SourceModels.Num());
+	WasReduced.AddZeroed(StaticMesh->GetNumSourceModels());
 	if (!Builder.ReduceLODs(LODGroup, Module.GetStaticMeshReductionInterface(), WasReduced))
 	{
 		return false;
@@ -3058,7 +3055,6 @@ bool FMeshUtilities::BuildStaticMesh(FStaticMeshRenderData& OutRenderData, UStat
 
 bool FMeshUtilities::GenerateStaticMeshLODs(UStaticMesh* StaticMesh, const FStaticMeshLODGroup& LODGroup)
 {
-	TArray<FStaticMeshSourceModel>& Models = StaticMesh->GetSourceModels();
 	int32 LightmapUVVersion = StaticMesh->GetLightmapUVVersion();
 
 	FStaticMeshUtilityBuilder Builder(StaticMesh);
@@ -3069,7 +3065,7 @@ bool FMeshUtilities::GenerateStaticMeshLODs(UStaticMesh* StaticMesh, const FStat
 	}
 
 	TArray<bool> WasReduced;
-	WasReduced.AddZeroed(Models.Num());
+	WasReduced.AddZeroed(StaticMesh->GetNumSourceModels());
 	if (!Builder.ReduceLODs(LODGroup, Module.GetStaticMeshReductionInterface(), WasReduced))
 	{
 		return false;
diff --git a/Engine/Source/Editor/Experimental/EditorInteractiveToolsFramework/Private/ToolTargets/StaticMeshComponentToolTarget.cpp b/Engine/Source/Editor/Experimental/EditorInteractiveToolsFramework/Private/ToolTargets/StaticMeshComponentToolTarget.cpp
--- a/Engine/Source/Editor/Experimental/EditorInteractiveToolsFramework/Private/ToolTargets/StaticMeshComponentToolTarget.cpp
+++ b/Engine/Source/Editor/Experimental/EditorInteractiveToolsFramework/Private/ToolTargets/StaticMeshComponentToolTarget.cpp
@@ -144,8 +144,8 @@ void UStaticMeshComponentToolTarget::CommitMeshDescription(const FCommitter& Com
 	// make sure transactional flag is on for this asset
 	StaticMesh->SetFlags(RF_Transactional);
 
-	bool bSavedToTransactionBuffer = StaticMesh->Modify();
-	check(bSavedToTransactionBuffer);
+	verify(StaticMesh->Modify());
+	verify(StaticMesh->ModifyMeshDescription(LODIndex));
 
 	FCommitterParams CommitterParams;
 	CommitterParams.MeshDescriptionOut = StaticMesh->GetMeshDescription(LODIndex);
diff --git a/Engine/Source/Editor/StaticMeshEditor/Private/StaticMeshEditor.cpp b/Engine/Source/Editor/StaticMeshEditor/Private/StaticMeshEditor.cpp
--- a/Engine/Source/Editor/StaticMeshEditor/Private/StaticMeshEditor.cpp
+++ b/Engine/Source/Editor/StaticMeshEditor/Private/StaticMeshEditor.cpp
@@ -1568,12 +1568,11 @@ void FStaticMeshEditor::HandleReimportAllMesh()
 		//Reimport base LOD, generated mesh will be rebuild here, the static mesh is always using the base mesh to reduce LOD
 		if (FReimportManager::Instance()->Reimport(StaticMesh, true))
 		{
-			TArray<FStaticMeshSourceModel>& SourceModels = StaticMesh->GetSourceModels();
 			//Reimport all custom LODs
 			for (int32 LodIndex = 1; LodIndex < StaticMesh->GetNumLODs(); ++LodIndex)
 			{
 				//Skip LOD import in the same file as the base mesh, they are already re-import
-				if (SourceModels[LodIndex].bImportWithBaseMesh)
+				if (StaticMesh->GetSourceModel(LodIndex).bImportWithBaseMesh)
 				{
 					continue;
 				}
diff --git a/Engine/Source/Editor/StaticMeshEditor/Private/StaticMeshEditorSubsystem.cpp b/Engine/Source/Editor/StaticMeshEditor/Private/StaticMeshEditorSubsystem.cpp
--- a/Engine/Source/Editor/StaticMeshEditor/Private/StaticMeshEditorSubsystem.cpp
+++ b/Engine/Source/Editor/StaticMeshEditor/Private/StaticMeshEditorSubsystem.cpp
@@ -1573,10 +1573,9 @@ bool UStaticMeshEditorSubsystem::SetGenerateLightmapUVs(UStaticMesh* StaticMesh,
 	if (AnySettingsToChange)
 	{
 		StaticMesh->Modify();
-		for (FStaticMeshSourceModel& SourceModel : StaticMesh->GetSourceModels())
+		for (int32 LodIndex = 0; LodIndex < StaticMesh->GetNumSourceModels(); LodIndex++)
 		{
-			SourceModel.BuildSettings.bGenerateLightmapUVs = bGenerateLightmapUVs;
-
+			StaticMesh->GetSourceModel(LodIndex).BuildSettings.bGenerateLightmapUVs = bGenerateLightmapUVs;
 		}
 
 		StaticMesh->Build();
diff --git a/Engine/Source/Editor/UnrealEd/Private/Tools/EditorComponentSourceFactory.cpp b/Engine/Source/Editor/UnrealEd/Private/Tools/EditorComponentSourceFactory.cpp
--- a/Engine/Source/Editor/UnrealEd/Private/Tools/EditorComponentSourceFactory.cpp
+++ b/Engine/Source/Editor/UnrealEd/Private/Tools/EditorComponentSourceFactory.cpp
@@ -225,8 +225,15 @@ void FStaticMeshComponentTarget::CommitMesh( const FCommitter& Committer )
 	// make sure transactional flag is on for this asset
 	StaticMesh->SetFlags(RF_Transactional);
 
-	bool bSavedToTransactionBuffer = StaticMesh->Modify();
-	check(bSavedToTransactionBuffer);
+	verify(StaticMesh->Modify());
+	if (EditingLOD == EStaticMeshEditingLOD::HiResSource)
+	{
+		verify(StaticMesh->ModifyHiResMeshDescription());
+	}
+	else
+	{
+		verify(StaticMesh->ModifyMeshDescription((int32)EditingLOD));
+	}
 
 	FCommitParams CommitParams;
 	CommitParams.MeshDescription = GetMesh();
diff --git a/Engine/Source/Runtime/Engine/Classes/Engine/StaticMesh.h b/Engine/Source/Runtime/Engine/Classes/Engine/StaticMesh.h
--- a/Engine/Source/Runtime/Engine/Classes/Engine/StaticMesh.h
+++ b/Engine/Source/Runtime/Engine/Classes/Engine/StaticMesh.h
@@ -15,9 +15,9 @@
 #include "Components.h"
 #include "Interfaces/Interface_CollisionDataProvider.h"
 #include "Interfaces/Interface_AsyncCompilation.h"
-#include "Engine/MeshMerging.h"
 #include "Engine/StreamableRenderAsset.h"
 #include "Templates/UniquePtr.h"
+#include "StaticMeshSourceData.h"
 #include "StaticMeshResources.h"
 #include "PerPlatformProperties.h"
 #include "RenderAssetUpdate.h"
@@ -31,11 +31,9 @@ class UMaterialInterface;
 class UNavCollisionBase;
 class UStaticMeshComponent;
 class UStaticMeshDescription;
-class UMeshDescriptionBulkDataWrapper;
 class FStaticMeshUpdate;
 struct FMeshDescription;
 struct FStaticMeshLODResources;
-struct FMeshDescriptionBulkData;
 
 /*-----------------------------------------------------------------------------
 	Async Static Mesh Compilation
@@ -62,6 +60,7 @@ enum class EStaticMeshAsyncProperties : uint32
 	MinLOD                  = 1 << 15,
 	LightMapCoordinateIndex = 1 << 16,
 	LightMapResolution      = 1 << 17,
+	HiResSourceModel		= 1 << 18,
 
 	All                     = MAX_uint32
 };
@@ -108,6 +107,8 @@ inline const TCHAR* ToString(EStaticMeshAsyncProperties Value)
 			return TEXT("LightMapCoordinateIndex");
 		case EStaticMeshAsyncProperties::LightMapResolution:
 			return TEXT("LightMapResolution");
+		case EStaticMeshAsyncProperties::HiResSourceModel:
+			return TEXT("HiResSourceModel");
 		default: 
 			check(false); 
 			return TEXT("Unknown");
@@ -325,170 +326,6 @@ struct FStaticMeshOptimizationSettings
 	UStaticMesh
 -----------------------------------------------------------------------------*/
 
-/**
- * Source model from which a renderable static mesh is built.
- */
-USTRUCT()
-struct FStaticMeshSourceModel
-{
-	GENERATED_USTRUCT_BODY()
-
-	friend class UStaticMesh;
-
-#if WITH_EDITOR
-	/**
-	 * Imported raw mesh data. Optional for all but the first LOD.
-	 *
-	 * This is a member for legacy assets only.
-	 * If it is non-empty, this means that it has been de-serialized from the asset, and
-	 * the asset hence pre-dates MeshDescription.
-	 */
-	class FRawMeshBulkData* RawMeshBulkData;
-	
-	/*
-	 * The staticmesh owner of this source model. We need the SM to be able to convert between MeshDesription and RawMesh.
-	 * RawMesh use int32 material index and MeshDescription use FName material slot name.
-	 * This memeber is fill in the PostLoad of the static mesh.
-	 */
-	class UStaticMesh* StaticMeshOwner;
-
-	/*
-	 * Accessor to Load and save the raw mesh or the mesh description depending on the editor settings.
-	 * Temporary until we deprecate the RawMesh.
-	 */
-	ENGINE_API bool IsRawMeshEmpty() const;
-	ENGINE_API void LoadRawMesh(struct FRawMesh& OutRawMesh) const;
-	ENGINE_API void SaveRawMesh(struct FRawMesh& InRawMesh, bool bConvertToMeshdescription = true);
-
-#endif // #if WITH_EDITOR
-
-#if WITH_EDITORONLY_DATA
-	/**
-	 * Mesh description unpacked from bulk data.
-	 *
-	 * If this is valid, this means the mesh description has either been unpacked from the bulk data stored in the asset,
-	 * or one has been generated by the build tools (or converted from legacy RawMesh).
-	 */
-	UPROPERTY(Transient)
-	UStaticMeshDescription* StaticMeshDescription;
-
-	/**
-	 * Bulk data containing mesh description. LOD0 must be valid, but autogenerated lower LODs may be invalid.
-	 *
-	 * New assets store their source data here instead of in the RawMeshBulkData.
-	 * If this is invalid, either the LOD is autogenerated (for LOD1+), or the asset is a legacy asset whose
-	 * data is in the RawMeshBulkData.
-	 */
-	UPROPERTY()
-	UMeshDescriptionBulkDataWrapper* StaticMeshDescriptionBulkData;
-#endif
-
-#if WITH_EDITOR
-	/**
-	 * Create a new MeshDescription object
-	 */
-	ENGINE_API FMeshDescription* CreateMeshDescription();
-
-	/**
-	 * Tries to initialize a MeshDescription into the given FMeshDescription object, returning whether it was successful.
-	 * It will try the following steps in order to obtain it:
-	 * - if there's valid MeshDescriptionBulkData, unpack it from that
-	 * - if there's a valid DDC entry from a previously converted RawMesh, load it and unpack it
-	 * - if there's a valid RawMesh, convert it to MeshDescription and return that.
-	 * If all the above failed, then return failure.
-	 */
-	ENGINE_API bool LoadMeshDescription(FMeshDescription& OutMeshDescription) const;
-
-	/**
-	 * If there's a valid cached MeshDescription, return a copy of it.
-	 * Otherwise unpack the mesh description from bulk data and return the result (without caching it).
-	 */
-	ENGINE_API bool CloneMeshDescription(FMeshDescription& OutMeshDescription) const;
-
-	/**
-	 * Gets the cached mesh description if there was one, otherwise loads, caches and returns it.
-	 */
-	ENGINE_API FMeshDescription* GetOrCacheMeshDescription();
-
-	/**
-	 * Gets the cached mesh description, or nullptr if there wasn't one
-	 */
-	ENGINE_API FMeshDescription* GetCachedMeshDescription() const;
-
-	/**
-	 * Gets the cached StaticMeshDescription if there is one, or nullptr if not.
-	 */
-	UStaticMeshDescription* GetCachedStaticMeshDescription() const { return StaticMeshDescription; }
-
-	/**
-	 * Gets the mesh description bulk data, or nullptr if not valid 
-	 */
-	ENGINE_API FMeshDescriptionBulkData* GetMeshDescriptionBulkData() const;
-
-	/**
-	 * Returns whether the cached mesh description is valid or not
-	 */
-	ENGINE_API bool IsMeshDescriptionValid() const;
-
-	/**
-	 * Commits the cached mesh description to bulk data
-	 */
-	ENGINE_API void CommitMeshDescription(bool bUseHashAsGuid);
-
-	/**
-	 * Clears the cached mesh description
-	 */
-	ENGINE_API void ClearMeshDescription();
-
-#endif
-
-public:
-	/** Settings applied when building the mesh. */
-	UPROPERTY(EditAnywhere, Category=BuildSettings)
-	FMeshBuildSettings BuildSettings;
-
-	/** Reduction settings to apply when building render data. */
-	UPROPERTY(EditAnywhere, Category=ReductionSettings)
-	FMeshReductionSettings ReductionSettings; 
-
-	UPROPERTY()
-	float LODDistance_DEPRECATED;
-
-	/** 
-	 * ScreenSize to display this LOD.
-	 * The screen size is based around the projected diameter of the bounding
-	 * sphere of the model. i.e. 0.5 means half the screen's maximum dimension.
-	 */
-	UPROPERTY(EditAnywhere, Category=ReductionSettings)
-	FPerPlatformFloat ScreenSize;
-
-	/** The file path that was used to import this LOD. */
-	UPROPERTY(VisibleAnywhere, Category = StaticMeshSourceModel, AdvancedDisplay)
-	FString SourceImportFilename;
-
-#if WITH_EDITORONLY_DATA
-	/** Whether this LOD was imported in the same file as the base mesh. */
-	UPROPERTY()
-	bool bImportWithBaseMesh;
-#endif
-
-	/** Default constructor. */
-	ENGINE_API FStaticMeshSourceModel();
-
-	/** Destructor. */
-	ENGINE_API ~FStaticMeshSourceModel();
-
-#if WITH_EDITOR
-	/** Serializes bulk data. */
-	void SerializeBulkData(FArchive& Ar, UObject* Owner);
-
-#endif
-};
-
-// Make FStaticMeshSourceModel non-assignable
-template<> struct TStructOpsTypeTraits<FStaticMeshSourceModel> : public TStructOpsTypeTraitsBase2<FStaticMeshSourceModel> { enum { WithCopy = false }; };
-
-
 /**
  * Per-section settings.
  */
@@ -797,6 +634,7 @@ private:
 	TArray<FStaticMeshSourceModel> SourceModels;
 
 	/** Optional hi-res source data */
+	UE_DEPRECATED(5.00, "This must be protected for async build, always use the accessors even internally.")
 	UPROPERTY()
 	FStaticMeshSourceModel HiResSourceModel;
 
@@ -1440,17 +1278,17 @@ public:
 	/**
 	 * Performs a Modify on the StaticMeshDescription object pertaining to the given LODIndex
 	 */
-	ENGINE_API void ModifyMeshDescription(int32 LodIndex, bool bAlwaysMarkDirty = true);
+	ENGINE_API bool ModifyMeshDescription(int32 LodIndex, bool bAlwaysMarkDirty = true);
 
 	/**
 	 * Performs a Modify on StaticMeshDescription objects for all LODs
 	 */
-	ENGINE_API void ModifyAllMeshDescriptions(bool bAlwaysMarkDirty = true);
+	ENGINE_API bool ModifyAllMeshDescriptions(bool bAlwaysMarkDirty = true);
 
 	/**
 	 * Performs a Modify on the hi-res StaticMeshDescription
 	 */
-	ENGINE_API void ModifyHiResMeshDescription(bool bAlwaysMarkDirty = true);
+	ENGINE_API bool ModifyHiResMeshDescription(bool bAlwaysMarkDirty = true);
 
 	/**
 	 * Adds an empty UV channel at the end of the existing channels on the given LOD of a StaticMesh.
@@ -1583,14 +1421,18 @@ public:
 	ENGINE_API void SetNumSourceModels(int32 Num);
 
 	ENGINE_API void RemoveSourceModel(int32 Index);
-	ENGINE_API TArray<FStaticMeshSourceModel>& GetSourceModels();
 	ENGINE_API const TArray<FStaticMeshSourceModel>& GetSourceModels() const;
 	ENGINE_API FStaticMeshSourceModel& GetSourceModel(int32 Index);
 	ENGINE_API const FStaticMeshSourceModel& GetSourceModel(int32 Index) const;
 	ENGINE_API int32 GetNumSourceModels() const;
 	ENGINE_API bool IsSourceModelValid(int32 Index) const;
+	ENGINE_API TArray<FStaticMeshSourceModel>&& MoveSourceModels();
+	ENGINE_API void SetSourceModels(TArray<FStaticMeshSourceModel>&& SourceModels);
+
 	ENGINE_API FStaticMeshSourceModel& GetHiResSourceModel();
 	ENGINE_API const FStaticMeshSourceModel& GetHiResSourceModel() const;
+	ENGINE_API FStaticMeshSourceModel&& MoveHiResSourceModel();
+	ENGINE_API void SetHiResSourceModel(FStaticMeshSourceModel&& SourceModel);
 
 	ENGINE_API FMeshSectionInfoMap& GetSectionInfoMap();
 	ENGINE_API const FMeshSectionInfoMap& GetSectionInfoMap() const;
@@ -1923,9 +1765,9 @@ private:
 	void FixupZeroTriangleSections();
 
 	/**
-	* Caches mesh data.
+	* Converts legacy RawMesh to MeshDescription.
 	*/
-	void CacheMeshData();
+	void ConvertLegacySourceData();
 	
 	/**
 	 * Verify if the static mesh can be built.
diff --git a/Engine/Source/Runtime/Engine/Classes/Engine/StaticMeshSourceData.h b/Engine/Source/Runtime/Engine/Classes/Engine/StaticMeshSourceData.h
new file mode 100644
--- /dev/null
+++ b/Engine/Source/Runtime/Engine/Classes/Engine/StaticMeshSourceData.h
@@ -0,0 +1,204 @@
+// Copyright Epic Games, Inc. All Rights Reserved.
+
+#pragma once
+
+#include "CoreMinimal.h"
+#include "MeshDescriptionBaseBulkData.h"
+#include "Engine/EngineTypes.h"
+#include "Engine/MeshMerging.h"
+#include "UObject/ObjectMacros.h"
+
+#include "StaticMeshSourceData.generated.h"
+
+
+struct FMeshDescription;
+class UStaticMeshDescription;
+
+
+/**
+ * UObject wrapper for FMeshDescriptionBulkData
+ */
+UCLASS()
+class ENGINE_API UStaticMeshDescriptionBulkData : public UMeshDescriptionBaseBulkData
+{
+	GENERATED_BODY()
+
+public:
+	virtual UClass* GetMeshDescriptionType() const override;
+
+};
+
+
+/**
+ * Source model from which a renderable static mesh is built.
+ */
+USTRUCT()
+struct FStaticMeshSourceModel
+{
+	GENERATED_BODY()
+
+	friend class UStaticMesh;
+
+public:
+#if WITH_EDITOR
+	/**
+	 * Imported raw mesh data. Optional for all but the first LOD.
+	 *
+	 * This is a member for legacy assets only.
+	 * If it is non-empty, this means that it has been de-serialized from the asset, and
+	 * the asset hence pre-dates MeshDescription.
+	 */
+	class FRawMeshBulkData* RawMeshBulkData;
+	
+	/*
+	 * Accessor to Load and save the raw mesh or the mesh description depending on the editor settings.
+	 * Temporary until we deprecate the RawMesh.
+	 */
+	ENGINE_API bool IsRawMeshEmpty() const;
+	ENGINE_API void LoadRawMesh(struct FRawMesh& OutRawMesh) const;
+	ENGINE_API void SaveRawMesh(struct FRawMesh& InRawMesh, bool bConvertToMeshdescription = true);
+	
+	void FillMaterialName(TMap<int32, FName>& OutMaterialMap) const;
+
+
+#endif // #if WITH_EDITOR
+
+#if WITH_EDITORONLY_DATA
+	/*
+	 * The staticmesh owner of this source model. We need the SM to be able to convert between MeshDesription and RawMesh.
+	 * RawMesh use int32 material index and MeshDescription use FName material slot name.
+	 * This member is fill in the PostLoad of the static mesh.
+	 */
+	UPROPERTY(Transient)
+	class UStaticMesh* StaticMeshOwner;
+
+	/**
+	 * Bulk data containing mesh description.
+	 * New assets store their source data here instead of in the RawMeshBulkData.
+	 * If the bulk data within is empty, the LOD is autogenerated (for LOD1+).
+	 */
+	UPROPERTY()
+	UStaticMeshDescriptionBulkData* StaticMeshDescriptionBulkData;
+#endif
+
+#if WITH_EDITOR
+	/**
+	 * Create a new MeshDescription object
+	 */
+	ENGINE_API FMeshDescription* CreateMeshDescription();
+
+	/**
+	 * Tries to initialize a MeshDescription into the given FMeshDescription object, returning whether it was successful.
+	 * It will try the following steps in order to obtain it:
+	 * - if there's valid StaticMeshDescriptionBulkData, unpack it from that
+	 * - if there's a valid DDC entry from a previously converted RawMesh, load it and unpack it
+	 * - if there's a valid RawMesh, convert it to MeshDescription and return that.
+	 * If all the above failed, then return failure.
+	 */
+	ENGINE_API bool LoadMeshDescription(FMeshDescription& OutMeshDescription) const;
+
+	/**
+	 * If there's a valid cached MeshDescription, return a copy of it.
+	 * Otherwise unpack the mesh description from bulk data and return the result (without caching it).
+	 */
+	ENGINE_API bool CloneMeshDescription(FMeshDescription& OutMeshDescription) const;
+
+	/**
+	 * Gets the cached mesh description if there was one, otherwise loads, caches and returns it.
+	 */
+	ENGINE_API FMeshDescription* GetOrCacheMeshDescription();
+
+	/**
+	 * Gets the cached mesh description, or nullptr if there wasn't one
+	 */
+	ENGINE_API FMeshDescription* GetCachedMeshDescription() const;
+
+	/**
+	 * Gets the cached StaticMeshDescription if there is one, or nullptr if not.
+	 */
+	ENGINE_API UStaticMeshDescription* GetCachedStaticMeshDescription() const;
+
+	/**
+	 * Gets the mesh description bulk data, or nullptr if not valid 
+	 */
+	ENGINE_API FMeshDescriptionBulkData* GetMeshDescriptionBulkData() const;
+
+	/**
+	 * Returns whether the cached mesh description is valid or not
+	 */
+	ENGINE_API bool IsMeshDescriptionValid() const;
+
+	/**
+	 * Commits the cached mesh description to bulk data
+	 */
+	ENGINE_API void CommitMeshDescription(bool bUseHashAsGuid);
+
+	/**
+	 * Clears the cached mesh description
+	 */
+	ENGINE_API void ClearMeshDescription();
+
+#endif
+
+public:
+	/** Settings applied when building the mesh. */
+	UPROPERTY(EditAnywhere, Category=BuildSettings)
+	FMeshBuildSettings BuildSettings;
+
+	/** Reduction settings to apply when building render data. */
+	UPROPERTY(EditAnywhere, Category=ReductionSettings)
+	FMeshReductionSettings ReductionSettings; 
+
+	UPROPERTY()
+	float LODDistance_DEPRECATED;
+
+	/** 
+	 * ScreenSize to display this LOD.
+	 * The screen size is based around the projected diameter of the bounding
+	 * sphere of the model. i.e. 0.5 means half the screen's maximum dimension.
+	 */
+	UPROPERTY(EditAnywhere, Category=ReductionSettings)
+	FPerPlatformFloat ScreenSize;
+
+	/** The file path that was used to import this LOD. */
+	UPROPERTY(VisibleAnywhere, Category = StaticMeshSourceModel, AdvancedDisplay)
+	FString SourceImportFilename;
+
+#if WITH_EDITORONLY_DATA
+	/** Whether this LOD was imported in the same file as the base mesh. */
+	UPROPERTY()
+	bool bImportWithBaseMesh;
+#endif
+
+	/** Default constructor. */
+	ENGINE_API FStaticMeshSourceModel();
+
+	/** Destructor. */
+	ENGINE_API ~FStaticMeshSourceModel();
+
+	/** Copying not allowed */
+	FStaticMeshSourceModel(const FStaticMeshSourceModel&) = delete;
+	FStaticMeshSourceModel& operator=(const FStaticMeshSourceModel&) = delete;
+
+	/** Handle moving gracefully */
+	ENGINE_API FStaticMeshSourceModel(FStaticMeshSourceModel&& Other);
+	ENGINE_API FStaticMeshSourceModel& operator=(FStaticMeshSourceModel&& Other);
+
+	/** Initialize */
+	ENGINE_API void CreateSubObjects(UStaticMesh* Owner);
+
+#if WITH_EDITOR
+	/** Serializes bulk data. */
+	void SerializeBulkData(FArchive& Ar, UObject* Owner);
+
+	void ConvertRawMesh(int32 LodIndex);
+
+#endif
+};
+
+
+// Make FStaticMeshSourceModel non-assignable
+template<> struct TStructOpsTypeTraits<FStaticMeshSourceModel> : public TStructOpsTypeTraitsBase2<FStaticMeshSourceModel>
+{
+	enum { WithCopy = false };
+};
diff --git a/Engine/Source/Runtime/Engine/Private/StaticMesh.cpp b/Engine/Source/Runtime/Engine/Private/StaticMesh.cpp
--- a/Engine/Source/Runtime/Engine/Private/StaticMesh.cpp
+++ b/Engine/Source/Runtime/Engine/Private/StaticMesh.cpp
@@ -2466,19 +2466,14 @@ static FString BuildStaticMeshDerivedDataKeySuffix(const ITargetPlatform* Target
 	{
 		FStaticMeshSourceModel& SrcModel = Mesh->GetSourceModel(LODIndex);
 		
-		if (SrcModel.GetMeshDescriptionBulkData() != nullptr)
+		check(SrcModel.RawMeshBulkData->IsEmpty());
+		if (!SrcModel.GetMeshDescriptionBulkData()->IsEmpty())
 		{
-			KeySuffix += "MD";
 			KeySuffix += SrcModel.GetMeshDescriptionBulkData()->GetIdString();
 		}
-		else if (!SrcModel.RawMeshBulkData->IsEmpty())
-		{
-			// Legacy path for old assets
-			KeySuffix += SrcModel.RawMeshBulkData->GetIdString();
-		}
 		else
 		{
-			// If neither mesh description nor raw mesh bulk data are valid, this is a generated LOD
+			// If mesh description bulk data is empty, this is a generated LOD
 			KeySuffix += "_";
 		}
 
@@ -2504,6 +2499,16 @@ static FString BuildStaticMeshDerivedDataKeySuffix(const ITargetPlatform* Target
 		}
 	}
 
+	// Add hi-res mesh description into DDC key
+	if (!Mesh->GetHiResSourceModel().GetMeshDescriptionBulkData()->IsEmpty())
+	{
+		KeySuffix += Mesh->GetHiResSourceModel().GetMeshDescriptionBulkData()->GetIdString();
+	}
+	else
+	{
+		KeySuffix += "_";
+	}
+
 	// Mesh LOD streaming settings that need to trigger recache when changed
 	const bool bAllowLODStreaming = TargetPlatform->SupportsFeature(ETargetPlatformFeatures::MeshLODStreaming) && LODGroup.IsLODStreamingSupported();
 	KeySuffix += bAllowLODStreaming ? TEXT("LS1") : TEXT("LS0");
@@ -2883,6 +2888,8 @@ UStaticMesh::UStaticMesh(const FObjectInitializer& ObjectInitializer)
 	ImportVersion = EImportStaticMeshVersion::BeforeImportStaticMeshVersionWasAdded;
 	LODForOccluderMesh = -1;
 	NumStreamedLODs.Default = -1;
+	GetHiResSourceModel().StaticMeshOwner = this;
+	GetHiResSourceModel().StaticMeshDescriptionBulkData = CreateEditorOnlyDefaultSubobject<UStaticMeshDescriptionBulkData>(TEXT("HiResMeshDescription"));
 #endif // #if WITH_EDITORONLY_DATA
 	SetLightMapResolution(4);
 	SetMinLOD(0);
@@ -3483,14 +3490,6 @@ void UStaticMesh::PostEditChangeProperty(FPropertyChangedEvent& PropertyChangedE
 	//Don't use the render data here because the property that just changed might be invalidating the current RenderData.
 	EnforceLightmapRestrictions(/*bUseRenderData=*/false);
 
-	// Following an undo or other operation which can change the SourceModels, ensure the StaticMeshOwner is up to date
-	for (int32 Index = 0; Index < GetNumSourceModels(); ++Index)
-	{
-		GetSourceModel(Index).StaticMeshOwner = this;
-	}
-
-	GetHiResSourceModel().StaticMeshOwner = this;
-
 	Build(/*bSilent=*/ true);
 
 	if (PropertyName == GET_MEMBER_NAME_CHECKED(UStaticMesh, bHasNavigationData)
@@ -3529,14 +3528,6 @@ void UStaticMesh::PostEditChangeProperty(FPropertyChangedEvent& PropertyChangedE
 
 void UStaticMesh::PostEditUndo()
 {
-	// Following an undo or other operation which can change the SourceModels, ensure the StaticMeshOwner is up to date
-	for (int32 Index = 0; Index < GetNumSourceModels(); ++Index)
-	{
-		GetSourceModel(Index).StaticMeshOwner = this;
-	}
-
-	GetHiResSourceModel().StaticMeshOwner = this;
-
 	// The super will cause a Build() via PostEditChangeProperty().
 	Super::PostEditUndo();
 }
@@ -3674,37 +3665,42 @@ const FStaticMeshSourceModel& UStaticMesh::GetSourceModel(int32 Index) const
 	return GetSourceModels()[Index];
 }
 
-TArray<FStaticMeshSourceModel>& UStaticMesh::GetSourceModels()
+FStaticMeshSourceModel& UStaticMesh::GetSourceModel(int32 Index)
 { 
 	WaitUntilAsyncPropertyReleased(EStaticMeshAsyncProperties::SourceModels);
 	PRAGMA_DISABLE_DEPRECATION_WARNINGS
-	return SourceModels; 
+	return SourceModels[Index];
 	PRAGMA_ENABLE_DEPRECATION_WARNINGS
 }
 
-FStaticMeshSourceModel& UStaticMesh::GetSourceModel(int32 Index)
-{ 
-	return GetSourceModels()[Index];
-}
-
 FStaticMeshSourceModel& UStaticMesh::AddSourceModel()
 {
-	int32 LodModelIndex = GetSourceModels().AddDefaulted();
-	FStaticMeshSourceModel& NewSourceModel = GetSourceModel(LodModelIndex);
-	NewSourceModel.StaticMeshOwner = this;
+	WaitUntilAsyncPropertyReleased(EStaticMeshAsyncProperties::SourceModels);
+	PRAGMA_DISABLE_DEPRECATION_WARNINGS
+	int32 LodModelIndex = SourceModels.AddDefaulted();
+	FStaticMeshSourceModel& NewSourceModel = SourceModels[LodModelIndex];
+	NewSourceModel.CreateSubObjects(this);
 	return NewSourceModel;
+	PRAGMA_ENABLE_DEPRECATION_WARNINGS
 }
 
 void UStaticMesh::SetNumSourceModels(const int32 Num)
 {
-	const int32 OldNum = GetNumSourceModels();
-	GetSourceModels().SetNum(Num);
+	WaitUntilAsyncPropertyReleased(EStaticMeshAsyncProperties::SourceModels);
+	PRAGMA_DISABLE_DEPRECATION_WARNINGS
+	const int32 OldNum = SourceModels.Num();
 
 	//Shrink the SectionInfoMap if some SourceModel are removed
 	if (OldNum > Num)
 	{
 		for (int32 RemoveLODIndex = Num; RemoveLODIndex < OldNum; ++RemoveLODIndex)
 		{
+			// Remove MeshDescription allocations
+			FStaticMeshSourceModel& ThisSourceModel = SourceModels[RemoveLODIndex];
+			ThisSourceModel.ClearMeshDescription();
+			check(ThisSourceModel.GetMeshDescriptionBulkData());
+			ThisSourceModel.GetMeshDescriptionBulkData()->Empty();
+
 			int32 SectionCount = GetSectionInfoMap().GetSectionNumber(RemoveLODIndex);
 			for (int32 SectionIndex = 0; SectionIndex < SectionCount; ++SectionIndex)
 			{
@@ -3718,16 +3714,18 @@ void UStaticMesh::SetNumSourceModels(const int32 Num)
 		}
 	}
 
+	SourceModels.SetNum(Num);
+
 	for (int32 Index = OldNum; Index < Num; ++Index)
 	{
-		FStaticMeshSourceModel& ThisSourceModel = GetSourceModel(Index);
+		FStaticMeshSourceModel& ThisSourceModel = SourceModels[Index];
 
-		ThisSourceModel.StaticMeshOwner = this;
+		ThisSourceModel.CreateSubObjects(this);
 		int32 PreviousCustomLODIndex = 0;
 		//Find the previous custom LOD
 		for (int32 ReverseIndex = Index - 1; ReverseIndex > 0; ReverseIndex--)
 		{
-			const FStaticMeshSourceModel& StaticMeshModel = GetSourceModel(ReverseIndex);
+			const FStaticMeshSourceModel& StaticMeshModel = SourceModels[ReverseIndex];
 			//If the custom import LOD is reduce and is not using himself as the source, do not consider it
 			if (IsMeshDescriptionValid(ReverseIndex) && !(IsReductionActive(ReverseIndex) && StaticMeshModel.ReductionSettings.BaseLODModel != ReverseIndex))
 			{
@@ -3742,11 +3740,20 @@ void UStaticMesh::SetNumSourceModels(const int32 Num)
 			ThisSourceModel.ReductionSettings.PercentTriangles = FMath::Pow(0.5f, (float)(Index-PreviousCustomLODIndex));
 		}
 	}
+	PRAGMA_ENABLE_DEPRECATION_WARNINGS
 }
 
 void UStaticMesh::RemoveSourceModel(const int32 Index)
 {
-	check(IsSourceModelValid(Index));
+	WaitUntilAsyncPropertyReleased(EStaticMeshAsyncProperties::SourceModels);
+	PRAGMA_DISABLE_DEPRECATION_WARNINGS
+	check(SourceModels.IsValidIndex(Index));
+
+	// Remove MeshDescription allocations
+	FStaticMeshSourceModel& ThisSourceModel = SourceModels[Index];
+	ThisSourceModel.ClearMeshDescription();
+	check(ThisSourceModel.GetMeshDescriptionBulkData());
+	ThisSourceModel.GetMeshDescriptionBulkData()->Empty();
 
 	//Remove the SectionInfoMap of the LOD we remove
 	{
@@ -3763,9 +3770,9 @@ void UStaticMesh::RemoveSourceModel(const int32 Index)
 	}
 
 	//Move down all SectionInfoMap for the next LOD
-	if (Index < GetNumSourceModels() - 1)
+	if (Index < SourceModels.Num() - 1)
 	{
-		for (int32 MoveIndex = Index + 1; MoveIndex < GetNumSourceModels(); ++MoveIndex)
+		for (int32 MoveIndex = Index + 1; MoveIndex < SourceModels.Num(); ++MoveIndex)
 		{
 			int32 SectionCount = GetSectionInfoMap().GetSectionNumber(MoveIndex);
 			for (int32 SectionIndex = 0; SectionIndex < SectionCount; ++SectionIndex)
@@ -3785,17 +3792,71 @@ void UStaticMesh::RemoveSourceModel(const int32 Index)
 	}
 
 	//Remove the LOD
-	GetSourceModels().RemoveAt(Index);
+	SourceModels.RemoveAt(Index);
+	PRAGMA_ENABLE_DEPRECATION_WARNINGS
+}
+
+TArray<FStaticMeshSourceModel>&& UStaticMesh::MoveSourceModels()
+{
+	WaitUntilAsyncPropertyReleased(EStaticMeshAsyncProperties::SourceModels);
+	PRAGMA_DISABLE_DEPRECATION_WARNINGS
+	return MoveTemp(SourceModels);
+	PRAGMA_ENABLE_DEPRECATION_WARNINGS
 }
 
+void UStaticMesh::SetSourceModels(TArray<FStaticMeshSourceModel>&& InSourceModels)
+{
+	WaitUntilAsyncPropertyReleased(EStaticMeshAsyncProperties::SourceModels);
+	PRAGMA_DISABLE_DEPRECATION_WARNINGS
+	SourceModels = MoveTemp(InSourceModels);
+
+	for (FStaticMeshSourceModel& SourceModel : SourceModels)
+	{
+		SourceModel.StaticMeshOwner = this;
+		if (SourceModel.StaticMeshDescriptionBulkData)
+		{
+			SourceModel.StaticMeshDescriptionBulkData->Rename(nullptr, this, REN_DontCreateRedirectors | REN_NonTransactional);
+		}
+	}
+	PRAGMA_ENABLE_DEPRECATION_WARNINGS
+}
+
+
 const FStaticMeshSourceModel& UStaticMesh::GetHiResSourceModel() const
 {
+	WaitUntilAsyncPropertyReleased(EStaticMeshAsyncProperties::HiResSourceModel);
+	PRAGMA_DISABLE_DEPRECATION_WARNINGS
 	return HiResSourceModel;
+	PRAGMA_ENABLE_DEPRECATION_WARNINGS
 }
 
 FStaticMeshSourceModel& UStaticMesh::GetHiResSourceModel()
 {
+	WaitUntilAsyncPropertyReleased(EStaticMeshAsyncProperties::HiResSourceModel);
+	PRAGMA_DISABLE_DEPRECATION_WARNINGS
 	return HiResSourceModel;
+	PRAGMA_ENABLE_DEPRECATION_WARNINGS
+}
+
+FStaticMeshSourceModel&& UStaticMesh::MoveHiResSourceModel()
+{
+	WaitUntilAsyncPropertyReleased(EStaticMeshAsyncProperties::HiResSourceModel);
+	PRAGMA_DISABLE_DEPRECATION_WARNINGS
+	return MoveTemp(HiResSourceModel);
+	PRAGMA_ENABLE_DEPRECATION_WARNINGS
+}
+
+void UStaticMesh::SetHiResSourceModel(FStaticMeshSourceModel&& InSourceModel)
+{
+	WaitUntilAsyncPropertyReleased(EStaticMeshAsyncProperties::HiResSourceModel);
+	PRAGMA_DISABLE_DEPRECATION_WARNINGS
+	HiResSourceModel = MoveTemp(InSourceModel);
+	HiResSourceModel.StaticMeshOwner = this;
+	if (HiResSourceModel.StaticMeshDescriptionBulkData)
+	{
+		HiResSourceModel.StaticMeshDescriptionBulkData->Rename(nullptr, this, REN_DontCreateRedirectors | REN_NonTransactional);
+	}
+	PRAGMA_ENABLE_DEPRECATION_WARNINGS
 }
 
 
@@ -4063,304 +4124,6 @@ bool UStaticMesh::IsCachedCookedPlatformDataLoaded(const ITargetPlatform* Target
 	FStaticMeshSourceModel
 ------------------------------------------------------------------------------*/
 
-FStaticMeshSourceModel::FStaticMeshSourceModel()
-{
-	LODDistance_DEPRECATED = 0.0f;
-#if WITH_EDITOR
-	RawMeshBulkData = new FRawMeshBulkData();
-	ScreenSize.Default = 0.0f;
-	StaticMeshOwner = nullptr;
-#endif // #if WITH_EDITOR
-	SourceImportFilename = FString();
-#if WITH_EDITORONLY_DATA
-	bImportWithBaseMesh = false;
-	StaticMeshDescription = nullptr;
-	StaticMeshDescriptionBulkData = nullptr;
-#endif
-}
-
-FStaticMeshSourceModel::~FStaticMeshSourceModel()
-{
-#if WITH_EDITOR
-	if (RawMeshBulkData)
-	{
-		delete RawMeshBulkData;
-	}
-#endif // #if WITH_EDITOR
-}
-
-
-#if WITH_EDITOR
-bool FStaticMeshSourceModel::IsRawMeshEmpty() const
-{
-	// Even if the RawMeshBulkData is empty, it may be because it's a new asset serialized as MeshDescription.
-	// Hence MeshDescriptionBulkData must also be invalid (and, by consequence, also MeshDescription).
-	return RawMeshBulkData->IsEmpty() && StaticMeshDescriptionBulkData == nullptr && StaticMeshDescription == nullptr;
-}
-
-void FStaticMeshSourceModel::LoadRawMesh(FRawMesh& OutRawMesh) const
-{
-	if (RawMeshBulkData->IsEmpty())
-	{
-		// If the RawMesh is empty, consider the possibility that it's a new asset with a valid MeshDescription which needs loading.
-		// We require the FStaticMeshSourceModel to be in the UStaticMesh::SourceModels array, so that we can infer which LOD it
-		// corresponds to. This would normally be unreasonably limiting, but since these methods are deprecated, we'll go with it.
-		check(StaticMeshOwner != nullptr);
-		const int32 LODIndex = this - &StaticMeshOwner->GetSourceModel(0);
-		check(LODIndex < StaticMeshOwner->GetNumSourceModels());
-		if (FMeshDescription* CachedMeshDescription = StaticMeshOwner->GetMeshDescription(LODIndex))
-		{
-			TMap<FName, int32> MaterialMap;
-			check(StaticMeshOwner != nullptr);
-			const UStaticMesh* ConstStaticMeshOwner = StaticMeshOwner;
-			for (int32 MaterialIndex = 0; MaterialIndex < ConstStaticMeshOwner->GetStaticMaterials().Num(); ++MaterialIndex)
-			{
-				MaterialMap.Add(ConstStaticMeshOwner->GetStaticMaterials()[MaterialIndex].ImportedMaterialSlotName, MaterialIndex);
-			}
-			FStaticMeshOperations::ConvertToRawMesh(*CachedMeshDescription, OutRawMesh, MaterialMap);
-		}
-	}
-	else
-	{
-		RawMeshBulkData->LoadRawMesh(OutRawMesh);
-	}
-}
-
-void FStaticMeshSourceModel::SaveRawMesh(FRawMesh& InRawMesh, bool /* unused */)
-{
-	if (!InRawMesh.IsValid())
-	{
-		return;
-	}
-
-	TRACE_CPUPROFILER_EVENT_SCOPE(FStaticMeshSourceModel::SaveRawMesh);
-
-	//Save both format
-	RawMeshBulkData->SaveRawMesh(InRawMesh);
-
-	FMeshDescription* MeshDescription = CreateMeshDescription();
-
-	TMap<int32, FName> MaterialMap;
-	check(StaticMeshOwner != nullptr);
-	const UStaticMesh* ConstStaticMeshOwner = StaticMeshOwner;
-	FillMaterialName(ConstStaticMeshOwner->GetStaticMaterials(), MaterialMap);
-	FStaticMeshOperations::ConvertFromRawMesh(InRawMesh, *MeshDescription, MaterialMap);
-	
-	// Package up mesh description into bulk data
-	if (StaticMeshDescriptionBulkData == nullptr)
-	{
-		StaticMeshDescriptionBulkData = NewObject<UMeshDescriptionBulkDataWrapper>(StaticMeshOwner, NAME_None, RF_Transactional);
-	}
-
-	StaticMeshDescriptionBulkData->GetBulkData().SaveMeshDescription(*MeshDescription);
-	StaticMeshDescriptionBulkData->AtomicallyClearInternalFlags(EInternalObjectFlags::Async);
-
-	// Converted to MeshDescription - can already lose the RawMeshBulkData
-	RawMeshBulkData->Empty();
-}
-
-
-bool FStaticMeshSourceModel::LoadMeshDescription(FMeshDescription& OutMeshDescription) const
-{
-	TRACE_CPUPROFILER_EVENT_SCOPE(FStaticMeshSourceModel::LoadMeshDescription);
-	check(StaticMeshOwner != nullptr);
-
-	// Ensure default MeshDescription result is empty, with no attributes registered
-	OutMeshDescription = FMeshDescription();
-
-	// If we have valid MeshDescriptionBulkData, unpack it and return it...
-	if (StaticMeshDescriptionBulkData != nullptr)
-	{
-		// Unpack MeshDescription from the bulk data which was deserialized
-		StaticMeshDescriptionBulkData->GetBulkData().LoadMeshDescription(OutMeshDescription);
-		return true;
-	}
-
-	// RawMeshBulkData should always be empty now (soon to be deprecated)
-	check(RawMeshBulkData->IsEmpty());
-
-	// This LOD doesn't have a MeshDescriptionBulkData, or a RawMesh, so we presume that it's a generated LOD (and return false)
-	return false;
-}
-
-bool FStaticMeshSourceModel::CloneMeshDescription(FMeshDescription& OutMeshDescription) const
-{
-	TRACE_CPUPROFILER_EVENT_SCOPE(FStaticMeshSourceModel::CloneMeshDescription);
-
-	if (StaticMeshDescription != nullptr)
-	{
-		OutMeshDescription = StaticMeshDescription->GetMeshDescription();
-		return true;
-	}
-
-	return LoadMeshDescription(OutMeshDescription);
-}
-
-FMeshDescription* FStaticMeshSourceModel::GetOrCacheMeshDescription()
-{
-	TRACE_CPUPROFILER_EVENT_SCOPE(FStaticMeshSourceModel::GetMeshDescription);
-	check(StaticMeshOwner != nullptr);
-
-	if (StaticMeshDescription == nullptr)
-	{
-		FMeshDescription MeshDescription;
-		if (LoadMeshDescription(MeshDescription))
-		{
-			// This will sometimes be run on a different thread.
-			// However once we have assigned the new UObject to the StaticMeshDescription UPROPERTY, we are free to clear any Async flag which has been assigned
-			// as GC will work as intended.
-			StaticMeshDescription = NewObject<UStaticMeshDescription>(StaticMeshOwner, NAME_None, RF_Transient | RF_Transactional);
-			StaticMeshDescription->SetMeshDescription(MoveTemp(MeshDescription));
-			StaticMeshDescription->AtomicallyClearInternalFlags(EInternalObjectFlags::Async);
-		}
-	}
-
-	if (StaticMeshDescription != nullptr)
-	{
-		return &StaticMeshDescription->GetMeshDescription();
-	}
-	else
-	{
-		return nullptr;
-	}
-}
-
-
-FMeshDescription* FStaticMeshSourceModel::GetCachedMeshDescription() const
-{
-	if (StaticMeshDescription)
-	{
-		return &StaticMeshDescription->GetMeshDescription();
-	}
-
-	return nullptr;
-}
-
-
-FMeshDescriptionBulkData* FStaticMeshSourceModel::GetMeshDescriptionBulkData() const
-{
-	if (StaticMeshDescriptionBulkData)
-	{
-		return &StaticMeshDescriptionBulkData->GetBulkData();
-	}
-
-	return nullptr;
-}
-
-
-bool FStaticMeshSourceModel::IsMeshDescriptionValid() const
-{
-	// Determine whether a mesh description is valid without requiring it to be loaded first.
-	// If there is a valid MeshDescriptionBulkData, we know this implies a valid mesh description.
-	// If not, then it's a legacy asset which will have a mesh description built from the RawMeshBulkData, if non-empty.
-	return StaticMeshDescription != nullptr ||
-		StaticMeshDescriptionBulkData != nullptr ||
-		!RawMeshBulkData->IsEmpty();
-}
-
-
-FMeshDescription* FStaticMeshSourceModel::CreateMeshDescription()
-{
-	if (StaticMeshDescription == nullptr)
-	{
-		// Assign the new UObject directly to a UPROPERTY.
-		// If we were run from a different thread, we can immediately clear the Async internal flag, now that GC refs are set up correctly.
-		check(StaticMeshOwner != nullptr);
-		StaticMeshDescription = NewObject<UStaticMeshDescription>(StaticMeshOwner, NAME_None, RF_Transient | RF_Transactional);
-		StaticMeshDescription->AtomicallyClearInternalFlags(EInternalObjectFlags::Async);
-	}
-
-	// Reset the mesh description and register its attributes
-	// Do this instead of always creating a new UStaticMeshDescription object, to save memory and improve performance
-	StaticMeshDescription->Reset();
-
-	return &StaticMeshDescription->GetMeshDescription();
-}
-
-
-void FStaticMeshSourceModel::CommitMeshDescription(bool bUseHashAsGuid)
-{
-	TRACE_CPUPROFILER_EVENT_SCOPE(FStaticMeshSourceModel::CommitMeshDescription);
-
-	check(RawMeshBulkData->IsEmpty());
-
-	// This part should remain thread-safe so it can be called from any thread
-	// as long as no more than one thread is calling it for the same UStaticMesh.
-
-	if (StaticMeshDescription != nullptr)
-	{
-		// Package up mesh description into bulk data
-		if (StaticMeshDescriptionBulkData == nullptr)
-		{
-			check(StaticMeshOwner != nullptr);
-			StaticMeshDescriptionBulkData = NewObject<UMeshDescriptionBulkDataWrapper>(StaticMeshOwner, NAME_None, RF_Transactional);
-		}
-
-		StaticMeshDescriptionBulkData->GetBulkData().SaveMeshDescription(StaticMeshDescription->GetMeshDescription());
-		if (bUseHashAsGuid)
-		{
-			StaticMeshDescriptionBulkData->GetBulkData().UseHashAsGuid();
-		}
-
-		// Clear the async flag when the object is fully initialized and assigned to a UPROPERTY
-		StaticMeshDescriptionBulkData->AtomicallyClearInternalFlags(EInternalObjectFlags::Async);
-	}
-	else
-	{
-		FGCScopeGuard GCScopeGuard;
-		StaticMeshDescriptionBulkData->MarkPendingKill();
-		StaticMeshDescriptionBulkData = nullptr;
-	}
-}
-
-void FStaticMeshSourceModel::ClearMeshDescription()
-{
-	TRACE_CPUPROFILER_EVENT_SCOPE(UStaticMesh::ClearMeshDescription);
-
-	if (StaticMeshDescription != nullptr)
-	{
-		FGCScopeGuard GCScopeGuard;
-		StaticMeshDescription->MarkPendingKill();
-		StaticMeshDescription = nullptr;
-	}
-}
-
-
-void FStaticMeshSourceModel::SerializeBulkData(FArchive& Ar, UObject* Owner)
-{
-	// Initialize the StaticMeshOwner
-	if (Ar.IsLoading())
-	{
-		StaticMeshOwner = Cast<UStaticMesh>(Owner);
-	}
-
-	if (Ar.IsLoading() && Ar.CustomVer(FEditorObjectVersion::GUID) < FEditorObjectVersion::StaticMeshDeprecatedRawMesh)
-	{
-		// If loading a legacy asset with RawMesh bulk data, serialize it here
-		check(RawMeshBulkData != NULL);
-		RawMeshBulkData->Serialize(Ar, Owner);
-	}
-	else
-	{
-		if (Ar.IsLoading() && Ar.CustomVer(FUE5MainStreamObjectVersion::GUID) < FUE5MainStreamObjectVersion::SerializeMeshDescriptionBase)
-		{
-			// If there's an inline mesh description bulk data (legacy version), serialize it here, and copy it into the new separate inner object
-			bool bIsValid;
-			Ar << bIsValid;
-
-			if (bIsValid)
-			{
-				StaticMeshDescriptionBulkData = NewObject<UMeshDescriptionBulkDataWrapper>(StaticMeshOwner, NAME_None, RF_Transactional);
-				StaticMeshDescriptionBulkData->GetBulkData().Serialize(Ar, StaticMeshOwner);
-				StaticMeshDescriptionBulkData->AtomicallyClearInternalFlags(EInternalObjectFlags::Async);
-
-				// Any cached mesh description is now invalid, so clear it out
-				ClearMeshDescription();
-			}
-		}
-	}
-}
-#endif // #if WITH_EDITOR
 
 /*------------------------------------------------------------------------------
 	FMeshSectionInfoMap
@@ -4715,74 +4478,38 @@ void UStaticMesh::ClearHiResMeshDescription()
 }
 
 
-void UStaticMesh::ModifyMeshDescription(int32 LodIndex, bool bAlwaysMarkDirty)
+bool UStaticMesh::ModifyMeshDescription(int32 LodIndex, bool bAlwaysMarkDirty)
 {
-	UStaticMeshDescription* StaticMeshDescription = GetSourceModel(LodIndex).StaticMeshDescription;
-	if (StaticMeshDescription)
-	{
-		StaticMeshDescription->Modify(bAlwaysMarkDirty);
-	}
+	FStaticMeshSourceModel& SourceModel = GetSourceModel(LodIndex);
+	check(SourceModel.StaticMeshDescriptionBulkData != nullptr);
+	return SourceModel.StaticMeshDescriptionBulkData->Modify(bAlwaysMarkDirty);
 }
 
 
-void UStaticMesh::ModifyAllMeshDescriptions(bool bAlwaysMarkDirty)
+bool UStaticMesh::ModifyAllMeshDescriptions(bool bAlwaysMarkDirty)
 {
+	bool bResult = true;
 	for (int LODIndex = 0; LODIndex < GetNumSourceModels(); LODIndex++)
 	{
-		ModifyMeshDescription(LODIndex);
+		if (!ModifyMeshDescription(LODIndex))
+		{
+			bResult = false;
+		}
 	}
-}
 
-
-void UStaticMesh::ModifyHiResMeshDescription(bool bAlwaysMarkDirty)
-{
-	UStaticMeshDescription* StaticMeshDescription = GetHiResSourceModel().StaticMeshDescription;
-	if (StaticMeshDescription)
-	{
-		StaticMeshDescription->Modify(bAlwaysMarkDirty);
-	}
+	return bResult;
 }
 
 
-// This is the key which is used for DDC data when legacy RawMesh is converted to MeshDescription.
-// If static mesh derived data needs to be rebuilt (new format, serialization differences, etc.) replace the version GUID below with a new one.
-// In case of merge conflicts with DDC versions, you MUST generate a new GUID and set this new GUID as the version.
-#if UE_USE_VIRTUALBULKDATA
-#define MESHDATAKEY_STATICMESH_DERIVEDDATA_VER TEXT("7F1A02EF80374F0F86F9A30BAFF3F850")
-#else
-#define MESHDATAKEY_STATICMESH_DERIVEDDATA_VER TEXT("E09E8DD68C864D82B511A97080B04837")
-#endif
-
-static bool GetMeshDataKey(FString& OutKey, const FRawMeshBulkData* RawMeshBulkData, int LodIndex)
+bool UStaticMesh::ModifyHiResMeshDescription(bool bAlwaysMarkDirty)
 {
-	OutKey.Empty();
-
-	FSHA1 Sha;
-
-	check(!RawMeshBulkData->IsEmpty());
-	FString LodIndexString = FString::Printf(TEXT("%d_"), LodIndex);
-	LodIndexString += RawMeshBulkData->GetIdString();
-
-	const TArray<TCHAR>& LodIndexArray = LodIndexString.GetCharArray();
-	Sha.Update((uint8*)LodIndexArray.GetData(), LodIndexArray.Num() * LodIndexArray.GetTypeSize());
-	Sha.Final();
-
-	// Retrieve the hash and use it to construct a pseudo-GUID.
-	uint32 Hash[5];
-	Sha.GetHash((uint8*)Hash);
-	FGuid Guid = FGuid(Hash[0] ^ Hash[4], Hash[1], Hash[2], Hash[3]);
-	FString MeshLodData = Guid.ToString(EGuidFormats::Digits);
-
-	OutKey = FDerivedDataCacheInterface::BuildCacheKey(
-		TEXT("MESHDATAKEY_STATICMESH"),
-		MESHDATAKEY_STATICMESH_DERIVEDDATA_VER,
-		*MeshLodData
-	);
-	return true;
+	FStaticMeshSourceModel& SourceModel = GetHiResSourceModel();
+	check(SourceModel.StaticMeshDescriptionBulkData != nullptr);
+	return SourceModel.StaticMeshDescriptionBulkData->Modify(bAlwaysMarkDirty);
 }
 
 
-void UStaticMesh::CacheMeshData()
+void UStaticMesh::ConvertLegacySourceData()
 {
 	TRACE_CPUPROFILER_EVENT_SCOPE(UStaticMesh::CacheMeshData);
 
@@ -4790,65 +4517,10 @@ void UStaticMesh::CacheMeshData()
 	for (int32 LodIndex = 0; LodIndex < GetNumSourceModels(); ++LodIndex)
 	{
 		FStaticMeshSourceModel& SourceModel = GetSourceModel(LodIndex);
-		if (SourceModel.GetMeshDescriptionBulkData() == nullptr)
-		{
-			// Legacy assets used to store their source data in the RawMeshBulkData
-			// Migrate it to the new description if present
-			if (!SourceModel.RawMeshBulkData->IsEmpty())
-			{
-				FString MeshDataKey;
-				if (GetMeshDataKey(MeshDataKey, SourceModel.RawMeshBulkData, LodIndex))
-				{
-					COOK_STAT(auto Timer = StaticMeshCookStats::UsageStats.TimeSyncWork());
-
-					TArray<uint8> DerivedData;
-
-					if (GetDerivedDataCacheRef().GetSynchronous(*MeshDataKey, DerivedData, GetPathName()))
-					{
-						COOK_STAT(Timer.AddHit(DerivedData.Num()));
-
-						SourceModel.StaticMeshDescriptionBulkData = NewObject<UMeshDescriptionBulkDataWrapper>(this, NAME_None, RF_Transactional);
 
-						// Load from the DDC
-						const bool bIsPersistent = true;
-						FMemoryReader Ar(DerivedData, bIsPersistent);
-						SourceModel.StaticMeshDescriptionBulkData->GetBulkData().Serialize(Ar, const_cast<UStaticMesh*>(this));
-
-						check(SourceModel.GetCachedMeshDescription() == nullptr);
-					}
-					else
-					{
-						// If the DDC key doesn't exist, convert the data and save it to DDC
-						// Get the RawMesh for this LOD
-						FRawMesh TempRawMesh;
-						SourceModel.RawMeshBulkData->LoadRawMesh(TempRawMesh);
-
-						// Create a new MeshDescription
-						FMeshDescription* MeshDescription = SourceModel.CreateMeshDescription();
-
-						// Convert the RawMesh to MeshDescription
-						TMap<int32, FName> MaterialMap;
-						FillMaterialName(GetStaticMaterials(), MaterialMap);
-						FStaticMeshOperations::ConvertFromRawMesh(TempRawMesh, *MeshDescription, MaterialMap);
-
-						// Pack MeshDescription into bulk data
-						SourceModel.StaticMeshDescriptionBulkData = NewObject<UMeshDescriptionBulkDataWrapper>(this, NAME_None, RF_Transactional);
-						SourceModel.StaticMeshDescriptionBulkData->GetBulkData().SaveMeshDescription(*MeshDescription);
-						SourceModel.StaticMeshDescriptionBulkData->AtomicallyClearInternalFlags(EInternalObjectFlags::Async);
-
-						// Write the DDC cache
-						const bool bIsPersistent = true;
-						FMemoryWriter Ar(DerivedData, bIsPersistent);
-						SourceModel.StaticMeshDescriptionBulkData->GetBulkData().Serialize(Ar, this);
-						GetDerivedDataCacheRef().Put(*MeshDataKey, DerivedData, GetPathName());
-						COOK_STAT(Timer.AddMiss(DerivedData.Num()));
-					}
-				}
-			}
-
-			// We now have a MeshDescription instead of a RawMesh, so get rid of the RawMesh completely
-			SourceModel.RawMeshBulkData->Empty();
-		}
+		// Legacy assets used to store their source data in the RawMeshBulkData
+		// Migrate it to the new description if present
+		SourceModel.ConvertRawMesh(LodIndex);
 	}
 }
 
@@ -5013,6 +4685,9 @@ void UStaticMesh::CacheDerivedData()
 	TRACE_CPUPROFILER_EVENT_SCOPE(UStaticMesh::CacheDerivedData);
 	LLM_SCOPE_BYNAME(TEXT("AssetCompilation/StaticMesh"));
 
+#if WITH_EDITORONLY_DATA
+	ConvertLegacySourceData();
+#endif
 	// Cache derived data for the running platform.
 	ITargetPlatformManagerModule& TargetPlatformManager = GetTargetPlatformManagerRef();
 	ITargetPlatform* RunningPlatform = TargetPlatformManager.GetRunningTargetPlatform();
@@ -5249,10 +4924,13 @@ void UStaticMesh::Serialize(FArchive& Ar)
 			SrcModel.SerializeBulkData(Ar, this);
 		}
 
-		// If loading, perform immediate conversion from legacy RawMesh to MeshDescription if necessary
 		if (Ar.IsLoading())
 		{
-			CacheMeshData();
+			GetHiResSourceModel().StaticMeshOwner = this;
+			if (GetHiResSourceModel().StaticMeshDescriptionBulkData == nullptr)
+			{
+				GetHiResSourceModel().StaticMeshDescriptionBulkData = NewObject<UStaticMeshDescriptionBulkData>(this, NAME_None, RF_Transactional);
+			}
 		}
 
 		if (Ar.CustomVer(FEditorObjectVersion::GUID) < FEditorObjectVersion::UPropertryForMeshSection)
@@ -5514,16 +5192,6 @@ void UStaticMesh::BeginPostLoadInternal(FStaticMeshPostLoadContext& Context)
 		{
 			DistanceFieldReplacementMesh->ConditionalPostLoad();
 		}
-
-		//Fill the static mesh owner
-		int32 NumLODs = GetNumSourceModels();
-		for (int32 LODIndex = 0; LODIndex < NumLODs; ++LODIndex)
-		{
-			FStaticMeshSourceModel& SrcModel = GetSourceModel(LODIndex);
-			SrcModel.StaticMeshOwner = this;
-		}
-
-		GetHiResSourceModel().StaticMeshOwner = this;
 	}
 
 	Context.bIsCookedForEditor = GetOutermost()->bIsCookedForEditor;
@@ -5563,10 +5231,15 @@ void UStaticMesh::BeginPostLoadInternal(FStaticMeshPostLoadContext& Context)
 
 			for (int32 i = 0; i < GetNumSourceModels(); ++i)
 			{
-				FMeshDescription MeshDescription;
-				if (GetSourceModel(i).CloneMeshDescription(MeshDescription))
+				// At this stage in loading, we have not yet converted a legacy asset's RawMesh to MeshDescription,
+				// so access RawMesh directly instead of through the FStaticMeshSourceModel API,
+				// because we don't want to perform an automatic conversion to MeshDescription at this point -
+				// this will be done below in CacheDerivedData().
+				// This is a path for legacy assets.
+				if (!GetSourceModel(i).RawMeshBulkData->IsEmpty())
 				{
-					TotalIndexCount += MeshDescription.VertexInstances().Num();
+					GetSourceModel(i).RawMeshBulkData->LoadRawMesh(TempRawMesh);
+					TotalIndexCount += TempRawMesh.WedgeIndices.Num();
 				}
 			}
 
@@ -5897,6 +5570,7 @@ void UStaticMesh::ExecutePostLoadInternal(FStaticMeshPostLoadContext& Context)
 	}
 
 	ReleaseAsyncProperty(EStaticMeshAsyncProperties::SourceModels);
+	ReleaseAsyncProperty(EStaticMeshAsyncProperties::HiResSourceModel);
 #endif // #if WITH_EDITOR
 }
 
@@ -6120,7 +5794,7 @@ UStaticMeshDescription* UStaticMesh::GetStaticMeshDescription(int32 LODIndex)
 #if WITH_EDITOR
 	if (LODIndex < GetNumSourceModels())
 	{
-		return GetSourceModel(LODIndex).StaticMeshDescription;
+		return GetSourceModel(LODIndex).GetCachedStaticMeshDescription();
 	}
 #endif
 	return nullptr;
@@ -6836,8 +6510,9 @@ ENGINE_API void UStaticMesh::RemoveVertexColors()
 #if WITH_EDITOR
 	bool bRemovedVertexColors = false;
 
-	for (FStaticMeshSourceModel& SourceModel : GetSourceModels())
+	for (int32 LodIndex = 0; LodIndex < GetNumSourceModels(); LodIndex++)
 	{
+		FStaticMeshSourceModel& SourceModel = GetSourceModel(LodIndex);
 		if (!SourceModel.IsRawMeshEmpty())
 		{
 			FRawMesh RawMesh;
diff --git a/Engine/Source/Runtime/Engine/Private/StaticMeshBuild.cpp b/Engine/Source/Runtime/Engine/Private/StaticMeshBuild.cpp
--- a/Engine/Source/Runtime/Engine/Private/StaticMeshBuild.cpp
+++ b/Engine/Source/Runtime/Engine/Private/StaticMeshBuild.cpp
@@ -400,9 +400,9 @@ void UStaticMesh::BeginBuildInternal(FStaticMeshBuildContext* Context)
 		// Extended bounds are super important to avoid many stalls since a lot of different
 		// systems rely on mesh bounds. If mesh description is available and doesn't require being
 		// loaded, compute bounds right now so we can unlock the property before going async.
-		if (IsSourceModelValid(0) && GetSourceModel(0).StaticMeshDescription != nullptr)
+		if (IsSourceModelValid(0) && GetSourceModel(0).GetCachedMeshDescription() != nullptr)
 		{
-			CachedMeshDescriptionBounds = GetSourceModel(0).StaticMeshDescription->GetMeshDescription().GetBounds();
+			CachedMeshDescriptionBounds = GetSourceModel(0).GetCachedMeshDescription()->GetBounds();
 		}
 
 		// CommitMeshDescription will also fill out CachedMeshDescriptionBounds
diff --git a/Engine/Source/Runtime/Engine/Private/StaticMeshSourceData.cpp b/Engine/Source/Runtime/Engine/Private/StaticMeshSourceData.cpp
new file mode 100644
--- /dev/null
+++ b/Engine/Source/Runtime/Engine/Private/StaticMeshSourceData.cpp
@@ -0,0 +1,451 @@
+// Copyright Epic Games, Inc. All Rights Reserved.
+
+#include "Engine/StaticMeshSourceData.h"
+#include "MeshDescription.h"
+#include "StaticMeshDescription.h"
+#include "StaticMeshOperations.h"
+#include "ProfilingDebugging/CookStats.h"
+#if WITH_EDITOR
+#include "DerivedDataCacheInterface.h"
+#include "RawMesh.h"
+#endif
+
+
+UClass* UStaticMeshDescriptionBulkData::GetMeshDescriptionType() const
+{
+	return UStaticMeshDescription::StaticClass();
+}
+
+
+FStaticMeshSourceModel::FStaticMeshSourceModel()
+{
+#if WITH_EDITOR
+	RawMeshBulkData = new FRawMeshBulkData();
+	StaticMeshOwner = nullptr;
+#endif // #if WITH_EDITOR
+	LODDistance_DEPRECATED = 0.0f;
+	ScreenSize.Default = 0.0f;
+#if WITH_EDITORONLY_DATA
+	bImportWithBaseMesh = false;
+	StaticMeshDescriptionBulkData = nullptr;
+#endif
+}
+
+FStaticMeshSourceModel::~FStaticMeshSourceModel()
+{
+#if WITH_EDITOR
+	if (RawMeshBulkData)
+	{
+		delete RawMeshBulkData;
+	}
+#endif // #if WITH_EDITOR
+}
+
+FStaticMeshSourceModel::FStaticMeshSourceModel(FStaticMeshSourceModel&& Other)
+{
+	*this = MoveTemp(Other);
+}
+
+FStaticMeshSourceModel& FStaticMeshSourceModel::operator=(FStaticMeshSourceModel&& Other)
+{
+#if WITH_EDITOR
+	RawMeshBulkData = Other.RawMeshBulkData;
+	Other.RawMeshBulkData = nullptr;
+
+	StaticMeshOwner = Other.StaticMeshOwner;
+	Other.StaticMeshOwner = nullptr;
+#endif
+
+#if WITH_EDITORONLY_DATA
+	StaticMeshDescriptionBulkData = Other.StaticMeshDescriptionBulkData;
+	Other.StaticMeshDescriptionBulkData = nullptr;
+
+	bImportWithBaseMesh = Other.bImportWithBaseMesh;
+#endif
+
+	BuildSettings = MoveTemp(Other.BuildSettings);
+	ReductionSettings = MoveTemp(Other.ReductionSettings);
+	LODDistance_DEPRECATED = Other.LODDistance_DEPRECATED;
+	ScreenSize = MoveTemp(Other.ScreenSize);
+	SourceImportFilename = MoveTemp(Other.SourceImportFilename);
+
+	return *this;
+}
+
+
+void FStaticMeshSourceModel::CreateSubObjects(UStaticMesh* InOwner)
+{
+#if WITH_EDITOR
+	check(StaticMeshOwner == nullptr);
+	StaticMeshOwner = InOwner;
+#endif
+
+#if WITH_EDITORONLY_DATA
+	ensure(IsInGameThread());
+	if (StaticMeshDescriptionBulkData == nullptr)
+	{
+		StaticMeshDescriptionBulkData = NewObject<UStaticMeshDescriptionBulkData>(StaticMeshOwner, NAME_None, RF_Transactional);
+		StaticMeshDescriptionBulkData->AtomicallyClearInternalFlags(EInternalObjectFlags::Async);
+	}
+
+	check(StaticMeshDescriptionBulkData->GetMeshDescription() == nullptr);
+#endif
+}
+
+
+#if WITH_EDITOR
+bool FStaticMeshSourceModel::IsRawMeshEmpty() const
+{
+	// The RawMeshBulkData will always be empty, so the test here is whether we are able to construct a valid non-empty RawMesh
+	// from what exists in the StaticMeshDescription and its bulk data.
+	check(RawMeshBulkData->IsEmpty());
+	check(StaticMeshDescriptionBulkData != nullptr);
+	return !StaticMeshDescriptionBulkData->IsBulkDataValid() && !StaticMeshDescriptionBulkData->HasCachedMeshDescription();
+}
+
+void FStaticMeshSourceModel::LoadRawMesh(FRawMesh& OutRawMesh) const
+{
+	check(RawMeshBulkData->IsEmpty());
+
+	// If the RawMesh is empty, consider the possibility that it's a new asset with a valid MeshDescription which needs loading.
+	check(StaticMeshOwner != nullptr);
+	FStaticMeshSourceModel* MutableThis = const_cast<FStaticMeshSourceModel*>(this);
+	if (FMeshDescription* CachedMeshDescription = MutableThis->GetOrCacheMeshDescription())
+	{
+		TMap<FName, int32> MaterialMap;
+		check(StaticMeshOwner != nullptr);
+		for (int32 MaterialIndex = 0; MaterialIndex < StaticMeshOwner->GetStaticMaterials().Num(); ++MaterialIndex)
+		{
+			MaterialMap.Add(StaticMeshOwner->GetStaticMaterials()[MaterialIndex].ImportedMaterialSlotName, MaterialIndex);
+		}
+		FStaticMeshOperations::ConvertToRawMesh(*CachedMeshDescription, OutRawMesh, MaterialMap);
+	}
+}
+
+void FStaticMeshSourceModel::SaveRawMesh(FRawMesh& InRawMesh, bool /* unused */)
+{
+	if (!InRawMesh.IsValid())
+	{
+		return;
+	}
+
+	TRACE_CPUPROFILER_EVENT_SCOPE(FStaticMeshSourceModel::SaveRawMesh);
+
+	check(RawMeshBulkData->IsEmpty());
+	check(StaticMeshDescriptionBulkData != nullptr);
+
+	UMeshDescriptionBase* MeshDescriptionBase = StaticMeshDescriptionBulkData->CreateMeshDescription();
+
+	TMap<int32, FName> MaterialMap;
+	FillMaterialName(MaterialMap);
+	FStaticMeshOperations::ConvertFromRawMesh(InRawMesh, MeshDescriptionBase->GetMeshDescription(), MaterialMap);
+
+	// Package up mesh description into bulk data
+	const bool bUseHashAsGuid = false;
+	StaticMeshDescriptionBulkData->CommitMeshDescription(bUseHashAsGuid);
+}
+
+
+bool FStaticMeshSourceModel::LoadMeshDescription(FMeshDescription& OutMeshDescription) const
+{
+	TRACE_CPUPROFILER_EVENT_SCOPE(FStaticMeshSourceModel::LoadMeshDescription);
+	check(StaticMeshOwner != nullptr);
+	check(StaticMeshDescriptionBulkData != nullptr);
+
+	// Ensure default MeshDescription result is empty, with no attributes registered
+	OutMeshDescription = FMeshDescription();
+
+	// If we have valid MeshDescriptionBulkData, unpack it and return it...
+	if (StaticMeshDescriptionBulkData->IsBulkDataValid())
+	{
+		// Unpack MeshDescription from the bulk data which was deserialized
+		StaticMeshDescriptionBulkData->GetBulkData().LoadMeshDescription(OutMeshDescription);
+		return true;
+	}
+
+	// RawMeshBulkData should always be empty now (soon to be deprecated)
+	check(RawMeshBulkData->IsEmpty());
+
+	// This LOD doesn't have a MeshDescriptionBulkData, or a RawMesh, so we presume that it's a generated LOD (and return false)
+	return false;
+}
+
+
+bool FStaticMeshSourceModel::CloneMeshDescription(FMeshDescription& OutMeshDescription) const
+{
+	TRACE_CPUPROFILER_EVENT_SCOPE(FStaticMeshSourceModel::CloneMeshDescription);
+
+	if (StaticMeshDescriptionBulkData->HasCachedMeshDescription())
+	{
+		OutMeshDescription = StaticMeshDescriptionBulkData->GetMeshDescription()->GetMeshDescription();
+		return true;
+	}
+
+	return LoadMeshDescription(OutMeshDescription);
+}
+
+
+FMeshDescription* FStaticMeshSourceModel::GetOrCacheMeshDescription()
+{
+	TRACE_CPUPROFILER_EVENT_SCOPE(FStaticMeshSourceModel::GetMeshDescription);
+	check(StaticMeshOwner != nullptr);
+
+	if (!StaticMeshDescriptionBulkData->HasCachedMeshDescription())
+	{
+		FMeshDescription MeshDescription;
+		if (LoadMeshDescription(MeshDescription))
+		{
+			UMeshDescriptionBase* MeshDescriptionBase = StaticMeshDescriptionBulkData->CreateMeshDescription();
+			MeshDescriptionBase->SetMeshDescription(MoveTemp(MeshDescription));
+		}
+	}
+
+	if (StaticMeshDescriptionBulkData->HasCachedMeshDescription())
+	{
+		return &StaticMeshDescriptionBulkData->GetMeshDescription()->GetMeshDescription();
+	}
+	else
+	{
+		return nullptr;
+	}
+}
+
+
+FMeshDescription* FStaticMeshSourceModel::GetCachedMeshDescription() const
+{
+	if (StaticMeshDescriptionBulkData->HasCachedMeshDescription())
+	{
+		return &StaticMeshDescriptionBulkData->GetMeshDescription()->GetMeshDescription();
+	}
+
+	return nullptr;
+}
+
+
+UStaticMeshDescription* FStaticMeshSourceModel::GetCachedStaticMeshDescription() const
+{
+	return Cast<UStaticMeshDescription>(StaticMeshDescriptionBulkData->GetMeshDescription());
+}
+
+
+FMeshDescriptionBulkData* FStaticMeshSourceModel::GetMeshDescriptionBulkData() const
+{
+	check(StaticMeshDescriptionBulkData != nullptr);
+	return &StaticMeshDescriptionBulkData->GetBulkData();
+}
+
+
+bool FStaticMeshSourceModel::IsMeshDescriptionValid() const
+{
+	// Determine whether a mesh description is valid without requiring it to be loaded first.
+	// If there is a valid MeshDescriptionBulkData, we know this implies a valid mesh description.
+	check(StaticMeshDescriptionBulkData != nullptr);
+	check(RawMeshBulkData->IsEmpty());
+	return StaticMeshDescriptionBulkData->HasCachedMeshDescription() || StaticMeshDescriptionBulkData->IsBulkDataValid();
+}
+
+
+FMeshDescription* FStaticMeshSourceModel::CreateMeshDescription()
+{
+	return &StaticMeshDescriptionBulkData->CreateMeshDescription()->GetMeshDescription();
+}
+
+
+void FStaticMeshSourceModel::CommitMeshDescription(bool bUseHashAsGuid)
+{
+	TRACE_CPUPROFILER_EVENT_SCOPE(FStaticMeshSourceModel::CommitMeshDescription);
+
+	check(RawMeshBulkData->IsEmpty());
+
+	// This part should remain thread-safe so it can be called from any thread
+	// as long as no more than one thread is calling it for the same UStaticMesh.
+
+	check(StaticMeshDescriptionBulkData != nullptr);
+	if (StaticMeshDescriptionBulkData->HasCachedMeshDescription())
+	{
+		// Package up mesh description into bulk data
+		StaticMeshDescriptionBulkData->CommitMeshDescription(bUseHashAsGuid);
+	}
+	else
+	{
+		StaticMeshDescriptionBulkData->Empty();
+	}
+}
+
+void FStaticMeshSourceModel::ClearMeshDescription()
+{
+	TRACE_CPUPROFILER_EVENT_SCOPE(UStaticMesh::ClearMeshDescription);
+
+	StaticMeshDescriptionBulkData->RemoveMeshDescription();
+}
+
+
+void FStaticMeshSourceModel::SerializeBulkData(FArchive& Ar, UObject* Owner)
+{
+	// Initialize the StaticMeshOwner
+	if (Ar.IsLoading())
+	{
+		StaticMeshOwner = Cast<UStaticMesh>(Owner);
+
+		// If this was a legacy asset, or is being created for the first time, create a bulkdata UObject wrapper
+		if (StaticMeshDescriptionBulkData == nullptr)
+		{
+			StaticMeshDescriptionBulkData = NewObject<UStaticMeshDescriptionBulkData>(StaticMeshOwner, NAME_None, RF_Transactional);
+		}
+	}
+
+	if (Ar.IsLoading() && Ar.CustomVer(FEditorObjectVersion::GUID) < FEditorObjectVersion::StaticMeshDeprecatedRawMesh)
+	{
+		// If loading a legacy asset with RawMesh bulk data, serialize it here
+		// The conversion to MeshDescription will be done asynchronously, during PostLoad.
+		check(RawMeshBulkData != NULL);
+		RawMeshBulkData->Serialize(Ar, Owner);
+	}
+	else
+	{
+		if (Ar.IsLoading())
+		{
+			if (Ar.CustomVer(FUE5MainStreamObjectVersion::GUID) < FUE5MainStreamObjectVersion::SerializeMeshDescriptionBase)
+			{
+				// If there's an inline mesh description bulk data (legacy version), serialize it here, and copy it into the new separate inner object
+				bool bIsValid;
+				Ar << bIsValid;
+
+				if (bIsValid)
+				{
+					StaticMeshDescriptionBulkData->GetBulkData().Serialize(Ar, StaticMeshOwner);
+				}
+			}
+
+		}
+	}
+}
+
+
+void FStaticMeshSourceModel::FillMaterialName(TMap<int32, FName>& OutMaterialMap) const
+{
+	check(StaticMeshOwner != nullptr);
+	const TArray<FStaticMaterial>& StaticMaterials = StaticMeshOwner->GetStaticMaterials();
+
+	OutMaterialMap.Empty(StaticMaterials.Num());
+
+	for (int32 MaterialIndex = 0; MaterialIndex < StaticMaterials.Num(); ++MaterialIndex)
+	{
+		FName MaterialName = StaticMaterials[MaterialIndex].ImportedMaterialSlotName;
+		if (MaterialName == NAME_None)
+		{
+			MaterialName = *(TEXT("MaterialSlot_") + FString::FromInt(MaterialIndex));
+		}
+		OutMaterialMap.Add(MaterialIndex, MaterialName);
+	}
+}
+
+
+// This is the key which is used for DDC data when legacy RawMesh is converted to MeshDescription.
+// If static mesh derived data needs to be rebuilt (new format, serialization differences, etc.) replace the version GUID below with a new one.
+// In case of merge conflicts with DDC versions, you MUST generate a new GUID and set this new GUID as the version.
+#if UE_USE_VIRTUALBULKDATA
+#define MESHDATAKEY_STATICMESH_DERIVEDDATA_VER TEXT("7F1A02EF80374F0F86F9A30BAFF3F850")
+#else
+#define MESHDATAKEY_STATICMESH_DERIVEDDATA_VER TEXT("E09E8DD68C864D82B511A97080B04837")
+#endif
+
+static bool GetMeshDataKey(FString& OutKey, const FRawMeshBulkData* RawMeshBulkData, int LodIndex)
+{
+	OutKey.Empty();
+
+	FSHA1 Sha;
+
+	check(!RawMeshBulkData->IsEmpty());
+	FString LodIndexString = FString::Printf(TEXT("%d_"), LodIndex);
+	LodIndexString += RawMeshBulkData->GetIdString();
+
+	const TArray<TCHAR>& LodIndexArray = LodIndexString.GetCharArray();
+	Sha.Update((uint8*)LodIndexArray.GetData(), LodIndexArray.Num() * LodIndexArray.GetTypeSize());
+	Sha.Final();
+
+	// Retrieve the hash and use it to construct a pseudo-GUID.
+	uint32 Hash[5];
+	Sha.GetHash((uint8*)Hash);
+	FGuid Guid = FGuid(Hash[0] ^ Hash[4], Hash[1], Hash[2], Hash[3]);
+	FString MeshLodData = Guid.ToString(EGuidFormats::Digits);
+
+	OutKey = FDerivedDataCacheInterface::BuildCacheKey(
+		TEXT("MESHDATAKEY_STATICMESH"),
+		MESHDATAKEY_STATICMESH_DERIVEDDATA_VER,
+		*MeshLodData
+	);
+	return true;
+}
+
+
+#if ENABLE_COOK_STATS
+namespace StaticMeshConvertStats
+{
+	static FCookStats::FDDCResourceUsageStats UsageStats;
+	static FCookStatsManager::FAutoRegisterCallback RegisterCookStats([](FCookStatsManager::AddStatFuncRef AddStat)
+		{
+			UsageStats.LogStats(AddStat, TEXT("StaticMeshConvert.Usage"), TEXT(""));
+		});
+}
+#endif
+
+
+void FStaticMeshSourceModel::ConvertRawMesh(int32 LodIndex)
+{
+	check(StaticMeshDescriptionBulkData != nullptr);
+	check(StaticMeshOwner != nullptr);
+
+	if (!RawMeshBulkData->IsEmpty() && !StaticMeshDescriptionBulkData->IsBulkDataValid())
+	{
+		FString MeshDataKey;
+		if (GetMeshDataKey(MeshDataKey, RawMeshBulkData, LodIndex))
+		{
+			COOK_STAT(auto Timer = StaticMeshConvertStats::UsageStats.TimeSyncWork());
+
+			TArray<uint8> DerivedData;
+
+			if (GetDerivedDataCacheRef().GetSynchronous(*MeshDataKey, DerivedData, StaticMeshOwner->GetPathName()))
+			{
+				COOK_STAT(Timer.AddHit(DerivedData.Num()));
+
+				// Load from the DDC
+				const bool bIsPersistent = true;
+				FMemoryReader Ar(DerivedData, bIsPersistent);
+				StaticMeshDescriptionBulkData->GetBulkData().Serialize(Ar, StaticMeshOwner);
+
+				check(GetCachedMeshDescription() == nullptr);
+			}
+			else
+			{
+				// If the DDC key doesn't exist, convert the data and save it to DDC
+				// Get the RawMesh for this LOD
+				FRawMesh TempRawMesh;
+				RawMeshBulkData->LoadRawMesh(TempRawMesh);
+
+				// Create a new MeshDescription
+				FMeshDescription* MeshDescription = CreateMeshDescription();
+
+				// Convert the RawMesh to MeshDescription
+				TMap<int32, FName> MaterialMap;
+				FillMaterialName(MaterialMap);
+				FStaticMeshOperations::ConvertFromRawMesh(TempRawMesh, *MeshDescription, MaterialMap);
+
+				// Pack MeshDescription into bulk data
+				StaticMeshDescriptionBulkData->GetBulkData().SaveMeshDescription(*MeshDescription);
+
+				// Write the DDC cache
+				const bool bIsPersistent = true;
+				FMemoryWriter Ar(DerivedData, bIsPersistent);
+				StaticMeshDescriptionBulkData->GetBulkData().Serialize(Ar, StaticMeshOwner);
+				GetDerivedDataCacheRef().Put(*MeshDataKey, DerivedData, StaticMeshOwner->GetPathName());
+				COOK_STAT(Timer.AddMiss(DerivedData.Num()));
+			}
+		}
+
+		// We now have a MeshDescription instead of a RawMesh, so get rid of the RawMesh completely
+		RawMeshBulkData->Empty();
+	}
+}
+
+#endif // #if WITH_EDITOR
diff --git a/Engine/Source/Runtime/MeshDescription/Private/MeshDescription.cpp b/Engine/Source/Runtime/MeshDescription/Private/MeshDescription.cpp
--- a/Engine/Source/Runtime/MeshDescription/Private/MeshDescription.cpp
+++ b/Engine/Source/Runtime/MeshDescription/Private/MeshDescription.cpp
@@ -2172,15 +2172,3 @@ void FMeshDescriptionBulkData::UseHashAsGuid()
 
 #endif // #if WITH_EDITORONLY_DATA
 
-
-void UMeshDescriptionBulkDataWrapper::Serialize(FArchive& Ar)
-{
-	Super::Serialize(Ar);
-
-#if WITH_EDITORONLY_DATA
-	if (!Ar.IsCooking())
-	{
-		BulkData.Serialize(Ar, GetOuter());
-	}
-#endif
-}
diff --git a/Engine/Source/Runtime/MeshDescription/Private/MeshDescriptionBaseBulkData.cpp b/Engine/Source/Runtime/MeshDescription/Private/MeshDescriptionBaseBulkData.cpp
new file mode 100644
--- /dev/null
+++ b/Engine/Source/Runtime/MeshDescription/Private/MeshDescriptionBaseBulkData.cpp
@@ -0,0 +1,145 @@
+// Copyright Epic Games, Inc. All Rights Reserved.
+
+#include "MeshDescriptionBaseBulkData.h"
+#include "MeshDescriptionBase.h"
+#include "MeshDescription.h"
+
+
+UMeshDescriptionBaseBulkData::UMeshDescriptionBaseBulkData()
+{
+#if WITH_EDITORONLY_DATA
+	BulkData = MakePimpl<FMeshDescriptionBulkData>();
+	MeshDescription = nullptr;
+#endif
+}
+
+
+void UMeshDescriptionBaseBulkData::Serialize(FArchive& Ar)
+{
+	Super::Serialize(Ar);
+
+#if WITH_EDITORONLY_DATA
+	if (!Ar.IsCooking())
+	{
+		check(BulkData.IsValid());
+		BulkData->Serialize(Ar, GetOuter());
+	}
+
+	if (MeshDescription != nullptr)
+	{
+		// If there was a cached mesh description, it could be out of sync with the bulk data we just serialized so re-cache it.
+		CacheMeshDescription();
+	}
+#endif
+}
+
+
+UClass* UMeshDescriptionBaseBulkData::GetMeshDescriptionType() const
+{
+	return UMeshDescriptionBase::StaticClass();
+}
+
+
+#if WITH_EDITORONLY_DATA
+
+void UMeshDescriptionBaseBulkData::Empty()
+{
+	BulkData->Empty();
+}
+
+
+UMeshDescriptionBase* UMeshDescriptionBaseBulkData::CreateMeshDescription()
+{
+	if (MeshDescription == nullptr)
+	{
+		// Assign the new UObject directly to a UPROPERTY.
+		// If we were run from a different thread, we can immediately clear the Async internal flag, now that GC refs are set up correctly.
+		UClass* MeshDescriptionType = GetMeshDescriptionType();
+		MeshDescription = NewObject<UMeshDescriptionBase>(this, MeshDescriptionType, NAME_None, RF_Transient | RF_Transactional);
+		MeshDescription->AtomicallyClearInternalFlags(EInternalObjectFlags::Async);
+	}
+
+	// Reset the mesh description and register its attributes
+	// Do this instead of always creating a new UStaticMeshDescription object, to save memory and improve performance
+	MeshDescription->Reset();
+
+	return MeshDescription;
+}
+
+
+UMeshDescriptionBase* UMeshDescriptionBaseBulkData::GetMeshDescription() const
+{
+	return MeshDescription;
+}
+
+
+bool UMeshDescriptionBaseBulkData::HasCachedMeshDescription() const
+{
+	return (MeshDescription != nullptr);
+}
+
+
+bool UMeshDescriptionBaseBulkData::CacheMeshDescription()
+{
+	check(MeshDescription);
+
+	if (!BulkData->IsEmpty())
+	{
+		FMeshDescription NewMeshDescription;
+		BulkData->LoadMeshDescription(NewMeshDescription);
+		MeshDescription->SetMeshDescription(MoveTemp(NewMeshDescription));
+		return true;
+	}
+	else
+	{
+		RemoveMeshDescription();
+		return false;
+	}
+}
+
+
+void UMeshDescriptionBaseBulkData::CommitMeshDescription(bool bUseHashAsGuid)
+{
+	if (MeshDescription != nullptr)
+	{
+		BulkData->SaveMeshDescription(MeshDescription->GetMeshDescription());
+		if (bUseHashAsGuid)
+		{
+			BulkData->UseHashAsGuid();
+		}
+	}
+	else
+	{
+		BulkData->Empty();
+	}
+}
+
+
+void UMeshDescriptionBaseBulkData::RemoveMeshDescription()
+{
+	if (MeshDescription != nullptr)
+	{
+		MeshDescription->Empty();
+		MeshDescription = nullptr;
+	}
+}
+
+
+bool UMeshDescriptionBaseBulkData::IsBulkDataValid() const
+{
+	return !BulkData->IsEmpty();
+}
+
+
+const FMeshDescriptionBulkData& UMeshDescriptionBaseBulkData::GetBulkData() const
+{
+	return *BulkData;
+}
+
+
+FMeshDescriptionBulkData& UMeshDescriptionBaseBulkData::GetBulkData()
+{
+	return *BulkData;
+}
+
+#endif
diff --git a/Engine/Source/Runtime/MeshDescription/Public/MeshDescription.h b/Engine/Source/Runtime/MeshDescription/Public/MeshDescription.h
--- a/Engine/Source/Runtime/MeshDescription/Public/MeshDescription.h
+++ b/Engine/Source/Runtime/MeshDescription/Public/MeshDescription.h
@@ -1340,25 +1340,3 @@ private:
 };
 #endif // WITH_EDITORONLY_DATA
 
-
-/**
- * UObject wrapper for FMeshDescriptionBulkData
- */
-UCLASS()
-class MESHDESCRIPTION_API UMeshDescriptionBulkDataWrapper : public UObject
-{
-	GENERATED_BODY()
-
-public:
-	virtual void Serialize(FArchive& Ar) override;
-
-#if WITH_EDITORONLY_DATA
-
-	const FMeshDescriptionBulkData& GetBulkData() const { return BulkData; }
-	FMeshDescriptionBulkData& GetBulkData() { return BulkData; }
-
-private:
-	FMeshDescriptionBulkData BulkData;
-
-#endif //WITH_EDITORONLY_DATA
-};
diff --git a/Engine/Source/Runtime/MeshDescription/Public/MeshDescriptionBaseBulkData.h b/Engine/Source/Runtime/MeshDescription/Public/MeshDescriptionBaseBulkData.h
new file mode 100644
--- /dev/null
+++ b/Engine/Source/Runtime/MeshDescription/Public/MeshDescriptionBaseBulkData.h
@@ -0,0 +1,51 @@
+// Copyright Epic Games, Inc. All Rights Reserved.
+
+#pragma once
+
+#include "CoreMinimal.h"
+#include "UObject/Object.h"
+#include "UObject/ObjectMacros.h"
+#include "Templates/PimplPtr.h"
+
+#include "MeshDescriptionBaseBulkData.generated.h"
+
+
+class UMeshDescriptionBase;
+struct FMeshDescriptionBulkData;
+
+
+/**
+ * UObject wrapper for FMeshDescriptionBulkData
+ */
+UCLASS()
+class MESHDESCRIPTION_API UMeshDescriptionBaseBulkData : public UObject
+{
+	GENERATED_BODY()
+
+public:
+	UMeshDescriptionBaseBulkData();
+	virtual void Serialize(FArchive& Ar) override;
+	virtual UClass* GetMeshDescriptionType() const;
+
+#if WITH_EDITORONLY_DATA
+
+	void Empty();
+	UMeshDescriptionBase* CreateMeshDescription();
+	UMeshDescriptionBase* GetMeshDescription() const;
+	bool HasCachedMeshDescription() const;
+	bool CacheMeshDescription();
+	void CommitMeshDescription(bool bUseHashAsGuid);
+	void RemoveMeshDescription();
+	bool IsBulkDataValid() const;
+
+	const FMeshDescriptionBulkData& GetBulkData() const;
+	FMeshDescriptionBulkData& GetBulkData();
+
+protected:
+	TPimplPtr<FMeshDescriptionBulkData> BulkData;
+
+	UPROPERTY(Transient)
+	UMeshDescriptionBase* MeshDescription;
+
+#endif //WITH_EDITORONLY_DATA
+};
