Updated textures to use the new DDC TryToPrefetch

diff --git a/Engine/Source/Runtime/Engine/Private/TextureDerivedData.cpp b/Engine/Source/Runtime/Engine/Private/TextureDerivedData.cpp
--- a/Engine/Source/Runtime/Engine/Private/TextureDerivedData.cpp
+++ b/Engine/Source/Runtime/Engine/Private/TextureDerivedData.cpp
@@ -1388,33 +1388,20 @@ bool FTexturePlatformData::AreDerivedMipsAvailable() const
 			MipKeys.Add(Mip.DerivedDataKey);
 		}
 	}
-	
-	// If this returns true, everything is already present in the local cache with fast access
-	bool bDataProbablyExists = GetDerivedDataCacheRef().AllCachedDataProbablyExists(MipKeys);
-
-	// This is a temporary solution that should be replaced by a Prefetch command
-	// once the new DDC is in place. 
-	// CachedDataProbablyExists doesn't work on slow cache, so it means we rebuild the 
-	// textures every time we don't have these locally.. which is bad, so Get them instead.
-	if (!bDataProbablyExists && !IsInGameThread())
+	if (IsInGameThread())
+	{
+		return GetDerivedDataCacheRef().AllCachedDataProbablyExists(MipKeys);
+	}
+	else
 	{
-		bDataProbablyExists = true;
 		// When using a shared DDC and performing async loading, 
 		// prefetch the data so we don't pay the latency on the game-thread
 		// to fetch the asset from a remote location once
 		// the mips needs to be accessed
 		TRACE_CPUPROFILER_EVENT_SCOPE(PrefetchMips);
-		TArray<uint8> ThrowAwayData;
-		for (FString& Key : MipKeys)
-		{
-			ThrowAwayData.Reset();
-			bDataProbablyExists &= GetDerivedDataCacheRef().GetSynchronous(*Key, ThrowAwayData, TEXT("DerivedMips"));
-		}
+		return GetDerivedDataCacheRef().TryToPrefetch(MipKeys, TEXT("DerivedMips"_SV));
 	}
-
-	return bDataProbablyExists;
 }
-
 bool FTexturePlatformData::AreDerivedVTChunksAvailable() const
 {
 	check(VTData);
@@ -1426,31 +1413,19 @@ bool FTexturePlatformData::AreDerivedVTChunksAvailable() const
 			ChunkKeys.Add(Chunk.DerivedDataKey);
 		}
 	}
-
-	// If this returns true, everything is already present in the local cache with fast access
-	bool bDataProbablyExists = GetDerivedDataCacheRef().AllCachedDataProbablyExists(ChunkKeys);
-
-	// This is a temporary solution that should be replaced by a Prefetch command
-	// once the new DDC is in place. 
-	// CachedDataProbablyExists doesn't work on slow cache, so it means we rebuild the 
-	// textures every time we don't have these locally.. which is bad, so Get them instead.
-	if (!bDataProbablyExists && !IsInGameThread())
+	if (IsInGameThread())
+	{
+		return GetDerivedDataCacheRef().AllCachedDataProbablyExists(ChunkKeys);
+	}
+	else
 	{
-		bDataProbablyExists = true;
 		// When using a shared DDC and performing async loading, 
 		// prefetch the data so we don't pay the latency on the game-thread
 		// to fetch the asset from a remote location once
 		// the VT chunks needs to be accessed
 		TRACE_CPUPROFILER_EVENT_SCOPE(PrefetchDerivedVTChunks);
-		TArray<uint8> ThrowAwayData;
-		for (FString& Key : ChunkKeys)
-		{
-			ThrowAwayData.Reset();
-			bDataProbablyExists &= GetDerivedDataCacheRef().GetSynchronous(*Key, ThrowAwayData, TEXT("DerivedVTChunks"));
-		}
+		return GetDerivedDataCacheRef().TryToPrefetch(ChunkKeys, TEXT("DerivedVTChunks"_SV));
 	}
-
-	return bDataProbablyExists;
 }
 #endif // #if WITH_EDITOR
 
