Exposed montage function to Engine_API
#rnx

diff --git a/Engine/Source/Runtime/Engine/Classes/Animation/AnimMontage.h b/Engine/Source/Runtime/Engine/Classes/Animation/AnimMontage.h
--- a/Engine/Source/Runtime/Engine/Classes/Animation/AnimMontage.h
+++ b/Engine/Source/Runtime/Engine/Classes/Animation/AnimMontage.h
@@ -803,7 +803,7 @@ public:
 	ENGINE_API const TArray<class UAnimMetaData*> GetSectionMetaData(FName SectionName, bool bIncludeSequence=true, FName SlotName = NAME_None);
 
 	/** Get Section Index from CurrentTime with PosWithinCompositeSection */
-	int32 GetAnimCompositeSectionIndexFromPos(float CurrentTime, float& PosWithinCompositeSection) const;
+	ENGINE_API int32 GetAnimCompositeSectionIndexFromPos(float CurrentTime, float& PosWithinCompositeSection) const;
 
 	/** Return time left to end of section from given position. -1.f if not a valid position */
 	ENGINE_API float GetSectionTimeLeftFromPos(float Position);
