Fixed use-after-free of FVirtualTextureCodec which could happen if FTranscodeTask takes a long time to complete.

diff --git a/Engine/Source/Runtime/Engine/Private/VT/UploadingVirtualTexture.cpp b/Engine/Source/Runtime/Engine/Private/VT/UploadingVirtualTexture.cpp
--- a/Engine/Source/Runtime/Engine/Private/VT/UploadingVirtualTexture.cpp
+++ b/Engine/Source/Runtime/Engine/Private/VT/UploadingVirtualTexture.cpp
@@ -163,7 +163,7 @@ void FVirtualTextureCodec::RetireOldCodecs()
 
 		bool bRetiredCodec = false;
 		// Can't retire codec if it's not even finished loading yet
-		if (Codec.Owner && (!Codec.CompletedEvent || Codec.CompletedEvent->IsComplete()))
+		if (Codec.Owner && Codec.IsIdle())
 		{
 			check(CurrentFrame >= Codec.LastFrameUsed);
 			const uint32 Age = CurrentFrame - Codec.LastFrameUsed;
@@ -247,7 +247,9 @@ FVirtualTextureCodec::~FVirtualTextureCodec()
 {
 	if (Owner)
 	{
-		check(IsComplete());
+		checkf(IsCreationComplete(), TEXT("Codec is being released before its construction task hasn't finished."));
+		checkf(AllTranscodeTasksComplete(), TEXT("Codec is being released while there are tasks that still reference it."));
+
 		check(!IsLinked());
 
 		const FVirtualTextureBuiltData* VTData = Owner->GetVTData();
diff --git a/Engine/Source/Runtime/Engine/Private/VT/UploadingVirtualTexture.h b/Engine/Source/Runtime/Engine/Private/VT/UploadingVirtualTexture.h
--- a/Engine/Source/Runtime/Engine/Private/VT/UploadingVirtualTexture.h
+++ b/Engine/Source/Runtime/Engine/Private/VT/UploadingVirtualTexture.h
@@ -10,18 +10,24 @@
 #include "Templates/UniquePtr.h"
 #include "Async/TaskGraphInterfaces.h"
 
+#include <atomic>
+
 class FUploadingVirtualTexture;
 class IFileCacheHandle;
 
 class FVirtualTextureCodec : public TIntrusiveLinkedList<FVirtualTextureCodec>
 {
+	friend struct FTranscodeTask;
 public:
 	static void RetireOldCodecs();
 
 	~FVirtualTextureCodec();
 	void Init(IMemoryReadStreamRef& HeaderData);
 
-	inline bool IsComplete() const { return !CompletedEvent || CompletedEvent->IsComplete(); }
+	// Tracks completion of FCreateCodecTask.
+	inline bool IsCreationComplete() const { return !CompletedEvent || CompletedEvent->IsComplete(); }
+
+	inline bool IsIdle() const { return IsCreationComplete() && AllTranscodeTasksComplete(); }
 
 	void LinkGlobalHead();
 	void LinkGlobalTail();
@@ -36,6 +42,16 @@ public:
 	void* Contexts[VIRTUALTEXTURE_DATA_MAXLAYERS] = { nullptr };
 	uint32 ChunkIndex = 0u;
 	uint32 LastFrameUsed = 0u;
+
+protected:
+
+	// This codec may be used by transcode tasks (see FTranscodeTask),
+	// which must be all complete before the codec is released by RetireOldTasks().
+	inline bool AllTranscodeTasksComplete() const { return TaskRefs.load() == 0; }
+	inline void BeginTranscodeTask() const { TaskRefs.fetch_add(1); }
+	inline void EndTranscodeTask() const { TaskRefs.fetch_sub(1); }
+
+	mutable std::atomic<uint32> TaskRefs {};
 };
 
 struct FVTCodecAndStatus
diff --git a/Engine/Source/Runtime/Engine/Private/VT/VirtualTextureTranscodeCache.cpp b/Engine/Source/Runtime/Engine/Private/VT/VirtualTextureTranscodeCache.cpp
--- a/Engine/Source/Runtime/Engine/Private/VT/VirtualTextureTranscodeCache.cpp
+++ b/Engine/Source/Runtime/Engine/Private/VT/VirtualTextureTranscodeCache.cpp
@@ -58,9 +58,22 @@ struct FTranscodeTask
 	FTranscodeTask(const FVTUploadTileBuffer* InStagingBuffers, const FVTTranscodeParams& InParams)
 		: Params(InParams)
 	{
+		if (Params.Codec)
+		{
+			Params.Codec->BeginTranscodeTask();
+		}
+
 		FMemory::Memcpy(StagingBuffer, InStagingBuffers, sizeof(FVTUploadTileBuffer) * VIRTUALTEXTURE_SPACE_MAXLAYERS);
 	}
 
+	~FTranscodeTask()
+	{
+		if (Params.Codec)
+		{
+			Params.Codec->EndTranscodeTask();
+		}
+	}
+
 	void DoTask(ENamedThreads::Type CurrentThread, const FGraphEventRef& MyCompletionGraphEvent)
 	{
 		DoTask(true);
@@ -86,7 +99,7 @@ struct FTranscodeTask
 		const uint32 TileIndex = Params.VTData->GetTileIndex(vLevel, vAddress);
 
 		// code must be fully loaded by the time we start transcoding
-		check(!Params.Codec || Params.Codec->IsComplete());
+		check(!Params.Codec || Params.Codec->IsCreationComplete());
 
 		// Used to allocate any temp memory needed to decode tile
 		// Inline allocator hopefully avoids heap allocation in most cases
