Fix for race condition initializing FormatVersionMap for FShaderCompileUtilities::DoWriteTasks

diff --git a/Engine/Source/Runtime/Engine/Private/ShaderCompiler/ShaderCompiler.cpp b/Engine/Source/Runtime/Engine/Private/ShaderCompiler/ShaderCompiler.cpp
--- a/Engine/Source/Runtime/Engine/Private/ShaderCompiler/ShaderCompiler.cpp
+++ b/Engine/Source/Runtime/Engine/Private/ShaderCompiler/ShaderCompiler.cpp
@@ -978,24 +978,25 @@ namespace SCWErrorCode
 	}
 }
 
-static inline void GetFormatVersionMap(TMap<FString, uint32>& OutFormatVersionMap)
+static TMap<FString, uint32> GetFormatVersionMap()
 {
-	if (OutFormatVersionMap.Num() == 0)
+	TMap<FString, uint32> FormatVersionMap;
+
+	const TArray<const class IShaderFormat*>& ShaderFormats = GetTargetPlatformManagerRef().GetShaderFormats();
+	check(ShaderFormats.Num());
+	for (int32 Index = 0; Index < ShaderFormats.Num(); Index++)
 	{
-		const TArray<const class IShaderFormat*>& ShaderFormats = GetTargetPlatformManagerRef().GetShaderFormats();
-		check(ShaderFormats.Num());
-		for (int32 Index = 0; Index < ShaderFormats.Num(); Index++)
+		TArray<FName> OutFormats;
+		ShaderFormats[Index]->GetSupportedFormats(OutFormats);
+		check(OutFormats.Num());
+		for (int32 InnerIndex = 0; InnerIndex < OutFormats.Num(); InnerIndex++)
 		{
-			TArray<FName> OutFormats;
-			ShaderFormats[Index]->GetSupportedFormats(OutFormats);
-			check(OutFormats.Num());
-			for (int32 InnerIndex = 0; InnerIndex < OutFormats.Num(); InnerIndex++)
-			{
-				uint32 Version = ShaderFormats[Index]->GetVersion(OutFormats[InnerIndex]);
-				OutFormatVersionMap.Add(OutFormats[InnerIndex].ToString(), Version);
-			}
+			uint32 Version = ShaderFormats[Index]->GetVersion(OutFormats[InnerIndex]);
+			FormatVersionMap.Add(OutFormats[InnerIndex].ToString(), Version);
 		}
 	}
+
+	return FormatVersionMap;
 }
 
 static int32 GetNumTotalJobs(const TArray<FShaderCommonCompileJobPtr>& Jobs)
@@ -1035,8 +1036,7 @@ bool FShaderCompileUtilities::DoWriteTasks(const TArray<FShaderCommonCompileJobP
 	int32 InputVersion = ShaderCompileWorkerInputVersion;
 	TransferFile << InputVersion;
 
-	thread_local static TMap<FString, uint32> FormatVersionMap;
-	GetFormatVersionMap(FormatVersionMap);
+	static TMap<FString, uint32> FormatVersionMap = GetFormatVersionMap();
 
 	TransferFile << FormatVersionMap;
 
@@ -1654,7 +1654,7 @@ struct FShaderCompileWorkerInfo
 	/** Time at which the worker started the most recent batch of tasks. */
 	double StartTime;
 
-	/** Time at which the worker finished its most recent batch of tasks. */
+	/** Jobs that this worker is responsible for compiling. */
 	double FinishTime = 0.0;
 
 	/** Jobs that this worker is responsible for compiling. */
@@ -4232,7 +4232,7 @@ void GlobalBeginCompileShader(
 		);
 }
 
-/** Enqueues a shader compile job with GShaderCompilingManager. */
+
 void GlobalBeginCompileShader(
 	const FString& DebugGroupName,
 	const FVertexFactoryType* VFType,
