Fix missing data in skeletalmesh ddc, out of sync indexbuffer and vertex count can happen if reduction or chunking change for an asset.

Change edigrate from UE4 main

diff --git a/Engine/Source/Runtime/Engine/Private/SkeletalMeshRenderData.cpp b/Engine/Source/Runtime/Engine/Private/SkeletalMeshRenderData.cpp
--- a/Engine/Source/Runtime/Engine/Private/SkeletalMeshRenderData.cpp
+++ b/Engine/Source/Runtime/Engine/Private/SkeletalMeshRenderData.cpp
@@ -69,7 +69,7 @@ static void SerializeLODInfoForDDC(USkeletalMesh* SkeletalMesh, FString& KeySuff
 // differences, etc.) replace the version GUID below with a new one.
 // In case of merge conflicts with DDC versions, you MUST generate a new GUID
 // and set this new GUID as the version.
-#define SKELETALMESH_DERIVEDDATA_VER TEXT("25A0331C18BA45B1AD69D06F23E84001")
+#define SKELETALMESH_DERIVEDDATA_VER TEXT("E6CA95D22F564092BE4AE5A8716349B0")
 
 static const FString& GetSkeletalMeshDerivedDataVersion()
 {
@@ -144,6 +144,57 @@ FString BuildSkeletalMeshDerivedDataKey(const ITargetPlatform* TargetPlatform, U
 	);
 }
 
+/** This code verify that the data is all in sync index buffer versus sections data. It is active only in debug build*/
+void VerifyAllLodSkeletalMeshModelIntegrity(USkeletalMesh* Owner)
+{
+	if (!Owner || !Owner->GetImportedModel())
+	{
+		return;
+	}
+
+	FSkeletalMeshModel* SkelMeshModel = Owner->GetImportedModel();
+	for (int32 LODIndex = 0; LODIndex < SkelMeshModel->LODModels.Num(); LODIndex++)
+	{
+		FSkeletalMeshLODModel* LODModel = &(SkelMeshModel->LODModels[LODIndex]);
+		int32 SectionsVerticeNum = 0;
+		int32 SectionsTriangleNum = 0;
+		for (const FSkelMeshSection& Section : LODModel->Sections)
+		{
+			SectionsVerticeNum += Section.GetNumVertices();
+			SectionsTriangleNum += Section.NumTriangles;
+			int32 LastSectionIndexBuffer = Section.BaseIndex + (Section.NumTriangles * 3);
+			if (Section.NumTriangles > 0)
+			{
+				//Remove 1 if we have at least one triangle
+				LastSectionIndexBuffer--;
+			}
+
+			if (LODModel->IndexBuffer.IsValidIndex(LastSectionIndexBuffer))
+			{
+				uint32 FirstSectionIndexBufferValue = LODModel->IndexBuffer[Section.BaseIndex];
+				uint32 LastSectionIndexBufferValue = LODModel->IndexBuffer[LastSectionIndexBuffer];
+				if (FirstSectionIndexBufferValue < Section.BaseVertexIndex || LastSectionIndexBufferValue >= Section.BaseVertexIndex + Section.GetNumVertices())
+				{
+					UE_ASSET_LOG(LogSkeletalMesh, Error, Owner, TEXT("The source model is corrupted! Section triangle refer to a vertex not in the section. LOD %"), LODIndex);
+				}
+			}
+			else
+			{
+				UE_ASSET_LOG(LogSkeletalMesh, Error, Owner, TEXT("The source model is corrupted! Section index buffer is invalid. LOD %"), LODIndex);
+			}
+		}
+
+		if (LODModel->NumVertices != SectionsVerticeNum)
+		{
+			UE_ASSET_LOG(LogSkeletalMesh, Error, Owner, TEXT("The source model is corrupted! Total sections vertice count is different from source model vertice count. LOD %"), LODIndex);
+		}
+		if ((LODModel->IndexBuffer.Num() / 3) != SectionsTriangleNum)
+		{
+			UE_ASSET_LOG(LogSkeletalMesh, Error, Owner, TEXT("The source model is corrupted! Total sections triangle count is different from source model triangle count (index count divide by 3). LOD %"), LODIndex);
+		}
+	}
+}
+
 FString FSkeletalMeshRenderData::GetDerivedDataKey(const ITargetPlatform* TargetPlatform, USkeletalMesh* Owner)
 {
 	return BuildSkeletalMeshDerivedDataKey(TargetPlatform, Owner);
@@ -156,6 +207,19 @@ void FSkeletalMeshRenderData::Cache(const ITargetPlatform* TargetPlatform, USkel
 
 	check(LODRenderData.Num() == 0); // Should only be called on new, empty RenderData
 
+	auto SerializeLodModelDdcData = [&Owner](FSkeletalMeshLODModel* LODModel, FArchive& Ar)
+	{
+		//Make sure we add everything FSkeletalMeshLODModel got modified by the skeletalmesh builder
+		Ar << LODModel->Sections;
+		Ar << LODModel->NumVertices;
+		Ar << LODModel->NumTexCoords;
+		Ar << LODModel->IndexBuffer;
+		Ar << LODModel->ActiveBoneIndices;
+		Ar << LODModel->RequiredBones;
+		Ar << LODModel->MeshToImportVertexMap;
+		Ar << LODModel->MaxImportVertex;
+		LODModel->RawPointIndices.Serialize(Ar, Owner);
+	};
 
 	{
 		COOK_STAT(auto Timer = SkeletalMeshCookStats::UsageStats.TimeSyncWork());
@@ -225,15 +289,7 @@ void FSkeletalMeshRenderData::Cache(const ITargetPlatform* TargetPlatform, USkel
 				for (int32 LODIndex = 0; LODIndex < SkelMeshModel->LODModels.Num(); LODIndex++)
 				{
 					FSkeletalMeshLODModel* LODModel = &(SkelMeshModel->LODModels[LODIndex]);
-					int32 SectionNum = 0;
-					Ar << SectionNum;
-					LODModel->Sections.Reset(SectionNum);
-					LODModel->Sections.AddDefaulted(SectionNum);
-					for (int32 SectionIndex = 0; SectionIndex < LODModel->Sections.Num(); ++SectionIndex)
-					{
-						Ar << LODModel->Sections[SectionIndex];
-					}
-					Ar << LODModel->MeshToImportVertexMap;
+					SerializeLodModelDdcData(LODModel, Ar);
 					LODModel->SyncronizeUserSectionsDataArray();
 				}
 			}
@@ -334,13 +390,7 @@ void FSkeletalMeshRenderData::Cache(const ITargetPlatform* TargetPlatform, USkel
 				for (int32 LODIndex = 0; LODIndex < SkelMeshModel->LODModels.Num(); LODIndex++)
 				{
 					FSkeletalMeshLODModel* LODModel = &(SkelMeshModel->LODModels[LODIndex]);
-					int32 SectionNum = LODModel->Sections.Num();
-					Ar << SectionNum;
-					for (int32 SectionIndex = 0; SectionIndex < LODModel->Sections.Num(); ++SectionIndex)
-					{
-						Ar << LODModel->Sections[SectionIndex];
-					}
-					Ar << LODModel->MeshToImportVertexMap;
+					SerializeLodModelDdcData(LODModel, Ar);
 				}
 			}
 
@@ -382,6 +432,7 @@ void FSkeletalMeshRenderData::Cache(const ITargetPlatform* TargetPlatform, USkel
 			COOK_STAT(Timer.AddMiss(DerivedData.Num()));
 		}
 	}
+	VerifyAllLodSkeletalMeshModelIntegrity(Owner);
 }
 
 void FSkeletalMeshRenderData::SyncUVChannelData(const TArray<FSkeletalMaterial>& ObjectData)
