World Partition - Do not auto load last loaded cells when running from a commandlet

diff --git a/Engine/Source/Runtime/Engine/Private/WorldPartition/WorldPartition.cpp b/Engine/Source/Runtime/Engine/Private/WorldPartition/WorldPartition.cpp
--- a/Engine/Source/Runtime/Engine/Private/WorldPartition/WorldPartition.cpp
+++ b/Engine/Source/Runtime/Engine/Private/WorldPartition/WorldPartition.cpp
@@ -86,16 +86,20 @@ struct FWorldPartionCellUpdateContext
 					WorldPartition->WorldPartitionEditor->Refresh();
 				}
 
-				TArray<FVector> EditorGridLastLoadedCells;
-				WorldPartition->EditorHash->ForEachCell([this, &EditorGridLastLoadedCells](UWorldPartitionEditorCell* Cell)
+				// Save last loaded cells settings (ignore while running commandlets)
+				if (!IsRunningCommandlet())
 				{
-					if ((Cell != WorldPartition->EditorHash->GetAlwaysLoadedCell()) && Cell->bLoaded)
+					TArray<FVector> EditorGridLastLoadedCells;
+					WorldPartition->EditorHash->ForEachCell([this, &EditorGridLastLoadedCells](UWorldPartitionEditorCell* Cell)
 					{
-						EditorGridLastLoadedCells.Add(Cell->Bounds.GetCenter());
-					}
-				});
+						if ((Cell != WorldPartition->EditorHash->GetAlwaysLoadedCell()) && Cell->bLoaded)
+						{
+							EditorGridLastLoadedCells.Add(Cell->Bounds.GetCenter());
+						}
+					});
 
-				GetMutableDefault<UWorldPartitionEditorPerProjectUserSettings>()->SetEditorGridLoadedCells(WorldPartition->GetWorld(), EditorGridLastLoadedCells);
+					GetMutableDefault<UWorldPartitionEditorPerProjectUserSettings>()->SetEditorGridLoadedCells(WorldPartition->GetWorld(), EditorGridLastLoadedCells);
+				}
 			}
 		}
 	}
@@ -254,35 +258,40 @@ void UWorldPartition::Initialize(UWorld* InWorld, const FTransform& InTransform)
 			// Load the always loaded cell, don't call LoadCells to avoid creating a transaction
 			UpdateLoadingEditorCell(EditorHash->GetAlwaysLoadedCell(), true);
 
-			// Autoload all cells if the world is smaller than the project setting's value
-			IWorldPartitionEditorModule& WorldPartitionEditorModule = FModuleManager::LoadModuleChecked<IWorldPartitionEditorModule>("WorldPartitionEditor");
-			const float AutoCellLoadingMaxWorldSize = WorldPartitionEditorModule.GetAutoCellLoadingMaxWorldSize();
-			FVector WorldSize = GetWorldBounds().GetSize();
-			const bool bItsASmallWorld = WorldSize.X <= AutoCellLoadingMaxWorldSize && WorldSize.Y <= AutoCellLoadingMaxWorldSize && WorldSize.Z <= AutoCellLoadingMaxWorldSize;
-			
-			// When loading a subworld, load all actors
-			const bool bWorldIsSubPartition = !IsMainWorldPartition();
-
-			if (bWorldIsSubPartition || bItsASmallWorld)
+			// Load more cells depending on the user's settings
+			// Skipped when running from a commandlet
+			if (!IsRunningCommandlet())
 			{
-				EditorHash->ForEachCell([this](UWorldPartitionEditorCell* Cell)
-				{
-					UpdateLoadingEditorCell(Cell, true);
-				});
-			}
-
-			// Load last loaded cells
-			if (WorldPartitionEditorModule.GetEnableLoadingOfLastLoadedCells())
-			{
-				const TArray<FVector>& EditorGridLastLoadedCells = GetMutableDefault<UWorldPartitionEditorPerProjectUserSettings>()->GetEditorGridLoadedCells(InWorld);
+				// Autoload all cells if the world is smaller than the project setting's value
+				IWorldPartitionEditorModule& WorldPartitionEditorModule = FModuleManager::LoadModuleChecked<IWorldPartitionEditorModule>("WorldPartitionEditor");
+				const float AutoCellLoadingMaxWorldSize = WorldPartitionEditorModule.GetAutoCellLoadingMaxWorldSize();
+				FVector WorldSize = GetWorldBounds().GetSize();
+				const bool bItsASmallWorld = WorldSize.X <= AutoCellLoadingMaxWorldSize && WorldSize.Y <= AutoCellLoadingMaxWorldSize && WorldSize.Z <= AutoCellLoadingMaxWorldSize;
+			
+				// When loading a subworld, load all actors
+				const bool bWorldIsSubPartition = !IsMainWorldPartition();
 
-				for (const FVector& EditorGridLastLoadedCellCenter : EditorGridLastLoadedCells)
+				if (bWorldIsSubPartition || bItsASmallWorld)
 				{
-					EditorHash->ForEachIntersectingCell(FBox(&EditorGridLastLoadedCellCenter, 1), [this](UWorldPartitionEditorCell* Cell)
+					EditorHash->ForEachCell([this](UWorldPartitionEditorCell* Cell)
 					{
 						UpdateLoadingEditorCell(Cell, true);
 					});
 				}
+
+				// Load last loaded cells
+				if (WorldPartitionEditorModule.GetEnableLoadingOfLastLoadedCells())
+				{
+					const TArray<FVector>& EditorGridLastLoadedCells = GetMutableDefault<UWorldPartitionEditorPerProjectUserSettings>()->GetEditorGridLoadedCells(InWorld);
+
+					for (const FVector& EditorGridLastLoadedCellCenter : EditorGridLastLoadedCells)
+					{
+						EditorHash->ForEachIntersectingCell(FBox(&EditorGridLastLoadedCellCenter, 1), [this](UWorldPartitionEditorCell* Cell)
+						{
+							UpdateLoadingEditorCell(Cell, true);
+						});
+					}
+				}
 			}
 		}
 	}
