When running commandlet, UWorldPartition::UpdateEditorCells now skips detection of ModifiedPackages.

diff --git a/Engine/Source/Runtime/Engine/Private/WorldPartition/WorldPartition.cpp b/Engine/Source/Runtime/Engine/Private/WorldPartition/WorldPartition.cpp
--- a/Engine/Source/Runtime/Engine/Private/WorldPartition/WorldPartition.cpp
+++ b/Engine/Source/Runtime/Engine/Private/WorldPartition/WorldPartition.cpp
@@ -580,7 +580,9 @@ bool UWorldPartition::UpdateEditorCells(TFunctionRef<bool(TArray<UWorldPartition
 
 	// Make sure we save modified actor packages before unloading
 	FEditorFileUtils::EPromptReturnCode RetCode = FEditorFileUtils::PR_Success;
-	if (ModifiedPackages.Num())
+	// Skip this when running commandlet because MarkPackageDirty() is allowed to dirty packages at loading when running a commandlet.
+	// Usually, the main reason actor packages are dirtied is when RerunConstructionScripts is called on the actor when it is added to the level.
+	if (ModifiedPackages.Num() && !IsRunningCommandlet())
 	{
 		const bool bCheckDirty = false;
 		const bool bAlreadyCheckedOut = false;
