Added graceful handling of out-of-bounds accesses in FStaticMeshInstanceData. Never corrupt memory and don't crash, but report an error and break when debugger is attached.

diff --git a/Engine/Source/Runtime/Engine/Public/StaticMeshResources.h b/Engine/Source/Runtime/Engine/Public/StaticMeshResources.h
--- a/Engine/Source/Runtime/Engine/Public/StaticMeshResources.h
+++ b/Engine/Source/Runtime/Engine/Public/StaticMeshResources.h
@@ -1677,49 +1677,69 @@ private:
 	{
 		FInstanceTransformMatrix<T>* ElementData = reinterpret_cast<FInstanceTransformMatrix<T>*>(InstanceTransformDataPtr);
 		uint32 CurrentSize = InstanceTransformData->Num() * InstanceTransformData->GetStride();
-		check((void*)((&ElementData[InstanceIndex]) + 1) <= (void*)(InstanceTransformDataPtr + CurrentSize));
-		check((void*)((&ElementData[InstanceIndex]) + 0) >= (void*)(InstanceTransformDataPtr));
-		
-		Transform[0][0] = ElementData[InstanceIndex].InstanceTransform1[0];
-		Transform[0][1] = ElementData[InstanceIndex].InstanceTransform1[1];
-		Transform[0][2] = ElementData[InstanceIndex].InstanceTransform1[2];
-		Transform[0][3] = ElementData[InstanceIndex].InstanceTransform1[3];
+
+		if (ensure((void*)((&ElementData[InstanceIndex]) + 1) <= (void*)(InstanceTransformDataPtr + CurrentSize))
+			&& ensure((void*)((&ElementData[InstanceIndex]) + 0) >= (void*)(InstanceTransformDataPtr)))
+		{		
+			Transform[0][0] = ElementData[InstanceIndex].InstanceTransform1[0];
+			Transform[0][1] = ElementData[InstanceIndex].InstanceTransform1[1];
+			Transform[0][2] = ElementData[InstanceIndex].InstanceTransform1[2];
+			Transform[0][3] = ElementData[InstanceIndex].InstanceTransform1[3];
 		
-		Transform[1][0] = ElementData[InstanceIndex].InstanceTransform2[0];
-		Transform[1][1] = ElementData[InstanceIndex].InstanceTransform2[1];
-		Transform[1][2] = ElementData[InstanceIndex].InstanceTransform2[2];
-		Transform[1][3] = ElementData[InstanceIndex].InstanceTransform2[3];
+			Transform[1][0] = ElementData[InstanceIndex].InstanceTransform2[0];
+			Transform[1][1] = ElementData[InstanceIndex].InstanceTransform2[1];
+			Transform[1][2] = ElementData[InstanceIndex].InstanceTransform2[2];
+			Transform[1][3] = ElementData[InstanceIndex].InstanceTransform2[3];
 		
-		Transform[2][0] = ElementData[InstanceIndex].InstanceTransform3[0];
-		Transform[2][1] = ElementData[InstanceIndex].InstanceTransform3[1];
-		Transform[2][2] = ElementData[InstanceIndex].InstanceTransform3[2];
-		Transform[2][3] = ElementData[InstanceIndex].InstanceTransform3[3];
+			Transform[2][0] = ElementData[InstanceIndex].InstanceTransform3[0];
+			Transform[2][1] = ElementData[InstanceIndex].InstanceTransform3[1];
+			Transform[2][2] = ElementData[InstanceIndex].InstanceTransform3[2];
+			Transform[2][3] = ElementData[InstanceIndex].InstanceTransform3[3];
+		}
+		else
+		{
+			Transform[0] = FVector4(1.0f, 0.0f, 0.0f, 0.0f);
+			Transform[1] = FVector4(0.0f, 1.0f, 0.0f, 0.0f);
+			Transform[2] = FVector4(0.0f, 0.0f, 1.0f, 0.0f);
+		}
 	}
 
 	FORCEINLINE_DEBUGGABLE void GetInstanceOriginInternal(int32 InstanceIndex, FVector4 &Origin) const
 	{
 		FVector4* ElementData = reinterpret_cast<FVector4*>(InstanceOriginDataPtr);
 		uint32 CurrentSize = InstanceOriginData->Num() * InstanceOriginData->GetStride();
-		check((void*)((&ElementData[InstanceIndex]) + 1) <= (void*)(InstanceOriginDataPtr + CurrentSize));
-		check((void*)((&ElementData[InstanceIndex]) + 0) >= (void*)(InstanceOriginDataPtr));
 
-		Origin = ElementData[InstanceIndex];
+		if (ensure((void*)((&ElementData[InstanceIndex]) + 1) <= (void*)(InstanceOriginDataPtr + CurrentSize))
+			&& ensure((void*)((&ElementData[InstanceIndex]) + 0) >= (void*)(InstanceOriginDataPtr)))
+		{
+			Origin = ElementData[InstanceIndex];
+		}
+		else
+		{
+			Origin = FVector4(0.0f, 0.0f, 0.0f, 0.0f);
+		}
 	}
 
 	FORCEINLINE_DEBUGGABLE void GetInstanceLightMapDataInternal(int32 InstanceIndex, FVector4 &LightmapData) const
 	{
 		FInstanceLightMapVector* ElementData = reinterpret_cast<FInstanceLightMapVector*>(InstanceLightmapDataPtr);
 		uint32 CurrentSize = InstanceLightmapData->Num() * InstanceLightmapData->GetStride();
-		check((void*)((&ElementData[InstanceIndex]) + 1) <= (void*)(InstanceLightmapDataPtr + CurrentSize));
-		check((void*)((&ElementData[InstanceIndex]) + 0) >= (void*)(InstanceLightmapDataPtr));
 
-		LightmapData = FVector4
-		(
-			float(ElementData[InstanceIndex].InstanceLightmapAndShadowMapUVBias[0]) / 32767.0f, 
-			float(ElementData[InstanceIndex].InstanceLightmapAndShadowMapUVBias[1]) / 32767.0f,
-			float(ElementData[InstanceIndex].InstanceLightmapAndShadowMapUVBias[2]) / 32767.0f,
-			float(ElementData[InstanceIndex].InstanceLightmapAndShadowMapUVBias[3]) / 32767.0f
-		);
+		if (ensure((void*)((&ElementData[InstanceIndex]) + 1) <= (void*)(InstanceLightmapDataPtr + CurrentSize))
+			&& ensure((void*)((&ElementData[InstanceIndex]) + 0) >= (void*)(InstanceLightmapDataPtr)))
+		{
+			LightmapData = FVector4
+			(
+				float(ElementData[InstanceIndex].InstanceLightmapAndShadowMapUVBias[0]) / 32767.0f, 
+				float(ElementData[InstanceIndex].InstanceLightmapAndShadowMapUVBias[1]) / 32767.0f,
+				float(ElementData[InstanceIndex].InstanceLightmapAndShadowMapUVBias[2]) / 32767.0f,
+				float(ElementData[InstanceIndex].InstanceLightmapAndShadowMapUVBias[3]) / 32767.0f
+			);
+		}
+		else
+		{
+			LightmapData = FVector4(0.0f, 0.0f, 0.0f, 0.0f);
+		}
 	}
 
 	FORCEINLINE_DEBUGGABLE void GetInstanceCustomDataInternal(int32 InstanceIndex, TArray<float>& CustomData) const
@@ -1733,10 +1753,11 @@ private:
 		{
 			int32 CustomDataIndex = NumCustomDataFloats * InstanceIndex + i;
 			
-			check((void*)((&ElementData[CustomDataIndex]) + 1) <= (void*)(InstanceCustomDataPtr + CurrentSize));
-			check((void*)((&ElementData[CustomDataIndex]) + 0) >= (void*)(InstanceCustomDataPtr));
-
-			CustomData[i] = ElementData[CustomDataIndex];
+			if (ensure((void*)((&ElementData[CustomDataIndex]) + 1) <= (void*)(InstanceCustomDataPtr + CurrentSize))
+				&& ensure((void*)((&ElementData[CustomDataIndex]) + 0) >= (void*)(InstanceCustomDataPtr)))
+			{
+				CustomData[i] = ElementData[CustomDataIndex];
+			}
 		}
 	}
 
@@ -1745,46 +1766,54 @@ private:
 	{
 		FInstanceTransformMatrix<T>* ElementData = reinterpret_cast<FInstanceTransformMatrix<T>*>(InstanceTransformDataPtr);
 		uint32 CurrentSize = InstanceTransformData->Num() * InstanceTransformData->GetStride();
-		check((void*)((&ElementData[InstanceIndex]) + 1) <= (void*)(InstanceTransformDataPtr + CurrentSize));
-		check((void*)((&ElementData[InstanceIndex]) + 0) >= (void*)(InstanceTransformDataPtr));
-
-		ElementData[InstanceIndex].InstanceTransform1[0] = Transform[0][0];
-		ElementData[InstanceIndex].InstanceTransform1[1] = Transform[0][1];
-		ElementData[InstanceIndex].InstanceTransform1[2] = Transform[0][2];
-		ElementData[InstanceIndex].InstanceTransform1[3] = Transform[0][3];
 
-		ElementData[InstanceIndex].InstanceTransform2[0] = Transform[1][0];
-		ElementData[InstanceIndex].InstanceTransform2[1] = Transform[1][1];
-		ElementData[InstanceIndex].InstanceTransform2[2] = Transform[1][2];
-		ElementData[InstanceIndex].InstanceTransform2[3] = Transform[1][3];
-
-		ElementData[InstanceIndex].InstanceTransform3[0] = Transform[2][0];
-		ElementData[InstanceIndex].InstanceTransform3[1] = Transform[2][1];
-		ElementData[InstanceIndex].InstanceTransform3[2] = Transform[2][2];
-		ElementData[InstanceIndex].InstanceTransform3[3] = Transform[2][3];
+		if (ensure((void*)((&ElementData[InstanceIndex]) + 1) <= (void*)(InstanceTransformDataPtr + CurrentSize))
+			&& ensure((void*)((&ElementData[InstanceIndex]) + 0) >= (void*)(InstanceTransformDataPtr)))
+		{
+			ElementData[InstanceIndex].InstanceTransform1[0] = Transform[0][0];
+			ElementData[InstanceIndex].InstanceTransform1[1] = Transform[0][1];
+			ElementData[InstanceIndex].InstanceTransform1[2] = Transform[0][2];
+			ElementData[InstanceIndex].InstanceTransform1[3] = Transform[0][3];
+
+			ElementData[InstanceIndex].InstanceTransform2[0] = Transform[1][0];
+			ElementData[InstanceIndex].InstanceTransform2[1] = Transform[1][1];
+			ElementData[InstanceIndex].InstanceTransform2[2] = Transform[1][2];
+			ElementData[InstanceIndex].InstanceTransform2[3] = Transform[1][3];
+
+			ElementData[InstanceIndex].InstanceTransform3[0] = Transform[2][0];
+			ElementData[InstanceIndex].InstanceTransform3[1] = Transform[2][1];
+			ElementData[InstanceIndex].InstanceTransform3[2] = Transform[2][2];
+			ElementData[InstanceIndex].InstanceTransform3[3] = Transform[2][3];
+		}
 	}
 
 	FORCEINLINE_DEBUGGABLE void SetInstanceOriginInternal(int32 InstanceIndex, const FVector4& Origin) const
 	{
 		FVector4* ElementData = reinterpret_cast<FVector4*>(InstanceOriginDataPtr);
 		uint32 CurrentSize = InstanceOriginData->Num() * InstanceOriginData->GetStride();
-		checkf((void*)((&ElementData[InstanceIndex]) + 1) <= (void*)(InstanceOriginDataPtr + CurrentSize), TEXT("OOB Instance Set Under: %i, %u, %p, %p"), InstanceIndex, CurrentSize, &ElementData, InstanceOriginDataPtr);
-		checkf((void*)((&ElementData[InstanceIndex]) + 0) >= (void*)(InstanceOriginDataPtr), TEXT("OOB Instance Set: %i, %u, %p, %p"), InstanceIndex, CurrentSize, &ElementData, InstanceOriginDataPtr);
 
-		ElementData[InstanceIndex] = Origin;
+		if (ensureMsgf((void*)((&ElementData[InstanceIndex]) + 1) <= (void*)(InstanceOriginDataPtr + CurrentSize), TEXT("OOB Instance Set Under: %i, %u, %p, %p"), InstanceIndex, CurrentSize, &ElementData, InstanceOriginDataPtr)
+			&& ensureMsgf((void*)((&ElementData[InstanceIndex]) + 0) >= (void*)(InstanceOriginDataPtr), TEXT("OOB Instance Set: %i, %u, %p, %p"), InstanceIndex, CurrentSize, &ElementData, InstanceOriginDataPtr))
+		{
+			ElementData[InstanceIndex] = Origin;
+		}
 	}
 
 	FORCEINLINE_DEBUGGABLE void SetInstanceLightMapDataInternal(int32 InstanceIndex, const FVector4& LightmapData) const
 	{
 		FInstanceLightMapVector* ElementData = reinterpret_cast<FInstanceLightMapVector*>(InstanceLightmapDataPtr);
 		uint32 CurrentSize = InstanceLightmapData->Num() * InstanceLightmapData->GetStride();
-		check((void*)((&ElementData[InstanceIndex]) + 1) <= (void*)(InstanceLightmapDataPtr + CurrentSize));
-		check((void*)((&ElementData[InstanceIndex]) + 0) >= (void*)(InstanceLightmapDataPtr));
+		
+		if (ensure((void*)((&ElementData[InstanceIndex]) + 1) <= (void*)(InstanceLightmapDataPtr + CurrentSize))
+			&& ensure((void*)((&ElementData[InstanceIndex]) + 0) >= (void*)(InstanceLightmapDataPtr)))
+		{
+
+			ElementData[InstanceIndex].InstanceLightmapAndShadowMapUVBias[0] = FMath::Clamp<int32>(FMath::TruncToInt(LightmapData.X * 32767.0f), MIN_int16, MAX_int16);
+			ElementData[InstanceIndex].InstanceLightmapAndShadowMapUVBias[1] = FMath::Clamp<int32>(FMath::TruncToInt(LightmapData.Y * 32767.0f), MIN_int16, MAX_int16);
+			ElementData[InstanceIndex].InstanceLightmapAndShadowMapUVBias[2] = FMath::Clamp<int32>(FMath::TruncToInt(LightmapData.Z * 32767.0f), MIN_int16, MAX_int16);
+			ElementData[InstanceIndex].InstanceLightmapAndShadowMapUVBias[3] = FMath::Clamp<int32>(FMath::TruncToInt(LightmapData.W * 32767.0f), MIN_int16, MAX_int16);
+		}
 
-		ElementData[InstanceIndex].InstanceLightmapAndShadowMapUVBias[0] = FMath::Clamp<int32>(FMath::TruncToInt(LightmapData.X * 32767.0f), MIN_int16, MAX_int16);
-		ElementData[InstanceIndex].InstanceLightmapAndShadowMapUVBias[1] = FMath::Clamp<int32>(FMath::TruncToInt(LightmapData.Y * 32767.0f), MIN_int16, MAX_int16);
-		ElementData[InstanceIndex].InstanceLightmapAndShadowMapUVBias[2] = FMath::Clamp<int32>(FMath::TruncToInt(LightmapData.Z * 32767.0f), MIN_int16, MAX_int16);
-		ElementData[InstanceIndex].InstanceLightmapAndShadowMapUVBias[3] = FMath::Clamp<int32>(FMath::TruncToInt(LightmapData.W * 32767.0f), MIN_int16, MAX_int16);
 	}
 
 	FORCEINLINE_DEBUGGABLE void SetInstanceCustomDataInternal(int32 InstanceIndex, int32 DataIndex, float CustomData)
@@ -1799,10 +1828,12 @@ private:
 
 		const int32 CustomDataIndex = NumCustomDataFloats * InstanceIndex + DataIndex;
 
-		check((void*)((&ElementData[CustomDataIndex]) + 1) <= (void*)(InstanceCustomDataPtr + CurrentSize));
-		check((void*)((&ElementData[CustomDataIndex]) + 0) >= (void*)(InstanceCustomDataPtr));
+		if (ensure((void*)((&ElementData[CustomDataIndex]) + 1) <= (void*)(InstanceCustomDataPtr + CurrentSize))
+			&& ensure((void*)((&ElementData[CustomDataIndex]) + 0) >= (void*)(InstanceCustomDataPtr)))
+		{
+			ElementData[CustomDataIndex] = CustomData;
+		}
 
-		ElementData[CustomDataIndex] = CustomData;
 	}
 
 	void AllocateBuffers(int32 InNumInstances, EResizeBufferFlags BufferFlags = EResizeBufferFlags::None)
