Fixed invalid check in LevelInstanceLevelStreaming (Level's bAlreadyMovedActors flag is not reset when recycling an already loaded package).

diff --git a/Engine/Source/Runtime/Engine/Private/LevelInstance/LevelInstanceLevelStreaming.cpp b/Engine/Source/Runtime/Engine/Private/LevelInstance/LevelInstanceLevelStreaming.cpp
--- a/Engine/Source/Runtime/Engine/Private/LevelInstance/LevelInstanceLevelStreaming.cpp
+++ b/Engine/Source/Runtime/Engine/Private/LevelInstance/LevelInstanceLevelStreaming.cpp
@@ -154,10 +154,10 @@ void ULevelStreamingLevelInstance::SetLoadedLevel(ULevel* Level)
 
 	if (ULevel* NewLoadedLevel = GetLoadedLevel())
 	{
-		check(!NewLoadedLevel->bAlreadyMovedActors);
-		if (AWorldSettings* WorldSettings = NewLoadedLevel->GetWorldSettings())
+		if (!NewLoadedLevel->bAlreadyMovedActors)
 		{
+			AWorldSettings* WorldSettings = NewLoadedLevel->GetWorldSettings();
+			check(WorldSettings);
 			LevelTransform.AddToTranslation(WorldSettings->LevelInstancePivotOffset);
 		}
 	}
-}
