Fixed AActor::SupportsExternalPackaging() : now also returns false for classes marked transient (CLASS_Transient).

diff --git a/Engine/Source/Runtime/Engine/Private/Actor.cpp b/Engine/Source/Runtime/Engine/Private/Actor.cpp
--- a/Engine/Source/Runtime/Engine/Private/Actor.cpp
+++ b/Engine/Source/Runtime/Engine/Private/Actor.cpp
@@ -3138,7 +3138,7 @@ bool AActor::IsSelectedInEditor() const
 
 bool AActor::SupportsExternalPackaging() const
 {
-	if (HasAllFlags(RF_Transient))
+	if (HasAllFlags(RF_Transient) || GetClass()->HasAllClassFlags(CLASS_Transient))
 	{
 		return false;
 	}
