Buildfix, missing last line in file.

diff --git a/Engine/Source/Runtime/Engine/Private/LevelInstance/LevelInstanceLevelStreaming.cpp b/Engine/Source/Runtime/Engine/Private/LevelInstance/LevelInstanceLevelStreaming.cpp
--- a/Engine/Source/Runtime/Engine/Private/LevelInstance/LevelInstanceLevelStreaming.cpp
+++ b/Engine/Source/Runtime/Engine/Private/LevelInstance/LevelInstanceLevelStreaming.cpp
@@ -161,3 +161,4 @@ void ULevelStreamingLevelInstance::SetLoadedLevel(ULevel* Level)
 			LevelTransform.AddToTranslation(WorldSettings->LevelInstancePivotOffset);
 		}
 	}
+}
\ No newline at end of file
