CIS: fix for Build DDC warning in Sound Wave

diff --git a/Engine/Source/Runtime/Engine/Private/SoundWave.cpp b/Engine/Source/Runtime/Engine/Private/SoundWave.cpp
--- a/Engine/Source/Runtime/Engine/Private/SoundWave.cpp
+++ b/Engine/Source/Runtime/Engine/Private/SoundWave.cpp
@@ -2425,39 +2425,27 @@ void USoundWave::UpdatePlatformData()
 {
 	if (IsStreaming(nullptr))
 	{
-		if (InternalProxy.IsValid())
+		// Make sure there are no pending requests in flight.
+		while (IStreamingManager::Get().GetAudioStreamingManager().IsStreamingInProgress(GetThisAsProxy()))
 		{
-			// Make sure there are no pending requests in flight.
-			while (IStreamingManager::Get().GetAudioStreamingManager().IsStreamingInProgress(*InternalProxy))
-			{
-				// Give up timeslice.
-				FPlatformProcess::Sleep(0);
-			}
+			// Give up timeslice.
+			FPlatformProcess::Sleep(0);
+		}
 
 #if WITH_EDITORONLY_DATA
-			// Temporarily remove from streaming manager to release currently used data chunks
-			IStreamingManager::Get().GetAudioStreamingManager().RemoveStreamingSoundWave(*InternalProxy);
-
-			// Recache platform data if the source has changed.
-			CachePlatformData();
+		// Temporarily remove from streaming manager to release currently used data chunks
+		IStreamingManager::Get().GetAudioStreamingManager().RemoveStreamingSoundWave(GetThisAsProxy());
 
-			// Release proxy holding reference to old data and replace with a new one
-			InternalProxy = CreateSoundWaveProxy();
+		// Recache platform data if the source has changed.
+		CachePlatformData();
 
-			if (InternalProxy.IsValid())
-			{
-				// Add back to the streaming manager to reload first chunk
-				IStreamingManager::Get().GetAudioStreamingManager().AddStreamingSoundWave(*InternalProxy);
-			}
+		// Add back to the streaming manager to reload first chunk
+		IStreamingManager::Get().GetAudioStreamingManager().AddStreamingSoundWave(GetThisAsProxy());
 #endif
-		}
 	}
 	else
 	{
-		if (InternalProxy.IsValid())
-		{
-			IStreamingManager::Get().GetAudioStreamingManager().RemoveStreamingSoundWave(*InternalProxy);
-		}
+		IStreamingManager::Get().GetAudioStreamingManager().RemoveStreamingSoundWave(GetThisAsProxy());
 	}
 }
 
