Deprecated the StaticMeshOwner property in FStaticMeshSourceModel.
It is no longer necessary to maintain it, as the owner is the same as the outer of the UStaticMeshDescriptionBulkData object, which must always be valid.
Implementing StaticMeshOwner as a UPROPERTY led to an issue when deleting a static mesh asset, as its reference to the UStaticMesh* owner was being nulled, but then the PostEditChangeProperty on the StaticMesh was attempting to Build() while no longer having a valid owner. Removing it completely simplifies everything.
Fixed a bug in which UStaticMesh::GetStaticMeshDescription() was not attempting to cache a mesh description if it wasn't valid.
Made the HiResMeshDescription bulk data a transactable object.

diff --git a/Engine/Source/Runtime/Engine/Classes/Engine/StaticMeshSourceData.h b/Engine/Source/Runtime/Engine/Classes/Engine/StaticMeshSourceData.h
--- a/Engine/Source/Runtime/Engine/Classes/Engine/StaticMeshSourceData.h
+++ b/Engine/Source/Runtime/Engine/Classes/Engine/StaticMeshSourceData.h
@@ -65,11 +65,11 @@ public:
 
 #if WITH_EDITORONLY_DATA
 	/*
-	 * The staticmesh owner of this source model. We need the SM to be able to convert between MeshDesription and RawMesh.
-	 * RawMesh use int32 material index and MeshDescription use FName material slot name.
-	 * This member is fill in the PostLoad of the static mesh.
+	 * We used to maintain the UStaticMesh which owned the source model here.
+	 * Now this is not necessary as we can get it from the Outer of the StaticMeshDescriptionBulkData.
+	 * This means less to maintain
 	 */
-	UPROPERTY(Transient)
+	UE_DEPRECATED(5.0, "StaticMeshOwner is now deprecated.")
 	class UStaticMesh* StaticMeshOwner;
 
 	/**
@@ -82,6 +82,11 @@ public:
 #endif
 
 #if WITH_EDITOR
+	/**
+	 * Returns the static mesh which owns this source model.
+	 */
+	ENGINE_API UStaticMesh* GetStaticMeshOwner() const;
+
 	/**
 	 * Create a new MeshDescription object
 	 */
diff --git a/Engine/Source/Runtime/Engine/Private/StaticMesh.cpp b/Engine/Source/Runtime/Engine/Private/StaticMesh.cpp
--- a/Engine/Source/Runtime/Engine/Private/StaticMesh.cpp
+++ b/Engine/Source/Runtime/Engine/Private/StaticMesh.cpp
@@ -2888,8 +2888,8 @@ UStaticMesh::UStaticMesh(const FObjectInitializer& ObjectInitializer)
 	ImportVersion = EImportStaticMeshVersion::BeforeImportStaticMeshVersionWasAdded;
 	LODForOccluderMesh = -1;
 	NumStreamedLODs.Default = -1;
-	GetHiResSourceModel().StaticMeshOwner = this;
 	GetHiResSourceModel().StaticMeshDescriptionBulkData = CreateEditorOnlyDefaultSubobject<UStaticMeshDescriptionBulkData>(TEXT("HiResMeshDescription"));
+	GetHiResSourceModel().StaticMeshDescriptionBulkData->SetFlags(RF_Transactional);
 #endif // #if WITH_EDITORONLY_DATA
 	SetLightMapResolution(4);
 	SetMinLOD(0);
@@ -4924,15 +4924,6 @@ void UStaticMesh::Serialize(FArchive& Ar)
 			SrcModel.SerializeBulkData(Ar, this);
 		}
 
-		if (Ar.IsLoading())
-		{
-			GetHiResSourceModel().StaticMeshOwner = this;
-			if (GetHiResSourceModel().StaticMeshDescriptionBulkData == nullptr)
-			{
-				GetHiResSourceModel().StaticMeshDescriptionBulkData = NewObject<UStaticMeshDescriptionBulkData>(this, NAME_None, RF_Transactional);
-			}
-		}
-
 		if (Ar.CustomVer(FEditorObjectVersion::GUID) < FEditorObjectVersion::UPropertryForMeshSection)
 		{
 			GetSectionInfoMap().Serialize(Ar);
@@ -5794,6 +5785,7 @@ UStaticMeshDescription* UStaticMesh::GetStaticMeshDescription(int32 LODIndex)
 #if WITH_EDITOR
 	if (LODIndex < GetNumSourceModels())
 	{
+		GetSourceModel(LODIndex).GetOrCacheMeshDescription();
 		return GetSourceModel(LODIndex).GetCachedStaticMeshDescription();
 	}
 #endif
diff --git a/Engine/Source/Runtime/Engine/Private/StaticMeshSourceData.cpp b/Engine/Source/Runtime/Engine/Private/StaticMeshSourceData.cpp
--- a/Engine/Source/Runtime/Engine/Private/StaticMeshSourceData.cpp
+++ b/Engine/Source/Runtime/Engine/Private/StaticMeshSourceData.cpp
@@ -21,11 +21,13 @@ FStaticMeshSourceModel::FStaticMeshSourceModel()
 {
 #if WITH_EDITOR
 	RawMeshBulkData = new FRawMeshBulkData();
-	StaticMeshOwner = nullptr;
 #endif // #if WITH_EDITOR
 	LODDistance_DEPRECATED = 0.0f;
 	ScreenSize.Default = 0.0f;
 #if WITH_EDITORONLY_DATA
+	PRAGMA_DISABLE_DEPRECATION_WARNINGS
+	StaticMeshOwner = nullptr;
+	PRAGMA_ENABLE_DEPRECATION_WARNINGS
 	bImportWithBaseMesh = false;
 	StaticMeshDescriptionBulkData = nullptr;
 #endif
@@ -51,9 +53,6 @@ FStaticMeshSourceModel& FStaticMeshSourceModel::operator=(FStaticMeshSourceModel
 #if WITH_EDITOR
 	RawMeshBulkData = Other.RawMeshBulkData;
 	Other.RawMeshBulkData = nullptr;
-
-	StaticMeshOwner = Other.StaticMeshOwner;
-	Other.StaticMeshOwner = nullptr;
 #endif
 
 #if WITH_EDITORONLY_DATA
@@ -75,16 +74,11 @@ FStaticMeshSourceModel& FStaticMeshSourceModel::operator=(FStaticMeshSourceModel
 
 void FStaticMeshSourceModel::CreateSubObjects(UStaticMesh* InOwner)
 {
-#if WITH_EDITOR
-	check(StaticMeshOwner == nullptr);
-	StaticMeshOwner = InOwner;
-#endif
-
 #if WITH_EDITORONLY_DATA
 	ensure(IsInGameThread());
 	if (StaticMeshDescriptionBulkData == nullptr)
 	{
-		StaticMeshDescriptionBulkData = NewObject<UStaticMeshDescriptionBulkData>(StaticMeshOwner, NAME_None, RF_Transactional);
+		StaticMeshDescriptionBulkData = NewObject<UStaticMeshDescriptionBulkData>(InOwner, NAME_None, RF_Transactional);
 		StaticMeshDescriptionBulkData->AtomicallyClearInternalFlags(EInternalObjectFlags::Async);
 	}
 
@@ -106,17 +100,17 @@ bool FStaticMeshSourceModel::IsRawMeshEmpty() const
 void FStaticMeshSourceModel::LoadRawMesh(FRawMesh& OutRawMesh) const
 {
 	check(RawMeshBulkData->IsEmpty());
+	check(StaticMeshDescriptionBulkData != nullptr);
 
-	// If the RawMesh is empty, consider the possibility that it's a new asset with a valid MeshDescription which needs loading.
-	check(StaticMeshOwner != nullptr);
 	FStaticMeshSourceModel* MutableThis = const_cast<FStaticMeshSourceModel*>(this);
 	if (FMeshDescription* CachedMeshDescription = MutableThis->GetOrCacheMeshDescription())
 	{
 		TMap<FName, int32> MaterialMap;
-		check(StaticMeshOwner != nullptr);
-		for (int32 MaterialIndex = 0; MaterialIndex < StaticMeshOwner->GetStaticMaterials().Num(); ++MaterialIndex)
+		UStaticMesh* StaticMesh = GetStaticMeshOwner();
+		check(StaticMesh != nullptr);
+		for (int32 MaterialIndex = 0; MaterialIndex < StaticMesh->GetStaticMaterials().Num(); ++MaterialIndex)
 		{
-			MaterialMap.Add(StaticMeshOwner->GetStaticMaterials()[MaterialIndex].ImportedMaterialSlotName, MaterialIndex);
+			MaterialMap.Add(StaticMesh->GetStaticMaterials()[MaterialIndex].ImportedMaterialSlotName, MaterialIndex);
 		}
 		FStaticMeshOperations::ConvertToRawMesh(*CachedMeshDescription, OutRawMesh, MaterialMap);
 	}
@@ -146,10 +140,16 @@ void FStaticMeshSourceModel::SaveRawMesh(FRawMesh& InRawMesh, bool /* unused */)
 }
 
 
+UStaticMesh* FStaticMeshSourceModel::GetStaticMeshOwner() const
+{
+	check(StaticMeshDescriptionBulkData != nullptr);
+	return Cast<UStaticMesh>(StaticMeshDescriptionBulkData->GetOuter());
+}
+
+
 bool FStaticMeshSourceModel::LoadMeshDescription(FMeshDescription& OutMeshDescription) const
 {
 	TRACE_CPUPROFILER_EVENT_SCOPE(FStaticMeshSourceModel::LoadMeshDescription);
-	check(StaticMeshOwner != nullptr);
 	check(StaticMeshDescriptionBulkData != nullptr);
 
 	// Ensure default MeshDescription result is empty, with no attributes registered
@@ -175,6 +175,7 @@ bool FStaticMeshSourceModel::CloneMeshDescription(FMeshDescription& OutMeshDescr
 {
 	TRACE_CPUPROFILER_EVENT_SCOPE(FStaticMeshSourceModel::CloneMeshDescription);
 
+	check(StaticMeshDescriptionBulkData != nullptr);
 	if (StaticMeshDescriptionBulkData->HasCachedMeshDescription())
 	{
 		OutMeshDescription = StaticMeshDescriptionBulkData->GetMeshDescription()->GetMeshDescription();
@@ -188,7 +189,7 @@ bool FStaticMeshSourceModel::CloneMeshDescription(FMeshDescription& OutMeshDescr
 FMeshDescription* FStaticMeshSourceModel::GetOrCacheMeshDescription()
 {
 	TRACE_CPUPROFILER_EVENT_SCOPE(FStaticMeshSourceModel::GetMeshDescription);
-	check(StaticMeshOwner != nullptr);
+	check(StaticMeshDescriptionBulkData != nullptr);
 
 	if (!StaticMeshDescriptionBulkData->HasCachedMeshDescription())
 	{
@@ -213,6 +214,7 @@ FMeshDescription* FStaticMeshSourceModel::GetOrCacheMeshDescription()
 
 FMeshDescription* FStaticMeshSourceModel::GetCachedMeshDescription() const
 {
+	check(StaticMeshDescriptionBulkData != nullptr);
 	if (StaticMeshDescriptionBulkData->HasCachedMeshDescription())
 	{
 		return &StaticMeshDescriptionBulkData->GetMeshDescription()->GetMeshDescription();
@@ -224,6 +226,7 @@ FMeshDescription* FStaticMeshSourceModel::GetCachedMeshDescription() const
 
 UStaticMeshDescription* FStaticMeshSourceModel::GetCachedStaticMeshDescription() const
 {
+	check(StaticMeshDescriptionBulkData != nullptr);
 	return Cast<UStaticMeshDescription>(StaticMeshDescriptionBulkData->GetMeshDescription());
 }
 
@@ -247,6 +250,7 @@ bool FStaticMeshSourceModel::IsMeshDescriptionValid() const
 
 FMeshDescription* FStaticMeshSourceModel::CreateMeshDescription()
 {
+	check(StaticMeshDescriptionBulkData != nullptr);
 	return &StaticMeshDescriptionBulkData->CreateMeshDescription()->GetMeshDescription();
 }
 
@@ -276,21 +280,20 @@ void FStaticMeshSourceModel::ClearMeshDescription()
 {
 	TRACE_CPUPROFILER_EVENT_SCOPE(UStaticMesh::ClearMeshDescription);
 
+	check(StaticMeshDescriptionBulkData != nullptr);
 	StaticMeshDescriptionBulkData->RemoveMeshDescription();
 }
 
 
 void FStaticMeshSourceModel::SerializeBulkData(FArchive& Ar, UObject* Owner)
 {
-	// Initialize the StaticMeshOwner
+	// Initialize the StaticMeshDescriptionBulkData
 	if (Ar.IsLoading())
 	{
-		StaticMeshOwner = Cast<UStaticMesh>(Owner);
-
 		// If this was a legacy asset, or is being created for the first time, create a bulkdata UObject wrapper
 		if (StaticMeshDescriptionBulkData == nullptr)
 		{
-			StaticMeshDescriptionBulkData = NewObject<UStaticMeshDescriptionBulkData>(StaticMeshOwner, NAME_None, RF_Transactional);
+			StaticMeshDescriptionBulkData = NewObject<UStaticMeshDescriptionBulkData>(Owner, NAME_None, RF_Transactional);
 		}
 	}
 
@@ -313,7 +316,7 @@ void FStaticMeshSourceModel::SerializeBulkData(FArchive& Ar, UObject* Owner)
 
 				if (bIsValid)
 				{
-					StaticMeshDescriptionBulkData->GetBulkData().Serialize(Ar, StaticMeshOwner);
+					StaticMeshDescriptionBulkData->GetBulkData().Serialize(Ar, Owner);
 				}
 			}
 
@@ -324,8 +327,12 @@ void FStaticMeshSourceModel::SerializeBulkData(FArchive& Ar, UObject* Owner)
 
 void FStaticMeshSourceModel::FillMaterialName(TMap<int32, FName>& OutMaterialMap) const
 {
-	check(StaticMeshOwner != nullptr);
-	const TArray<FStaticMaterial>& StaticMaterials = StaticMeshOwner->GetStaticMaterials();
+	check(StaticMeshDescriptionBulkData != nullptr);
+	
+	UStaticMesh* StaticMesh = GetStaticMeshOwner();
+	check(StaticMesh != nullptr);
+
+	const TArray<FStaticMaterial>& StaticMaterials = StaticMesh->GetStaticMaterials();
 
 	OutMaterialMap.Empty(StaticMaterials.Num());
 
@@ -394,7 +401,8 @@ namespace StaticMeshConvertStats
 void FStaticMeshSourceModel::ConvertRawMesh(int32 LodIndex)
 {
 	check(StaticMeshDescriptionBulkData != nullptr);
-	check(StaticMeshOwner != nullptr);
+	UStaticMesh* StaticMesh = GetStaticMeshOwner();
+	check(StaticMesh != nullptr);
 
 	if (!RawMeshBulkData->IsEmpty() && !StaticMeshDescriptionBulkData->IsBulkDataValid())
 	{
@@ -405,14 +413,14 @@ void FStaticMeshSourceModel::ConvertRawMesh(int32 LodIndex)
 
 			TArray<uint8> DerivedData;
 
-			if (GetDerivedDataCacheRef().GetSynchronous(*MeshDataKey, DerivedData, StaticMeshOwner->GetPathName()))
+			if (GetDerivedDataCacheRef().GetSynchronous(*MeshDataKey, DerivedData, StaticMesh->GetPathName()))
 			{
 				COOK_STAT(Timer.AddHit(DerivedData.Num()));
 
 				// Load from the DDC
 				const bool bIsPersistent = true;
 				FMemoryReader Ar(DerivedData, bIsPersistent);
-				StaticMeshDescriptionBulkData->GetBulkData().Serialize(Ar, StaticMeshOwner);
+				StaticMeshDescriptionBulkData->GetBulkData().Serialize(Ar, StaticMesh);
 
 				check(GetCachedMeshDescription() == nullptr);
 			}
@@ -437,8 +445,8 @@ void FStaticMeshSourceModel::ConvertRawMesh(int32 LodIndex)
 				// Write the DDC cache
 				const bool bIsPersistent = true;
 				FMemoryWriter Ar(DerivedData, bIsPersistent);
-				StaticMeshDescriptionBulkData->GetBulkData().Serialize(Ar, StaticMeshOwner);
-				GetDerivedDataCacheRef().Put(*MeshDataKey, DerivedData, StaticMeshOwner->GetPathName());
+				StaticMeshDescriptionBulkData->GetBulkData().Serialize(Ar, StaticMesh);
+				GetDerivedDataCacheRef().Put(*MeshDataKey, DerivedData, StaticMesh->GetPathName());
 				COOK_STAT(Timer.AddMiss(DerivedData.Num()));
 			}
 		}
