LevelInstance: Prevent Commit of LevelInstance that has a modified child LevelInstance being edited

diff --git a/Engine/Source/Runtime/Engine/Private/LevelInstance/LevelInstanceSubsystem.cpp b/Engine/Source/Runtime/Engine/Private/LevelInstance/LevelInstanceSubsystem.cpp
--- a/Engine/Source/Runtime/Engine/Private/LevelInstance/LevelInstanceSubsystem.cpp
+++ b/Engine/Source/Runtime/Engine/Private/LevelInstance/LevelInstanceSubsystem.cpp
@@ -1290,6 +1290,15 @@ bool ULevelInstanceSubsystem::CanCommitLevelInstance(const ALevelInstance* Level
 		return false;
 	}
 
+	if (HasDirtyChildrenLevelInstances(LevelInstanceActor))
+	{
+		if (OutReason)
+		{
+			*OutReason = LOCTEXT("CanCommitLevelInstanceUnsavedChild", "Level Instance has an unsaved child level instance");
+		}
+		return false;
+	}
+
 	return true;
 }
 
