Initial support for skeletal remapping, which allows you to play animation assets from one skeleton onto the other.

diff --git a/Engine/Source/Editor/AnimGraph/Private/AnimGraphNode_AimOffsetLookAt.cpp b/Engine/Source/Editor/AnimGraph/Private/AnimGraphNode_AimOffsetLookAt.cpp
--- a/Engine/Source/Editor/AnimGraph/Private/AnimGraphNode_AimOffsetLookAt.cpp
+++ b/Engine/Source/Editor/AnimGraph/Private/AnimGraphNode_AimOffsetLookAt.cpp
@@ -174,9 +174,9 @@ void UAnimGraphNode_AimOffsetLookAt::ValidateAnimNodeDuringCompilation(class USk
 	{
 		const USkeleton* BlendSpaceSkeleton = BlendSpaceToCheck->GetSkeleton();
 		if (BlendSpaceSkeleton && // if blend space doesn't have skeleton, it might be due to blend space not loaded yet, @todo: wait with anim blueprint compilation until all assets are loaded?
-			!BlendSpaceSkeleton->IsCompatible(ForSkeleton))
+			!ForSkeleton->IsCompatible(BlendSpaceSkeleton))
 		{
-			MessageLog.Error(TEXT("@@ references blendspace that uses different skeleton @@"), this, BlendSpaceSkeleton);
+			MessageLog.Error(TEXT("@@ references blendspace that uses an incompatible skeleton @@"), this, BlendSpaceSkeleton);
 		}
 
 		// Make sure that the source socket name is a valid one for the skeleton
diff --git a/Engine/Source/Editor/AnimGraph/Private/AnimGraphNode_BlendSpaceEvaluator.cpp b/Engine/Source/Editor/AnimGraph/Private/AnimGraphNode_BlendSpaceEvaluator.cpp
--- a/Engine/Source/Editor/AnimGraph/Private/AnimGraphNode_BlendSpaceEvaluator.cpp
+++ b/Engine/Source/Editor/AnimGraph/Private/AnimGraphNode_BlendSpaceEvaluator.cpp
@@ -126,9 +126,9 @@ void UAnimGraphNode_BlendSpaceEvaluator::ValidateAnimNodeDuringCompilation(class
 	{
 		USkeleton* BlendSpaceSkeleton = BlendSpaceToCheck->GetSkeleton();
 		if (BlendSpaceSkeleton&& // if blend space doesn't have skeleton, it might be due to blend space not loaded yet, @todo: wait with anim blueprint compilation until all assets are loaded?
-			!BlendSpaceSkeleton->IsCompatible(ForSkeleton))
+			!ForSkeleton->IsCompatible(BlendSpaceSkeleton))
 		{
-			MessageLog.Error(TEXT("@@ references blendspace that uses different skeleton @@"), this, BlendSpaceSkeleton);
+			MessageLog.Error(TEXT("@@ references blendspace that uses an incompatible skeleton @@"), this, BlendSpaceSkeleton);
 		}
 	}
 }
diff --git a/Engine/Source/Editor/AnimGraph/Private/AnimGraphNode_BlendSpacePlayer.cpp b/Engine/Source/Editor/AnimGraph/Private/AnimGraphNode_BlendSpacePlayer.cpp
--- a/Engine/Source/Editor/AnimGraph/Private/AnimGraphNode_BlendSpacePlayer.cpp
+++ b/Engine/Source/Editor/AnimGraph/Private/AnimGraphNode_BlendSpacePlayer.cpp
@@ -149,9 +149,9 @@ void UAnimGraphNode_BlendSpacePlayer::ValidateAnimNodeDuringCompilation(class US
 	{
 		USkeleton* BlendSpaceSkeleton = BlendSpaceToCheck->GetSkeleton();
 		if (BlendSpaceSkeleton && // if blend space doesn't have skeleton, it might be due to blend space not loaded yet, @todo: wait with anim blueprint compilation until all assets are loaded?
-			!BlendSpaceSkeleton->IsCompatible(ForSkeleton))
+			!ForSkeleton->IsCompatible(BlendSpaceSkeleton))
 		{
-			MessageLog.Error(TEXT("@@ references blendspace that uses different skeleton @@"), this, BlendSpaceSkeleton);
+			MessageLog.Error(TEXT("@@ references blendspace that uses an incompatible skeleton @@"), this, BlendSpaceSkeleton);
 		}
 	}
 }
diff --git a/Engine/Source/Editor/AnimGraph/Private/AnimGraphNode_LinkedAnimGraphBase.cpp b/Engine/Source/Editor/AnimGraph/Private/AnimGraphNode_LinkedAnimGraphBase.cpp
--- a/Engine/Source/Editor/AnimGraph/Private/AnimGraphNode_LinkedAnimGraphBase.cpp
+++ b/Engine/Source/Editor/AnimGraph/Private/AnimGraphNode_LinkedAnimGraphBase.cpp
@@ -415,8 +415,7 @@ bool UAnimGraphNode_LinkedAnimGraphBase::OnShouldFilterInstanceBlueprint(const F
 	{
 		if (UAnimBlueprint* CurrentBlueprint = Cast<UAnimBlueprint>(GetBlueprint()))
 		{
-			FString TargetSkeletonName = FString::Printf(TEXT("%s'%s'"), *CurrentBlueprint->TargetSkeleton->GetClass()->GetName(), *CurrentBlueprint->TargetSkeleton->GetPathName());
-			if(Result.GetValue() != TargetSkeletonName)
+			if (!CurrentBlueprint->TargetSkeleton->IsCompatibleSkeletonByAssetString(Result.GetValue()))
 			{
 				return true;
 			}
diff --git a/Engine/Source/Editor/AnimGraph/Private/AnimGraphNode_PoseByName.cpp b/Engine/Source/Editor/AnimGraph/Private/AnimGraphNode_PoseByName.cpp
--- a/Engine/Source/Editor/AnimGraph/Private/AnimGraphNode_PoseByName.cpp
+++ b/Engine/Source/Editor/AnimGraph/Private/AnimGraphNode_PoseByName.cpp
@@ -123,10 +123,10 @@ void UAnimGraphNode_PoseByName::ValidateAnimNodeDuringCompilation(class USkeleto
 	else
 	{
 		USkeleton* SeqSkeleton = PoseAssetToCheck->GetSkeleton();
-		if (SeqSkeleton&& // if anim sequence doesn't have skeleton, it might be due to anim sequence not loaded yet, @todo: wait with anim blueprint compilation until all assets are loaded?
-			!SeqSkeleton->IsCompatible(ForSkeleton))
+		if (SeqSkeleton && // if anim sequence doesn't have skeleton, it might be due to anim sequence not loaded yet, @todo: wait with anim blueprint compilation until all assets are loaded?
+			!ForSkeleton->IsCompatible(SeqSkeleton))
 		{
-			MessageLog.Error(TEXT("@@ references sequence that uses different skeleton @@"), this, SeqSkeleton);
+			MessageLog.Error(TEXT("@@ references sequence that uses an incompatible skeleton @@"), this, SeqSkeleton);
 		}
 	}
 }
diff --git a/Engine/Source/Editor/AnimGraph/Private/AnimGraphNode_PoseHandler.cpp b/Engine/Source/Editor/AnimGraph/Private/AnimGraphNode_PoseHandler.cpp
--- a/Engine/Source/Editor/AnimGraph/Private/AnimGraphNode_PoseHandler.cpp
+++ b/Engine/Source/Editor/AnimGraph/Private/AnimGraphNode_PoseHandler.cpp
@@ -31,9 +31,9 @@ void UAnimGraphNode_PoseHandler::ValidateAnimNodeDuringCompilation(USkeleton* Fo
 	{
 		USkeleton* SeqSkeleton = PoseAssetToCheck->GetSkeleton();
 		if (SeqSkeleton && // if PoseAsset doesn't have skeleton, it might be due to PoseAsset not loaded yet, @todo: wait with anim blueprint compilation until all assets are loaded?
-			!SeqSkeleton->IsCompatible(ForSkeleton))
+			!ForSkeleton->IsCompatible(SeqSkeleton))
 		{
-			MessageLog.Error(TEXT("@@ references poseasset that uses different skeleton @@"), this, SeqSkeleton);
+			MessageLog.Error(TEXT("@@ references poseasset that uses an incompatible skeleton @@"), this, SeqSkeleton);
 		}
 	}
 
diff --git a/Engine/Source/Editor/AnimGraph/Private/AnimGraphNode_RotationOffsetBlendSpace.cpp b/Engine/Source/Editor/AnimGraph/Private/AnimGraphNode_RotationOffsetBlendSpace.cpp
--- a/Engine/Source/Editor/AnimGraph/Private/AnimGraphNode_RotationOffsetBlendSpace.cpp
+++ b/Engine/Source/Editor/AnimGraph/Private/AnimGraphNode_RotationOffsetBlendSpace.cpp
@@ -177,9 +177,9 @@ void UAnimGraphNode_RotationOffsetBlendSpace::ValidateAnimNodeDuringCompilation(
 	{
 		USkeleton* BlendSpaceSkeleton = BlendSpaceToCheck->GetSkeleton();
 		if (BlendSpaceSkeleton && // if blend space doesn't have skeleton, it might be due to blend space not loaded yet, @todo: wait with anim blueprint compilation until all assets are loaded?
-			!BlendSpaceSkeleton->IsCompatible(ForSkeleton))
+			!ForSkeleton->IsCompatible(BlendSpaceSkeleton))
 		{
-			MessageLog.Error(TEXT("@@ references blendspace that uses different skeleton @@"), this, BlendSpaceSkeleton);
+			MessageLog.Error(TEXT("@@ references blendspace that uses an incompatible skeleton @@"), this, BlendSpaceSkeleton);
 		}
 	}
 
diff --git a/Engine/Source/Editor/AnimGraph/Private/AnimGraphNode_SequenceEvaluator.cpp b/Engine/Source/Editor/AnimGraph/Private/AnimGraphNode_SequenceEvaluator.cpp
--- a/Engine/Source/Editor/AnimGraph/Private/AnimGraphNode_SequenceEvaluator.cpp
+++ b/Engine/Source/Editor/AnimGraph/Private/AnimGraphNode_SequenceEvaluator.cpp
@@ -149,10 +149,10 @@ void UAnimGraphNode_SequenceEvaluator::ValidateAnimNodeDuringCompilation(class U
 	else
 	{
 		USkeleton* SeqSkeleton = SequenceToCheck->GetSkeleton();
-		if (SeqSkeleton&& // if anim sequence doesn't have skeleton, it might be due to anim sequence not loaded yet, @todo: wait with anim blueprint compilation until all assets are loaded?
-			!SeqSkeleton->IsCompatible(ForSkeleton))
+		if (SeqSkeleton && // if anim sequence doesn't have skeleton, it might be due to anim sequence not loaded yet, @todo: wait with anim blueprint compilation until all assets are loaded?
+			!ForSkeleton->IsCompatible(SeqSkeleton))
 		{
-			MessageLog.Error(TEXT("@@ references sequence that uses different skeleton @@"), this, SeqSkeleton);
+			MessageLog.Error(TEXT("@@ references sequence that uses an incompatible skeleton @@"), this, SeqSkeleton);
 		}
 	}
 }
diff --git a/Engine/Source/Editor/AnimGraph/Private/AnimGraphNode_SequencePlayer.cpp b/Engine/Source/Editor/AnimGraph/Private/AnimGraphNode_SequencePlayer.cpp
--- a/Engine/Source/Editor/AnimGraph/Private/AnimGraphNode_SequencePlayer.cpp
+++ b/Engine/Source/Editor/AnimGraph/Private/AnimGraphNode_SequencePlayer.cpp
@@ -312,17 +312,16 @@ bool UAnimGraphNode_SequencePlayer::IsActionFilteredOut(class FBlueprintActionFi
 		{
 			if(Node.Sequence)
 			{
-				if(Node.Sequence->GetSkeleton() != AnimBlueprint->TargetSkeleton)
+				if (!AnimBlueprint->TargetSkeleton->IsCompatible(Node.Sequence->GetSkeleton()))
 				{
-					// Sequence does not use the same skeleton as the Blueprint, cannot use
+					// Asset does not use a compatible skeleton with the Blueprint, cannot use
 					bIsFilteredOut = true;
 					break;
 				}
 			}
 			else 
 			{
-				FAssetData SkeletonData(AnimBlueprint->TargetSkeleton);
-				if(UnloadedSkeletonName != SkeletonData.GetExportTextName())
+				if (!AnimBlueprint->TargetSkeleton->IsCompatibleSkeletonByAssetString(UnloadedSkeletonName))
 				{
 					bIsFilteredOut = true;
 					break;
@@ -387,10 +386,10 @@ void UAnimGraphNode_SequencePlayer::ValidateAnimNodeDuringCompilation(class USke
 	else
 	{
 		USkeleton* SeqSkeleton = SequenceToCheck->GetSkeleton();
-		if (SeqSkeleton&& // if anim sequence doesn't have skeleton, it might be due to anim sequence not loaded yet, @todo: wait with anim blueprint compilation until all assets are loaded?
-			!SeqSkeleton->IsCompatible(ForSkeleton))
+		if (SeqSkeleton && // if anim sequence doesn't have skeleton, it might be due to anim sequence not loaded yet, @todo: wait with anim blueprint compilation until all assets are loaded?
+			!ForSkeleton->IsCompatible(SeqSkeleton))
 		{
-			MessageLog.Error(TEXT("@@ references sequence that uses different skeleton @@"), this, SeqSkeleton);
+			MessageLog.Error(TEXT("@@ references sequence that uses an incompatible skeleton @@"), this, SeqSkeleton);
 		}
 	}
 }
diff --git a/Engine/Source/Editor/AnimGraph/Private/AnimationGraphSchema.cpp b/Engine/Source/Editor/AnimGraph/Private/AnimationGraphSchema.cpp
--- a/Engine/Source/Editor/AnimGraph/Private/AnimationGraphSchema.cpp
+++ b/Engine/Source/Editor/AnimGraph/Private/AnimationGraphSchema.cpp
@@ -406,7 +406,7 @@ void UAnimationGraphSchema::SpawnNodeFromAsset(UAnimationAsset* Asset, const FVe
 
 	UAnimBlueprint* AnimBlueprint = Cast<UAnimBlueprint>(FBlueprintEditorUtils::FindBlueprintForGraph(Graph));
 
-	const bool bSkelMatch = (AnimBlueprint != NULL) && (AnimBlueprint->TargetSkeleton == Asset->GetSkeleton());
+	const bool bSkelMatch = (AnimBlueprint != NULL) && (AnimBlueprint->TargetSkeleton->IsCompatible(Asset->GetSkeleton()));
 	const bool bTypeMatch = (PinIfAvailable == NULL) || UAnimationGraphSchema::IsLocalSpacePosePin(PinIfAvailable->PinType);
 	const bool bDirectionMatch = (PinIfAvailable == NULL) || (PinIfAvailable->Direction == EGPD_Input);
 
@@ -514,7 +514,7 @@ void UAnimationGraphSchema::GetAssetsNodeHoverMessage(const TArray<FAssetData>&
 
 	// this one only should happen when there is an Anim Blueprint
 	UAnimBlueprint* AnimBlueprint = Cast<UAnimBlueprint>(FBlueprintEditorUtils::FindBlueprintForNode(HoverNode));
-	const bool bSkelMatch = (AnimBlueprint != NULL) && (AnimBlueprint->TargetSkeleton == Asset->GetSkeleton());
+	const bool bSkelMatch = (AnimBlueprint != NULL) && (AnimBlueprint->TargetSkeleton->IsCompatible(Asset->GetSkeleton()));
 
 	if (!bSkelMatch)
 	{
@@ -546,7 +546,7 @@ void UAnimationGraphSchema::GetAssetsPinHoverMessage(const TArray<FAssetData>& A
 	// this one only should happen when there is an Anim Blueprint
 	UAnimBlueprint* AnimBlueprint = Cast<UAnimBlueprint>(FBlueprintEditorUtils::FindBlueprintForNode(HoverPin->GetOwningNode()));
 
-	const bool bSkelMatch = (AnimBlueprint != NULL) && (AnimBlueprint->TargetSkeleton == Asset->GetSkeleton());
+	const bool bSkelMatch = (AnimBlueprint != NULL) && (AnimBlueprint->TargetSkeleton->IsCompatible(Asset->GetSkeleton()));
 	const bool bTypeMatch = UAnimationGraphSchema::IsLocalSpacePosePin(HoverPin->PinType);
 	const bool bDirectionMatch = HoverPin->Direction == EGPD_Input;
 
@@ -567,7 +567,7 @@ void UAnimationGraphSchema::GetAssetsGraphHoverMessage(const TArray<FAssetData>&
 	if (UAnimationAsset* AnimationAsset = FAssetData::GetFirstAsset<UAnimationAsset>(Assets))
 	{
 		UAnimBlueprint* AnimBlueprint = Cast<UAnimBlueprint>(FBlueprintEditorUtils::FindBlueprintForGraph(HoverGraph));
-		const bool bSkelMatch = (AnimBlueprint != NULL) && (AnimBlueprint->TargetSkeleton == AnimationAsset->GetSkeleton());
+		const bool bSkelMatch = (AnimBlueprint != NULL) && (AnimBlueprint->TargetSkeleton->IsCompatible(AnimationAsset->GetSkeleton()));
 		if (!bSkelMatch)
 		{
 			OutOkIcon = false;
diff --git a/Engine/Source/Editor/DetailCustomizations/Private/AnimSequenceDetails.cpp b/Engine/Source/Editor/DetailCustomizations/Private/AnimSequenceDetails.cpp
--- a/Engine/Source/Editor/DetailCustomizations/Private/AnimSequenceDetails.cpp
+++ b/Engine/Source/Editor/DetailCustomizations/Private/AnimSequenceDetails.cpp
@@ -550,7 +550,7 @@ void SAnimationRefPoseViewport::Tick( const FGeometry& AllottedGeometry, const d
 		{
 			Description->SetText( FText::Format( LOCTEXT( "Previewing", "Previewing {0}" ), FText::FromString( Component->AnimClass->GetName() ) ) );
 		}
-		else if ( AnimRef && AnimRef->GetSkeleton() != TargetSkeleton )
+		else if ( AnimRef && !AnimRef->GetSkeleton()->IsCompatible(TargetSkeleton) )
 		{
 			Description->SetText( FText::Format( LOCTEXT( "IncorrectSkeleton", "The preview asset doesn't work for the skeleton '{0}'" ), FText::FromString( TargetSkeletonName ) ) );
 		}
diff --git a/Engine/Source/Editor/DetailCustomizations/Private/AnimationAssetDetail.cpp b/Engine/Source/Editor/DetailCustomizations/Private/AnimationAssetDetail.cpp
--- a/Engine/Source/Editor/DetailCustomizations/Private/AnimationAssetDetail.cpp
+++ b/Engine/Source/Editor/DetailCustomizations/Private/AnimationAssetDetail.cpp
@@ -28,7 +28,7 @@ void FAnimationAssetDetails::CustomizeDetails(IDetailLayoutBuilder& DetailBuilde
 	{
 		if (UAnimationAsset* TestAsset = Cast<UAnimationAsset>(SelectionIt->Get()))
 		{
-			if (TargetSkeleton.IsValid() && TestAsset->GetSkeleton() != TargetSkeleton.Get())
+			if (TargetSkeleton.IsValid() && !TestAsset->GetSkeleton()->IsCompatible(TargetSkeleton.Get()))
 			{
 				TargetSkeleton = nullptr;
 				break;
@@ -72,9 +72,7 @@ bool FAnimationAssetDetails::ShouldFilterAsset(const FAssetData& AssetData)
 {
 	if (TargetSkeleton.IsValid())
 	{
-		FString SkeletonString = FAssetData(TargetSkeleton.Get()).GetExportTextName();
-		FAssetDataTagMapSharedView::FFindTagResult Result = AssetData.TagsAndValues.FindTag("Skeleton");
-		return (!Result.IsSet() || SkeletonString != Result.GetValue());
+		return !TargetSkeleton->IsCompatibleSkeletonByAssetData(AssetData);
 	}
 
 	return true;
diff --git a/Engine/Source/Editor/DetailCustomizations/Private/PoseAssetDetails.cpp b/Engine/Source/Editor/DetailCustomizations/Private/PoseAssetDetails.cpp
--- a/Engine/Source/Editor/DetailCustomizations/Private/PoseAssetDetails.cpp
+++ b/Engine/Source/Editor/DetailCustomizations/Private/PoseAssetDetails.cpp
@@ -314,9 +314,7 @@ bool FPoseAssetDetails::ShouldFilterAsset(const FAssetData& AssetData)
 {
 	if (TargetSkeleton.IsValid())
 	{
-		FString SkeletonString = FAssetData(TargetSkeleton.Get()).GetExportTextName();
-		FAssetDataTagMapSharedView::FFindTagResult Result = AssetData.TagsAndValues.FindTag("Skeleton");
-		return (!Result.IsSet() || SkeletonString != Result.GetValue());
+		return !TargetSkeleton->IsCompatibleSkeletonByAssetData(AssetData);
 	}
 
 	return true;
@@ -595,7 +593,7 @@ void FPoseAssetDetails::CachePoseAssetData()
  {
 	 if (PoseAsset.IsValid())
 	 {
-		 FScopedTransaction Transaction(LOCTEXT("ApplyAdditiveSetting_Transaciton", "Apply Additive Setting"));
+		 FScopedTransaction Transaction(LOCTEXT("ApplyAdditiveSetting_Transaction", "Apply Additive Setting"));
 		 PoseAsset->Modify();
 
 		 PoseAsset->ConvertSpace(bCachedAdditive, CachedBasePoseIndex);
@@ -636,9 +634,9 @@ void FPoseAssetDetails::CachePoseAssetData()
 		 SourceAnimationPropertyHandle->GetValue(ObjectSet);
 
 		 UAnimSequence* AnimSequenceSelected = Cast<UAnimSequence>(ObjectSet);
-		 if (AnimSequenceSelected && AnimSequenceSelected->GetSkeleton() == PoseAsset->GetSkeleton())
+		 if (AnimSequenceSelected && PoseAsset->GetSkeleton()->IsCompatible(AnimSequenceSelected->GetSkeleton()))
 		 {
-			 FScopedTransaction Transaction(LOCTEXT("UpdatePoseSourceAnimation_Transaciton", "Update Pose"));
+			 FScopedTransaction Transaction(LOCTEXT("UpdatePoseSourceAnimation_Transaction", "Update Pose"));
 			 PoseAsset->Modify();
 			 PoseAsset->UpdatePoseFromAnimation(AnimSequenceSelected);
 
@@ -664,7 +662,7 @@ void FPoseAssetDetails::CachePoseAssetData()
 			 Args.Add(TEXT("AssetName"), FText::FromString(PoseAsset->GetName()));
 			 Args.Add(TEXT("SourceAsset"), FText::FromString(GetNameSafe(AnimSequenceSelected)));
 			 Args.Add(TEXT("SkeletonName"), FText::FromString(GetNameSafe(PoseAsset->GetSkeleton())));
-			 FText ResultText = FText::Format(LOCTEXT("UpdatePoseWithInvalidSkeleton", "Source Asset {SourceAsset} is invalid or does not have matching skeleton {SkeletonName} with {AssetName}"), Args);
+			 FText ResultText = FText::Format(LOCTEXT("UpdatePoseWithInvalidSkeleton", "Source Asset {SourceAsset} is invalid or does not have a compatible skeleton {SkeletonName} with {AssetName}"), Args);
 
 			 FNotificationInfo Info(ResultText);
 			 Info.bFireAndForget = true;
diff --git a/Engine/Source/Editor/DetailCustomizations/Private/SkeletalMeshComponentDetails.cpp b/Engine/Source/Editor/DetailCustomizations/Private/SkeletalMeshComponentDetails.cpp
--- a/Engine/Source/Editor/DetailCustomizations/Private/SkeletalMeshComponentDetails.cpp
+++ b/Engine/Source/Editor/DetailCustomizations/Private/SkeletalMeshComponentDetails.cpp
@@ -49,6 +49,7 @@ public:
 
 FSkeletalMeshComponentDetails::FSkeletalMeshComponentDetails()
 	: CurrentDetailBuilder(NULL)
+	, Skeleton(nullptr)
 	, bAnimPickerEnabled(false)
 {
 
@@ -223,8 +224,13 @@ EVisibility FSkeletalMeshComponentDetails::VisibilityForAnimationMode(EAnimation
 
 bool FSkeletalMeshComponentDetails::OnShouldFilterAnimAsset( const FAssetData& AssetData )
 {
-	const FString SkeletonName = AssetData.GetTagValueRef<FString>("Skeleton");
-	return SkeletonName != SelectedSkeletonName;
+	// Check the compatible skeletons.
+	if (Skeleton && Skeleton->IsCompatibleSkeletonByAssetData(AssetData))
+	{
+		return false;
+	}
+
+	return true;
 }
 
 void FSkeletalMeshComponentDetails::SkeletalMeshPropertyChanged()
@@ -235,7 +241,7 @@ void FSkeletalMeshComponentDetails::SkeletalMeshPropertyChanged()
 void FSkeletalMeshComponentDetails::UpdateSkeletonNameAndPickerVisibility()
 {
 	// Update the selected skeleton name and the picker visibility
-	USkeleton* Skeleton = GetValidSkeletonFromRegisteredMeshes();
+	Skeleton = GetValidSkeletonFromRegisteredMeshes();
 
 	if (Skeleton)
 	{
@@ -356,8 +362,7 @@ void FSkeletalMeshComponentDetails::UseSelectedAnimBlueprint()
 		{
 			if(USkeleton* AnimBlueprintSkeleton = AnimBlueprintToAssign->TargetSkeleton)
 			{
-				FString BlueprintSkeletonName = FString::Printf(TEXT("%s'%s'"), *AnimBlueprintSkeleton->GetClass()->GetName(), *AnimBlueprintSkeleton->GetPathName());
-				if (BlueprintSkeletonName == SelectedSkeletonName)
+				if (Skeleton && Skeleton->IsCompatible(AnimBlueprintSkeleton))
 				{
 					OnClassPicked(AnimBlueprintToAssign->GetAnimBlueprintGeneratedClass());
 				}
@@ -385,7 +390,7 @@ void FSkeletalMeshComponentDetails::PerformInitialRegistrationOfSkeletalMeshes(I
 
 USkeleton* FSkeletalMeshComponentDetails::GetValidSkeletonFromRegisteredMeshes() const
 {
-	USkeleton* Skeleton = NULL;
+	USkeleton* ResultSkeleton = NULL;
 
 	for (auto ObjectIter = SelectedObjects.CreateConstIterator(); ObjectIter; ++ObjectIter)
 	{
@@ -396,22 +401,22 @@ USkeleton* FSkeletalMeshComponentDetails::GetValidSkeletonFromRegisteredMeshes()
 		}
 
 		// If we've not come across a valid skeleton yet, store this one.
-		if (!Skeleton)
+		if (!ResultSkeleton)
 		{
-			Skeleton = Mesh->SkeletalMesh->GetSkeleton();
+			ResultSkeleton = Mesh->SkeletalMesh->GetSkeleton();
 			continue;
 		}
 
 		// We've encountered a valid skeleton before.
 		// If this skeleton is not the same one, that means there are multiple
 		// skeletons selected, so we don't want to take any action.
-		if (Mesh->SkeletalMesh->GetSkeleton() != Skeleton)
+		if (Mesh->SkeletalMesh->GetSkeleton() != ResultSkeleton)
 		{
 			return NULL;
 		}
 	}
 
-	return Skeleton;
+	return ResultSkeleton;
 }
 
 #undef LOCTEXT_NAMESPACE
diff --git a/Engine/Source/Editor/DetailCustomizations/Private/SkeletalMeshComponentDetails.h b/Engine/Source/Editor/DetailCustomizations/Private/SkeletalMeshComponentDetails.h
--- a/Engine/Source/Editor/DetailCustomizations/Private/SkeletalMeshComponentDetails.h
+++ b/Engine/Source/Editor/DetailCustomizations/Private/SkeletalMeshComponentDetails.h
@@ -12,6 +12,7 @@ struct FAssetData;
 class IDetailLayoutBuilder;
 class IPropertyHandle;
 class SComboButton;
+class USkeleton;
 
 class DETAILCUSTOMIZATIONS_API FSkeletalMeshComponentDetails : public IDetailCustomization
 {
@@ -99,6 +100,9 @@ private:
 	/** Full name of the currently selected skeleton to use for filtering animation assets */
 	FString SelectedSkeletonName;
 
+	/** The skeleton that we grab the name from for filtering. */
+	USkeleton* Skeleton;
+
 	/** Current enabled state of the animation asset picker in the details panel */
 	bool bAnimPickerEnabled;
 
diff --git a/Engine/Source/Editor/MovieSceneTools/Private/TrackEditors/SkeletalAnimationTrackEditor.cpp b/Engine/Source/Editor/MovieSceneTools/Private/TrackEditors/SkeletalAnimationTrackEditor.cpp
--- a/Engine/Source/Editor/MovieSceneTools/Private/TrackEditors/SkeletalAnimationTrackEditor.cpp
+++ b/Engine/Source/Editor/MovieSceneTools/Private/TrackEditors/SkeletalAnimationTrackEditor.cpp
@@ -424,7 +424,7 @@ public:
 				{
 					// Store the compatible skeleton's name, and create a property widget with a filter that will check
 					// for animations that match that skeleton.
-					USkeleton* Skeleton = AcquireSkeletonFromObjectGuid(Params.ParentObjectBindingGuid, Params.Sequencer);
+					Skeleton = AcquireSkeletonFromObjectGuid(Params.ParentObjectBindingGuid, Params.Sequencer);
 					SkeletonName = FAssetData(Skeleton).GetExportTextName();
 
 					TSharedPtr<IPropertyUtilities> PropertyUtilities = CustomizationUtils.GetPropertyUtilities();
@@ -463,21 +463,13 @@ public:
 			return true;
 		}
 
-		if (!SkeletonName.IsEmpty())
-		{
-			const FString& SkeletonTag = AssetData.GetTagValueRef<FString>(TEXT("Skeleton"));
-			if (SkeletonTag != SkeletonName)
-			{
-				return true;
-			}
-		}
-
-		return false;
+		return !(Skeleton && Skeleton->IsCompatibleSkeletonByAssetData(AssetData));
 	}
 
 private:
 	FSequencerSectionPropertyDetailsViewCustomizationParams Params;
 	FString SkeletonName;
+	USkeleton* Skeleton = nullptr;
 };
 
 
@@ -1085,7 +1077,7 @@ bool FSkeletalAnimationTrackEditor::HandleAssetAdded(UObject* Asset, const FGuid
 		{
 			USkeleton* Skeleton = AcquireSkeletonFromObjectGuid(TargetObjectGuid, GetSequencer());
 
-			if (Skeleton && Skeleton == AnimSequence->GetSkeleton())
+			if (Skeleton && Skeleton->IsCompatible(AnimSequence->GetSkeleton()))
 			{
 				UObject* Object = SequencerPtr->FindSpawnedObjectOrTemplate(TargetObjectGuid);
 				
@@ -1758,7 +1750,7 @@ bool FSkeletalAnimationTrackEditor::OnAllowDrop(const FDragDropEvent& DragDropEv
 		UAnimSequenceBase* AnimSequence = Cast<UAnimSequenceBase>(AssetData.GetAsset());
 
 		const bool bValidAnimSequence = AnimSequence && AnimSequence->CanBeUsedInComposition();
-		if (bValidAnimSequence && Skeleton && Skeleton == AnimSequence->GetSkeleton())
+		if (bValidAnimSequence && Skeleton && Skeleton->IsCompatible(AnimSequence->GetSkeleton()))
 		{
 			return true;
 		}
@@ -1800,7 +1792,7 @@ FReply FSkeletalAnimationTrackEditor::OnDrop(const FDragDropEvent& DragDropEvent
 	{
 		UAnimSequenceBase* AnimSequence = Cast<UAnimSequenceBase>(AssetData.GetAsset());
 		const bool bValidAnimSequence = AnimSequence && AnimSequence->CanBeUsedInComposition();
-		if (bValidAnimSequence && Skeleton && Skeleton == AnimSequence->GetSkeleton())
+		if (bValidAnimSequence && Skeleton && Skeleton->IsCompatible(AnimSequence->GetSkeleton()))
 		{
 			UObject* Object = GetSequencer()->FindSpawnedObjectOrTemplate(TargetObjectGuid);
 				
diff --git a/Engine/Source/Editor/Persona/Private/AnimGraphNodeDetails.cpp b/Engine/Source/Editor/Persona/Private/AnimGraphNodeDetails.cpp
--- a/Engine/Source/Editor/Persona/Private/AnimGraphNodeDetails.cpp
+++ b/Engine/Source/Editor/Persona/Private/AnimGraphNodeDetails.cpp
@@ -262,7 +262,7 @@ TSharedRef<SWidget> FAnimGraphNodeDetails::CreatePropertyWidget(FProperty* Targe
 bool FAnimGraphNodeDetails::OnShouldFilterAnimAsset( const FAssetData& AssetData, UClass* NodeToFilterFor ) const
 {
 	FAssetDataTagMapSharedView::FFindTagResult Result = AssetData.TagsAndValues.FindTag("Skeleton");
-	if (Result.IsSet() && Result.GetValue() == TargetSkeletonName)
+	if (Result.IsSet() && TargetSkeleton && TargetSkeleton->IsCompatibleSkeletonByAssetData(AssetData))
 	{
 		const UClass* AssetClass = AssetData.GetClass();
 		// If node is an 'asset player', only let you select the right kind of asset for it
@@ -1020,18 +1020,8 @@ TSharedRef<SWidget> SParentPlayerTreeRow::GenerateWidgetForColumn(const FName& C
 
 bool SParentPlayerTreeRow::OnShouldFilterAsset(const FAssetData& AssetData)
 {
-	const FString SkeletonName = AssetData.GetTagValueRef<FString>("Skeleton");
-
-	if(!SkeletonName.IsEmpty())
-	{
-		USkeleton* CurrentSkeleton = GraphNode->GetAnimBlueprint()->TargetSkeleton;
-		if(SkeletonName == FString::Printf(TEXT("%s'%s'"), *CurrentSkeleton->GetClass()->GetName(), *CurrentSkeleton->GetPathName()))
-		{
-			return false;
-		}
-	}
-
-	return true;
+	const USkeleton* CurrentSkeleton = GraphNode->GetAnimBlueprint()->TargetSkeleton;
+	return !CurrentSkeleton->IsCompatibleSkeletonByAssetData(AssetData);
 }
 
 void SParentPlayerTreeRow::OnAssetSelected(const FAssetData& AssetData)
diff --git a/Engine/Source/Editor/Persona/Private/AnimInstanceDetails.cpp b/Engine/Source/Editor/Persona/Private/AnimInstanceDetails.cpp
--- a/Engine/Source/Editor/Persona/Private/AnimInstanceDetails.cpp
+++ b/Engine/Source/Editor/Persona/Private/AnimInstanceDetails.cpp
@@ -9,6 +9,7 @@
 #include "DetailCategoryBuilder.h"
 #include "PropertyCustomizationHelpers.h"
 #include "Animation/AnimInstance.h"
+#include "Animation/Skeleton.h"
 #include "ObjectEditorUtils.h"
 
 /////////////////////////////////////////////////////
@@ -54,8 +55,6 @@ void FAnimInstanceDetails::CustomizeDetails(class IDetailLayoutBuilder& DetailBu
 	// If everything uses the same skeleton we can filter based on that skeleton
 	if (TargetSkeleton)
 	{
-		TargetSkeletonName = FString::Printf(TEXT("%s'%s'"), *TargetSkeleton->GetClass()->GetName(), *TargetSkeleton->GetPathName());
-
 		for (TFieldIterator<FProperty> It(CommonBaseClass, EFieldIteratorFlags::IncludeSuper); It; ++It)
 		{
 			FProperty* TargetProperty = *It;
@@ -72,7 +71,7 @@ void FAnimInstanceDetails::CustomizeDetails(class IDetailLayoutBuilder& DetailBu
 				FDetailWidgetRow Row;
 				PropertyRow.GetDefaultWidgets(NameWidget, ValueWidget, Row);
 
-				TSharedRef<SWidget> TempWidget = CreateFilteredObjectPropertyWidget(TargetProperty, TargetPropertyHandle);
+				TSharedRef<SWidget> TempWidget = CreateFilteredObjectPropertyWidget(TargetProperty, TargetPropertyHandle, TargetSkeleton);
 				ValueWidget = (TempWidget == SNullWidget::NullWidget) ? ValueWidget : TempWidget;
 
 				const bool bShowChildren = true;
@@ -94,7 +93,7 @@ void FAnimInstanceDetails::CustomizeDetails(class IDetailLayoutBuilder& DetailBu
 	}
 }
 
-TSharedRef<SWidget> FAnimInstanceDetails::CreateFilteredObjectPropertyWidget(FProperty* TargetProperty, TSharedRef<IPropertyHandle> TargetPropertyHandle)
+TSharedRef<SWidget> FAnimInstanceDetails::CreateFilteredObjectPropertyWidget(FProperty* TargetProperty, TSharedRef<IPropertyHandle> TargetPropertyHandle, const USkeleton* TargetSkeleton)
 {
 	if(const FObjectPropertyBase* ObjectProperty = CastField<const FObjectPropertyBase>( TargetProperty ))
 	{
@@ -106,15 +105,14 @@ TSharedRef<SWidget> FAnimInstanceDetails::CreateFilteredObjectPropertyWidget(FPr
 				.PropertyHandle(TargetPropertyHandle)
 				.AllowedClass(ObjectProperty->PropertyClass)
 				.AllowClear(bAllowClear)
-				.OnShouldFilterAsset(FOnShouldFilterAsset::CreateSP(this, &FAnimInstanceDetails::OnShouldFilterAnimAsset));
+				.OnShouldFilterAsset(this, &FAnimInstanceDetails::OnShouldFilterAnimAsset, TargetSkeleton);
 		}
 	}
 
 	return SNullWidget::NullWidget;
 }
 
-bool FAnimInstanceDetails::OnShouldFilterAnimAsset(const FAssetData& AssetData) const
+bool FAnimInstanceDetails::OnShouldFilterAnimAsset(const FAssetData& AssetData, const USkeleton* TargetSkeleton) const
 {
-	const FString SkeletonName = AssetData.GetTagValueRef<FString>("Skeleton");
-	return SkeletonName != TargetSkeletonName;
+	return !(TargetSkeleton && TargetSkeleton->IsCompatibleSkeletonByAssetData(AssetData));
 }
diff --git a/Engine/Source/Editor/Persona/Private/AnimInstanceDetails.h b/Engine/Source/Editor/Persona/Private/AnimInstanceDetails.h
--- a/Engine/Source/Editor/Persona/Private/AnimInstanceDetails.h
+++ b/Engine/Source/Editor/Persona/Private/AnimInstanceDetails.h
@@ -8,6 +8,7 @@
 #include "IDetailCustomization.h"
 
 struct FAssetData;
+class USkeleton;
 
 /////////////////////////////////////////////////////
 // FAnimInstanceDetails 
@@ -24,11 +25,8 @@ public:
 protected:
 
 	// Creates a filtered object widget if the supplied property is an object property
-	TSharedRef<SWidget> CreateFilteredObjectPropertyWidget(FProperty* TargetProperty, TSharedRef<IPropertyHandle> TargetPropertyHandle);
+	TSharedRef<SWidget> CreateFilteredObjectPropertyWidget(FProperty* TargetProperty, TSharedRef<IPropertyHandle> TargetPropertyHandle, const USkeleton* TargetSkeleton);
 
 	/** Delegate to handle filtering of asset pickers */
-	bool OnShouldFilterAnimAsset( const FAssetData& AssetData ) const;
-
-	/** Path to the current blueprints skeleton to allow us to filter asset pickers */
-	FString TargetSkeletonName;
+	bool OnShouldFilterAnimAsset(const FAssetData& AssetData, const USkeleton* TargetSkeleton) const;
 };
diff --git a/Engine/Source/Editor/Persona/Private/AnimationEditorPreviewScene.cpp b/Engine/Source/Editor/Persona/Private/AnimationEditorPreviewScene.cpp
--- a/Engine/Source/Editor/Persona/Private/AnimationEditorPreviewScene.cpp
+++ b/Engine/Source/Editor/Persona/Private/AnimationEditorPreviewScene.cpp
@@ -613,7 +613,7 @@ void FAnimationEditorPreviewScene::SetPreviewAnimationAsset(UAnimationAsset* Ani
 			}
 
 			// Treat it as invalid if it's got a bogus skeleton pointer
-			if (AnimAsset->GetSkeleton() != Skeleton && Skeleton != nullptr)
+			if (Skeleton != nullptr && !Skeleton->IsCompatible(AnimAsset->GetSkeleton()))
 			{
 				return;
 			}
diff --git a/Engine/Source/Editor/Persona/Private/Customization/BlendSampleDetails.cpp b/Engine/Source/Editor/Persona/Private/Customization/BlendSampleDetails.cpp
--- a/Engine/Source/Editor/Persona/Private/Customization/BlendSampleDetails.cpp
+++ b/Engine/Source/Editor/Persona/Private/Customization/BlendSampleDetails.cpp
@@ -300,10 +300,10 @@ bool FBlendSampleDetails::ShouldFilterAssetStatic(const FAssetData& AssetData, c
 	FString SkeletonName;
 	if (AssetData.GetTagValue(SkeletonTagName, SkeletonName))
 	{
-		// Check whether or not the skeletons match
-		if (SkeletonName.Contains(BlendSpaceBase->GetSkeleton()->GetPathName()))
+		// Check whether or not the skeletons are compatible
+		if (BlendSpaceBase->GetSkeleton()->IsCompatibleSkeletonByAssetData(AssetData))
 		{
-			// If so check if the additive animation tpye is compatible with the blend space
+			// If so check if the additive animation type is compatible with the blend space
 			const FName AdditiveTypeTagName = GET_MEMBER_NAME_CHECKED(UAnimSequence, AdditiveAnimType);
 			FString AnimationTypeName;
 			if (AssetData.GetTagValue(AdditiveTypeTagName, AnimationTypeName))
diff --git a/Engine/Source/Editor/Persona/Private/PersonaAssetFamily.cpp b/Engine/Source/Editor/Persona/Private/PersonaAssetFamily.cpp
--- a/Engine/Source/Editor/Persona/Private/PersonaAssetFamily.cpp
+++ b/Engine/Source/Editor/Persona/Private/PersonaAssetFamily.cpp
@@ -255,7 +255,7 @@ bool FPersonaAssetFamily::IsAssetCompatible(const FAssetData& InAssetData) const
 	{
 		if (Class->IsChildOf<USkeleton>())
 		{
-			return FAssetData(Skeleton.Get()) == InAssetData;
+			return Skeleton.Get()->IsCompatibleSkeletonByAssetData(InAssetData);
 		}
 		else if (Class->IsChildOf<UAnimationAsset>() || Class->IsChildOf<USkeletalMesh>())
 		{
@@ -263,7 +263,10 @@ bool FPersonaAssetFamily::IsAssetCompatible(const FAssetData& InAssetData) const
 
 			if (Result.IsSet())
 			{
-				return Result.GetValue() == FAssetData(Skeleton.Get()).GetExportTextName();
+				if (Skeleton.Get())
+				{
+					return Skeleton.Get()->IsCompatibleSkeletonByAssetString(Result.GetValue());
+				}
 			}
 		}
 		else if (Class->IsChildOf<UAnimBlueprint>())
@@ -272,7 +275,7 @@ bool FPersonaAssetFamily::IsAssetCompatible(const FAssetData& InAssetData) const
 
 			if (Result.IsSet())
 			{
-				return Result.GetValue() == FAssetData(Skeleton.Get()).GetExportTextName();
+				return Skeleton.Get()->IsCompatibleSkeletonByAssetString(Result.GetValue());
 			}
 		}
 		else if (Class->IsChildOf<UPhysicsAsset>())
diff --git a/Engine/Source/Editor/Persona/Private/PersonaMeshDetails.cpp b/Engine/Source/Editor/Persona/Private/PersonaMeshDetails.cpp
--- a/Engine/Source/Editor/Persona/Private/PersonaMeshDetails.cpp
+++ b/Engine/Source/Editor/Persona/Private/PersonaMeshDetails.cpp
@@ -3909,10 +3909,8 @@ bool FPersonaMeshDetails::OnShouldFilterPostProcessBlueprint(const FAssetData& A
 {
 	if(USkeletalMesh* SkelMesh = GetPersonaToolkit()->GetMesh())
 	{
-		const FString CurrentMeshSkeletonName = FString::Printf(TEXT("%s'%s'"), *SkelMesh->GetSkeleton()->GetClass()->GetName(), *SkelMesh->GetSkeleton()->GetPathName());
 		const FString SkeletonName = AssetData.GetTagValueRef<FString>("TargetSkeleton");
-
-		return SkeletonName != CurrentMeshSkeletonName;
+		return !SkelMesh->GetSkeleton()->IsCompatibleSkeletonByAssetString(SkeletonName);
 	}
 
 	return true;
@@ -6145,10 +6143,7 @@ void FPersonaMeshDetails::OnPreviewMeshChanged(USkeletalMesh* OldSkeletalMesh, U
 
 bool FPersonaMeshDetails::FilterOutBakePose(const FAssetData& AssetData, USkeleton* Skeleton) const
 {
-	FString SkeletonName;
-	AssetData.GetTagValue("Skeleton", SkeletonName);
-	FAssetData SkeletonData(Skeleton);
-	return (SkeletonName != SkeletonData.GetExportTextName());
+	return !(Skeleton && Skeleton->IsCompatibleSkeletonByAssetData(AssetData));
 }
 
 #undef LOCTEXT_NAMESPACE
diff --git a/Engine/Source/Editor/Persona/Private/PersonaModule.cpp b/Engine/Source/Editor/Persona/Private/PersonaModule.cpp
--- a/Engine/Source/Editor/Persona/Private/PersonaModule.cpp
+++ b/Engine/Source/Editor/Persona/Private/PersonaModule.cpp
@@ -937,7 +937,7 @@ void FPersonaModule::AddCommonToolbarExtensions(FToolBarBuilder& InToolbarBuilde
 						FString TagValue;
 						if (AssetData.GetTagValue("Skeleton", TagValue))
 						{
-							return TagValue != FAssetData(WeakPersonaToolkit.Pin()->GetSkeleton()).GetExportTextName();
+							return !WeakPersonaToolkit.Pin()->GetSkeleton()->IsCompatibleSkeletonByAssetData(AssetData);
 						}
 					}
 					return true;
@@ -1009,7 +1009,7 @@ void FPersonaModule::AddCommonToolbarExtensions(FToolBarBuilder& InToolbarBuilde
 						FString TagValue;
 						if (AssetData.GetTagValue("Skeleton", TagValue))
 						{
-							return TagValue != FAssetData(WeakPersonaToolkit.Pin()->GetSkeleton()).GetExportTextName();
+							return !WeakPersonaToolkit.Pin()->GetSkeleton()->IsCompatibleSkeletonByAssetData(AssetData);
 						}
 					}
 					return true;
diff --git a/Engine/Source/Editor/Persona/Private/PersonaPreviewSceneAnimationController.cpp b/Engine/Source/Editor/Persona/Private/PersonaPreviewSceneAnimationController.cpp
--- a/Engine/Source/Editor/Persona/Private/PersonaPreviewSceneAnimationController.cpp
+++ b/Engine/Source/Editor/Persona/Private/PersonaPreviewSceneAnimationController.cpp
@@ -7,6 +7,7 @@
 #include "IEditableSkeleton.h"
 #include "DetailWidgetRow.h"
 #include "PropertyCustomizationHelpers.h"
+#include "Animation/Skeleton.h"
 
 void UPersonaPreviewSceneAnimationController::InitializeView(UPersonaPreviewSceneDescription* SceneDescription, IPersonaPreviewScene* PreviewScene) const
 {
@@ -25,8 +26,6 @@ IDetailPropertyRow* UPersonaPreviewSceneAnimationController::AddPreviewControlle
 	const USkeleton* Skeleton = PersonaToolkit->GetPreviewMeshComponent()->SkeletalMesh ? PersonaToolkit->GetPreviewMeshComponent()->SkeletalMesh->GetSkeleton() : nullptr;
 	if (Skeleton)
 	{
-		FString SkeletonName = FAssetData(Skeleton).GetExportTextName();
-
 		IDetailPropertyRow* NewRow = Category.AddExternalObjectProperty(ListOfPreviewController, Property->GetFName(), PropertyLocation);
 	
 		NewRow->CustomWidget()
@@ -41,7 +40,7 @@ IDetailPropertyRow* UPersonaPreviewSceneAnimationController::AddPreviewControlle
 			SNew(SObjectPropertyEntryBox)
 			.AllowedClass(UAnimationAsset::StaticClass())
 			.PropertyHandle(NewRow->GetPropertyHandle())
-			.OnShouldFilterAsset(FOnShouldFilterAsset::CreateUObject(this, &UPersonaPreviewSceneAnimationController::HandleShouldFilterAsset, SkeletonName))
+			.OnShouldFilterAsset(FOnShouldFilterAsset::CreateUObject(this, &UPersonaPreviewSceneAnimationController::HandleShouldFilterAsset, Skeleton))
 			.ThumbnailPool(DetailBuilder.GetThumbnailPool())
 		];
 		return NewRow;
@@ -50,10 +49,9 @@ IDetailPropertyRow* UPersonaPreviewSceneAnimationController::AddPreviewControlle
 	return nullptr;
 }
 
-bool UPersonaPreviewSceneAnimationController::HandleShouldFilterAsset(const FAssetData& InAssetData, const FString SkeletonName) const
+bool UPersonaPreviewSceneAnimationController::HandleShouldFilterAsset(const FAssetData& InAssetData, const USkeleton* InSkeleton) const
 {
-	FString SkeletonTag = InAssetData.GetTagValueRef<FString>("Skeleton");
-	if (SkeletonTag == SkeletonName)
+	if (InSkeleton && InSkeleton->IsCompatibleSkeletonByAssetData(InAssetData))
 	{
 		return false;
 	}
diff --git a/Engine/Source/Editor/Persona/Private/PersonaPreviewSceneAnimationController.h b/Engine/Source/Editor/Persona/Private/PersonaPreviewSceneAnimationController.h
--- a/Engine/Source/Editor/Persona/Private/PersonaPreviewSceneAnimationController.h
+++ b/Engine/Source/Editor/Persona/Private/PersonaPreviewSceneAnimationController.h
@@ -7,6 +7,7 @@
 #include "PersonaPreviewSceneAnimationController.generated.h"
 
 class UAnimationAsset;
+class USkeleton;
 
 UCLASS(DisplayName="Use Specific Animation")
 class UPersonaPreviewSceneAnimationController : public UPersonaPreviewSceneController
@@ -23,5 +24,5 @@ public:
 
 	virtual IDetailPropertyRow* AddPreviewControllerPropertyToDetails(const TSharedRef<IPersonaToolkit>& PersonaToolkit, IDetailLayoutBuilder& DetailBuilder, IDetailCategoryBuilder& Category, const FProperty* Property, const EPropertyLocation::Type PropertyLocation);
 
-	bool HandleShouldFilterAsset(const FAssetData& InAssetData, const FString SkeletonName) const;
+	bool HandleShouldFilterAsset(const FAssetData& InAssetData, const USkeleton* InSkeleton) const;
 };
diff --git a/Engine/Source/Editor/Persona/Private/PersonaToolkit.cpp b/Engine/Source/Editor/Persona/Private/PersonaToolkit.cpp
--- a/Engine/Source/Editor/Persona/Private/PersonaToolkit.cpp
+++ b/Engine/Source/Editor/Persona/Private/PersonaToolkit.cpp
@@ -213,7 +213,7 @@ void FPersonaToolkit::SetAnimationAsset(class UAnimationAsset* InAnimationAsset)
 {
 	if (InAnimationAsset != nullptr)
 	{
-		check(InAnimationAsset->GetSkeleton() == Skeleton);
+		check(Skeleton->IsCompatible(InAnimationAsset->GetSkeleton()));
 	}
 
 	AnimationAsset = InAnimationAsset;
diff --git a/Engine/Source/Editor/Persona/Private/PreviewSceneCustomizations.cpp b/Engine/Source/Editor/Persona/Private/PreviewSceneCustomizations.cpp
--- a/Engine/Source/Editor/Persona/Private/PreviewSceneCustomizations.cpp
+++ b/Engine/Source/Editor/Persona/Private/PreviewSceneCustomizations.cpp
@@ -23,6 +23,7 @@
 #include "AssetToolsModule.h"
 #include "IAssetTools.h"
 #include "Animation/AnimBlueprint.h"
+#include "Animation/Skeleton.h"
 #include "UObject/UObjectIterator.h"
 #include "Widgets/Input/SComboBox.h"
 #include "PhysicsEngine/PhysicsAsset.h"
@@ -523,8 +524,10 @@ bool FPreviewSceneDescriptionCustomization::HandleShouldFilterAsset(const FAsset
 		return false;
 	}
 
-	FString SkeletonTag = InAssetData.GetTagValueRef<FString>(InTag);
-	if (SkeletonName.IsEmpty() || SkeletonTag == SkeletonName)
+	const UPersonaPreviewSceneDescription* PersonaPreviewSceneDescription = PreviewScene.Pin()->GetPreviewSceneDescription();
+	const USkeleton* Skeleton = PersonaPreviewSceneDescription->PreviewMesh->GetSkeleton();
+	const FString SkeletonTag = InAssetData.GetTagValueRef<FString>(InTag);
+	if (Skeleton && Skeleton->IsCompatibleSkeletonByAssetString(SkeletonTag))
 	{
 		return false;
 	}
@@ -676,7 +679,7 @@ void FPreviewMeshCollectionEntryCustomization::CustomizeHeader(TSharedRef<IPrope
 	if (OuterObjects[0] != nullptr)
 	{
 		FString SkeletonName = FAssetData(CastChecked<UPreviewMeshCollection>(OuterObjects[0])->Skeleton).GetExportTextName();
-
+		USkeleton* Skeleton = CastChecked<UPreviewMeshCollection>(OuterObjects[0])->Skeleton;
 		PropertyHandle->GetParentHandle()->SetOnPropertyValueChanged(FSimpleDelegate::CreateSP(this, &FPreviewMeshCollectionEntryCustomization::HandleMeshesArrayChanged, CustomizationUtils.GetPropertyUtilities()));
 
 		TSharedPtr<IPropertyHandle> SkeletalMeshProperty = PropertyHandle->GetChildHandle(GET_MEMBER_NAME_CHECKED(FPreviewMeshCollectionEntry, SkeletalMesh));
@@ -693,7 +696,7 @@ void FPreviewMeshCollectionEntryCustomization::CustomizeHeader(TSharedRef<IPrope
 				SNew(SObjectPropertyEntryBox)
 				.AllowedClass(USkeletalMesh::StaticClass())
 				.PropertyHandle(SkeletalMeshProperty)
-				.OnShouldFilterAsset(this, &FPreviewMeshCollectionEntryCustomization::HandleShouldFilterAsset, SkeletonName)
+				.OnShouldFilterAsset(this, &FPreviewMeshCollectionEntryCustomization::HandleShouldFilterAsset, SkeletonName, Skeleton)
 				.OnObjectChanged(this, &FPreviewMeshCollectionEntryCustomization::HandleMeshChanged)
 				.ThumbnailPool(CustomizationUtils.GetThumbnailPool())
 			];
@@ -701,15 +704,14 @@ void FPreviewMeshCollectionEntryCustomization::CustomizeHeader(TSharedRef<IPrope
 	}
 }
 
-bool FPreviewMeshCollectionEntryCustomization::HandleShouldFilterAsset(const FAssetData& InAssetData, FString SkeletonName)
+bool FPreviewMeshCollectionEntryCustomization::HandleShouldFilterAsset(const FAssetData& InAssetData, FString SkeletonName, USkeleton* Skeleton)
 {
 	if (GetDefault<UPersonaOptions>()->bAllowPreviewMeshCollectionsToSelectFromDifferentSkeletons)
 	{
 		return false;
 	}
 
-	FString SkeletonTag = InAssetData.GetTagValueRef<FString>("Skeleton");
-	if (SkeletonTag == SkeletonName)
+	if (Skeleton && Skeleton->IsCompatibleSkeletonByAssetData(InAssetData))
 	{
 		return false;
 	}
diff --git a/Engine/Source/Editor/Persona/Private/PreviewSceneCustomizations.h b/Engine/Source/Editor/Persona/Private/PreviewSceneCustomizations.h
--- a/Engine/Source/Editor/Persona/Private/PreviewSceneCustomizations.h
+++ b/Engine/Source/Editor/Persona/Private/PreviewSceneCustomizations.h
@@ -21,6 +21,7 @@ class IDetailChildrenBuilder;
 class IDetailLayoutBuilder;
 class IPropertyUtilities;
 class UPreviewMeshCollectionFactory;
+class USkeleton;
 
 // An entry in the preview mode choice box
 struct FPersonaModeComboEntry
@@ -153,7 +154,7 @@ public:
 	virtual void CustomizeChildren(TSharedRef<IPropertyHandle> PropertyHandle, IDetailChildrenBuilder& ChildBuilder, IPropertyTypeCustomizationUtils& CustomizationUtils) override {}
 
 private:
-	bool HandleShouldFilterAsset(const FAssetData& InAssetData, FString SkeletonName);
+	bool HandleShouldFilterAsset(const FAssetData& InAssetData, FString SkeletonName, USkeleton* Skeleton);
 
 	void HandleMeshChanged(const FAssetData& InAssetData);
 
diff --git a/Engine/Source/Editor/Persona/Private/SAnimCurveViewer.cpp b/Engine/Source/Editor/Persona/Private/SAnimCurveViewer.cpp
--- a/Engine/Source/Editor/Persona/Private/SAnimCurveViewer.cpp
+++ b/Engine/Source/Editor/Persona/Private/SAnimCurveViewer.cpp
@@ -1167,6 +1167,7 @@ void SAnimCurveViewer::HandleSmartNamesChange(const FName& InContainerName)
 {
 	AnimCurveList.Empty();
 	RefreshCurveList(true);
+	EditableSkeletonPtr.Pin()->GetSkeleton().OnSmartNamesChangedEvent.Broadcast();
 }
 
 #undef LOCTEXT_NAMESPACE
diff --git a/Engine/Source/Editor/Persona/Private/SAnimationSequenceBrowser.cpp b/Engine/Source/Editor/Persona/Private/SAnimationSequenceBrowser.cpp
--- a/Engine/Source/Editor/Persona/Private/SAnimationSequenceBrowser.cpp
+++ b/Engine/Source/Editor/Persona/Private/SAnimationSequenceBrowser.cpp
@@ -1543,12 +1543,10 @@ bool SAnimationSequenceBrowser::HandleFilterAsset(const FAssetData& InAssetData)
 {
 	if (InAssetData.GetClass()->IsChildOf(UAnimationAsset::StaticClass()))
 	{
-		USkeleton* DesiredSkeleton = PersonaToolkitPtr.Pin()->GetSkeleton();
+		const USkeleton* DesiredSkeleton = PersonaToolkitPtr.Pin()->GetSkeleton();
 		if (DesiredSkeleton)
 		{
-			FString SkeletonString = FAssetData(DesiredSkeleton).GetExportTextName();
-
-			return (InAssetData.TagsAndValues.FindTag(TEXT("Skeleton")) != SkeletonString);
+			return !DesiredSkeleton->IsCompatibleSkeletonByAssetData(InAssetData);
 		}
 	}
 
diff --git a/Engine/Source/Editor/SkeletonEditor/Private/EditableSkeleton.cpp b/Engine/Source/Editor/SkeletonEditor/Private/EditableSkeleton.cpp
--- a/Engine/Source/Editor/SkeletonEditor/Private/EditableSkeleton.cpp
+++ b/Engine/Source/Editor/SkeletonEditor/Private/EditableSkeleton.cpp
@@ -266,7 +266,9 @@ bool FEditableSkeleton::AddSmartname(const FName& InContainerName, const FName&
 {
 	if (const FSmartNameMapping* NameMapping = Skeleton->GetSmartNameContainer(InContainerName))
 	{
-		return Skeleton->AddSmartNameAndModify(InContainerName, InNewName, OutSmartName);
+		const bool bResult = Skeleton->AddSmartNameAndModify(InContainerName, InNewName, OutSmartName);
+		OnSmartNameChanged.Broadcast(OutSmartName.DisplayName);
+		return bResult;
 	}
 
 	return false;
diff --git a/Engine/Source/Editor/UnrealEd/Private/Animation/DebugSkelMeshComponent.cpp b/Engine/Source/Editor/UnrealEd/Private/Animation/DebugSkelMeshComponent.cpp
--- a/Engine/Source/Editor/UnrealEd/Private/Animation/DebugSkelMeshComponent.cpp
+++ b/Engine/Source/Editor/UnrealEd/Private/Animation/DebugSkelMeshComponent.cpp
@@ -262,11 +262,11 @@ FString UDebugSkelMeshComponent::GetPreviewText() const
 
 void UDebugSkelMeshComponent::InitAnim(bool bForceReinit)
 {
-	// If we already have PreviewInstnace and its asset's Skeleton does not match with mesh's Skeleton
+	// If we already have PreviewInstance and its asset's Skeleton isn't compatible with the mesh's Skeleton
 	// then we need to clear it up to avoid an issue
 	if ( PreviewInstance && PreviewInstance->GetCurrentAsset() && SkeletalMesh )
 	{
-		if ( PreviewInstance->GetCurrentAsset()->GetSkeleton() != SkeletalMesh->GetSkeleton() )
+		if (SkeletalMesh->GetSkeleton()->IsCompatible(PreviewInstance->GetCurrentAsset()->GetSkeleton()))
 		{
 			// if it doesn't match, just clear it
 			PreviewInstance->SetAnimationAsset(NULL);
diff --git a/Engine/Source/Editor/UnrealEd/Private/Animation/EditorAnimCompositeSegment.cpp b/Engine/Source/Editor/UnrealEd/Private/Animation/EditorAnimCompositeSegment.cpp
--- a/Engine/Source/Editor/UnrealEd/Private/Animation/EditorAnimCompositeSegment.cpp
+++ b/Engine/Source/Editor/UnrealEd/Private/Animation/EditorAnimCompositeSegment.cpp
@@ -33,7 +33,7 @@ bool UEditorAnimCompositeSegment::ApplyChangesToMontage()
 	{
 		if(Composite->AnimationTrack.AnimSegments.IsValidIndex(AnimSegmentIndex))
 		{
-			if (AnimSegment.AnimReference && Composite->GetSkeleton() == AnimSegment.AnimReference->GetSkeleton())
+			if (AnimSegment.AnimReference && Composite->GetSkeleton()->IsCompatible(AnimSegment.AnimReference->GetSkeleton()))
 			{
 				Composite->AnimationTrack.AnimSegments[AnimSegmentIndex] = AnimSegment;
 				return true;
diff --git a/Engine/Source/Editor/UnrealEd/Private/Animation/EditorAnimSegment.cpp b/Engine/Source/Editor/UnrealEd/Private/Animation/EditorAnimSegment.cpp
--- a/Engine/Source/Editor/UnrealEd/Private/Animation/EditorAnimSegment.cpp
+++ b/Engine/Source/Editor/UnrealEd/Private/Animation/EditorAnimSegment.cpp
@@ -36,7 +36,7 @@ bool UEditorAnimSegment::ApplyChangesToMontage()
 	{
 		if(Montage->SlotAnimTracks.IsValidIndex(AnimSlotIndex) && Montage->SlotAnimTracks[AnimSlotIndex].AnimTrack.AnimSegments.IsValidIndex(AnimSegmentIndex) )
 		{
-			if (AnimSegment.AnimReference && AnimSegment.AnimReference->GetSkeleton() == Montage->GetSkeleton())
+			if (AnimSegment.AnimReference && Montage->GetSkeleton()->IsCompatible(AnimSegment.AnimReference->GetSkeleton()))
 			{
 				Montage->SlotAnimTracks[AnimSlotIndex].AnimTrack.AnimSegments[AnimSegmentIndex] = AnimSegment;
 				Montage->UpdateLinkableElements(AnimSlotIndex, AnimSegmentIndex);
diff --git a/Engine/Source/Editor/UnrealEd/Private/LevelEditorViewport.cpp b/Engine/Source/Editor/UnrealEd/Private/LevelEditorViewport.cpp
--- a/Engine/Source/Editor/UnrealEd/Private/LevelEditorViewport.cpp
+++ b/Engine/Source/Editor/UnrealEd/Private/LevelEditorViewport.cpp
@@ -666,7 +666,7 @@ static bool AttemptApplyObjToComponent(UObject* ObjToUse, USceneComponent* Compo
 						SkeletalMeshComponent->Modify();
 
 						// If the component doesn't have a mesh or the anim blueprint's skeleton isn't compatible with the existing mesh's skeleton, the mesh should change
-						const bool bShouldChangeMesh = !SkeletalMeshComponent->SkeletalMesh || !AnimBPSkeleton->IsCompatible(SkeletalMeshComponent->SkeletalMesh->GetSkeleton());
+						const bool bShouldChangeMesh = !SkeletalMeshComponent->SkeletalMesh || !SkeletalMeshComponent->SkeletalMesh->GetSkeleton()->IsCompatible(AnimBPSkeleton);
 
 						if (bShouldChangeMesh)
 						{
@@ -674,7 +674,7 @@ static bool AttemptApplyObjToComponent(UObject* ObjToUse, USceneComponent* Compo
 						}
 
 						// Verify that the skeletons are compatible before changing the anim BP
-						if (SkeletalMeshComponent->SkeletalMesh && AnimBPSkeleton->IsCompatible(SkeletalMeshComponent->SkeletalMesh->GetSkeleton()))
+						if (SkeletalMeshComponent->SkeletalMesh && SkeletalMeshComponent->SkeletalMesh->GetSkeleton()->IsCompatible(AnimBPSkeleton))
 						{
 							SkeletalMeshComponent->SetAnimInstanceClass(DroppedObjAsAnimBlueprint->GeneratedClass);
 							bResult = true;
@@ -701,7 +701,7 @@ static bool AttemptApplyObjToComponent(UObject* ObjToUse, USceneComponent* Compo
 						SkeletalMeshComponent->Modify();
 
 						// If the component doesn't have a mesh or the anim blueprint's skeleton isn't compatible with the existing mesh's skeleton, the mesh should change
-						const bool bShouldChangeMesh = !SkeletalMeshComponent->SkeletalMesh || !AnimSkeleton->IsCompatible(SkeletalMeshComponent->SkeletalMesh->GetSkeleton());
+						const bool bShouldChangeMesh = !SkeletalMeshComponent->SkeletalMesh || !SkeletalMeshComponent->SkeletalMesh->GetSkeleton()->IsCompatible(AnimSkeleton);
 
 						if (bShouldChangeMesh)
 						{
diff --git a/Engine/Source/Runtime/AnimGraphRuntime/Private/AnimNodes/AnimNode_PoseByName.cpp b/Engine/Source/Runtime/AnimGraphRuntime/Private/AnimNodes/AnimNode_PoseByName.cpp
--- a/Engine/Source/Runtime/AnimGraphRuntime/Private/AnimNodes/AnimNode_PoseByName.cpp
+++ b/Engine/Source/Runtime/AnimGraphRuntime/Private/AnimNodes/AnimNode_PoseByName.cpp
@@ -19,7 +19,7 @@ void FAnimNode_PoseByName::RebuildPoseList(const FBoneContainer& InBoneContainer
 	const TArray<FSmartName>& PoseNames = InPoseAsset->GetPoseNames();
 	const int32 PoseIndex = InPoseAsset->GetPoseIndexByName(PoseName);
 	TArray<uint16> const& LUTIndex = InBoneContainer.GetUIDToArrayLookupTable();
-	if (PoseIndex != INDEX_NONE && ensure(LUTIndex.IsValidIndex(PoseNames[PoseIndex].UID)) && LUTIndex[PoseNames[PoseIndex].UID] != MAX_uint16)
+	if (PoseIndex != INDEX_NONE && LUTIndex.IsValidIndex(PoseNames[PoseIndex].UID) && LUTIndex[PoseNames[PoseIndex].UID] != MAX_uint16)
 	{
 		// we keep pose index as that is the fastest way to search when extracting pose asset
 		PoseExtractContext.PoseCurves.Add(FPoseCurve(PoseIndex, PoseNames[PoseIndex].UID, 0.f));
diff --git a/Engine/Source/Runtime/AnimGraphRuntime/Private/AnimNodes/AnimNode_PoseHandler.cpp b/Engine/Source/Runtime/AnimGraphRuntime/Private/AnimNodes/AnimNode_PoseHandler.cpp
--- a/Engine/Source/Runtime/AnimGraphRuntime/Private/AnimNodes/AnimNode_PoseHandler.cpp
+++ b/Engine/Source/Runtime/AnimGraphRuntime/Private/AnimNodes/AnimNode_PoseHandler.cpp
@@ -60,7 +60,7 @@ void FAnimNode_PoseHandler::RebuildPoseList(const FBoneContainer& InBoneContaine
 		for (int32 PoseIndex = 0; PoseIndex < PoseNames.Num(); ++PoseIndex)
 		{
 			const FSmartName& PoseName = PoseNames[PoseIndex];
-			if (ensureMsgf(LUTIndex.IsValidIndex(PoseName.UID), TEXT("Invalid PoseName %s in PoseAsset %s for BoneContainer using %s"), *PoseName.DisplayName.ToString(), *GetPathNameSafe(InPoseAsset), *GetPathNameSafe(InBoneContainer.GetAsset())) && LUTIndex[PoseName.UID] != MAX_uint16)
+			if (LUTIndex.IsValidIndex(PoseName.UID) && LUTIndex[PoseName.UID] != MAX_uint16)
 			{
 				// we keep pose index as that is the fastest way to search when extracting pose asset
 				PoseExtractContext.PoseCurves.Add(FPoseCurve(PoseIndex, PoseName.UID, 0.f));
diff --git a/Engine/Source/Runtime/AnimGraphRuntime/Private/AnimNodes/AnimNode_SequenceEvaluator.cpp b/Engine/Source/Runtime/AnimGraphRuntime/Private/AnimNodes/AnimNode_SequenceEvaluator.cpp
--- a/Engine/Source/Runtime/AnimGraphRuntime/Private/AnimNodes/AnimNode_SequenceEvaluator.cpp
+++ b/Engine/Source/Runtime/AnimGraphRuntime/Private/AnimNodes/AnimNode_SequenceEvaluator.cpp
@@ -97,7 +97,7 @@ void FAnimNode_SequenceEvaluator::Evaluate_AnyThread(FPoseContext& Output)
 {
 	DECLARE_SCOPE_HIERARCHICAL_COUNTER_ANIMNODE(Evaluate_AnyThread)
 	check(Output.AnimInstanceProxy != nullptr);
-	if ((Sequence != nullptr) && (Output.AnimInstanceProxy->IsSkeletonCompatible(Sequence->GetSkeleton())))
+	if ((Sequence != nullptr) && Output.AnimInstanceProxy->IsSkeletonCompatible(Sequence->GetSkeleton()))
 	{
 		FAnimationPoseData AnimationPoseData(Output);
 		Sequence->GetAnimationPose(AnimationPoseData, FAnimExtractContext(InternalTimeAccumulator, Output.AnimInstanceProxy->ShouldExtractRootMotion()));
diff --git a/Engine/Source/Runtime/Engine/Classes/Animation/Skeleton.h b/Engine/Source/Runtime/Engine/Classes/Animation/Skeleton.h
--- a/Engine/Source/Runtime/Engine/Classes/Animation/Skeleton.h
+++ b/Engine/Source/Runtime/Engine/Classes/Animation/Skeleton.h
@@ -15,8 +15,10 @@
 #include "Animation/PreviewAssetAttachComponent.h"
 #include "Animation/SmartName.h"
 #include "Engine/AssetUserData.h"
+#include "HAL/CriticalSection.h"
 #include "Interfaces/Interface_AssetUserData.h"
 #include "Interfaces/Interface_PreviewMeshProvider.h"
+#include "Delegates/DelegateCombinations.h"
 #include "Skeleton.generated.h"
 
 class UAnimSequence;
@@ -25,6 +27,7 @@ class URig;
 class USkeletalMesh;
 class USkeletalMeshSocket;
 class FPackageReloadedEvent;
+struct FAssetData;
 enum class EPackageReloadPhase : uint8;
 
 /** This is a mapping table between bone in a particular skeletal mesh and bone of this skeleton set. */
@@ -265,6 +268,175 @@ public:
 	}
 };
 
+struct FSkeletonRemapping
+{
+	FSkeletonRemapping() = default;
+	FSkeletonRemapping(const USkeleton* InSourceSkeleton, const USkeleton* InTargetSkeleton);
+
+	const TWeakObjectPtr<const USkeleton>& GetSourceSkeleton() const { return SourceSkeleton; }
+	const TWeakObjectPtr<const USkeleton>& GetTargetSkeleton() const { return TargetSkeleton; }
+
+	/**
+	 * Compose this remapping with another remapping in place.  The other remapping's source skeleton must match
+	 * this remapping's target skeleton.  The result will map from this remapping's source skeleton to the other
+	 * remapping's target skeleton
+	 *
+	 * @param	OtherSkeletonRemapping		Skeleton remapping to compose into this remapping
+	 */
+	void ComposeWith(const FSkeletonRemapping& OtherSkeletonRemapping);
+
+	/**
+	 * Get the target skeleton bone index that corresponds to the specified bone on the source skeleton
+	 *
+	 * @param	SourceSkeletonBoneIndex		Skeleton bone index on the source skeleton
+	 * @return								Skeleton bone index on the target skeleton (or INDEX_NONE)
+	 */
+	int32 GetTargetSkeletonBoneIndex(int32 SourceSkeletonBoneIndex) const
+	{
+		return (SourceToTargetBoneIndexes.IsValidIndex(SourceSkeletonBoneIndex))
+				   ? SourceToTargetBoneIndexes[SourceSkeletonBoneIndex]
+				   : INDEX_NONE;
+	}
+
+	/**
+	 * Get the source skeleton bone index that corresponds to the specified bone on the target skeleton
+	 *
+	 * @param	TargetSkeletonBoneIndex		Skeleton bone index on the target skeleton
+	 * @return								Skeleton bone index on the source skeleton (or INDEX_NONE)
+	 */
+	int32 GetSourceSkeletonBoneIndex(int32 TargetSkeletonBoneIndex) const
+	{
+		return (TargetToSourceBoneIndexes.IsValidIndex(TargetSkeletonBoneIndex))
+				   ? TargetToSourceBoneIndexes[TargetSkeletonBoneIndex]
+				   : INDEX_NONE;
+	}
+
+	/**
+	 * Get the specified bone transform retargeted from the source skeleton onto the target skeleton, corrected
+	 * for differences between source and target rest poses
+	 *
+	 * @param	TargetSkeletonBoneIndex		Skeleton bone index on the target skeleton
+	 * @param	SourceTransform				Bone transform from the corresponding bone on the source skeleton
+	 * @return								Transform mapped onto the target skeleton
+	 */
+	FTransform RetargetBoneTransformToTargetSkeleton(int32 TargetSkeletonBoneIndex, const FTransform& SourceTransform) const
+	{
+		return FTransform(
+			RetargetBoneRotationToTargetSkeleton(TargetSkeletonBoneIndex, SourceTransform.GetRotation()),
+			RetargetBoneTranslationToTargetSkeleton(TargetSkeletonBoneIndex, SourceTransform.GetTranslation()),
+			SourceTransform.GetScale3D());
+	}
+
+	/**
+	 * Get the specified bone translation retargeted from the source skeleton onto the target skeleton, corrected
+	 * for differences between source and target rest poses
+	 *
+	 * @param	TargetSkeletonBoneIndex		Skeleton bone index on the target skeleton
+	 * @param	SourceTranslation			Bone translation from the corresponding bone on the source skeleton
+	 * @return								Translation mapped onto the target skeleton
+	 */
+	FVector RetargetBoneTranslationToTargetSkeleton(int32 TargetSkeletonBoneIndex, const FVector& SourceTranslation) const
+	{
+		// Compute the translation part of FTransform(Q1) * Source * FTransform(Q0)
+		const TTuple<FQuat, FQuat>& QQ = RetargetingTable[TargetSkeletonBoneIndex];
+		return QQ.Get<0>().RotateVector(SourceTranslation);
+	}
+
+	/**
+	 * Get the specified bone rotation retargeted from the source skeleton onto the target skeleton, corrected
+	 * for differences between source and target rest poses
+	 *
+	 * @param	TargetSkeletonBoneIndex		Skeleton bone index on the target skeleton
+	 * @param	SourceRotation				Bone rotation from the corresponding bone on the source skeleton
+	 * @return								Rotation mapped onto the target skeleton
+	 */
+	FQuat RetargetBoneRotationToTargetSkeleton(int32 TargetSkeletonBoneIndex, const FQuat& SourceRotation) const
+	{
+		// Compute the rotation part of FTransform(Q1) * Source * FTransform(Q0)
+		const TTuple<FQuat, FQuat>& QQ = RetargetingTable[TargetSkeletonBoneIndex];
+		return QQ.Get<0>() * SourceRotation * QQ.Get<1>();
+	}
+
+	/**
+	 * Get the specified additive transform retargeted from the source skeleton onto the target skeleton, corrected
+	 * for differences between source and target rest poses
+	 *
+	 * @param	TargetSkeletonBoneIndex		Skeleton bone index on the target skeleton
+	 * @param	SourceTransform				Bone transform from the corresponding bone on the source skeleton
+	 * @return								Transform mapped onto the target skeleton
+	 */
+	FTransform RetargetAdditiveTransformToTargetSkeleton(int32 TargetSkeletonBoneIndex, const FTransform& SourceTransform) const
+	{
+		return FTransform(
+			RetargetAdditiveRotationToTargetSkeleton(TargetSkeletonBoneIndex, SourceTransform.GetRotation()),
+			RetargetAdditiveTranslationToTargetSkeleton(TargetSkeletonBoneIndex, SourceTransform.GetTranslation()),
+			SourceTransform.GetScale3D());
+	}
+
+	/**
+	 * Get the specified additive translation retargeted from the source skeleton onto the target skeleton, corrected
+	 * for differences between source and target rest poses
+	 *
+	 * @param	TargetSkeletonBoneIndex		Skeleton bone index on the target skeleton
+	 * @param	SourceTranslation			Bone translation from the corresponding bone on the source skeleton
+	 * @return								Translation mapped onto the target skeleton
+	 */
+	FVector RetargetAdditiveTranslationToTargetSkeleton(int32 TargetSkeletonBoneIndex, const FVector& SourceTranslation) const
+	{
+		// Compute the translation part of FTransform(Q0.Inverse) * Source * FTransform(Q0)
+		const TTuple<FQuat, FQuat>& QQ = RetargetingTable[TargetSkeletonBoneIndex];
+		return QQ.Get<0>().RotateVector(SourceTranslation);
+	}
+
+	/**
+	 * Get the specified additive rotation retargeted from the source skeleton onto the target skeleton, corrected
+	 * for differences between source and target rest poses
+	 *
+	 * @param	TargetSkeletonBoneIndex		Skeleton bone index on the target skeleton
+	 * @param	SourceRotation				Bone rotation from the corresponding bone on the source skeleton
+	 * @return								Rotation mapped onto the target skeleton
+	 */
+	FQuat RetargetAdditiveRotationToTargetSkeleton(int32 TargetSkeletonBoneIndex, const FQuat& SourceRotation) const
+	{
+		// Compute the rotation part of FTransform(Q0.Inverse) * Source * FTransform(Q0)
+		const TTuple<FQuat, FQuat>& QQ = RetargetingTable[TargetSkeletonBoneIndex];
+		return QQ.Get<0>() * SourceRotation * QQ.Get<0>().Inverse();
+	}
+
+	/**
+	 * Get the curve mapping array, which maps from the source curve UID to the target curve UID.
+	 * You can access this array like this:
+	 * 
+	 * \code{.cpp}
+	 * const SmartName::UID_Type TargetCurveUID = GetSourceToTargetCurveMapping()[SourceCurveUID];
+	 * \endcode
+	 */
+	const TArray<SmartName::UID_Type>& GetSourceToTargetCurveMapping() const { return SourceToTargetCurveMapping; }
+
+	/**
+	 * Generates the mapping table for the curves, based on curve names.
+	 * Basically this will look at the curve names in the smart name table on the source, and tries to find matching ones in the table of 
+	 * the target skeleton. It then maps the UID's.
+	 */
+	void GenerateCurveMapping();
+
+private:
+	TWeakObjectPtr<const USkeleton> SourceSkeleton;
+	TWeakObjectPtr<const USkeleton> TargetSkeleton;
+
+	// Table of target skeleton bone indexes (indexed by source skeleton bone index)
+	TArray<int32> SourceToTargetBoneIndexes;
+
+	// Table of source skeleton bone indexes (indexed by target skeleton bone index)
+	TArray<int32> TargetToSourceBoneIndexes;
+	
+	// Maps curve UIDs between source and target (indexed by source curve UID).
+	TArray<SmartName::UID_Type> SourceToTargetCurveMapping;
+
+	// Table of precalculated constants for retargeting from source to target (indexed by target skeleton bone index)
+	TArray<TTuple<FQuat, FQuat>> RetargetingTable;
+};
+
 /**
  *	USkeleton : that links between mesh and animation
  *		- Bone hierarchy for animations
@@ -310,9 +482,17 @@ protected:
 	UPROPERTY()
 	TArray<FVirtualBone> VirtualBones;
 
+	/**
+	 * The list of compatible skeletons.
+	 * This is an array of TSoftObjectPtr in order to prevent all skeletons to be loaded, as we only want to load things on demand.
+	 **/
+	UPROPERTY(EditAnywhere, BlueprintReadOnly, Category = CompatibleSkeletons)
+	TArray<TSoftObjectPtr<USkeleton>> CompatibleSkeletons;
+
 public:
 	//~ Begin UObject Interface.
 #if WITH_EDITOR
+	ENGINE_API virtual void PostEditChangeProperty(struct FPropertyChangedEvent& PropertyChangedEvent);
 	ENGINE_API virtual void PreEditUndo() override;
 	ENGINE_API virtual void PostEditUndo() override;
 #endif
@@ -370,6 +550,9 @@ public:
 	uint16 GetAnimCurveUidVersion() const { return AnimCurveUidVersion;  }
 	const TArray<uint16>& GetDefaultCurveUIDList() const { return DefaultCurveUIDList; }
 
+	ENGINE_API const TArray<TSoftObjectPtr<USkeleton>>& GetCompatibleSkeletons() const { return CompatibleSkeletons; }
+	ENGINE_API const FSkeletonRemapping* GetSkeletonRemapping(const USkeleton* SourceSkeleton) const;
+
 #if WITH_EDITOR
 	// Get existing (seen) sync marker names for this Skeleton
 	const TArray<FName>& GetExistingMarkerNames() const { return ExistingMarkerNames; }
@@ -549,7 +732,6 @@ public:
 		OnRetargetSourceChanged.Broadcast();
 	}
 
-
 	typedef TArray<FBoneNode> FBoneTreeType;
 
 	/** Non-serialised cache of linkups between different skeletal meshes and this Skeleton. */
@@ -558,6 +740,9 @@ public:
 	/** Runtime built mapping table between SkeletalMeshes, and LinkupCache array indices. */
 	TMap<TWeakObjectPtr<USkeletalMesh>, int32> SkelMesh2LinkupCache;
 
+	/** A cached soft object pointer of this skeleton. This is done for performance reasons when searching for compatible skeletons when using IsCompatible(Skeleton). */
+	TSoftObjectPtr<USkeleton> CachedSoftObjectPtr;
+
 	/** IInterface_PreviewMeshProvider interface */
 	virtual USkeletalMesh* GetPreviewMesh(bool bFindIfNotSet = false) override;
 	virtual USkeletalMesh* GetPreviewMesh() const override;
@@ -610,12 +795,82 @@ public:
 	ENGINE_API int32 GetChildBones(int32 ParentBoneIndex, TArray<int32> & Children) const;
 
 #endif
+	/**
+	 * Check if animation content authored on the supplied skeleton may be played on this skeleton
+	 */
+	ENGINE_API bool IsCompatible(const USkeleton* InSkeleton) const;
+
+	/**
+	 * Check if this skeleton is compatible with a given other asset, if that is a skeleton.
+	 */
+	ENGINE_API bool IsCompatibleSkeletonByAssetData(const FAssetData& AssetData) const;
 
 	/**
-	 *	Check if this skeleton may be used with other skeleton
+	 * Check if this skeleton is compatible with another skeleton asset that is identified by the string returned by AssetData(SkeletonPtr).GetExportTextName().
 	 */
-	ENGINE_API bool IsCompatible(USkeleton const * InSkeleton) const { return (InSkeleton && this == InSkeleton); }
+	ENGINE_API bool IsCompatibleSkeletonByAssetString(const FString& SkeletonAssetString) const;
+
+	DECLARE_EVENT_OneParam(USkeleton, FSkeletonDestructEvent, const USkeleton*);
+	FSkeletonDestructEvent OnSkeletonDestructEvent;
 
+	DECLARE_EVENT(USkeleton, FSmartNamesChangedEvent);
+	FSmartNamesChangedEvent OnSmartNamesChangedEvent;
+
+protected:
+	mutable FCriticalSection SkeletonRemappingMutex;
+
+	/** The skeleton remappings, which map bones from a source skeleton to this skeleton (the target). */
+	TMap<const USkeleton*, FSkeletonRemapping*> SkeletonRemappings;
+
+	/**
+ 	 * The function that is called when another skeleton is being destructed.
+	 * This is used to unregister skeleton remappings that aren't needed anymore.
+	 * @param Skeleton The skeleton that is being destructed.
+	 */
+	void HandleSkeletonDestruct(const USkeleton* Skeleton);
+
+	/**
+	 * Handle smart name changes.
+	 * This will internally trigger skeleton remappings to be regenerated where needed, for example when editing curves in the animation editor.
+	 **/
+	void HandleSmartNamesChangedEvent();
+
+	/**
+	 * Build or update all skeleton remappings for all loaded skeletons that we are compatible with.
+	 * And if bBidirectional is set to true also update all remappings in other skeletons that are compatible with this skeleton.
+	 * Keep in mind this only generates mappings to skeletons that are loaded, as otherwise we can't generate the mappings yet, even though the list of compatible skeletons might contain
+	 * a larger set of skeletons. We do not load all skeletons upfront.
+	 * @param bBidirectional When enabled we update or create all mappings for both skeletons we are compatible with and skeletons that are compatible with ourselves.
+	 *        When set to false it only updates or creates mappings to skeletons we are compatible with.
+	 */
+	void BuildSkeletonRemappings(bool bBidirectional);
+
+	/**
+	 * Remove all skeleton remappings stored inside our skeleton.
+	 * This frees up memory.
+	 */
+	void ClearSkeletonRemappings();
+
+	/**
+	 * Remove the skeleton remappings for a specific source skeleton.
+	 * This only works if the source skeleton is inside the compatible skeleton list of this skeleton.
+	 * If a skeleton is passed in that isn't in the compatibility list this method will essentially do nothing.
+	 * The skeleton will remain inside the compatibility list, it is just the remapping data that is being removed.
+	 * @param SourceSkeleton The source skeleton we are compatible with and want to remove the generated mapping data for.
+	 */
+	void RemoveSkeletonRemapping(const USkeleton* SourceSkeleton);
+
+	/**
+	 * Create or update a skeletal remapping to a given source skeleton.
+	 * @param SourceSkeleton The skeleton from which we receive animation data.
+	 * @param bRebuildIfExists When set to true this method will delete the old mapping and generate a new one. This is useful when the hierarchy changes for example.
+	 */
+	void CreateSkeletonRemappingIfNeeded(const USkeleton* SourceSkeleton, bool bRebuildIfExists = false);
+
+	static TArray<USkeleton*> LoadedSkeletons; // The set of skeletons that are currently loaded.
+	static FCriticalSection LoadedSkeletonsMutex; // The mutex used for thread safety when registering and removing items from the LoadedSkeletons array.
+
+public:
 	/** 
 	 * Indexing naming convention
 	 * 
@@ -886,9 +1141,6 @@ private:
 	// Handle skeletons being reloaded via the content browser
 	static void HandlePackageReloaded(const EPackageReloadPhase InPackageReloadPhase, FPackageReloadedEvent* InPackageReloadedEvent);
 
-
-
-
 public:
 	//~ Begin IInterface_AssetUserData Interface
 	ENGINE_API virtual void AddAssetUserData(UAssetUserData* InUserData) override;
diff --git a/Engine/Source/Runtime/Engine/Private/Animation/AnimCompressionTypes.cpp b/Engine/Source/Runtime/Engine/Private/Animation/AnimCompressionTypes.cpp
--- a/Engine/Source/Runtime/Engine/Private/Animation/AnimCompressionTypes.cpp
+++ b/Engine/Source/Runtime/Engine/Private/Animation/AnimCompressionTypes.cpp
@@ -698,19 +698,22 @@ struct FGetBonePoseScratchArea : public TThreadSingleton<FGetBonePoseScratchArea
 	BoneTrackArray OrientAndScaleRetargetingPairs;
 };
 
-void DecompressPose(FCompactPose& OutPose, const FCompressedAnimSequence& CompressedData, const FAnimExtractContext& ExtractionContext, USkeleton* Skeleton, float SequenceLength, EAnimInterpolationType Interpolation, bool bIsBakedAdditive, FName RetargetSource, FName SourceName, const FRootMotionReset& RootMotionReset)
+void DecompressPose(FCompactPose& OutPose, const FCompressedAnimSequence& CompressedData, const FAnimExtractContext& ExtractionContext, USkeleton* SourceSkeleton, float SequenceLength, EAnimInterpolationType Interpolation, bool bIsBakedAdditive, FName RetargetSource, FName SourceName, const FRootMotionReset& RootMotionReset)
 {
-	const TArray<FTransform>& RetargetTransforms = Skeleton->GetRefLocalPoses(RetargetSource);
-	DecompressPose(OutPose, CompressedData, ExtractionContext, Skeleton, SequenceLength, Interpolation, bIsBakedAdditive, RetargetTransforms, SourceName, RootMotionReset);
+	const TArray<FTransform>& RetargetTransforms = SourceSkeleton->GetRefLocalPoses(RetargetSource);
+	DecompressPose(OutPose, CompressedData, ExtractionContext, SourceSkeleton, SequenceLength, Interpolation, bIsBakedAdditive, RetargetTransforms, SourceName, RootMotionReset);
 }
 
-void DecompressPose(FCompactPose& OutPose, const FCompressedAnimSequence& CompressedData, const FAnimExtractContext& ExtractionContext, USkeleton* Skeleton, float SequenceLength, EAnimInterpolationType Interpolation, bool bIsBakedAdditive, const TArray<FTransform>& RetargetTransforms, FName SourceName, const FRootMotionReset& RootMotionReset)
+void DecompressPose(FCompactPose& OutPose, const FCompressedAnimSequence& CompressedData, const FAnimExtractContext& ExtractionContext, USkeleton* SourceSkeleton, float SequenceLength, EAnimInterpolationType Interpolation, bool bIsBakedAdditive, const TArray<FTransform>& RetargetTransforms, FName SourceName, const FRootMotionReset& RootMotionReset)
 {
 	const FBoneContainer& RequiredBones = OutPose.GetBoneContainer();
 	const int32 NumTracks = CompressedData.CompressedTrackToSkeletonMapTable.Num();
 
 	TArray<int32> const& SkeletonToPoseBoneIndexArray = RequiredBones.GetSkeletonToPoseBoneIndexArray();
 
+	const USkeleton* TargetSkeleton = RequiredBones.GetSkeletonAsset();
+	const FSkeletonRemapping* SkeletonRemapping = TargetSkeleton->GetSkeletonRemapping(SourceSkeleton);
+
 	BoneTrackArray& RotationScalePairs = FGetBonePoseScratchArea::Get().RotationScalePairs;
 	BoneTrackArray& TranslationPairs = FGetBonePoseScratchArea::Get().TranslationPairs;
 	BoneTrackArray& AnimScaleRetargetingPairs = FGetBonePoseScratchArea::Get().AnimScaleRetargetingPairs;
@@ -735,11 +738,12 @@ void DecompressPose(FCompactPose& OutPose, const FCompressedAnimSequence& Compre
 		// Handle root bone separately if it is track 0. so we start w/ Index 1.
 		for (int32 TrackIndex = (bFirstTrackIsRootBone ? 1 : 0); TrackIndex < NumTracks; TrackIndex++)
 		{
-			const int32 SkeletonBoneIndex = CompressedData.GetSkeletonIndexFromTrackIndex(TrackIndex);
-			// not sure it's safe to assume that SkeletonBoneIndex can never be INDEX_NONE
-			if (SkeletonBoneIndex != INDEX_NONE)
+			const int32 SourceSkeletonBoneIndex = CompressedData.GetSkeletonIndexFromTrackIndex(TrackIndex);
+			const int32 TargetSkeletonBoneIndex = (SkeletonRemapping) ? SkeletonRemapping->GetTargetSkeletonBoneIndex(SourceSkeletonBoneIndex) : SourceSkeletonBoneIndex;
+
+			if (TargetSkeletonBoneIndex != INDEX_NONE)
 			{
-				const FCompactPoseBoneIndex BoneIndex = RequiredBones.GetCompactPoseIndexFromSkeletonIndex(SkeletonBoneIndex);
+				const FCompactPoseBoneIndex BoneIndex = RequiredBones.GetCompactPoseIndexFromSkeletonIndex(TargetSkeletonBoneIndex);
 				//Nasty, we break our type safety, code in the lower levels should be adjusted for this
 				const int32 CompactPoseBoneIndex = BoneIndex.GetInt();
 				if (CompactPoseBoneIndex != INDEX_NONE)
@@ -747,14 +751,14 @@ void DecompressPose(FCompactPose& OutPose, const FCompressedAnimSequence& Compre
 					RotationScalePairs.Add(BoneTrackPair(CompactPoseBoneIndex, TrackIndex));
 
 					// Skip extracting translation component for EBoneTranslationRetargetingMode::Skeleton.
-					switch (Skeleton->GetBoneTranslationRetargetingMode(SkeletonBoneIndex))
+					switch (TargetSkeleton->GetBoneTranslationRetargetingMode(TargetSkeletonBoneIndex))
 					{
 					case EBoneTranslationRetargetingMode::Animation:
 						TranslationPairs.Add(BoneTrackPair(CompactPoseBoneIndex, TrackIndex));
 						break;
 					case EBoneTranslationRetargetingMode::AnimationScaled:
 						TranslationPairs.Add(BoneTrackPair(CompactPoseBoneIndex, TrackIndex));
-						AnimScaleRetargetingPairs.Add(BoneTrackPair(CompactPoseBoneIndex, SkeletonBoneIndex));
+						AnimScaleRetargetingPairs.Add(BoneTrackPair(CompactPoseBoneIndex, SourceSkeletonBoneIndex));
 						break;
 					case EBoneTranslationRetargetingMode::AnimationRelative:
 						TranslationPairs.Add(BoneTrackPair(CompactPoseBoneIndex, TrackIndex));
@@ -763,7 +767,7 @@ void DecompressPose(FCompactPose& OutPose, const FCompressedAnimSequence& Compre
 						// (A1 + Rel) - (A2 + Rel) = A1 - A2.
 						if (!bIsBakedAdditive)
 						{
-							AnimRelativeRetargetingPairs.Add(BoneTrackPair(CompactPoseBoneIndex, SkeletonBoneIndex));
+							AnimRelativeRetargetingPairs.Add(BoneTrackPair(CompactPoseBoneIndex, SourceSkeletonBoneIndex));
 						}
 						break;
 					case EBoneTranslationRetargetingMode::OrientAndScale:
@@ -772,7 +776,7 @@ void DecompressPose(FCompactPose& OutPose, const FCompressedAnimSequence& Compre
 						// Additives remain additives, they're not retargeted.
 						if (!bIsBakedAdditive)
 						{
-							OrientAndScaleRetargetingPairs.Add(BoneTrackPair(CompactPoseBoneIndex, SkeletonBoneIndex));
+							OrientAndScaleRetargetingPairs.Add(BoneTrackPair(CompactPoseBoneIndex, SourceSkeletonBoneIndex));
 						}
 						break;
 					}
@@ -798,8 +802,31 @@ void DecompressPose(FCompactPose& OutPose, const FCompressedAnimSequence& Compre
 
 			CompressedData.BoneCompressionCodec->DecompressBone(EvalDecompContext, TrackIndex, RootAtom);
 
+			// Retarget the root onto the target skeleton (correcting for differences in rest poses)
+			if (SkeletonRemapping)
+			{
+				const int32 TargetSkeletonBoneIndex = 0;
+
+				if (bIsBakedAdditive)
+				{
+					RootAtom.SetRotation(SkeletonRemapping->RetargetAdditiveRotationToTargetSkeleton(TargetSkeletonBoneIndex, RootAtom.GetRotation()));
+					if (TargetSkeleton->GetBoneTranslationRetargetingMode(TargetSkeletonBoneIndex) != EBoneTranslationRetargetingMode::Skeleton)
+					{
+						RootAtom.SetTranslation(SkeletonRemapping->RetargetAdditiveTranslationToTargetSkeleton(TargetSkeletonBoneIndex, RootAtom.GetTranslation()));
+					}
+				}
+				else
+				{
+					RootAtom.SetRotation(SkeletonRemapping->RetargetBoneRotationToTargetSkeleton(TargetSkeletonBoneIndex, RootAtom.GetRotation()));
+					if (TargetSkeleton->GetBoneTranslationRetargetingMode(TargetSkeletonBoneIndex) != EBoneTranslationRetargetingMode::Skeleton)
+					{
+						RootAtom.SetTranslation(SkeletonRemapping->RetargetBoneTranslationToTargetSkeleton(TargetSkeletonBoneIndex, RootAtom.GetTranslation()));
+					}
+				}
+			}
+
 			// @laurent - we should look into splitting rotation and translation tracks, so we don't have to process translation twice.
-			FAnimationRuntime::RetargetBoneTransform(Skeleton, SourceName, RetargetTransforms, RootAtom, 0, RootBone, RequiredBones, bIsBakedAdditive);
+			FAnimationRuntime::RetargetBoneTransform(SourceSkeleton, SourceName, RetargetTransforms, RootAtom, 0, RootBone, RequiredBones, bIsBakedAdditive);
 		}
 
 		if (RotationScalePairs.Num() > 0)
@@ -810,6 +837,35 @@ void DecompressPose(FCompactPose& OutPose, const FCompressedAnimSequence& Compre
 		}
 	}
 
+	// Retarget the pose onto the target skeleton (correcting for differences in rest poses)
+	if (SkeletonRemapping)
+	{
+		if (bIsBakedAdditive)
+		{
+			for (FCompactPoseBoneIndex BoneIndex(bFirstTrackIsRootBone ? 1 : 0); BoneIndex < OutPose.GetNumBones(); ++BoneIndex)
+			{
+				const int32 TargetSkeletonBoneIndex = RequiredBones.GetSkeletonIndex(BoneIndex);
+				OutPose[BoneIndex].SetRotation(SkeletonRemapping->RetargetAdditiveRotationToTargetSkeleton(TargetSkeletonBoneIndex, OutPose[BoneIndex].GetRotation()));
+				if (TargetSkeleton->GetBoneTranslationRetargetingMode(TargetSkeletonBoneIndex) != EBoneTranslationRetargetingMode::Skeleton)
+				{
+					OutPose[BoneIndex].SetTranslation(SkeletonRemapping->RetargetAdditiveTranslationToTargetSkeleton(TargetSkeletonBoneIndex, OutPose[BoneIndex].GetTranslation()));
+				}
+			}
+		}
+		else
+		{
+			for (FCompactPoseBoneIndex BoneIndex(bFirstTrackIsRootBone ? 1 : 0); BoneIndex < OutPose.GetNumBones(); ++BoneIndex)
+			{
+				const int32 TargetSkeletonBoneIndex = RequiredBones.GetSkeletonIndex(BoneIndex);
+				OutPose[BoneIndex].SetRotation(SkeletonRemapping->RetargetBoneRotationToTargetSkeleton(TargetSkeletonBoneIndex, OutPose[BoneIndex].GetRotation()));
+				if (TargetSkeleton->GetBoneTranslationRetargetingMode(TargetSkeletonBoneIndex) != EBoneTranslationRetargetingMode::Skeleton)
+				{
+					OutPose[BoneIndex].SetTranslation(SkeletonRemapping->RetargetBoneTranslationToTargetSkeleton(TargetSkeletonBoneIndex, OutPose[BoneIndex].GetTranslation()));
+				}
+			}
+		}
+	}
+
 	// Once pose has been extracted, snap root bone back to first frame if we are extracting root motion.
 	if ((ExtractionContext.bExtractRootMotion && RootMotionReset.bEnableRootMotion) || RootMotionReset.bForceRootLock)
 	{
@@ -825,10 +881,10 @@ void DecompressPose(FCompactPose& OutPose, const FCompressedAnimSequence& Compre
 		for (const BoneTrackPair& BonePair : AnimScaleRetargetingPairs)
 		{
 			const FCompactPoseBoneIndex BoneIndex(BonePair.AtomIndex); //Nasty, we break our type safety, code in the lower levels should be adjusted for this
-			int32 const& SkeletonBoneIndex = BonePair.TrackIndex;
+			const int32 SourceSkeletonBoneIndex = BonePair.TrackIndex;
 
 			// @todo - precache that in FBoneContainer when we have SkeletonIndex->TrackIndex mapping. So we can just apply scale right away.
-			float const SourceTranslationLength = AuthoredOnRefSkeleton[SkeletonBoneIndex].GetTranslation().Size();
+			float const SourceTranslationLength = AuthoredOnRefSkeleton[SourceSkeletonBoneIndex].GetTranslation().Size();
 			if (SourceTranslationLength > KINDA_SMALL_NUMBER)
 			{
 				float const TargetTranslationLength = RequiredBones.GetRefPoseTransform(BoneIndex).GetTranslation().Size();
@@ -846,14 +902,22 @@ void DecompressPose(FCompactPose& OutPose, const FCompressedAnimSequence& Compre
 		for (const BoneTrackPair& BonePair : AnimRelativeRetargetingPairs)
 		{
 			const FCompactPoseBoneIndex BoneIndex(BonePair.AtomIndex); //Nasty, we break our type safety, code in the lower levels should be adjusted for this
-			int32 const& SkeletonBoneIndex = BonePair.TrackIndex;
+			const int32 SourceSkeletonBoneIndex = BonePair.TrackIndex;
+
+			const FTransform& RefPoseTransform = RequiredBones.GetRefPoseTransform(BoneIndex);
 
-			const FTransform& RefPose = RequiredBones.GetRefPoseTransform(BoneIndex);
+			// Remap the base pose onto the target skeleton so that we are working entirely in target space
+			FTransform BaseTransform = AuthoredOnRefSkeleton[SourceSkeletonBoneIndex];
+			if (SkeletonRemapping)
+			{
+				const int32 TargetSkeletonBoneIndex = SkeletonRemapping->GetTargetSkeletonBoneIndex(SourceSkeletonBoneIndex);
+				BaseTransform = SkeletonRemapping->RetargetBoneTransformToTargetSkeleton(TargetSkeletonBoneIndex, BaseTransform);
+			}
 
 			// Apply the retargeting as if it were an additive difference between the current skeleton and the retarget skeleton. 
-			OutPose[BoneIndex].SetRotation(OutPose[BoneIndex].GetRotation() * AuthoredOnRefSkeleton[SkeletonBoneIndex].GetRotation().Inverse() * RefPose.GetRotation());
-			OutPose[BoneIndex].SetTranslation(OutPose[BoneIndex].GetTranslation() + (RefPose.GetTranslation() - AuthoredOnRefSkeleton[SkeletonBoneIndex].GetTranslation()));
-			OutPose[BoneIndex].SetScale3D(OutPose[BoneIndex].GetScale3D() * (RefPose.GetScale3D() * AuthoredOnRefSkeleton[SkeletonBoneIndex].GetSafeScaleReciprocal(AuthoredOnRefSkeleton[SkeletonBoneIndex].GetScale3D())));
+			OutPose[BoneIndex].SetRotation(OutPose[BoneIndex].GetRotation() * BaseTransform.GetRotation().Inverse() * RefPoseTransform.GetRotation());
+			OutPose[BoneIndex].SetTranslation(OutPose[BoneIndex].GetTranslation() + (RefPoseTransform.GetTranslation() - BaseTransform.GetTranslation()));
+			OutPose[BoneIndex].SetScale3D(OutPose[BoneIndex].GetScale3D() * (RefPoseTransform.GetScale3D() * BaseTransform.GetSafeScaleReciprocal(BaseTransform.GetScale3D())));
 			OutPose[BoneIndex].NormalizeRotation();
 		}
 	}
diff --git a/Engine/Source/Runtime/Engine/Private/Animation/AnimMontage.cpp b/Engine/Source/Runtime/Engine/Private/Animation/AnimMontage.cpp
--- a/Engine/Source/Runtime/Engine/Private/Animation/AnimMontage.cpp
+++ b/Engine/Source/Runtime/Engine/Private/Animation/AnimMontage.cpp
@@ -425,7 +425,7 @@ void UAnimMontage::PostLoad()
 			if ( SlotAnimTracks[I].AnimTrack.AnimSegments.Num() > 0 )
 			{
 				UAnimSequenceBase* SequenceBase = SlotAnimTracks[I].AnimTrack.AnimSegments[0].AnimReference;
-				if (SequenceBase && SequenceBase->GetSkeleton() != MySkeleton )
+				if (SequenceBase && !MySkeleton->IsCompatible(SequenceBase->GetSkeleton()))
 				{
 					SlotAnimTracks[I].AnimTrack.AnimSegments[0].AnimReference = nullptr;
 					MarkPackageDirty();
diff --git a/Engine/Source/Runtime/Engine/Private/Animation/AnimNode_SequencePlayer.cpp b/Engine/Source/Runtime/Engine/Private/Animation/AnimNode_SequencePlayer.cpp
--- a/Engine/Source/Runtime/Engine/Private/Animation/AnimNode_SequencePlayer.cpp
+++ b/Engine/Source/Runtime/Engine/Private/Animation/AnimNode_SequencePlayer.cpp
@@ -76,7 +76,7 @@ void FAnimNode_SequencePlayer::UpdateAssetPlayer(const FAnimationUpdateContext&
 		Sequence = nullptr;
 	}
 
-	if ((Sequence != nullptr) && (Context.AnimInstanceProxy->IsSkeletonCompatible(Sequence->GetSkeleton())))
+	if ((Sequence != nullptr) && Context.AnimInstanceProxy->IsSkeletonCompatible(Sequence->GetSkeleton()))
 	{
 		InternalTimeAccumulator = FMath::Clamp(InternalTimeAccumulator, 0.f, Sequence->GetPlayLength());
 		const float AdjustedPlayRate = PlayRateScaleBiasClamp.ApplyTo(FMath::IsNearlyZero(PlayRateBasis) ? 0.f : (PlayRate / PlayRateBasis), Context.GetDeltaTime());
diff --git a/Engine/Source/Runtime/Engine/Private/Animation/AnimSequence.cpp b/Engine/Source/Runtime/Engine/Private/Animation/AnimSequence.cpp
--- a/Engine/Source/Runtime/Engine/Private/Animation/AnimSequence.cpp
+++ b/Engine/Source/Runtime/Engine/Private/Animation/AnimSequence.cpp
@@ -846,9 +846,9 @@ void UAnimSequence::PostLoad()
 	// but this causes issue since once this happens this is unrecoverable until you delete from outside of editor
 	if (IsValidAdditive())
 	{
-		if (RefPoseSeq && RefPoseSeq->GetSkeleton() != GetSkeleton())
+		if (RefPoseSeq && !GetSkeleton()->IsCompatible(RefPoseSeq->GetSkeleton()))
 		{
-			// if this happens, there was a issue with retargeting, 
+			// if this happens, there was a issue with retargeting
 			UE_LOG(LogAnimation, Warning, TEXT("Animation %s - Invalid additive animation base animation (%s)"), *GetName(), *RefPoseSeq->GetName());
 			RefPoseSeq = nullptr;
 		}
@@ -1432,7 +1432,7 @@ void UAnimSequence::GetBonePose(FAnimationPoseData& OutAnimationPoseData, const
 	{
 		if (bIsBakedAdditive)
 		{
-			OutPose.ResetToAdditiveIdentity(); 
+			OutPose.ResetToAdditiveIdentity();
 		}
 		else
 		{
@@ -1473,8 +1473,12 @@ void UAnimSequence::GetBonePose(FAnimationPoseData& OutAnimationPoseData, const
 		}
 	}
 
-	// extract curve data . Even if no track, it can contain curve data
-	EvaluateCurveData(OutAnimationPoseData.GetCurve(), ExtractionContext.CurrentTime, bUseRawDataForPoseExtraction);
+	// Evaluate the (remapped if needed) curves.
+	// Scoped so that the RemappedCurve destructs and isn't kept alive longer than needed.
+	{
+		FSkeletonRemappingCurve RemappedCurve(OutAnimationPoseData.GetCurve(), RequiredBones, GetSkeleton());
+		EvaluateCurveData(RemappedCurve.GetCurve(), ExtractionContext.CurrentTime, bUseRawDataForPoseExtraction);
+	}
 
 #if WITH_EDITOR
 	ValidateModel();
diff --git a/Engine/Source/Runtime/Engine/Private/Animation/AnimSingleNodeInstance.cpp b/Engine/Source/Runtime/Engine/Private/Animation/AnimSingleNodeInstance.cpp
--- a/Engine/Source/Runtime/Engine/Private/Animation/AnimSingleNodeInstance.cpp
+++ b/Engine/Source/Runtime/Engine/Private/Animation/AnimSingleNodeInstance.cpp
@@ -48,7 +48,7 @@ void UAnimSingleNodeInstance::SetAnimationAsset(class UAnimationAsset* NewAsset,
 		else if (CurrentAsset != nullptr)
 		{
 			// if we have an asset, make sure their skeleton matches, otherwise, null it
-			if (MeshComponent->SkeletalMesh->GetSkeleton() != CurrentAsset->GetSkeleton())
+			if (!MeshComponent->SkeletalMesh->GetSkeleton()->IsCompatible(CurrentAsset->GetSkeleton()))
 			{
 				// clear asset since we do not have matching skeleton
 				CurrentAsset = nullptr;
diff --git a/Engine/Source/Runtime/Engine/Private/Animation/AnimSingleNodeInstanceProxy.cpp b/Engine/Source/Runtime/Engine/Private/Animation/AnimSingleNodeInstanceProxy.cpp
--- a/Engine/Source/Runtime/Engine/Private/Animation/AnimSingleNodeInstanceProxy.cpp
+++ b/Engine/Source/Runtime/Engine/Private/Animation/AnimSingleNodeInstanceProxy.cpp
@@ -434,12 +434,13 @@ void FAnimNode_SingleNode::Evaluate_AnyThread(FPoseContext& Output)
 		// have to propagate output curve before pose asset as it can use pose curve data
 		Proxy->PropagatePreviewCurve(Output);
 
-		// if it has preview pose asset, we have to handle that after we do all animation
+		// if it has a preview pose asset, we have to handle that after we do all animation
 		if (const UPoseAsset* PoseAsset = Proxy->CurrentAsset->PreviewPoseAsset)
 		{
 			USkeleton* MySkeleton = Proxy->CurrentAsset->GetSkeleton();
+
 			// if skeleton doesn't match it won't work
-			if (PoseAsset->GetSkeleton() == MySkeleton)
+			if (MySkeleton->IsCompatible(PoseAsset->GetSkeleton()))
 			{
 				const TArray<FSmartName>& PoseNames = PoseAsset->GetPoseNames();
 
diff --git a/Engine/Source/Runtime/Engine/Private/Animation/AnimationRuntime.cpp b/Engine/Source/Runtime/Engine/Private/Animation/AnimationRuntime.cpp
--- a/Engine/Source/Runtime/Engine/Private/Animation/AnimationRuntime.cpp
+++ b/Engine/Source/Runtime/Engine/Private/Animation/AnimationRuntime.cpp
@@ -2015,7 +2015,7 @@ bool FAnimationRuntime::ContainsNaN(TArray<FBoneIndexType>& RequiredBoneIndices,
 }
 #endif
 
-FTransform FAnimationRuntime::GetComponentSpaceTransform(const FReferenceSkeleton& RefSkeleton, const TArray<FTransform> &BoneSpaceTransforms, int32 BoneIndex)
+FTransform FAnimationRuntime::GetComponentSpaceTransform(const FReferenceSkeleton& RefSkeleton, const TArrayView<const FTransform> &BoneSpaceTransforms, int32 BoneIndex)
 {
 	if (RefSkeleton.IsValidIndex(BoneIndex))
 	{
@@ -2069,7 +2069,7 @@ const FTransform& FAnimationRuntime::GetComponentSpaceTransformWithCache(const F
 	return CachedTransforms[BoneIndex];
 }
 
-void FAnimationRuntime::FillUpComponentSpaceTransforms(const FReferenceSkeleton& RefSkeleton, const TArray<FTransform> &BoneSpaceTransforms, TArray<FTransform> &ComponentSpaceTransforms)
+void FAnimationRuntime::FillUpComponentSpaceTransforms(const FReferenceSkeleton& RefSkeleton, const TArrayView<const FTransform> &BoneSpaceTransforms, TArray<FTransform> &ComponentSpaceTransforms)
 {
 	ComponentSpaceTransforms.Empty(BoneSpaceTransforms.Num());
 	ComponentSpaceTransforms.AddUninitialized(BoneSpaceTransforms.Num());
@@ -2287,26 +2287,32 @@ int32 FAnimationRuntime::GetStringDistance(const FString& First, const FString&
 	return NextRow[SecondLength];
 }
 
-void FAnimationRuntime::RetargetBoneTransform(const USkeleton* MySkeleton, const FName& RetargetSource, FTransform& BoneTransform, const int32 SkeletonBoneIndex, const FCompactPoseBoneIndex& BoneIndex, const FBoneContainer& RequiredBones, const bool bIsBakedAdditive)
+void FAnimationRuntime::RetargetBoneTransform(const USkeleton* SourceSkeleton, const FName& RetargetSource, FTransform& BoneTransform, const int32 SkeletonBoneIndex, const FCompactPoseBoneIndex& BoneIndex, const FBoneContainer& RequiredBones, const bool bIsBakedAdditive)
 {
-	if (MySkeleton)
+	if (SourceSkeleton)
 	{
-		const TArray<FTransform>& RetargetTransforms = MySkeleton->GetRefLocalPoses(RetargetSource);
-		RetargetBoneTransform(MySkeleton, RetargetSource, RetargetTransforms, BoneTransform, SkeletonBoneIndex, BoneIndex, RequiredBones, bIsBakedAdditive);
+		const TArray<FTransform>& RetargetTransforms = SourceSkeleton->GetRefLocalPoses(RetargetSource);
+		RetargetBoneTransform(SourceSkeleton, RetargetSource, RetargetTransforms, BoneTransform, SkeletonBoneIndex, BoneIndex, RequiredBones, bIsBakedAdditive);
 	}
 }
 
-void FAnimationRuntime::RetargetBoneTransform(const USkeleton* MySkeleton, const FName& SourceName, const TArray<FTransform>& RetargetTransforms, FTransform& BoneTransform, const int32 SkeletonBoneIndex, const FCompactPoseBoneIndex& BoneIndex, const FBoneContainer& RequiredBones, const bool bIsBakedAdditive)
+void FAnimationRuntime::RetargetBoneTransform(const USkeleton* SourceSkeleton, const FName& SourceName, const TArray<FTransform>& RetargetTransforms, FTransform& BoneTransform, const int32 SkeletonBoneIndex, const FCompactPoseBoneIndex& BoneIndex, const FBoneContainer& RequiredBones, const bool bIsBakedAdditive)
 {
-	if (MySkeleton)
+	if (SourceSkeleton)
 	{
-		switch (MySkeleton->GetBoneTranslationRetargetingMode(SkeletonBoneIndex))
+		const USkeleton* TargetSkeleton = RequiredBones.GetSkeletonAsset();
+		const FSkeletonRemapping* SkeletonRemapping = TargetSkeleton->GetSkeletonRemapping(SourceSkeleton);
+
+		const int32 TargetSkeletonBoneIndex = RequiredBones.GetSkeletonIndex(BoneIndex);
+		const int32 SourceSkeletonBoneIndex = (SkeletonRemapping) ? SkeletonRemapping->GetSourceSkeletonBoneIndex(TargetSkeletonBoneIndex) : SkeletonBoneIndex;
+
+		switch (TargetSkeleton->GetBoneTranslationRetargetingMode(TargetSkeletonBoneIndex))
 		{
 			case EBoneTranslationRetargetingMode::AnimationScaled:
 			{
 				// @todo - precache that in FBoneContainer when we have SkeletonIndex->TrackIndex mapping. So we can just apply scale right away.
 				const TArray<FTransform>& SkeletonRefPoseArray = RetargetTransforms;
-				const float SourceTranslationLength = SkeletonRefPoseArray[SkeletonBoneIndex].GetTranslation().Size();
+				const float SourceTranslationLength = SkeletonRefPoseArray[SourceSkeletonBoneIndex].GetTranslation().Size();
 				if (SourceTranslationLength > KINDA_SMALL_NUMBER)
 				{
 					const float TargetTranslationLength = RequiredBones.GetRefPoseTransform(BoneIndex).GetTranslation().Size();
@@ -2332,10 +2338,17 @@ void FAnimationRuntime::RetargetBoneTransform(const USkeleton* MySkeleton, const
 
 					const FTransform& RefPoseTransform = RequiredBones.GetRefPoseTransform(BoneIndex);
 
+					// Remap the base pose onto the target skeleton so that we are working entirely in target space
+					FTransform BaseTransform = AuthoredOnRefSkeleton[SourceSkeletonBoneIndex];
+					if (SkeletonRemapping)
+					{
+						BaseTransform = SkeletonRemapping->RetargetBoneTransformToTargetSkeleton(TargetSkeletonBoneIndex, BaseTransform);
+					}
+
 					// Apply the retargeting as if it were an additive difference between the current skeleton and the retarget skeleton. 
-					BoneTransform.SetRotation(BoneTransform.GetRotation() * AuthoredOnRefSkeleton[SkeletonBoneIndex].GetRotation().Inverse() * RefPoseTransform.GetRotation());
-					BoneTransform.SetTranslation(BoneTransform.GetTranslation() + (RefPoseTransform.GetTranslation() - AuthoredOnRefSkeleton[SkeletonBoneIndex].GetTranslation()));
-					BoneTransform.SetScale3D(BoneTransform.GetScale3D() * (RefPoseTransform.GetScale3D() * AuthoredOnRefSkeleton[SkeletonBoneIndex].GetSafeScaleReciprocal(AuthoredOnRefSkeleton[SkeletonBoneIndex].GetScale3D())));
+					BoneTransform.SetRotation(BoneTransform.GetRotation() * BaseTransform.GetRotation().Inverse() * RefPoseTransform.GetRotation());
+					BoneTransform.SetTranslation(BoneTransform.GetTranslation() + (RefPoseTransform.GetTranslation() - BaseTransform.GetTranslation()));
+					BoneTransform.SetScale3D(BoneTransform.GetScale3D() * (RefPoseTransform.GetScale3D() * BaseTransform.GetSafeScaleReciprocal(BaseTransform.GetScale3D())));
 					BoneTransform.NormalizeRotation();
 				}
 				break;
diff --git a/Engine/Source/Runtime/Engine/Private/Animation/BlendSpace.cpp b/Engine/Source/Runtime/Engine/Private/Animation/BlendSpace.cpp
--- a/Engine/Source/Runtime/Engine/Private/Animation/BlendSpace.cpp
+++ b/Engine/Source/Runtime/Engine/Private/Animation/BlendSpace.cpp
@@ -679,7 +679,7 @@ void UBlendSpace::GetAnimationPose_Internal(TArray<FBlendSampleData>& BlendSampl
 				if (Sample.Animation
 #if WITH_EDITOR
 					// verify if Sample.Animation->GetSkeleton matches
-					&& ensure(Sample.Animation->GetSkeleton() == GetSkeleton())
+					&& ensure(GetSkeleton()->IsCompatible(Sample.Animation->GetSkeleton()))
 #endif // WITH_EDITOR
 					)
 				{
@@ -1130,9 +1130,9 @@ bool UBlendSpace::ShouldAnimationBeAdditive() const
 
 bool UBlendSpace::IsAnimationCompatibleWithSkeleton(const UAnimSequence* AnimationSequence) const
 {
-	// Check if the animation sequences skeleton is compatible iwth the blendspace one
+	// Check if the animation sequences skeleton is compatible with the blendspace one
 	const USkeleton* MySkeleton = GetSkeleton();
-	const bool bIsAnimationCompatible = AnimationSequence && MySkeleton && AnimationSequence->GetSkeleton() && AnimationSequence->GetSkeleton()->IsCompatible(MySkeleton);
+	const bool bIsAnimationCompatible = AnimationSequence && MySkeleton && AnimationSequence->GetSkeleton() && MySkeleton->IsCompatible(AnimationSequence->GetSkeleton());
 	return bIsAnimationCompatible;
 }
 
diff --git a/Engine/Source/Runtime/Engine/Private/Animation/Skeleton.cpp b/Engine/Source/Runtime/Engine/Private/Animation/Skeleton.cpp
--- a/Engine/Source/Runtime/Engine/Private/Animation/Skeleton.cpp
+++ b/Engine/Source/Runtime/Engine/Private/Animation/Skeleton.cpp
@@ -31,7 +31,7 @@
 #include "Animation/AnimBlueprint.h"
 #include "UObject/AnimObjectVersion.h"
 #include "EngineUtils.h"
-#include "Animation/AnimData/AnimDataModel.h"
+#include "Misc/ScopeLock.h"
 
 #define LOCTEXT_NAMESPACE "Skeleton"
 #define ROOT_BONE_PARENT	INDEX_NONE
@@ -53,6 +53,11 @@ const FName USkeleton::AnimTrackCurveMappingName = FName(TEXT("AnimationTrackCur
 const FName FAnimSlotGroup::DefaultGroupName = FName(TEXT("DefaultGroup"));
 const FName FAnimSlotGroup::DefaultSlotName = FName(TEXT("DefaultSlot"));
 
+// Skeleton remapping related.
+TArray<USkeleton*> USkeleton::LoadedSkeletons;		// We keep track of a list of loaded skeletons, because we have no thread safe object iterator at the time this code is written.
+FCriticalSection USkeleton::LoadedSkeletonsMutex;	// The mutex we use to lock the LoadedSkeletons array.
+
+
 void SerializeReferencePose(FArchive& Ar, FReferencePose& P, UObject* Outer)
 {
 	Ar.UsingCustomVersion(FAnimPhysObjectVersion::GUID);
@@ -107,13 +112,336 @@ namespace VirtualBoneNameHelpers
 	}
 }
 
+bool USkeleton::IsCompatible(const USkeleton* InSkeleton) const
+{
+	if (InSkeleton == this)
+	{
+		return true;
+	}
+
+	for (const auto& CompatibleSkeleton : CompatibleSkeletons)
+	{
+		if (CompatibleSkeleton == InSkeleton->CachedSoftObjectPtr)
+		{
+			return true;
+		}
+	}
+
+	return false;
+}
+
+bool USkeleton::IsCompatibleSkeletonByAssetString(const FString& SkeletonAssetString) const
+{
+	// First check against itself.
+	const FString SkeletonString = FAssetData(this).GetExportTextName();
+	if (SkeletonString == SkeletonAssetString)
+	{
+		return true;
+	}
+
+	// Now check against the list of compatible skeletons.
+	FString CompatibleName;
+	for (auto& CompatibleSkeleton : CompatibleSkeletons)
+	{
+		CompatibleName = FString::Printf(TEXT("%s'%s'"), *GetClass()->GetName(), *CompatibleSkeleton.ToString());
+		if (CompatibleName == SkeletonAssetString)
+		{
+			return true;
+		}
+	}
+
+	return false;
+}
+
+bool USkeleton::IsCompatibleSkeletonByAssetData(const FAssetData& AssetData) const
+{
+	return IsCompatibleSkeletonByAssetString(AssetData.GetTagValueRef<FString>(TEXT("Skeleton")));
+}
+
+FSkeletonRemapping::FSkeletonRemapping(const USkeleton* InSourceSkeleton, const USkeleton* InTargetSkeleton)
+	: SourceSkeleton(InSourceSkeleton)
+	, TargetSkeleton(InTargetSkeleton)
+{
+	if (!InSourceSkeleton || !InTargetSkeleton)
+	{
+		return;
+	}
+
+	const FReferenceSkeleton& SourceReferenceSkeleton = InSourceSkeleton->GetReferenceSkeleton();
+	const FReferenceSkeleton& TargetReferenceSkeleton = InTargetSkeleton->GetReferenceSkeleton();
+
+	const int32 SourceNumBones = SourceReferenceSkeleton.GetNum();
+	const int32 TargetNumBones = TargetReferenceSkeleton.GetNum();
+
+	SourceToTargetBoneIndexes.SetNumUninitialized(SourceNumBones);
+	TargetToSourceBoneIndexes.SetNumUninitialized(TargetNumBones);
+	RetargetingTable.SetNumUninitialized(TargetNumBones);
+
+	TArrayView<const FTransform> SourceLocalTransforms;
+	TArrayView<const FTransform> TargetLocalTransforms;
+/*
+	// Build the mapping from source to target bones through the remapping rig if one exists
+	if (true)	// TODO_SKELETON_REMAPPING: Use a remapping rig if it exists
+	{
+		for (int32 SourceBoneIndex = 0; SourceBoneIndex < SourceNumBones; ++SourceBoneIndex)
+		{
+			const FName BoneName = SourceReferenceSkeleton.GetBoneName(SourceBoneIndex);
+			int32 TargetBoneIndex = TargetReferenceSkeleton.FindBoneIndex(BoneName);
+			SourceToTargetBoneIndexes[SourceBoneIndex] = TargetBoneIndex;
+		}
+
+		// Get the matched source and target rest poses from the remapping rig
+		//
+		// TODO_SKELETON_REMAPPING: For now we'll just use the skeleton's rest poses, but we should really be pulling these
+		// poses from a remapping rig so that the use can better align the poses.  Note that we'll want the remapping rig
+		// to be independent from the USkeleton itself because a skeleton may want to participate in multiple remappings with
+		// different other skeletons and may have to pose itself differently to align with the different other skeletons
+		//
+		SourceLocalTransforms = MakeArrayView(SourceReferenceSkeleton.GetRefBonePose());
+		TargetLocalTransforms = MakeArrayView(TargetReferenceSkeleton.GetRefBonePose());
+	}
+	else // Fall back to simple name matching if there is no rig
+*/
+	{
+		// Match source to target bones by name lookup between source and target skeletons
+		for (int32 SourceBoneIndex = 0; SourceBoneIndex < SourceNumBones; ++SourceBoneIndex)
+		{
+			const FName BoneName = SourceReferenceSkeleton.GetBoneName(SourceBoneIndex);
+			const int32 TargetBoneIndex = TargetReferenceSkeleton.FindBoneIndex(BoneName);
+			SourceToTargetBoneIndexes[SourceBoneIndex] = TargetBoneIndex;
+		}
+
+		// Get the matched (hopefully) source and target rest poses from the source and target skeletons
+		SourceLocalTransforms = MakeArrayView(SourceReferenceSkeleton.GetRefBonePose());
+		TargetLocalTransforms = MakeArrayView(TargetReferenceSkeleton.GetRefBonePose());
+	}
+
+	// Force the roots to map onto each other regardless of their names
+	SourceToTargetBoneIndexes[0] = 0;
+
+	// Build the reverse mapping from target back to source bones
+	FMemory::Memset(TargetToSourceBoneIndexes.GetData(), static_cast<uint8>(INDEX_NONE), TargetNumBones * TargetToSourceBoneIndexes.GetTypeSize());
+	for (int32 SourceBoneIndex = 0; SourceBoneIndex < SourceNumBones; ++SourceBoneIndex)
+	{
+		const int32 TargetBoneIndex = SourceToTargetBoneIndexes[SourceBoneIndex];
+		if (TargetBoneIndex != INDEX_NONE)
+		{
+			TargetToSourceBoneIndexes[TargetBoneIndex] = SourceBoneIndex;
+		}
+	}
+
+	TArray<FTransform> SourceComponentTransforms;
+	TArray<FTransform> TargetComponentTransforms;
+	FAnimationRuntime::FillUpComponentSpaceTransforms(SourceReferenceSkeleton, SourceLocalTransforms, SourceComponentTransforms);
+	FAnimationRuntime::FillUpComponentSpaceTransforms(TargetReferenceSkeleton, TargetLocalTransforms, TargetComponentTransforms);
+
+	// Calculate the retargeting constants to map from source skeleton space to target skeleton space
+	//
+	// Simply remapping joint indices is usually not sufficient to give us the desired result pose if the source and target
+	// skeletons are built with different conventions for joint orientations.  We therefore need to compute a remapping
+	// between the source and target joint orientations, which we do in terms of delta rotations from the rest pose:
+	//
+	//		Q = P * D * R
+	//
+	// where:
+	//
+	//		Q is the final joint orientation in component space
+	//		P is the parent joint orientation in component space
+	//		D is the delta rotation that we want to remap
+	//		R is the local space rest pose orientation for the joint
+	//
+	// We want to find a mapping from the source to the target:
+	//
+	//		Ps * Ds * Rs  -->  Pt * Dt * Rt
+	//
+	// such that the deltas produce equivalent rotations on the mesh even if their parent rotation frames are different.
+	// In other words, we need to find Dt such that its component space rotation is equivalent to Ds.  To convert a rotation
+	// from local space (D) to component space (C), given its parent (P), we have:
+	//
+	//		C * P = P * D
+	//		C = P * D * P⁻¹
+	//
+	// Setting the source and target component space rotations to be equal (Cs = Ct) then gives us:
+	//
+	//		Ps * Ds * Ps⁻¹ = Pt * Dt * Pt⁻¹
+	//
+	// which we then solve for Dt to get:
+	//
+	//		Dt = Pt⁻¹ * Ps * Ds * Ps⁻¹ * Pt
+	//
+	// However, when we're remapping an animation pose, we will have the local transforms rather than the deltas from the
+	// rest pose, so we also need to convert between these local transforms and the equivalent deltas:
+	//
+	//		L = D * R
+	//		D = L * R⁻¹
+	//
+	// Combining that with our equation for Dt, we get:
+	//
+	//		Lt * Rt⁻¹ = Pt⁻¹ * Ps * Ls * Rs⁻¹ * Ps⁻¹ * Pt
+	//
+	// Solving for Lt then gives us:
+	//
+	//		Lt = Pt⁻¹ * Ps * Ls * Rs⁻¹ * Ps⁻¹ * Pt * Rt
+	//
+	// Finally, factoring out the constant terms (which we precompute here) gives us:
+	//
+	//		Lt = Q0 * Ls * Q1
+	//		Q0 = Pt⁻¹ * Ps
+	//		Q1 = Rs⁻¹ * Ps⁻¹ * Pt * Rt
+	//
+	// Note that when remapping additive animations, we drop the rest pose terms, but we still need to convert between the
+	// source and target rotation frames.  Dropping Rs and Rt from the equations above gives us:
+	//
+	//		Lt = Pt⁻¹ * Ps * Ls * Ps⁻¹ * Pt			(for additive animations)
+	//
+	// which is equivalent to the following in terms of our precomputed constants:
+	//
+	//		Lt = Q0 * Ls * Q0⁻¹						(for additive animations)
+	//
+	RetargetingTable[0] = MakeTuple(FQuat::Identity, SourceLocalTransforms[0].GetRotation().Inverse() * TargetLocalTransforms[0].GetRotation());
+	for (int32 TargetBoneIndex = 1; TargetBoneIndex < TargetNumBones; ++TargetBoneIndex)
+	{
+		const int32 SourceBoneIndex = TargetToSourceBoneIndexes[TargetBoneIndex];
+		if (SourceBoneIndex != INDEX_NONE)
+		{
+			const int32 SourceParentIndex = SourceReferenceSkeleton.GetParentIndex(SourceBoneIndex);
+			const int32 TargetParentIndex = TargetReferenceSkeleton.GetParentIndex(TargetBoneIndex);
+			check(SourceParentIndex != INDEX_NONE);
+			check(TargetParentIndex != INDEX_NONE);
+
+			const FQuat PS = SourceComponentTransforms[SourceParentIndex].GetRotation();
+			const FQuat PT = TargetComponentTransforms[TargetParentIndex].GetRotation();
+
+			const FQuat RS = SourceLocalTransforms[SourceBoneIndex].GetRotation();
+			const FQuat RT = TargetLocalTransforms[TargetBoneIndex].GetRotation();
+
+			const FQuat Q0 = PT.Inverse() * PS;
+			const FQuat Q1 = RS.Inverse() * PS.Inverse() * PT * RT;
+
+			RetargetingTable[TargetBoneIndex] = MakeTuple(Q0, Q1);
+		}
+		else
+		{
+			RetargetingTable[TargetBoneIndex] = MakeTuple(FQuat::Identity, FQuat::Identity);
+		}
+	}
+
+	GenerateCurveMapping();
+}
+
+void FSkeletonRemapping::ComposeWith(const FSkeletonRemapping& OtherSkeletonRemapping)
+{
+	check(OtherSkeletonRemapping.SourceSkeleton == TargetSkeleton);
+
+	TargetSkeleton = OtherSkeletonRemapping.TargetSkeleton;
+
+	const int32 SourceNumBones = SourceToTargetBoneIndexes.Num();
+	const int32 TargetNumBones = OtherSkeletonRemapping.TargetToSourceBoneIndexes.Num();
+
+	TArray<int32> NewTargetToSourceBoneIndexes;
+	TArray<TTuple<FQuat, FQuat>> NewRetargetingTable;
+
+	NewTargetToSourceBoneIndexes.SetNumUninitialized(TargetNumBones);
+	NewRetargetingTable.SetNumUninitialized(TargetNumBones);
+
+	// Compose the retargeting constants
+	for (int32 NewTargetBoneIndex = 0; NewTargetBoneIndex < TargetNumBones; ++NewTargetBoneIndex)
+	{
+		const int32 OldTargetBoneIndex = OtherSkeletonRemapping.TargetToSourceBoneIndexes[NewTargetBoneIndex];
+		const int32 OldSourceBoneIndex = (OldTargetBoneIndex != INDEX_NONE) ? TargetToSourceBoneIndexes[OldTargetBoneIndex] : INDEX_NONE;
+
+		if (OldSourceBoneIndex != INDEX_NONE)
+		{
+			const TTuple<FQuat, FQuat>& OldQQ = RetargetingTable[OldTargetBoneIndex];
+			const TTuple<FQuat, FQuat>& NewQQ = OtherSkeletonRemapping.RetargetingTable[NewTargetBoneIndex];
+			const FQuat Q0 = NewQQ.Get<0>() * OldQQ.Get<0>();
+			const FQuat Q1 = OldQQ.Get<1>() * NewQQ.Get<1>();
+
+			NewTargetToSourceBoneIndexes[NewTargetBoneIndex] = OldSourceBoneIndex;
+			NewRetargetingTable[NewTargetBoneIndex] = MakeTuple(Q0, Q1);
+		}
+		else
+		{
+			NewTargetToSourceBoneIndexes[NewTargetBoneIndex] = INDEX_NONE;
+			NewRetargetingTable[NewTargetBoneIndex] = MakeTuple(FQuat::Identity, FQuat::Identity);
+		}
+	}
+
+	TargetToSourceBoneIndexes = MoveTemp(NewTargetToSourceBoneIndexes);
+	RetargetingTable = MoveTemp(NewRetargetingTable);
+
+	// Rebuild the mapping from source bones to the target bones
+	FMemory::Memset(SourceToTargetBoneIndexes.GetData(), static_cast<uint8>(INDEX_NONE), SourceNumBones * SourceToTargetBoneIndexes.GetTypeSize());
+	for (int32 TargetBoneIndex = 0; TargetBoneIndex < TargetNumBones; ++TargetBoneIndex)
+	{
+		const int32 SourceBoneIndex = TargetToSourceBoneIndexes[TargetBoneIndex];
+		if (SourceBoneIndex != INDEX_NONE)
+		{
+			SourceToTargetBoneIndexes[SourceBoneIndex] = TargetBoneIndex;
+		}
+	}
+}
+
+void FSkeletonRemapping::GenerateCurveMapping()
+{
+	const USkeleton* Source = GetSourceSkeleton().Get();
+	const USkeleton* Target = GetTargetSkeleton().Get();
+	check(Source && Target);
+
+	// Get the curve mappings.
+	const FSmartNameMapping* SourceMapping = Source->GetSmartNameContainer(USkeleton::AnimCurveMappingName);
+	const FSmartNameMapping* TargetMapping = Target->GetSmartNameContainer(USkeleton::AnimCurveMappingName);
+	if (!SourceMapping) // Source has no curves, so there is nothing to add.
+	{
+		return;
+	}
+
+	// Check whether we have curves or not in both source and target.
+	bool bHasSourceCurves = false;
+	bool bHasTargetCurves = false;
+	if (TargetMapping)
+	{
+		bHasTargetCurves = (TargetMapping->GetMaxUID() != SmartName::MaxUID);
+		bHasSourceCurves = (SourceMapping->GetMaxUID() != SmartName::MaxUID);
+	}
+
+	// Init the mapping table.
+	const SmartName::UID_Type NumItems = SourceMapping->GetMaxUID() + 1;
+	if (bHasSourceCurves)
+	{
+		SourceToTargetCurveMapping.Init(MAX_uint16, NumItems);
+	}
+
+	// If we have no source or target curves (while having valid smart name mappings), there is nothing to do.
+	if (!bHasTargetCurves && !bHasSourceCurves)
+	{
+		return;
+	}
+
+	// For every source curve, try to find the curve with the same name in the target.
+	for (SmartName::UID_Type Index = 0; Index < NumItems; ++Index)
+	{
+		FName CurveName;
+		SourceMapping->GetName(Index, CurveName);
+
+		// Make sure the curve name is valid and we can actually find a curve with the same name in the target.
+		FSmartName TargetSmartName;
+		if (!CurveName.IsValid() || !TargetMapping->FindSmartName(CurveName, TargetSmartName))
+		{
+			continue;
+		}
+
+		SourceToTargetCurveMapping[Index] = TargetSmartName.UID;
+	}
+}
+
 USkeleton::USkeleton(const FObjectInitializer& ObjectInitializer)
 	: Super(ObjectInitializer)
 {
 	// Make sure we have somewhere for curve names.
 	AnimCurveMapping = SmartNames.AddContainer(AnimCurveMappingName);
-
 	AnimCurveUidVersion = 0;
+	CachedSoftObjectPtr = TSoftObjectPtr<USkeleton>(this);
 
 	if (HasAnyFlags(RF_ClassDefaultObject))
 	{
@@ -128,6 +456,27 @@ void USkeleton::BeginDestroy()
 	{
 		FCoreUObjectDelegates::OnPackageReloaded.RemoveAll(this);
 	}
+
+	// Clear our own skeleton mappings.
+	ClearSkeletonRemappings();
+	OnSkeletonDestructEvent.Broadcast(this); // Inform others about this skeleton being destructed, so we can remove mappings inside those skeletons.
+	OnSkeletonDestructEvent.Clear();
+	OnSmartNamesChangedEvent.Clear();
+
+	//  Remove the loaded skeleton from the list.
+	{
+		FScopeLock Lock(&LoadedSkeletonsMutex);
+		LoadedSkeletons.Remove(this);
+
+		// Also unregister this skeleton from other skeletons this one is listening to (events it is listening to).
+		for (USkeleton* Skeleton : LoadedSkeletons)
+		{
+			if (Skeleton->IsCompatible(this))
+			{
+				Skeleton->OnSkeletonDestructEvent.RemoveAll(this);
+			}
+		}
+	}
 }
 
 void USkeleton::PostInitProperties()
@@ -159,7 +508,7 @@ void USkeleton::PostLoad()
 
 	// catch any case if guid isn't valid
 	check(Guid.IsValid());
-
+	
 	// Cache smart name uids for animation curve names
 	IncreaseAnimCurveUidVersion();
 
@@ -171,6 +520,120 @@ void USkeleton::PostLoad()
 	{
 		CurveMappingTable->InitializeCurveMetaData(this);
 	}
+
+	// Add this skeleton to the loaded skeletons list and register some event listeners.
+	{
+		FScopeLock Lock(&LoadedSkeletonsMutex);
+
+		// Make any loaded skeleton that is compatible with this skeleton listen to destruct events.
+		// This is done so when this skeleton gets destructed, other skeletons can remove mappings to this skeleton.
+		for (USkeleton* Skeleton : LoadedSkeletons)
+		{
+			if (!Skeleton || Skeleton == this)
+			{
+				continue;
+			}
+
+			// If the already loaded other skeleton is compatible with this newly loaded skeleton, then let the
+			// already loaded other skeleton listen to the destruction event of this newly loaded skeleton.
+			if (Skeleton->IsCompatible(this))
+			{
+				Skeleton->CreateSkeletonRemappingIfNeeded(this);
+				if (!OnSkeletonDestructEvent.IsBoundToObject(Skeleton))
+				{
+					OnSkeletonDestructEvent.AddUObject(Skeleton, &USkeleton::HandleSkeletonDestruct);
+				}
+			}
+
+			// If we are compatible with the already loaded skeleton, we want to start listening to destruct events of that one.
+			if (IsCompatible(Skeleton))
+			{
+				CreateSkeletonRemappingIfNeeded(Skeleton);
+				if (!Skeleton->OnSkeletonDestructEvent.IsBoundToObject(this))
+				{
+					Skeleton->OnSkeletonDestructEvent.AddUObject(this, &USkeleton::HandleSkeletonDestruct);
+				}
+			}
+		}
+
+		LoadedSkeletons.Add(this);
+	}
+
+	// Listen to smart name changes so we can regenerate our skeleton mappings.
+	if (!OnSmartNamesChangedEvent.IsBoundToObject(this))
+	{
+		OnSmartNamesChangedEvent.AddUObject(this, &USkeleton::HandleSmartNamesChangedEvent);
+	}
+}
+
+// Regenerate all required skeleton remappings whenever curves change.
+void USkeleton::HandleSmartNamesChangedEvent()
+{
+	FScopeLock Lock(&LoadedSkeletonsMutex);
+	for (USkeleton* Skeleton : LoadedSkeletons)
+	{
+		if (Skeleton == this)
+		{
+			continue;
+		}
+
+		if (!Skeleton)
+		{
+			continue;
+		}
+
+		if (Skeleton->IsCompatible(this))
+		{
+			Skeleton->CreateSkeletonRemappingIfNeeded(this, true);
+		}
+
+		if (IsCompatible(Skeleton))
+		{
+			CreateSkeletonRemappingIfNeeded(Skeleton, true);
+		}
+	}
+}
+
+void USkeleton::HandleSkeletonDestruct(const USkeleton* Skeleton)
+{
+	OnSkeletonDestructEvent.RemoveAll(Skeleton);
+
+	// We aren't interested to hear that we loaded ourselves as we need no mappings to ourselves.
+	// Also if this skeleton isn't in the compatible list, we can ignore it.
+	if (Skeleton == this || !Skeleton || !IsCompatible(Skeleton))
+	{
+		return;
+	}
+
+	// Check if we already have a mapping. If so, remove the existing one and create a new one.
+	// This might only happen when a skeleton asset gets reimported.
+	RemoveSkeletonRemapping(Skeleton);
+}
+
+void USkeleton::CreateSkeletonRemappingIfNeeded(const USkeleton* SourceSkeleton, bool bRebuildIfExists)
+{
+	FScopeLock Lock(&SkeletonRemappingMutex);
+	check(SourceSkeleton != nullptr && SourceSkeleton != this && IsCompatible(SourceSkeleton));
+
+	// Check if we already have a mapping. If so, remove the existing one and create a new one.
+	// This might only happen when a skeleton asset gets reimported.
+	FSkeletonRemapping** Mapping = SkeletonRemappings.Find(SourceSkeleton);
+	if (Mapping)	
+	{
+		if (!bRebuildIfExists)
+		{
+			return;
+		}
+
+		delete *Mapping;
+		SkeletonRemappings.Remove(SourceSkeleton);
+	}
+
+	// Create, add and return the new entry.
+	FSkeletonRemapping* NewMapping = new FSkeletonRemapping(SourceSkeleton, this);
+	SkeletonRemappings.Add(SourceSkeleton, NewMapping);
+
+	UE_LOG(LogAnimation, Verbose, TEXT("Recreating remapping between %s and %s"), *FAssetData(this).AssetName.ToString(), *FAssetData(SourceSkeleton).AssetName.ToString());
 }
 
 void USkeleton::PostDuplicate(bool bDuplicateForPIE)
@@ -182,6 +645,10 @@ void USkeleton::PostDuplicate(bool bDuplicateForPIE)
 		// regenerate Guid
 		RegenerateGuid();
 	}
+
+	// Initialize skeleton remappings.
+	ClearSkeletonRemappings();
+	BuildSkeletonRemappings(true);
 }
 
 void USkeleton::Serialize( FArchive& Ar )
@@ -288,6 +755,18 @@ void USkeleton::Serialize( FArchive& Ar )
 }
 
 #if WITH_EDITOR
+void USkeleton::PostEditChangeProperty(FPropertyChangedEvent& PropertyChangedEvent)
+{
+	Super::PostEditChangeProperty(PropertyChangedEvent);
+
+	// Rebuild skeleton mappings involving this skeleton when we modified the compatible skeleton list.
+	if (PropertyChangedEvent.Property && PropertyChangedEvent.Property->GetFName() == GET_MEMBER_NAME_CHECKED(USkeleton, CompatibleSkeletons))
+	{
+		ClearSkeletonRemappings();
+		BuildSkeletonRemappings(false); // Build it, but not bidirectional, meaning only rebuild mappings to skeletons we are compatible with.
+	}
+}
+
 void USkeleton::PreEditUndo()
 {
 	// Undoing so clear cached data as it will now be stale
@@ -330,6 +809,73 @@ void USkeleton::ConvertToFReferenceSkeleton()
 	ensure (AnimRetargetSources.Num() == 0);
 }
 
+void USkeleton::ClearSkeletonRemappings()
+{
+	FScopeLock Lock(&SkeletonRemappingMutex);
+	for (auto& Mapping : SkeletonRemappings)
+	{
+		delete Mapping.Value;
+	}
+
+	SkeletonRemappings.Reset();
+}
+
+const FSkeletonRemapping* USkeleton::GetSkeletonRemapping(const USkeleton* SourceSkeleton) const
+{
+	check(SourceSkeleton);
+	if (SourceSkeleton == this)
+	{
+		return nullptr;
+	}
+
+	// Find the existing one.
+	FScopeLock Lock(&SkeletonRemappingMutex);
+	const FSkeletonRemapping* const* Mapping = SkeletonRemappings.Find(SourceSkeleton);
+	if (Mapping)
+	{
+		return *Mapping;
+	}
+
+	return nullptr;
+}
+
+void USkeleton::BuildSkeletonRemappings(bool bBidirectional)
+{
+	for (USkeleton* Skeleton : LoadedSkeletons)
+	{
+		if (Skeleton == this)
+		{
+			continue;
+		}
+
+		// Update the remapping in the other skeleton.
+		if (bBidirectional && Skeleton->IsCompatible(this))
+		{
+			Skeleton->CreateSkeletonRemappingIfNeeded(this, true);
+		}
+
+		// Update the remapping to this skeleton inside our own remapping table.
+		if (IsCompatible(Skeleton))
+		{
+			CreateSkeletonRemappingIfNeeded(Skeleton, true);
+		}
+	}
+}
+
+void USkeleton::RemoveSkeletonRemapping(const USkeleton* SourceSkeleton)
+{
+	check(SourceSkeleton && SourceSkeleton != this);
+
+	// Locate the mapping and remove if found.
+	FScopeLock Lock(&SkeletonRemappingMutex);
+	FSkeletonRemapping** Mapping = SkeletonRemappings.Find(SourceSkeleton);
+	if (Mapping)
+	{
+		delete *Mapping;
+		SkeletonRemappings.Remove(SourceSkeleton);
+	}
+}
+
 bool USkeleton::DoesParentChainMatch(int32 StartBoneIndex, const USkeletalMesh* InSkelMesh) const
 {
 	const FReferenceSkeleton& SkeletonRefSkel = ReferenceSkeleton;
@@ -775,25 +1321,24 @@ FName USkeleton::GetRetargetSourceForMesh(USkeletalMesh* InMesh) const
 
 int32 USkeleton::GetRawAnimationTrackIndex(const int32 InSkeletonBoneIndex, const UAnimSequence* InAnimSeq)
 {
-	if( InSkeletonBoneIndex != INDEX_NONE )
+	if (InSkeletonBoneIndex != INDEX_NONE)
 	{
 #if WITH_EDITOR
 		return InAnimSeq->GetDataModel()->GetBoneAnimationTracks().IndexOfByPredicate([InSkeletonBoneIndex](const FBoneAnimationTrack& AnimationTrack)
-		{
-			return AnimationTrack.BoneTreeIndex == InSkeletonBoneIndex;
-		});
+			{
+				return AnimationTrack.BoneTreeIndex == InSkeletonBoneIndex;
+			});
 #else
 		return InAnimSeq->GetCompressedTrackToSkeletonMapTable().IndexOfByPredicate([InSkeletonBoneIndex](const FTrackToSkeletonMap& Mapping)
-		{
-			return Mapping.BoneTreeIndex == InSkeletonBoneIndex;
-		});
+			{
+				return Mapping.BoneTreeIndex == InSkeletonBoneIndex;
+			});
 #endif
 	}
 
 	return INDEX_NONE;
 }
 
-
 int32 USkeleton::GetSkeletonBoneIndexFromMeshBoneIndex(const USkeletalMesh* InSkelMesh, const int32 MeshBoneIndex)
 {
 	check(MeshBoneIndex != INDEX_NONE);
@@ -1076,6 +1621,9 @@ void USkeleton::HandleSkeletonHierarchyChange()
 		}
 	}
 
+	// Full rebuild of all compatible with this and with ones we are compatible with.
+	BuildSkeletonRemappings(true);
+
 	// Fix up loaded animations (any animations that aren't loaded will be fixed on load)
 	int32 NumLoadedAssets = 0;
 	for (TObjectIterator<UAnimationAsset> It; It; ++It)
diff --git a/Engine/Source/Runtime/Engine/Private/BoneContainer.cpp b/Engine/Source/Runtime/Engine/Private/BoneContainer.cpp
--- a/Engine/Source/Runtime/Engine/Private/BoneContainer.cpp
+++ b/Engine/Source/Runtime/Engine/Private/BoneContainer.cpp
@@ -8,6 +8,44 @@
 
 DEFINE_LOG_CATEGORY(LogSkeletalControl);
 
+//////////////////////////////////////////////////////////////////////////
+// FSkeletonRemappingCurve
+
+FSkeletonRemappingCurve::FSkeletonRemappingCurve(FBlendedCurve& InCurve, const FBoneContainer& InBoneContainer, const FSkeletonRemapping* InSkeletonRemapping)
+	: Curve(InCurve)
+	, BoneContainer(InBoneContainer)
+	, bIsRemapping(true)
+{
+	const FCachedSkeletonCurveMapping& CurveMapping = InBoneContainer.GetOrCreateCachedCurveMapping(InSkeletonRemapping);
+	Curve.UIDToArrayIndexLUT = &CurveMapping.UIDToArrayIndices;
+}
+
+FSkeletonRemappingCurve::FSkeletonRemappingCurve(FBlendedCurve& InCurve, const FBoneContainer& InBoneContainer, const USkeleton* SourceSkeleton)
+	: Curve(InCurve)
+	, BoneContainer(InBoneContainer)
+{
+	const FSkeletonRemapping* SkeletonRemapping = BoneContainer.GetSkeletonAsset()->GetSkeletonRemapping(SourceSkeleton);
+	if (SkeletonRemapping) // No remapping is required, just continue as we normally would.
+	{
+		const FCachedSkeletonCurveMapping& CurveMapping = BoneContainer.GetOrCreateCachedCurveMapping(SkeletonRemapping);
+		Curve.UIDToArrayIndexLUT = &CurveMapping.UIDToArrayIndices;
+		bIsRemapping = true;
+	}
+	else
+	{
+		bIsRemapping = false;
+	}
+}
+
+FSkeletonRemappingCurve::~FSkeletonRemappingCurve()
+{
+	if (bIsRemapping)
+	{
+		Curve.UIDToArrayIndexLUT = &BoneContainer.GetUIDToArrayLookupTableBackup();
+	}
+}
+
+
 //////////////////////////////////////////////////////////////////////////
 // FBoneContainer
 
@@ -290,6 +328,64 @@ void FBoneContainer::CacheRequiredAnimCurveUids(const FCurveEvaluationOption& Cu
 	{
 		UIDToArrayIndexLUT.Reset();
 	}
+
+	// Make a backup, used for skeleton remapping of curves.
+	UIDToArrayIndexLUTBackup = UIDToArrayIndexLUT;
+
+	// Make sure we regenerate our cached curve mappings next time they are requested.
+	MarkAllCachedCurveMappingsDirty();
+}
+
+void FBoneContainer::MarkAllCachedCurveMappingsDirty()
+{
+	for (auto& CurveMapping : CachedCurveMappingTable)
+	{
+		CurveMapping.Value.bIsDirty = true;
+	}
+}
+
+const FCachedSkeletonCurveMapping& FBoneContainer::GetOrCreateCachedCurveMapping(const FSkeletonRemapping* SkeletonRemapping) const
+{
+	check(SkeletonRemapping);
+
+	const USkeleton* SourceSkeleton = SkeletonRemapping->GetSourceSkeleton().Get();
+	check(SourceSkeleton);
+
+	// Check if we already have some cached data for this skeleton.
+	FCachedSkeletonCurveMapping* CachedData = CachedCurveMappingTable.Find(SourceSkeleton);
+	if (CachedData)
+	{
+		// Only return the object when we don't need to update its mapping.
+		if (!CachedData->bIsDirty)
+		{
+			return *CachedData;
+		}
+	}
+	else
+	{
+		CachedData = &CachedCurveMappingTable.Add(const_cast<USkeleton*>(SourceSkeleton));
+	}
+
+	UE_LOG(LogAnimation, Verbose, TEXT("Generating mapping for %s to %s"), 
+		*FAssetData(SkeletonRemapping->GetSourceSkeleton().Get()).AssetName.ToString(),
+		*FAssetData(SkeletonRemapping->GetTargetSkeleton().Get()).AssetName.ToString());
+
+	CachedData->UIDToArrayIndices = SkeletonRemapping->GetSourceToTargetCurveMapping();
+
+	// Now remap the UIDs to our local curve indexes.
+	const int32 NumMappings = CachedData->UIDToArrayIndices.Num();
+	for (int32 MappingIndex = 0; MappingIndex < NumMappings; ++MappingIndex)
+	{
+		const SmartName::UID_Type CurrentUID = CachedData->UIDToArrayIndices[MappingIndex];
+		if (CurrentUID != MAX_uint16)
+		{
+			const int32 CurveIndex = UIDToArrayIndexLUT.Find(CurrentUID);
+			CachedData->UIDToArrayIndices[MappingIndex] = (CurveIndex != INDEX_NONE) ? static_cast<SmartName::UID_Type>(CurveIndex) : MAX_uint16;
+		}
+	}
+
+	CachedData->bIsDirty = false;
+	return *CachedData;
 }
 
 const FRetargetSourceCachedData& FBoneContainer::GetRetargetSourceCachedData(const FName& InRetargetSourceName) const
diff --git a/Engine/Source/Runtime/Engine/Private/Components/SkeletalMeshComponent.cpp b/Engine/Source/Runtime/Engine/Private/Components/SkeletalMeshComponent.cpp
--- a/Engine/Source/Runtime/Engine/Private/Components/SkeletalMeshComponent.cpp
+++ b/Engine/Source/Runtime/Engine/Private/Components/SkeletalMeshComponent.cpp
@@ -729,7 +729,7 @@ void USkeletalMeshComponent::InitAnim(bool bForceReinit)
 
 		const bool bClearAnimInstance = AnimScriptInstance && !AnimSkeleton;
 		const bool bSkeletonMismatch = AnimSkeleton && (AnimScriptInstance->CurrentSkeleton!=SkeletalMesh->GetSkeleton());
-		const bool bSkeletonNotCompatible = AnimSkeleton && !bSkeletonMismatch && (AnimSkeleton->IsCompatibleMesh(SkeletalMesh) == false);
+		const bool bSkeletonNotCompatible = AnimSkeleton && !bSkeletonMismatch && !SkeletalMesh->GetSkeleton()->IsCompatible(AnimSkeleton);
 
 		LastPoseTickFrame = 0;
 
@@ -1031,7 +1031,7 @@ void USkeletalMeshComponent::PostEditChangeProperty(FPropertyChangedEvent& Prope
 		if ( PropertyThatChanged->GetFName() == GET_MEMBER_NAME_CHECKED( FSingleAnimationPlayData, AnimToPlay ))
 		{
 			// make sure the animation skeleton matches the current skeletalmesh
-			if (AnimationData.AnimToPlay != nullptr && SkeletalMesh && AnimationData.AnimToPlay->GetSkeleton() != SkeletalMesh->GetSkeleton())
+			if (AnimationData.AnimToPlay != nullptr && SkeletalMesh && !SkeletalMesh->GetSkeleton()->IsCompatible(AnimationData.AnimToPlay->GetSkeleton()))
 			{
 				UE_LOG(LogAnimation, Warning, TEXT("Invalid animation"));
 				AnimationData.AnimToPlay = nullptr;
@@ -3638,7 +3638,7 @@ void USkeletalMeshComponent::ValidateAnimation()
 
 	if(AnimationMode == EAnimationMode::AnimationSingleNode)
 	{
-		if(AnimationData.AnimToPlay && SkeletalMesh && AnimationData.AnimToPlay->GetSkeleton() != SkeletalMesh->GetSkeleton())
+		if(AnimationData.AnimToPlay && SkeletalMesh && !AnimationData.AnimToPlay->GetSkeleton()->IsCompatible(SkeletalMesh->GetSkeleton()))
 		{
 			if (SkeletalMesh->GetSkeleton())
 			{
@@ -3655,7 +3655,7 @@ void USkeletalMeshComponent::ValidateAnimation()
 	else if (AnimationMode == EAnimationMode::AnimationBlueprint)
 	{
 		IAnimClassInterface* AnimClassInterface = IAnimClassInterface::GetFromClass(AnimClass);
-		if (AnimClassInterface && SkeletalMesh && AnimClassInterface->GetTargetSkeleton() != SkeletalMesh->GetSkeleton())
+		if (AnimClassInterface && SkeletalMesh && !AnimClassInterface->GetTargetSkeleton()->IsCompatible(SkeletalMesh->GetSkeleton()))
 		{
 			if(SkeletalMesh->GetSkeleton())
 			{
diff --git a/Engine/Source/Runtime/Engine/Public/Animation/AnimCompressionTypes.h b/Engine/Source/Runtime/Engine/Public/Animation/AnimCompressionTypes.h
--- a/Engine/Source/Runtime/Engine/Public/Animation/AnimCompressionTypes.h
+++ b/Engine/Source/Runtime/Engine/Public/Animation/AnimCompressionTypes.h
@@ -801,7 +801,7 @@ extern void DecompressPose(FCompactPose& OutPose,
 extern void DecompressPose(	FCompactPose& OutPose,
 							const FCompressedAnimSequence& CompressedData,
 							const FAnimExtractContext& ExtractionContext,
-							USkeleton* Skeleton,
+							USkeleton* SourceSkeleton,
 							float SequenceLength,
 							EAnimInterpolationType Interpolation,
 							bool bIsBakedAdditive,
diff --git a/Engine/Source/Runtime/Engine/Public/Animation/AnimInstanceProxy.h b/Engine/Source/Runtime/Engine/Public/Animation/AnimInstanceProxy.h
--- a/Engine/Source/Runtime/Engine/Public/Animation/AnimInstanceProxy.h
+++ b/Engine/Source/Runtime/Engine/Public/Animation/AnimInstanceProxy.h
@@ -271,10 +271,10 @@ public:
 	/** Trigger any anim notifies */
 	void TriggerAnimNotifies(USkeletalMeshComponent* SkelMeshComp, float DeltaSeconds);
 
-	/** Check whether the supplied skeleton is compatible with this instance's skeleton */
-	bool IsSkeletonCompatible(USkeleton const* InSkeleton) const 
-	{ 
-		return InSkeleton->IsCompatible(Skeleton); 
+	/** Check whether animation content authored on the supplied skeleton may be played on this instance's skeleton */
+	bool IsSkeletonCompatible(USkeleton const* InSkeleton) const
+	{
+		return Skeleton->IsCompatible(InSkeleton);
 	}
 
 	/** Check whether we should extract root motion */
diff --git a/Engine/Source/Runtime/Engine/Public/AnimationRuntime.h b/Engine/Source/Runtime/Engine/Public/AnimationRuntime.h
--- a/Engine/Source/Runtime/Engine/Public/AnimationRuntime.h
+++ b/Engine/Source/Runtime/Engine/Public/AnimationRuntime.h
@@ -506,8 +506,8 @@ public:
 	static void SetSpaceTransform(FA2CSPose& Pose, int32 Index, FTransform& NewTransform);
 	// space bases
 	static FTransform GetComponentSpaceTransformRefPose(const FReferenceSkeleton& RefSkeleton, int32 BoneIndex);
-	static FTransform GetComponentSpaceTransform(const FReferenceSkeleton& RefSkeleton, const TArray<FTransform> &BoneSpaceTransforms, int32 BoneIndex);
-	static void FillUpComponentSpaceTransforms(const FReferenceSkeleton& RefSkeleton, const TArray<FTransform> &BoneSpaceTransforms, TArray<FTransform> &ComponentSpaceTransforms);
+	static FTransform GetComponentSpaceTransform(const FReferenceSkeleton& RefSkeleton, const TArrayView<const FTransform> &BoneSpaceTransforms, int32 BoneIndex);
+	static void FillUpComponentSpaceTransforms(const FReferenceSkeleton& RefSkeleton, const TArrayView<const FTransform> &BoneSpaceTransforms, TArray<FTransform> &ComponentSpaceTransforms);
 	static void MakeSkeletonRefPoseFromMesh(const USkeletalMesh* InMesh, const USkeleton* InSkeleton, TArray<FTransform>& OutBoneBuffer);
 
 	/**
@@ -540,16 +540,15 @@ public:
 	/**
 	* Retarget a single bone transform, to apply right after extraction.
 	*
-	* @param	MySkeleton			Skeleton this is retargeting
+	* @param	SourceSkeleton		Skeleton from which this is retargeting
 	* @param	RetargetSource		Retarget Source for the retargeting
-	* @param	BoneTransform		BoneTransform to read/write from.
-	* @param	SkeletonBoneIndex	Bone Index in USkeleton.
-	* @param	BoneIndex			Bone Index in Bone Transform array.
-	* @param	RequiredBones		BoneContainer
+	* @param	BoneTransform		BoneTransform to read/write from
+	* @param	SkeletonBoneIndex	Source Bone Index in SourceSkeleton
+	* @param	BoneIndex			Target Bone Index in Bone Transform array
+	* @param	RequiredBones		BoneContainer to which this is retargeting
 	*/
-	static void RetargetBoneTransform(const USkeleton* MySkeleton, const FName& RetargetSource, FTransform& BoneTransform, const int32 SkeletonBoneIndex, const FCompactPoseBoneIndex& BoneIndex, const FBoneContainer& RequiredBones, const bool bIsBakedAdditive);
+	static void RetargetBoneTransform(const USkeleton* SourceSkeleton, const FName& RetargetSource, FTransform& BoneTransform, const int32 SkeletonBoneIndex, const FCompactPoseBoneIndex& BoneIndex, const FBoneContainer& RequiredBones, const bool bIsBakedAdditive);
 	static void RetargetBoneTransform(const USkeleton* MySkeleton, const FName& SourceName, const TArray<FTransform>& RetargetTransforms, FTransform& BoneTransform, const int32 SkeletonBoneIndex, const FCompactPoseBoneIndex& BoneIndex, const FBoneContainer& RequiredBones, const bool bIsBakedAdditive);
-
 	/** 
 	 * Calculate distance how close two strings are. 
 	 * By close, it calculates how many operations to transform First to Second 
diff --git a/Engine/Source/Runtime/Engine/Public/BoneContainer.h b/Engine/Source/Runtime/Engine/Public/BoneContainer.h
--- a/Engine/Source/Runtime/Engine/Public/BoneContainer.h
+++ b/Engine/Source/Runtime/Engine/Public/BoneContainer.h
@@ -13,6 +13,9 @@
 class USkeletalMesh;
 class USkeleton;
 class USkeletalMesh;
+struct FBoneContainer;
+struct FSkeletonRemapping;
+struct FBlendedCurve;
 
 /** in the future if we need more bools, please convert to bitfield 
  * These are not saved in asset but per skeleton. 
@@ -91,6 +94,58 @@ struct FOrientAndScaleRetargetingCachedData
 	}
 };
 
+/**
+ * An array of cached curve remappings. This is used to remap curves between skeletons.
+ * It is used in the FBoneContainer and is generated using a lazy approach. 
+ */
+struct FCachedSkeletonCurveMapping
+{
+	TArray<SmartName::UID_Type>	UIDToArrayIndices; /** The mapping table used for mapping curves. This is indexed by UID and returns the curve index, or MAX_uint16 in case its not used. */
+	bool bIsDirty = true; /** Specifies whether we need to rebuild this cached data or not. */
+};
+
+/**
+ * A scoped curve remapping.
+ * This object is used in a RAII pattern. It initializes the curve remapping on the provided curve.
+ * When the object gets destructed it will automatically revert the remapping on the curve object.
+ * 
+ * Usage would be something like this:
+ * \code{.cpp}
+ * FSkeletonRemapping* Mapping = TargetSkeleton->GetSkeletonRemapping(SourceSkeleton);
+ * if (Mapping)
+ * {
+ *		FSkeletonRemappingCurve Context(BlendedCurve, RequiredBones, Mapping);
+ *		EvaluateCurveData(BlendedCurve);
+ * }
+ * \endcode
+ * 
+ * @param InCurve The curve object that will be used during sampling of the animation data. We will modify its UIDToArrayIndexLUT member to point to a remapped version.
+ * @param InBoneContainer The bone container that we're using. This stores the remappings.
+ * @param InSkeletonMapping The skeleton mapping to use. This can be requested with USkeleton::GetSkeletonRemapping(SourceSkeleton).
+ */
+struct FSkeletonRemappingCurve
+{
+public:
+	FSkeletonRemappingCurve() = delete;
+	FSkeletonRemappingCurve(const FSkeletonRemappingCurve&) = delete;
+	FSkeletonRemappingCurve(FSkeletonRemappingCurve&&) = delete;
+	FSkeletonRemappingCurve(FBlendedCurve& InCurve, const FBoneContainer& InBoneContainer, const FSkeletonRemapping* InSkeletonMapping);
+	FSkeletonRemappingCurve(FBlendedCurve& InCurve, const FBoneContainer& InBoneContainer, const USkeleton* SourceSkeleton);
+	~FSkeletonRemappingCurve();
+
+	FSkeletonRemappingCurve& operator = (const FSkeletonRemappingCurve&) = delete;
+	FSkeletonRemappingCurve& operator = (FSkeletonRemappingCurve&&) = delete;
+
+	FBlendedCurve& GetCurve() { return Curve;  }
+	const FBlendedCurve& GetCurve() const { return Curve; }
+
+private:
+	FBlendedCurve& Curve;
+	const FBoneContainer& BoneContainer;
+	bool bIsRemapping = false;
+};
+
+
 /** Retargeting cached data for a specific Retarget Source */
 struct FRetargetSourceCachedData
 {
@@ -333,6 +388,12 @@ public:
 		return UIDToCurveTypeLUT;
 	}
 
+	/** Get the array that maps UIDs to array indexes. */
+	TArray<SmartName::UID_Type> const& GetUIDToArrayLookupTableBackup() const
+	{
+		return UIDToArrayIndexLUTBackup;
+	}
+
 	/**
 	* Serializes the bones
 	*
@@ -410,7 +471,17 @@ public:
 	const FRetargetSourceCachedData& GetRetargetSourceCachedData(const FName& InRetargetSource) const;
 	const FRetargetSourceCachedData& GetRetargetSourceCachedData(const FName& InSourceName, const TArray<FTransform>& InRetargetTransforms) const;
 
+	// Curve remapping
+	const FCachedSkeletonCurveMapping& GetOrCreateCachedCurveMapping(const FSkeletonRemapping* SkeletonRemapping) const;
+	void MarkAllCachedCurveMappingsDirty();
+
 private:
+	/** The map of cached curve mapping indexes, which is used for skeleton remapping. The key of this table is the source skeleton of the asset we are sampling curve values from. */
+	mutable TMap<USkeleton*, FCachedSkeletonCurveMapping> CachedCurveMappingTable;
+
+	/** The backup of the original UIDToArrayIndexLUT array. This is needed for skeleton remapping curves, as we have to modify this array and later need to restore it again. */
+	TArray<SmartName::UID_Type> UIDToArrayIndexLUTBackup;
+
 	/** 
 	 * Runtime cached data for retargeting from a specific RetargetSource to this current SkelMesh LOD.
 	 * @todo: We could also cache this once per skelmesh per lod, rather than creating it at runtime for each skelmesh instance.
