Wrap bLockLocation into Setter and virtual Getter

diff --git a/Engine/Plugins/Experimental/ActorPalette/Source/ActorPalette/Private/ActorPaletteViewportClient.cpp b/Engine/Plugins/Experimental/ActorPalette/Source/ActorPalette/Private/ActorPaletteViewportClient.cpp
--- a/Engine/Plugins/Experimental/ActorPalette/Source/ActorPalette/Private/ActorPaletteViewportClient.cpp
+++ b/Engine/Plugins/Experimental/ActorPalette/Source/ActorPalette/Private/ActorPaletteViewportClient.cpp
@@ -81,7 +81,7 @@ bool FActorPaletteViewportClient::InputKey(FViewport* InViewport, int32 Controll
 
 			if (HActor* ActorProxy = HitProxyCast<HActor>(HitProxy))
 			{
-				if ((ActorProxy->Actor != nullptr) && !ActorProxy->Actor->bLockLocation)
+				if ((ActorProxy->Actor != nullptr) && !ActorProxy->Actor->IsLockLocation())
 				{
 					TArray<UObject*> Assets;
 					ActorProxy->Actor->GetReferencedContentObjects(Assets);
diff --git a/Engine/Source/Editor/LevelEditor/Private/Elements/Actor/ActorElementLevelEditorViewportInteractionCustomization.cpp b/Engine/Source/Editor/LevelEditor/Private/Elements/Actor/ActorElementLevelEditorViewportInteractionCustomization.cpp
--- a/Engine/Source/Editor/LevelEditor/Private/Elements/Actor/ActorElementLevelEditorViewportInteractionCustomization.cpp
+++ b/Engine/Source/Editor/LevelEditor/Private/Elements/Actor/ActorElementLevelEditorViewportInteractionCustomization.cpp
@@ -127,7 +127,7 @@ bool FActorElementLevelEditorViewportInteractionCustomization::CanMoveActorInVie
 	}
 
 	// The actor cannot be location locked
-	if (InActor->bLockLocation)
+	if (InActor->IsLockLocation())
 	{
 		return false;
 	}
diff --git a/Engine/Source/Editor/LevelEditor/Private/SActorPilotViewportToolbar.h b/Engine/Source/Editor/LevelEditor/Private/SActorPilotViewportToolbar.h
--- a/Engine/Source/Editor/LevelEditor/Private/SActorPilotViewportToolbar.h
+++ b/Engine/Source/Editor/LevelEditor/Private/SActorPilotViewportToolbar.h
@@ -50,7 +50,7 @@ public:
 			Pilot = ViewportPtr->GetLevelViewportClient().GetActiveActorLock().Get();
 		}
 
-		return Pilot && Pilot->bLockLocation ? EVisibility::Visible : EVisibility::Collapsed;
+		return Pilot && Pilot->IsLockLocation() ? EVisibility::Visible : EVisibility::Collapsed;
 	}
 
 	void Construct(const FArguments& InArgs)
diff --git a/Engine/Source/Editor/UnrealEd/Private/EditorEngine.cpp b/Engine/Source/Editor/UnrealEd/Private/EditorEngine.cpp
--- a/Engine/Source/Editor/UnrealEd/Private/EditorEngine.cpp
+++ b/Engine/Source/Editor/UnrealEd/Private/EditorEngine.cpp
@@ -2925,7 +2925,7 @@ void UEditorEngine::ToggleSelectedActorMovementLock()
 
 		// If nothing is locked then we'll turn on locked for all selected actors
 		// Otherwise, we'll turn off locking for any actors that are locked
-		Actor->bLockLocation = !bFoundLockedActor;
+		Actor->SetLockLocation(!bFoundLockedActor);
 
 		LevelDirtyCallback.Request();
 	}
@@ -2943,7 +2943,7 @@ bool UEditorEngine::HasLockedActors()
 			AActor* Actor = Cast<AActor>( *It );
 			checkSlow( Actor );
 
-			if( Actor->bLockLocation )
+			if( Actor->IsLockLocation() )
 			{
 				bHasLockedActors = true;
 				break;
diff --git a/Engine/Source/Editor/UnrealEd/Private/LevelEditorViewport.cpp b/Engine/Source/Editor/UnrealEd/Private/LevelEditorViewport.cpp
--- a/Engine/Source/Editor/UnrealEd/Private/LevelEditorViewport.cpp
+++ b/Engine/Source/Editor/UnrealEd/Private/LevelEditorViewport.cpp
@@ -2995,7 +2995,7 @@ void FLevelEditorViewportClient::TrackingStarted( const FInputEventState& InInpu
 		else
 		{
 			AActor* ActiveActorLock = GetActiveActorLock().Get();
-			if (ActiveActorLock && !ActiveActorLock->bLockLocation && TrackingTransaction.TransCount == 0)
+			if (ActiveActorLock && !ActiveActorLock->IsLockLocation() && TrackingTransaction.TransCount == 0)
 			{
 				// Open a tracking transaction to contain the locked actor changes
 				TrackingTransaction.TransCount++;
@@ -3390,7 +3390,7 @@ void FLevelEditorViewportClient::MoveLockedActorToCamera()
 	AActor* ActiveActorLock = GetActiveActorLock().Get();
 	if (ActiveActorLock)
 	{
-		if (!ActiveActorLock->bLockLocation)
+		if (!ActiveActorLock->IsLockLocation())
 		{
 			if (TrackingTransaction.TransCount > 0)
 			{
diff --git a/Engine/Source/Runtime/Engine/Classes/GameFramework/Actor.h b/Engine/Source/Runtime/Engine/Classes/GameFramework/Actor.h
--- a/Engine/Source/Runtime/Engine/Classes/GameFramework/Actor.h
+++ b/Engine/Source/Runtime/Engine/Classes/GameFramework/Actor.h
@@ -827,6 +827,12 @@ public:
 	/** Returns this actor's current Guid. Actor Guids are only available in development builds. */
 	inline const FGuid& GetActorGuid() const { return ActorGuid; }
 
+	/** Returns true if actor location should be locked. */
+	virtual bool IsLockLocation() const { return bLockLocation; }
+
+	/** Set the bLockLocation flag */
+	void SetLockLocation(bool bInLockLocation) { bLockLocation = bInLockLocation; }
+
 private:
 	friend class UActorDescContainer;
 
@@ -885,15 +891,15 @@ public:
 	/** Whether this actor is hidden by the level browser. */
 	UPROPERTY(Transient)
 	uint8 bHiddenEdLevel:1;
+		
+	/** If true during PostEditMove the construction script will be run every time. If false it will only run when the drag finishes. */
+	uint8 bRunConstructionScriptOnDrag:1;
 
+protected:
 	/** If true, prevents the actor from being moved in the editor viewport. */
 	UPROPERTY()
 	uint8 bLockLocation:1;
 
-	/** If true during PostEditMove the construction script will be run every time. If false it will only run when the drag finishes. */
-	uint8 bRunConstructionScriptOnDrag:1;
-
-protected:
 	/** Is the actor label editable by the user? */
 	UPROPERTY()
 	uint8 bActorLabelEditable:1; 
diff --git a/Engine/Source/Runtime/Engine/Private/ActorPartition/ActorPartitionSubsystem.cpp b/Engine/Source/Runtime/Engine/Private/ActorPartition/ActorPartitionSubsystem.cpp
--- a/Engine/Source/Runtime/Engine/Private/ActorPartition/ActorPartitionSubsystem.cpp
+++ b/Engine/Source/Runtime/Engine/Private/ActorPartition/ActorPartitionSubsystem.cpp
@@ -271,7 +271,7 @@ public:
 			FVector CellCenter(CellBounds.GetCenter());
 			FoundActor = CastChecked<APartitionActor>(World->SpawnActor(InActorPartitionId.GetClass(), &CellCenter, nullptr, SpawnParams));
 			FoundActor->GridSize = InGridSize;
-			FoundActor->bLockLocation = true;
+			FoundActor->SetLockLocation(true);
 			
 			InActorCreated(FoundActor);
 
diff --git a/Engine/Source/Runtime/Engine/Private/Components/BillboardComponent.cpp b/Engine/Source/Runtime/Engine/Private/Components/BillboardComponent.cpp
--- a/Engine/Source/Runtime/Engine/Private/Components/BillboardComponent.cpp
+++ b/Engine/Source/Runtime/Engine/Private/Components/BillboardComponent.cpp
@@ -99,7 +99,7 @@ public:
 
 			//save off override states
 #if WITH_EDITORONLY_DATA
-			bIsActorLocked = Owner->bLockLocation;
+			bIsActorLocked = Owner->IsLockLocation();
 #else // WITH_EDITORONLY_DATA
 			bIsActorLocked = false;
 #endif // WITH_EDITORONLY_DATA
diff --git a/Engine/Source/Runtime/Engine/Private/Components/SceneComponent.cpp b/Engine/Source/Runtime/Engine/Private/Components/SceneComponent.cpp
--- a/Engine/Source/Runtime/Engine/Private/Components/SceneComponent.cpp
+++ b/Engine/Source/Runtime/Engine/Private/Components/SceneComponent.cpp
@@ -3355,7 +3355,7 @@ bool USceneComponent::CanEditChange( const FProperty* Property ) const
 			   Property->GetFName() == TEXT( "RelativeRotation" ) ||
 			   Property->GetFName() == TEXT( "RelativeScale3D" ))
 			{
-				bIsEditable = !Owner->bLockLocation;
+				bIsEditable = !Owner->IsLockLocation();
 			}
 		}
 
diff --git a/Engine/Source/Runtime/Engine/Private/LevelInstance/LevelInstanceActor.cpp b/Engine/Source/Runtime/Engine/Private/LevelInstance/LevelInstanceActor.cpp
--- a/Engine/Source/Runtime/Engine/Private/LevelInstance/LevelInstanceActor.cpp
+++ b/Engine/Source/Runtime/Engine/Private/LevelInstance/LevelInstanceActor.cpp
@@ -24,6 +24,7 @@ ALevelInstance::ALevelInstance(const FObjectInitializer& ObjectInitializer)
 #if WITH_EDITOR
 	, CachedLevelInstanceID(InvalidLevelInstanceID)
 	, bGuardLoadUnload(false)
+	, bEditLockLocation(false)
 #endif
 	, LevelInstanceID(InvalidLevelInstanceID)
 {
@@ -207,12 +208,12 @@ AActor* ALevelInstance::FindEditorInstanceActor() const
 
 void ALevelInstance::OnEdit()
 {
-	bLockLocation = true;
+	bEditLockLocation = true;
 }
 
 void ALevelInstance::OnCommit()
 {
-	bLockLocation = false;
+	bEditLockLocation = false;
 }
 
 ALevelInstance::FOnLevelInstanceActorPostLoad ALevelInstance::OnLevelInstanceActorPostLoad;
@@ -220,7 +221,7 @@ ALevelInstance::FOnLevelInstanceActorPostLoad ALevelInstance::OnLevelInstanceAct
 void ALevelInstance::PostLoad()
 {
 	Super::PostLoad();
-
+		
 	OnLevelInstanceActorPostLoad.Broadcast(this);
 }
 
@@ -581,6 +582,11 @@ void ALevelInstance::GetActorLocationBounds(bool bOnlyCollidingComponents, FVect
 	}
 }
 
+bool ALevelInstance::IsLockLocation() const
+{
+	return Super::IsLockLocation() || bEditLockLocation;
+}
+
 FBox ALevelInstance::GetComponentsBoundingBox(bool bNonColliding, bool bIncludeFromChildActors) const
 {
 	FBox Box = Super::GetComponentsBoundingBox(bNonColliding, bIncludeFromChildActors);
diff --git a/Engine/Source/Runtime/Engine/Public/LevelInstance/LevelInstanceActor.h b/Engine/Source/Runtime/Engine/Public/LevelInstance/LevelInstanceActor.h
--- a/Engine/Source/Runtime/Engine/Public/LevelInstance/LevelInstanceActor.h
+++ b/Engine/Source/Runtime/Engine/Public/LevelInstance/LevelInstanceActor.h
@@ -56,6 +56,7 @@ public:
 	virtual void PushLevelInstanceEditingStateToProxies(bool bInEditingState) override;
 	virtual FBox GetComponentsBoundingBox(bool bNonColliding = false, bool bIncludeFromChildActors = false) const override;
 	virtual void GetActorLocationBounds(bool bOnlyCollidingComponents, FVector& Origin, FVector& BoxExtent, bool bIncludeFromChildActors = false) const override;
+	virtual bool IsLockLocation() const override;
 
 	bool CanEdit(FText* OutReason = nullptr) const;
 	bool CanCommit(FText* OutReason = nullptr) const;
@@ -97,6 +98,7 @@ private:
 	FLevelInstanceID CachedLevelInstanceID;
 	bool bCachedIsTemporarilyHiddenInEditor;
 	bool bGuardLoadUnload;
+	bool bEditLockLocation;
 #else
 	// This Guid is used to compute the LevelInstanceID. Because in non-editor build we don't have an ActorGuid, we store it at cook time.
 	FGuid LevelInstanceActorGuid;
diff --git a/Engine/Source/Runtime/Landscape/Private/Landscape.cpp b/Engine/Source/Runtime/Landscape/Private/Landscape.cpp
--- a/Engine/Source/Runtime/Landscape/Private/Landscape.cpp
+++ b/Engine/Source/Runtime/Landscape/Private/Landscape.cpp
@@ -3184,7 +3184,9 @@ void ULandscapeInfo::RegisterActor(ALandscapeProxy* Proxy, bool bMapCheck)
 			checkf(!LandscapeActor || LandscapeActor == Landscape, TEXT("Multiple landscapes with the same GUID detected: %s vs %s"), *LandscapeActor->GetPathName(), *Landscape->GetPathName());
 			LandscapeActor = Landscape;
 			// In world composition user is not allowed to move landscape in editor, only through WorldBrowser 
-			LandscapeActor->bLockLocation |= OwningWorld != nullptr ? OwningWorld->WorldComposition != nullptr : false;
+			bool bIsLockLocation = LandscapeActor->IsLockLocation();
+			bIsLockLocation |= OwningWorld != nullptr ? OwningWorld->WorldComposition != nullptr : false;
+			LandscapeActor->SetLockLocation(bIsLockLocation);
 
 			// update proxies reference actor
 			for (ALandscapeStreamingProxy* StreamingProxy : Proxies)
