Fixed the actor details panel not refreshing correctly when objects were re-instanced (eg, from a Blueprint compile)

diff --git a/Engine/Source/Editor/Kismet/Private/SSCSEditor.cpp b/Engine/Source/Editor/Kismet/Private/SSCSEditor.cpp
--- a/Engine/Source/Editor/Kismet/Private/SSCSEditor.cpp
+++ b/Engine/Source/Editor/Kismet/Private/SSCSEditor.cpp
@@ -3856,6 +3856,7 @@ void SSCSEditor::Construct( const FArguments& InArgs )
 	OnSelectionUpdated = InArgs._OnSelectionUpdated;
 	OnItemDoubleClicked = InArgs._OnItemDoubleClicked;
 	OnHighlightPropertyInDetailsView = InArgs._OnHighlightPropertyInDetailsView;
+	OnObjectReplaced = InArgs._OnObjectReplaced;
 	bUpdatingSelection = false;
 	bAllowTreeUpdates = true;
 	bIsDiffing = InArgs._IsDiffing;
@@ -4125,18 +4126,23 @@ void SSCSEditor::OnLevelComponentRequestRename(const UActorComponent* InComponen
 
 void SSCSEditor::OnObjectsReplaced(const TMap<UObject*, UObject*>& OldToNewInstanceMap)
 {
-	ReplaceComponentReferencesInTree(GetActorNode(), OldToNewInstanceMap);
+	bool bHasChanges = false;
+	ReplaceComponentReferencesInTree(GetActorNode(), OldToNewInstanceMap, bHasChanges);
+	if (bHasChanges)
+	{
+		OnObjectReplaced.ExecuteIfBound();
+	}
 }
 
-void SSCSEditor::ReplaceComponentReferencesInTree(FSCSEditorActorNodePtrType InActorNode, const TMap<UObject*, UObject*>& OldToNewInstanceMap)
+void SSCSEditor::ReplaceComponentReferencesInTree(FSCSEditorActorNodePtrType InActorNode, const TMap<UObject*, UObject*>& OldToNewInstanceMap, bool& OutHasChanges)
 {
 	if (InActorNode.IsValid())
 	{
-		ReplaceComponentReferencesInTree(InActorNode->GetComponentNodes(), OldToNewInstanceMap);
+		ReplaceComponentReferencesInTree(InActorNode->GetComponentNodes(), OldToNewInstanceMap, OutHasChanges);
 	}
 }
 
-void SSCSEditor::ReplaceComponentReferencesInTree(const TArray<FSCSEditorTreeNodePtrType>& Nodes, const TMap<UObject*, UObject*>& OldToNewInstanceMap)
+void SSCSEditor::ReplaceComponentReferencesInTree(const TArray<FSCSEditorTreeNodePtrType>& Nodes, const TMap<UObject*, UObject*>& OldToNewInstanceMap, bool& OutHasChanges)
 {
 	for (const FSCSEditorTreeNodePtrType& Node : Nodes)
 	{
@@ -4151,10 +4157,11 @@ void SSCSEditor::ReplaceComponentReferencesInTree(const TArray<FSCSEditorTreeNod
 				if (NewObject)
 				{
 					Node->SetObject(*NewObject);
+					OutHasChanges = true;
 				}
 			}
 
-			ReplaceComponentReferencesInTree(Node->GetChildren(), OldToNewInstanceMap);
+			ReplaceComponentReferencesInTree(Node->GetChildren(), OldToNewInstanceMap, OutHasChanges);
 		}
 	}
 }
diff --git a/Engine/Source/Editor/Kismet/Public/SSCSEditor.h b/Engine/Source/Editor/Kismet/Public/SSCSEditor.h
--- a/Engine/Source/Editor/Kismet/Public/SSCSEditor.h
+++ b/Engine/Source/Editor/Kismet/Public/SSCSEditor.h
@@ -1008,6 +1008,7 @@ public:
 		SLATE_EVENT(FOnSelectionUpdated, OnSelectionUpdated)
 		SLATE_EVENT(FOnItemDoubleClicked, OnItemDoubleClicked)
 		SLATE_EVENT(FOnHighlightPropertyInDetailsView, OnHighlightPropertyInDetailsView)
+		SLATE_EVENT(FSimpleDelegate, OnObjectReplaced)
 
 	SLATE_END_ARGS()
 
@@ -1302,10 +1303,10 @@ protected:
 	void OnObjectsReplaced(const TMap<UObject*, UObject*>& OldToNewInstanceMap);
 
 	/** Helper method to update component pointers held by the given actor node's subtree */
-	void ReplaceComponentReferencesInTree(FSCSEditorActorNodePtrType InActorNode, const TMap<UObject*, UObject*>& OldToNewInstanceMap);
+	void ReplaceComponentReferencesInTree(FSCSEditorActorNodePtrType InActorNode, const TMap<UObject*, UObject*>& OldToNewInstanceMap, bool& OutHasChanges);
 
 	/** Update component pointers held by tree nodes if components have been replaced following construction script execution */
-	void ReplaceComponentReferencesInTree(const TArray<FSCSEditorTreeNodePtrType>& Nodes, const TMap<UObject*, UObject*>& OldToNewInstanceMap);
+	void ReplaceComponentReferencesInTree(const TArray<FSCSEditorTreeNodePtrType>& Nodes, const TMap<UObject*, UObject*>& OldToNewInstanceMap, bool& OutHasChanges);
 
 	/**
 	 * Function to create events for the current selection
@@ -1473,6 +1474,9 @@ public:
 	/** Delegate to invoke when the given property should be highlighted in the details view (e.g. diff). */
 	FOnHighlightPropertyInDetailsView OnHighlightPropertyInDetailsView;
 
+	/** Delegate to invoke when objects within the SCS tree are replaced (eg, via re-instancing from a BP compile) */
+	FSimpleDelegate OnObjectReplaced;
+
 	/** Returns the Actor context for which we are viewing/editing the SCS.  Can return null.  Should not be cached as it may change from frame to frame. */
 	class AActor* GetActorContext() const;
 
diff --git a/Engine/Source/Editor/LevelEditor/Private/SActorDetails.cpp b/Engine/Source/Editor/LevelEditor/Private/SActorDetails.cpp
--- a/Engine/Source/Editor/LevelEditor/Private/SActorDetails.cpp
+++ b/Engine/Source/Editor/LevelEditor/Private/SActorDetails.cpp
@@ -90,8 +90,9 @@ void SActorDetails::Construct(const FArguments& InArgs, UTypedElementSelectionSe
 {
 	SelectionSet = InSelectionSet;
 	checkf(SelectionSet, TEXT("SActorDetails must be constructed with a valid selection set!"));
-	SelectionSet->OnChanged().AddRaw(this, &SActorDetails::OnElementSelectionChanged);
-	
+
+	FCoreUObjectDelegates::OnObjectsReplaced.AddRaw(this, &SActorDetails::OnObjectsReplaced);
+
 	FLevelEditorModule& LevelEditor = FModuleManager::GetModuleChecked<FLevelEditorModule>("LevelEditor");
 	LevelEditor.OnComponentsEdited().AddRaw(this, &SActorDetails::OnComponentsEditedInWorld);
 
@@ -140,6 +141,7 @@ void SActorDetails::Construct(const FArguments& InArgs, UTypedElementSelectionSe
 			.ActorContext(this, &SActorDetails::GetActorContext)
 			.OnSelectionUpdated(this, &SActorDetails::OnSCSEditorTreeViewSelectionChanged)
 			.OnItemDoubleClicked(this, &SActorDetails::OnSCSEditorTreeViewItemDoubleClicked)
+			.OnObjectReplaced(this, &SActorDetails::OnSCSEditorTreeViewObjectReplaced)
 		];
 
 	TSharedRef<SWidget> ButtonBox = SCSEditor->GetToolButtonsBox().ToSharedRef();
@@ -221,7 +223,7 @@ SActorDetails::~SActorDetails()
 		GEditor->UnregisterForUndo(this);
 	}
 
-	SelectionSet->OnChanged().RemoveAll(this);
+	FCoreUObjectDelegates::OnObjectsReplaced.RemoveAll(this);
 
 	RemoveBPComponentCompileEventDelegate();
 
@@ -239,6 +241,11 @@ bool SActorDetails::IsObservingSelectionSet(const UTypedElementSelectionSet* InS
 
 void SActorDetails::RefreshSelection(const bool bForceRefresh)
 {
+	if (bSelectionGuard)
+	{
+		return;
+	}
+
 	TArray<TTypedElement<UTypedElementDetailsInterface>> DetailsElements;
 	DetailsElements.Reserve(SelectionSet->GetNumSelectedElements());
 	SelectionSet->ForEachSelectedElement<UTypedElementDetailsInterface>([&DetailsElements](const TTypedElement<UTypedElementDetailsInterface>& InDetailsElement)
@@ -247,6 +254,9 @@ void SActorDetails::RefreshSelection(const bool bForceRefresh)
 		return true;
 	});
 
+	bHasSelectionOverride = false;
+	SelectionOverrideActors.Reset();
+
 	RefreshTopLevelElements(DetailsElements, bForceRefresh, /*bOverrideLock*/false);
 }
 
@@ -267,6 +277,9 @@ void SActorDetails::OverrideSelection(const TArray<AActor*>& InActors, const boo
 		}
 	}
 
+	bHasSelectionOverride = true;
+	SelectionOverrideActors = InActors;
+
 	RefreshTopLevelElements(DetailsElements, bForceRefresh, /*bOverrideLock*/false);
 }
 
@@ -295,14 +308,12 @@ void SActorDetails::RefreshTopLevelElements(TArrayView<const TTypedElement<UType
 	SetElementDetailsObjects(TopLevelElements, bForceRefresh, bOverrideLock);
 
 	// Update the SCS tree if we were asked to edit a single actor
-	if (InDetailsElements.Num() == 1)
+	if (AActor* Actor = GetActorContext())
 	{
-		if (AActor* Actor = GetActorContext())
-		{
-			// Enable the selection guard to prevent OnTreeSelectionChanged() from altering the editor's component selection
-			TGuardValue<bool> SelectionGuard(bSelectionGuard, true);
-			SCSEditor->UpdateTree();
-		}
+		// Enable the selection guard to prevent OnTreeSelectionChanged() from altering the editor's component selection
+		TGuardValue<bool> SelectionGuard(bSelectionGuard, true);
+		SCSEditor->UpdateTree();
+		UpdateComponentTreeFromEditorSelection();
 	}
 
 	// Draw attention to this tab if needed
@@ -452,15 +463,6 @@ void SActorDetails::OnComponentsEditedInWorld()
 	}
 }
 
-void SActorDetails::OnElementSelectionChanged(const UTypedElementSelectionSet* InSelectionSet)
-{
-	check(IsObservingSelectionSet(InSelectionSet));
-	if (!bSelectionGuard && SCSEditor)
-	{
-		UpdateComponentTreeFromEditorSelection();
-	}
-}
-
 bool SActorDetails::GetAllowComponentTreeEditing() const
 {
 	return GEditor->PlayWorld == nullptr;
@@ -568,6 +570,13 @@ void SActorDetails::OnSCSEditorTreeViewItemDoubleClicked(const TSharedPtr<class
 	}
 }
 
+void SActorDetails::OnSCSEditorTreeViewObjectReplaced()
+{
+	// Enable the selection guard to prevent OnTreeSelectionChanged() from altering the editor's component selection
+	TGuardValue<bool> SelectionGuard(bSelectionGuard, true);
+	SCSEditor->UpdateTree();
+}
+
 void SActorDetails::UpdateComponentTreeFromEditorSelection()
 {
 	if (DetailsView->IsLocked())
@@ -803,3 +812,37 @@ void SActorDetails::OnBlueprintComponentCompiled(UBlueprint* ComponentBlueprint)
 	TGuardValue<bool> SelectedComponentRecompiledGuard(bSelectedComponentRecompiled, true);
 	UpdateComponentTreeFromEditorSelection();
 }
+
+void SActorDetails::OnObjectsReplaced(const TMap<UObject*, UObject*>& InReplacementObjects)
+{
+	if (bHasSelectionOverride && SelectionOverrideActors.Num() > 0)
+	{
+		bool bHasChanges = false;
+
+		for (auto It = SelectionOverrideActors.CreateIterator(); It; ++It)
+		{
+			AActor*& Actor = *It;
+
+			if (UObject* const* ReplacementObjectPtr = InReplacementObjects.Find(Actor))
+			{
+				bHasChanges = true;
+
+				AActor* ReplacementActor = Cast<AActor>(*ReplacementObjectPtr);
+				if (ReplacementActor)
+				{
+					Actor = ReplacementActor;
+				}
+				else
+				{
+					It.RemoveCurrent();
+				}
+			}
+		}
+
+		if (bHasChanges)
+		{
+			TArray<AActor*> NewSelection = SelectionOverrideActors;
+			OverrideSelection(NewSelection);
+		}
+	}
+}
diff --git a/Engine/Source/Editor/LevelEditor/Private/SActorDetails.h b/Engine/Source/Editor/LevelEditor/Private/SActorDetails.h
--- a/Engine/Source/Editor/LevelEditor/Private/SActorDetails.h
+++ b/Engine/Source/Editor/LevelEditor/Private/SActorDetails.h
@@ -77,9 +77,9 @@ private:
 	bool GetAllowComponentTreeEditing() const;
 
 	void OnComponentsEditedInWorld();
-	void OnElementSelectionChanged(const UTypedElementSelectionSet* InSelectionSet);
 	void OnSCSEditorTreeViewSelectionChanged(const TArray<TSharedPtr<class FSCSEditorTreeNode> >& SelectedNodes);
 	void OnSCSEditorTreeViewItemDoubleClicked(const TSharedPtr<class FSCSEditorTreeNode> ClickedNode);
+	void OnSCSEditorTreeViewObjectReplaced();
 	void UpdateComponentTreeFromEditorSelection();
 
 	bool IsPropertyReadOnly(const struct FPropertyAndParent& PropertyAndParent) const;
@@ -94,6 +94,7 @@ private:
 	void AddBPComponentCompileEventDelegate(UBlueprint* ComponentBlueprint);
 	void RemoveBPComponentCompileEventDelegate();
 	void OnBlueprintComponentCompiled(UBlueprint* ComponentBlueprint);
+	void OnObjectsReplaced(const TMap<UObject*, UObject*>& InReplacementObjects);
 
 private:
 	TSharedPtr<SSplitter> DetailsSplitter;
@@ -104,6 +105,10 @@ private:
 	// The selection set this details panel is observing
 	UTypedElementSelectionSet* SelectionSet = nullptr;
 
+	// The selection override, if any
+	bool bHasSelectionOverride = false;
+	TArray<AActor*> SelectionOverrideActors;
+
 	// Array of top-level elements that are currently being edited
 	TArray<TUniquePtr<ITypedElementDetailsObject>> TopLevelElements;
 
diff --git a/Engine/Source/Editor/PropertyEditor/Private/PropertyEditorModule.cpp b/Engine/Source/Editor/PropertyEditor/Private/PropertyEditorModule.cpp
--- a/Engine/Source/Editor/PropertyEditor/Private/PropertyEditorModule.cpp
+++ b/Engine/Source/Editor/PropertyEditor/Private/PropertyEditorModule.cpp
@@ -90,6 +90,8 @@ const FPropertyTypeLayoutCallback& FPropertyTypeLayoutCallbackList::Find( const
 void FPropertyEditorModule::StartupModule()
 {
 	StructOnScopePropertyOwner = nullptr;
+
+	FCoreUObjectDelegates::OnObjectsReplaced.AddRaw(this, &FPropertyEditorModule::ReplaceViewedObjects);
 }
 
 void FPropertyEditorModule::ShutdownModule()
@@ -103,6 +105,8 @@ void FPropertyEditorModule::ShutdownModule()
 	//       templates, such as delegate wrapper objects that are allocated by the module!
 	DestroyColorPicker();
 
+	FCoreUObjectDelegates::OnObjectsReplaced.RemoveAll(this);
+
 	AllDetailViews.Empty();
 	AllSinglePropertyViews.Empty();
 }
diff --git a/Engine/Source/Editor/UnrealEd/Classes/Editor/UnrealEdEngine.h b/Engine/Source/Editor/UnrealEd/Classes/Editor/UnrealEdEngine.h
--- a/Engine/Source/Editor/UnrealEd/Classes/Editor/UnrealEdEngine.h
+++ b/Engine/Source/Editor/UnrealEd/Classes/Editor/UnrealEdEngine.h
@@ -260,7 +260,7 @@ public:
 	void DrawComponentVisualizersHUD(const FViewport* Viewport, const FSceneView* View, FCanvas* Canvas);
 
 	/** Updates the property windows of selected actors */
-	virtual void UpdateFloatingPropertyWindows(bool bForceRefresh=false);
+	void UpdateFloatingPropertyWindows(bool bForceRefresh=false, bool bNotifyActorSelectionChanged=true);
 
 	/**
 	*	Updates the property windows to show the data of the supplied ActorList
@@ -268,7 +268,7 @@ public:
 	*	@param	ActorList	The list of actors to show the properties for
 	*
 	*/
-	virtual void UpdateFloatingPropertyWindowsFromActorList(const TArray<AActor*>& ActorList, bool bForceRefresh=false);
+	void UpdateFloatingPropertyWindowsFromActorList(const TArray<AActor*>& ActorList, bool bForceRefresh=false);
 
 	/**
 	 * Called whenever the actor selection has changed to invalidate any cached state
diff --git a/Engine/Source/Editor/UnrealEd/Private/EditorEngine.cpp b/Engine/Source/Editor/UnrealEd/Private/EditorEngine.cpp
--- a/Engine/Source/Editor/UnrealEd/Private/EditorEngine.cpp
+++ b/Engine/Source/Editor/UnrealEd/Private/EditorEngine.cpp
@@ -5989,41 +5989,8 @@ void UEditorEngine::DoConvertActors( const TArray<AActor*>& ActorsToConvert, UCl
 
 void UEditorEngine::NotifyToolsOfObjectReplacement(const TMap<UObject*, UObject*>& OldToNewInstanceMap)
 {
-	// This can be called early on during startup if blueprints need to be compiled.  
-	// If the property module isn't loaded then there aren't any property windows to update
-	if( FModuleManager::Get().IsModuleLoaded( "PropertyEditor" ) )
-	{
-		FPropertyEditorModule& PropertyEditorModule = FModuleManager::GetModuleChecked<FPropertyEditorModule>( "PropertyEditor" );
-		PropertyEditorModule.ReplaceViewedObjects( OldToNewInstanceMap );
-	}
-
 	// Allow any other observers to act upon the object replacement
 	FCoreUObjectDelegates::OnObjectsReplaced.Broadcast(OldToNewInstanceMap);
-
-	// Check to see if any selected components were reinstanced, as a final step.
-	USelection* ComponentSelection = GetSelectedComponents();
-	if (ComponentSelection)
-	{
-		TArray<TWeakObjectPtr<UObject> > SelectedComponents;
-		ComponentSelection->GetSelectedObjects(SelectedComponents);
-
-		ComponentSelection->BeginBatchSelectOperation();
-		for (int32 i = 0; i < SelectedComponents.Num(); ++i)
-		{
-			UObject* Component = SelectedComponents[i].GetEvenIfUnreachable();
-
-			// If the component corresponds to a new instance in the map, update the selection accordingly
-			if (OldToNewInstanceMap.Contains(Component))
-			{
-				if (UActorComponent* NewComponent = CastChecked<UActorComponent>(OldToNewInstanceMap[Component], ECastCheckedType::NullAllowed))
-				{
-					ComponentSelection->Deselect(Component);
-					SelectComponent(NewComponent, true, false);
-				}
-			}
-		}
-		ComponentSelection->EndBatchSelectOperation();
-	}
 }
 
 void UEditorEngine::SetViewportsRealtimeOverride(bool bShouldBeRealtime, FText SystemDisplayName)
diff --git a/Engine/Source/Editor/UnrealEd/Private/EditorHook.cpp b/Engine/Source/Editor/UnrealEd/Private/EditorHook.cpp
--- a/Engine/Source/Editor/UnrealEd/Private/EditorHook.cpp
+++ b/Engine/Source/Editor/UnrealEd/Private/EditorHook.cpp
@@ -20,24 +20,27 @@ void UUnrealEdEngine::NotifyPostChange(const FPropertyChangedEvent& PropertyChan
 	GLevelEditorModeTools().ActorPropChangeNotify();
 }
 
-void UUnrealEdEngine::UpdateFloatingPropertyWindows(bool bForceRefresh)
+void UUnrealEdEngine::UpdateFloatingPropertyWindows(bool bForceRefresh, bool bNotifyActorSelectionChanged)
 {
 	if (const UTypedElementSelectionSet* SelectionSet = GetSelectedActors()->GetElementSelectionSet())
 	{
 		FLevelEditorModule& LevelEditor = FModuleManager::LoadModuleChecked<FLevelEditorModule>(TEXT("LevelEditor"));
 		LevelEditor.BroadcastElementSelectionChanged(SelectionSet, bForceRefresh);
 
-		// Assemble a set of valid selected actors.
-		TArray<UObject*> SelectedActors;
-		SelectionSet->ForEachSelectedObject<AActor>([&SelectedActors](AActor* InActor)
+		if (bNotifyActorSelectionChanged)
 		{
-			if (!InActor->IsPendingKill())
+			// Assemble a set of valid selected actors.
+			TArray<UObject*> SelectedActors;
+			SelectionSet->ForEachSelectedObject<AActor>([&SelectedActors](AActor* InActor)
 			{
-				SelectedActors.Add(InActor);
-			}
-			return true;
-		});
-		LevelEditor.BroadcastActorSelectionChanged(SelectedActors, bForceRefresh);
+				if (!InActor->IsPendingKill())
+				{
+					SelectedActors.Add(InActor);
+				}
+				return true;
+			});
+			LevelEditor.BroadcastActorSelectionChanged(SelectedActors, bForceRefresh);
+		}
 	}
 }
 
diff --git a/Engine/Source/Editor/UnrealEd/Private/EditorSelectUtils.cpp b/Engine/Source/Editor/UnrealEd/Private/EditorSelectUtils.cpp
--- a/Engine/Source/Editor/UnrealEd/Private/EditorSelectUtils.cpp
+++ b/Engine/Source/Editor/UnrealEd/Private/EditorSelectUtils.cpp
@@ -419,9 +419,9 @@ void UUnrealEdEngine::NoteSelectionChange(bool bNotify)
 	if (!bComponentSelectionChanged)
 	{
 		PostActorSelectionChanged();
-		UpdateFloatingPropertyWindows();
 	}
 
+	UpdateFloatingPropertyWindows(/*bForceRefresh*/false, !bComponentSelectionChanged);
 	RedrawLevelEditingViewports();
 }
 
diff --git a/Engine/Source/Runtime/Engine/Private/ActorConstruction.cpp b/Engine/Source/Runtime/Engine/Private/ActorConstruction.cpp
--- a/Engine/Source/Runtime/Engine/Private/ActorConstruction.cpp
+++ b/Engine/Source/Runtime/Engine/Private/ActorConstruction.cpp
@@ -593,6 +593,8 @@ void AActor::RerunConstructionScripts()
 		// Now iterate through all previous construction script created components, looking for a match with reinstanced components.
 		for (const FComponentData& ComponentData : ComponentMapping)
 		{
+			UActorComponent* ResolvedNewComponent = nullptr;
+			
 			if (ComponentData.OldComponent->CreationMethod == EComponentCreationMethod::UserConstructionScript)
 			{
 				if (ComponentData.UCSComponentIndex >= 0)
@@ -606,7 +608,7 @@ void AActor::RerunConstructionScripts()
 							if (   ComponentData.OldComponent->GetClass() == NewComponent->GetClass() 
 							    && ComponentData.OldArchetype == ComponentToArchetypeMap[NewComponent])
 							{
-								OldToNewComponentMapping.Add(ComponentData.OldComponent, NewComponent);
+								ResolvedNewComponent = NewComponent;
 								NewUCSComponentsToConsider->RemoveAtSwap(Index, 1, false);
 								break;
 							}
@@ -643,12 +645,14 @@ void AActor::RerunConstructionScripts()
 					{
 						if (!OuterToMatch || GetComponentAddedByConstructionScript(MatchedComponent) == OuterToMatch)
 						{
-							OldToNewComponentMapping.Add(ComponentData.OldComponent, MatchedComponent);
+							ResolvedNewComponent = MatchedComponent;
 							break;
 						}
 					}
 				}
 			}
+
+			OldToNewComponentMapping.Add(ComponentData.OldComponent, ResolvedNewComponent);
 		}
 
 		if (GEditor && (OldToNewComponentMapping.Num() > 0))
diff --git a/Engine/Source/Runtime/Engine/Private/Elements/Framework/EngineElementsLibrary.cpp b/Engine/Source/Runtime/Engine/Private/Elements/Framework/EngineElementsLibrary.cpp
--- a/Engine/Source/Runtime/Engine/Private/Elements/Framework/EngineElementsLibrary.cpp
+++ b/Engine/Source/Runtime/Engine/Private/Elements/Framework/EngineElementsLibrary.cpp
@@ -72,21 +72,14 @@ FTypedElementHandle AcquireEditorTypedElementHandle(const ObjectClass* Object, T
 }
 
 template <typename ObjectClass, typename ElementDataType>
-void ReplaceEditorTypedElementHandles(const TMap<UObject*, UObject*>& ReplacementObjects, TTypedElementOwnerStore<ElementDataType, const ObjectClass*>& ElementOwnerStore, TFunctionRef<void(const ObjectClass*, TTypedElementOwner<ElementDataType>&)> UpdateElement, TFunctionRef<void(const ObjectClass*, TTypedElementOwner<ElementDataType>&)> DestroyElement)
+void ReplaceEditorTypedElementHandles(TArray<FTypedElementHandle>& OutUpdatedElements, const TMap<UObject*, UObject*>& ReplacementObjects, TTypedElementOwnerStore<ElementDataType, const ObjectClass*>& ElementOwnerStore, TFunctionRef<void(const ObjectClass*, TTypedElementOwner<ElementDataType>&)> UpdateElement, TFunctionRef<void(const ObjectClass*, TTypedElementOwner<ElementDataType>&)> DestroyElement)
 {
 	for (const TTuple<UObject*, UObject*>& ReplacementObjectPair : ReplacementObjects)
 	{
 		const ObjectClass* OldObject = Cast<ObjectClass>(ReplacementObjectPair.Key);
-		if (!OldObject || !ReplacementObjectPair.Value)
+		if (!OldObject)
 		{
-			// We can't do anything for objects of the incorrect type, or redirection to null...
-			continue;
-		}
-
-		const ObjectClass* NewObject = Cast<ObjectClass>(ReplacementObjectPair.Value);
-		if (!NewObject)
-		{
-			UE_LOG(LogCore, Error, TEXT("Typed element replacement failed! Objects must be compatible types - Expected: %s, Old: %s (%s), New: %s (%s)"), *ObjectClass::StaticClass()->GetName(), *OldObject->GetClass()->GetName(), *OldObject->GetPathName(), *ReplacementObjectPair.Key->GetClass()->GetName(), *ReplacementObjectPair.Key->GetPathName());
+			// We can't do anything for objects of the incorrect type...
 			continue;
 		}
 
@@ -95,16 +88,35 @@ void ReplaceEditorTypedElementHandles(const TMap<UObject*, UObject*>& Replacemen
 		{
 			if (OldEditorElement.Private_GetInternalData()->GetRefCount() > 1)
 			{
-				// The old element has external references, so we're going to destroy the new element (if any) and re-point the old element to the new object...
-				// Note: This requires that the new element has no external references - if both old and new elements have external references then we'll need to support redirection at the element level!
-				if (TTypedElementOwner<ElementDataType> NewEditorElement = ElementOwnerStore.UnregisterElementOwner(NewObject))
+				if (const ObjectClass* NewObject = Cast<ObjectClass>(ReplacementObjectPair.Value))
 				{
-					checkf(NewEditorElement.Private_GetInternalData()->GetRefCount() <= 1, TEXT("The old and new element both have external references! Replacing these will require support for redirection at the element level!"));
-					DestroyElement(NewObject, NewEditorElement);
+					// The old element has external references, so we're going to destroy the new element (if any) and re-point the old element to the new object...
+					// Note: This requires that the new element has no external references - if both old and new elements have external references then we'll need to support redirection at the element level!
+					if (TTypedElementOwner<ElementDataType> NewEditorElement = ElementOwnerStore.UnregisterElementOwner(NewObject))
+					{
+						if (NewEditorElement.Private_GetInternalData()->GetRefCount() > 1)
+						{
+							// Both elements currently have external references, so try and update anything referencing the new element so that it references the old one instead...
+							// Note: If we had redirection at the element level, then we'd try and do this redirection from old->new instead
+							TTuple<FTypedElementHandle, FTypedElementHandle> ElementRedirect = MakeTuple(NewEditorElement.AcquireHandle(), OldEditorElement.AcquireHandle());
+							UTypedElementRegistry::GetInstance()->OnElementReplaced().Broadcast(MakeArrayView(&ElementRedirect, 1));
+						}
+
+						checkf(NewEditorElement.Private_GetInternalData()->GetRefCount() <= 1, TEXT("The old and new element both have external references! Replacing these will require support for redirection at the element level!"));
+						DestroyElement(NewObject, NewEditorElement);
+					}
+
+					UpdateElement(NewObject, OldEditorElement);
+					OutUpdatedElements.Emplace(OldEditorElement.AcquireHandle());
+					ElementOwnerStore.RegisterElementOwner(NewObject, MoveTemp(OldEditorElement));
+				}
+				else
+				{
+					// The object has been redirected to null, so try and clear any external references to the old element...
+					TTuple<FTypedElementHandle, FTypedElementHandle> ElementRedirect = MakeTuple(OldEditorElement.AcquireHandle(), FTypedElementHandle());
+					UTypedElementRegistry::GetInstance()->OnElementReplaced().Broadcast(MakeArrayView(&ElementRedirect, 1));
+					DestroyElement(OldObject, OldEditorElement);
 				}
-
-				UpdateElement(NewObject, OldEditorElement);
-				ElementOwnerStore.RegisterElementOwner(NewObject, MoveTemp(OldEditorElement));
 			}
 			else
 			{
@@ -141,20 +153,27 @@ UEngineElementsLibrary::UEngineElementsLibrary()
 #if WITH_EDITOR
 void UEngineElementsLibrary::OnObjectsReplaced(const TMap<UObject*, UObject*>& InReplacementObjects)
 {
-	EngineElementsLibraryUtil::ReplaceEditorTypedElementHandles<UObject, FObjectElementData>(InReplacementObjects, GObjectElementOwnerStore, [](const UObject* InObject, TTypedElementOwner<FObjectElementData>& InOutObjectElement)
+	TArray<FTypedElementHandle> UpdatedElements;
+
+	EngineElementsLibraryUtil::ReplaceEditorTypedElementHandles<UObject, FObjectElementData>(UpdatedElements, InReplacementObjects, GObjectElementOwnerStore, [](const UObject* InObject, TTypedElementOwner<FObjectElementData>& InOutObjectElement)
 	{
 		InOutObjectElement.GetDataChecked().Object = const_cast<UObject*>(InObject);
 	}, &UEngineElementsLibrary::DestroyObjectElement);
 
-	EngineElementsLibraryUtil::ReplaceEditorTypedElementHandles<AActor, FActorElementData>(InReplacementObjects, GActorElementOwnerStore, [](const AActor* InActor, TTypedElementOwner<FActorElementData>& InOutActorElement)
+	EngineElementsLibraryUtil::ReplaceEditorTypedElementHandles<AActor, FActorElementData>(UpdatedElements, InReplacementObjects, GActorElementOwnerStore, [](const AActor* InActor, TTypedElementOwner<FActorElementData>& InOutActorElement)
 	{
 		InOutActorElement.GetDataChecked().Actor = const_cast<AActor*>(InActor);
 	}, &UEngineElementsLibrary::DestroyActorElement);
 
-	EngineElementsLibraryUtil::ReplaceEditorTypedElementHandles<UActorComponent, FComponentElementData>(InReplacementObjects, GComponentElementOwnerStore, [](const UActorComponent* InComponent, TTypedElementOwner<FComponentElementData>& InOutComponentElement)
+	EngineElementsLibraryUtil::ReplaceEditorTypedElementHandles<UActorComponent, FComponentElementData>(UpdatedElements, InReplacementObjects, GComponentElementOwnerStore, [](const UActorComponent* InComponent, TTypedElementOwner<FComponentElementData>& InOutComponentElement)
 	{
 		InOutComponentElement.GetDataChecked().Component = const_cast<UActorComponent*>(InComponent);
 	}, &UEngineElementsLibrary::DestroyComponentElement);
+
+	if (UpdatedElements.Num() > 0)
+	{
+		UTypedElementRegistry::GetInstance()->OnElementUpdated().Broadcast(UpdatedElements);
+	}
 }
 #endif
 
diff --git a/Engine/Source/Runtime/TypedElementFramework/Public/Elements/Framework/TypedElementRegistry.h b/Engine/Source/Runtime/TypedElementFramework/Public/Elements/Framework/TypedElementRegistry.h
--- a/Engine/Source/Runtime/TypedElementFramework/Public/Elements/Framework/TypedElementRegistry.h
+++ b/Engine/Source/Runtime/TypedElementFramework/Public/Elements/Framework/TypedElementRegistry.h
@@ -41,6 +41,25 @@ public:
 	UFUNCTION(BlueprintPure, Category="TypedElementFramework|Registry")
 	static UTypedElementRegistry* GetInstance();
 
+	/**
+	 * Event fired when references to one element should be replaced with a reference to a different element.
+	 */
+	using FOnElementReplaced_Payload = TArrayView<const TTuple<FTypedElementHandle, FTypedElementHandle>>; // Alias for the macro below
+	DECLARE_MULTICAST_DELEGATE_OneParam(FOnElementReplaced, FOnElementReplaced_Payload /*ReplacedElements*/);
+	FOnElementReplaced& OnElementReplaced()
+	{
+		return OnElementReplacedDelegate;
+	}
+
+	/**
+	 * Event fired when an element has been internally updated and data cached from it should be refreshed.
+	 */
+	DECLARE_MULTICAST_DELEGATE_OneParam(FOnElementUpdated, TArrayView<const FTypedElementHandle> /*UpdatedElements*/);
+	FOnElementUpdated& OnElementUpdated()
+	{
+		return OnElementUpdatedDelegate;
+	}
+
 	/**
 	 * Get the element type ID for the associated element type name, if any.
 	 * @return The element type ID, or 0 if the given name wasn't registered.
@@ -487,4 +506,7 @@ private:
 
 	uint8 DisableElementDestructionOnGCCount = 0;
 	bool bIsWithinFrame = false;
+
+	FOnElementReplaced OnElementReplacedDelegate;
+	FOnElementUpdated OnElementUpdatedDelegate;
 };
diff --git a/Engine/Source/Runtime/TypedElementRuntime/Private/Elements/Framework/TypedElementSelectionSet.cpp b/Engine/Source/Runtime/TypedElementRuntime/Private/Elements/Framework/TypedElementSelectionSet.cpp
--- a/Engine/Source/Runtime/TypedElementRuntime/Private/Elements/Framework/TypedElementSelectionSet.cpp
+++ b/Engine/Source/Runtime/TypedElementRuntime/Private/Elements/Framework/TypedElementSelectionSet.cpp
@@ -10,7 +10,11 @@ UTypedElementSelectionSet::UTypedElementSelectionSet()
 {
 	if (!HasAnyFlags(RF_ClassDefaultObject))
 	{
-		ElementList = UTypedElementRegistry::GetInstance()->CreateElementList();
+		UTypedElementRegistry* Registry = UTypedElementRegistry::GetInstance();
+		Registry->OnElementReplaced().AddUObject(this, &UTypedElementSelectionSet::OnElementReplaced);
+		Registry->OnElementUpdated().AddUObject(this, &UTypedElementSelectionSet::OnElementUpdated);
+
+		ElementList = Registry->CreateElementList();
 		ElementList->OnPreChange().AddUObject(this, &UTypedElementSelectionSet::OnElementListPreChange);
 		ElementList->OnChanged().AddUObject(this, &UTypedElementSelectionSet::OnElementListChanged);
 	}
@@ -303,6 +307,47 @@ FTypedElementSelectionSetElement UTypedElementSelectionSet::ResolveSelectionSetE
 		: FTypedElementSelectionSetElement();
 }
 
+void UTypedElementSelectionSet::OnElementReplaced(TArrayView<const TTuple<FTypedElementHandle, FTypedElementHandle>> InReplacedElements)
+{
+	const FTypedElementSelectionOptions SelectionOptions = FTypedElementSelectionOptions()
+		.SetAllowHidden(true)
+		.SetAllowGroups(false)
+		.SetAllowLegacyNotifications(false)
+		.SetWarnIfLocked(false);
+
+	for (const TTuple<FTypedElementHandle, FTypedElementHandle>& ReplacedElement : InReplacedElements)
+	{
+		if (ElementList->Num() == 0)
+		{
+			break;
+		}
+
+		if (ElementList->Remove(ReplacedElement.Key))
+		{
+			if (ReplacedElement.Value)
+			{
+				SelectElement(ReplacedElement.Value, SelectionOptions);
+			}
+		}
+	}
+}
+
+void UTypedElementSelectionSet::OnElementUpdated(TArrayView<const FTypedElementHandle> InUpdatedElements)
+{
+	if (ElementList->Num() == 0)
+	{
+		return;
+	}
+	
+	for (const FTypedElementHandle& UpdatedElement : InUpdatedElements)
+	{
+		if (ElementList->Remove(UpdatedElement))
+		{
+			ElementList->Add(UpdatedElement);
+		}
+	}
+}
+
 void UTypedElementSelectionSet::OnElementListPreChange(const UTypedElementList* InElementList)
 {
 	check(InElementList == ElementList);
diff --git a/Engine/Source/Runtime/TypedElementRuntime/Public/Elements/Framework/TypedElementSelectionSet.h b/Engine/Source/Runtime/TypedElementRuntime/Public/Elements/Framework/TypedElementSelectionSet.h
--- a/Engine/Source/Runtime/TypedElementRuntime/Public/Elements/Framework/TypedElementSelectionSet.h
+++ b/Engine/Source/Runtime/TypedElementRuntime/Public/Elements/Framework/TypedElementSelectionSet.h
@@ -490,6 +490,16 @@ private:
 	 */
 	FTypedElementSelectionSetElement ResolveSelectionSetElement(const FTypedElementHandle& InElementHandle) const;
 
+	/**
+	 * Update the selection from a replacement request.
+	 */
+	void OnElementReplaced(TArrayView<const TTuple<FTypedElementHandle, FTypedElementHandle>> InReplacedElements);
+
+	/**
+	 * Force a selection update if an element internal state changes.
+	 */
+	void OnElementUpdated(TArrayView<const FTypedElementHandle> InUpdatedElements);
+
 	/**
 	 * Proxy the internal OnPreChange event from the underlying element list.
 	 */
