Automated test - Allow image comparison of arbitrary image data

Implement:
 RequestImageComparison(FString& ImageName, int32 Width, int32 Height, TArray<FColor>& ImageData, EAutomationComparisonToleranceLevel Tolerance, FString& Notes);

diff --git a/Engine/Source/Developer/FunctionalTesting/Public/AutomationBlueprintFunctionLibrary.h b/Engine/Source/Developer/FunctionalTesting/Public/AutomationBlueprintFunctionLibrary.h
--- a/Engine/Source/Developer/FunctionalTesting/Public/AutomationBlueprintFunctionLibrary.h
+++ b/Engine/Source/Developer/FunctionalTesting/Public/AutomationBlueprintFunctionLibrary.h
@@ -25,7 +25,7 @@ public:
 	virtual void SetDone() { Done = true; };
 
 protected:
-	bool Done;
+	bool Done = false;
 };
 
 /**
diff --git a/Engine/Source/Runtime/Core/Public/Misc/AutomationTest.h b/Engine/Source/Runtime/Core/Public/Misc/AutomationTest.h
--- a/Engine/Source/Runtime/Core/Public/Misc/AutomationTest.h
+++ b/Engine/Source/Runtime/Core/Public/Misc/AutomationTest.h
@@ -646,6 +646,60 @@ struct CORE_API FAutomationScreenshotCompareResults
 	FAutomationEvent ToAutomationEvent(const FString& ScreenhotName) const;
 };
 
+enum class EAutomationComparisonToleranceLevel : uint8
+{
+	Zero,
+	Low,
+	Medium,
+	High
+};
+
+struct FAutomationComparisonToleranceAmount
+{
+public:
+
+	FAutomationComparisonToleranceAmount()
+		: Red(0)
+		, Green(0)
+		, Blue(0)
+		, Alpha(0)
+		, MinBrightness(0)
+		, MaxBrightness(255)
+	{
+	}
+
+	FAutomationComparisonToleranceAmount(uint8 R, uint8 G, uint8 B, uint8 A, uint8 InMinBrightness, uint8 InMaxBrightness)
+		: Red(R)
+		, Green(G)
+		, Blue(B)
+		, Alpha(A)
+		, MinBrightness(InMinBrightness)
+		, MaxBrightness(InMaxBrightness)
+	{
+	}
+
+	static FAutomationComparisonToleranceAmount FromToleranceLevel(EAutomationComparisonToleranceLevel InTolerance)
+	{
+		switch (InTolerance)
+		{
+		case EAutomationComparisonToleranceLevel::Low:
+			return FAutomationComparisonToleranceAmount(16, 16, 16, 16, 16, 240);
+		case EAutomationComparisonToleranceLevel::Medium:
+			return FAutomationComparisonToleranceAmount(24, 24, 24, 24, 24, 220);
+		case EAutomationComparisonToleranceLevel::High:
+			return FAutomationComparisonToleranceAmount(32, 32, 32, 32, 64, 96);
+		}
+		// Zero
+		return FAutomationComparisonToleranceAmount(0, 0, 0, 0, 0, 255);
+	}
+
+	uint8 Red;
+	uint8 Green;
+	uint8 Blue;
+	uint8 Alpha;
+	uint8 MinBrightness;
+	uint8 MaxBrightness;
+};
 
 /**
  * Delegate type for when a test screenshot has been captured
diff --git a/Engine/Source/Runtime/Engine/Private/Tests/AutomationCommon.cpp b/Engine/Source/Runtime/Engine/Private/Tests/AutomationCommon.cpp
--- a/Engine/Source/Runtime/Engine/Private/Tests/AutomationCommon.cpp
+++ b/Engine/Source/Runtime/Engine/Private/Tests/AutomationCommon.cpp
@@ -189,6 +189,77 @@ namespace AutomationCommon
 
 		return FrameTrace;
 	}
+
+	class FAutomationImageComparisonRequest : public IAutomationLatentCommand
+	{
+	public:
+		FAutomationImageComparisonRequest(const FString& InImageName, int32 InWidth, int32 InHeight, const TArray<FColor>& InImageData, const FAutomationComparisonToleranceAmount& InTolerance, const FString& InNotes)
+			: ImageName(InImageName), TaskCompleted(false)
+		{
+			FString Context;
+			if (FAutomationTestBase* CurrentTest = FAutomationTestFramework::Get().GetCurrentTest())
+			{
+				Context = CurrentTest->GetTestContext();
+				if (Context.IsEmpty()) { Context = CurrentTest->GetTestName(); }
+			}
+
+			FAutomationScreenshotData Data = BuildScreenshotData(Context, TEXT(""), ImageName, InWidth, InHeight);
+
+			// Copy the relevant data into the metadata for the screenshot.
+			Data.bHasComparisonRules = true;
+			Data.ToleranceRed = InTolerance.Red;
+			Data.ToleranceGreen = InTolerance.Green;
+			Data.ToleranceBlue = InTolerance.Blue;
+			Data.ToleranceAlpha = InTolerance.Alpha;
+			Data.ToleranceMinBrightness = InTolerance.MinBrightness;
+			Data.ToleranceMaxBrightness = InTolerance.MaxBrightness;
+			Data.bIgnoreAntiAliasing = true;
+			Data.bIgnoreColors = false;
+			Data.MaximumLocalError = 0.10f;
+			Data.MaximumGlobalError = 0.02f;
+
+			// Record any user notes that were made to accompany this shot.
+			Data.Notes = InNotes;
+
+			FAutomationTestFramework::Get().OnScreenshotCaptured().ExecuteIfBound(InImageData, Data);
+
+			UE_LOG(LogEditorAutomationTests, Log, TEXT("Requesting image %s to be compared."), *Data.ScreenshotName);
+
+			FAutomationTestFramework::Get().OnScreenshotCompared.AddRaw(this, &FAutomationImageComparisonRequest::OnComparisonComplete);
+		}
+
+		virtual ~FAutomationImageComparisonRequest()
+		{
+			FAutomationTestFramework::Get().OnScreenshotCompared.RemoveAll(this);
+		}
+
+		void OnComparisonComplete(const FAutomationScreenshotCompareResults& CompareResults)
+		{
+			FAutomationTestFramework::Get().OnScreenshotCompared.RemoveAll(this);
+
+			if (FAutomationTestBase* CurrentTest = FAutomationTestFramework::Get().GetCurrentTest())
+			{
+				CurrentTest->AddEvent(CompareResults.ToAutomationEvent(ImageName));
+			}
+
+			TaskCompleted = true;
+		}
+
+		bool IsTaskCompleted()
+		{
+			return TaskCompleted;
+		}
+
+		bool Update() override
+		{
+			return IsTaskCompleted();
+		}
+
+	private:
+		FString	ImageName;
+		bool TaskCompleted;
+	};
+
 #endif
 
 	/** These save a PNG and get sent over the network */
@@ -589,4 +660,12 @@ bool FWaitForShadersToFinishCompilingInGame::Update()
 	return true;
 }
 
+void RequestImageComparison(const FString& InImageName, int32 InWidth, int32 InHeight, const TArray<FColor>& InImageData, EAutomationComparisonToleranceLevel InTolerance, const FString& InNotes)
+{
+#if WITH_AUTOMATION_TESTS
+	FAutomationComparisonToleranceAmount ToleranceAmount = FAutomationComparisonToleranceAmount::FromToleranceLevel(InTolerance);
+	ADD_LATENT_AUTOMATION_COMMAND(AutomationCommon::FAutomationImageComparisonRequest(InImageName, InWidth, InHeight, InImageData, ToleranceAmount, InNotes));
+#endif
+}
+
 #endif //WITH_DEV_AUTOMATION_TESTS
diff --git a/Engine/Source/Runtime/Engine/Public/Tests/AutomationCommon.h b/Engine/Source/Runtime/Engine/Public/Tests/AutomationCommon.h
--- a/Engine/Source/Runtime/Engine/Public/Tests/AutomationCommon.h
+++ b/Engine/Source/Runtime/Engine/Public/Tests/AutomationCommon.h
@@ -203,4 +203,10 @@ private:
 	float MaxWaitTime;
 };
 
+/**
+* Request an Image Comparison and wait for it to complete
+**/
+ENGINE_API void RequestImageComparison(const FString& InImageName, int32 InWidth, int32 InHeight, const TArray<FColor>& InImageData, EAutomationComparisonToleranceLevel InTolerance = EAutomationComparisonToleranceLevel::Low, const FString& InNotes = TEXT(""));
+
+
 #endif
