Got rid of the explicit subtree hashing as it has been made redundant by string-based deduplication.

diff --git a/Engine/Shaders/Public/ShaderVersion.ush b/Engine/Shaders/Public/ShaderVersion.ush
--- a/Engine/Shaders/Public/ShaderVersion.ush
+++ b/Engine/Shaders/Public/ShaderVersion.ush
@@ -3,5 +3,4 @@
 // in Platform.ush (which should be included in any shader) it allows to invalidate the shader DDC.
 // 
 // If you are merging streams and there is a conflict with this GUID you should make a new GUID rather than taking one or the other.
-// GUID = 0FC91888-D8E4-4845-87E8-1165934456DA
-
+// GUID = 595DE415-EE7A-47A0-BEAD-0D1478697889
diff --git a/Engine/Source/Runtime/Engine/Private/Materials/HLSLMaterialTranslator.cpp b/Engine/Source/Runtime/Engine/Private/Materials/HLSLMaterialTranslator.cpp
--- a/Engine/Source/Runtime/Engine/Private/Materials/HLSLMaterialTranslator.cpp
+++ b/Engine/Source/Runtime/Engine/Private/Materials/HLSLMaterialTranslator.cpp
@@ -3266,16 +3266,6 @@ int32 FHLSLMaterialTranslator::AddCodeChunk(EMaterialValueType Type,const TCHAR*
 	return CodeIndex;
 }
 
-static inline uint32 GetTCharStringLength(const TCHAR* String)
-{
-	uint32 Length = 0u;
-	while (String[Length])
-	{
-		++Length;
-	}
-	return Length * sizeof(TCHAR);
-}
-
 int32 FHLSLMaterialTranslator::AddCodeChunkWithHash(uint64 BaseHash, EMaterialValueType Type, const TCHAR* Format, ...)
 {
 	int32	BufferSize = 256;
@@ -3290,7 +3280,8 @@ int32 FHLSLMaterialTranslator::AddCodeChunkWithHash(uint64 BaseHash, EMaterialVa
 	};
 	FormattedCode[Result] = 0;
 
-	const uint64 Hash = CityHash64WithSeed((char*)Format, GetTCharStringLength(Format), BaseHash);
+	uint64 Hash = CityHash64((char*)FormattedCode, Result * sizeof(TCHAR));
+	Hash = CityHash128to64({ BaseHash, Hash });
 	const int32 CodeIndex = AddCodeChunkInner(Hash, FormattedCode, Type, false);
 	FMemory::Free(FormattedCode);
 
@@ -3322,27 +3313,6 @@ int32 FHLSLMaterialTranslator::AddInlinedCodeChunk(EMaterialValueType Type, cons
 	return CodeIndex;
 }
 
-int32 FHLSLMaterialTranslator::AddInlinedCodeChunkWithHash(uint64 BaseHash, EMaterialValueType Type, const TCHAR* Format, ...)
-{
-	int32	BufferSize = 256;
-	TCHAR*	FormattedCode = NULL;
-	int32	Result = -1;
-
-	while (Result == -1)
-	{
-		FormattedCode = (TCHAR*)FMemory::Realloc(FormattedCode, BufferSize * sizeof(TCHAR));
-		GET_VARARGS_RESULT(FormattedCode, BufferSize, BufferSize - 1, Format, Format, Result);
-		BufferSize *= 2;
-	};
-	FormattedCode[Result] = 0;
-
-	const uint64 Hash = CityHash64WithSeed((char*)Format, GetTCharStringLength(Format), BaseHash);
-	const int32 CodeIndex = AddCodeChunkInner(Hash, FormattedCode, Type, true);
-	FMemory::Free(FormattedCode);
-
-	return CodeIndex;
-}
-
 int32 FHLSLMaterialTranslator::AddUniformExpressionInner(uint64 Hash, FMaterialUniformExpression* UniformExpression, EMaterialValueType Type, const TCHAR* FormattedCode)
 {
 	check(bAllowCodeChunkGeneration);
@@ -3478,26 +3448,6 @@ int32 FHLSLMaterialTranslator::AddUniformExpression(FMaterialUniformExpression*
 	return CodeIndex;
 }
 
-int32 FHLSLMaterialTranslator::AddUniformExpressionWithHash(uint64 BaseHash, FMaterialUniformExpression* UniformExpression, EMaterialValueType Type, const TCHAR* Format, ...)
-{
-	int32	BufferSize = 256;
-	TCHAR*	FormattedCode = NULL;
-	int32	Result = -1;
-
-	while (Result == -1)
-	{
-		FormattedCode = (TCHAR*)FMemory::Realloc(FormattedCode, BufferSize * sizeof(TCHAR));
-		GET_VARARGS_RESULT(FormattedCode, BufferSize, BufferSize - 1, Format, Format, Result);
-		BufferSize *= 2;
-	};
-	FormattedCode[Result] = 0;
-
-	const uint64 Hash = CityHash64WithSeed((char*)Format, GetTCharStringLength(Format), BaseHash);
-	const int32 CodeIndex = AddUniformExpressionInner(Hash, UniformExpression, Type, FormattedCode);
-	FMemory::Free(FormattedCode);
-	return CodeIndex;
-}
-
 // AccessUniformExpression - Adds code to access the value of a uniform expression to the Code array and returns its index.
 int32 FHLSLMaterialTranslator::AccessUniformExpression(int32 Index)
 {
@@ -7044,10 +6994,9 @@ int32 FHLSLMaterialTranslator::Add(int32 A,int32 B)
 		return INDEX_NONE;
 	}
 
-	const uint64 Hash = CityHash128to64({ GetParameterHash(A), GetParameterHash(B) });
 	if(GetParameterUniformExpression(A) && GetParameterUniformExpression(B))
 	{
-		return AddUniformExpressionWithHash(Hash, new FMaterialUniformExpressionFoldedMath(GetParameterUniformExpression(A),GetParameterUniformExpression(B),FMO_Add),GetArithmeticResultType(A,B),TEXT("(%s + %s)"),*GetParameterCode(A),*GetParameterCode(B));
+		return AddUniformExpression(new FMaterialUniformExpressionFoldedMath(GetParameterUniformExpression(A),GetParameterUniformExpression(B),FMO_Add),GetArithmeticResultType(A,B),TEXT("(%s + %s)"),*GetParameterCode(A),*GetParameterCode(B));
 	}
 	else
 	{
@@ -7057,7 +7006,7 @@ int32 FHLSLMaterialTranslator::Add(int32 A,int32 B)
 		}
 		else
 		{
-			return AddCodeChunkWithHash(Hash, GetArithmeticResultType(A,B),TEXT("(%s + %s)"),*GetParameterCode(A),*GetParameterCode(B));
+			return AddCodeChunk(GetArithmeticResultType(A,B),TEXT("(%s + %s)"),*GetParameterCode(A),*GetParameterCode(B));
 		}
 	}
 }
@@ -7069,10 +7018,9 @@ int32 FHLSLMaterialTranslator::Sub(int32 A,int32 B)
 		return INDEX_NONE;
 	}
 
-	const uint64 Hash = CityHash128to64({ GetParameterHash(A), GetParameterHash(B) });
 	if(GetParameterUniformExpression(A) && GetParameterUniformExpression(B))
 	{
-		return AddUniformExpressionWithHash(Hash, new FMaterialUniformExpressionFoldedMath(GetParameterUniformExpression(A),GetParameterUniformExpression(B),FMO_Sub),GetArithmeticResultType(A,B),TEXT("(%s - %s)"),*GetParameterCode(A),*GetParameterCode(B));
+		return AddUniformExpression(new FMaterialUniformExpressionFoldedMath(GetParameterUniformExpression(A),GetParameterUniformExpression(B),FMO_Sub),GetArithmeticResultType(A,B),TEXT("(%s - %s)"),*GetParameterCode(A),*GetParameterCode(B));
 	}
 	else
 	{
@@ -7082,7 +7030,7 @@ int32 FHLSLMaterialTranslator::Sub(int32 A,int32 B)
 		}
 		else
 		{
-			return AddCodeChunkWithHash(Hash, GetArithmeticResultType(A,B),TEXT("(%s - %s)"),*GetParameterCode(A),*GetParameterCode(B));
+			return AddCodeChunk(GetArithmeticResultType(A,B),TEXT("(%s - %s)"),*GetParameterCode(A),*GetParameterCode(B));
 		}
 	}
 }
@@ -7094,10 +7042,9 @@ int32 FHLSLMaterialTranslator::Mul(int32 A,int32 B)
 		return INDEX_NONE;
 	}
 
-	const uint64 Hash = CityHash128to64({ GetParameterHash(A), GetParameterHash(B) });
 	if(GetParameterUniformExpression(A) && GetParameterUniformExpression(B))
 	{
-		return AddUniformExpressionWithHash(Hash, new FMaterialUniformExpressionFoldedMath(GetParameterUniformExpression(A),GetParameterUniformExpression(B),FMO_Mul),GetArithmeticResultType(A,B),TEXT("(%s * %s)"),*GetParameterCode(A),*GetParameterCode(B));
+		return AddUniformExpression(new FMaterialUniformExpressionFoldedMath(GetParameterUniformExpression(A),GetParameterUniformExpression(B),FMO_Mul),GetArithmeticResultType(A,B),TEXT("(%s * %s)"),*GetParameterCode(A),*GetParameterCode(B));
 	}
 	else
 	{
@@ -7107,7 +7054,7 @@ int32 FHLSLMaterialTranslator::Mul(int32 A,int32 B)
 		}
 		else
 		{
-			return AddCodeChunkWithHash(Hash, GetArithmeticResultType(A,B),TEXT("(%s * %s)"),*GetParameterCode(A),*GetParameterCode(B));
+			return AddCodeChunk(GetArithmeticResultType(A,B),TEXT("(%s * %s)"),*GetParameterCode(A),*GetParameterCode(B));
 		}
 	}
 }
@@ -7119,10 +7066,9 @@ int32 FHLSLMaterialTranslator::Div(int32 A,int32 B)
 		return INDEX_NONE;
 	}
 
-	const uint64 Hash = CityHash128to64({ GetParameterHash(A), GetParameterHash(B) });
 	if(GetParameterUniformExpression(A) && GetParameterUniformExpression(B))
 	{
-		return AddUniformExpressionWithHash(Hash, new FMaterialUniformExpressionFoldedMath(GetParameterUniformExpression(A),GetParameterUniformExpression(B),FMO_Div),GetArithmeticResultType(A,B),TEXT("(%s / %s)"),*GetParameterCode(A),*GetParameterCode(B));
+		return AddUniformExpression(new FMaterialUniformExpressionFoldedMath(GetParameterUniformExpression(A),GetParameterUniformExpression(B),FMO_Div),GetArithmeticResultType(A,B),TEXT("(%s / %s)"),*GetParameterCode(A),*GetParameterCode(B));
 	}
 	else
 	{
@@ -7132,7 +7078,7 @@ int32 FHLSLMaterialTranslator::Div(int32 A,int32 B)
 		}
 		else
 		{
-			return AddCodeChunkWithHash(Hash, GetArithmeticResultType(A,B),TEXT("(%s / %s)"),*GetParameterCode(A),*GetParameterCode(B));
+			return AddCodeChunk(GetArithmeticResultType(A,B),TEXT("(%s / %s)"),*GetParameterCode(A),*GetParameterCode(B));
 		}
 	}
 }
diff --git a/Engine/Source/Runtime/Engine/Private/Materials/HLSLMaterialTranslator.h b/Engine/Source/Runtime/Engine/Private/Materials/HLSLMaterialTranslator.h
--- a/Engine/Source/Runtime/Engine/Private/Materials/HLSLMaterialTranslator.h
+++ b/Engine/Source/Runtime/Engine/Private/Materials/HLSLMaterialTranslator.h
@@ -603,13 +603,11 @@ protected:
 	 * This should be used instead of AddCodeChunk when the code chunk does not add any actual shader instructions, for example a component mask.
 	 */
 	int32 AddInlinedCodeChunk(EMaterialValueType Type, const TCHAR* Format, ...);
-	int32 AddInlinedCodeChunkWithHash(uint64 BaseHash, EMaterialValueType Type, const TCHAR* Format, ...);
 
 	int32 AddUniformExpressionInner(uint64 Hash, FMaterialUniformExpression* UniformExpression, EMaterialValueType Type, const TCHAR* FormattedCode);
 
 	// AddUniformExpression - Adds an input to the Code array and returns its index.
 	int32 AddUniformExpression(FMaterialUniformExpression* UniformExpression, EMaterialValueType Type, const TCHAR* Format, ...);
-	int32 AddUniformExpressionWithHash(uint64 BaseHash, FMaterialUniformExpression* UniformExpression, EMaterialValueType Type, const TCHAR* Format, ...);
 
 	// AccessUniformExpression - Adds code to access the value of a uniform expression to the Code array and returns its index.
 	int32 AccessUniformExpression(int32 Index);
