WorldPartition: Move some BeginDestroy code up into UActorDescContainer & call Uninitialize in BeginDestroy

diff --git a/Engine/Source/Runtime/Engine/Private/WorldPartition/ActorDescContainer.cpp b/Engine/Source/Runtime/Engine/Private/WorldPartition/ActorDescContainer.cpp
--- a/Engine/Source/Runtime/Engine/Private/WorldPartition/ActorDescContainer.cpp
+++ b/Engine/Source/Runtime/Engine/Private/WorldPartition/ActorDescContainer.cpp
@@ -72,6 +72,19 @@ void UActorDescContainer::Uninitialize()
 #endif
 }
 
+void UActorDescContainer::BeginDestroy()
+{
+	Super::BeginDestroy();
+
+	Uninitialize();
+#if WITH_EDITOR
+	for (TUniquePtr<FWorldPartitionActorDesc>& ActorDescPtr : ActorDescList)
+	{
+		ActorDescPtr.Release();
+	}
+#endif
+}
+
 #if WITH_EDITOR
 TUniquePtr<FWorldPartitionActorDesc> UActorDescContainer::GetActorDescriptor(const FAssetData& InAssetData)
 {
diff --git a/Engine/Source/Runtime/Engine/Private/WorldPartition/WorldPartition.cpp b/Engine/Source/Runtime/Engine/Private/WorldPartition/WorldPartition.cpp
--- a/Engine/Source/Runtime/Engine/Private/WorldPartition/WorldPartition.cpp
+++ b/Engine/Source/Runtime/Engine/Private/WorldPartition/WorldPartition.cpp
@@ -802,11 +802,6 @@ void UWorldPartition::BeginDestroy()
 
 #if WITH_EDITOR
 	LoadedSubobjects.Empty();
-
-	for (TUniquePtr<FWorldPartitionActorDesc>& ActorDescPtr : ActorDescList)
-	{
-		ActorDescPtr.Release();
-	}
 #endif
 }
 
diff --git a/Engine/Source/Runtime/Engine/Public/WorldPartition/ActorDescContainer.h b/Engine/Source/Runtime/Engine/Public/WorldPartition/ActorDescContainer.h
--- a/Engine/Source/Runtime/Engine/Public/WorldPartition/ActorDescContainer.h
+++ b/Engine/Source/Runtime/Engine/Public/WorldPartition/ActorDescContainer.h
@@ -31,6 +31,10 @@ public:
 #endif
 
 protected:
+	//~ Begin UObject Interface
+	virtual void BeginDestroy() override;
+	//~ End UObject Interface
+
 #if WITH_EDITOR
 	TUniquePtr<FWorldPartitionActorDesc> GetActorDescriptor(const FAssetData& InAssetData);
 	
