Prevent dirtying actor packages while loading & registering external actors

diff --git a/Engine/Source/Runtime/Engine/Private/WorldPartition/WorldPartition.cpp b/Engine/Source/Runtime/Engine/Private/WorldPartition/WorldPartition.cpp
--- a/Engine/Source/Runtime/Engine/Private/WorldPartition/WorldPartition.cpp
+++ b/Engine/Source/Runtime/Engine/Private/WorldPartition/WorldPartition.cpp
@@ -677,8 +677,6 @@ void UWorldPartition::UpdateLoadingEditorCell(UWorldPartitionEditorCell* Cell, b
 	}
 	else
 	{
-		TGuardValue<bool> IsEditorLoadingPackageGuard(GIsEditorLoadingPackage, true);
-
 		for (FWorldPartitionHandle& ActorHandle: Cell->Actors)
 		{
 			AActor* Actor = ActorHandle->GetActor();
diff --git a/Engine/Source/Runtime/Engine/Private/WorldPartition/WorldPartitionHandle.cpp b/Engine/Source/Runtime/Engine/Private/WorldPartition/WorldPartitionHandle.cpp
--- a/Engine/Source/Runtime/Engine/Private/WorldPartition/WorldPartitionHandle.cpp
+++ b/Engine/Source/Runtime/Engine/Private/WorldPartition/WorldPartitionHandle.cpp
@@ -34,6 +34,7 @@ void FWorldPartitionReferenceImpl::IncRefCount(FWorldPartitionActorDesc* ActorDe
 {
 	if (ActorDesc->IncHardRefCount() == 1)
 	{
+		TGuardValue<bool> IsEditorLoadingPackageGuard(GIsEditorLoadingPackage, true);
 		ActorDesc->Load();
 		ActorDesc->RegisterActor();
 	}
