Checked nullptr when creating MeshDescriptionBulkData editor-only subobject. This may in certain circumstances legitimately yield null.

diff --git a/Engine/Source/Runtime/Engine/Private/StaticMesh.cpp b/Engine/Source/Runtime/Engine/Private/StaticMesh.cpp
--- a/Engine/Source/Runtime/Engine/Private/StaticMesh.cpp
+++ b/Engine/Source/Runtime/Engine/Private/StaticMesh.cpp
@@ -2898,7 +2898,10 @@ UStaticMesh::UStaticMesh(const FObjectInitializer& ObjectInitializer)
 	LODForOccluderMesh = -1;
 	NumStreamedLODs.Default = -1;
 	GetHiResSourceModel().StaticMeshDescriptionBulkData = CreateEditorOnlyDefaultSubobject<UStaticMeshDescriptionBulkData>(TEXT("HiResMeshDescription"));
-	GetHiResSourceModel().StaticMeshDescriptionBulkData->SetFlags(RF_Transactional);
+	if (GetHiResSourceModel().StaticMeshDescriptionBulkData)
+	{
+		GetHiResSourceModel().StaticMeshDescriptionBulkData->SetFlags(RF_Transactional);
+	}
 #endif // #if WITH_EDITORONLY_DATA
 	SetLightMapResolution(4);
 	SetMinLOD(0);
