Fixes Log Warnings about missing ViewportEditor Brushes

diff --git a/Engine/Source/Runtime/Engine/Private/ViewModeNames.cpp b/Engine/Source/Runtime/Engine/Private/ViewModeNames.cpp
--- a/Engine/Source/Runtime/Engine/Private/ViewModeNames.cpp
+++ b/Engine/Source/Runtime/Engine/Private/ViewModeNames.cpp
@@ -237,7 +237,7 @@ TArray<const FSlateBrush*> FillViewModeDisplayIcons()
 		// Colored according to light count - showing lightmap texel density on texture mapped objects
 		else if (ViewModeIndex == VMI_LitLightmapDensity)
 		{
-			ViewModeDisplayIcons.Emplace(FAppStyle::Get().GetBrush("EditorViewport.LitLightmapDensityMode"));
+			ViewModeDisplayIcons.Emplace(FAppStyle::Get().GetBrush("EditorViewport.LightmapDensityMode"));
 		}
 		else if (ViewModeIndex == VMI_ReflectionOverride)
 		{
@@ -325,7 +325,7 @@ TArray<const FSlateBrush*> FillViewModeDisplayIcons()
 		// VMI_Max
 		else if (ViewModeIndex == VMI_Max)
 		{
-			ViewModeDisplayIcons.Emplace(FAppStyle::Get().GetBrush("EditorViewport.MaxMode"));
+			ViewModeDisplayIcons.Emplace(FAppStyle::Get().GetBrush("EditorViewport.LitMode"));
 		}
 		// VMI_Unknown
 		else 
