Defer deletion of the old global SM until new one is complete (UE-108443).

diff --git a/Engine/Source/Runtime/Engine/Private/ShaderCompiler/ShaderCompiler.cpp b/Engine/Source/Runtime/Engine/Private/ShaderCompiler/ShaderCompiler.cpp
--- a/Engine/Source/Runtime/Engine/Private/ShaderCompiler/ShaderCompiler.cpp
+++ b/Engine/Source/Runtime/Engine/Private/ShaderCompiler/ShaderCompiler.cpp
@@ -113,6 +113,12 @@ static FAutoConsoleVariableRef CVarShaderCompilerCacheStatsPrintoutInterval(
 	ECVF_Default
 );
 
+/** Storage for the global shadar map(s) that have been replaced by new one(s), which aren't yet compiled.
+ * 
+ *	Sometimes a mesh drawing command references a pointer to global SM's memory. To nix these MDCs when we're replacing a global SM, we would just recreate the render state for all the components, but
+ *	we may need to access a global shader during such an update, creating a catch 22. So deleting the global SM and updating components is deferred until the new one is compiled. 
+ */
+FGlobalShaderMap* GGlobalShaderMap_DeferredDeleteCopy[SP_NumPlatforms] = {nullptr};
 
 #if ENABLE_COOK_STATS
 namespace GlobalShaderCookStats
@@ -5744,8 +5750,9 @@ void CompileGlobalShaderMap(EShaderPlatform Platform, const ITargetPlatform* Tar
 
 	if (bRefreshShaderMap || GGlobalShaderTargetPlatform[Platform] != TargetPlatform)
 	{
-		// delete the current global shader map
-		delete GGlobalShaderMap[Platform];
+		// defer the deletion the current global shader map, delete the previous one if it is still valid
+		delete GGlobalShaderMap_DeferredDeleteCopy[Platform];	// deleting null is Okay
+		GGlobalShaderMap_DeferredDeleteCopy[Platform] = GGlobalShaderMap[Platform];
 		GGlobalShaderMap[Platform] = nullptr;
 
 		GGlobalShaderTargetPlatform[Platform] = TargetPlatform;
@@ -5893,8 +5900,6 @@ void CompileGlobalShaderMap(EShaderPlatform Platform, const ITargetPlatform* Tar
 		{
 			GGlobalShaderMap[Platform]->BeginCreateAllShaders();
 		}
-
-		PropagateGlobalShadersToAllPrimitives();
 	}
 }
 
@@ -6292,6 +6297,10 @@ void ProcessCompiledGlobalShaders(const TArray<FShaderCommonCompileJobPtr>& Comp
 					GlobalShaderMap->FindOrAddShaderPipeline(ShaderPipelineType, ShaderPipeline);
 				}
 			}
+
+			// at this point the new global sm is populated and we can delete the deferred copy, if any
+			delete GGlobalShaderMap_DeferredDeleteCopy[ShaderPlatformsProcessed[PlatformIndex]];	// even if it was nullptr, deleting null is Okay
+			GGlobalShaderMap_DeferredDeleteCopy[ShaderPlatformsProcessed[PlatformIndex]] = nullptr;
 		}
 
 		// Save the global shader map for any platforms that were recompiled
