Fix crash when rebuilding a visible skeletal mesh because of RecreateRenderStateContext destructor being called during what should have been a simple assignment

diff --git a/Engine/Source/Runtime/Engine/Private/SkeletalMesh.cpp b/Engine/Source/Runtime/Engine/Private/SkeletalMesh.cpp
--- a/Engine/Source/Runtime/Engine/Private/SkeletalMesh.cpp
+++ b/Engine/Source/Runtime/Engine/Private/SkeletalMesh.cpp
@@ -1981,7 +1981,7 @@ void USkeletalMesh::BeginBuildInternal(FSkeletalMeshBuildContext& Context)
 	TRACE_CPUPROFILER_EVENT_SCOPE(USkeletalMesh::BeginBuildInternal);
 
 	// Unregister all instances of this component
-	Context.RecreateRenderStateContext = FSkinnedMeshComponentRecreateRenderStateContext(this, false);
+	Context.RecreateRenderStateContext.Emplace(this, false);
 
 	// Release the static mesh's resources.
 	ReleaseResources();
@@ -6530,9 +6530,9 @@ FSkinnedMeshComponentRecreateRenderStateContext::FSkinnedMeshComponentRecreateRe
 				It->DestroyRenderState_Concurrent();
 				bRequireFlushRenderingCommands = true;
 			}
-			MeshComponents.Add(*It);
+				MeshComponents.Add(*It);
+			}
 		}
-	}
 
 	// Flush the rendering commands generated by the detachments.
 	// The static mesh scene proxies reference the UStaticMesh, and this ensures that they are cleaned up before the UStaticMesh changes.
