Fixed linker error due to FSkeletalMeshAsyncBuildScope::SkeletalMeshBeingAsyncCompiled.

FSkeletalMeshAsyncBuildScope can't be inline as SkeletalMeshBeingAsyncCompiled is a thread-local, which can't be exported from a dynamic library.

diff --git a/Engine/Source/Runtime/Engine/Classes/Engine/SkeletalMesh.h b/Engine/Source/Runtime/Engine/Classes/Engine/SkeletalMesh.h
--- a/Engine/Source/Runtime/Engine/Classes/Engine/SkeletalMesh.h
+++ b/Engine/Source/Runtime/Engine/Classes/Engine/SkeletalMesh.h
@@ -140,22 +140,9 @@ public:
 class FSkeletalMeshAsyncBuildScope
 {
 public:
-	ENGINE_API FSkeletalMeshAsyncBuildScope(const USkeletalMesh* SkeletalMesh)
-	{
-		PreviousScope = SkeletalMeshBeingAsyncCompiled;
-		SkeletalMeshBeingAsyncCompiled = SkeletalMesh;
-	}
-
-	ENGINE_API ~FSkeletalMeshAsyncBuildScope()
-	{
-		check(SkeletalMeshBeingAsyncCompiled);
-		SkeletalMeshBeingAsyncCompiled = PreviousScope;
-	}
-
-	ENGINE_API static bool ShouldWaitOnLockedProperties(const USkeletalMesh* SkeletalMesh)
-	{
-		return SkeletalMeshBeingAsyncCompiled != SkeletalMesh;
-	}
+	ENGINE_API FSkeletalMeshAsyncBuildScope(const USkeletalMesh* SkeletalMesh);
+	ENGINE_API ~FSkeletalMeshAsyncBuildScope();
+	ENGINE_API static bool ShouldWaitOnLockedProperties(const USkeletalMesh* SkeletalMesh);
 
 private:
 	const USkeletalMesh* PreviousScope = nullptr;
diff --git a/Engine/Source/Runtime/Engine/Private/SkeletalMesh.cpp b/Engine/Source/Runtime/Engine/Private/SkeletalMesh.cpp
--- a/Engine/Source/Runtime/Engine/Private/SkeletalMesh.cpp
+++ b/Engine/Source/Runtime/Engine/Private/SkeletalMesh.cpp
@@ -1878,6 +1878,23 @@ uint32 USkeletalMesh::GetVertexBufferFlags() const
 
 thread_local const USkeletalMesh* FSkeletalMeshAsyncBuildScope::SkeletalMeshBeingAsyncCompiled = nullptr;
 
+FSkeletalMeshAsyncBuildScope::FSkeletalMeshAsyncBuildScope(const USkeletalMesh* SkeletalMesh)
+{
+	PreviousScope = SkeletalMeshBeingAsyncCompiled;
+	SkeletalMeshBeingAsyncCompiled = SkeletalMesh;
+}
+
+FSkeletalMeshAsyncBuildScope::~FSkeletalMeshAsyncBuildScope()
+{
+	check(SkeletalMeshBeingAsyncCompiled);
+	SkeletalMeshBeingAsyncCompiled = PreviousScope;
+}
+
+bool FSkeletalMeshAsyncBuildScope::ShouldWaitOnLockedProperties(const USkeletalMesh* SkeletalMesh)
+{
+	return SkeletalMeshBeingAsyncCompiled != SkeletalMesh;
+}
+
 void USkeletalMesh::WaitUntilAsyncPropertyReleased(ESkeletalMeshAsyncProperties AsyncProperties, EAsyncPropertyLockType LockType) const
 {
 	// We need to protect internal skeletal mesh data from race conditions during async build
