LevelInstance: Fix crash on null Actor when pushing selection state to proxies

diff --git a/Engine/Source/Runtime/Engine/Private/LevelInstance/LevelInstanceSubsystem.cpp b/Engine/Source/Runtime/Engine/Private/LevelInstance/LevelInstanceSubsystem.cpp
--- a/Engine/Source/Runtime/Engine/Private/LevelInstance/LevelInstanceSubsystem.cpp
+++ b/Engine/Source/Runtime/Engine/Private/LevelInstance/LevelInstanceSubsystem.cpp
@@ -1417,7 +1417,10 @@ void ULevelInstanceSubsystem::EditLevelInstance(ALevelInstance* LevelInstanceAct
 	for (const auto& Actor : LevelStreaming->LoadedLevel->Actors)
 	{
 		const bool bEditing = true;
-		Actor->PushLevelInstanceEditingStateToProxies(bEditing);
+		if (Actor)
+		{
+			Actor->PushLevelInstanceEditingStateToProxies(bEditing);
+		}
 	}
 
 	const bool bHasLevelInstanceEdits = LevelInstanceEdits.Num() > 0;
