Declared UMeshDescriptionBaseBulkData an editor-only class by implementing NeedsLoadForClient/Server() returning false. This prevents it from being cooked.
Made the HiResSourceModel MeshDescriptionBulkData a regular default subobject, so that it is still valid in Editor -game builds.

diff --git a/Engine/Source/Runtime/Engine/Private/StaticMesh.cpp b/Engine/Source/Runtime/Engine/Private/StaticMesh.cpp
--- a/Engine/Source/Runtime/Engine/Private/StaticMesh.cpp
+++ b/Engine/Source/Runtime/Engine/Private/StaticMesh.cpp
@@ -2897,11 +2897,8 @@ UStaticMesh::UStaticMesh(const FObjectInitializer& ObjectInitializer)
 	ImportVersion = EImportStaticMeshVersion::BeforeImportStaticMeshVersionWasAdded;
 	LODForOccluderMesh = -1;
 	NumStreamedLODs.Default = -1;
-	GetHiResSourceModel().StaticMeshDescriptionBulkData = CreateEditorOnlyDefaultSubobject<UStaticMeshDescriptionBulkData>(TEXT("HiResMeshDescription"));
-	if (GetHiResSourceModel().StaticMeshDescriptionBulkData)
-	{
-		GetHiResSourceModel().StaticMeshDescriptionBulkData->SetFlags(RF_Transactional);
-	}
+	GetHiResSourceModel().StaticMeshDescriptionBulkData = CreateDefaultSubobject<UStaticMeshDescriptionBulkData>(TEXT("HiResMeshDescription"));
+	GetHiResSourceModel().StaticMeshDescriptionBulkData->SetFlags(RF_Transactional);
 #endif // #if WITH_EDITORONLY_DATA
 	SetLightMapResolution(4);
 	SetMinLOD(0);
diff --git a/Engine/Source/Runtime/MeshDescription/Public/MeshDescriptionBaseBulkData.h b/Engine/Source/Runtime/MeshDescription/Public/MeshDescriptionBaseBulkData.h
--- a/Engine/Source/Runtime/MeshDescription/Public/MeshDescriptionBaseBulkData.h
+++ b/Engine/Source/Runtime/MeshDescription/Public/MeshDescriptionBaseBulkData.h
@@ -25,6 +25,8 @@ class MESHDESCRIPTION_API UMeshDescriptionBaseBulkData : public UObject
 public:
 	UMeshDescriptionBaseBulkData();
 	virtual void Serialize(FArchive& Ar) override;
+	virtual bool NeedsLoadForClient() const override { return false; }
+	virtual bool NeedsLoadForServer() const override { return false; }
 	virtual UClass* GetMeshDescriptionType() const;
 
 #if WITH_EDITORONLY_DATA
