Fixed crash when a level instance actor refers to an invalid asset.
Added log error+ map check when that is the case

diff --git a/Engine/Source/Runtime/Engine/Private/LevelInstance/LevelInstanceActor.cpp b/Engine/Source/Runtime/Engine/Private/LevelInstance/LevelInstanceActor.cpp
--- a/Engine/Source/Runtime/Engine/Private/LevelInstance/LevelInstanceActor.cpp
+++ b/Engine/Source/Runtime/Engine/Private/LevelInstance/LevelInstanceActor.cpp
@@ -322,6 +322,17 @@ void ALevelInstance::CheckForErrors()
 
 		Error->AddToken(FMapErrorToken::Create(FName(TEXT("LevelInstanceActor_Loop_CheckForErrors"))));
 	}
+
+	FPackagePath WorldAssetPath;
+	if (!FPackagePath::TryFromPackageName(GetWorldAssetPackage(), WorldAssetPath) || !FPackageName::DoesPackageExist(WorldAssetPath))
+	{
+		TSharedRef<FTokenizedMessage> Error = FMessageLog("MapCheck").Error()
+			->AddToken(FTextToken::Create(LOCTEXT("LevelInstanceActor_InvalidPackage", "LevelInstance actor")))
+			->AddToken(FUObjectToken::Create(this))
+			->AddToken(FTextToken::Create(FText::FromString(TEXT("refers to an invalid asset:"))))
+			->AddToken(FAssetNameToken::Create(GetWorldAsset().GetLongPackageName(), FText::FromString(GetWorldAsset().GetLongPackageName())))
+			->AddToken(FMapErrorToken::Create(FName(TEXT("LevelInstanceActor_InvalidPackage_CheckForErrors"))));
+	}
 }
 
 bool ALevelInstance::CheckForLoop(TSoftObjectPtr<UWorld> InLevelInstance, TArray<TPair<FText,TSoftObjectPtr<UWorld>>>* LoopInfo, const ALevelInstance** LoopStart) const
diff --git a/Engine/Source/Runtime/Engine/Private/LevelInstance/LevelInstanceLevelStreaming.cpp b/Engine/Source/Runtime/Engine/Private/LevelInstance/LevelInstanceLevelStreaming.cpp
--- a/Engine/Source/Runtime/Engine/Private/LevelInstance/LevelInstanceLevelStreaming.cpp
+++ b/Engine/Source/Runtime/Engine/Private/LevelInstance/LevelInstanceLevelStreaming.cpp
@@ -51,6 +51,13 @@ ULevelStreamingLevelInstance* ULevelStreamingLevelInstance::LoadInstance(ALevelI
 		UE_LOG(LogLevelInstance, Error, TEXT("Failed to load LevelInstance Actor '%s' because that would cause a loop. Run Map Check for more details."), *LevelInstanceActor->GetPathName());
 		return nullptr;
 	}
+
+	FPackagePath WorldAssetPath;
+	if (!FPackagePath::TryFromPackageName(LevelInstanceActor->GetWorldAssetPackage(), WorldAssetPath) || !FPackageName::DoesPackageExist(WorldAssetPath))
+	{
+		UE_LOG(LogLevelInstance, Error, TEXT("Failed to load LevelInstance Actor '%s' because it refers to an invalid package ('%s'). Run Map Check for more details."), *LevelInstanceActor->GetPathName(), *LevelInstanceActor->GetWorldAsset().GetLongPackageName());
+		return nullptr;
+	}
 #endif
 
 	bool bOutSuccess = false;
