WorldPartition: Actor desc leak fix + Do the ptr reset in Uninitialize for simmetry

diff --git a/Engine/Source/Runtime/Engine/Private/WorldPartition/ActorDescContainer.cpp b/Engine/Source/Runtime/Engine/Private/WorldPartition/ActorDescContainer.cpp
--- a/Engine/Source/Runtime/Engine/Private/WorldPartition/ActorDescContainer.cpp
+++ b/Engine/Source/Runtime/Engine/Private/WorldPartition/ActorDescContainer.cpp
@@ -69,6 +69,11 @@ void UActorDescContainer::Uninitialize()
 		UnregisterDelegates();
 		bContainerInitialized = false;
 	}
+
+	for (TUniquePtr<FWorldPartitionActorDesc>& ActorDescPtr : ActorDescList)
+	{
+		ActorDescPtr.Reset();
+	}
 #endif
 }
 
@@ -77,12 +82,6 @@ void UActorDescContainer::BeginDestroy()
 	Super::BeginDestroy();
 
 	Uninitialize();
-#if WITH_EDITOR
-	for (TUniquePtr<FWorldPartitionActorDesc>& ActorDescPtr : ActorDescList)
-	{
-		ActorDescPtr.Release();
-	}
-#endif
 }
 
 #if WITH_EDITOR
