Correcting typo in pre-declare

diff --git a/Engine/Source/Runtime/Engine/Classes/Animation/AnimData/AnimDataNotifyCollector.h b/Engine/Source/Runtime/Engine/Classes/Animation/AnimData/AnimDataNotifyCollector.h
--- a/Engine/Source/Runtime/Engine/Classes/Animation/AnimData/AnimDataNotifyCollector.h
+++ b/Engine/Source/Runtime/Engine/Classes/Animation/AnimData/AnimDataNotifyCollector.h
@@ -6,7 +6,7 @@
 #include "Containers/Set.h"
 #include "Containers/Array.h"
 
-enum class EAnimDataModelNotifType : uint8;
+enum class EAnimDataModelNotifyType : uint8;
 
 namespace UE {
 namespace Anim {
