Store last loaded editor cells by name instead of center.

diff --git a/Engine/Source/Runtime/Engine/Private/WorldPartition/WorldPartition.cpp b/Engine/Source/Runtime/Engine/Private/WorldPartition/WorldPartition.cpp
--- a/Engine/Source/Runtime/Engine/Private/WorldPartition/WorldPartition.cpp
+++ b/Engine/Source/Runtime/Engine/Private/WorldPartition/WorldPartition.cpp
@@ -89,12 +89,12 @@ struct FWorldPartionCellUpdateContext
 				// Save last loaded cells settings (ignore while running commandlets)
 				if (!IsRunningCommandlet())
 				{
-					TArray<FVector> EditorGridLastLoadedCells;
+					TArray<FName> EditorGridLastLoadedCells;
 					WorldPartition->EditorHash->ForEachCell([this, &EditorGridLastLoadedCells](UWorldPartitionEditorCell* Cell)
 					{
 						if ((Cell != WorldPartition->EditorHash->GetAlwaysLoadedCell()) && Cell->bLoaded)
 						{
-							EditorGridLastLoadedCells.Add(Cell->Bounds.GetCenter());
+							EditorGridLastLoadedCells.Add(Cell->GetFName());
 						}
 					});
 
@@ -282,14 +282,14 @@ void UWorldPartition::Initialize(UWorld* InWorld, const FTransform& InTransform)
 				// Load last loaded cells
 				if (WorldPartitionEditorModule.GetEnableLoadingOfLastLoadedCells())
 				{
-					const TArray<FVector>& EditorGridLastLoadedCells = GetMutableDefault<UWorldPartitionEditorPerProjectUserSettings>()->GetEditorGridLoadedCells(InWorld);
+					const TArray<FName>& EditorGridLastLoadedCells = GetMutableDefault<UWorldPartitionEditorPerProjectUserSettings>()->GetEditorGridLoadedCells(InWorld);
 
-					for (const FVector& EditorGridLastLoadedCellCenter : EditorGridLastLoadedCells)
+					for (FName EditorGridLastLoadedCell : EditorGridLastLoadedCells)
 					{
-						EditorHash->ForEachIntersectingCell(FBox(&EditorGridLastLoadedCellCenter, 1), [this](UWorldPartitionEditorCell* Cell)
+						if (UWorldPartitionEditorCell* Cell = FindObject<UWorldPartitionEditorCell>(EditorHash, *EditorGridLastLoadedCell.ToString()))
 						{
 							UpdateLoadingEditorCell(Cell, true);
-						});
+						}
 					}
 				}
 			}
diff --git a/Engine/Source/Runtime/Engine/Private/WorldPartition/WorldPartitionEditorSpatialHash.cpp b/Engine/Source/Runtime/Engine/Private/WorldPartition/WorldPartitionEditorSpatialHash.cpp
--- a/Engine/Source/Runtime/Engine/Private/WorldPartition/WorldPartitionEditorSpatialHash.cpp
+++ b/Engine/Source/Runtime/Engine/Private/WorldPartition/WorldPartitionEditorSpatialHash.cpp
@@ -30,7 +30,8 @@ void UWorldPartitionEditorSpatialHash::Initialize()
 	AlwaysLoadedCell->Bounds.Init();
 
 	FHashBuilder EditorGridConfigHashBuilder;
-	EditorGridConfigHashBuilder << CellSize;
+	static uint32 Version = 1;
+	EditorGridConfigHashBuilder << Version << CellSize;	
 	GetMutableDefault<UWorldPartitionEditorPerProjectUserSettings>()->SetEditorGridConfigHash(EditorGridConfigHashBuilder.GetHash());
 }
 
@@ -162,7 +163,7 @@ void UWorldPartitionEditorSpatialHash::HashActor(FWorldPartitionHandle& InActorH
 			}
 			else
 			{
-				EditorCell = NewObject<UWorldPartitionEditorCell>(this, *FString::Printf(TEXT("EditorCell_X%lld_Y%lld_Z%lld_L%d"), CellCoord.X, CellCoord.Y, CellCoord.Z, CellCoord.Level), RF_Transient);
+				EditorCell = NewObject<UWorldPartitionEditorCell>(this, *FString::Printf(TEXT("EditorCell_X%lld_Y%lld_Z%lld"), CellCoord.X, CellCoord.Y, CellCoord.Z), RF_Transient);
 				EditorCell->Bounds = GetCellBounds(CellCoord);
 
 				Cells.Add(EditorCell);
diff --git a/Engine/Source/Runtime/Engine/Public/WorldPartition/WorldPartitionEditorPerProjectUserSettings.h b/Engine/Source/Runtime/Engine/Public/WorldPartition/WorldPartitionEditorPerProjectUserSettings.h
--- a/Engine/Source/Runtime/Engine/Public/WorldPartition/WorldPartitionEditorPerProjectUserSettings.h
+++ b/Engine/Source/Runtime/Engine/Public/WorldPartition/WorldPartitionEditorPerProjectUserSettings.h
@@ -18,14 +18,14 @@ struct FWorldPartitionPerWorldSettings
 	FWorldPartitionPerWorldSettings()
 	{}
 
-	FWorldPartitionPerWorldSettings(TArray<FVector>& InLoadedEditorGridCells)
+	FWorldPartitionPerWorldSettings(TArray<FName>& InLoadedEditorGridCells)
 		: LoadedEditorGridCells(InLoadedEditorGridCells)
 	{}
 #endif
 
 #if WITH_EDITORONLY_DATA
 	UPROPERTY()
-	TArray<FVector> LoadedEditorGridCells;
+	TArray<FName> LoadedEditorGridCells;
 
 	UPROPERTY()
 	TArray<FName> NotLoadedDataLayers;
@@ -58,14 +58,14 @@ public:
 		}
 	}
 
-	const TArray<FVector>& GetEditorGridLoadedCells(UWorld* InWorld)
+	const TArray<FName>& GetEditorGridLoadedCells(UWorld* InWorld)
 	{ 
 		return PerWorldEditorSettings.FindOrAdd(TSoftObjectPtr<UWorld>(InWorld)).LoadedEditorGridCells;
 	}
 
-	void SetEditorGridLoadedCells(UWorld* InWorld, TArray<FVector>& InEditorGridLoadedCells)
+	void SetEditorGridLoadedCells(UWorld* InWorld, TArray<FName>& InEditorGridLoadedCells)
 	{
-		TArray<FVector>& GridLoadedCells = PerWorldEditorSettings.FindOrAdd(TSoftObjectPtr<UWorld>(InWorld)).LoadedEditorGridCells;
+		TArray<FName>& GridLoadedCells = PerWorldEditorSettings.FindOrAdd(TSoftObjectPtr<UWorld>(InWorld)).LoadedEditorGridCells;
 		GridLoadedCells = InEditorGridLoadedCells;
 		SaveConfig();
 	}
