Fixed non-unity build.

diff --git a/Engine/Source/Runtime/Engine/Private/StaticMeshSourceData.cpp b/Engine/Source/Runtime/Engine/Private/StaticMeshSourceData.cpp
--- a/Engine/Source/Runtime/Engine/Private/StaticMeshSourceData.cpp
+++ b/Engine/Source/Runtime/Engine/Private/StaticMeshSourceData.cpp
@@ -1,6 +1,7 @@
 // Copyright Epic Games, Inc. All Rights Reserved.
 
 #include "Engine/StaticMeshSourceData.h"
+#include "Engine/StaticMesh.h"
 #include "MeshDescription.h"
 #include "StaticMeshDescription.h"
 #include "StaticMeshOperations.h"
