Changed how WorldPartition handles AlwaysLoaded actors in PIE : They now go through the same path as non-partitioned world (i.e. they are duplicated when duplicating world for PIE).

diff --git a/Engine/Source/Runtime/Engine/Classes/GameFramework/Actor.h b/Engine/Source/Runtime/Engine/Classes/GameFramework/Actor.h
--- a/Engine/Source/Runtime/Engine/Classes/GameFramework/Actor.h
+++ b/Engine/Source/Runtime/Engine/Classes/GameFramework/Actor.h
@@ -920,6 +920,9 @@ private:
 	/** Whether this actor is temporarily hidden within the editor; used for show/hide/etc functionality w/o dirtying the actor. */
 	UPROPERTY(Transient)
 	uint8 bHiddenEdTemporary:1;
+
+	UPROPERTY(Transient)
+	uint8 bForceExternalActorLevelReferenceForPIE : 1;
 #endif // WITH_EDITORONLY_DATA
 
 	/** Set while actor is being constructed. Used to ensure that construction is not re-entrant. */
@@ -2087,6 +2090,20 @@ public:
 	UFUNCTION(BlueprintCallable, Category="Editor Scripting | Actor Editing")
 	virtual void SetIsTemporarilyHiddenInEditor( bool bIsHidden );
 
+	/** Returns if level should keep a reference to the external actor for PIE (used for always loaded actors). */
+	bool IsForceExternalActorLevelReferenceForPIE() const
+	{
+		return bForceExternalActorLevelReferenceForPIE;
+	}
+
+	void SetForceExternalActorLevelReferenceForPIE(bool bValue)
+	{
+		if (ensure(IsPackageExternal()))
+		{
+			bForceExternalActorLevelReferenceForPIE = bValue;
+		}
+	}
+
 	/**
 	 * Returns whether or not this actor was explicitly hidden in the editor for the duration of the current editor session
 	 * @param bIncludeParent - Whether to recurse up child actor hierarchy or not
diff --git a/Engine/Source/Runtime/Engine/Private/Actor.cpp b/Engine/Source/Runtime/Engine/Private/Actor.cpp
--- a/Engine/Source/Runtime/Engine/Private/Actor.cpp
+++ b/Engine/Source/Runtime/Engine/Private/Actor.cpp
@@ -130,6 +130,7 @@ void AActor::InitializeDefaults()
 	SpriteScale = 1.0f;
 	bEnableAutoLODGeneration = true;
 	bOptimizeBPComponentData = false;
+	bForceExternalActorLevelReferenceForPIE = false;
 #endif // WITH_EDITORONLY_DATA
 	NetCullDistanceSquared = 225000000.0f;
 	NetDriverName = NAME_GameNetDriver;
diff --git a/Engine/Source/Runtime/Engine/Private/Level.cpp b/Engine/Source/Runtime/Engine/Private/Level.cpp
--- a/Engine/Source/Runtime/Engine/Private/Level.cpp
+++ b/Engine/Source/Runtime/Engine/Private/Level.cpp
@@ -453,10 +453,11 @@ void ULevel::Serialize( FArchive& Ar )
 				{
 					return false;
 				}
-				// When the world is partitioned and duplicating for PIE, we don't want to duplicate any external actors as they will be streamed by the runtime grids.
+				// When the world is partitioned and duplicating for PIE, we only duplicate external actors that were explicitely marked (see IsForceExternalActorLevelReferenceForPIE).
+				// All the other actors will be streamed by the runtime grids.
 				else if (bIsPartitioned && (Ar.GetPortFlags() & PPF_DuplicateForPIE))
 				{
-					return false;
+					return Actor->IsForceExternalActorLevelReferenceForPIE();
 				}
 			}
 #endif
diff --git a/Engine/Source/Runtime/Engine/Private/WorldPartition/WorldPartition.cpp b/Engine/Source/Runtime/Engine/Private/WorldPartition/WorldPartition.cpp
--- a/Engine/Source/Runtime/Engine/Private/WorldPartition/WorldPartition.cpp
+++ b/Engine/Source/Runtime/Engine/Private/WorldPartition/WorldPartition.cpp
@@ -140,12 +140,14 @@ void UWorldPartition::OnPreBeginPIE(bool bStartSimulate)
 {
 	check(IsMainWorldPartition());
 	GenerateStreaming(EWorldPartitionStreamingMode::PIE);
+	RuntimeHash->OnPreBeginPIE();
 }
 
 void UWorldPartition::OnEndPIE(bool bStartSimulate)
 {
 	check(IsMainWorldPartition());
 	FlushStreaming();
+	RuntimeHash->OnEndPIE();
 }
 
 FName UWorldPartition::GetWorldPartitionEditorName()
diff --git a/Engine/Source/Runtime/Engine/Private/WorldPartition/WorldPartitionLevelStreamingPolicy.cpp b/Engine/Source/Runtime/Engine/Private/WorldPartition/WorldPartitionLevelStreamingPolicy.cpp
--- a/Engine/Source/Runtime/Engine/Private/WorldPartition/WorldPartitionLevelStreamingPolicy.cpp
+++ b/Engine/Source/Runtime/Engine/Private/WorldPartition/WorldPartitionLevelStreamingPolicy.cpp
@@ -144,7 +144,6 @@ TSubclassOf<UWorldPartitionRuntimeCell> UWorldPartitionLevelStreamingPolicy::Get
 
 void UWorldPartitionLevelStreamingPolicy::PrepareForPIE()
 {
-	TSet<const UWorldPartitionRuntimeCell*> AlwaysLoadedCells;
 	TSet<const UWorldPartitionRuntimeCell*> StreamingCells;
 	WorldPartition->RuntimeHash->GetAllStreamingCells(StreamingCells, /*bIncludeDataLayers*/ true);
 
@@ -157,26 +156,12 @@ void UWorldPartitionLevelStreamingPolicy::PrepareForPIE()
 		{
 			ActorToCellRemapping.Add(CellObjectMap.Path, StreamingCell->GetFName());
 		}
-
-		if (StreamingCell->IsAlwaysLoaded())
-		{
-			AlwaysLoadedCells.Add(Cell);
-		}
 	}
-
-	// In PIE, this is where we load Always Loaded cells
-	Super::LoadCells(AlwaysLoadedCells);
 }
 
 void UWorldPartitionLevelStreamingPolicy::OnPreFixupForPIE(int32 InPIEInstanceID, FSoftObjectPath& ObjectPath)
 {
 	// Once we fix up the path to the proper streaming Level, normal flow of FixupForPIE will adapt the path for PIE.
-	FixupSoftObjectPath(ObjectPath);
-}
-
-// Fix up FSoftObjectPath to remap to proper streaming level
-void UWorldPartitionLevelStreamingPolicy::FixupSoftObjectPath(FSoftObjectPath& ObjectPath)
-{
 	FString SrcPath = ObjectPath.ToString();
 
 	FName* CellName = ActorToCellRemapping.Find(FName(*SrcPath));
diff --git a/Engine/Source/Runtime/Engine/Private/WorldPartition/WorldPartitionRuntimeHash.cpp b/Engine/Source/Runtime/Engine/Private/WorldPartition/WorldPartitionRuntimeHash.cpp
--- a/Engine/Source/Runtime/Engine/Private/WorldPartition/WorldPartitionRuntimeHash.cpp
+++ b/Engine/Source/Runtime/Engine/Private/WorldPartition/WorldPartitionRuntimeHash.cpp
@@ -5,3 +5,33 @@
 UWorldPartitionRuntimeHash::UWorldPartitionRuntimeHash(const FObjectInitializer& ObjectInitializer)
 : Super(ObjectInitializer)
 {}
+
+#if WITH_EDITOR
+void UWorldPartitionRuntimeHash::OnPreBeginPIE()
+{
+	// Mark always loaded actors so that the Level will force reference to these actors for PIE.
+	// These actor will then be duplicated for PIE during the PIE world duplication process
+	ForceExternalActorLevelReference(/*bForceExternalActorLevelReferenceForPIE*/true);
+}
+
+void UWorldPartitionRuntimeHash::OnEndPIE()
+{
+	// Unmark always loaded actors
+	ForceExternalActorLevelReference(/*bForceExternalActorLevelReferenceForPIE*/false);
+
+	// Release references (will unload actors that were not already loaded in the Editor)
+	AlwaysLoadedActorsForPIE.Empty();
+}
+
+void UWorldPartitionRuntimeHash::ForceExternalActorLevelReference(bool bForceExternalActorLevelReferenceForPIE)
+{
+	for (const FWorldPartitionReference& AlwaysLoadedActor : AlwaysLoadedActorsForPIE)
+	{
+		if (AActor* Actor = (*AlwaysLoadedActor)->GetActor())
+		{
+			Actor->SetForceExternalActorLevelReferenceForPIE(bForceExternalActorLevelReferenceForPIE);
+		}
+	}
+}
+
+#endif
diff --git a/Engine/Source/Runtime/Engine/Private/WorldPartition/WorldPartitionRuntimeSpatialHash.cpp b/Engine/Source/Runtime/Engine/Private/WorldPartition/WorldPartitionRuntimeSpatialHash.cpp
--- a/Engine/Source/Runtime/Engine/Private/WorldPartition/WorldPartitionRuntimeSpatialHash.cpp
+++ b/Engine/Source/Runtime/Engine/Private/WorldPartition/WorldPartitionRuntimeSpatialHash.cpp
@@ -510,6 +510,7 @@ bool UWorldPartitionRuntimeSpatialHash::CreateStreamingGrid(const FSpatialHashRu
 	TRACE_CPUPROFILER_EVENT_SCOPE(UWorldPartitionRuntimeSpatialHash::CreateStreamingGrid);
 
 	UWorldPartition* WorldPartition = GetOuterUWorldPartition();
+	const bool bIsMainWorldPartition = (WorldPartition->GetWorld() == WorldPartition->GetTypedOuter<UWorld>());
 
 	check(FMath::IsPowerOfTwo(PartionedActors.GridSize));
 	FSpatialHashStreamingGrid& CurrentStreamingGrid = StreamingGrids.AddDefaulted_GetRef();
@@ -556,14 +557,28 @@ bool UWorldPartitionRuntimeSpatialHash::CreateStreamingGrid(const FSpatialHashRu
 					}, [](const FGuid& ActorGuid) { return ActorGuid; });
 				}
 
+				// Cell cannot be treated as always loaded if it has data layers
+				const bool bIsCellAlwaysLoaded = (&TempCell == &PartionedActors.GetAlwaysLoadedCell()) && !GridCellDataChunk.HasDataLayers();
+
+				// In PIE, Always loaded cell is not generated. Instead, always loaded actors will be added to AlwaysLoadedActorsForPIE.
+				// This will trigger loading/registration of these actors in the PersistentLevel (if not already loaded).
+				// Then, duplication of world for PIE will duplicate only these actors. 
+				// When stopping PIE, WorldPartition will release these FWorldPartitionReferences which 
+				// will unload actors that were not already loaded in the non PIE world.
+				if (bIsCellAlwaysLoaded && bIsMainWorldPartition && (Mode == EWorldPartitionStreamingMode::PIE))
+				{
+					for (const FGuid& ActorGuid : FilteredActors)
+					{
+						AlwaysLoadedActorsForPIE.Emplace(WorldPartition, ActorGuid);
+					}
+					continue;
+				}
+
 				if (!FilteredActors.Num())
 				{
 					continue;
 				}
 				
-				// Cell cannot be treated as always loaded if it has data layers
-				const bool bIsCellAlwaysLoaded = (&TempCell == &PartionedActors.GetAlwaysLoadedCell()) && !GridCellDataChunk.HasDataLayers();
-
 				FName CellName = GetCellName(CurrentStreamingGrid.GridName, Level, CellCoordX, CellCoordY, GridCellDataChunk.GetDataLayersID());
 
 				UWorldPartitionRuntimeSpatialHashCell* StreamingCell = NewObject<UWorldPartitionRuntimeSpatialHashCell>(WorldPartition, StreamingPolicy->GetRuntimeCellClass(), CellName);
diff --git a/Engine/Source/Runtime/Engine/Public/WorldPartition/WorldPartitionLevelStreamingPolicy.h b/Engine/Source/Runtime/Engine/Public/WorldPartition/WorldPartitionLevelStreamingPolicy.h
--- a/Engine/Source/Runtime/Engine/Public/WorldPartition/WorldPartitionLevelStreamingPolicy.h
+++ b/Engine/Source/Runtime/Engine/Public/WorldPartition/WorldPartitionLevelStreamingPolicy.h
@@ -36,8 +36,6 @@ private:
 	int32 GetCellLoadingCount() const;
 
 #if WITH_EDITOR
-	void FixupSoftObjectPath(FSoftObjectPath& ObjectPath);
-
 	TMap<FName, FName> ActorToCellRemapping;
 #endif
 };
\ No newline at end of file
diff --git a/Engine/Source/Runtime/Engine/Public/WorldPartition/WorldPartitionRuntimeHash.h b/Engine/Source/Runtime/Engine/Public/WorldPartition/WorldPartitionRuntimeHash.h
--- a/Engine/Source/Runtime/Engine/Public/WorldPartition/WorldPartitionRuntimeHash.h
+++ b/Engine/Source/Runtime/Engine/Public/WorldPartition/WorldPartitionRuntimeHash.h
@@ -6,6 +6,7 @@
 #include "UObject/Object.h"
 #include "WorldPartition.h"
 #include "WorldPartitionStreamingPolicy.h"
+#include "WorldPartition/WorldPartitionHandle.h"
 #include "WorldPartitionRuntimeHash.generated.h"
 
 class UWorldPartitionRuntimeCell;
@@ -26,6 +27,10 @@ class ENGINE_API UWorldPartitionRuntimeHash : public UObject
 	virtual bool GenerateNavigationData() { return false; }
 	virtual FName GetActorRuntimeGrid(const AActor* Actor) const { return NAME_None; }
 	virtual void DrawPreview() const {}
+
+	// PIE methods
+	void OnPreBeginPIE();
+	void OnEndPIE();
 #endif
 
 	// Streaming interface
@@ -42,4 +47,14 @@ class ENGINE_API UWorldPartitionRuntimeHash : public UObject
 
 	virtual void Draw2D(class UCanvas* Canvas, const TArray<FWorldPartitionStreamingSource>& Sources, const FVector2D& PartitionCanvasOffset, const FVector2D& PartitionCanvasSize) const {}
 	virtual void Draw3D(const TArray<FWorldPartitionStreamingSource>& Sources) const {}
+
+private:
+#if WITH_EDITOR
+	void ForceExternalActorLevelReference(bool bForceExternalActorLevelReferenceForPIE);
+#endif
+
+protected:
+#if WITH_EDITORONLY_DATA
+	TArray<FWorldPartitionReference> AlwaysLoadedActorsForPIE;
+#endif
 };
\ No newline at end of file
