Fix for double-ticks of MediaFramework when Sequencer is used in game mode
#rnx

diff --git a/Engine/Source/Runtime/Engine/Private/GameEngine.cpp b/Engine/Source/Runtime/Engine/Private/GameEngine.cpp
--- a/Engine/Source/Runtime/Engine/Private/GameEngine.cpp
+++ b/Engine/Source/Runtime/Engine/Private/GameEngine.cpp
@@ -1708,17 +1708,6 @@ void UGameEngine::Tick( float DeltaSeconds, bool bIdleMode )
 		UE_LOG(LogEngine, Log, TEXT("Slow GT frame detected (GT frame %u, delta time %f s)"), GFrameCounter - 1, DeltaSeconds);
 	}
 
-#if !UE_SERVER
-	// tick media framework
-	static const FName MediaModuleName(TEXT("Media"));
-	IMediaModule* MediaModule = FModuleManager::LoadModulePtr<IMediaModule>(MediaModuleName);
-
-	if (MediaModule != nullptr)
-	{
-		MediaModule->TickPreEngine();
-	}
-#endif
-
 	if (IsRunningDedicatedServer())
 	{
 		double CurrentTime = FPlatformTime::Seconds();
@@ -1744,13 +1733,42 @@ void UGameEngine::Tick( float DeltaSeconds, bool bIdleMode )
 		return;
 	}
 
-	if ( GameViewport != NULL )
+	if (GameViewport != NULL)
 	{
 		// Decide whether to drop high detail because of frame rate.
 		QUICK_SCOPE_CYCLE_COUNTER(STAT_UGameEngine_Tick_SetDropDetail);
 		GameViewport->SetDropDetail(DeltaSeconds);
 	}
 
+#if !UE_SERVER
+	// Media module present?
+	static const FName MediaModuleName(TEXT("Media"));
+	IMediaModule* MediaModule = FModuleManager::LoadModulePtr<IMediaModule>(MediaModuleName);
+	if (MediaModule != nullptr)
+	{
+		// Yes. Will a world trigger the MediaFramework tick due to an active Sequencer?
+		bool bWorldWillTickMediaFramework = false;
+		if (!bIdleMode)
+		{
+			for (int32 i = 0; i < WorldList.Num(); ++i)
+			{
+				FWorldContext& Context = WorldList[i];
+				if (Context.World() != nullptr && Context.World()->ShouldTick() && Context.World()->IsMovieSceneSequenceTickHandlerBound())
+				{
+					bWorldWillTickMediaFramework = true;
+					break;
+				}
+			}
+		}
+		if (!bWorldWillTickMediaFramework)
+		{
+			// tick media framework if no world would do it later on
+			// (so we can normally - no Sequencer active - assume that the media state changes are all done early)
+			MediaModule->TickPreEngine();
+		}
+	}
+#endif
+
 	// Update subsystems.
 	{
 		// This assumes that UObject::StaticTick only calls ProcessAsyncLoading.
