Removed FHashBuilderArchive, added FArchiveCrc32

diff --git a/Engine/Source/Runtime/Core/Public/Misc/HashBuilder.h b/Engine/Source/Runtime/Core/Public/Misc/HashBuilder.h
--- a/Engine/Source/Runtime/Core/Public/Misc/HashBuilder.h
+++ b/Engine/Source/Runtime/Core/Public/Misc/HashBuilder.h
@@ -80,30 +80,3 @@ public:
 private:
 	uint32 Hash;
 };
-
-/**
- * FArchive adapter for FHashBuilder
- */
-class FHashBuilderArchive : public FArchive
-{
-public:
-	FHashBuilderArchive()
-	{
-		SetIsLoading(false);
-		SetIsSaving(true);
-		SetIsPersistent(false);
-	}
-
-	virtual FString GetArchiveName() const { return TEXT("FHashBuilderArchive"); }
-
-	void Serialize(void* Data, int64 Num) override { HashBuilder.AppendRaw(Data, Num); }
-
-	using FArchive::operator<<;
-	virtual FArchive& operator<<(class FName& Value) override { HashBuilder << Value; return *this; }
-	virtual FArchive& operator<<(class UObject*& Value) override { check(0); return *this; }
-
-	uint32 GetHash() const { return HashBuilder.GetHash(); }
-
-protected:
-	FHashBuilder HashBuilder;
-};
\ No newline at end of file
diff --git a/Engine/Source/Runtime/CoreUObject/Private/Serialization/ArchiveCrc32.cpp b/Engine/Source/Runtime/CoreUObject/Private/Serialization/ArchiveCrc32.cpp
new file mode 100644
--- /dev/null
+++ b/Engine/Source/Runtime/CoreUObject/Private/Serialization/ArchiveCrc32.cpp
@@ -0,0 +1,35 @@
+// Copyright Epic Games, Inc. All Rights Reserved.
+
+#include "Serialization/ArchiveCrc32.h"
+#include "Serialization/ArchiveObjectCrc32.h"
+
+FArchiveCrc32::FArchiveCrc32(uint32 InCRC, UObject* InRootObject)
+	: CRC(InCRC)
+	, RootObject(InRootObject)
+{
+}
+
+void FArchiveCrc32::Serialize(void* Data, int64 Num)
+{
+	CRC = FCrc::MemCrc32(Data, Num, CRC);
+}
+
+FArchive& FArchiveCrc32::operator<<(class FName& Name)
+{
+	CRC = FCrc::StrCrc32(*Name.ToString(), CRC);
+	return *this;
+}
+
+FArchive& FArchiveCrc32::operator<<(class UObject*& Object)
+{
+	if (!RootObject || !Object || !Object->IsIn(RootObject))
+	{
+		CRC = FCrc::StrCrc32(*GetPathNameSafe(Object), CRC);
+	}
+	else
+	{
+		CRC = FArchiveObjectCrc32().Crc32(Object, RootObject, CRC);
+	}
+
+	return *this;
+}
diff --git a/Engine/Source/Runtime/CoreUObject/Private/Serialization/ArchiveObjectCrc32.cpp b/Engine/Source/Runtime/CoreUObject/Private/Serialization/ArchiveObjectCrc32.cpp
--- a/Engine/Source/Runtime/CoreUObject/Private/Serialization/ArchiveObjectCrc32.cpp
+++ b/Engine/Source/Runtime/CoreUObject/Private/Serialization/ArchiveObjectCrc32.cpp
@@ -47,7 +47,7 @@ FArchive& FArchiveObjectCrc32::operator<<(class UObject*& Object)
 {
 	FArchive& Ar = *this;
 
-	if (!Object || !Object->IsIn(RootObject))
+	if (!RootObject || !Object || !Object->IsIn(RootObject))
 	{
 		auto UniqueName = GetPathNameSafe(Object);
 		Ar << UniqueName;
@@ -60,13 +60,13 @@ FArchive& FArchiveObjectCrc32::operator<<(class UObject*& Object)
 	return Ar;
 }
 
-uint32 FArchiveObjectCrc32::Crc32(UObject* Object, uint32 CRC)
+uint32 FArchiveObjectCrc32::Crc32(UObject* Object, UObject* Root, uint32 CRC)
 {
 #ifdef DEBUG_ARCHIVE_OBJECT_CRC32
 	const double StartTime = FPlatformTime::Seconds();
 	UE_LOG(LogArchiveObjectCrc32, Log, TEXT("### Calculating CRC for object: %s with outer: %s"), *Object->GetName(), Object->GetOuter() ? *Object->GetOuter()->GetName() : TEXT("NULL"));
 #endif
-	RootObject = Object;
+	RootObject = Root;
 	if (Object)
 	{
 		TSet<UObject*> SerializedObjects;
@@ -114,3 +114,8 @@ uint32 FArchiveObjectCrc32::Crc32(UObject* Object, uint32 CRC)
 #endif
 	return CRC;
 }
+
+uint32 FArchiveObjectCrc32::Crc32(UObject* Object, uint32 CRC)
+{
+	return Crc32(Object, Object, CRC);
+}
diff --git a/Engine/Source/Runtime/CoreUObject/Public/Serialization/ArchiveCrc32.h b/Engine/Source/Runtime/CoreUObject/Public/Serialization/ArchiveCrc32.h
new file mode 100644
--- /dev/null
+++ b/Engine/Source/Runtime/CoreUObject/Public/Serialization/ArchiveCrc32.h
@@ -0,0 +1,49 @@
+// Copyright Epic Games, Inc. All Rights Reserved.
+
+#pragma once
+
+#include "CoreMinimal.h"
+#include "Serialization/ArchiveUObject.h"
+#include "Serialization/MemoryLayout.h"
+
+/**
+* Calculates a checksum from the input provided to the archive.
+*/
+class COREUOBJECT_API FArchiveCrc32 : public FArchiveUObject
+{
+public:
+	/**
+	* Default constructor.
+	*/
+	FArchiveCrc32(uint32 InCRC = 0, UObject* InRootObject = nullptr);
+
+	/**
+	 * @return The checksum computed so far.
+	 */
+	uint32 GetCrc() const { return CRC; }
+
+	//~ Begin FArchive Interface
+	virtual void Serialize(void* Data, int64 Num) override;
+	virtual FArchive& operator<<(class FName& Name);
+	virtual FArchive& operator<<(class UObject*& Object);
+	virtual FString GetArchiveName() const { return TEXT("FArchiveCrc32"); }
+	//~ End FArchive Interface
+
+	/**
+	 * Serializes a USTRUCT value from or into an archive.
+	 *
+	 * @param Ar The archive to serialize from or to.
+	 * @param Value The value to serialize.
+	 */
+	template <typename StructType>
+	FORCEINLINE friend typename TEnableIf<TProvidesStaticStruct<StructType>::Value, FArchive&>::Type operator <<(FArchive& Ar, const StructType& Value)
+	{
+		StructType* MutableValue = const_cast<StructType*>(&Value);
+		StructType::StaticStruct()->SerializeItem(Ar, MutableValue, nullptr);
+		return Ar;
+	}
+
+private:
+	uint32 CRC;
+	UObject* RootObject;
+};
diff --git a/Engine/Source/Runtime/CoreUObject/Public/Serialization/ArchiveObjectCrc32.h b/Engine/Source/Runtime/CoreUObject/Public/Serialization/ArchiveObjectCrc32.h
--- a/Engine/Source/Runtime/CoreUObject/Public/Serialization/ArchiveObjectCrc32.h
+++ b/Engine/Source/Runtime/CoreUObject/Public/Serialization/ArchiveObjectCrc32.h
@@ -33,6 +33,7 @@ public:
 	* Serialize the given object, calculate and return its checksum.
 	*/
 	uint32 Crc32(UObject* Object, uint32 CRC = 0);
+	uint32 Crc32(UObject* Object, UObject* Root, uint32 CRC);
 
 protected:
 	/** Return if object was already serialized */
diff --git a/Engine/Source/Runtime/Engine/Private/HLODProxy.cpp b/Engine/Source/Runtime/Engine/Private/HLODProxy.cpp
--- a/Engine/Source/Runtime/Engine/Private/HLODProxy.cpp
+++ b/Engine/Source/Runtime/Engine/Private/HLODProxy.cpp
@@ -6,7 +6,7 @@
 
 #if WITH_EDITOR
 #include "Misc/ConfigCacheIni.h"
-#include "Serialization/ArchiveObjectCrc32.h"
+#include "Serialization/ArchiveCrc32.h"
 #include "HierarchicalLOD.h"
 #include "ObjectTools.h"
 #endif
@@ -317,7 +317,7 @@ void UHLODProxy::ExtractComponents(const ALODActor* LODActor, TArray<UPrimitiveC
 
 uint32 UHLODProxy::GetCRC(UMaterialInterface* InMaterialInterface, uint32 InCRC)
 {
-	TArray<uint8> KeyBuffer;
+	FArchiveCrc32 Ar(InCRC);
 
 	UMaterialInterface* MaterialInterface = InMaterialInterface;
 	while(MaterialInterface)
@@ -327,18 +327,18 @@ uint32 UHLODProxy::GetCRC(UMaterialInterface* InMaterialInterface, uint32 InCRC)
 		{
 			if(UMaterialInstanceConstant* MIC = Cast<UMaterialInstanceConstant>(MI))
 			{
-				KeyBuffer.Append((uint8*)&MIC->ParameterStateId, sizeof(FGuid));
+				Ar << MIC->ParameterStateId;
 			}
 			MaterialInterface = MI->Parent;
 		}
 		else if(UMaterial* Material = Cast<UMaterial>(MaterialInterface))
 		{
-			KeyBuffer.Append((uint8*)&Material->StateId, sizeof(FGuid));
+			Ar << Material->StateId;
 			MaterialInterface = nullptr;
 		}
 	}
 
-	return FCrc::MemCrc32(KeyBuffer.GetData(), KeyBuffer.Num(), InCRC);
+	return Ar.GetCrc();
 }
 
 uint32 UHLODProxy::GetCRC(UTexture* InTexture, uint32 InCRC)
@@ -359,49 +359,53 @@ uint32 UHLODProxy::GetCRC(UTexture* InTexture, uint32 InCRC)
 
 uint32 UHLODProxy::GetCRC(UStaticMesh* InStaticMesh, uint32 InCRC)
 {
-	TArray<uint8> KeyBuffer;
+	FArchiveCrc32 Ar(InCRC);
 
 	// Default to just the path name if we don't have render data
 	FString DerivedDataKey = InStaticMesh->GetRenderData() ? InStaticMesh->GetRenderData()->DerivedDataKey : InStaticMesh->GetPathName();
-	KeyBuffer.Append((uint8*)DerivedDataKey.GetCharArray().GetData(), DerivedDataKey.GetCharArray().Num() * DerivedDataKey.GetCharArray().GetTypeSize());
+	Ar << DerivedDataKey;
+
+	int32 LightMapCoordinateIndex = InStaticMesh->GetLightMapCoordinateIndex();
+	Ar << LightMapCoordinateIndex;
 
-	const int32 LightMapCoordinateIndex = InStaticMesh->GetLightMapCoordinateIndex();
-	KeyBuffer.Append((uint8*)&LightMapCoordinateIndex, sizeof(int32));
 	if(InStaticMesh->GetBodySetup())
 	{
 		// Incorporate physics data
-		KeyBuffer.Append((uint8*)&InStaticMesh->GetBodySetup()->BodySetupGuid, sizeof(FGuid));;
+		Ar << InStaticMesh->GetBodySetup()->BodySetupGuid;
 	}
-	return FCrc::MemCrc32(KeyBuffer.GetData(), KeyBuffer.Num(), InCRC);
+
+	return Ar.GetCrc();
 }
 
-static void AppendRoundedTransform(const FRotator& ComponentRotation, const FVector& ComponentLocation, const FVector& ComponentScale, TArray<uint8>& OutKeyBuffer)
+static void AppendRoundedTransform(const FRotator& ComponentRotation, const FVector& ComponentLocation, const FVector& ComponentScale, FArchive& Ar)
 {
 	// Include transform - round sufficiently to ensure stability
 	FIntVector Location(FMath::RoundToInt(ComponentLocation.X), FMath::RoundToInt(ComponentLocation.Y), FMath::RoundToInt(ComponentLocation.Z));
-	OutKeyBuffer.Append((uint8*)&Location, sizeof(Location));
+	Ar << Location;
+
 	FVector RotationVector(ComponentRotation.GetNormalized().Vector());
 	FIntVector Rotation(FMath::RoundToInt(RotationVector.X), FMath::RoundToInt(RotationVector.Y), FMath::RoundToInt(RotationVector.Z));
-	OutKeyBuffer.Append((uint8*)&Rotation, sizeof(Rotation));
+	Ar << Rotation;
+
 	FIntVector Scale(FMath::RoundToInt(ComponentScale.X), FMath::RoundToInt(ComponentScale.Y), FMath::RoundToInt(ComponentScale.Z));
-	OutKeyBuffer.Append((uint8*)&Scale, sizeof(Scale));
+	Ar << Scale;
 }
 
-static void AppendRoundedTransform(const FTransform& InTransform, TArray<uint8>& OutKeyBuffer)
+static void AppendRoundedTransform(const FTransform& InTransform, FArchive& Ar)
 {
-	AppendRoundedTransform(InTransform.Rotator(), InTransform.GetLocation(), InTransform.GetScale3D(), OutKeyBuffer);
+	AppendRoundedTransform(InTransform.Rotator(), InTransform.GetLocation(), InTransform.GetScale3D(), Ar);
 }
 
-static int32 GetTransformCRC(const FTransform& InTransform, uint32 InCRC)
+static int32 GetTransformCRC(const FTransform& InTransform, uint32 InCRC = 0)
 {
-	TArray<uint8> KeyBuffer;
-	AppendRoundedTransform(InTransform, KeyBuffer);
-	return FCrc::MemCrc32(KeyBuffer.GetData(), KeyBuffer.Num(), InCRC);
+	FArchiveCrc32 Ar(InCRC);
+	AppendRoundedTransform(InTransform, Ar);
+	return Ar.GetCrc();
 }
 
 uint32 UHLODProxy::GetCRC(UStaticMeshComponent* InComponent, uint32 InCRC, const FTransform& TransformComponents)
 {
-	TArray<uint8> KeyBuffer;
+	FArchiveCrc32 Ar(InCRC);
 
 	FVector  ComponentLocation = InComponent->GetComponentLocation();
 	FRotator ComponentRotation = InComponent->GetComponentRotation();
@@ -409,71 +413,52 @@ uint32 UHLODProxy::GetCRC(UStaticMeshComponent* InComponent, uint32 InCRC, const
 
 	ComponentLocation = TransformComponents.TransformPosition(ComponentLocation);
 	ComponentRotation = TransformComponents.TransformRotation(ComponentRotation.Quaternion()).Rotator();
-	AppendRoundedTransform(ComponentRotation, ComponentLocation, ComponentScale, KeyBuffer);
+	AppendRoundedTransform(ComponentRotation, ComponentLocation, ComponentScale, Ar);
 
 	// Include other relevant properties
-	KeyBuffer.Append((uint8*)&InComponent->ForcedLodModel, sizeof(int32));
+	Ar << InComponent->ForcedLodModel;
 	bool bUseMaxLODAsImposter = InComponent->bUseMaxLODAsImposter;
-	KeyBuffer.Append((uint8*)&bUseMaxLODAsImposter, sizeof(bool));
+	Ar << bUseMaxLODAsImposter;
 	bool bCastShadow = InComponent->CastShadow;
-	KeyBuffer.Append((uint8*)&bCastShadow, sizeof(bool));
+	Ar << bCastShadow;
 	bool bCastStaticShadow = InComponent->bCastStaticShadow;
-	KeyBuffer.Append((uint8*)&bCastStaticShadow, sizeof(bool));
+	Ar << bCastStaticShadow;
 	bool bCastDynamicShadow = InComponent->bCastDynamicShadow;
-	KeyBuffer.Append((uint8*)&bCastDynamicShadow, sizeof(bool));
+	Ar << bCastDynamicShadow;
 	bool bCastFarShadow = InComponent->bCastFarShadow;
-	KeyBuffer.Append((uint8*)&bCastFarShadow, sizeof(bool));
+	Ar << bCastFarShadow;
 	int32 Width, Height;
 	InComponent->GetLightMapResolution(Width, Height);
-	KeyBuffer.Append((uint8*)&Width, sizeof(int32));
-	KeyBuffer.Append((uint8*)&Height, sizeof(int32));
+	Ar << Width;
+	Ar << Height;
 	
 	// incorporate vertex colors
 	for(FStaticMeshComponentLODInfo& LODInfo : InComponent->LODData)
 	{
 		if(LODInfo.OverrideVertexColors)
 		{
-			KeyBuffer.Append((uint8*)LODInfo.OverrideVertexColors->GetVertexData(), LODInfo.OverrideVertexColors->GetNumVertices() * LODInfo.OverrideVertexColors->GetStride());
+			Ar.Serialize((uint8*)LODInfo.OverrideVertexColors->GetVertexData(), LODInfo.OverrideVertexColors->GetNumVertices() * LODInfo.OverrideVertexColors->GetStride());
 		}
 	}
 
-	return FCrc::MemCrc32(KeyBuffer.GetData(), KeyBuffer.Num(), InCRC);
-}
-
-// Key that forms the basis of the HLOD proxy key. Bump this key (i.e. generate a new GUID) when you want to force a rebuild of ALL HLOD proxies
-#define HLOD_PROXY_BASE_KEY		TEXT("174C29B19AB34A21894058E058F253B3")
-
-namespace 
-{
-	class FHLODProxyCRCArchive : public FArchive
+	// Include instances data in case of a ISMC
+	if (UInstancedStaticMeshComponent* ISMC = Cast<UInstancedStaticMeshComponent>(InComponent))
 	{
-	public:
-		FHLODProxyCRCArchive()
-			: Hash(0)
+		for (const FInstancedStaticMeshInstanceData& InstanceData : ISMC->PerInstanceSMData)
 		{
-			SetIsLoading(false);
-			SetIsSaving(true);
-			SetUseUnversionedPropertySerialization(true);
+			AppendRoundedTransform(FTransform(InstanceData.Transform), Ar);
 		}
 
-		uint32 GetHash() const
-		{
-			return Hash;
-		}
-
-	private:
-		virtual FArchive& operator<<(class UObject*& Value) override { check(0); return *this; }
-		virtual FArchive& operator<<(class FName& Value) override { check(0); return *this; }
-
-		virtual void Serialize(void* Data, int64 Num) override
-		{
-			Hash = FCrc::MemCrc32(Data, Num, Hash);
-		}
+		Ar << ISMC->PerInstanceSMCustomData;
+		Ar << ISMC->InstancingRandomSeed;
+	}
 
-		uint32 Hash;
-	};
+	return Ar.GetCrc();
 }
 
+// Key that forms the basis of the HLOD proxy key. Bump this key (i.e. generate a new GUID) when you want to force a rebuild of ALL HLOD proxies
+#define HLOD_PROXY_BASE_KEY		TEXT("174C29B19AB34A21894058E058F253B3")
+
 FName UHLODProxy::GenerateKeyForActor(const ALODActor* LODActor, bool bMustUndoLevelTransform)
 {
 	FString Key = HLOD_PROXY_BASE_KEY;
@@ -495,10 +480,9 @@ FName UHLODProxy::GenerateKeyForActor(const ALODActor* LODActor, bool bMustUndoL
 			TArray<FHierarchicalSimplification>& BuildLODLevelSettings = LODActor->GetLevel()->GetWorldSettings()->GetHierarchicalLODSetup();
 			if(BuildLODLevelSettings.IsValidIndex(LODActor->LODLevel - 1))
 			{
-				FHierarchicalSimplification BuildLODLevelSetting = BuildLODLevelSettings[LODActor->LODLevel - 1];
-				FHLODProxyCRCArchive Ar;
-				FHierarchicalSimplification::StaticStruct()->SerializeItem(Ar, &BuildLODLevelSetting, nullptr);
-				CRC = HashCombine(CRC, Ar.GetHash());
+				FArchiveCrc32 Ar(CRC);
+				Ar << BuildLODLevelSettings[LODActor->LODLevel - 1];
+				CRC = Ar.GetCrc();
 			}
 		}
 
@@ -517,10 +501,9 @@ FName UHLODProxy::GenerateKeyForActor(const ALODActor* LODActor, bool bMustUndoL
 		{
 			if (LODActor->bOverrideMaterialMergeSettings)
 			{
-				FMaterialProxySettings MaterialProxySettings = LODActor->MaterialSettings;
-				FHLODProxyCRCArchive Ar;
-				FMaterialProxySettings::StaticStruct()->SerializeItem(Ar, &MaterialProxySettings, nullptr);
-				CRC = HashCombine(CRC, Ar.GetHash());
+				FArchiveCrc32 Ar(CRC);
+				Ar << const_cast<ALODActor*>(LODActor)->MaterialSettings;
+				CRC = Ar.GetCrc();
 			}
 		}
 
diff --git a/Engine/Source/Runtime/Engine/Private/ISMPartition/ISMComponentDescriptor.cpp b/Engine/Source/Runtime/Engine/Private/ISMPartition/ISMComponentDescriptor.cpp
--- a/Engine/Source/Runtime/Engine/Private/ISMPartition/ISMComponentDescriptor.cpp
+++ b/Engine/Source/Runtime/Engine/Private/ISMPartition/ISMComponentDescriptor.cpp
@@ -4,7 +4,7 @@
 
 #if WITH_EDITOR
 
-#include "Serialization/ArchiveObjectCrc32.h"
+#include "Serialization/ArchiveCrc32.h"
 #include "Components/InstancedStaticMeshComponent.h"
 #include "Components/HierarchicalInstancedStaticMeshComponent.h"
 
@@ -129,14 +129,13 @@ bool FISMComponentDescriptor::operator==(const FISMComponentDescriptor& Other) c
 
 uint32 FISMComponentDescriptor::ComputeHash() const
 {
-	FArchiveObjectCrc32 CrcArchive;
-	UISMComponentDescriptorHasher* Hasher = NewObject<UISMComponentDescriptorHasher>();
+	FArchiveCrc32 CrcArchive;
+
 	Hash = 0; // we don't want the hash to impact the calculation
-	Hasher->Descriptor = *this;
-	uint32 Crc = CrcArchive.Crc32(Hasher);
-	Hash = Crc;
+	CrcArchive << *this;
+	Hash = CrcArchive.GetCrc();
 
-	return Crc;
+	return Hash;
 }
 
 UInstancedStaticMeshComponent* FISMComponentDescriptor::CreateComponent(UObject* Outer, FName Name, EObjectFlags ObjectFlags) const
@@ -190,8 +189,3 @@ void FISMComponentDescriptor::InitComponent(UInstancedStaticMeshComponent* ISMCo
 }
 
 #endif
-
-UISMComponentDescriptorHasher::UISMComponentDescriptorHasher(const FObjectInitializer& ObjectInitializer)
-	: Super(ObjectInitializer)
-{
-}
\ No newline at end of file
diff --git a/Engine/Source/Runtime/Engine/Public/ISMPartition/ISMComponentDescriptor.h b/Engine/Source/Runtime/Engine/Public/ISMPartition/ISMComponentDescriptor.h
--- a/Engine/Source/Runtime/Engine/Public/ISMPartition/ISMComponentDescriptor.h
+++ b/Engine/Source/Runtime/Engine/Public/ISMPartition/ISMComponentDescriptor.h
@@ -155,14 +155,3 @@ public:
 	uint8 bIsEditorOnly : 1;
 #endif
 };
-
-/** Utility object to hash the FISMComponentDescriptor through FArchiveObjectCrc32 */
-UCLASS()
-class UISMComponentDescriptorHasher : public UObject
-{
-	GENERATED_UCLASS_BODY()
-
-public:
-	UPROPERTY()
-	FISMComponentDescriptor Descriptor;
-};
