Missing file part of 15477598

diff --git a/Engine/Source/Runtime/Engine/Private/Materials/StrataMaterial.h b/Engine/Source/Runtime/Engine/Private/Materials/StrataMaterial.h
--- a/Engine/Source/Runtime/Engine/Private/Materials/StrataMaterial.h
+++ b/Engine/Source/Runtime/Engine/Private/Materials/StrataMaterial.h
@@ -27,7 +27,8 @@ struct FStrataMaterialCompilationInfo
 	{
 		uint8 Type;
 		uint8 SharedNormalIndex;
-		bool  bHasScattering;
+		bool  bHasSSS;
+		bool  bHasDMFPPluggedIn;
 		bool  bHasEdgeColor;
 		bool  bHasThinFilm;
 	};
@@ -56,7 +57,7 @@ FString GetStrataBSDFName(uint8 BSDFType);
 uint8 StrataCompilationInfoCreateSharedNormal(FMaterialCompiler* Compiler, int32 NormalCodeChunk);
 uint8 StrataCompilationInfoCreateSharedNormal(FMaterialCompiler* Compiler, int32 NormalCodeChunk, int32 TangentCodeChunk);
 void  StrataCompilationInfoCreateNullBSDF(FMaterialCompiler* Compiler, int32 CodeChunk);
-void  StrataCompilationInfoCreateSingleBSDFMaterial(FMaterialCompiler* Compiler, int32 CodeChunk, uint8 SharedNormalIndex, uint8 BSDFType, bool bHasEdgeColor = false, bool bHasScattering = false, bool bHasThinFilm = false);
+void  StrataCompilationInfoCreateSingleBSDFMaterial(FMaterialCompiler* Compiler, int32 CodeChunk, uint8 SharedNormalIndex, uint8 BSDFType, bool bHasSSS = false, bool bHasDMFPPluggedIn = false, bool bHasEdgeColor = false, bool bHasThinFilm = false);
 
 FStrataMaterialCompilationInfo StrataCompilationInfoAdd(FMaterialCompiler* Compiler, const FStrataMaterialCompilationInfo& A, const FStrataMaterialCompilationInfo& B);
 FStrataMaterialCompilationInfo StrataCompilationInfoMultiply(FMaterialCompiler* Compiler, const FStrataMaterialCompilationInfo& A);
