Force saving newly created level instance actor and all moved actors when creating a new level instance from selection.

diff --git a/Engine/Source/Runtime/Engine/Private/LevelInstance/LevelInstanceSubsystem.cpp b/Engine/Source/Runtime/Engine/Private/LevelInstance/LevelInstanceSubsystem.cpp
--- a/Engine/Source/Runtime/Engine/Private/LevelInstance/LevelInstanceSubsystem.cpp
+++ b/Engine/Source/Runtime/Engine/Private/LevelInstance/LevelInstanceSubsystem.cpp
@@ -795,9 +795,17 @@ ALevelInstance* ULevelInstanceSubsystem::CreateLevelInstanceFrom(const TArray<AA
 	LevelStreaming->LevelInstanceID = NewLevelInstanceActor->GetLevelInstanceID();
 		
 	GetWorld()->SetCurrentLevel(LoadedLevel);
+
+	TSet<UPackage*> DirtyPackages;
+	DirtyPackages.Add(NewLevelInstanceActor->GetPackage());
+
+	for (AActor* ActorToMove : ActorsToMove)
+	{
+		DirtyPackages.Add(ActorToMove->GetPackage());
+	}
 				
 	// Use LevelStreaming->GetLevelInstanceActor() because OnWorldAssetSaved could've reinstanced the LevelInstanceActor
-	return CommitLevelInstance(LevelStreaming->GetLevelInstanceActor());
+	return CommitLevelInstance(LevelStreaming->GetLevelInstanceActor(), /*bDiscardEdits*/false, &DirtyPackages);
 }
 
 bool ULevelInstanceSubsystem::BreakLevelInstance(ALevelInstance* LevelInstanceActor, uint32 Levels /* = 1 */, TArray<AActor*>* OutMovedActors /* = nullptr */)
@@ -1454,7 +1462,7 @@ void ULevelInstanceSubsystem::CommitChildrenLevelInstances(ALevelInstance* Level
 	});
 }
 
-ALevelInstance* ULevelInstanceSubsystem::CommitLevelInstance(ALevelInstance* LevelInstanceActor, bool bDiscardEdits)
+ALevelInstance* ULevelInstanceSubsystem::CommitLevelInstance(ALevelInstance* LevelInstanceActor, bool bDiscardEdits, TSet<UPackage*>* DirtyPackages)
 {
 	check(CanCommitLevelInstance(LevelInstanceActor));
 
@@ -1476,6 +1484,10 @@ ALevelInstance* ULevelInstanceSubsystem::CommitLevelInstance(ALevelInstance* Lev
 		if (!FEditorFileUtils::SaveDirtyPackages(bPromptUserToSave, bSaveMapPackages, bSaveContentPackages, bFastSave, bNotifyNoPackagesSaved, bCanBeDeclined, nullptr,
 			[=](UPackage* DirtyPackage)
 			{
+				if (DirtyPackages && DirtyPackages->Contains(DirtyPackage))
+				{
+					return false;
+				}
 				return ShouldIgnoreDirtyPackage(DirtyPackage, EditingWorld);
 			}))
 		{
diff --git a/Engine/Source/Runtime/Engine/Public/LevelInstance/LevelInstanceSubsystem.h b/Engine/Source/Runtime/Engine/Public/LevelInstance/LevelInstanceSubsystem.h
--- a/Engine/Source/Runtime/Engine/Public/LevelInstance/LevelInstanceSubsystem.h
+++ b/Engine/Source/Runtime/Engine/Public/LevelInstance/LevelInstanceSubsystem.h
@@ -54,7 +54,7 @@ public:
 	bool CanEditLevelInstance(const ALevelInstance* LevelInstanceActor, FText* OutReason = nullptr) const;
 	bool CanCommitLevelInstance(const ALevelInstance* LevelInstanceActor, FText* OutReason = nullptr) const;
 	void EditLevelInstance(ALevelInstance* LevelInstanceActor, TWeakObjectPtr<AActor> ContextActorPtr = nullptr);
-	ALevelInstance* CommitLevelInstance(ALevelInstance* LevelInstanceActor, bool bDiscardEdits = false);
+	ALevelInstance* CommitLevelInstance(ALevelInstance* LevelInstanceActor, bool bDiscardEdits = false, TSet<UPackage*>* DirtyPackages = nullptr);
 	void SaveLevelInstanceAs(ALevelInstance* LevelInstanceActor);
 	bool IsEditingLevelInstanceDirty(const ALevelInstance* LevelInstanceActor) const;
 	bool IsEditingLevelInstance(const ALevelInstance* LevelInstanceActor) const { return GetLevelInstanceEdit(LevelInstanceActor) != nullptr; }
