Fixed issues with editor -game mode.
StaticMesh hi res source model bulk data is no longer a default subobject; instead it is instanced if it was not present in PostLoad, thus fixing issues with object name differences.
MeshDescriptionBaseBulkData is marked as NeedsLoadForEditorGame, as well as not needing load for client or server, so that it can be loaded correctly.

diff --git a/Engine/Source/Runtime/Engine/Private/StaticMesh.cpp b/Engine/Source/Runtime/Engine/Private/StaticMesh.cpp
--- a/Engine/Source/Runtime/Engine/Private/StaticMesh.cpp
+++ b/Engine/Source/Runtime/Engine/Private/StaticMesh.cpp
@@ -2897,8 +2897,6 @@ UStaticMesh::UStaticMesh(const FObjectInitializer& ObjectInitializer)
 	ImportVersion = EImportStaticMeshVersion::BeforeImportStaticMeshVersionWasAdded;
 	LODForOccluderMesh = -1;
 	NumStreamedLODs.Default = -1;
-	GetHiResSourceModel().StaticMeshDescriptionBulkData = CreateDefaultSubobject<UStaticMeshDescriptionBulkData>(TEXT("HiResMeshDescription"));
-	GetHiResSourceModel().StaticMeshDescriptionBulkData->SetFlags(RF_Transactional);
 #endif // #if WITH_EDITORONLY_DATA
 	SetLightMapResolution(4);
 	SetMinLOD(0);
@@ -5285,6 +5283,11 @@ void UStaticMesh::BeginPostLoadInternal(FStaticMeshPostLoadContext& Context)
 	{
 		ComplexCollisionMesh->ConditionalPostLoad();
 	}
+
+	if (GetHiResSourceModel().StaticMeshDescriptionBulkData == nullptr)
+	{
+		GetHiResSourceModel().StaticMeshDescriptionBulkData = NewObject<UStaticMeshDescriptionBulkData>(this, NAME_None, RF_Transactional);
+	}
 #endif //WITH_EDITORONLY_DATA
 
 	// We want to always have a BodySetup, its used for per-poly collision as well
diff --git a/Engine/Source/Runtime/Engine/Private/StaticMeshSourceData.cpp b/Engine/Source/Runtime/Engine/Private/StaticMeshSourceData.cpp
--- a/Engine/Source/Runtime/Engine/Private/StaticMeshSourceData.cpp
+++ b/Engine/Source/Runtime/Engine/Private/StaticMeshSourceData.cpp
@@ -353,9 +353,9 @@ void FStaticMeshSourceModel::FillMaterialName(TMap<int32, FName>& OutMaterialMap
 // If static mesh derived data needs to be rebuilt (new format, serialization differences, etc.) replace the version GUID below with a new one.
 // In case of merge conflicts with DDC versions, you MUST generate a new GUID and set this new GUID as the version.
 #if UE_USE_VIRTUALBULKDATA
-#define MESHDATAKEY_STATICMESH_DERIVEDDATA_VER TEXT("7F1A02EF80374F0F86F9A30BAFF3F850")
+#define MESHDATAKEY_STATICMESH_DERIVEDDATA_VER TEXT("80264C99E2AD4D83A1DDBBB8D93D8AEF")
 #else
-#define MESHDATAKEY_STATICMESH_DERIVEDDATA_VER TEXT("E09E8DD68C864D82B511A97080B04837")
+#define MESHDATAKEY_STATICMESH_DERIVEDDATA_VER TEXT("B3CE4F8E59E542D4A721EE56068AD214")
 #endif
 
 static bool GetMeshDataKey(FString& OutKey, const FRawMeshBulkData* RawMeshBulkData, int LodIndex)
diff --git a/Engine/Source/Runtime/MeshDescription/Private/MeshDescriptionBaseBulkData.cpp b/Engine/Source/Runtime/MeshDescription/Private/MeshDescriptionBaseBulkData.cpp
--- a/Engine/Source/Runtime/MeshDescription/Private/MeshDescriptionBaseBulkData.cpp
+++ b/Engine/Source/Runtime/MeshDescription/Private/MeshDescriptionBaseBulkData.cpp
@@ -34,6 +34,30 @@ void UMeshDescriptionBaseBulkData::Serialize(FArchive& Ar)
 }
 
 
+bool UMeshDescriptionBaseBulkData::IsEditorOnly() const
+{
+	return true;
+}
+
+
+bool UMeshDescriptionBaseBulkData::NeedsLoadForClient() const
+{
+	return false;
+}
+
+
+bool UMeshDescriptionBaseBulkData::NeedsLoadForServer() const
+{
+	return false;
+}
+
+
+bool UMeshDescriptionBaseBulkData::NeedsLoadForEditorGame() const
+{
+	return true;
+}
+
+
 UClass* UMeshDescriptionBaseBulkData::GetMeshDescriptionType() const
 {
 	return UMeshDescriptionBase::StaticClass();
diff --git a/Engine/Source/Runtime/MeshDescription/Public/MeshDescriptionBaseBulkData.h b/Engine/Source/Runtime/MeshDescription/Public/MeshDescriptionBaseBulkData.h
--- a/Engine/Source/Runtime/MeshDescription/Public/MeshDescriptionBaseBulkData.h
+++ b/Engine/Source/Runtime/MeshDescription/Public/MeshDescriptionBaseBulkData.h
@@ -25,8 +25,10 @@ class MESHDESCRIPTION_API UMeshDescriptionBaseBulkData : public UObject
 public:
 	UMeshDescriptionBaseBulkData();
 	virtual void Serialize(FArchive& Ar) override;
-	virtual bool NeedsLoadForClient() const override { return false; }
-	virtual bool NeedsLoadForServer() const override { return false; }
+	virtual bool IsEditorOnly() const override;
+	virtual bool NeedsLoadForClient() const override;
+	virtual bool NeedsLoadForServer() const override;
+	virtual bool NeedsLoadForEditorGame() const override;
 	virtual UClass* GetMeshDescriptionType() const;
 
 #if WITH_EDITORONLY_DATA
