Initial support for manipulating Static Mesh instances directly in the editor

Adds support for selection, gizmo manipulation, property editing, duplication, deletion, and undo/redo for Static Mesh instances within the main Level Editor, without having to enter an alternative editing mode.

Note: This is currently disabled via UE_ENABLE_SMINSTANCE_ELEMENTS

diff --git a/Engine/Source/Editor/EditorWidgets/Private/SObjectNameEditableTextBox.cpp b/Engine/Source/Editor/EditorWidgets/Private/SObjectNameEditableTextBox.cpp
--- a/Engine/Source/Editor/EditorWidgets/Private/SObjectNameEditableTextBox.cpp
+++ b/Engine/Source/Editor/EditorWidgets/Private/SObjectNameEditableTextBox.cpp
@@ -186,34 +186,9 @@ FText SObjectNameEditableTextBox::GetNameTooltipText() const
 
 EVisibility SObjectNameEditableTextBox::GetNameVisibility() const
 {
-	if (Objects.Num() == 1 && Objects[0].IsValid())
-	{
-		if (!IsReadOnly())
-		{
-			return EVisibility::Visible;
-		}
-		else if (Objects[0].Get()->IsA(AActor::StaticClass()))
-		{
-			return EVisibility::Visible;
-		}
-		else
-		{
-			return EVisibility::Collapsed;
-		}
-	}
-	else if (Objects.Num() > 1)
-	{
-		if (!IsReadOnly())
-		{
-			return EVisibility::Visible;
-		}
-		else
-		{
-			return EVisibility::Collapsed;
-		}
-	}
-
-	return EVisibility::Collapsed;
+	return ((Objects.Num() == 1 && Objects[0].IsValid()) || Objects.Num() > 1)
+		? EVisibility::Visible
+		: EVisibility::Collapsed;
 } 
 
 
diff --git a/Engine/Source/Editor/LevelEditor/Private/Elements/SMInstance/SMInstanceElementLevelEditorSelectionCustomization.cpp b/Engine/Source/Editor/LevelEditor/Private/Elements/SMInstance/SMInstanceElementLevelEditorSelectionCustomization.cpp
new file mode 100644
--- /dev/null
+++ b/Engine/Source/Editor/LevelEditor/Private/Elements/SMInstance/SMInstanceElementLevelEditorSelectionCustomization.cpp
@@ -0,0 +1,60 @@
+// Copyright Epic Games, Inc. All Rights Reserved.
+
+#include "Elements/SMInstance/SMInstanceElementLevelEditorSelectionCustomization.h"
+#include "Elements/SMInstance/SMInstanceElementData.h"
+#include "Components/InstancedStaticMeshComponent.h"
+
+#include "Elements/Framework/TypedElementList.h"
+#include "Elements/Framework/EngineElementsLibrary.h"
+
+DEFINE_LOG_CATEGORY_STATIC(LogSMInstanceLevelEditorSelection, Log, All);
+
+bool FSMInstanceElementLevelEditorSelectionCustomization::CanSelectElement(const TTypedElement<UTypedElementSelectionInterface>& InElementSelectionHandle, const FTypedElementSelectionOptions& InSelectionOptions)
+{
+	const FSMInstanceId SMInstance = SMInstanceElementDataUtil::GetSMInstanceFromHandleChecked(InElementSelectionHandle);
+	
+	// Bail if global selection is locked
+	return !GEdSelectionLock;
+}
+
+bool FSMInstanceElementLevelEditorSelectionCustomization::CanDeselectElement(const TTypedElement<UTypedElementSelectionInterface>& InElementSelectionHandle, const FTypedElementSelectionOptions& InSelectionOptions)
+{
+	const FSMInstanceId SMInstance = SMInstanceElementDataUtil::GetSMInstanceFromHandleChecked(InElementSelectionHandle);
+
+	// Bail if global selection is locked
+	return !GEdSelectionLock;
+}
+
+bool FSMInstanceElementLevelEditorSelectionCustomization::SelectElement(const TTypedElement<UTypedElementSelectionInterface>& InElementSelectionHandle, UTypedElementList* InSelectionSet, const FTypedElementSelectionOptions& InSelectionOptions)
+{
+	const FSMInstanceId SMInstance = SMInstanceElementDataUtil::GetSMInstanceFromHandleChecked(InElementSelectionHandle);
+
+	if (!InElementSelectionHandle.SelectElement(InSelectionSet, InSelectionOptions))
+	{
+		return false;
+	}
+
+	UE_LOG(LogSMInstanceLevelEditorSelection, Verbose, TEXT("Selected SMInstance: %s (%s), Index %d"), *SMInstance.ISMComponent->GetPathName(), *SMInstance.ISMComponent->GetClass()->GetName(), SMInstance.InstanceIndex);
+
+	// Update the internal selection state for viewport selection rendering
+	SMInstance.ISMComponent->SelectInstance(/*bIsSelected*/true, SMInstance.InstanceIndex);
+	
+	return true;
+}
+
+bool FSMInstanceElementLevelEditorSelectionCustomization::DeselectElement(const TTypedElement<UTypedElementSelectionInterface>& InElementSelectionHandle, UTypedElementList* InSelectionSet, const FTypedElementSelectionOptions& InSelectionOptions)
+{
+	const FSMInstanceId SMInstance = SMInstanceElementDataUtil::GetSMInstanceFromHandleChecked(InElementSelectionHandle);
+
+	if (!InElementSelectionHandle.DeselectElement(InSelectionSet, InSelectionOptions))
+	{
+		return false;
+	}
+
+	UE_LOG(LogSMInstanceLevelEditorSelection, Verbose, TEXT("Deselected SMInstance: %s (%s), Index %d"), *SMInstance.ISMComponent->GetPathName(), *SMInstance.ISMComponent->GetClass()->GetName(), SMInstance.InstanceIndex);
+
+	// Update the internal selection state for viewport selection rendering
+	SMInstance.ISMComponent->SelectInstance(/*bIsSelected*/false, SMInstance.InstanceIndex);
+
+	return true;
+}
diff --git a/Engine/Source/Editor/LevelEditor/Private/Elements/SMInstance/SMInstanceElementLevelEditorSelectionCustomization.h b/Engine/Source/Editor/LevelEditor/Private/Elements/SMInstance/SMInstanceElementLevelEditorSelectionCustomization.h
new file mode 100644
--- /dev/null
+++ b/Engine/Source/Editor/LevelEditor/Private/Elements/SMInstance/SMInstanceElementLevelEditorSelectionCustomization.h
@@ -0,0 +1,15 @@
+// Copyright Epic Games, Inc. All Rights Reserved.
+
+#pragma once
+
+#include "Elements/Framework/TypedElementSelectionSet.h"
+#include "Elements/Framework/TypedElementAssetEditorToolkitHostMixin.h"
+
+class FSMInstanceElementLevelEditorSelectionCustomization : public FTypedElementSelectionCustomization, public FTypedElementAssetEditorToolkitHostMixin
+{
+public:
+	virtual bool CanSelectElement(const TTypedElement<UTypedElementSelectionInterface>& InElementSelectionHandle, const FTypedElementSelectionOptions& InSelectionOptions) override;
+	virtual bool CanDeselectElement(const TTypedElement<UTypedElementSelectionInterface>& InElementSelectionHandle, const FTypedElementSelectionOptions& InSelectionOptions) override;
+	virtual bool SelectElement(const TTypedElement<UTypedElementSelectionInterface>& InElementSelectionHandle, UTypedElementList* InSelectionSet, const FTypedElementSelectionOptions& InSelectionOptions) override;
+	virtual bool DeselectElement(const TTypedElement<UTypedElementSelectionInterface>& InElementSelectionHandle, UTypedElementList* InSelectionSet, const FTypedElementSelectionOptions& InSelectionOptions) override;
+};
diff --git a/Engine/Source/Editor/LevelEditor/Private/SActorDetails.cpp b/Engine/Source/Editor/LevelEditor/Private/SActorDetails.cpp
--- a/Engine/Source/Editor/LevelEditor/Private/SActorDetails.cpp
+++ b/Engine/Source/Editor/LevelEditor/Private/SActorDetails.cpp
@@ -34,6 +34,7 @@
 #include "Elements/Framework/TypedElementRegistry.h"
 #include "Elements/Framework/TypedElementSelectionSet.h"
 #include "Elements/Framework/EngineElementsLibrary.h"
+#include "Elements/SMInstance/SMInstanceElementId.h" // TODO: Only included to define UE_ENABLE_SMINSTANCE_ELEMENTS
 
 class SActorDetailsUneditableComponentWarning : public SCompoundWidget
 {
@@ -101,7 +102,12 @@ void SActorDetails::Construct(const FArguments& InArgs, UTypedElementSelectionSe
 	DetailsViewArgs.bUpdatesFromSelection = true;
 	DetailsViewArgs.bLockable = true;
 	DetailsViewArgs.bAllowFavoriteSystem = true;
+#if UE_ENABLE_SMINSTANCE_ELEMENTS
+	// TODO: This allows the proxy objects used by SMInstance elements to show their name, but it does also change the behavior for actors/components
+	DetailsViewArgs.NameAreaSettings = FDetailsViewArgs::ObjectsUseNameArea;
+#else
 	DetailsViewArgs.NameAreaSettings = FDetailsViewArgs::ComponentsAndActorsUseNameArea;
+#endif
 	DetailsViewArgs.NotifyHook = GUnrealEd;
 	DetailsViewArgs.ViewIdentifier = TabIdentifier;
 	DetailsViewArgs.bCustomNameAreaLocation = true;
diff --git a/Engine/Source/Editor/LevelEditor/Private/SLevelEditor.cpp b/Engine/Source/Editor/LevelEditor/Private/SLevelEditor.cpp
--- a/Engine/Source/Editor/LevelEditor/Private/SLevelEditor.cpp
+++ b/Engine/Source/Editor/LevelEditor/Private/SLevelEditor.cpp
@@ -60,13 +60,15 @@
 #include "SourceCodeNavigation.h"
 #include "Editor/EnvironmentLightingViewer/Public/EnvironmentLightingModule.h"
 #include "Misc/MessageDialog.h"
+#include "ThumbnailRendering/ThumbnailManager.h"
 #include "Elements/Framework/TypedElementSelectionSet.h"
 #include "Elements/Framework/TypedElementCommonActions.h"
 #include "Elements/Actor/ActorElementLevelEditorSelectionCustomization.h"
 #include "Elements/Actor/ActorElementLevelEditorCommonActionsCustomization.h"
 #include "Elements/Component/ComponentElementLevelEditorSelectionCustomization.h"
 #include "Elements/Component/ComponentElementLevelEditorCommonActionsCustomization.h"
-#include "ThumbnailRendering/ThumbnailManager.h"
+#include "Elements/SMInstance/SMInstanceElementId.h"
+#include "Elements/SMInstance/SMInstanceElementLevelEditorSelectionCustomization.h"
 #include <atomic>
 
 #define LOCTEXT_NAMESPACE "SLevelEditor"
@@ -243,6 +245,11 @@ void SLevelEditor::Initialize( const TSharedRef<SDockTab>& OwnerTab, const TShar
 		ComponentCustomization->SetToolkitHost(this);
 		SelectedElements->RegisterInterfaceCustomizationByTypeName(NAME_Components, MoveTemp(ComponentCustomization));
 	}
+	{
+		TUniquePtr<FSMInstanceElementLevelEditorSelectionCustomization> SMInstanceCustomization = MakeUnique<FSMInstanceElementLevelEditorSelectionCustomization>();
+		SMInstanceCustomization->SetToolkitHost(this);
+		SelectedElements->RegisterInterfaceCustomizationByTypeName(NAME_SMInstance, MoveTemp(SMInstanceCustomization));
+	}
 
 	// Allow USelection to bridge to our selected element list
 	GUnrealEd->GetSelectedActors()->SetElementSelectionSet(SelectedElements);
diff --git a/Engine/Source/Editor/UnrealEd/Private/Elements/Framework/EditorElements.cpp b/Engine/Source/Editor/UnrealEd/Private/Elements/Framework/EditorElements.cpp
--- a/Engine/Source/Editor/UnrealEd/Private/Elements/Framework/EditorElements.cpp
+++ b/Engine/Source/Editor/UnrealEd/Private/Elements/Framework/EditorElements.cpp
@@ -3,18 +3,26 @@
 #include "Elements/Framework/EditorElements.h"
 #include "Elements/Framework/TypedElementRegistry.h"
 
+#include "Elements/Object/ObjectElementData.h"
 #include "Elements/Object/ObjectElementDetailsInterface.h"
 #include "Elements/Object/ObjectElementEditorSelectionInterface.h"
 
+#include "Elements/Actor/ActorElementData.h"
 #include "Elements/Actor/ActorElementDetailsInterface.h"
 #include "Elements/Actor/ActorElementEditorWorldInterface.h"
 #include "Elements/Actor/ActorElementEditorSelectionInterface.h"
 #include "Elements/Actor/ActorElementEditorAssetDataInterface.h"
 
+#include "Elements/Component/ComponentElementData.h"
 #include "Elements/Component/ComponentElementDetailsInterface.h"
 #include "Elements/Component/ComponentElementEditorWorldInterface.h"
 #include "Elements/Component/ComponentElementEditorSelectionInterface.h"
 
+#include "Elements/SMInstance/SMInstanceElementData.h"
+#include "Elements/SMInstance/SMInstanceElementDetailsInterface.h"
+#include "Elements/SMInstance/SMInstanceElementEditorWorldInterface.h"
+#include "Elements/SMInstance/SMInstanceElementEditorSelectionInterface.h"
+
 FSimpleMulticastDelegate OnRegisterEditorElementsDelegate;
 
 void RegisterEditorObjectElements()
@@ -44,11 +52,21 @@ void RegisterEditorComponentElements()
 	Registry->RegisterElementInterface<UTypedElementSelectionInterface>(NAME_Components, NewObject<UComponentElementEditorSelectionInterface>(), /*bAllowOverride*/true);
 }
 
+void RegisterEditorSMInstanceElements()
+{
+	UTypedElementRegistry* Registry = UTypedElementRegistry::GetInstance();
+
+	Registry->RegisterElementInterface<UTypedElementDetailsInterface>(NAME_SMInstance, NewObject<USMInstanceElementDetailsInterface>());
+	Registry->RegisterElementInterface<UTypedElementWorldInterface>(NAME_SMInstance, NewObject<USMInstanceElementEditorWorldInterface>(), /*bAllowOverride*/true);
+	Registry->RegisterElementInterface<UTypedElementSelectionInterface>(NAME_SMInstance, NewObject<USMInstanceElementEditorSelectionInterface>(), /*bAllowOverride*/true);
+}
+
 void RegisterEditorElements()
 {
 	RegisterEditorObjectElements();
 	RegisterEditorActorElements();
 	RegisterEditorComponentElements();
+	RegisterEditorSMInstanceElements();
 
 	OnRegisterEditorElementsDelegate.Broadcast();
 }
diff --git a/Engine/Source/Editor/UnrealEd/Private/Elements/SMInstance/SMInstanceElementDetailsInterface.cpp b/Engine/Source/Editor/UnrealEd/Private/Elements/SMInstance/SMInstanceElementDetailsInterface.cpp
new file mode 100644
--- /dev/null
+++ b/Engine/Source/Editor/UnrealEd/Private/Elements/SMInstance/SMInstanceElementDetailsInterface.cpp
@@ -0,0 +1,43 @@
+// Copyright Epic Games, Inc. All Rights Reserved.
+
+#include "Elements/SMInstance/SMInstanceElementDetailsInterface.h"
+#include "Elements/SMInstance/SMInstanceElementDetailsProxyObject.h"
+#include "Elements/SMInstance/SMInstanceElementData.h"
+#include "Components/InstancedStaticMeshComponent.h"
+
+class FSMInstanceTypedElementDetailsObject : public ITypedElementDetailsObject
+{
+public:
+	explicit FSMInstanceTypedElementDetailsObject(const FSMInstanceElementId& InSMInstanceElementId)
+		: InstanceProxyObject(NewObject<USMInstanceElementDetailsProxyObject>())
+	{
+		InstanceProxyObject->Initialize(InSMInstanceElementId);
+	}
+
+	~FSMInstanceTypedElementDetailsObject()
+	{
+		InstanceProxyObject->Shutdown();
+	}
+
+	virtual UObject* GetObject() override
+	{
+		return InstanceProxyObject;
+	}
+
+	virtual void AddReferencedObjects(FReferenceCollector& Collector) override
+	{
+		Collector.AddReferencedObject(InstanceProxyObject);
+	}
+
+private:
+	USMInstanceElementDetailsProxyObject* InstanceProxyObject = nullptr;
+};
+
+TUniquePtr<ITypedElementDetailsObject> USMInstanceElementDetailsInterface::GetDetailsObject(const FTypedElementHandle& InElementHandle)
+{
+	if (const FSMInstanceElementData* SMInstanceElement = InElementHandle.GetData<FSMInstanceElementData>())
+	{
+		return MakeUnique<FSMInstanceTypedElementDetailsObject>(SMInstanceElement->InstanceElementId);
+	}
+	return nullptr;
+}
diff --git a/Engine/Source/Editor/UnrealEd/Private/Elements/SMInstance/SMInstanceElementDetailsInterface.h b/Engine/Source/Editor/UnrealEd/Private/Elements/SMInstance/SMInstanceElementDetailsInterface.h
new file mode 100644
--- /dev/null
+++ b/Engine/Source/Editor/UnrealEd/Private/Elements/SMInstance/SMInstanceElementDetailsInterface.h
@@ -0,0 +1,15 @@
+// Copyright Epic Games, Inc. All Rights Reserved.
+
+#pragma once
+
+#include "Elements/Interfaces/TypedElementDetailsInterface.h"
+#include "SMInstanceElementDetailsInterface.generated.h"
+
+UCLASS()
+class USMInstanceElementDetailsInterface : public UTypedElementDetailsInterface
+{
+	GENERATED_BODY()
+
+public:
+	virtual TUniquePtr<ITypedElementDetailsObject> GetDetailsObject(const FTypedElementHandle& InElementHandle) override;
+};
diff --git a/Engine/Source/Editor/UnrealEd/Private/Elements/SMInstance/SMInstanceElementDetailsProxyObject.cpp b/Engine/Source/Editor/UnrealEd/Private/Elements/SMInstance/SMInstanceElementDetailsProxyObject.cpp
new file mode 100644
--- /dev/null
+++ b/Engine/Source/Editor/UnrealEd/Private/Elements/SMInstance/SMInstanceElementDetailsProxyObject.cpp
@@ -0,0 +1,82 @@
+// Copyright Epic Games, Inc. All Rights Reserved.
+
+#include "Elements/SMInstance/SMInstanceElementDetailsProxyObject.h"
+#include "Components/InstancedStaticMeshComponent.h"
+
+#include "UnrealEdGlobals.h"
+#include "Editor/UnrealEdEngine.h"
+
+void USMInstanceElementDetailsProxyObject::Initialize(const FSMInstanceElementId& InSMInstanceElementId)
+{
+	ISMComponent = InSMInstanceElementId.ISMComponent;
+	ISMInstanceId = InSMInstanceElementId.InstanceId;
+
+	TickHandle = FTicker::GetCoreTicker().AddTicker(TEXT("USMInstanceElementDetailsProxyObject"), 0.1f, [this](float)
+	{
+		SyncProxyStateFromInstance();
+		return true;
+	});
+	SyncProxyStateFromInstance();
+}
+
+void USMInstanceElementDetailsProxyObject::Shutdown()
+{
+	if (TickHandle.IsValid())
+	{
+		FTicker::GetCoreTicker().RemoveTicker(TickHandle);
+		TickHandle.Reset();
+	}
+
+	ISMComponent.Reset();
+	ISMInstanceId = 0;
+
+	SyncProxyStateFromInstance();
+}
+
+bool USMInstanceElementDetailsProxyObject::Modify(bool bAlwaysMarkDirty)
+{
+	if (FSMInstanceId InstanceId = GetInstanceId())
+	{
+		return InstanceId.ISMComponent->Modify(bAlwaysMarkDirty);
+	}
+
+	return false;
+}
+
+void USMInstanceElementDetailsProxyObject::PostEditChangeChainProperty(FPropertyChangedChainEvent& PropertyChangedEvent)
+{
+	if (PropertyChangedEvent.Property)
+	{
+		if (PropertyChangedEvent.Property->GetFName() == GET_MEMBER_NAME_CHECKED(USMInstanceElementDetailsProxyObject, Transform))
+		{
+			if (FSMInstanceId InstanceId = GetInstanceId())
+			{
+				// TODO: Need flag for local/world space, like FComponentTransformDetails
+				InstanceId.ISMComponent->UpdateInstanceTransform(InstanceId.InstanceIndex, Transform, /*bWorldSpace*/false, /*bMarkRenderStateDirty*/true);
+				
+				GUnrealEd->UpdatePivotLocationForSelection();
+				GUnrealEd->RedrawLevelEditingViewports();
+			}
+		}
+	}
+
+	Super::PostEditChangeChainProperty(PropertyChangedEvent);
+}
+
+void USMInstanceElementDetailsProxyObject::SyncProxyStateFromInstance()
+{
+	if (FSMInstanceId InstanceId = GetInstanceId())
+	{
+		// TODO: Need flag for local/world space, like FComponentTransformDetails
+		InstanceId.ISMComponent->GetInstanceTransform(InstanceId.InstanceIndex, Transform, /*bWorldSpace*/false);
+	}
+	else
+	{
+		Transform = FTransform::Identity;
+	}
+}
+
+FSMInstanceId USMInstanceElementDetailsProxyObject::GetInstanceId() const
+{
+	return FSMInstanceElementIdMap::Get().GetSMInstanceIdFromSMInstanceElementId(FSMInstanceElementId{ ISMComponent.Get(), ISMInstanceId });
+}
diff --git a/Engine/Source/Editor/UnrealEd/Private/Elements/SMInstance/SMInstanceElementDetailsProxyObject.h b/Engine/Source/Editor/UnrealEd/Private/Elements/SMInstance/SMInstanceElementDetailsProxyObject.h
new file mode 100644
--- /dev/null
+++ b/Engine/Source/Editor/UnrealEd/Private/Elements/SMInstance/SMInstanceElementDetailsProxyObject.h
@@ -0,0 +1,34 @@
+// Copyright Epic Games, Inc. All Rights Reserved.
+
+#pragma once
+
+#include "UObject/Object.h"
+#include "Elements/SMInstance/SMInstanceElementId.h"
+#include "SMInstanceElementDetailsProxyObject.generated.h"
+
+UCLASS(Transient)
+class USMInstanceElementDetailsProxyObject : public UObject
+{
+	GENERATED_BODY()
+
+public:
+	void Initialize(const FSMInstanceElementId& InSMInstanceElementId);
+	void Shutdown();
+
+	//~ UObject interface
+	virtual bool Modify(bool bAlwaysMarkDirty = true) override;
+	virtual void PostEditChangeChainProperty(FPropertyChangedChainEvent& PropertyChangedEvent) override;
+
+	UPROPERTY(EditAnywhere, Category="Transform", meta=(ShowOnlyInnerProperties))
+	FTransform Transform;
+
+private:
+	void SyncProxyStateFromInstance();
+
+	FSMInstanceId GetInstanceId() const;
+
+	TWeakObjectPtr<UInstancedStaticMeshComponent> ISMComponent;
+	uint64 ISMInstanceId = 0;
+
+	FDelegateHandle TickHandle;
+};
diff --git a/Engine/Source/Editor/UnrealEd/Private/Elements/SMInstance/SMInstanceElementEditorSelectionInterface.cpp b/Engine/Source/Editor/UnrealEd/Private/Elements/SMInstance/SMInstanceElementEditorSelectionInterface.cpp
new file mode 100644
--- /dev/null
+++ b/Engine/Source/Editor/UnrealEd/Private/Elements/SMInstance/SMInstanceElementEditorSelectionInterface.cpp
@@ -0,0 +1,79 @@
+// Copyright Epic Games, Inc. All Rights Reserved.
+
+#include "Elements/SMInstance/SMInstanceElementEditorSelectionInterface.h"
+#include "Elements/SMInstance/SMInstanceElementData.h"
+#include "Components/InstancedStaticMeshComponent.h"
+
+#include "Elements/Framework/TypedElementList.h"
+#include "Elements/Framework/EngineElementsLibrary.h"
+#include "Elements/Object/ObjectElementEditorSelectionInterface.h"
+
+class FSMInstanceElementTransactedElement : public ITypedElementTransactedElement
+{
+private:
+	virtual TUniquePtr<ITypedElementTransactedElement> CloneImpl() const override
+	{
+		return MakeUnique<FSMInstanceElementTransactedElement>(*this);
+	}
+
+	virtual FTypedElementHandle GetElementImpl() const override
+	{
+		FSMInstanceElementId SMInstanceElementId{ ISMComponentPtr.Get(/*bEvenIfPendingKill*/true), InstanceId };
+		return SMInstanceElementId
+			? UEngineElementsLibrary::AcquireEditorSMInstanceElementHandle(SMInstanceElementId)
+			: FTypedElementHandle();
+	}
+
+	virtual void SetElementImpl(const FTypedElementHandle& InElementHandle) override
+	{
+		const FSMInstanceElementData& SMInstanceElementData = InElementHandle.GetDataChecked<FSMInstanceElementData>();
+		ISMComponentPtr = SMInstanceElementData.InstanceElementId.ISMComponent;
+		InstanceId = SMInstanceElementData.InstanceElementId.InstanceId;
+	}
+
+	virtual void SerializeImpl(FArchive& InArchive) override
+	{
+		InArchive << ISMComponentPtr;
+		InArchive << InstanceId;
+	}
+
+	TWeakObjectPtr<UInstancedStaticMeshComponent> ISMComponentPtr;
+	uint64 InstanceId = 0;
+};
+
+bool USMInstanceElementEditorSelectionInterface::IsElementSelected(const FTypedElementHandle& InElementHandle, const UTypedElementList* InSelectionSet, const FTypedElementIsSelectedOptions& InSelectionOptions)
+{
+	if (const FSMInstanceId SMInstance = SMInstanceElementDataUtil::GetSMInstanceFromHandle(InElementHandle))
+	{
+		if (InSelectionSet->Num() == 0)
+		{
+			return false;
+		}
+
+		if (InSelectionSet->Contains(InElementHandle))
+		{
+			return true;
+		}
+
+		if (InSelectionOptions.AllowIndirect())
+		{
+			if (FTypedElementHandle ISMComponentElement = UEngineElementsLibrary::AcquireEditorComponentElementHandle(SMInstance.ISMComponent, /*bAllowCreate*/false))
+			{
+				return InSelectionSet->Contains(ISMComponentElement);
+			}
+		}
+	}
+
+	return false;
+}
+
+bool USMInstanceElementEditorSelectionInterface::ShouldPreventTransactions(const FTypedElementHandle& InElementHandle)
+{
+	const FSMInstanceId SMInstance = SMInstanceElementDataUtil::GetSMInstanceFromHandle(InElementHandle);
+	return SMInstance && UObjectElementEditorSelectionInterface::ShouldObjectPreventTransactions(SMInstance.ISMComponent);
+}
+
+TUniquePtr<ITypedElementTransactedElement> USMInstanceElementEditorSelectionInterface::CreateTransactedElementImpl()
+{
+	return MakeUnique<FSMInstanceElementTransactedElement>();
+}
diff --git a/Engine/Source/Editor/UnrealEd/Private/Elements/SMInstance/SMInstanceElementEditorSelectionInterface.h b/Engine/Source/Editor/UnrealEd/Private/Elements/SMInstance/SMInstanceElementEditorSelectionInterface.h
new file mode 100644
--- /dev/null
+++ b/Engine/Source/Editor/UnrealEd/Private/Elements/SMInstance/SMInstanceElementEditorSelectionInterface.h
@@ -0,0 +1,17 @@
+// Copyright Epic Games, Inc. All Rights Reserved.
+
+#pragma once
+
+#include "Elements/SMInstance/SMInstanceElementSelectionInterface.h"
+#include "SMInstanceElementEditorSelectionInterface.generated.h"
+
+UCLASS()
+class USMInstanceElementEditorSelectionInterface : public USMInstanceElementSelectionInterface
+{
+	GENERATED_BODY()
+
+public:
+	virtual bool IsElementSelected(const FTypedElementHandle& InElementHandle, const UTypedElementList* InSelectionSet, const FTypedElementIsSelectedOptions& InSelectionOptions) override;
+	virtual bool ShouldPreventTransactions(const FTypedElementHandle& InElementHandle) override;
+	virtual TUniquePtr<ITypedElementTransactedElement> CreateTransactedElementImpl() override;
+};
diff --git a/Engine/Source/Editor/UnrealEd/Private/Elements/SMInstance/SMInstanceElementEditorWorldInterface.cpp b/Engine/Source/Editor/UnrealEd/Private/Elements/SMInstance/SMInstanceElementEditorWorldInterface.cpp
new file mode 100644
--- /dev/null
+++ b/Engine/Source/Editor/UnrealEd/Private/Elements/SMInstance/SMInstanceElementEditorWorldInterface.cpp
@@ -0,0 +1,122 @@
+// Copyright Epic Games, Inc. All Rights Reserved.
+
+#include "Elements/SMInstance/SMInstanceElementEditorWorldInterface.h"
+#include "Elements/SMInstance/SMInstanceElementData.h"
+#include "Components/InstancedStaticMeshComponent.h"
+
+#include "Elements/Framework/EngineElementsLibrary.h"
+#include "Elements/Framework/TypedElementSelectionSet.h"
+
+bool USMInstanceElementEditorWorldInterface::CanDeleteElement(const FTypedElementHandle& InElementHandle)
+{
+	const FSMInstanceId SMInstance = SMInstanceElementDataUtil::GetSMInstanceFromHandle(InElementHandle);
+	return SMInstance && CanEditSMInstance(SMInstance);
+}
+
+bool USMInstanceElementEditorWorldInterface::DeleteElement(const FTypedElementHandle& InElementHandle, UWorld* InWorld, UTypedElementSelectionSet* InSelectionSet, const FTypedElementDeletionOptions& InDeletionOptions)
+{
+	if (const FSMInstanceId SMInstance = SMInstanceElementDataUtil::GetSMInstanceFromHandle(InElementHandle))
+	{
+		if (CanEditSMInstance(SMInstance))
+		{
+			SMInstance.ISMComponent->Modify();
+			InSelectionSet->DeselectElement(InElementHandle, FTypedElementSelectionOptions());
+			return SMInstance.ISMComponent->RemoveInstance(SMInstance.InstanceIndex);
+		}
+	}
+
+	return false;
+}
+
+bool USMInstanceElementEditorWorldInterface::DeleteElements(TArrayView<const FTypedElementHandle> InElementHandles, UWorld* InWorld, UTypedElementSelectionSet* InSelectionSet, const FTypedElementDeletionOptions& InDeletionOptions)
+{
+	const TArray<FSMInstanceId> SMInstancesToDelete = SMInstanceElementDataUtil::GetSMInstancesFromHandles(InElementHandles);
+
+	if (SMInstancesToDelete.Num() > 0)
+	{
+		// Batch by the ISM component
+		TMap<UInstancedStaticMeshComponent*, TArray<int32>> BatchedInstancesToDelete;
+		for (const FSMInstanceId& SMInstance : SMInstancesToDelete)
+		{
+			if (CanEditSMInstance(SMInstance))
+			{
+				TArray<int32>& InstanceIndices = BatchedInstancesToDelete.FindOrAdd(SMInstance.ISMComponent);
+				InstanceIndices.Add(SMInstance.InstanceIndex);
+			}
+		}
+
+		bool bDidDelete = false;
+
+		FTypedElementListLegacySyncScopedBatch LegacySyncBatch(InSelectionSet->GetElementList());
+		for (TTuple<UInstancedStaticMeshComponent*, TArray<int32>>& BatchedInstancesToDeletePair : BatchedInstancesToDelete)
+		{
+			BatchedInstancesToDeletePair.Key->Modify();
+
+			// Sort the instance indices in descending order to avoid invalidating resolved indices mid-delete
+			BatchedInstancesToDeletePair.Value.Sort(TGreater<int32>());
+			for (const int32 InstanceIndex : BatchedInstancesToDeletePair.Value)
+			{
+				InSelectionSet->DeselectElement(UEngineElementsLibrary::AcquireEditorSMInstanceElementHandle(FSMInstanceId{ BatchedInstancesToDeletePair.Key, InstanceIndex }), FTypedElementSelectionOptions());
+				bDidDelete |= BatchedInstancesToDeletePair.Key->RemoveInstance(InstanceIndex);
+			}
+		}
+
+		return bDidDelete;
+	}
+
+	return false;
+}
+
+FTypedElementHandle USMInstanceElementEditorWorldInterface::DuplicateElement(const FTypedElementHandle& InElementHandle, UWorld* InWorld, bool bOffsetLocations)
+{
+	if (const FSMInstanceId SMInstance = SMInstanceElementDataUtil::GetSMInstanceFromHandle(InElementHandle))
+	{
+		if (CanEditSMInstance(SMInstance))
+		{
+			FTransform NewInstanceTransform = FTransform::Identity;
+			SMInstance.ISMComponent->GetInstanceTransform(SMInstance.InstanceIndex, NewInstanceTransform);
+
+			const int32 NewInstanceIndex = SMInstance.ISMComponent->AddInstance(NewInstanceTransform);
+			return UEngineElementsLibrary::AcquireEditorSMInstanceElementHandle(FSMInstanceId{ SMInstance.ISMComponent, NewInstanceIndex });
+		}
+	}
+
+	return FTypedElementHandle();
+}
+
+void USMInstanceElementEditorWorldInterface::DuplicateElements(TArrayView<const FTypedElementHandle> InElementHandles, UWorld* InWorld, bool bOffsetLocations, TArray<FTypedElementHandle>& OutNewElements)
+{
+	const TArray<FSMInstanceId> SMInstancesToDuplicate = SMInstanceElementDataUtil::GetSMInstancesFromHandles(InElementHandles);
+
+	if (SMInstancesToDuplicate.Num() > 0)
+	{
+		// Batch by the ISM component
+		TMap<UInstancedStaticMeshComponent*, TArray<int32>> BatchedInstancesToDuplicate;
+		for (const FSMInstanceId& SMInstance : SMInstancesToDuplicate)
+		{
+			if (CanEditSMInstance(SMInstance))
+			{
+				TArray<int32>& InstanceIndices = BatchedInstancesToDuplicate.FindOrAdd(SMInstance.ISMComponent);
+				InstanceIndices.Add(SMInstance.InstanceIndex);
+			}
+		}
+
+		for (const TTuple<UInstancedStaticMeshComponent*, TArray<int32>>& BatchedInstancesToDuplicatePair : BatchedInstancesToDuplicate)
+		{
+			TArray<FTransform> NewInstanceTransforms;
+			NewInstanceTransforms.Reserve(BatchedInstancesToDuplicatePair.Value.Num());
+			for (const int32 InstanceIndex : BatchedInstancesToDuplicatePair.Value)
+			{
+				FTransform& NewInstanceTransform = NewInstanceTransforms.Add_GetRef(FTransform::Identity);
+				BatchedInstancesToDuplicatePair.Key->GetInstanceTransform(InstanceIndex, NewInstanceTransform);
+			}
+
+			const TArray<int32> NewInstanceIndices = BatchedInstancesToDuplicatePair.Key->AddInstances(NewInstanceTransforms, /*bShouldReturnIndices*/true);
+			OutNewElements.Reserve(OutNewElements.Num() + NewInstanceIndices.Num());
+			for (const int32 NewInstanceIndex : NewInstanceIndices)
+			{
+				OutNewElements.Add(UEngineElementsLibrary::AcquireEditorSMInstanceElementHandle(FSMInstanceId{ BatchedInstancesToDuplicatePair.Key, NewInstanceIndex }));
+			}
+		}
+	}
+}
diff --git a/Engine/Source/Editor/UnrealEd/Private/Elements/SMInstance/SMInstanceElementEditorWorldInterface.h b/Engine/Source/Editor/UnrealEd/Private/Elements/SMInstance/SMInstanceElementEditorWorldInterface.h
new file mode 100644
--- /dev/null
+++ b/Engine/Source/Editor/UnrealEd/Private/Elements/SMInstance/SMInstanceElementEditorWorldInterface.h
@@ -0,0 +1,19 @@
+// Copyright Epic Games, Inc. All Rights Reserved.
+
+#pragma once
+
+#include "Elements/SMInstance/SMInstanceElementWorldInterface.h"
+#include "SMInstanceElementEditorWorldInterface.generated.h"
+
+UCLASS()
+class USMInstanceElementEditorWorldInterface : public USMInstanceElementWorldInterface
+{
+	GENERATED_BODY()
+
+public:
+	virtual bool CanDeleteElement(const FTypedElementHandle& InElementHandle) override;
+	virtual bool DeleteElement(const FTypedElementHandle& InElementHandle, UWorld* InWorld, UTypedElementSelectionSet* InSelectionSet, const FTypedElementDeletionOptions& InDeletionOptions) override;
+	virtual bool DeleteElements(TArrayView<const FTypedElementHandle> InElementHandles, UWorld* InWorld, UTypedElementSelectionSet* InSelectionSet, const FTypedElementDeletionOptions& InDeletionOptions) override;
+	virtual FTypedElementHandle DuplicateElement(const FTypedElementHandle& InElementHandle, UWorld* InWorld, bool bOffsetLocations) override;
+	virtual void DuplicateElements(TArrayView<const FTypedElementHandle> InElementHandles, UWorld* InWorld, bool bOffsetLocations, TArray<FTypedElementHandle>& OutNewElements) override;
+};
diff --git a/Engine/Source/Editor/UnrealEd/Private/Selection.cpp b/Engine/Source/Editor/UnrealEd/Private/Selection.cpp
--- a/Engine/Source/Editor/UnrealEd/Private/Selection.cpp
+++ b/Engine/Source/Editor/UnrealEd/Private/Selection.cpp
@@ -187,13 +187,13 @@ UObject* USelection::GetObjectForElementHandle(const FTypedElementHandle& InElem
 {
 	check(InElementHandle && ElementSelectionSet);
 
-	TTypedElement<UTypedElementObjectInterface> ObjectElement = ElementSelectionSet->GetElementList()->GetElement<UTypedElementObjectInterface>(InElementHandle);
-	check(ObjectElement);
-
-	UObject* Object = ObjectElement.GetObject();
-	if (Object && SelectionElementBridge->IsValidObjectType(Object))
+	if (TTypedElement<UTypedElementObjectInterface> ObjectElement = ElementSelectionSet->GetElementList()->GetElement<UTypedElementObjectInterface>(InElementHandle))
 	{
-		return Object;
+		UObject* Object = ObjectElement.GetObject();
+		if (Object && SelectionElementBridge->IsValidObjectType(Object))
+		{
+			return Object;
+		}
 	}
 
 	return nullptr;
diff --git a/Engine/Source/Runtime/Engine/Classes/Components/InstancedStaticMeshComponent.h b/Engine/Source/Runtime/Engine/Classes/Components/InstancedStaticMeshComponent.h
--- a/Engine/Source/Runtime/Engine/Classes/Components/InstancedStaticMeshComponent.h
+++ b/Engine/Source/Runtime/Engine/Classes/Components/InstancedStaticMeshComponent.h
@@ -439,6 +439,7 @@ struct HInstancedStaticMeshInstance : public HHitProxy
 	HInstancedStaticMeshInstance(UInstancedStaticMeshComponent* InComponent, int32 InInstanceIndex) : HHitProxy(HPP_World), Component(InComponent), InstanceIndex(InInstanceIndex) {}
 
 	virtual void AddReferencedObjects(FReferenceCollector& Collector) override;
+	virtual FTypedElementHandle GetElementHandle() const override;
 
 	virtual EMouseCursor::Type GetMouseCursor() override
 	{
diff --git a/Engine/Source/Runtime/Engine/Classes/Engine/InstancedStaticMesh.h b/Engine/Source/Runtime/Engine/Classes/Engine/InstancedStaticMesh.h
--- a/Engine/Source/Runtime/Engine/Classes/Engine/InstancedStaticMesh.h
+++ b/Engine/Source/Runtime/Engine/Classes/Engine/InstancedStaticMesh.h
@@ -464,9 +464,16 @@ public:
 	,	StaticMesh(InComponent->GetStaticMesh())
 	,	InstancedRenderData(InComponent, InFeatureLevel)
 #if WITH_EDITOR
-	,	bHasSelectedInstances(InComponent->SelectedInstances.Num() > 0)
+	,	bHasSelectedInstances(false)
 #endif
 	{
+#if WITH_EDITOR
+		for (int32 InstanceIndex = 0; InstanceIndex < InComponent->SelectedInstances.Num() && !bHasSelectedInstances; ++InstanceIndex)
+		{
+			bHasSelectedInstances |= InComponent->SelectedInstances[InstanceIndex];
+		}
+#endif
+
 	#if !defined(GPUCULL_TODO)
 		bVFRequiresPrimitiveUniformBuffer = true;
 	#endif
diff --git a/Engine/Source/Runtime/Engine/Private/Elements/Framework/EngineElements.cpp b/Engine/Source/Runtime/Engine/Private/Elements/Framework/EngineElements.cpp
--- a/Engine/Source/Runtime/Engine/Private/Elements/Framework/EngineElements.cpp
+++ b/Engine/Source/Runtime/Engine/Private/Elements/Framework/EngineElements.cpp
@@ -19,6 +19,10 @@
 #include "Elements/Component/ComponentElementWorldInterface.h"
 #include "Elements/Component/ComponentElementSelectionInterface.h"
 
+#include "Elements/SMInstance/SMInstanceElementData.h"
+#include "Elements/SMInstance/SMInstanceElementWorldInterface.h"
+#include "Elements/SMInstance/SMInstanceElementSelectionInterface.h"
+
 #include "Modules/ModuleManager.h"
 
 FSimpleMulticastDelegate OnRegisterEngineElementsDelegate;
@@ -54,6 +58,15 @@ void RegisterEngineComponentElements()
 	Registry->RegisterElementInterface<UTypedElementSelectionInterface>(NAME_Components, NewObject<UComponentElementSelectionInterface>());
 }
 
+void RegisterEngineSMInstanceElements()
+{
+	UTypedElementRegistry* Registry = UTypedElementRegistry::GetInstance();
+
+	Registry->RegisterElementType<FSMInstanceElementData>(NAME_SMInstance);
+	Registry->RegisterElementInterface<UTypedElementWorldInterface>(NAME_SMInstance, NewObject<USMInstanceElementWorldInterface>());
+	Registry->RegisterElementInterface<UTypedElementSelectionInterface>(NAME_SMInstance, NewObject<USMInstanceElementSelectionInterface>());
+}
+
 void RegisterEngineElements()
 {
 	// Ensure the framework and base interfaces are also loaded
@@ -63,6 +76,7 @@ void RegisterEngineElements()
 	RegisterEngineObjectElements();
 	RegisterEngineActorElements();
 	RegisterEngineComponentElements();
+	RegisterEngineSMInstanceElements();
 
 	OnRegisterEngineElementsDelegate.Broadcast();
 }
diff --git a/Engine/Source/Runtime/Engine/Private/Elements/Framework/EngineElementsLibrary.cpp b/Engine/Source/Runtime/Engine/Private/Elements/Framework/EngineElementsLibrary.cpp
--- a/Engine/Source/Runtime/Engine/Private/Elements/Framework/EngineElementsLibrary.cpp
+++ b/Engine/Source/Runtime/Engine/Private/Elements/Framework/EngineElementsLibrary.cpp
@@ -15,6 +15,9 @@
 #include "Elements/Component/ComponentElementData.h"
 #include "Components/ActorComponent.h"
 
+#include "Elements/SMInstance/SMInstanceElementData.h"
+#include "Components/InstancedStaticMeshComponent.h"
+
 // The editor requires ref-counting for object replacement to function correctly
 static_assert(!WITH_EDITOR || UE_TYPED_ELEMENT_HAS_REFCOUNTING, "The editor requires that ref-counting is enabled for typed elements!");
 
@@ -133,6 +136,7 @@ void ReplaceEditorTypedElementHandles(TArray<FTypedElementHandle>& OutUpdatedEle
 TTypedElementOwnerStore<FObjectElementData, const UObject*> GObjectElementOwnerStore;
 TTypedElementOwnerStore<FActorElementData, const AActor*> GActorElementOwnerStore;
 TTypedElementOwnerStore<FComponentElementData, const UActorComponent*> GComponentElementOwnerStore;
+TTypedElementOwnerStore<FSMInstanceElementData, FSMInstanceElementId> GSMInstanceElementOwnerStore;
 #endif
 
 UEngineElementsLibrary::UEngineElementsLibrary()
@@ -146,6 +150,9 @@ UEngineElementsLibrary::UEngineElementsLibrary()
 		// UObject exists inside CoreUObject, so it cannot call through directly to clean-up any element handles that have been created
 		// Instead we rely on this GC hook to clean-up any element handles for unreachable objects prior to them being destroyed
 		FCoreUObjectDelegates::PreGarbageCollectConditionalBeginDestroy.AddStatic(&UEngineElementsLibrary::DestroyUnreachableEditorObjectElements);
+
+		// Static Mesh Instances are unmapped when removed from their owner component, so we must also destroy any corresponding element handle when that happens
+		FSMInstanceElementIdMap::Get().OnInstanceRemoved().AddStatic([](const FSMInstanceElementId& SMInstanceElementId, int32 InstanceIndex) { UEngineElementsLibrary::DestroyEditorSMInstanceElement(SMInstanceElementId); });
 	}
 #endif
 }
@@ -306,3 +313,93 @@ FTypedElementHandle UEngineElementsLibrary::AcquireEditorComponentElementHandle(
 	return EngineElementsLibraryUtil::AcquireEditorTypedElementHandle<UActorComponent, FComponentElementData>(Component, GComponentElementOwnerStore, &UEngineElementsLibrary::CreateComponentElement, bAllowCreate);
 }
 #endif
+
+TTypedElementOwner<FSMInstanceElementData> UEngineElementsLibrary::CreateSMInstanceElement(const FSMInstanceId& InSMInstanceId)
+{
+	const FSMInstanceElementId SMInstanceElementId = FSMInstanceElementIdMap::Get().GetSMInstanceElementIdFromSMInstanceId(InSMInstanceId);
+	checkf(SMInstanceElementId, TEXT("Static Mesh Instance Index failed to map to a valid Static Mesh Instance Element ID!"));
+	return CreateSMInstanceElement(SMInstanceElementId);
+}
+
+TTypedElementOwner<FSMInstanceElementData> UEngineElementsLibrary::CreateSMInstanceElement(const FSMInstanceElementId& InSMInstanceElementId)
+{
+	UTypedElementRegistry* Registry = UTypedElementRegistry::GetInstance();
+
+	TTypedElementOwner<FSMInstanceElementData> SMInstanceElement;
+	if (ensureAlways(Registry))
+	{
+#if UE_ENABLE_SMINSTANCE_ELEMENTS
+		SMInstanceElement = Registry->CreateElement<FSMInstanceElementData>(NAME_SMInstance);
+		if (SMInstanceElement)
+		{
+			SMInstanceElement.GetDataChecked().InstanceElementId = InSMInstanceElementId;
+		}
+#endif	// UE_ENABLE_SMINSTANCE_ELEMENTS
+	}
+
+	return SMInstanceElement;
+}
+
+void UEngineElementsLibrary::DestroySMInstanceElement(const FSMInstanceElementId& InSMInstanceElementId, TTypedElementOwner<FSMInstanceElementData>& InOutSMInstanceElement)
+{
+	if (InOutSMInstanceElement)
+	{
+		checkf(InOutSMInstanceElement.GetDataChecked().InstanceElementId == InSMInstanceElementId, TEXT("Static Mesh Instance element was not for this instance! %s"), *InSMInstanceElementId.ISMComponent->GetPathName());
+		UTypedElementRegistry::GetInstance()->DestroyElement(InOutSMInstanceElement);
+	}
+}
+
+#if WITH_EDITOR
+void UEngineElementsLibrary::CreateEditorSMInstanceElement(const FSMInstanceId& SMInstanceId)
+{
+	if (GIsEditor)
+	{
+		const FSMInstanceElementId SMInstanceElementId = FSMInstanceElementIdMap::Get().GetSMInstanceElementIdFromSMInstanceId(SMInstanceId);
+		checkf(SMInstanceElementId, TEXT("Static Mesh Instance Index failed to map to a valid Static Mesh Instance Element ID!"));
+		GSMInstanceElementOwnerStore.RegisterElementOwner(SMInstanceElementId, CreateSMInstanceElement(SMInstanceElementId));
+	}
+}
+
+void UEngineElementsLibrary::DestroyEditorSMInstanceElement(const FSMInstanceElementId& SMInstanceElementId)
+{
+	if (TTypedElementOwner<FSMInstanceElementData> SMInstanceElement = GSMInstanceElementOwnerStore.UnregisterElementOwner(SMInstanceElementId))
+	{
+		DestroySMInstanceElement(SMInstanceElementId, SMInstanceElement);
+	}
+}
+
+FTypedElementHandle UEngineElementsLibrary::AcquireEditorSMInstanceElementHandle(const UInstancedStaticMeshComponent* ISMComponent, const int32 InstanceIndex, const bool bAllowCreate)
+{
+	return AcquireEditorSMInstanceElementHandle(FSMInstanceId{ const_cast<UInstancedStaticMeshComponent*>(ISMComponent), InstanceIndex }, bAllowCreate);
+}
+
+FTypedElementHandle UEngineElementsLibrary::AcquireEditorSMInstanceElementHandle(const FSMInstanceId& SMInstanceId, const bool bAllowCreate)
+{
+	if (GIsEditor)
+	{
+		const FSMInstanceElementId SMInstanceElementId = FSMInstanceElementIdMap::Get().GetSMInstanceElementIdFromSMInstanceId(SMInstanceId, bAllowCreate);
+		checkf(!bAllowCreate || SMInstanceElementId, TEXT("Static Mesh Instance Index failed to map to a valid Static Mesh Instance Element ID!"));
+		return AcquireEditorSMInstanceElementHandle(SMInstanceElementId, bAllowCreate);
+	}
+
+	return FTypedElementHandle();
+}
+
+FTypedElementHandle UEngineElementsLibrary::AcquireEditorSMInstanceElementHandle(const FSMInstanceElementId& SMInstanceElementId, const bool bAllowCreate)
+{
+	if (GIsEditor)
+	{
+		TTypedElementOwnerScopedAccess<FSMInstanceElementData> SMInstanceElement = bAllowCreate
+			? GSMInstanceElementOwnerStore.FindOrRegisterElementOwner(SMInstanceElementId, [&SMInstanceElementId]() { return CreateSMInstanceElement(SMInstanceElementId); })
+			: GSMInstanceElementOwnerStore.FindElementOwner(SMInstanceElementId);
+
+		if (SMInstanceElement)
+		{
+			return SMInstanceElement->AcquireHandle();
+		}
+	}
+
+	return FTypedElementHandle();
+}
+#endif
+
diff --git a/Engine/Source/Runtime/Engine/Private/Elements/SMInstance/SMInstanceElementData.cpp b/Engine/Source/Runtime/Engine/Private/Elements/SMInstance/SMInstanceElementData.cpp
new file mode 100644
--- /dev/null
+++ b/Engine/Source/Runtime/Engine/Private/Elements/SMInstance/SMInstanceElementData.cpp
@@ -0,0 +1,25 @@
+// Copyright Epic Games, Inc. All Rights Reserved.
+
+#include "Elements/SMInstance/SMInstanceElementData.h"
+#include "Elements/Framework/TypedElementHandle.h"
+
+UE_DEFINE_TYPED_ELEMENT_DATA_RTTI(FSMInstanceElementData);
+
+namespace SMInstanceElementDataUtil
+{
+
+FSMInstanceId GetSMInstanceFromHandle(const FTypedElementHandle& InHandle, const bool bSilent)
+{
+	const FSMInstanceElementData* SMInstanceElement = InHandle.GetData<FSMInstanceElementData>(bSilent);
+	return SMInstanceElement ? FSMInstanceElementIdMap::Get().GetSMInstanceIdFromSMInstanceElementId(SMInstanceElement->InstanceElementId) : FSMInstanceId();
+}
+
+FSMInstanceId GetSMInstanceFromHandleChecked(const FTypedElementHandle& InHandle)
+{
+	const FSMInstanceElementData& SMInstanceElement = InHandle.GetDataChecked<FSMInstanceElementData>();
+	FSMInstanceId SMInstanceId = FSMInstanceElementIdMap::Get().GetSMInstanceIdFromSMInstanceElementId(SMInstanceElement.InstanceElementId);
+	checkf(SMInstanceId, TEXT("Static Mesh Instance Element ID failed to map to a valid Static Mesh Instance Index!"));
+	return SMInstanceId;
+}
+
+} // namespace SMInstanceElementDataUtil
diff --git a/Engine/Source/Runtime/Engine/Private/Elements/SMInstance/SMInstanceElementId.cpp b/Engine/Source/Runtime/Engine/Private/Elements/SMInstance/SMInstanceElementId.cpp
new file mode 100644
--- /dev/null
+++ b/Engine/Source/Runtime/Engine/Private/Elements/SMInstance/SMInstanceElementId.cpp
@@ -0,0 +1,313 @@
+// Copyright Epic Games, Inc. All Rights Reserved.
+
+#include "Elements/SMInstance/SMInstanceElementId.h"
+
+#include "Misc/ScopeLock.h"
+#include "UObject/Package.h"
+#include "Components/InstancedStaticMeshComponent.h"
+
+FSMInstanceElementIdMapEntry::FSMInstanceElementIdMapEntry(FSMInstanceElementIdMap* InOwner, UInstancedStaticMeshComponent* InComponent)
+	: Owner(InOwner)
+	, Component(InComponent)
+{
+	checkf(Owner, TEXT("FSMInstanceElementIdMapEntry must be constructed with a valid owner!"));
+	checkf(Component, TEXT("FSMInstanceElementIdMapEntry must be constructed with a valid component!"));
+
+#if WITH_EDITORONLY_DATA
+	Transactor = NewObject<USMInstanceElementIdMapTransactor>(GetTransientPackage(), NAME_None, RF_Transient | RF_Transactional);
+	Transactor->SetOwnerEntry(this);
+#endif	// WITH_EDITORONLY_DATA
+}
+
+FSMInstanceElementIdMapEntry::~FSMInstanceElementIdMapEntry()
+{
+#if WITH_EDITORONLY_DATA
+	Transactor->SetOwnerEntry(nullptr);
+#endif	// WITH_EDITORONLY_DATA
+}
+
+
+FSMInstanceElementIdMap& FSMInstanceElementIdMap::Get()
+{
+	static FSMInstanceElementIdMap Instance;
+	return Instance;
+}
+
+FSMInstanceElementIdMap::~FSMInstanceElementIdMap()
+{
+	UnregisterCallbacks();
+}
+
+FSMInstanceId FSMInstanceElementIdMap::GetSMInstanceIdFromSMInstanceElementId(const FSMInstanceElementId& InSMInstanceElementId)
+{
+	if (InSMInstanceElementId)
+	{
+		FScopeLock Lock(&ISMComponentsCS);
+
+		if (TSharedPtr<FSMInstanceElementIdMapEntry> ISMEntry = ISMComponents.FindRef(InSMInstanceElementId.ISMComponent))
+		{
+			if (const int32* InstanceIndexPtr = ISMEntry->InstanceIdToIndexMap.Find(InSMInstanceElementId.InstanceId))
+			{
+				return FSMInstanceId{ InSMInstanceElementId.ISMComponent, *InstanceIndexPtr };
+			}
+		}
+	}
+
+	return FSMInstanceId();
+}
+
+FSMInstanceElementId FSMInstanceElementIdMap::GetSMInstanceElementIdFromSMInstanceId(const FSMInstanceId& InSMInstanceId, const bool bAllowCreate)
+{
+	if (InSMInstanceId)
+	{
+		FScopeLock Lock(&ISMComponentsCS);
+
+		TSharedPtr<FSMInstanceElementIdMapEntry> ISMEntry = ISMComponents.FindRef(InSMInstanceId.ISMComponent);
+		if (ISMEntry)
+		{
+			if (const uint64* InstanceIdPtr = ISMEntry->InstanceIndexToIdMap.Find(InSMInstanceId.InstanceIndex))
+			{
+				return FSMInstanceElementId{ InSMInstanceId.ISMComponent, *InstanceIdPtr };
+			}
+		}
+
+		if (bAllowCreate)
+		{
+			if (!ISMEntry)
+			{
+				ISMEntry = ISMComponents.Add(InSMInstanceId.ISMComponent, MakeShared<FSMInstanceElementIdMapEntry>(this, InSMInstanceId.ISMComponent));
+				RegisterCallbacks();
+			}
+
+			checkf(ISMEntry->NextInstanceId < MAX_uint64, TEXT("Ran out of instance IDs for ISM component '%s'!"), *InSMInstanceId.ISMComponent->GetPathName());
+			uint64 InstanceId = ISMEntry->NextInstanceId++;
+
+			ISMEntry->InstanceIndexToIdMap.Add(InSMInstanceId.InstanceIndex, InstanceId);
+			ISMEntry->InstanceIdToIndexMap.Add(InstanceId, InSMInstanceId.InstanceIndex);
+
+			return FSMInstanceElementId{ InSMInstanceId.ISMComponent, InstanceId };
+		}
+	}
+
+	return FSMInstanceElementId();
+}
+
+void FSMInstanceElementIdMap::SerializeIdMappings(FSMInstanceElementIdMapEntry* InEntry, FArchive& Ar)
+{
+	FScopeLock Lock(&ISMComponentsCS);
+
+	TMap<uint64, int32> PreviousInstanceIdToIndexMap;
+	if (Ar.IsLoading() && (OnInstanceRemappedDelegate.IsBound() || OnInstanceRemovedDelegate.IsBound()))
+	{
+		// Take a copy of the current instance IDs, so that we can find out which elements have been unmapped or remapped post-load
+		PreviousInstanceIdToIndexMap = InEntry->InstanceIdToIndexMap;
+	}
+
+	// The two maps are mirrors of each other, so we only need to serialize one as we can restore the other again on load
+	Ar << InEntry->InstanceIdToIndexMap;
+	if (Ar.IsLoading())
+	{
+		InEntry->InstanceIndexToIdMap.Reset();
+		InEntry->InstanceIndexToIdMap.Reserve(InEntry->InstanceIdToIndexMap.Num());
+		for (const TTuple<uint64, int32>& InstanceIdToIndexPair : InEntry->InstanceIdToIndexMap)
+		{
+			InEntry->InstanceIndexToIdMap.Add(InstanceIdToIndexPair.Value, InstanceIdToIndexPair.Key);
+		}
+	}
+
+	if (Ar.IsLoading() && (OnInstanceRemappedDelegate.IsBound() || OnInstanceRemovedDelegate.IsBound()))
+	{
+		// Emit events for elements that were unmapped or remapped from this load
+		for (const TTuple<uint64, int32>& PreviousInstanceIdToIndexPair : PreviousInstanceIdToIndexMap)
+		{
+			if (const int32* InstanceIndexPtr = InEntry->InstanceIdToIndexMap.Find(PreviousInstanceIdToIndexPair.Key))
+			{
+				if (PreviousInstanceIdToIndexPair.Value != *InstanceIndexPtr)
+				{
+					OnInstanceRemappedDelegate.Broadcast(FSMInstanceElementId{ InEntry->Component, PreviousInstanceIdToIndexPair.Key }, PreviousInstanceIdToIndexPair.Value, *InstanceIndexPtr);
+				}
+			}
+			else
+			{
+				OnInstanceRemovedDelegate.Broadcast(FSMInstanceElementId{ InEntry->Component, PreviousInstanceIdToIndexPair.Key }, PreviousInstanceIdToIndexPair.Value);
+			}
+		}
+		if (OnInstanceRemappedDelegate.IsBound())
+		{
+			for (const TTuple<uint64, int32>& InstanceIdToIndexPair : InEntry->InstanceIdToIndexMap)
+			{
+				if (!PreviousInstanceIdToIndexMap.Contains(InstanceIdToIndexPair.Key))
+				{
+					OnInstanceRemappedDelegate.Broadcast(FSMInstanceElementId{ InEntry->Component, InstanceIdToIndexPair.Key }, INDEX_NONE, InstanceIdToIndexPair.Value);
+				}
+			}
+		}
+	}
+}
+
+void FSMInstanceElementIdMap::OnInstanceIndexUpdated(UInstancedStaticMeshComponent* InComponent, TArrayView<const FInstancedStaticMeshDelegates::FInstanceIndexUpdateData> InIndexUpdates)
+{
+	FScopeLock Lock(&ISMComponentsCS);
+
+	TSharedPtr<FSMInstanceElementIdMapEntry> ISMEntry = ISMComponents.FindRef(InComponent);
+	if (!ISMEntry)
+	{
+		// We don't process 'Added' updates below, so we can bail immediately if the component in question isn't registered
+		return;
+	}
+
+	for (const FInstancedStaticMeshDelegates::FInstanceIndexUpdateData& IndexUpdateData : InIndexUpdates)
+	{
+		switch (IndexUpdateData.Type)
+		{
+		case FInstancedStaticMeshDelegates::EInstanceIndexUpdateType::Added:
+			break; // We don't need to process 'Added' updates as they can't affect existing IDs
+
+		case FInstancedStaticMeshDelegates::EInstanceIndexUpdateType::Removed:
+			{
+				if (const uint64* InstanceIdPtr = ISMEntry->InstanceIndexToIdMap.Find(IndexUpdateData.Index))
+				{
+					const uint64 InstanceId = *InstanceIdPtr;
+
+					ISMEntry->InstanceIndexToIdMap.Remove(IndexUpdateData.Index);
+					ISMEntry->InstanceIdToIndexMap.Remove(InstanceId);
+
+					OnInstanceRemovedDelegate.Broadcast(FSMInstanceElementId{ InComponent, InstanceId }, IndexUpdateData.Index);
+				}
+				// Note: We can't remove things from the ISMComponents map until we get a 'Destroyed' notification for the owner ISM component, as the 
+				// per-component transactors may be referenced by the undo buffer, so the transactor needs to persist in order for redo to work correctly
+			}
+			break;
+
+		case FInstancedStaticMeshDelegates::EInstanceIndexUpdateType::Relocated:
+			{
+				if (const uint64* InstanceIdPtr = ISMEntry->InstanceIndexToIdMap.Find(IndexUpdateData.OldIndex))
+				{
+					const uint64 InstanceId = *InstanceIdPtr;
+
+					ISMEntry->InstanceIndexToIdMap.Remove(IndexUpdateData.OldIndex);
+					ISMEntry->InstanceIndexToIdMap.Add(IndexUpdateData.Index, InstanceId);
+					ISMEntry->InstanceIdToIndexMap[InstanceId] = IndexUpdateData.Index;
+
+					OnInstanceRemappedDelegate.Broadcast(FSMInstanceElementId{ InComponent, InstanceId }, IndexUpdateData.OldIndex, IndexUpdateData.Index);
+				}
+			}
+			break;
+
+		case FInstancedStaticMeshDelegates::EInstanceIndexUpdateType::Cleared:
+			{
+				// Take the current instance IDs, so that we can find notify the elements have been unmapped
+				TMap<uint64, int32> PreviousInstanceIdToIndexMap = MoveTemp(ISMEntry->InstanceIdToIndexMap);
+
+				// Note: We can't remove things from the ISMComponents map until we get a 'Destroyed' notification for the owner ISM component, as the 
+				// per-component transactors may be referenced by the undo buffer, so the transactor needs to persist in order for redo to work correctly
+				ISMEntry->InstanceIndexToIdMap.Reset();
+				ISMEntry->InstanceIdToIndexMap.Reset();
+				
+				if (OnInstanceRemovedDelegate.IsBound())
+				{
+					for (const TTuple<uint64, int32>& PreviousInstanceIdToIndexPair : PreviousInstanceIdToIndexMap)
+					{
+						OnInstanceRemovedDelegate.Broadcast(FSMInstanceElementId{ InComponent, PreviousInstanceIdToIndexPair.Key }, PreviousInstanceIdToIndexPair.Value);
+					}
+				}
+			}
+			break;
+
+		case FInstancedStaticMeshDelegates::EInstanceIndexUpdateType::Destroyed:
+			{
+				ISMComponents.Remove(InComponent);
+
+				if (OnInstanceRemovedDelegate.IsBound())
+				{
+					for (const TTuple<uint64, int32>& InstanceIdToIndexPair : ISMEntry->InstanceIdToIndexMap)
+					{
+						OnInstanceRemovedDelegate.Broadcast(FSMInstanceElementId{ InComponent, InstanceIdToIndexPair.Key }, InstanceIdToIndexPair.Value);
+					}
+				}
+			}
+			break;
+
+		default:
+			break;
+		}
+	}
+
+	if (ISMComponents.Num() == 0)
+	{
+		UnregisterCallbacks();
+	}
+}
+
+#if WITH_EDITOR
+void FSMInstanceElementIdMap::OnObjectModified(UObject* InObject)
+{
+	if (UInstancedStaticMeshComponent* Component = Cast<UInstancedStaticMeshComponent>(InObject))
+	{
+		FScopeLock Lock(&ISMComponentsCS);
+
+		if (TSharedPtr<FSMInstanceElementIdMapEntry> ISMEntry = ISMComponents.FindRef(Component))
+		{
+			// The component has been modified, so modify our transactor too to keep the maps in sync over an undo/redo
+			ISMEntry->Transactor->Modify();
+		}
+	}
+}
+#endif	// WITH_EDITOR
+
+void FSMInstanceElementIdMap::RegisterCallbacks()
+{
+	if (!OnInstanceIndexUpdatedHandle.IsValid())
+	{
+		OnInstanceIndexUpdatedHandle = FInstancedStaticMeshDelegates::OnInstanceIndexUpdated.AddRaw(this, &FSMInstanceElementIdMap::OnInstanceIndexUpdated);
+	}
+
+#if WITH_EDITOR
+	if (!OnObjectModifiedHandle.IsValid())
+	{
+		OnObjectModifiedHandle = FCoreUObjectDelegates::OnObjectModified.AddRaw(this, &FSMInstanceElementIdMap::OnObjectModified);
+	}
+#endif	// WITH_EDITOR
+}
+
+void FSMInstanceElementIdMap::UnregisterCallbacks()
+{
+	if (OnInstanceIndexUpdatedHandle.IsValid())
+	{
+		FInstancedStaticMeshDelegates::OnInstanceIndexUpdated.Remove(OnInstanceIndexUpdatedHandle);
+		OnInstanceIndexUpdatedHandle.Reset();
+	}
+
+#if WITH_EDITOR
+	if (OnObjectModifiedHandle.IsValid())
+	{
+		FCoreUObjectDelegates::OnObjectModified.Remove(OnObjectModifiedHandle);
+		OnObjectModifiedHandle.Reset();
+	}
+#endif	// WITH_EDITOR
+}
+
+void FSMInstanceElementIdMap::AddReferencedObjects(FReferenceCollector& Collector)
+{
+#if WITH_EDITORONLY_DATA
+	FScopeLock Lock(&ISMComponentsCS);
+
+	for (const TTuple<UInstancedStaticMeshComponent*, TSharedPtr<FSMInstanceElementIdMapEntry>>& ISMEntryPair : ISMComponents)
+	{
+		Collector.AddReferencedObject(ISMEntryPair.Value->Transactor);
+	}
+#endif	// WITH_EDITORONLY_DATA
+}
+
+
+#if WITH_EDITORONLY_DATA
+void USMInstanceElementIdMapTransactor::Serialize(FArchive& Ar)
+{
+	checkf(!Ar.IsPersistent(), TEXT("USMInstanceElementIdMapTransactor can only be serialized by transient archives!"));
+
+	if (OwnerEntry)
+	{
+		OwnerEntry->Owner->SerializeIdMappings(OwnerEntry, Ar);
+	}
+}
+#endif	// WITH_EDITORONLY_DATA
diff --git a/Engine/Source/Runtime/Engine/Private/Elements/SMInstance/SMInstanceElementSelectionInterface.cpp b/Engine/Source/Runtime/Engine/Private/Elements/SMInstance/SMInstanceElementSelectionInterface.cpp
new file mode 100644
--- /dev/null
+++ b/Engine/Source/Runtime/Engine/Private/Elements/SMInstance/SMInstanceElementSelectionInterface.cpp
@@ -0,0 +1,31 @@
+// Copyright Epic Games, Inc. All Rights Reserved.
+
+#include "Elements/SMInstance/SMInstanceElementSelectionInterface.h"
+#include "Elements/SMInstance/SMInstanceElementData.h"
+#include "Components/InstancedStaticMeshComponent.h"
+#include "Elements/Framework/TypedElementList.h"
+
+int32 USMInstanceElementSelectionInterface::GetNumSelectedSMInstances(const UTypedElementList* InCurrentSelection)
+{
+	int32 NumSelected = 0;
+	InCurrentSelection->ForEachElementHandle([&NumSelected](const FTypedElementHandle& InSelectedElement)
+	{
+		if (SMInstanceElementDataUtil::GetSMInstanceFromHandle(InSelectedElement, /*bSilent*/true))
+		{
+			++NumSelected;
+		}
+		return true;
+	});
+	return NumSelected;
+}
+
+bool USMInstanceElementSelectionInterface::HasSelectedSMInstances(const UTypedElementList* InCurrentSelection)
+{
+	bool bHasSelectedSMInstances = false;
+	InCurrentSelection->ForEachElementHandle([&bHasSelectedSMInstances](const FTypedElementHandle& InSelectedElement)
+	{
+		bHasSelectedSMInstances = !!SMInstanceElementDataUtil::GetSMInstanceFromHandle(InSelectedElement, /*bSilent*/true);
+		return !bHasSelectedSMInstances;
+	});
+	return bHasSelectedSMInstances;
+}
diff --git a/Engine/Source/Runtime/Engine/Private/Elements/SMInstance/SMInstanceElementWorldInterface.cpp b/Engine/Source/Runtime/Engine/Private/Elements/SMInstance/SMInstanceElementWorldInterface.cpp
new file mode 100644
--- /dev/null
+++ b/Engine/Source/Runtime/Engine/Private/Elements/SMInstance/SMInstanceElementWorldInterface.cpp
@@ -0,0 +1,78 @@
+// Copyright Epic Games, Inc. All Rights Reserved.
+
+#include "Elements/SMInstance/SMInstanceElementWorldInterface.h"
+#include "Elements/SMInstance/SMInstanceElementData.h"
+#include "Components/InstancedStaticMeshComponent.h"
+
+#include "Elements/Framework/EngineElementsLibrary.h"
+#include "Elements/Framework/TypedElementSelectionSet.h"
+
+bool USMInstanceElementWorldInterface::CanEditElement(const FTypedElementHandle& InElementHandle)
+{
+	const FSMInstanceId SMInstance = SMInstanceElementDataUtil::GetSMInstanceFromHandle(InElementHandle);
+	return SMInstance && CanEditSMInstance(SMInstance);
+}
+
+bool USMInstanceElementWorldInterface::IsTemplateElement(const FTypedElementHandle& InElementHandle)
+{
+	const FSMInstanceId SMInstance = SMInstanceElementDataUtil::GetSMInstanceFromHandle(InElementHandle);
+	return SMInstance && SMInstance.ISMComponent->IsTemplate();
+}
+
+ULevel* USMInstanceElementWorldInterface::GetOwnerLevel(const FTypedElementHandle& InElementHandle)
+{
+	if (const FSMInstanceId SMInstance = SMInstanceElementDataUtil::GetSMInstanceFromHandle(InElementHandle))
+	{
+		if (const AActor* ComponentOwner = SMInstance.ISMComponent->GetOwner())
+		{
+			return ComponentOwner->GetLevel();
+		}
+	}
+
+	return nullptr;
+}
+
+UWorld* USMInstanceElementWorldInterface::GetOwnerWorld(const FTypedElementHandle& InElementHandle)
+{
+	const FSMInstanceId SMInstance = SMInstanceElementDataUtil::GetSMInstanceFromHandle(InElementHandle);
+	return SMInstance ? SMInstance.ISMComponent->GetWorld() : nullptr;
+}
+
+bool USMInstanceElementWorldInterface::GetBounds(const FTypedElementHandle& InElementHandle, FBoxSphereBounds& OutBounds)
+{
+	if (const FSMInstanceId SMInstance = SMInstanceElementDataUtil::GetSMInstanceFromHandle(InElementHandle))
+	{
+		// TODO: ISM instance bounds?
+	}
+
+	return false;
+}
+
+bool USMInstanceElementWorldInterface::GetWorldTransform(const FTypedElementHandle& InElementHandle, FTransform& OutTransform)
+{
+	const FSMInstanceId SMInstance = SMInstanceElementDataUtil::GetSMInstanceFromHandle(InElementHandle);
+	return SMInstance && SMInstance.ISMComponent->GetInstanceTransform(SMInstance.InstanceIndex, OutTransform, /*bWorldSpace*/true);
+}
+
+bool USMInstanceElementWorldInterface::SetWorldTransform(const FTypedElementHandle& InElementHandle, const FTransform& InTransform)
+{
+	const FSMInstanceId SMInstance = SMInstanceElementDataUtil::GetSMInstanceFromHandle(InElementHandle);
+	return SMInstance && SMInstance.ISMComponent->UpdateInstanceTransform(SMInstance.InstanceIndex, InTransform, /*bWorldSpace*/true, /*bMarkRenderStateDirty*/true);
+}
+
+bool USMInstanceElementWorldInterface::GetRelativeTransform(const FTypedElementHandle& InElementHandle, FTransform& OutTransform)
+{
+	const FSMInstanceId SMInstance = SMInstanceElementDataUtil::GetSMInstanceFromHandle(InElementHandle);
+	return SMInstance && SMInstance.ISMComponent->GetInstanceTransform(SMInstance.InstanceIndex, OutTransform, /*bWorldSpace*/false);
+}
+
+bool USMInstanceElementWorldInterface::SetRelativeTransform(const FTypedElementHandle& InElementHandle, const FTransform& InTransform)
+{
+	const FSMInstanceId SMInstance = SMInstanceElementDataUtil::GetSMInstanceFromHandle(InElementHandle);
+	return SMInstance && SMInstance.ISMComponent->UpdateInstanceTransform(SMInstance.InstanceIndex, InTransform, /*bWorldSpace*/false, /*bMarkRenderStateDirty*/true);
+}
+
+bool USMInstanceElementWorldInterface::CanEditSMInstance(const FSMInstanceId& InSMInstanceId)
+{
+	return InSMInstanceId.ISMComponent->IsEditableWhenInherited();
+}
diff --git a/Engine/Source/Runtime/Engine/Private/HierarchicalInstancedStaticMesh.cpp b/Engine/Source/Runtime/Engine/Private/HierarchicalInstancedStaticMesh.cpp
--- a/Engine/Source/Runtime/Engine/Private/HierarchicalInstancedStaticMesh.cpp
+++ b/Engine/Source/Runtime/Engine/Private/HierarchicalInstancedStaticMesh.cpp
@@ -27,6 +27,7 @@
 #include "StaticMeshResources.h"
 #include "Components/HierarchicalInstancedStaticMeshComponent.h"
 #include "Engine/InstancedStaticMesh.h"
+#include "InstancedStaticMeshDelegates.h"
 #include "SceneManagement.h"
 #include "HAL/LowLevelMemTracker.h"
 #include "UObject/ReleaseObjectVersion.h"
@@ -1544,8 +1545,7 @@ void FHierarchicalStaticMeshSceneProxy::FillDynamicMeshElements(FMeshElementColl
 								}
 								else
 								{
-									NewBatchElement = new(MeshElement.Elements) FMeshBatchElement();
-									*NewBatchElement = MeshElement.Elements[0];
+									NewBatchElement = &MeshElement.Elements.Add_GetRef(MeshElement.Elements[0]);
 								}
 
 								const int32 InstanceOffset = RunArray[CurrentRun];
@@ -2237,6 +2237,8 @@ void UHierarchicalInstancedStaticMeshComponent::RemoveInstancesInternal(const in
 		}
 	#endif
 
+		const int32 LastInstanceIndex = PerInstanceSMData.Num();
+
 		// update the physics state
 		if (bPhysicsStateCreated)
 		{
@@ -2247,8 +2249,6 @@ void UHierarchicalInstancedStaticMeshComponent::RemoveInstancesInternal(const in
 				delete InstanceBodies[InstanceIndex];
 			}
 
-			int32 LastInstanceIndex = PerInstanceSMData.Num();
-
 			if (InstanceIndex == LastInstanceIndex)
 			{
 				// If we removed the last instance in the array we just need to remove it from the InstanceBodies array too.
@@ -2272,6 +2272,21 @@ void UHierarchicalInstancedStaticMeshComponent::RemoveInstancesInternal(const in
 				}
 			}
 		}
+
+		// Notify that these instances have been removed/relocated
+		if (FInstancedStaticMeshDelegates::OnInstanceIndexUpdated.IsBound())
+		{
+			TArray<FInstancedStaticMeshDelegates::FInstanceIndexUpdateData, TInlineAllocator<2>> IndexUpdates;
+
+			IndexUpdates.Add(FInstancedStaticMeshDelegates::FInstanceIndexUpdateData{ FInstancedStaticMeshDelegates::EInstanceIndexUpdateType::Removed, InstanceIndex });
+			if (InstanceIndex != LastInstanceIndex)
+			{
+				// HISMs use swap remove, so the last index has been moved to the spot we removed from
+				IndexUpdates.Add(FInstancedStaticMeshDelegates::FInstanceIndexUpdateData{ FInstancedStaticMeshDelegates::EInstanceIndexUpdateType::Relocated, InstanceIndex, LastInstanceIndex });
+			}
+
+			FInstancedStaticMeshDelegates::OnInstanceIndexUpdated.Broadcast(this, IndexUpdates);
+		}
 	}
 
 	PerInstanceSMData.Shrink();
@@ -2627,6 +2642,8 @@ void UHierarchicalInstancedStaticMeshComponent::ClearInstances()
 		}
 	}
 
+	const int32 PrevNumInstances = GetInstanceCount();
+
 	// Clear all the per-instance data
 	PerInstanceSMData.Empty();
 	PerInstanceSMCustomData.Empty();
@@ -2640,6 +2657,13 @@ void UHierarchicalInstancedStaticMeshComponent::ClearInstances()
 
 	MarkRenderStateDirty();
 
+	// Notify that these instances have been cleared
+	if (FInstancedStaticMeshDelegates::OnInstanceIndexUpdated.IsBound())
+	{
+		FInstancedStaticMeshDelegates::FInstanceIndexUpdateData IndexUpdate{ FInstancedStaticMeshDelegates::EInstanceIndexUpdateType::Cleared, PrevNumInstances - 1 };
+		FInstancedStaticMeshDelegates::OnInstanceIndexUpdated.Broadcast(this, MakeArrayView(&IndexUpdate, 1));
+	}
+
 	FNavigationSystem::UpdateComponentData(*this);
 }
 
diff --git a/Engine/Source/Runtime/Engine/Private/InstancedStaticMesh.cpp b/Engine/Source/Runtime/Engine/Private/InstancedStaticMesh.cpp
--- a/Engine/Source/Runtime/Engine/Private/InstancedStaticMesh.cpp
+++ b/Engine/Source/Runtime/Engine/Private/InstancedStaticMesh.cpp
@@ -5,6 +5,7 @@
 =============================================================================*/
 
 #include "Engine/InstancedStaticMesh.h"
+#include "InstancedStaticMeshDelegates.h"
 #include "AI/NavigationSystemBase.h"
 #include "Engine/MapBuildDataRegistry.h"
 #include "Components/LightComponent.h"
@@ -32,6 +33,9 @@
 #include "MeshMaterialShader.h"
 #include "ProfilingDebugging/LoadTimeTracker.h"
 
+#include "Elements/Framework/EngineElementsLibrary.h"
+#include "Elements/SMInstance/SMInstanceElementData.h"
+
 #if RHI_RAYTRACING
 #include "RayTracingInstance.h"
 #endif
@@ -119,12 +123,25 @@ TGlobalResource<FDummyFloatBuffer> GDummyFloatBuffer;
 
 extern int32 GRenderNaniteMeshes;
 
+FInstancedStaticMeshDelegates::FOnInstanceIndexUpdated FInstancedStaticMeshDelegates::OnInstanceIndexUpdated;
+
 /** InstancedStaticMeshInstance hit proxy */
 void HInstancedStaticMeshInstance::AddReferencedObjects(FReferenceCollector& Collector)
 {
 	Collector.AddReferencedObject(Component);
 }
 
+FTypedElementHandle HInstancedStaticMeshInstance::GetElementHandle() const
+{
+#if WITH_EDITOR
+	if (Component)
+	{
+		return UEngineElementsLibrary::AcquireEditorSMInstanceElementHandle(FSMInstanceId{ Component, InstanceIndex });
+	}
+#endif	// WITH_EDITOR
+	return FTypedElementHandle();
+}
+
 FInstanceUpdateCmdBuffer::FInstanceUpdateCmdBuffer()
 	: NumAdds(0)
 	, NumEdits(0)
@@ -2292,7 +2309,7 @@ int32 UInstancedStaticMeshComponent::AddInstanceInternal(int32 InstanceIndex, FI
 
 	if (NewInstanceData == nullptr)
 	{
-		NewInstanceData = new(PerInstanceSMData) FInstancedStaticMeshInstanceData();
+		NewInstanceData = &PerInstanceSMData.AddDefaulted_GetRef();
 	}
 
 	SetupNewInstanceData(*NewInstanceData, InstanceIndex, InstanceTransform);
@@ -2309,6 +2326,12 @@ int32 UInstancedStaticMeshComponent::AddInstanceInternal(int32 InstanceIndex, FI
 
 	PartialNavigationUpdate(InstanceIndex);
 
+	if (FInstancedStaticMeshDelegates::OnInstanceIndexUpdated.IsBound())
+	{
+		FInstancedStaticMeshDelegates::FInstanceIndexUpdateData IndexUpdate{ FInstancedStaticMeshDelegates::EInstanceIndexUpdateType::Added, InstanceIndex };
+		FInstancedStaticMeshDelegates::OnInstanceIndexUpdated.Broadcast(this, MakeArrayView(&IndexUpdate, 1));
+	}
+
 	InstanceUpdateCmdBuffer.Edit();
 	MarkRenderStateDirty();
 
@@ -2339,9 +2362,9 @@ TArray<int32> UInstancedStaticMeshComponent::AddInstancesInternal(int32 Count, c
 
 	for (int32 i = 0; i < Count; ++i)
 	{
-		FInstancedStaticMeshInstanceData* NewInstanceData = new(PerInstanceSMData) FInstancedStaticMeshInstanceData();
+		FInstancedStaticMeshInstanceData& NewInstanceData = PerInstanceSMData.AddDefaulted_GetRef();
 
-		SetupNewInstanceData(*NewInstanceData, InstanceIndex, InstanceTransforms[i]);
+		SetupNewInstanceData(NewInstanceData, InstanceIndex, InstanceTransforms[i]);
 
 		if (bShouldReturnIndices)
 		{
@@ -2353,6 +2376,12 @@ TArray<int32> UInstancedStaticMeshComponent::AddInstancesInternal(int32 Count, c
 			PartialNavigationUpdate(InstanceIndex);
 		}
 
+		if (FInstancedStaticMeshDelegates::OnInstanceIndexUpdated.IsBound())
+		{
+			FInstancedStaticMeshDelegates::FInstanceIndexUpdateData IndexUpdate{ FInstancedStaticMeshDelegates::EInstanceIndexUpdateType::Added, InstanceIndex };
+			FInstancedStaticMeshDelegates::OnInstanceIndexUpdated.Broadcast(this, MakeArrayView(&IndexUpdate, 1));
+		}
+
 		++InstanceIndex;
 	}
 
@@ -2432,12 +2461,12 @@ bool UInstancedStaticMeshComponent::SetCustomData(int32 InstanceIndex, const TAr
 
 bool UInstancedStaticMeshComponent::RemoveInstanceInternal(int32 InstanceIndex, bool InstanceAlreadyRemoved)
 {
-	// Request navigation update
-	PartialNavigationUpdate(InstanceIndex);
-
 	// remove instance
 	if (!InstanceAlreadyRemoved && PerInstanceSMData.IsValidIndex(InstanceIndex))
 	{
+		// Request navigation update
+		PartialNavigationUpdate(InstanceIndex);
+
 		PerInstanceSMData.RemoveAt(InstanceIndex);
 		PerInstanceSMCustomData.RemoveAt(InstanceIndex * NumCustomDataFloats, NumCustomDataFloats);
 	}
@@ -2469,6 +2498,22 @@ bool UInstancedStaticMeshComponent::RemoveInstanceInternal(int32 InstanceIndex,
 		}
 	}
 
+	// Notify that these instances have been removed/relocated
+	if (FInstancedStaticMeshDelegates::OnInstanceIndexUpdated.IsBound())
+	{
+		TArray<FInstancedStaticMeshDelegates::FInstanceIndexUpdateData, TInlineAllocator<128>> IndexUpdates;
+		IndexUpdates.Reserve(1 + (PerInstanceSMData.Num() - InstanceIndex));
+
+		IndexUpdates.Add(FInstancedStaticMeshDelegates::FInstanceIndexUpdateData{ FInstancedStaticMeshDelegates::EInstanceIndexUpdateType::Removed, InstanceIndex });
+		for (int32 MovedInstanceIndex = InstanceIndex; MovedInstanceIndex < PerInstanceSMData.Num(); ++MovedInstanceIndex)
+		{
+			// ISMs use standard remove, so each instance above our removal point is shuffled down by 1
+			IndexUpdates.Add(FInstancedStaticMeshDelegates::FInstanceIndexUpdateData{ FInstancedStaticMeshDelegates::EInstanceIndexUpdateType::Relocated, MovedInstanceIndex, MovedInstanceIndex + 1 });
+		}
+
+		FInstancedStaticMeshDelegates::OnInstanceIndexUpdated.Broadcast(this, IndexUpdates);
+	}
+
 	// Force recreation of the render data
 	InstanceUpdateCmdBuffer.Edit();
 	MarkRenderStateDirty();
@@ -2842,6 +2887,8 @@ void UInstancedStaticMeshComponent::GetStreamingRenderAssetInfo(FStreamingTextur
 
 void UInstancedStaticMeshComponent::ClearInstances()
 {
+	const int32 PrevNumInstances = GetInstanceCount();
+
 	// Clear all the per-instance data
 	PerInstanceSMData.Empty();
 	PerInstanceSMCustomData.Empty();
@@ -2858,6 +2905,13 @@ void UInstancedStaticMeshComponent::ClearInstances()
 	InstanceUpdateCmdBuffer.Edit();
 	MarkRenderStateDirty();
 
+	// Notify that these instances have been cleared
+	if (FInstancedStaticMeshDelegates::OnInstanceIndexUpdated.IsBound())
+	{
+		FInstancedStaticMeshDelegates::FInstanceIndexUpdateData IndexUpdate{ FInstancedStaticMeshDelegates::EInstanceIndexUpdateType::Cleared, PrevNumInstances - 1 };
+		FInstancedStaticMeshDelegates::OnInstanceIndexUpdated.Broadcast(this, MakeArrayView(&IndexUpdate, 1));
+	}
+
 	FNavigationSystem::UpdateComponentData(*this);
 }
 
@@ -3158,7 +3212,15 @@ void UInstancedStaticMeshComponent::GetResourceSizeEx(FResourceSizeEx& Cumulativ
 
 void UInstancedStaticMeshComponent::BeginDestroy()
 {
+	// Notify that these instances have been cleared due to the destroy
+	if (FInstancedStaticMeshDelegates::OnInstanceIndexUpdated.IsBound())
+	{
+		FInstancedStaticMeshDelegates::FInstanceIndexUpdateData IndexUpdate{ FInstancedStaticMeshDelegates::EInstanceIndexUpdateType::Destroyed, GetInstanceCount() - 1 };
+		FInstancedStaticMeshDelegates::OnInstanceIndexUpdated.Broadcast(this, MakeArrayView(&IndexUpdate, 1));
+	}
+
 	ReleasePerInstanceRenderData();
+
 	Super::BeginDestroy();
 }
 
diff --git a/Engine/Source/Runtime/Engine/Public/Elements/Framework/EngineElementsLibrary.h b/Engine/Source/Runtime/Engine/Public/Elements/Framework/EngineElementsLibrary.h
--- a/Engine/Source/Runtime/Engine/Public/Elements/Framework/EngineElementsLibrary.h
+++ b/Engine/Source/Runtime/Engine/Public/Elements/Framework/EngineElementsLibrary.h
@@ -15,6 +15,12 @@ struct FActorElementData;
 class UActorComponent;
 struct FComponentElementData;
 
+class UInstancedStaticMeshComponent;
+struct FSMInstanceElementData;
+
+struct FSMInstanceId;
+struct FSMInstanceElementId;
+
 class UWorld;
 class UTypedElementList;
 
@@ -53,6 +59,18 @@ public:
 	static FTypedElementHandle AcquireEditorComponentElementHandle(const UActorComponent* Component, const bool bAllowCreate = true);
 #endif
 
+	static TTypedElementOwner<FSMInstanceElementData> CreateSMInstanceElement(const FSMInstanceId& InSMInstanceId);
+	static TTypedElementOwner<FSMInstanceElementData> CreateSMInstanceElement(const FSMInstanceElementId& InSMInstanceElementId);
+	static void DestroySMInstanceElement(const FSMInstanceElementId& InSMInstanceElementId, TTypedElementOwner<FSMInstanceElementData>& InOutSMInstanceElement);
+#if WITH_EDITOR
+	static void CreateEditorSMInstanceElement(const FSMInstanceId& SMInstanceId);
+	static void DestroyEditorSMInstanceElement(const FSMInstanceElementId& SMInstanceElementId);
+	UFUNCTION(BlueprintCallable, BlueprintPure=false, Category="TypedElementFramework|StaticMeshInstance", meta=(ScriptMethod="AcquireEditorElementHandle"))
+	static FTypedElementHandle AcquireEditorSMInstanceElementHandle(const UInstancedStaticMeshComponent* ISMComponent, const int32 InstanceIndex, const bool bAllowCreate = true);
+	static FTypedElementHandle AcquireEditorSMInstanceElementHandle(const FSMInstanceId& SMInstanceId, const bool bAllowCreate = true);
+	static FTypedElementHandle AcquireEditorSMInstanceElementHandle(const FSMInstanceElementId& SMInstanceElementId, const bool bAllowCreate = true);
+#endif
+
 private:
 #if WITH_EDITOR
 	static void DestroyUnreachableEditorObjectElements();
diff --git a/Engine/Source/Runtime/Engine/Public/Elements/SMInstance/SMInstanceElementData.h b/Engine/Source/Runtime/Engine/Public/Elements/SMInstance/SMInstanceElementData.h
new file mode 100644
--- /dev/null
+++ b/Engine/Source/Runtime/Engine/Public/Elements/SMInstance/SMInstanceElementData.h
@@ -0,0 +1,78 @@
+// Copyright Epic Games, Inc. All Rights Reserved.
+
+#pragma once
+
+#include "Elements/Framework/TypedElementData.h"
+#include "Elements/SMInstance/SMInstanceElementId.h"
+
+struct FTypedElementHandle;
+
+/**
+ * Element data that represents a specific instance within an ISM.
+ */
+struct ENGINE_API FSMInstanceElementData
+{
+	UE_DECLARE_TYPED_ELEMENT_DATA_RTTI(FSMInstanceElementData);
+
+	FSMInstanceElementId InstanceElementId;
+};
+
+namespace SMInstanceElementDataUtil
+{
+
+/**
+ * Attempt to get the static mesh instance ID from the given element handle.
+ * @note This is not typically something you'd want to use outside of data access within an interface implementation.
+ * @return The static mesh instance ID if the element handle contains FSMInstanceElementData, otherwise an invalid ID.
+ */
+ENGINE_API FSMInstanceId GetSMInstanceFromHandle(const FTypedElementHandle& InHandle, const bool bSilent = false);
+
+/**
+ * Attempt to get the static mesh instance ID from the given element handle, asserting if the element handle doesn't contain FSMInstanceElementData.
+ * @note This is not typically something you'd want to use outside of data access within an interface implementation.
+ * @return The static mesh instance ID.
+ */
+ENGINE_API FSMInstanceId GetSMInstanceFromHandleChecked(const FTypedElementHandle& InHandle);
+
+/**
+ * Attempt to get the static mesh instance IDs from the given element handles.
+ * @note This is not typically something you'd want to use outside of data access within an interface implementation.
+ * @return The static mesh instance IDs of any element handles that contain FSMInstanceElementData, skipping any that don't.
+ */
+template <typename ElementHandleType>
+TArray<FSMInstanceId> GetSMInstancesFromHandles(TArrayView<const ElementHandleType> InHandles, const bool bSilent = false)
+{
+	TArray<FSMInstanceId> SMInstanceIds;
+	SMInstanceIds.Reserve(InHandles.Num());
+
+	for (const FTypedElementHandle& Handle : InHandles)
+	{
+		if (FSMInstanceId SMInstanceId = GetSMInstanceFromHandle(Handle, bSilent))
+		{
+			SMInstanceIds.Add(SMInstanceId);
+		}
+	}
+
+	return SMInstanceIds;
+}
+
+/**
+ * Attempt to get the static mesh instance IDs from the given element handles, asserting if any element handle doesn't contain FSMInstanceElementData.
+ * @note This is not typically something you'd want to use outside of data access within an interface implementation.
+ * @return The static mesh instance IDs.
+ */
+template <typename ElementHandleType>
+TArray<FSMInstanceId> GetSMInstancesFromHandlesChecked(TArrayView<const ElementHandleType> InHandles)
+{
+	TArray<FSMInstanceId> SMInstanceIds;
+	SMInstanceIds.Reserve(InHandles.Num());
+
+	for (const FTypedElementHandle& Handle : InHandles)
+	{
+		SMInstanceIds.Add(GetSMInstanceFromHandleChecked(Handle));
+	}
+
+	return SMInstanceIds;
+}
+
+} // namespace SMInstanceElementDataUtil
diff --git a/Engine/Source/Runtime/Engine/Public/Elements/SMInstance/SMInstanceElementId.h b/Engine/Source/Runtime/Engine/Public/Elements/SMInstance/SMInstanceElementId.h
new file mode 100644
--- /dev/null
+++ b/Engine/Source/Runtime/Engine/Public/Elements/SMInstance/SMInstanceElementId.h
@@ -0,0 +1,200 @@
+// Copyright Epic Games, Inc. All Rights Reserved.
+
+#pragma once
+
+#include "CoreMinimal.h"
+#include "UObject/GCObject.h"
+#include "InstancedStaticMeshDelegates.h"
+#include "SMInstanceElementId.generated.h"
+
+#define UE_ENABLE_SMINSTANCE_ELEMENTS (0)
+
+class FSMInstanceElementIdMap;
+class UInstancedStaticMeshComponent;
+class USMInstanceElementIdMapTransactor;
+
+static const FName NAME_SMInstance = "SMInstance";
+
+/**
+ * ID for a specific instance within an ISM, mapped from its instance index.
+ */
+struct FSMInstanceId
+{
+	explicit operator bool() const
+	{
+		return ISMComponent != nullptr
+			&& InstanceIndex != INDEX_NONE;
+	}
+
+	bool operator==(const FSMInstanceId& InRHS) const
+	{
+		return ISMComponent == InRHS.ISMComponent
+			&& InstanceIndex == InRHS.InstanceIndex;
+	}
+
+	bool operator!=(const FSMInstanceId& InRHS) const
+	{
+		return !(*this == InRHS);
+	}
+
+	friend inline uint32 GetTypeHash(const FSMInstanceId& InId)
+	{
+		return HashCombine(GetTypeHash(InId.InstanceIndex), GetTypeHash(InId.ISMComponent));
+	}
+
+	UInstancedStaticMeshComponent* ISMComponent = nullptr;
+	int32 InstanceIndex = INDEX_NONE;
+};
+
+/**
+ * ID for a specific instance within an ISM, mapped from the instance ID used by typed elements.
+ * @note The specific implementation of instance IDs for typed elements is considered private and may change without warning!
+ *       Use the FSMInstanceElementIdMap functions to convert this struct to/from a FSMInstanceId rather than access this data directly!
+ */
+struct FSMInstanceElementId
+{
+	explicit operator bool() const
+	{
+		return ISMComponent != nullptr
+			&& InstanceId != 0;
+	}
+
+	bool operator==(const FSMInstanceElementId& InRHS) const
+	{
+		return ISMComponent == InRHS.ISMComponent
+			&& InstanceId == InRHS.InstanceId;
+	}
+
+	bool operator!=(const FSMInstanceElementId& InRHS) const
+	{
+		return !(*this == InRHS);
+	}
+
+	friend inline uint32 GetTypeHash(const FSMInstanceElementId& InId)
+	{
+		return HashCombine(GetTypeHash(InId.InstanceId), GetTypeHash(InId.ISMComponent));
+	}
+
+	UInstancedStaticMeshComponent* ISMComponent = nullptr;
+	uint64 InstanceId = 0;
+};
+
+/**
+ * Entry within a FSMInstanceElementIdMap, tied to a specific ISM component.
+ */
+struct FSMInstanceElementIdMapEntry
+{
+	FSMInstanceElementIdMapEntry(FSMInstanceElementIdMap* InOwner, UInstancedStaticMeshComponent* InComponent);
+	~FSMInstanceElementIdMapEntry();
+
+	FSMInstanceElementIdMap* Owner = nullptr;
+
+	UInstancedStaticMeshComponent* Component = nullptr;
+
+#if WITH_EDITORONLY_DATA
+	USMInstanceElementIdMapTransactor* Transactor = nullptr;
+#endif	// WITH_EDITORONLY_DATA
+
+	TMap<int32, uint64> InstanceIndexToIdMap;
+	TMap<uint64, int32> InstanceIdToIndexMap;
+	uint64 NextInstanceId = 1;
+};
+
+/**
+ * Mapping between the instance ID used by typed elements and the corresponding instance index on the ISM components.
+ * This mapping will be kept up-to-date via the add/remove operations of the ISM components, as well as undo/redo.
+ */
+class ENGINE_API FSMInstanceElementIdMap : public FGCObject
+{
+public:
+	static FSMInstanceElementIdMap& Get();
+
+	~FSMInstanceElementIdMap();
+
+	/**
+	 * Delegate called when a static mesh instance has been relocated within the instances array on its component (eg, because another instance was removed or added to the array).
+	 * @note This may be called with INDEX_NONE as the PreviousInstanceIndex value if the element had previously been unmapped, but was mapped again via serialization (eg, via an undo/redo).
+	 */
+	DECLARE_MULTICAST_DELEGATE_ThreeParams(FOnInstanceRemapped, const FSMInstanceElementId& /*SMInstanceElementId*/, int32 /*PreviousInstanceIndex*/, int32 /*InstanceIndex*/);
+	FOnInstanceRemapped& OnInstanceRemapped()
+	{
+		return OnInstanceRemappedDelegate;
+	}
+
+	/**
+	 * Delegate called when a static mesh instance has been removed from the instances array on its component, and the associated ID for a SMInstance element has been unmapped.
+	 */
+	DECLARE_MULTICAST_DELEGATE_TwoParams(FOnInstanceRemoved, const FSMInstanceElementId& /*SMInstanceElementId*/, int32 /*InstanceIndex*/);
+	FOnInstanceRemoved& OnInstanceRemoved()
+	{
+		return OnInstanceRemovedDelegate;
+	}
+
+	/**
+	 * Given a FSMInstanceElementId, attempt to convert it into a FSMInstanceId.
+	 */
+	FSMInstanceId GetSMInstanceIdFromSMInstanceElementId(const FSMInstanceElementId& InSMInstanceElementId);
+
+	/**
+	 * Given a FSMInstanceId, attempt to convert it into a FSMInstanceElementId.
+	 */
+	FSMInstanceElementId GetSMInstanceElementIdFromSMInstanceId(const FSMInstanceId& InSMInstanceId, const bool bAllowCreate = true);
+
+public:
+	//~ FGCObject interface
+	virtual void AddReferencedObjects(FReferenceCollector& Collector) override;
+	virtual FString GetReferencerName() const override
+	{
+		return TEXT("FSMInstanceElementIdMap");
+	}
+
+public:
+	void SerializeIdMappings(FSMInstanceElementIdMapEntry* InEntry, FArchive& Ar);
+
+private:
+	void OnInstanceIndexUpdated(UInstancedStaticMeshComponent* InComponent, TArrayView<const FInstancedStaticMeshDelegates::FInstanceIndexUpdateData> InIndexUpdates);
+
+#if WITH_EDITOR
+	void OnObjectModified(UObject* InObject);
+#endif	// WITH_EDITOR
+
+	void RegisterCallbacks();
+
+	void UnregisterCallbacks();
+
+	mutable FCriticalSection ISMComponentsCS;
+	TMap<UInstancedStaticMeshComponent*, TSharedPtr<FSMInstanceElementIdMapEntry>> ISMComponents;
+
+	FOnInstanceRemapped OnInstanceRemappedDelegate;
+
+	FOnInstanceRemoved OnInstanceRemovedDelegate;
+
+	FDelegateHandle OnInstanceIndexUpdatedHandle;
+#if WITH_EDITOR
+	FDelegateHandle OnObjectModifiedHandle;
+#endif
+};
+
+/**
+ * Transient object instance used as a proxy for storing the current data mapping state
+ * in the transaction buffer (for retaining the correct mapping through undo/redo).
+ */
+UCLASS(Transient)
+class USMInstanceElementIdMapTransactor : public UObject
+{
+	GENERATED_BODY()
+
+public:
+#if WITH_EDITORONLY_DATA
+	void SetOwnerEntry(FSMInstanceElementIdMapEntry* InOwnerEntry)
+	{
+		OwnerEntry = InOwnerEntry;
+	}
+
+	//~ UObject interface
+	virtual void Serialize(FArchive& Ar) override;
+
+private:
+	FSMInstanceElementIdMapEntry* OwnerEntry = nullptr;
+#endif	// WITH_EDITORONLY_DATA
+};
diff --git a/Engine/Source/Runtime/Engine/Public/Elements/SMInstance/SMInstanceElementSelectionInterface.h b/Engine/Source/Runtime/Engine/Public/Elements/SMInstance/SMInstanceElementSelectionInterface.h
new file mode 100644
--- /dev/null
+++ b/Engine/Source/Runtime/Engine/Public/Elements/SMInstance/SMInstanceElementSelectionInterface.h
@@ -0,0 +1,16 @@
+// Copyright Epic Games, Inc. All Rights Reserved.
+
+#pragma once
+
+#include "Elements/Interfaces/TypedElementSelectionInterface.h"
+#include "SMInstanceElementSelectionInterface.generated.h"
+
+UCLASS()
+class ENGINE_API USMInstanceElementSelectionInterface : public UTypedElementSelectionInterface
+{
+	GENERATED_BODY()
+
+public:
+	static int32 GetNumSelectedSMInstances(const UTypedElementList* InCurrentSelection);
+	static bool HasSelectedSMInstances(const UTypedElementList* InCurrentSelection);
+};
diff --git a/Engine/Source/Runtime/Engine/Public/Elements/SMInstance/SMInstanceElementWorldInterface.h b/Engine/Source/Runtime/Engine/Public/Elements/SMInstance/SMInstanceElementWorldInterface.h
new file mode 100644
--- /dev/null
+++ b/Engine/Source/Runtime/Engine/Public/Elements/SMInstance/SMInstanceElementWorldInterface.h
@@ -0,0 +1,27 @@
+// Copyright Epic Games, Inc. All Rights Reserved.
+
+#pragma once
+
+#include "Elements/Interfaces/TypedElementWorldInterface.h"
+#include "SMInstanceElementWorldInterface.generated.h"
+
+struct FSMInstanceId;
+
+UCLASS()
+class ENGINE_API USMInstanceElementWorldInterface : public UTypedElementWorldInterface
+{
+	GENERATED_BODY()
+
+public:
+	virtual bool CanEditElement(const FTypedElementHandle& InElementHandle) override;
+	virtual bool IsTemplateElement(const FTypedElementHandle& InElementHandle) override;
+	virtual ULevel* GetOwnerLevel(const FTypedElementHandle& InElementHandle) override;
+	virtual UWorld* GetOwnerWorld(const FTypedElementHandle& InElementHandle) override;
+	virtual bool GetBounds(const FTypedElementHandle& InElementHandle, FBoxSphereBounds& OutBounds) override;
+	virtual bool GetWorldTransform(const FTypedElementHandle& InElementHandle, FTransform& OutTransform) override;
+	virtual bool SetWorldTransform(const FTypedElementHandle& InElementHandle, const FTransform& InTransform) override;
+	virtual bool GetRelativeTransform(const FTypedElementHandle& InElementHandle, FTransform& OutTransform) override;
+	virtual bool SetRelativeTransform(const FTypedElementHandle& InElementHandle, const FTransform& InTransform) override;
+
+	static bool CanEditSMInstance(const FSMInstanceId& InSMInstanceId);
+};
diff --git a/Engine/Source/Runtime/Engine/Public/InstancedStaticMeshDelegates.h b/Engine/Source/Runtime/Engine/Public/InstancedStaticMeshDelegates.h
new file mode 100644
--- /dev/null
+++ b/Engine/Source/Runtime/Engine/Public/InstancedStaticMeshDelegates.h
@@ -0,0 +1,51 @@
+// Copyright Epic Games, Inc. All Rights Reserved.
+
+#pragma once
+
+#include "CoreMinimal.h"
+#include "Containers/ArrayView.h"
+
+class UInstancedStaticMeshComponent;
+
+struct ENGINE_API FInstancedStaticMeshDelegates
+{
+	enum class EInstanceIndexUpdateType : uint8
+	{
+		/** An instance has been added */
+		Added,
+		/** An instance has been removed */
+		Removed,
+		/** An instance has been relocated within the array of instances */
+		Relocated,
+		/** All instances have been removed */
+		Cleared,
+		/** All instances have been removed, and the component is being destroyed */
+		Destroyed,
+	};
+
+	struct FInstanceIndexUpdateData
+	{
+		/**
+		 * The type of this update operation.
+		 */
+		EInstanceIndexUpdateType Type = EInstanceIndexUpdateType::Cleared;
+
+		/**
+		 * The index of the affected instance, when Type==Added, Type==Removed or Type==Relocated.
+		 * The index of the last instance, when Type==Cleared or Type==Destroyed.
+		 */
+		int32 Index = INDEX_NONE;
+
+		/**
+		 * The previous index of the affected instance, when Type==Relocated.
+		 */
+		int32 OldIndex = INDEX_NONE;
+	};
+
+	/**
+	 * Delegate called when the index of the instances within an ISM are updated.
+	 * @note: This is only fired for instances added or removed via the various Add/Remove/ClearInstances functions, and not for instances updated via serialization or undo/redo.
+	 */
+	DECLARE_MULTICAST_DELEGATE_TwoParams(FOnInstanceIndexUpdated, UInstancedStaticMeshComponent* /*Component*/, TArrayView<const FInstanceIndexUpdateData> /*IndexUpdates*/);
+	static FOnInstanceIndexUpdated OnInstanceIndexUpdated;
+};
