Convert Strata slab node from RoughnessX/Y into Roughness/Anisotropy.

diff --git a/Engine/Shaders/Private/Strata/Strata.ush b/Engine/Shaders/Private/Strata/Strata.ush
--- a/Engine/Shaders/Private/Strata/Strata.ush
+++ b/Engine/Shaders/Private/Strata/Strata.ush
@@ -432,7 +432,7 @@ FStrataData GetInitialisedStrataData()
 }
 
 FStrataData GetStrataSlabBSDF(
-	float3 BaseColor, float3 EdgeColor, float Specular, float Metallic, float2 Roughness,
+	float3 BaseColor, float3 EdgeColor, float Specular, float Metallic, float Roughness, float Anisotropy,
 	float SSSProfileID, float3 SSSDMFP, float SSSDMFPScale,
 	float3 Emissive, float Haziness, float ThinFilmThickness,  float Thickness, 
 	uint SharedNormalIndex, inout uint SharedNormalTypes)
@@ -443,12 +443,14 @@ FStrataData GetStrataSlabBSDF(
 	const uint bHasEdgeColor = any(EdgeColor < 1);
 	const uint bHasHaziness = Haziness > 0.0f ? 1 : 0;
 	const uint bHasThinFilm = ThinFilmThickness > 0.0f ? 1 : 0;
-	const uint bHasAnisotropy = StrataHasAnisotropy(Roughness.x, Roughness.y) ? 1 : 0;
+	const uint bHasAnisotropy = Anisotropy != 0;
 	if (bHasAnisotropy)
 	{
 		StrataRequestSharedNormalTangentBasis(SharedNormalTypes, SharedNormalIndex);
 	}
 
+	const float2 RoughnessXY = GetAnisotropicRoughness(Roughness, Anisotropy);
+
 	// STRATA_TODO apply SSS as a function of pixel footprint w.r.t. MFP
 	const bool bDMFPPlugged = any(SSSDMFP > 0);
 	const bool bHasSSSProfile = SSSProfileID  > 0.0f;	// When no profile is provide, we use Burley. STRATA_TODO we should clarify as burley is also an option of the profile.
@@ -472,8 +474,8 @@ FStrataData GetStrataSlabBSDF(
 	SLAB_EDGECOLOR			(StrataData.Layers[0].BSDFs[0])		= EdgeColor;
 	SLAB_SPECULAR			(StrataData.Layers[0].BSDFs[0])		= Specular;
 	SLAB_METALLIC			(StrataData.Layers[0].BSDFs[0])		= Metallic;
-	SLAB_ROUGHNESSX			(StrataData.Layers[0].BSDFs[0])		= Roughness.x;
-	SLAB_ROUGHNESSY			(StrataData.Layers[0].BSDFs[0])		= Roughness.y;
+	SLAB_ROUGHNESSX			(StrataData.Layers[0].BSDFs[0])		= RoughnessXY.x;
+	SLAB_ROUGHNESSY			(StrataData.Layers[0].BSDFs[0])		= RoughnessXY.y;
 	SLAB_HAZINESS			(StrataData.Layers[0].BSDFs[0])		= Haziness;
 	SLAB_THINFILM_THICKNESS	(StrataData.Layers[0].BSDFs[0])		= ThinFilmThickness;
 	if (bHasSSSProfile)
diff --git a/Engine/Shaders/Private/Strata/StrataVisualize.usf b/Engine/Shaders/Private/Strata/StrataVisualize.usf
--- a/Engine/Shaders/Private/Strata/StrataVisualize.usf
+++ b/Engine/Shaders/Private/Strata/StrataVisualize.usf
@@ -333,13 +333,24 @@ void VisualizeMaterialPS(
 				HORIZONTALTAB; PRINTFLOAT1(SLAB_METALLIC(BSDF));
 				NEXTLINE;
 
-				HORIZONTALTAB; PRINTPARAMCHAR(_R_); PRINTPARAMCHAR(_O_); PRINTPARAMCHAR(_U_); PRINTPARAMCHAR(_G_); PRINTPARAMCHAR(_H_); PRINTPARAMCHAR(_X_);
-				HORIZONTALTAB; PRINTFLOAT1(SLAB_ROUGHNESSX(BSDF));
-				NEXTLINE;
 				if (BSDF_GETANISOTROPY(BSDF))
 				{
-					HORIZONTALTAB; PRINTPARAMCHAR(_R_); PRINTPARAMCHAR(_O_); PRINTPARAMCHAR(_U_); PRINTPARAMCHAR(_G_); PRINTPARAMCHAR(_H_); PRINTPARAMCHAR(_Y_);
-					HORIZONTALTAB; PRINTFLOAT1(SLAB_ROUGHNESSY(BSDF));
+					float Roughness  = 0.f;
+					float Anisotropy = 0.f;
+					GetAnisotropicFactor(SLAB_ROUGHNESSX(BSDF), SLAB_ROUGHNESSY(BSDF), Anisotropy, Roughness);
+					
+					HORIZONTALTAB; PRINTPARAMCHAR(_R_); PRINTPARAMCHAR(_O_); PRINTPARAMCHAR(_U_); PRINTPARAMCHAR(_G_); PRINTPARAMCHAR(_H_); PRINTPARAMCHAR(_SPC_);
+					HORIZONTALTAB; PRINTFLOAT1(Roughness);
+					NEXTLINE;
+
+					HORIZONTALTAB; PRINTPARAMCHAR(_A_); PRINTPARAMCHAR(_N_); PRINTPARAMCHAR(_I_); PRINTPARAMCHAR(_S_); PRINTPARAMCHAR(_O_); PRINTPARAMCHAR(_SPC_);
+					HORIZONTALTAB; PRINTFLOAT1(abs(Anisotropy));
+					NEXTLINE;
+				}
+				else
+				{
+					HORIZONTALTAB; PRINTPARAMCHAR(_R_); PRINTPARAMCHAR(_O_); PRINTPARAMCHAR(_U_); PRINTPARAMCHAR(_G_); PRINTPARAMCHAR(_H_); PRINTPARAMCHAR(_SPC_);
+					HORIZONTALTAB; PRINTFLOAT1(SLAB_ROUGHNESSX(BSDF));
 					NEXTLINE;
 				}
 
diff --git a/Engine/Source/Runtime/Engine/Classes/Materials/MaterialExpressionStrata.h b/Engine/Source/Runtime/Engine/Classes/Materials/MaterialExpressionStrata.h
--- a/Engine/Source/Runtime/Engine/Classes/Materials/MaterialExpressionStrata.h
+++ b/Engine/Source/Runtime/Engine/Classes/Materials/MaterialExpressionStrata.h
@@ -64,13 +64,13 @@ class UMaterialExpressionStrataSlabBSDF : public UMaterialExpressionStrataBSDF
 	 * Controls how rough the Material is. Roughness of 0 (smooth) is a mirror reflection and 1 (rough) is completely matte or diffuse. When using anisotropy, it is the roughness used along the Tangent axis. (type = float, unit = unitless, defaults to 0.5)
 	 */
 	UPROPERTY()
-	FExpressionInput RoughnessX;
+	FExpressionInput Roughness;
 		
 	/**
-	 * Controls the roughness along the secondary surface tangent vector (perpendicular to Tangent). (type = float, unit = unitless). If not plugged in, RoughnessY is set to RoughnessX to disable anisotropy, resulting in an isotropic behavior.
+	 * Controls the anisotropy factor of the roughness. Positive value elongates the specular lobe along the Tangent vector, Negative value elongates the specular lobe along the perpendiculat of the Tangent. (type = float, unit = unitless).
 	 */
 	UPROPERTY()
-	FExpressionInput RoughnessY;
+	FExpressionInput Anisotropy;
 
 	/**
 	 * Take the surface normal as input. The normal is considered tangent or world space according to the space properties on the main material node. (type = float3, unit = unitless, defaults to vertex normal)
@@ -139,6 +139,7 @@ class UMaterialExpressionStrataSlabBSDF : public UMaterialExpressionStrataBSDF
 	bool HasSSS() const;
 	bool HasSSSProfile() const;
 	bool HasDMFPPluggedIn() const;
+	bool HasAnisotropy() const;
 #endif
 	//~ End UMaterialExpression Interface
 };
diff --git a/Engine/Source/Runtime/Engine/Private/Materials/HLSLMaterialTranslator.cpp b/Engine/Source/Runtime/Engine/Private/Materials/HLSLMaterialTranslator.cpp
--- a/Engine/Source/Runtime/Engine/Private/Materials/HLSLMaterialTranslator.cpp
+++ b/Engine/Source/Runtime/Engine/Private/Materials/HLSLMaterialTranslator.cpp
@@ -9023,18 +9023,18 @@ int32 FHLSLMaterialTranslator::StrataCreateAndRegisterNullMaterial()
 }
 
 int32 FHLSLMaterialTranslator::StrataSlabBSDF(
-	int32 BaseColor, int32 EdgeColor, int32 Specular, int32 Metallic, int32 RoughnessX, int32 RoughnessY,
+	int32 BaseColor, int32 EdgeColor, int32 Specular, int32 Metallic, int32 Roughness, int32 Anisotropy,
 	int32 SSSProfileId, int32 SSSDMFP, int32 SSSDMFPScale, int32 EmissiveColor,
 	int32 Haziness, int32 ThinFilmThickness, int32 Thickness, int32 Normal, int32 Tangent, uint8 SharedNormalIndex)
 {
 	return AddCodeChunk(
-		MCT_Strata, TEXT("GetStrataSlabBSDF(%s, %s, %s, %s, float2(%s, %s), %s, %s, %s, %s, %s, %s, %s, %u, Parameters.SharedNormals.NormalTypes) /* Normal = %s ; Tangent = %s */"),
+		MCT_Strata, TEXT("GetStrataSlabBSDF(%s, %s, %s, %s, %s, %s, %s, %s, %s, %s, %s, %s, %s, %u, Parameters.SharedNormals.NormalTypes) /* Normal = %s ; Tangent = %s */"),
 		*GetParameterCode(BaseColor),
 		*GetParameterCode(EdgeColor),
 		*GetParameterCode(Specular),
 		*GetParameterCode(Metallic),
-		*GetParameterCode(RoughnessX),
-		*GetParameterCode(RoughnessY),
+		*GetParameterCode(Roughness),
+		*GetParameterCode(Anisotropy),
 		*GetParameterCode(SSSProfileId),
 		*GetParameterCode(SSSDMFP),
 		*GetParameterCode(SSSDMFPScale),
diff --git a/Engine/Source/Runtime/Engine/Private/Materials/HLSLMaterialTranslator.h b/Engine/Source/Runtime/Engine/Private/Materials/HLSLMaterialTranslator.h
--- a/Engine/Source/Runtime/Engine/Private/Materials/HLSLMaterialTranslator.h
+++ b/Engine/Source/Runtime/Engine/Private/Materials/HLSLMaterialTranslator.h
@@ -989,7 +989,7 @@ protected:
 	// Strata
 	virtual int32 StrataCreateAndRegisterNullMaterial() override;
 	virtual int32 StrataSlabBSDF(
-		int32 BaseColor, int32 EdgeColor, int32 Specular, int32 Metallic, int32 RoughnessX, int32 RoughnessY,
+		int32 BaseColor, int32 EdgeColor, int32 Specular, int32 Metallic, int32 Roughness, int32 Anisotropy,
 		int32 SSSProfileId, int32 SSSDMFP, int32 SSSDMFPScale,
 		int32 EmissiveColor, int32 Haziness, int32 ThinFilmThickness, int32 Thickness,
 		int32 Normal, int32 Tangent, uint8 SharedNormalIndex) override;
diff --git a/Engine/Source/Runtime/Engine/Private/Materials/Material.cpp b/Engine/Source/Runtime/Engine/Private/Materials/Material.cpp
--- a/Engine/Source/Runtime/Engine/Private/Materials/Material.cpp
+++ b/Engine/Source/Runtime/Engine/Private/Materials/Material.cpp
@@ -3648,37 +3648,15 @@ void UMaterial::ConvertMaterialToStrataMaterial()
 		{
 			UMaterialExpressionStrataSlabBSDF* SlabBSDF = NewObject<UMaterialExpressionStrataSlabBSDF>(this);
 
-			if (Anisotropy.IsConnected())
-			{
-				UMaterialExpressionStrataAnisotropyToRoughness* AnisotropyToRoughness = NewObject<UMaterialExpressionStrataAnisotropyToRoughness>(this);
-				MoveConnectionTo(Roughness, AnisotropyToRoughness, 0);	// Roughness
-				CopyConnectionTo(Anisotropy, AnisotropyToRoughness, 1);	// Anisotropy
-
-				UMaterialExpressionComponentMask* RoughnessesMaskX = NewObject<UMaterialExpressionComponentMask>(this);
-				RoughnessesMaskX->GetInput(0)->Connect(0, AnisotropyToRoughness);
-				RoughnessesMaskX->R = 1;
-				RoughnessesMaskX->G = 0;
-				RoughnessesMaskX->B = 0;
-				RoughnessesMaskX->A = 0;
-
-				UMaterialExpressionComponentMask* RoughnessesMaskY = NewObject<UMaterialExpressionComponentMask>(this);
-				RoughnessesMaskY->GetInput(0)->Connect(0, AnisotropyToRoughness);
-				RoughnessesMaskY->R = 0;
-				RoughnessesMaskY->G = 1;
-				RoughnessesMaskY->B = 0;
-				RoughnessesMaskY->A = 0;
-
-				SlabBSDF->RoughnessX.Connect(0, RoughnessesMaskX);
-				SlabBSDF->RoughnessY.Connect(0, RoughnessesMaskY);
-			}
-			else
-			{
-				MoveConnectionTo(Roughness, SlabBSDF, 4);	// RoughnessX
-			}
 
 			MoveConnectionTo(BaseColor, SlabBSDF, 0);		// BaseColor
 			MoveConnectionTo(Metallic, SlabBSDF, 2);		// Metallic
 			MoveConnectionTo(Specular, SlabBSDF, 3);		// Specular
+			MoveConnectionTo(Roughness, SlabBSDF, 4);		// Roughness
+			if (Anisotropy.IsConnected())
+			{
+				MoveConnectionTo(Anisotropy, SlabBSDF, 5);	// Anisotropy
+			}
 			MoveConnectionTo(Normal, SlabBSDF, 6);			// Normal
 			MoveConnectionTo(Tangent, SlabBSDF, 7);			// Tangent
 			// Opacity can remain on the end point node
diff --git a/Engine/Source/Runtime/Engine/Private/Materials/MaterialExpressions.cpp b/Engine/Source/Runtime/Engine/Private/Materials/MaterialExpressions.cpp
--- a/Engine/Source/Runtime/Engine/Private/Materials/MaterialExpressions.cpp
+++ b/Engine/Source/Runtime/Engine/Private/Materials/MaterialExpressions.cpp
@@ -19738,18 +19738,17 @@ UMaterialExpressionStrataSlabBSDF::UMaterialExpressionStrataSlabBSDF(const FObje
 #if WITH_EDITOR
 int32 UMaterialExpressionStrataSlabBSDF::Compile(class FMaterialCompiler* Compiler, int32 OutputIndex)
 {
-	int32 RoughnessXCodeChunk = CompileWithDefaultFloat1(Compiler, RoughnessX, 0.5f);
-	// If not plugged in, RoughnessYCodeChunk is set to RoughnessXCodeChunk to get an isotropic behavior
-	int32 RoughnessYCodeChunk = CompileWithDefaultCodeChunk(Compiler, RoughnessY, RoughnessXCodeChunk);
+	int32 RoughnessCodeChunk = CompileWithDefaultFloat1(Compiler, Roughness, 0.5f);
+	int32 AnisotropyCodeChunk = CompileWithDefaultFloat1(Compiler, Anisotropy, 0.0f);
 	// As long as both roughness are potentially different, we must take it into account in our encoding.
 	// We also cannot ignore the tangent when using the default Tangent because GetTangentBasis
 	// used in StrataGetBSDFSharedBasis cannot be relied on for smooth tangent used for lighting on any mesh.
-	bool bAnisotropyPotentiallyUsed = RoughnessXCodeChunk != RoughnessYCodeChunk;
+	const bool bHasAnisotropy = HasAnisotropy();
 
 	int32 NormalCodeChunk = CompileWithDefaultNormalWS(Compiler, Normal);
 
 	int32 TangentCodeChunk = CompileWithDefaultTangentWS(Compiler, Tangent);
-	uint8 SharedNormalIndex = bAnisotropyPotentiallyUsed ? StrataCompilationInfoCreateSharedNormal(Compiler, NormalCodeChunk, TangentCodeChunk) : StrataCompilationInfoCreateSharedNormal(Compiler, NormalCodeChunk);
+	uint8 SharedNormalIndex = bHasAnisotropy ? StrataCompilationInfoCreateSharedNormal(Compiler, NormalCodeChunk, TangentCodeChunk) : StrataCompilationInfoCreateSharedNormal(Compiler, NormalCodeChunk);
 
 	const bool bHasEdgeColor = HasEdgeColor();
 	const bool bHasThinFilm = HasThinFilm();
@@ -19768,8 +19767,8 @@ int32 UMaterialExpressionStrataSlabBSDF::Compile(class FMaterialCompiler* Compil
 		CompileWithDefaultFloat3(Compiler, EdgeColor, 1.0f, 1.0f, 1.0f),
 		CompileWithDefaultFloat1(Compiler, Specular, 0.5f),
 		CompileWithDefaultFloat1(Compiler, Metallic, 0.0f),
-		RoughnessXCodeChunk,
-		RoughnessYCodeChunk,
+		RoughnessCodeChunk,
+		AnisotropyCodeChunk,
 		SSSProfileCodeChunk != INDEX_NONE ? SSSProfileCodeChunk : Compiler->Constant(0.0f),	
 		CompileWithDefaultFloat3(Compiler, SSSDMFP, 0.0f, 0.0f, 0.0f),
 		CompileWithDefaultFloat1(Compiler, SSSDMFPScale, 1.0f),
@@ -19812,10 +19811,10 @@ uint32 UMaterialExpressionStrataSlabBSDF::GetInputType(int32 InputIndex)
 		return MCT_Float1; // Specular
 		break;
 	case 4:
-		return MCT_Float1; // RoughnessX
+		return MCT_Float1; // Roughness
 		break;
 	case 5:
-		return MCT_Float1; // RoughnessY
+		return MCT_Float1; // Anisotropy
 		break;
 	case 6:
 		return MCT_Float3; // Normal
@@ -19867,11 +19866,11 @@ FName UMaterialExpressionStrataSlabBSDF::GetInputName(int32 InputIndex) const
 	}
 	else if (InputIndex == 4)
 	{
-		return TEXT("RoughnessX");
+		return TEXT("Roughness");
 	}
 	else if (InputIndex == 5)
 	{
-		return TEXT("RoughnessY");
+		return TEXT("Anisotropy");
 	}
 	else if (InputIndex == 6)
 	{
@@ -19955,6 +19954,10 @@ bool UMaterialExpressionStrataSlabBSDF::HasThinFilm() const
 	return ThinFilmThickness.IsConnected();
 }
 
+bool UMaterialExpressionStrataSlabBSDF::HasAnisotropy() const
+{
+	return Anisotropy.IsConnected();
+}
 #endif // WITH_EDITOR
 
 
diff --git a/Engine/Source/Runtime/Engine/Public/MaterialCompiler.h b/Engine/Source/Runtime/Engine/Public/MaterialCompiler.h
--- a/Engine/Source/Runtime/Engine/Public/MaterialCompiler.h
+++ b/Engine/Source/Runtime/Engine/Public/MaterialCompiler.h
@@ -410,7 +410,7 @@ public:
 	// Strata
 	virtual int32 StrataCreateAndRegisterNullMaterial() = 0;
 	virtual int32 StrataSlabBSDF(
-		int32 BaseColor, int32 EdgeColor, int32 Specular, int32 Metallic, int32 RoughnessX, int32 RoughnessY, 
+		int32 BaseColor, int32 EdgeColor, int32 Specular, int32 Metallic, int32 Roughness, int32 Anisotropy, 
 		int32 SSSProfileId, int32 SSSDMFP, int32 SSSDMFPScale,
 		int32 EmissiveColor, int32 Haziness, int32 ThinFilmThickness, int32 Thickness,
 		int32 Normal, int32 Tangent, uint8 SharedNormalIndex) = 0;
@@ -893,13 +893,13 @@ public:
 	}
 
 	virtual int32 StrataSlabBSDF(
-		int32 BaseColor, int32 EdgeColor, int32 Specular, int32 Metallic, int32 RoughnessX, int32 RoughnessY,
+		int32 BaseColor, int32 EdgeColor, int32 Specular, int32 Metallic, int32 Roughness, int32 Anisotropy,
 		int32 SSSProfileId, int32 SSSDMFP, int32 SSSDMFPScale,
 		int32 EmissiveColor, int32 Haziness, int32 ThinFilmThickness, int32 Thickness,
 		int32 Normal, int32 Tangent, uint8 SharedNormalIndex) override
 	{
 		return Compiler->StrataSlabBSDF(
-			BaseColor, EdgeColor, Specular, Metallic, RoughnessX, RoughnessY,
+			BaseColor, EdgeColor, Specular, Metallic, Roughness, Anisotropy,
 			SSSProfileId, SSSDMFP, SSSDMFPScale,
 			EmissiveColor, Haziness, ThinFilmThickness, Thickness, Normal, Tangent, SharedNormalIndex);
 	}
