Chaos Physics - Contraint blueprint library - Add Setter and getter for angular and linear soft limit parameters
Also added keywords to better match the UI wording to make the serach more intuitive

diff --git a/Engine/Source/Runtime/Engine/Classes/PhysicsEngine/ConstraintInstance.h b/Engine/Source/Runtime/Engine/Classes/PhysicsEngine/ConstraintInstance.h
--- a/Engine/Source/Runtime/Engine/Classes/PhysicsEngine/ConstraintInstance.h
+++ b/Engine/Source/Runtime/Engine/Classes/PhysicsEngine/ConstraintInstance.h
@@ -423,57 +423,81 @@ public:
 	}
 
 	/** Whether the linear limits are soft (only if at least one axis if Limited) */
-	bool GetIsSoftLinearLimit() const
-	{
-		return ProfileInstance.LinearLimit.bSoftConstraint;
-	}
+	bool GetIsSoftLinearLimit() const { return ProfileInstance.LinearLimit.bSoftConstraint;	}
 
 	/** Linear stiffness if the constraint is set to use soft linear limits */
-	float GetSoftLinearLimitStiffness() const
-	{
-		return ProfileInstance.LinearLimit.Stiffness;
-	}
+	float GetSoftLinearLimitStiffness() const {	return ProfileInstance.LinearLimit.Stiffness; }
 
 	/** Linear damping if the constraint is set to use soft linear limits */
-	float GetSoftLinearLimitDamping() const
+	float GetSoftLinearLimitDamping() const { return ProfileInstance.LinearLimit.Damping; }
+
+	/** Linear restitution if the constraint is set to use soft linear limits */
+	float GetSoftLinearLimitRestitution() const { return ProfileInstance.LinearLimit.Restitution; }
+
+	/** Linear contact distance if the constraint is set to use soft linear limits */
+	float GetSoftLinearLimitContactDistance() const { return ProfileInstance.LinearLimit.ContactDistance; }
+
+	/** Set twist soft limit parameters */
+	void SetSoftLinearLimitParams(bool bIsSoftLimit, float Stiffness, float Damping, float Restitution, float ContactDistance)
 	{
-		return ProfileInstance.LinearLimit.Damping;
+		ProfileInstance.LinearLimit.bSoftConstraint = bIsSoftLimit;
+		ProfileInstance.LinearLimit.Stiffness = Stiffness;
+		ProfileInstance.LinearLimit.Damping = Damping;
+		ProfileInstance.LinearLimit.Restitution = Restitution;
+		ProfileInstance.LinearLimit.ContactDistance = ContactDistance;
+		UpdateLinearLimit();
 	}
 
 	/** Whether the twist limits are soft (only available if twist is Limited) */
-	bool GetIsSoftTwistLimit() const
-	{
-		return ProfileInstance.TwistLimit.bSoftConstraint;
-	}
+	bool GetIsSoftTwistLimit() const { return ProfileInstance.TwistLimit.bSoftConstraint; }
 
 	/** Twist stiffness if the constraint is set to use soft limits */
-	float GetSoftTwistLimitStiffness() const
-	{
-		return ProfileInstance.TwistLimit.Stiffness;
-	}
+	float GetSoftTwistLimitStiffness() const {	return ProfileInstance.TwistLimit.Stiffness; }
 
 	/** Twist damping if the constraint is set to use soft limits */
-	float GetSoftTwistLimitDamping() const
-	{
-		return ProfileInstance.TwistLimit.Damping;
+	float GetSoftTwistLimitDamping() const { return ProfileInstance.TwistLimit.Damping;	}
+
+	/** Twist restitution if the constraint is set to use soft limits */
+	float GetSoftTwistLimitRestitution() const { return ProfileInstance.TwistLimit.Restitution; }
+
+	/** Twist contact distance if the constraint is set to use soft limits */
+	float GetSoftTwistLimitContactDistance() const { return ProfileInstance.TwistLimit.ContactDistance; }
+
+	/** Set twist soft limit parameters */
+	void SetSoftTwistLimitParams(bool bIsSoftLimit, float Stiffness, float Damping, float Restitution, float ContactDistance) 
+	{ 
+		ProfileInstance.TwistLimit.bSoftConstraint = bIsSoftLimit;
+		ProfileInstance.TwistLimit.Stiffness = Stiffness;
+		ProfileInstance.TwistLimit.Damping = Damping;
+		ProfileInstance.TwistLimit.Restitution = Restitution;
+		ProfileInstance.TwistLimit.ContactDistance = ContactDistance;
+		UpdateAngularLimit();
 	}
 
 	/** Whether the swing limits are soft (only available if swing1 and/or swing2 is Limited) */
-	bool GetIsSoftSwingLimit() const
-	{
-		return ProfileInstance.ConeLimit.bSoftConstraint;
-	}
+	bool GetIsSoftSwingLimit() const { return ProfileInstance.ConeLimit.bSoftConstraint; }
 
 	/** Swing stiffness if the constraint is set to use soft limits */
-	float GetSoftSwingLimitStiffness() const
-	{
-		return ProfileInstance.ConeLimit.Stiffness;
-	}
+	float GetSoftSwingLimitStiffness() const { return ProfileInstance.ConeLimit.Stiffness; }
 
 	/** Swing damping if the constraint is set to use soft limits */
-	float GetSoftSwingLimitDamping() const
+	float GetSoftSwingLimitDamping() const { return ProfileInstance.ConeLimit.Damping; }
+
+	/** Swing restitution if the constraint is set to use soft limits */
+	float GetSoftSwingLimitRestitution() const { return ProfileInstance.ConeLimit.Restitution; }
+
+	/** Swing Contact distance if the constraint is set to use soft limits */
+	float GetSoftSwingLimitContactDistance() const { return ProfileInstance.ConeLimit.ContactDistance; }
+
+	/** Set twist soft limit parameters */
+	void SetSoftSwingLimitParams(bool bIsSoftLimit, float Stiffness, float Damping, float Restitution, float ContactDistance)
 	{
-		return ProfileInstance.ConeLimit.Damping;
+		ProfileInstance.ConeLimit.bSoftConstraint = bIsSoftLimit;
+		ProfileInstance.ConeLimit.Stiffness = Stiffness;
+		ProfileInstance.ConeLimit.Damping = Damping;
+		ProfileInstance.ConeLimit.Restitution = Restitution;
+		ProfileInstance.ConeLimit.ContactDistance = ContactDistance;
+		UpdateAngularLimit();
 	}
 
 	/** Sets the Linear Breakable properties
diff --git a/Engine/Source/Runtime/Engine/Classes/PhysicsEngine/ConstraintInstanceBlueprintLibrary.h b/Engine/Source/Runtime/Engine/Classes/PhysicsEngine/ConstraintInstanceBlueprintLibrary.h
--- a/Engine/Source/Runtime/Engine/Classes/PhysicsEngine/ConstraintInstanceBlueprintLibrary.h
+++ b/Engine/Source/Runtime/Engine/Classes/PhysicsEngine/ConstraintInstanceBlueprintLibrary.h
@@ -21,28 +21,28 @@ public:
 	//---------------------------------------------------------------------------------------------------
 
 	/** Sets whether bodies attched to the constraint can collide or not
-	*	@param ConstraintInstance	Instance of the constraint to change
+	*	@param Accessor				Constraint accessor to change
 	*	@param bDisableCollision	true to disable collision between constrained bodies
 	*/
-	UFUNCTION(BlueprintCallable, Category = "Physics|Constraints")
+	UFUNCTION(BlueprintCallable, Category = "Physics|Constraints", meta = (Keywords = "disable enable collision"))
 	static void SetDisableCollision(
 			UPARAM(ref) FConstraintInstanceAccessor& Accessor,
 			bool bDisableCollision
 		);
 
 	/** Gets whether bodies attched to the constraint can collide or not
-	*	@param ConstraintInstance	Instance of the constraint to change
+	*	@param Accessor		Constraint accessor to query
 	*/
-	UFUNCTION(BlueprintCallable, Category = "Physics|Constraints")
+	UFUNCTION(BlueprintCallable, Category = "Physics|Constraints", meta = (Keywords = "disable enable collision"))
 	static bool GetDisableCollsion(FConstraintInstanceAccessor& Accessor);
 
 	/** Sets projection parameters of the constraint
-	*	@param ConstraintInstance		Instance of the constraint to change
+	*	@param Accessor					Constraint accessor to change
 	*	@param bEnableProjection		true to enable projection
 	*	@param ProjectionLinearAlpha	how much linear projection to apply in [0,1] range
 	*	@param ProjectionAngularAlpha	how much angular projection to apply in [0,1] range
 	*/
-	UFUNCTION(BlueprintCallable, Category = "Physics|Constraints")
+	UFUNCTION(BlueprintCallable, Category = "Physics|Constraints", meta = (Keywords = "enable linear angular"))
 	static void SetProjectionParams(
 		UPARAM(ref) FConstraintInstanceAccessor& Accessor,
 		bool bEnableProjection,
@@ -51,12 +51,12 @@ public:
 	);
 
 	/** Gets projection parameters of the constraint
-	*	@param ConstraintInstance		Instance of the constraint to change
+	*	@param Accessor					Constraint accessor to query
 	*	@param bEnableProjection		true to enable projection
 	*	@param ProjectionLinearAlpha	how much linear projection to apply in [0,1] range
 	*	@param ProjectionAngularAlpha	how much angular projection to apply in [0,1] range
 	*/
-	UFUNCTION(BlueprintCallable, Category = "Physics|Constraints")
+	UFUNCTION(BlueprintCallable, Category = "Physics|Constraints", meta = (Keywords = "enable linear angular"))
 	static void GetProjectionParams(
 		UPARAM(ref) FConstraintInstanceAccessor& Accessor,
 		bool& bEnableProjection,
@@ -65,7 +65,7 @@ public:
 	);
 
 	/** Sets whether the parent body is not affected by it's child motion 
-	*	@param ConstraintInstance	Instance of the constraint to change
+	*	@param Accessor				Constraint accessor to change
 	*	@param bParentDominates		true to avoid the parent being affected by its child motion
 	*/
 	UFUNCTION(BlueprintCallable, Category = "Physics|Constraints")
@@ -75,7 +75,7 @@ public:
 	);
 
 	/** Gets whether the parent body is not affected by it's child motion
-	*	@param ConstraintInstance	Instance of the constraint to change
+	*	@param Accessor Constraint accessor to query
 	*/
 	UFUNCTION(BlueprintCallable, Category = "Physics|Constraints")
 	static bool GetParentDominates(FConstraintInstanceAccessor& Accessor);
@@ -87,13 +87,13 @@ public:
 	//---------------------------------------------------------------------------------------------------
 
 	/** Sets Constraint Linear Motion Ranges
-	*	@param ConstraintInstance	Instance of the constraint to change
+	*	@param Accessor	Constraint accessor to change
 	*	@param XMotion	Type of motion along the X axis
 	*	@param YMotion	Type of motion along the Y axis
 	*	@param ZMotion	Type of motion along the Z axis
 	*	@param Limit	linear limit to apply to all axis
 	*/
-	UFUNCTION(BlueprintCallable, Category = "Physics|Constraints")
+	UFUNCTION(BlueprintCallable, Category = "Physics|Constraints", meta = (Keywords = "motion free limited locked linear angular"))
 	static void SetLinearLimits(
 		UPARAM(ref) FConstraintInstanceAccessor& Accessor,
 		TEnumAsByte<ELinearConstraintMotion> XMotion,
@@ -103,23 +103,59 @@ public:
 	);
 
 	/** Gets Constraint Linear Motion Ranges
-	*	@param ConstraintInstance	Instance of the constraint to query
+	*	@param Accessor	Constraint accessor to query
 	*	@param XMotion	Type of motion along the X axis
 	*	@param YMotion	Type of motion along the Y axis
 	*	@param ZMotion	Type of motion along the Z axis
 	*	@param Limit	linear limit applied to all axis
 	*/
-	UFUNCTION(BlueprintCallable, Category = "Physics|Constraints")
+	UFUNCTION(BlueprintCallable, Category = "Physics|Constraints", meta = (Keywords = "motion free limited locked linear angular"))
 	static void GetLinearLimits(
-		FConstraintInstanceAccessor& Accessor,
+		UPARAM(ref) FConstraintInstanceAccessor& Accessor,
 		TEnumAsByte<ELinearConstraintMotion>& XMotion,
 		TEnumAsByte<ELinearConstraintMotion>& YMotion,
 		TEnumAsByte<ELinearConstraintMotion>& ZMotion,
 		float& Limit
 	);
 
+	/** Sets Constraint Linear Soft Limit parameters
+	*	@param Accessor						Constraint accessor to change
+	* *	@param bSoftLinearLimit				True is the linear limit is soft
+	*	@param LinearLimitStiffness			Stiffness of the soft linear limit. Only used when Soft limit is on ( positive value )
+	*	@param LinearLimitDamping			Damping of the soft linear limit. Only used when Soft limit is on ( positive value )
+	*   @param LinearLimitRestitution		Controls the amount of bounce when the constraint is violated. A restitution value of 1 will bounce back with the same velocity the limit was hit. A value of 0 will stop dead.
+	*	@param LinearLimitContactDistance	Determines how close to the limit we have to get before turning the joint on. Larger value will be more expensive, but will do a better job not violating constraints. A smaller value will be more efficient, but easier to violate.
+	*/
+	UFUNCTION(BlueprintCallable, Category = "Physics|Constraints", meta = (Keywords = "soft stiffness damping restitution contact distance"))
+	static void SetLinearSoftLimitParams(
+		UPARAM(ref) FConstraintInstanceAccessor& Accessor,
+		bool bSoftLinearLimit,
+		float LinearLimitStiffness,
+		float LinearLimitDamping,
+		float LinearLimitRestitution,
+		float LinearLimitContactDistance
+	);
+
+	/** Gets Constraint Linear Soft Limit parameters
+	*	@param Accessor						Constraint accessor to query
+	* *	@param bSoftLinearLimit				True is the Linear limit is soft
+	*	@param LinearLimitStiffness			Stiffness of the soft Linear limit. Only used when Soft limit is on ( positive value )
+	*	@param LinearLimitDamping			Damping of the soft Linear limit. Only used when Soft limit is on ( positive value )
+	*   @param LinearLimitRestitution		Controls the amount of bounce when the constraint is violated. A restitution value of 1 will bounce back with the same velocity the limit was hit. A value of 0 will stop dead.
+	*	@param LinearLimitContactDistance	Determines how close to the limit we have to get before turning the joint on. Larger value will be more expensive, but will do a better job not violating constraints. A smaller value will be more efficient, but easier to violate.
+	*/
+	UFUNCTION(BlueprintCallable, Category = "Physics|Constraints", meta = (Keywords = "soft stiffness damping restitution contact distance"))
+	static void GetLinearSoftLimitParams(
+		UPARAM(ref) FConstraintInstanceAccessor& Accessor,
+		bool& bSoftLinearLimit,
+		float& LinearLimitStiffness,
+		float& LinearLimitDamping,
+		float& LinearLimitRestitution,
+		float& LinearLimitContactDistance
+	);
+
 	/** Sets the Linear Breakable properties
-	*	@param ConstraintInstance	Instance of the constraint to change
+	*	@param Accessor				Constraint accessor to change
 	*	@param bLinearBreakable		Whether it is possible to break the joint with linear force
 	*	@param LinearBreakThreshold	Force needed to break the joint
 	*/
@@ -131,7 +167,7 @@ public:
 	);
 
 	/** Gets Constraint Linear Breakable properties
-	*	@param ConstraintInstance	Instance of the constraint to query
+	*	@param Accessor				Constraint accessor to query
 	*	@param bLinearBreakable		Whether it is possible to break the joint with linear force
 	*	@param LinearBreakThreshold	Force needed to break the joint
 	*/
@@ -149,7 +185,7 @@ public:
 	//---------------------------------------------------------------------------------------------------
 
 	/** Sets COnstraint Angular Motion Ranges
-	*	@param ConstraintInstance	Instance of the constraint to change
+	*	@param Accessor				Constraint accessor to change
 	*	@param Swing1MotionType		Type of swing motion ( first axis )
 	*	@param Swing1LimitAngle		Size of limit in degrees in [0, 180] range
 	*   @param Swing2MotionType		Type of swing motion ( second axis )
@@ -157,7 +193,7 @@ public:
 	*   @param TwistMotionType		Type of twist motion
 	*	@param TwistLimitAngle		Size of limit in degrees in [0, 180] range	
 	*/
-	UFUNCTION(BlueprintCallable, Category = "Physics|Constraints")
+	UFUNCTION(BlueprintCallable, Category = "Physics|Constraints", meta = (Keywords = "swing twist motion free limited locked"))
 	static void SetAngularLimits(
 		UPARAM(ref) FConstraintInstanceAccessor& Accessor, 
 		TEnumAsByte<EAngularConstraintMotion> Swing1MotionType,
@@ -169,7 +205,7 @@ public:
 	);
 
 	/** Gets Constraint Angular Motion Ranges
-	*	@param ConstraintInstance	Instance of the constraint to query
+	*	@param Accessor				Constraint accessor to query
 	*	@param Swing1MotionType		Type of swing motion ( first axis )
 	*	@param Swing1LimitAngle		Size of limit in degrees in [0, 180] range
 	*   @param Swing2MotionType		Type of swing motion ( second axis )
@@ -177,7 +213,7 @@ public:
 	*   @param TwistMotionType		Type of twist motion
 	*	@param TwistLimitAngle		Size of limit in degrees in [0, 180] range
 	*/
-	UFUNCTION(BlueprintCallable, Category = "Physics|Constraints")
+	UFUNCTION(BlueprintCallable, Category = "Physics|Constraints", meta = (Keywords = "swing twist motion free limited locked"))
 	static void GetAngularLimits(
 		UPARAM(ref) FConstraintInstanceAccessor& Accessor,
 		TEnumAsByte<EAngularConstraintMotion>& Swing1MotionType,
@@ -188,8 +224,80 @@ public:
 		float& TwistLimitAngle
 	);
 
+	/** Sets Constraint Angular Soft Swing Limit parameters
+	*	@param Accessor						Constraint accessor to change
+	* *	@param bSoftSwingLimit				True is the swing limit is soft
+	*	@param SwingLimitStiffness			Stiffness of the soft swing limit. Only used when Soft limit is on ( positive value )
+	*	@param SwingLimitDamping			Damping of the soft swing limit. Only used when Soft limit is on ( positive value )
+	*   @param SwingLimitRestitution		Controls the amount of bounce when the constraint is violated. A restitution value of 1 will bounce back with the same velocity the limit was hit. A value of 0 will stop dead.
+	*	@param SwingLimitContactDistance	Determines how close to the limit we have to get before turning the joint on. Larger value will be more expensive, but will do a better job not violating constraints. A smaller value will be more efficient, but easier to violate.
+	*/
+	UFUNCTION(BlueprintCallable, Category = "Physics|Constraints", meta = (Keywords = "soft stiffness damping restitution contact distance"))
+	static void SetAngularSoftSwingLimitParams(
+		UPARAM(ref) FConstraintInstanceAccessor& Accessor,
+		bool bSoftSwingLimit,
+		float SwingLimitStiffness,
+		float SwingLimitDamping,
+		float SwingLimitRestitution,
+		float SwingLimitContactDistance
+	);
+
+	/** Gets Constraint Angular Soft Swing Limit parameters
+	*	@param Accessor						Constraint accessor to query
+	* *	@param bSoftSwingLimit				True is the swing limit is soft
+	*	@param SwingLimitStiffness			Stiffness of the soft swing limit. Only used when Soft limit is on ( positive value )
+	*	@param SwingLimitDamping			Damping of the soft swing limit. Only used when Soft limit is on ( positive value )
+	*   @param SwingLimitRestitution		Controls the amount of bounce when the constraint is violated. A restitution value of 1 will bounce back with the same velocity the limit was hit. A value of 0 will stop dead.
+	*	@param SwingLimitContactDistance	Determines how close to the limit we have to get before turning the joint on. Larger value will be more expensive, but will do a better job not violating constraints. A smaller value will be more efficient, but easier to violate.
+	*/
+	UFUNCTION(BlueprintCallable, Category = "Physics|Constraints", meta = (Keywords = "soft stiffness damping restitution contact distance"))
+	static void GetAngularSoftSwingLimitParams(
+		UPARAM(ref) FConstraintInstanceAccessor& Accessor,
+		bool& bSoftSwingLimit,
+		float& SwingLimitStiffness,
+		float& SwingLimitDamping,
+		float& SwingLimitRestitution,
+		float& SwingLimitContactDistance
+	);
+
+	/** Sets Constraint Angular Soft Twist Limit parameters
+	*	@param Accessor						Constraint accessor to change
+	* *	@param bSoftTwistLimit				True is the twist limit is soft
+	*	@param TwistLimitStiffness			Stiffness of the soft Twist limit. Only used when Soft limit is on ( positive value )
+	*	@param TwistLimitDamping			Damping of the soft Twist limit. Only used when Soft limit is on ( positive value )
+	*   @param TwistLimitRestitution		Controls the amount of bounce when the constraint is violated. A restitution value of 1 will bounce back with the same velocity the limit was hit. A value of 0 will stop dead.
+	*	@param TwistLimitContactDistance	Determines how close to the limit we have to get before turning the joint on. Larger value will be more expensive, but will do a better job not violating constraints. A smaller value will be more efficient, but easier to violate.
+	*/
+	UFUNCTION(BlueprintCallable, Category = "Physics|Constraints", meta = (Keywords = "soft stiffness damping restitution contact distance"))
+	static void SetAngularSoftTwistLimitParams(
+		UPARAM(ref) FConstraintInstanceAccessor& Accessor,
+		bool bSoftTwistLimit,
+		float TwistLimitStiffness,
+		float TwistLimitDamping,
+		float TwistLimitRestitution,
+		float TwistLimitContactDistance
+	);
+
+	/** Gets Constraint Angular Soft Twist Limit parameters
+	*	@param Accessor						Constraint accessor to query
+	* *	@param bSoftTwistLimit				True is the Twist limit is soft
+	*	@param TwistLimitStiffness			Stiffness of the soft Twist limit. Only used when Soft limit is on ( positive value )
+	*	@param TwistLimitDamping			Damping of the soft Twist limit. Only used when Soft limit is on ( positive value )
+	*   @param TwistLimitRestitution		Controls the amount of bounce when the constraint is violated. A restitution value of 1 will bounce back with the same velocity the limit was hit. A value of 0 will stop dead.
+	*	@param TwistLimitContactDistance	Determines how close to the limit we have to get before turning the joint on. Larger value will be more expensive, but will do a better job not violating constraints. A smaller value will be more efficient, but easier to violate.
+	*/
+	UFUNCTION(BlueprintCallable, Category = "Physics|Constraints", meta = (Keywords = "soft stiffness damping restitution contact distance"))
+		static void GetAngularSoftTwistLimitParams(
+			UPARAM(ref) FConstraintInstanceAccessor& Accessor,
+			bool& bSoftTwistLimit,
+			float& TwistLimitStiffness,
+			float& TwistLimitDamping,
+			float& TwistLimitRestitution,
+			float& TwistLimitContactDistance
+		);
+
 	/** Sets Constraint Angular Breakable properties
-	*	@param ConstraintInstance		Instance of the constraint to change
+	*	@param Accessor					Constraint accessor to change
 	*	@param bAngularBreakable		Whether it is possible to break the joint with angular force
 	*	@param AngularBreakThreshold	Torque needed to break the joint
 	*/
@@ -201,7 +309,7 @@ public:
 	);
 
 	/** Gets Constraint Angular Breakable properties
-	*	@param ConstraintInstance		Instance of the constraint to query
+	*	@param Accessor					Constraint accessor to query
 	*	@param bAngularBreakable		Whether it is possible to break the joint with angular force
 	*	@param AngularBreakThreshold	Torque needed to break the joint
 	*/
@@ -213,7 +321,7 @@ public:
 	);
 
 	/** Sets Constraint Angular Plasticity properties
-	*	@param ConstraintInstance			Instance of the constraint to change
+	*	@param Accessor						Constraint accessor to change
 	*	@param bAngularPlasticity			Whether it is possible to reset the target angle from the angular displacement
 	*	@param AngularPlasticityThreshold	Degrees needed to reset the rest state of the joint
 	*/
@@ -225,7 +333,7 @@ public:
 	);
 
 	/** Sets Constraint Angular Plasticity properties
-	*	@param ConstraintInstance			Instance of the constraint to query
+	*	@param Accessor						Constraint accessor to query
 	*	@param bAngularPlasticity			Whether it is possible to reset the target angle from the angular displacement
 	*	@param AngularPlasticityThreshold	Degrees needed to reset the rest state of the joint
 	*/
@@ -243,12 +351,12 @@ public:
 	//---------------------------------------------------------------------------------------------------
 
 	/** Enables/Disables linear position drive
-	*	@param ConstraintInstance	Instance of the constraint to change
+	*	@param Accessor				Constraint accessor to change
 	*	@param bEnableDriveX		Indicates whether the drive for the X-Axis should be enabled
 	*	@param bEnableDriveY		Indicates whether the drive for the Y-Axis should be enabled
 	*	@param bEnableDriveZ		Indicates whether the drive for the Z-Axis should be enabled
 	*/
-	UFUNCTION(BlueprintCallable, Category = "Physics|Constraints")
+	UFUNCTION(BlueprintCallable, Category = "Physics|Constraints", meta = (Keywords = "motor position target"))
 	static void SetLinearPositionDrive(
 		UPARAM(ref) FConstraintInstanceAccessor& Accessor, 
 		bool bEnableDriveX, 
@@ -257,12 +365,12 @@ public:
 	);
 
 	/** Gets whether linear position drive is enabled or not
-	*	@param ConstraintInstance	Instance of the constraint to query
+	*	@param Accessor				Constraint accessor to query
 	*	@param bOutEnableDriveX		Indicates whether the drive for the X-Axis is enabled
 	*	@param bOutEnableDriveY		Indicates whether the drive for the Y-Axis is enabled
 	*	@param bOutEnableDriveZ		Indicates whether the drive for the Z-Axis is enabled
 	*/
-	UFUNCTION(BlueprintCallable, Category = "Physics|Constraints")
+	UFUNCTION(BlueprintCallable, Category = "Physics|Constraints", meta = (Keywords = "motor position target"))
 	static void GetLinearPositionDrive(
 		UPARAM(ref) FConstraintInstanceAccessor& Accessor,
 		bool& bOutEnableDriveX,
@@ -271,12 +379,12 @@ public:
 	);
 
 	/** Enables/Disables linear velocity drive
-	*	@param ConstraintInstance	Instance of the constraint to change
+	*	@param Accessor				Constraint accessor to change
 	*	@param bEnableDriveX		Indicates whether the drive for the X-Axis should be enabled
 	*	@param bEnableDriveY		Indicates whether the drive for the Y-Axis should be enabled
 	*	@param bEnableDriveZ		Indicates whether the drive for the Z-Axis should be enabled
 	*/
-	UFUNCTION(BlueprintCallable, Category = "Physics|Constraints")
+	UFUNCTION(BlueprintCallable, Category = "Physics|Constraints", meta = (Keywords = "motor velocity target"))
 	static void SetLinearVelocityDrive(
 		UPARAM(ref) FConstraintInstanceAccessor& Accessor,
 		bool bEnableDriveX, 
@@ -285,12 +393,12 @@ public:
 	);
 
 	/** Gets whether linear velocity drive is enabled or not
-	*	@param ConstraintInstance	Instance of the constraint to query
+	*	@param Accessor				Constraint accessor to query
 	*	@param bOutEnableDriveX		Indicates whether the drive for the X-Axis is enabled
 	*	@param bOutEnableDriveY		Indicates whether the drive for the Y-Axis is enabled
 	*	@param bOutEnableDriveZ		Indicates whether the drive for the Z-Axis is enabled
 	*/
-	UFUNCTION(BlueprintCallable, Category = "Physics|Constraints")
+	UFUNCTION(BlueprintCallable, Category = "Physics|Constraints", meta = (Keywords = "motor velocity target"))
 	static void GetLinearVelocityDrive(
 		UPARAM(ref) FConstraintInstanceAccessor& Accessor,
 		bool& bOutEnableDriveX,
@@ -299,52 +407,52 @@ public:
 	);
 
 	/** Sets the target position for the linear drive.
-	*	@param ConstraintInstance	Instance of the constraint to change
+	*	@param Accessor				Constraint accessor to change
 	*	@param InPosTarget			Target position
 	*/
-	UFUNCTION(BlueprintCallable, Category = "Physics|Constraints")
+	UFUNCTION(BlueprintCallable, Category = "Physics|Constraints", meta = (Keywords = "motor"))
 	static void SetLinearPositionTarget(
 		UPARAM(ref) FConstraintInstanceAccessor& Accessor, 
 		const FVector& InPosTarget
 	);
 
 	/** Gets the target position for the linear drive.
-	*	@param ConstraintInstance	Instance of the constraint to query
+	*	@param Accessor				Constraint accessor to query
 	*	@param OutPosTarget			Target position
 	*/
-	UFUNCTION(BlueprintCallable, Category = "Physics|Constraints")
+	UFUNCTION(BlueprintCallable, Category = "Physics|Constraints", meta = (Keywords = "motor"))
 	static void GetLinearPositionTarget(
 		UPARAM(ref) FConstraintInstanceAccessor& Accessor,
 		FVector& OutPosTarget
 	);
 
 	/** Sets the target velocity for the linear drive.
-	*	@param ConstraintInstance	Instance of the constraint to change
+	*	@param Accessor				Constraint accessor to change
 	*	@param InVelTarget			Target velocity
 	*/
-	UFUNCTION(BlueprintCallable, Category = "Physics|Constraints")
+	UFUNCTION(BlueprintCallable, Category = "Physics|Constraints", meta = (Keywords = "motor"))
 	static void SetLinearVelocityTarget(
 		UPARAM(ref) FConstraintInstanceAccessor& Accessor, 
 		const FVector& InVelTarget
 	);
 
 	/** Gets the target velocity for the linear drive.
-	*	@param ConstraintInstance	Instance of the constraint to query
+	*	@param Accessor				Constraint accessor to query
 	*	@param OutVelTarget			Target velocity
 	*/
-	UFUNCTION(BlueprintCallable, Category = "Physics|Constraints")
+	UFUNCTION(BlueprintCallable, Category = "Physics|Constraints", meta = (Keywords = "motor"))
 	static void GetLinearVelocityTarget(
 		UPARAM(ref) FConstraintInstanceAccessor& Accessor,
 		FVector& OutVelTarget
 	);
 
 	/** Sets the drive params for the linear drive.
-	*	@param ConstraintInstance	Instance of the constraint to change
+	*	@param Accessor				Constraint accessor to change
 	*	@param PositionStrength		Positional strength for the drive (stiffness)
 	*	@param VelocityStrength		Velocity strength of the drive (damping)
 	*	@param InForceLimit			Max force applied by the drive
 	*/
-	UFUNCTION(BlueprintCallable, Category = "Physics|Constraints")
+	UFUNCTION(BlueprintCallable, Category = "Physics|Constraints", meta = (Keywords = "motor position velocity target strength max force"))
 	static void SetLinearDriveParams(
 		UPARAM(ref) FConstraintInstanceAccessor& Accessor,
 		float PositionStrength, 
@@ -353,12 +461,12 @@ public:
 	);
 
 	/** Gets the drive params for the linear drive.
-	*	@param ConstraintInstance	Instance of the constraint to query
+	*	@param Accessor				Constraint accessor to query
 	*	@param OutPositionStrength	Positional strength for the drive (stiffness)
 	*	@param OutVelocityStrength	Velocity strength of the drive (damping)
 	*	@param OutForceLimit		Max force applied by the drive
 	*/
-	UFUNCTION(BlueprintCallable, Category = "Physics|Constraints")
+	UFUNCTION(BlueprintCallable, Category = "Physics|Constraints", meta = (Keywords = "motor position velocity target strength max force"))
 	static void GetLinearDriveParams(
 		UPARAM(ref) FConstraintInstanceAccessor& Accessor,
 		float& OutPositionStrength,
@@ -373,11 +481,11 @@ public:
 	//---------------------------------------------------------------------------------------------------
 
 	/** Enables/Disables angular orientation drive. Only relevant if the AngularDriveMode is set to Twist and Swing
-	*	@param ConstraintInstance	Instance of the constraint to change
+	*	@param Accessor				Constraint accessor to change
 	*	@param bEnableSwingDrive	Indicates whether the drive for the swing axis should be enabled. Only relevant if the AngularDriveMode is set to Twist and Swing
 	*	@param bEnableTwistDrive	Indicates whether the drive for the twist axis should be enabled. Only relevant if the AngularDriveMode is set to Twist and Swing
 	*/
-	UFUNCTION(BlueprintCallable, Category = "Physics|Constraints")
+	UFUNCTION(BlueprintCallable, Category = "Physics|Constraints", meta = (Keywords = "motor target"))
 	static void SetOrientationDriveTwistAndSwing(
 		UPARAM(ref) FConstraintInstanceAccessor& Accessor, 
 		bool bEnableTwistDrive, 
@@ -385,11 +493,11 @@ public:
 	);
 
 	/** Gets whether angular orientation drive are enabled. Only relevant if the AngularDriveMode is set to Twist and Swing
-	*	@param ConstraintInstance	Instance of the constraint to query
+	*	@param Accessor				Constraint accessor to query
 	*	@param bOutEnableTwistDrive	Indicates whether the drive for the twist axis is enabled. Only relevant if the AngularDriveMode is set to Twist and Swing
 	*	@param bOutEnableSwingDrive	Indicates whether the drive for the swing axis is enabled. Only relevant if the AngularDriveMode is set to Twist and Swing
 	*/
-	UFUNCTION(BlueprintCallable, Category = "Physics|Constraints")
+	UFUNCTION(BlueprintCallable, Category = "Physics|Constraints", meta = (Keywords = "motor target"))
 	static void GetOrientationDriveTwistAndSwing(
 		UPARAM(ref) FConstraintInstanceAccessor& Accessor,
 		bool& bOutEnableTwistDrive,
@@ -397,31 +505,31 @@ public:
 	);
 
 	/** Enables/Disables the angular orientation slerp drive. Only relevant if the AngularDriveMode is set to SLERP
-	*	@param ConstraintInstance	Instance of the constraint to change
+	*	@param Accessor				Constraint accessor to change
 	*	@param bEnableSLERP			Indicates whether the SLERP drive should be enabled. Only relevant if the AngularDriveMode is set to SLERP
 	*/
-	UFUNCTION(BlueprintCallable, Category = "Physics|Constraints")
+	UFUNCTION(BlueprintCallable, Category = "Physics|Constraints", meta = (Keywords = "motor target"))
 	static void SetOrientationDriveSLERP(
 		UPARAM(ref) FConstraintInstanceAccessor& Accessor, 
 		bool bEnableSLERP
 	);
 
 	/** Gets whether the angular orientation slerp drive is enabled or not. Only relevant if the AngularDriveMode is set to SLERP
-	*	@param ConstraintInstance	Instance of the constraint to query
+	*	@param Accessor				Constraint accessor to query
 	*	@param bOutEnableSLERP		Indicates whether the SLERP drive should be enabled. Only relevant if the AngularDriveMode is set to SLERP
 	*/
-	UFUNCTION(BlueprintCallable, Category = "Physics|Constraints")
+	UFUNCTION(BlueprintCallable, Category = "Physics|Constraints", meta = (Keywords = "motor target"))
 	static void GetOrientationDriveSLERP(
 		UPARAM(ref) FConstraintInstanceAccessor& Accessor,
 		bool& bOutEnableSLERP
 	);
 
 	/** Enables/Disables angular velocity twist and swing drive. Only relevant if the AngularDriveMode is set to Twist and Swing
-	*	@param ConstraintInstance	Instance of the constraint to change
+	*	@param Accessor				Constraint accessor to change
 	*	@param bEnableTwistDrive	Indicates whether the drive for the twist axis should be enabled. Only relevant if the AngularDriveMode is set to Twist and Swing
 	*	@param bEnableSwingDrive	Indicates whether the drive for the swing axis should be enabled. Only relevant if the AngularDriveMode is set to Twist and Swing
 	*/
-	UFUNCTION(BlueprintCallable, Category = "Physics|Constraints")
+	UFUNCTION(BlueprintCallable, Category = "Physics|Constraints", meta = (Keywords = "motor target"))
 	static void SetAngularVelocityDriveTwistAndSwing(
 		UPARAM(ref) FConstraintInstanceAccessor& Accessor, 
 		bool bEnableTwistDrive, 
@@ -429,11 +537,11 @@ public:
 	);
 
 	/** Gets whether angular velocity twist and swing drive is enabled or not. Only relevant if the AngularDriveMode is set to Twist and Swing
-	*	@param ConstraintInstance	Instance of the constraint to query	
+	*	@param Accessor				Constraint accessor to query
 	*	@param bOutEnableTwistDrive	Indicates whether the drive for the twist axis should be enabled. Only relevant if the AngularDriveMode is set to Twist and Swing
 	*   @param bOutEnableSwingDrive	Indicates whether the drive for the swing axis should be enabled. Only relevant if the AngularDriveMode is set to Twist and Swing
 	*/
-	UFUNCTION(BlueprintCallable, Category = "Physics|Constraints")
+	UFUNCTION(BlueprintCallable, Category = "Physics|Constraints", meta = (Keywords = "motor target"))
 	static void GetAngularVelocityDriveTwistAndSwing(
 		UPARAM(ref) FConstraintInstanceAccessor& Accessor,
 		bool& bOutEnableTwistDrive,
@@ -441,92 +549,92 @@ public:
 	);
 
 	/** Enables/Disables the angular velocity slerp drive. Only relevant if the AngularDriveMode is set to SLERP
-	*	@param ConstraintInstance	Instance of the constraint to change
+	*	@param Accessor				Constraint accessor to change
 	*	@param bEnableSLERP			Indicates whether the SLERP drive should be enabled. Only relevant if the AngularDriveMode is set to SLERP
 	*/
-	UFUNCTION(BlueprintCallable, Category = "Physics|Constraints")
+	UFUNCTION(BlueprintCallable, Category = "Physics|Constraints", meta = (Keywords = "motor target"))
 	static void SetAngularVelocityDriveSLERP(
 		UPARAM(ref) FConstraintInstanceAccessor& Accessor,
 		bool bEnableSLERP
 	);
 
 	/** Gets whether the angular velocity slerp drive is enabled or not. Only relevant if the AngularDriveMode is set to SLERP
-	*	@param ConstraintInstance	Instance of the constraint to change
+	*	@param Accessor				Constraint accessor to query
 	*	@param bOutEnableSLERP		Indicates whether the SLERP drive is enabled. Only relevant if the AngularDriveMode is set to SLERP
 	*/
-	UFUNCTION(BlueprintCallable, Category = "Physics|Constraints")
+	UFUNCTION(BlueprintCallable, Category = "Physics|Constraints", meta = (Keywords = "motor target"))
 	static void GetAngularVelocityDriveSLERP(
 		UPARAM(ref) FConstraintInstanceAccessor& Accessor,
 		bool& bOutEnableSLERP
 	);
 
 	/** Switches the angular drive mode between SLERP and Twist And Swing
-	*	@param ConstraintInstance	Instance of the constraint to change
+	*	@param Accessor		Constraint accessor to change
 	*	@param DriveMode	The angular drive mode to use. SLERP uses shortest spherical path, but will not work if any angular constraints are locked. Twist and Swing decomposes the path into the different angular degrees of freedom but may experience gimbal lock
 	*/
-	UFUNCTION(BlueprintCallable, Category = "Physics|Constraints")
+	UFUNCTION(BlueprintCallable, Category = "Physics|Constraints", meta = (Keywords = "motor"))
 	static void SetAngularDriveMode(
 		UPARAM(ref) FConstraintInstanceAccessor& Accessor,
 		EAngularDriveMode::Type DriveMode
 	);
 
 	/** Gets the angular drive mode ( SLERP or Twist And Swing)
-	*	@param ConstraintInstance	Instance of the constraint to change
+	*	@param Accessor		Constraint accessor to query
 	*	@param OutDriveMode	The angular drive mode to use. SLERP uses shortest spherical path, but will not work if any angular constraints are locked. Twist and Swing decomposes the path into the different angular degrees of freedom but may experience gimbal lock
 	*/
-	UFUNCTION(BlueprintCallable, Category = "Physics|Constraints")
+	UFUNCTION(BlueprintCallable, Category = "Physics|Constraints", meta = (Keywords = "motor"))
 	static void GetAngularDriveMode(
 		UPARAM(ref) FConstraintInstanceAccessor& Accessor,
 		TEnumAsByte<EAngularDriveMode::Type>& OutDriveMode
 	);
 
 	/** Sets the target orientation for the angular drive.
-	*	@param ConstraintInstance	Instance of the constraint to change
+	*	@param Accessor				Constraint accessor to change
 	*	@param InPosTarget			Target orientation
 	*/
-	UFUNCTION(BlueprintCallable, Category = "Physics|Constraints")
+	UFUNCTION(BlueprintCallable, Category = "Physics|Constraints", meta = (Keywords = "motor"))
 	static void SetAngularOrientationTarget(
 		UPARAM(ref) FConstraintInstanceAccessor& Accessor,
 		const FRotator& InPosTarget
 	);
 
 	/** Gets the target orientation for the angular drive.
-	*	@param ConstraintInstance	Instance of the constraint to change
+	*	@param Accessor				Constraint accessor to query
 	*	@param OutPosTarget			Target orientation
 	*/
-	UFUNCTION(BlueprintCallable, Category = "Physics|Constraints")
+	UFUNCTION(BlueprintCallable, Category = "Physics|Constraints", meta = (Keywords = "motor"))
 	static void GetAngularOrientationTarget(
 		UPARAM(ref) FConstraintInstanceAccessor& Accessor,
 		FRotator& OutPosTarget
 	);
 
 	/** Sets the target velocity for the angular drive.
-	*	@param ConstraintInstance	Instance of the constraint to change
+	*	@param Accessor				Constraint accessor to change
 	*	@param InVelTarget			Target velocity
 	*/
-	UFUNCTION(BlueprintCallable, Category = "Physics|Constraints")
+	UFUNCTION(BlueprintCallable, Category = "Physics|Constraints", meta = (Keywords = "motor"))
 	static void SetAngularVelocityTarget(
 		UPARAM(ref) FConstraintInstanceAccessor& Accessor,
 		const FVector& InVelTarget
 	);
 
 	/** Gets the target velocity for the angular drive.
-	*	@param ConstraintInstance	Instance of the constraint to change
+	*	@param Accessor				Constraint accessor to query
 	*	@param OutVelTarget			Target velocity
 	*/
-	UFUNCTION(BlueprintCallable, Category = "Physics|Constraints")
+	UFUNCTION(BlueprintCallable, Category = "Physics|Constraints", meta = (Keywords = "motor"))
 	static void GetAngularVelocityTarget(
 		UPARAM(ref) FConstraintInstanceAccessor& Accessor,
 		FVector& OutVelTarget
 	);
 
 	/** Sets the drive params for the angular drive.
-	*	@param ConstraintInstance	Instance of the constraint to change
+	*	@param Accessor				Constraint accessor to change
 	*	@param PositionStrength		Positional strength for the drive (stiffness)
 	*	@param VelocityStrength 	Velocity strength of the drive (damping)
 	*	@param InForceLimit			Max force applied by the drive
 	*/
-	UFUNCTION(BlueprintCallable, Category = "Physics|Constraints")
+	UFUNCTION(BlueprintCallable, Category = "Physics|Constraints", meta = (Keywords = "motor target orientation velocity strength max force"))
 	static void SetAngularDriveParams(
 		UPARAM(ref) FConstraintInstanceAccessor& Accessor,
 		float PositionStrength, 
@@ -535,12 +643,12 @@ public:
 	);
 
 	/** Gets the drive params for the angular drive.
-*	@param ConstraintInstance	Instance of the constraint to change
+*	@param Accessor				Constraint accessor to query
 *	@param OutPositionStrength	Positional strength for the drive (stiffness)
 *	@param OutVelocityStrength 	Velocity strength of the drive (damping)
 *	@param OutForceLimit		Max force applied by the drive
 */
-	UFUNCTION(BlueprintCallable, Category = "Physics|Constraints")
+	UFUNCTION(BlueprintCallable, Category = "Physics|Constraints", meta = (Keywords = "motor target orientation velocity strength max force"))
 	static void GetAngularDriveParams(
 		UPARAM(ref) FConstraintInstanceAccessor& Accessor,
 		float OutPositionStrength,
diff --git a/Engine/Source/Runtime/Engine/Private/PhysicsEngine/ConstraintInstanceBlueprintLibrary.cpp b/Engine/Source/Runtime/Engine/Private/PhysicsEngine/ConstraintInstanceBlueprintLibrary.cpp
--- a/Engine/Source/Runtime/Engine/Private/PhysicsEngine/ConstraintInstanceBlueprintLibrary.cpp
+++ b/Engine/Source/Runtime/Engine/Private/PhysicsEngine/ConstraintInstanceBlueprintLibrary.cpp
@@ -132,6 +132,49 @@ void UConstraintInstanceBlueprintLibrary::GetLinearLimits(
 	}
 }
 
+void UConstraintInstanceBlueprintLibrary::SetLinearSoftLimitParams(
+	FConstraintInstanceAccessor& Accessor,
+	bool bSoftLinearLimit,
+	float LinearLimitStiffness,
+	float LinearLimitDamping,
+	float LinearLimitRestitution,
+	float LinearLimitContactDistance
+)
+{
+	if (FConstraintInstance* ConstraintInstance = Accessor.Get())
+	{
+		ConstraintInstance->SetSoftLinearLimitParams(bSoftLinearLimit, LinearLimitStiffness, LinearLimitDamping, LinearLimitRestitution, LinearLimitContactDistance);
+	}
+}
+
+void UConstraintInstanceBlueprintLibrary::GetLinearSoftLimitParams(
+	FConstraintInstanceAccessor& Accessor,
+	bool& bSoftLinearLimit,
+	float& LinearLimitStiffness,
+	float& LinearLimitDamping,
+	float& LinearLimitRestitution,
+	float& LinearLimitContactDistance
+)
+{
+	if (FConstraintInstance* ConstraintInstance = Accessor.Get())
+	{
+		bSoftLinearLimit = ConstraintInstance->GetIsSoftLinearLimit();
+		LinearLimitStiffness = ConstraintInstance->GetSoftLinearLimitStiffness();
+		LinearLimitDamping = ConstraintInstance->GetSoftLinearLimitDamping();
+		LinearLimitRestitution = ConstraintInstance->GetSoftLinearLimitRestitution();
+		LinearLimitContactDistance = ConstraintInstance->GetSoftLinearLimitContactDistance();
+	}
+	else
+	{
+		bSoftLinearLimit = false;
+		LinearLimitStiffness = 0.;
+		LinearLimitDamping = 0.;
+		LinearLimitRestitution = 0.;
+		LinearLimitContactDistance = 0.;
+
+	}
+}
+
 void UConstraintInstanceBlueprintLibrary::SetLinearBreakable(
 	FConstraintInstanceAccessor& Accessor,
 	bool bLinearBreakable,
@@ -216,6 +259,92 @@ void UConstraintInstanceBlueprintLibrary::GetAngularLimits(
 	}
 }
 
+void UConstraintInstanceBlueprintLibrary::SetAngularSoftSwingLimitParams(
+	FConstraintInstanceAccessor& Accessor,
+	bool bSoftSwingLimit,
+	float SwingLimitStiffness,
+	float SwingLimitDamping,
+	float SwingLimitRestitution,
+	float SwingLimitContactDistance
+)
+{
+	if (FConstraintInstance* ConstraintInstance = Accessor.Get())
+	{
+		ConstraintInstance->SetSoftSwingLimitParams(bSoftSwingLimit, SwingLimitStiffness, SwingLimitDamping, SwingLimitRestitution, SwingLimitContactDistance);
+	}
+}
+
+void UConstraintInstanceBlueprintLibrary::GetAngularSoftSwingLimitParams(
+	FConstraintInstanceAccessor& Accessor,
+	bool& bSoftSwingLimit,
+	float& SwingLimitStiffness,
+	float& SwingLimitDamping,
+	float& SwingLimitRestitution,
+	float& SwingLimitContactDistance
+)
+{
+	if (FConstraintInstance* ConstraintInstance = Accessor.Get())
+	{
+		bSoftSwingLimit = ConstraintInstance->GetIsSoftSwingLimit();
+		SwingLimitStiffness = ConstraintInstance->GetSoftSwingLimitStiffness();
+		SwingLimitDamping = ConstraintInstance->GetSoftSwingLimitDamping();
+		SwingLimitRestitution = ConstraintInstance->GetSoftSwingLimitRestitution();
+		SwingLimitContactDistance = ConstraintInstance->GetSoftSwingLimitContactDistance();
+	}
+	else
+	{
+		bSoftSwingLimit = false;
+		SwingLimitStiffness = 0.;
+		SwingLimitDamping = 0.;
+		SwingLimitRestitution = 0.;
+		SwingLimitContactDistance = 0.;
+		
+	}
+}
+
+void UConstraintInstanceBlueprintLibrary::SetAngularSoftTwistLimitParams(
+	FConstraintInstanceAccessor& Accessor,
+	bool bSoftTwistLimit,
+	float TwistLimitStiffness,
+	float TwistLimitDamping,
+	float TwistLimitRestitution,
+	float TwistLimitContactDistance
+)
+{
+	if (FConstraintInstance* ConstraintInstance = Accessor.Get())
+	{
+		ConstraintInstance->SetSoftTwistLimitParams(bSoftTwistLimit, TwistLimitStiffness, TwistLimitDamping, TwistLimitRestitution, TwistLimitContactDistance);
+	}
+}
+
+void UConstraintInstanceBlueprintLibrary::GetAngularSoftTwistLimitParams(
+	FConstraintInstanceAccessor& Accessor,
+	bool& bSoftTwistLimit,
+	float& TwistLimitStiffness,
+	float& TwistLimitDamping,
+	float& TwistLimitRestitution,
+	float& TwistLimitContactDistance
+)
+{
+	if (FConstraintInstance* ConstraintInstance = Accessor.Get())
+	{
+		bSoftTwistLimit = ConstraintInstance->GetIsSoftTwistLimit();
+		TwistLimitStiffness = ConstraintInstance->GetSoftTwistLimitStiffness();
+		TwistLimitDamping = ConstraintInstance->GetSoftTwistLimitDamping();
+		TwistLimitRestitution = ConstraintInstance->GetSoftTwistLimitRestitution();
+		TwistLimitContactDistance = ConstraintInstance->GetSoftTwistLimitContactDistance();
+	}
+	else
+	{
+		bSoftTwistLimit = false;
+		TwistLimitStiffness = 0.;
+		TwistLimitDamping = 0.;
+		TwistLimitRestitution = 0.;
+		TwistLimitContactDistance = 0.;
+
+	}
+}
+
 void UConstraintInstanceBlueprintLibrary::SetAngularBreakable(
 	FConstraintInstanceAccessor& Accessor,
 	bool bAngularBreakable,
