Fix crash due to HiRes bulk data being null after convertion from skeletal mesh to static mesh

diff --git a/Engine/Source/Runtime/Engine/Private/StaticMeshBuild.cpp b/Engine/Source/Runtime/Engine/Private/StaticMeshBuild.cpp
--- a/Engine/Source/Runtime/Engine/Private/StaticMeshBuild.cpp
+++ b/Engine/Source/Runtime/Engine/Private/StaticMeshBuild.cpp
@@ -376,6 +376,12 @@ void UStaticMesh::BeginBuildInternal(FStaticMeshBuildContext* Context)
 	GetStaticMeshDerivedDataVersion();
 
 	PreMeshBuild.Broadcast(this);
+	
+	// Make sure the BulkData is always created before going async
+	if (GetHiResSourceModel().StaticMeshDescriptionBulkData == nullptr)
+	{
+		GetHiResSourceModel().StaticMeshDescriptionBulkData = NewObject<UStaticMeshDescriptionBulkData>(this, NAME_None, RF_Transactional);
+	}
 
 	// Ensure we have a bodysetup.
 	CreateBodySetup();
