Reinstated the hi res source model MeshDescription bulk data as a default subobject; this ensures it is always initialized.

diff --git a/Engine/Source/Runtime/Engine/Private/StaticMesh.cpp b/Engine/Source/Runtime/Engine/Private/StaticMesh.cpp
--- a/Engine/Source/Runtime/Engine/Private/StaticMesh.cpp
+++ b/Engine/Source/Runtime/Engine/Private/StaticMesh.cpp
@@ -2897,6 +2897,8 @@ UStaticMesh::UStaticMesh(const FObjectInitializer& ObjectInitializer)
 	ImportVersion = EImportStaticMeshVersion::BeforeImportStaticMeshVersionWasAdded;
 	LODForOccluderMesh = -1;
 	NumStreamedLODs.Default = -1;
+	GetHiResSourceModel().StaticMeshDescriptionBulkData = CreateDefaultSubobject<UStaticMeshDescriptionBulkData>(TEXT("HiResMeshDescription"));
+	GetHiResSourceModel().StaticMeshDescriptionBulkData->SetFlags(RF_Transactional);
 #endif // #if WITH_EDITORONLY_DATA
 	SetLightMapResolution(4);
 	SetMinLOD(0);
@@ -5283,11 +5285,6 @@ void UStaticMesh::BeginPostLoadInternal(FStaticMeshPostLoadContext& Context)
 	{
 		ComplexCollisionMesh->ConditionalPostLoad();
 	}
-
-	if (GetHiResSourceModel().StaticMeshDescriptionBulkData == nullptr)
-	{
-		GetHiResSourceModel().StaticMeshDescriptionBulkData = NewObject<UStaticMeshDescriptionBulkData>(this, NAME_None, RF_Transactional);
-	}
 #endif //WITH_EDITORONLY_DATA
 
 	// We want to always have a BodySetup, its used for per-poly collision as well
diff --git a/Engine/Source/Runtime/Engine/Private/StaticMeshBuild.cpp b/Engine/Source/Runtime/Engine/Private/StaticMeshBuild.cpp
--- a/Engine/Source/Runtime/Engine/Private/StaticMeshBuild.cpp
+++ b/Engine/Source/Runtime/Engine/Private/StaticMeshBuild.cpp
@@ -377,12 +377,6 @@ void UStaticMesh::BeginBuildInternal(FStaticMeshBuildContext* Context)
 
 	PreMeshBuild.Broadcast(this);
 	
-	// Make sure the BulkData is always created before going async
-	if (GetHiResSourceModel().StaticMeshDescriptionBulkData == nullptr)
-	{
-		GetHiResSourceModel().StaticMeshDescriptionBulkData = NewObject<UStaticMeshDescriptionBulkData>(this, NAME_None, RF_Transactional);
-	}
-
 	// Ensure we have a bodysetup.
 	CreateBodySetup();
 	check(GetBodySetup() != nullptr);
