Strata material conversion now always automatically runs when strata is enabled.

diff --git a/Engine/Source/Runtime/Engine/Private/Materials/HLSLMaterialTranslator.cpp b/Engine/Source/Runtime/Engine/Private/Materials/HLSLMaterialTranslator.cpp
--- a/Engine/Source/Runtime/Engine/Private/Materials/HLSLMaterialTranslator.cpp
+++ b/Engine/Source/Runtime/Engine/Private/Materials/HLSLMaterialTranslator.cpp
@@ -1516,8 +1516,11 @@ bool FHLSLMaterialTranslator::Translate()
 			{
 				Errorf(TEXT("SingleLayerWater materials cannot be combined with other shading models.")); // Simply untested for now
 			}
+
+			static const auto CVarStrata = IConsoleManager::Get().FindTConsoleVariableDataInt(TEXT("r.Strata"));
+			const bool bStrata = CVarStrata ? CVarStrata->GetValueOnAnyThread() > 0 : false;
 			if (Material->GetMaterialInterface() && !Material->GetMaterialInterface()->GetMaterial()->HasAnyExpressionsInMaterialAndFunctionsOfType<UMaterialExpressionSingleLayerWaterMaterialOutput>()
-				&& CONVERT_MATERIAL_TO_STRATA_ON_LOAD == 0
+				&& !bStrata
 				)
 			{
 				Errorf(TEXT("SingleLayerWater materials requires the use of SingleLayerWaterMaterial output node."));
diff --git a/Engine/Source/Runtime/Engine/Private/Materials/Material.cpp b/Engine/Source/Runtime/Engine/Private/Materials/Material.cpp
--- a/Engine/Source/Runtime/Engine/Private/Materials/Material.cpp
+++ b/Engine/Source/Runtime/Engine/Private/Materials/Material.cpp
@@ -3616,7 +3616,15 @@ void UMaterial::GetQualityLevelUsage(TArray<bool, TInlineAllocator<EMaterialQual
 
 void UMaterial::ConvertMaterialToStrataMaterial()
 {
-#if WITH_EDITOR && CONVERT_MATERIAL_TO_STRATA_ON_LOAD
+#if WITH_EDITOR
+	static const auto CVarStrata = IConsoleManager::Get().FindTConsoleVariableDataInt(TEXT("r.Strata"));
+	const bool bStrata = CVarStrata ? CVarStrata->GetValueOnAnyThread() > 0 : false;
+
+	if (!bStrata)
+	{
+		return;
+	}
+
 	const URendererSettings* RendererSettings = GetDefault<URendererSettings>();
 	const bool bStrataEnabled = RendererSettings && RendererSettings->bEnableStrata;
 
diff --git a/Engine/Source/Runtime/Engine/Private/Materials/MaterialExpressions.cpp b/Engine/Source/Runtime/Engine/Private/Materials/MaterialExpressions.cpp
--- a/Engine/Source/Runtime/Engine/Private/Materials/MaterialExpressions.cpp
+++ b/Engine/Source/Runtime/Engine/Private/Materials/MaterialExpressions.cpp
@@ -19167,8 +19167,11 @@ int32 UMaterialExpressionSingleLayerWaterMaterialOutput::Compile(class FMaterial
 {
 	int32 CodeInput = INDEX_NONE;
 
+	static const auto CVarStrata = IConsoleManager::Get().FindTConsoleVariableDataInt(TEXT("r.Strata"));
+	const bool bStrata = CVarStrata ? CVarStrata->GetValueOnAnyThread() > 0 : false;
+
 	if (!ScatteringCoefficients.IsConnected() && !AbsorptionCoefficients.IsConnected() && !PhaseG.IsConnected()
-		&& CONVERT_MATERIAL_TO_STRATA_ON_LOAD == 0)
+		&& !bStrata)
 	{
 		Compiler->Error(TEXT("No inputs to Single Layer Water Material."));
 	}
diff --git a/Engine/Source/Runtime/Engine/Private/Materials/StrataMaterial.h b/Engine/Source/Runtime/Engine/Private/Materials/StrataMaterial.h
--- a/Engine/Source/Runtime/Engine/Private/Materials/StrataMaterial.h
+++ b/Engine/Source/Runtime/Engine/Private/Materials/StrataMaterial.h
@@ -8,15 +8,6 @@
 
 
 
-// Material conversion to strata is disabled until ready.
-// To enable it locally:
-//		- Set CONVERT_MATERIAL_TO_STRATA_ON_LOAD to 1
-//		- Locally change ShaderVersion.ush GUID
-// !!!!!! !!!!!! DO NOT CHECK THIS IN SET TO 1 !!!!!! !!!!!!
-#define CONVERT_MATERIAL_TO_STRATA_ON_LOAD 0
-
-
-
 class FMaterialCompiler;
 
 
