Added MaterialExpressionStrataTransmittanceToMFP helper node  and conversion from ThinTranslucency to Strata.

diff --git a/Engine/Shaders/Private/ParticipatingMediaCommon.ush b/Engine/Shaders/Private/ParticipatingMediaCommon.ush
--- a/Engine/Shaders/Private/ParticipatingMediaCommon.ush
+++ b/Engine/Shaders/Private/ParticipatingMediaCommon.ush
@@ -15,13 +15,15 @@ struct FParticipatingMedia
 	float3 Albedo;			// (unitless)
 };
 
+#define PARTICIPATING_MEDIA_MIN_MFP		0.000001f
+
 // The mean free path must be in meters
 FParticipatingMedia CreateMediumFromAlbedoMFP(float3 Albedo, float3 MeanFreePathMeters)
 {
 	FParticipatingMedia PM = (FParticipatingMedia)0;
 	PM.Albedo = Albedo;
 	PM.MeanFreePath = MeanFreePathMeters;
-	PM.ExtinctionCoef = 1.0f / max(0.000001f, PM.MeanFreePath);
+	PM.ExtinctionCoef = 1.0f / max(PARTICIPATING_MEDIA_MIN_MFP, PM.MeanFreePath);
 	PM.ScatteringCoef = Albedo * PM.ExtinctionCoef;
 	PM.AbsorptionCoef = max(0.0f, PM.ExtinctionCoef - PM.ScatteringCoef);
 	return PM;
@@ -86,3 +88,15 @@ float4 ImportanceSampleHenyeyGreensteinPhase(float2 E, float G)
 	return float4(H, HenyeyGreensteinPhase(G, CosTheta));
 }
 
+float3 TransmittanceToExtinction(in float3 TransmittanceColor, in float Thickness)
+{
+	// TransmittanceColor	= exp(-Extinction * Thickness)
+	// Extinction			= -log(TransmittanceColor) / Thickness
+	return -log(saturate(TransmittanceColor)) / max(PARTICIPATING_MEDIA_MIN_MFP, Thickness);
+}
+
+float3 TransmittanceToMeanFreePath(in float3 TransmittanceColor, in float Thickness)
+{
+	return 1.0f / TransmittanceToExtinction(TransmittanceColor, Thickness);
+}
+
diff --git a/Engine/Source/Runtime/Engine/Classes/Materials/MaterialExpressionStrata.h b/Engine/Source/Runtime/Engine/Classes/Materials/MaterialExpressionStrata.h
--- a/Engine/Source/Runtime/Engine/Classes/Materials/MaterialExpressionStrata.h
+++ b/Engine/Source/Runtime/Engine/Classes/Materials/MaterialExpressionStrata.h
@@ -550,3 +550,26 @@ class UMaterialExpressionStrataAnisotropyToRoughness : public UMaterialExpressio
 	//~ End UMaterialExpression Interface
 };
 
+UCLASS(MinimalAPI, collapsecategories, hidecategories = Object, DisplayName = "Strata Transmittance-To-MeanFreePath")
+class UMaterialExpressionStrataTransmittanceToMFP : public UMaterialExpressionStrataBSDF
+{
+	GENERATED_UCLASS_BODY()
+
+	/**
+	* The colored transmittance for a view perpendicular to the surface. The transmittance or other view orientation will automatically deduced according to surface thickness.
+	*/
+	UPROPERTY()
+	FExpressionInput TransmittanceColor;
+
+	//~ Begin UMaterialExpression Interface
+#if WITH_EDITOR
+	virtual int32 Compile(class FMaterialCompiler* Compiler, int32 OutputIndex) override;
+	virtual void GetCaption(TArray<FString>& OutCaptions) const override;
+	virtual uint32 GetOutputType(int32 OutputIndex) override;
+	virtual uint32 GetInputType(int32 InputIndex) override;
+	virtual void GetConnectorToolTip(int32 InputIndex, int32 OutputIndex, TArray<FString>& OutToolTip) override;
+	virtual void GetExpressionToolTip(TArray<FString>& OutToolTip) override;
+#endif
+	//~ End UMaterialExpression Interface
+};
+
diff --git a/Engine/Source/Runtime/Engine/Private/Materials/HLSLMaterialTranslator.cpp b/Engine/Source/Runtime/Engine/Private/Materials/HLSLMaterialTranslator.cpp
--- a/Engine/Source/Runtime/Engine/Private/Materials/HLSLMaterialTranslator.cpp
+++ b/Engine/Source/Runtime/Engine/Private/Materials/HLSLMaterialTranslator.cpp
@@ -9188,6 +9188,25 @@ int32 FHLSLMaterialTranslator::StrataAnisotropyToRoughness(int32 Roughness, int3
 		*GetParameterCode(Anisotropy));
 }
 
+int32 FHLSLMaterialTranslator::StrataTransmittanceToMFP(int32 TransmittanceColor, int32 OutputIndex)
+{
+	if (OutputIndex == INDEX_NONE)
+	{
+		return INDEX_NONE;
+	}
+
+	switch (OutputIndex)
+	{
+	case 0:
+		return AddCodeChunk(MCT_Float3, TEXT("TransmittanceToMeanFreePath(%s, STRATA_SIMPLEVOLUME_THICKNESS_CM)"), *GetParameterCode(TransmittanceColor));
+		break;
+	case 1:
+		return AddInlinedCodeChunk(MCT_Float1, TEXT("STRATA_SIMPLEVOLUME_THICKNESS_CM")); // Thickness
+		break;
+	}
+	return INDEX_NONE;
+}
+
 int32 FHLSLMaterialTranslator::MapARPassthroughCameraUV(int32 UV)
 {
 	if (UV == INDEX_NONE)
diff --git a/Engine/Source/Runtime/Engine/Private/Materials/HLSLMaterialTranslator.h b/Engine/Source/Runtime/Engine/Private/Materials/HLSLMaterialTranslator.h
--- a/Engine/Source/Runtime/Engine/Private/Materials/HLSLMaterialTranslator.h
+++ b/Engine/Source/Runtime/Engine/Private/Materials/HLSLMaterialTranslator.h
@@ -1006,6 +1006,7 @@ protected:
 	virtual int32 StrataAdd(int32 A, int32 B) override;
 	virtual int32 StrataMultiply(int32 A, int32 Weight) override;
 	virtual int32 StrataAnisotropyToRoughness(int32 Roughness, int32 Anisotropy, int32 OutputIndex) override;
+	virtual int32 StrataTransmittanceToMFP(int32 TransmittanceColor, int32 OutputIndex) override;
 
 	virtual void StrataCompilationInfoRegisterCodeChunk(int32 CodeChunk, FStrataMaterialCompilationInfo& StrataMaterialCompilationInfo) override;
 	virtual bool StrataCompilationInfoContainsCodeChunk(int32 CodeChunk) override;
diff --git a/Engine/Source/Runtime/Engine/Private/Materials/Material.cpp b/Engine/Source/Runtime/Engine/Private/Materials/Material.cpp
--- a/Engine/Source/Runtime/Engine/Private/Materials/Material.cpp
+++ b/Engine/Source/Runtime/Engine/Private/Materials/Material.cpp
@@ -80,6 +80,7 @@
 #include "Materials/MaterialExpressionSingleLayerWaterMaterialOutput.h"
 #include "Materials/MaterialExpressionStrata.h"
 #include "Materials/StrataMaterial.h"
+#include "Materials/MaterialExpressionThinTranslucentMaterialOutput.h"
 
 #if WITH_EDITOR
 #include "Logging/TokenizedMessage.h"
@@ -3740,6 +3741,59 @@ void UMaterial::ConvertMaterialToStrataMaterial()
 
 			FrontMaterial.Connect(0, SLWBSDF);
 		}
+		else if (MaterialDomain == MD_Surface && ShadingModel == MSM_ThinTranslucent)
+		{
+			// Top slab BSDF as a simple Disney material
+			UMaterialExpressionStrataSlabBSDF* TopSlabBSDF = NewObject<UMaterialExpressionStrataSlabBSDF>(this);
+			MoveConnectionTo(BaseColor, TopSlabBSDF, 0);								// BaseColor
+			MoveConnectionTo(Metallic, TopSlabBSDF, 2);									// Metallic
+			CopyConnectionTo(Specular, TopSlabBSDF, 3);									// Specular
+			CopyConnectionTo(Roughness, TopSlabBSDF, 4);								// Roughness
+			CopyConnectionTo(Normal, TopSlabBSDF, 6);									// Normal
+			MoveConnectionTo(EmissiveColor, TopSlabBSDF, 10);							// Emissive
+
+			// Now weight the top base material by opacity
+			UMaterialExpressionStrataMultiply* TopMaterialWithCoverage = NewObject<UMaterialExpressionStrataMultiply>(this);
+			TopMaterialWithCoverage->GetInput(0)->Connect(0, TopSlabBSDF);				// TopSlabBSDF -> A
+			MoveConnectionTo(Opacity, TopMaterialWithCoverage, 1);						// Opacity -> Weight
+
+			// Bottom slab BSDF will be a simple absorption only layer
+			UMaterialExpressionStrataSlabBSDF* BottomSlabBSDF = NewObject<UMaterialExpressionStrataSlabBSDF>(this);
+			// Assign specular properties shared with the top layer.
+			MoveConnectionTo(Specular, BottomSlabBSDF, 3);								// Specular
+			MoveConnectionTo(Roughness, BottomSlabBSDF, 4);								// Roughness
+			MoveConnectionTo(Normal, BottomSlabBSDF, 6);								// Normal
+
+			TArray<class UMaterialExpressionCustomOutput*> CustomOutputExpressions;
+			GetAllCustomOutputExpressions(CustomOutputExpressions);
+			UMaterialExpressionThinTranslucentMaterialOutput* TTCustomOutput = nullptr;
+			for (UMaterialExpressionCustomOutput* Expression : CustomOutputExpressions)
+			{
+				TTCustomOutput = Cast<UMaterialExpressionThinTranslucentMaterialOutput>(Expression);
+				if (TTCustomOutput)
+				{
+					break;
+				}
+			}
+			if (TTCustomOutput)
+			{
+				UMaterialExpressionStrataTransmittanceToMFP* TransToMDFP = NewObject<UMaterialExpressionStrataTransmittanceToMFP>(this);
+				MoveConnectionTo(*TTCustomOutput->GetInput(0), TransToMDFP, 0);			// TransmittanceColor -> TransmittanceColor
+				BottomSlabBSDF->GetInput(8 )->Connect(0, TransToMDFP);					// MFP -> MFP
+				BottomSlabBSDF->GetInput(13)->Connect(1, TransToMDFP);					// Thickness -> Thickness
+
+				UMaterialExpressionStrataVerticalLayering* VerticalLayering = NewObject<UMaterialExpressionStrataVerticalLayering>(this);
+				VerticalLayering->GetInput(0)->Connect(0, TopMaterialWithCoverage);		// Top -> Top
+				VerticalLayering->GetInput(1)->Connect(0, BottomSlabBSDF);				// Bottom -> Base
+
+				// Connect the result of the layering to the front material
+				FrontMaterial.Connect(0, VerticalLayering);
+			}
+			else
+			{
+				FrontMaterial.Connect(0, TopSlabBSDF);
+			}
+		}
 		else if (MaterialDomain == MD_Volume)
 		{
 			UMaterialExpressionStrataVolumetricFogCloudBSDF* VolBSDF = NewObject<UMaterialExpressionStrataVolumetricFogCloudBSDF>(this);
diff --git a/Engine/Source/Runtime/Engine/Private/Materials/MaterialExpressions.cpp b/Engine/Source/Runtime/Engine/Private/Materials/MaterialExpressions.cpp
--- a/Engine/Source/Runtime/Engine/Private/Materials/MaterialExpressions.cpp
+++ b/Engine/Source/Runtime/Engine/Private/Materials/MaterialExpressions.cpp
@@ -20729,6 +20729,79 @@ uint32 UMaterialExpressionStrataAnisotropyToRoughness::GetInputType(int32 InputI
 }
 #endif // WITH_EDITOR
 
+UMaterialExpressionStrataTransmittanceToMFP::UMaterialExpressionStrataTransmittanceToMFP(const FObjectInitializer& ObjectInitializer)
+	: Super(ObjectInitializer)
+{
+	struct FConstructorStatics
+	{
+		FText NAME_Strata;
+		FConstructorStatics() : NAME_Strata(LOCTEXT("Strata Helpers", "Strata Helpers")) { }
+	};
+	static FConstructorStatics ConstructorStatics;
+#if WITH_EDITORONLY_DATA
+	MenuCategories.Add(ConstructorStatics.NAME_Strata);
+
+	bShowOutputNameOnPin = true;
+
+	Outputs.Reset();
+	Outputs.Add(FExpressionOutput(TEXT("MFP")));
+	Outputs.Add(FExpressionOutput(TEXT("Thickness")));
+#endif
+}
+
+#if WITH_EDITOR
+int32 UMaterialExpressionStrataTransmittanceToMFP::Compile(class FMaterialCompiler* Compiler, int32 OutputIndex)
+{
+	return Compiler->StrataTransmittanceToMFP(TransmittanceColor.GetTracedInput().Expression ? TransmittanceColor.Compile(Compiler) : Compiler->Constant(0.5f), OutputIndex);
+}
+
+void UMaterialExpressionStrataTransmittanceToMFP::GetCaption(TArray<FString>& OutCaptions) const
+{
+	OutCaptions.Add(TEXT("Strata Transmittance-To-MeanFreePath"));
+}
+
+uint32 UMaterialExpressionStrataTransmittanceToMFP::GetOutputType(int32 OutputIndex)
+{
+	switch (OutputIndex)
+	{
+	case 0:
+		return MCT_Float3; // MFP
+		break;
+	case 1:
+		return MCT_Float1; // Thickness
+		break;
+	}
+
+	check(false);
+	return MCT_Float1;
+}
+
+uint32 UMaterialExpressionStrataTransmittanceToMFP::GetInputType(int32 InputIndex)
+{
+	return MCT_Float3;
+}
+void UMaterialExpressionStrataTransmittanceToMFP::GetConnectorToolTip(int32 InputIndex, int32 OutputIndex, TArray<FString>& OutToolTip)
+{
+	switch (OutputIndex)
+	{
+	case 0:
+		ConvertToMultilineToolTip(TEXT("The Mean Free Path in centimeters defining the participating media constituting the slab of material."), 80, OutToolTip);
+		break;
+	case 1:
+		ConvertToMultilineToolTip(TEXT("The thickness of the slab of material."), 80, OutToolTip);
+		break;
+	}
+}
+
+void UMaterialExpressionStrataTransmittanceToMFP::GetExpressionToolTip(TArray<FString>& OutToolTip)
+{
+	ConvertToMultilineToolTip(TEXT("Convert a transmittance color corresponding to a loab of participating media viewed perpendicularly to its surface.\nThis node directly maps to the Slab BSDF input."), 80, OutToolTip);
+
+}
+#endif // WITH_EDITOR
+
+
+
 UMaterialExpressionExecBegin::UMaterialExpressionExecBegin(const FObjectInitializer& ObjectInitializer)
 	: Super(ObjectInitializer)
 {
diff --git a/Engine/Source/Runtime/Engine/Public/MaterialCompiler.h b/Engine/Source/Runtime/Engine/Public/MaterialCompiler.h
--- a/Engine/Source/Runtime/Engine/Public/MaterialCompiler.h
+++ b/Engine/Source/Runtime/Engine/Public/MaterialCompiler.h
@@ -426,6 +426,7 @@ public:
 	virtual int32 StrataAdd(int32 A, int32 B) = 0;
 	virtual int32 StrataMultiply(int32 A, int32 Weight) = 0;
 	virtual int32 StrataAnisotropyToRoughness(int32 Roughness, int32 Anisotropy, int32 OutputIndex) = 0;
+	virtual int32 StrataTransmittanceToMFP(int32 TransmittanceColor, int32 OutputIndex) = 0;
 
 	virtual void StrataCompilationInfoRegisterCodeChunk(int32 CodeChunk, FStrataMaterialCompilationInfo& StrataMaterialCompilationInfo) = 0;
 	virtual bool StrataCompilationInfoContainsCodeChunk(int32 CodeChunk) = 0;
@@ -958,6 +959,11 @@ public:
 		return Compiler->StrataAnisotropyToRoughness(Roughness, Anisotropy, OutputIndex);
 	}
 
+	virtual int32 StrataTransmittanceToMFP(int32 TransmittanceColor, int32 OutputIndex) override
+	{
+		return Compiler->StrataTransmittanceToMFP(TransmittanceColor, OutputIndex);
+	}
+
 	virtual void StrataCompilationInfoRegisterCodeChunk(int32 CodeChunk, FStrataMaterialCompilationInfo& StrataMaterialCompilationInfo) override
 	{
 		Compiler->StrataCompilationInfoRegisterCodeChunk(CodeChunk, StrataMaterialCompilationInfo);
