Deleted now useless Strata AnisotropyToRoughness helper node.

diff --git a/Engine/Source/Runtime/Engine/Classes/Materials/MaterialExpressionStrata.h b/Engine/Source/Runtime/Engine/Classes/Materials/MaterialExpressionStrata.h
--- a/Engine/Source/Runtime/Engine/Classes/Materials/MaterialExpressionStrata.h
+++ b/Engine/Source/Runtime/Engine/Classes/Materials/MaterialExpressionStrata.h
@@ -523,33 +523,6 @@ class UMaterialExpressionStrataMultiply : public UMaterialExpressionStrataBSDF
 ///////////////////////////////////////////////////////////////////////////////
 // Utilities
 
-UCLASS(MinimalAPI, collapsecategories, hidecategories = Object, DisplayName = "Strata Anisotropy-To-Roughnesses")
-class UMaterialExpressionStrataAnisotropyToRoughness : public UMaterialExpressionStrataBSDF
-{
-	GENERATED_UCLASS_BODY()
-
-	/**
-	 * Input roughness
-	 */
-	UPROPERTY()
-	FExpressionInput Roughness;
-
-	/**
-	 * Anisotropy factor - 0: isotropic behavior, -1: anisotropy along the bitangent vector, 1: anisotropy along the tangent vector
-	 */
-	UPROPERTY()
-	FExpressionInput Anisotropy;
-
-	//~ Begin UMaterialExpression Interface
-#if WITH_EDITOR
-	virtual int32 Compile(class FMaterialCompiler* Compiler, int32 OutputIndex) override;
-	virtual void GetCaption(TArray<FString>& OutCaptions) const override;
-	virtual uint32 GetOutputType(int32 OutputIndex) override;
-	virtual uint32 GetInputType(int32 InputIndex) override;
-#endif
-	//~ End UMaterialExpression Interface
-};
-
 UCLASS(MinimalAPI, collapsecategories, hidecategories = Object, DisplayName = "Strata Transmittance-To-MeanFreePath")
 class UMaterialExpressionStrataTransmittanceToMFP : public UMaterialExpressionStrataBSDF
 {
diff --git a/Engine/Source/Runtime/Engine/Private/Materials/HLSLMaterialTranslator.cpp b/Engine/Source/Runtime/Engine/Private/Materials/HLSLMaterialTranslator.cpp
--- a/Engine/Source/Runtime/Engine/Private/Materials/HLSLMaterialTranslator.cpp
+++ b/Engine/Source/Runtime/Engine/Private/Materials/HLSLMaterialTranslator.cpp
@@ -9171,23 +9171,6 @@ int32 FHLSLMaterialTranslator::StrataMultiply(int32 A, int32 Weight)
 	);
 }
 
-int32 FHLSLMaterialTranslator::StrataAnisotropyToRoughness(int32 Roughness, int32 Anisotropy, int32 OutputIndex)
-{
-	if (Roughness == INDEX_NONE)
-	{
-		return INDEX_NONE;
-	}
-	else if (Anisotropy == INDEX_NONE)
-	{
-		return INDEX_NONE;
-	}
-	return AddCodeChunk(
-		MCT_Float2,
-		TEXT("GetAnisotropicRoughness(%s,%s)"),
-		*GetParameterCode(Roughness),
-		*GetParameterCode(Anisotropy));
-}
-
 int32 FHLSLMaterialTranslator::StrataTransmittanceToMFP(int32 TransmittanceColor, int32 OutputIndex)
 {
 	if (OutputIndex == INDEX_NONE)
diff --git a/Engine/Source/Runtime/Engine/Private/Materials/HLSLMaterialTranslator.h b/Engine/Source/Runtime/Engine/Private/Materials/HLSLMaterialTranslator.h
--- a/Engine/Source/Runtime/Engine/Private/Materials/HLSLMaterialTranslator.h
+++ b/Engine/Source/Runtime/Engine/Private/Materials/HLSLMaterialTranslator.h
@@ -1005,7 +1005,6 @@ protected:
 	virtual int32 StrataVerticalLayering(int32 Top, int32 Base) override;
 	virtual int32 StrataAdd(int32 A, int32 B) override;
 	virtual int32 StrataMultiply(int32 A, int32 Weight) override;
-	virtual int32 StrataAnisotropyToRoughness(int32 Roughness, int32 Anisotropy, int32 OutputIndex) override;
 	virtual int32 StrataTransmittanceToMFP(int32 TransmittanceColor, int32 OutputIndex) override;
 
 	virtual void StrataCompilationInfoRegisterCodeChunk(int32 CodeChunk, FStrataMaterialCompilationInfo& StrataMaterialCompilationInfo) override;
diff --git a/Engine/Source/Runtime/Engine/Private/Materials/MaterialExpressions.cpp b/Engine/Source/Runtime/Engine/Private/Materials/MaterialExpressions.cpp
--- a/Engine/Source/Runtime/Engine/Private/Materials/MaterialExpressions.cpp
+++ b/Engine/Source/Runtime/Engine/Private/Materials/MaterialExpressions.cpp
@@ -20674,61 +20674,6 @@ void UMaterialExpressionStrataMultiply::GatherStrataMaterialInfo(FStrataMaterial
 
 
 
-UMaterialExpressionStrataAnisotropyToRoughness::UMaterialExpressionStrataAnisotropyToRoughness(const FObjectInitializer& ObjectInitializer)
-	: Super(ObjectInitializer)
-{
-	struct FConstructorStatics
-	{
-		FText NAME_Strata;
-		FConstructorStatics() : NAME_Strata(LOCTEXT("Strata Helpers", "Strata Helpers")) { }
-	};
-	static FConstructorStatics ConstructorStatics;
-#if WITH_EDITORONLY_DATA
-	MenuCategories.Add(ConstructorStatics.NAME_Strata);
-
-	bShowOutputNameOnPin = true;
-
-	Outputs.Reset();
-	Outputs.Add(FExpressionOutput(TEXT("Roughnesses")));
-#endif
-}
-
-#if WITH_EDITOR
-int32 UMaterialExpressionStrataAnisotropyToRoughness::Compile(class FMaterialCompiler* Compiler, int32 OutputIndex)
-{
-	return Compiler->StrataAnisotropyToRoughness(
-		Roughness.GetTracedInput().Expression ? Roughness.Compile(Compiler) : Compiler->Constant(0.5f),
-		Anisotropy.GetTracedInput().Expression ? Anisotropy.Compile(Compiler) : Compiler->Constant(0.f),
-		OutputIndex);
-}
-
-void UMaterialExpressionStrataAnisotropyToRoughness::GetCaption(TArray<FString>& OutCaptions) const
-{
-	OutCaptions.Add(TEXT("Strata Anisotropy-To-Roughnesses"));
-}
-
-uint32 UMaterialExpressionStrataAnisotropyToRoughness::GetOutputType(int32 OutputIndex)
-{
-	return MCT_Float2;
-}
-
-uint32 UMaterialExpressionStrataAnisotropyToRoughness::GetInputType(int32 InputIndex)
-{
-	switch (InputIndex)
-	{
-	case 0:
-		return MCT_Float1;
-		break;
-	case 1:
-		return MCT_Float1;
-		break;
-	}
-
-	check(false);
-	return MCT_Float1;
-}
-#endif // WITH_EDITOR
-
 UMaterialExpressionStrataTransmittanceToMFP::UMaterialExpressionStrataTransmittanceToMFP(const FObjectInitializer& ObjectInitializer)
 	: Super(ObjectInitializer)
 {
diff --git a/Engine/Source/Runtime/Engine/Public/MaterialCompiler.h b/Engine/Source/Runtime/Engine/Public/MaterialCompiler.h
--- a/Engine/Source/Runtime/Engine/Public/MaterialCompiler.h
+++ b/Engine/Source/Runtime/Engine/Public/MaterialCompiler.h
@@ -425,7 +425,6 @@ public:
 	virtual int32 StrataVerticalLayering(int32 Top, int32 Base) = 0;
 	virtual int32 StrataAdd(int32 A, int32 B) = 0;
 	virtual int32 StrataMultiply(int32 A, int32 Weight) = 0;
-	virtual int32 StrataAnisotropyToRoughness(int32 Roughness, int32 Anisotropy, int32 OutputIndex) = 0;
 	virtual int32 StrataTransmittanceToMFP(int32 TransmittanceColor, int32 OutputIndex) = 0;
 
 	virtual void StrataCompilationInfoRegisterCodeChunk(int32 CodeChunk, FStrataMaterialCompilationInfo& StrataMaterialCompilationInfo) = 0;
@@ -954,11 +953,6 @@ public:
 		return Compiler->StrataMultiply(A, Weight);
 	}
 
-	virtual int32 StrataAnisotropyToRoughness(int32 Roughness, int32 Anisotropy, int32 OutputIndex) override
-	{
-		return Compiler->StrataAnisotropyToRoughness(Roughness, Anisotropy, OutputIndex);
-	}
-
 	virtual int32 StrataTransmittanceToMFP(int32 TransmittanceColor, int32 OutputIndex) override
 	{
 		return Compiler->StrataTransmittanceToMFP(TransmittanceColor, OutputIndex);
