Remove a number of memory allocations when cooking

- Replace a number of FString calls (which allocate on the heap) with TStringBuilder which will operate on the stack instead (although will allocate on the heap if the path exceeds 511 characters)
- When cooking ShooterGame from an empty DDC with shaders being built locally, FShaderCompileThreadRunnable::ReadAvailableResults can easily exceed 10 million allocations or more.
- Fixing this doesn't really result in any reduction of total cook time but it does help remove noise from memory allocation traces and can seriously reduce the amount of data collected in these traces (1.6gb to 1.3gb in our test scenario)
- We could also generate this string once and store in the FShaderCompileWorkerInfo if we wanted to be truely optimal. Since this code isn't performance cirtical I opted to try and reduce the scope of the changes in this submit.

#rnx

diff --git a/Engine/Source/Runtime/Engine/Private/ShaderCompiler/ShaderCompiler.cpp b/Engine/Source/Runtime/Engine/Private/ShaderCompiler/ShaderCompiler.cpp
--- a/Engine/Source/Runtime/Engine/Private/ShaderCompiler/ShaderCompiler.cpp
+++ b/Engine/Source/Runtime/Engine/Private/ShaderCompiler/ShaderCompiler.cpp
@@ -2088,11 +2088,10 @@ void FShaderCompileThreadRunnable::ReadAvailableResults()
 			TRACE_CPUPROFILER_EVENT_SCOPE(FShaderCompileThreadRunnable::ReadAvailableResults);
 
 			// Distributed compiles always use the same directory
-			const FString WorkingDirectory = Manager->AbsoluteShaderBaseWorkingDirectory + FString::FromInt(WorkerIndex) + TEXT("/");
 			// 'Only' indicates to the worker that it should log and continue checking for the input file after the first one is processed
-			const TCHAR* InputFileName = TEXT("WorkerInputOnly.in");
-			const FString OutputFileNameAndPath = WorkingDirectory + TEXT("WorkerOutputOnly.out");
-
+			TStringBuilder<512> OutputFileNameAndPath;
+			OutputFileNameAndPath << Manager->AbsoluteShaderBaseWorkingDirectory << WorkerIndex << TEXT("/WorkerOutputOnly.out");
+			
 			// In the common case the output file will not exist, so check for existence before opening
 			// This is only a win if FileExists is faster than CreateFileReader, which it is on Windows
 			if (FPlatformFileManager::Get().GetPlatformFile().FileExists(*OutputFileNameAndPath))
