WP HLOD Subsystem - Don't register delegate needed in editor only when running on game worlds

diff --git a/Engine/Source/Runtime/Engine/Private/WorldPartition/HLOD/HLODSubsystem.cpp b/Engine/Source/Runtime/Engine/Private/WorldPartition/HLOD/HLODSubsystem.cpp
--- a/Engine/Source/Runtime/Engine/Private/WorldPartition/HLOD/HLODSubsystem.cpp
+++ b/Engine/Source/Runtime/Engine/Private/WorldPartition/HLOD/HLODSubsystem.cpp
@@ -43,7 +43,7 @@ void UHLODSubsystem::Initialize(FSubsystemCollectionBase& Collection)
 	Super::Initialize(Collection);
 
 #if WITH_EDITOR
-	if (!IsRunningCommandlet())
+	if (!IsRunningCommandlet() && !GetWorld()->IsGameWorld())
 	{
 		GetWorld()->PersistentLevel->OnLoadedActorAddedToLevelEvent.AddUObject(this, &UHLODSubsystem::OnActorLoaded);
 		GetWorld()->PersistentLevel->OnLoadedActorRemovedFromLevelEvent.AddUObject(this, &UHLODSubsystem::OnActorUnloaded);
