Added HLOD info into FWorldPartitionActorDesc
* ActorIsHLODRelevant
* HLODLayer

diff --git a/Engine/Source/Runtime/Core/Public/UObject/UE5MainStreamObjectVersion.h b/Engine/Source/Runtime/Core/Public/UObject/UE5MainStreamObjectVersion.h
--- a/Engine/Source/Runtime/Core/Public/UObject/UE5MainStreamObjectVersion.h
+++ b/Engine/Source/Runtime/Core/Public/UObject/UE5MainStreamObjectVersion.h
@@ -91,6 +91,9 @@ struct CORE_API FUE5MainStreamObjectVersion
 		// Chaos FConvex uses array of FVec3s for vertices instead of particles
 		ConvexUsesVerticesArray,
 
+		// Serialize HLOD info in WorldPartitionActorDesc
+		WorldPartitionActorDescSerializeHLODInfo,
+
 		// -----<new versions can be added above this line>-------------------------------------------------
 		VersionPlusOne,
 		LatestVersion = VersionPlusOne - 1
diff --git a/Engine/Source/Runtime/Engine/Private/WorldPartition/WorldPartitionActorDesc.cpp b/Engine/Source/Runtime/Engine/Private/WorldPartition/WorldPartitionActorDesc.cpp
--- a/Engine/Source/Runtime/Engine/Private/WorldPartition/WorldPartitionActorDesc.cpp
+++ b/Engine/Source/Runtime/Engine/Private/WorldPartition/WorldPartitionActorDesc.cpp
@@ -17,6 +17,7 @@
 #include "WorldPartition/WorldPartition.h"
 #include "WorldPartition/ActorDescContainer.h"
 #include "WorldPartition/DataLayer/DataLayer.h"
+#include "WorldPartition/HLOD/HLODLayer.h"
 #include "Engine/Public/ActorReferencesUtils.h"
 #endif
 
@@ -63,6 +64,8 @@ void FWorldPartitionActorDesc::Init(const AActor* InActor)
 	RuntimeGrid = InActor->GetRuntimeGrid();
 	bActorIsEditorOnly = InActor->IsEditorOnly();
 	bLevelBoundsRelevant = InActor->IsLevelBoundsRelevant();
+	bActorIsHLODRelevant = InActor->IsHLODRelevant();
+	HLODLayer = InActor->GetHLODLayer() ? FName(InActor->GetHLODLayer()->GetPathName()) : FName();
 	DataLayers = InActor->GetDataLayerNames();
 	ActorPackage = InActor->GetPackage()->GetFName();
 	ActorPath = *InActor->GetPathName();
@@ -199,6 +202,17 @@ void FWorldPartitionActorDesc::Serialize(FArchive& Ar)
 	{
 		Ar << ActorLabel;
 	}
+
+	if (Ar.CustomVer(FUE5MainStreamObjectVersion::GUID) >= FUE5MainStreamObjectVersion::WorldPartitionActorDescSerializeHLODInfo)
+	{
+		Ar << bActorIsHLODRelevant;
+		Ar << HLODLayer;
+	}
+	else
+	{
+		bActorIsHLODRelevant = true;
+		HLODLayer = FName();
+	}
 }
 
 FBox FWorldPartitionActorDesc::GetBounds() const
@@ -206,6 +220,11 @@ FBox FWorldPartitionActorDesc::GetBounds() const
 	return FBox(BoundsLocation - BoundsExtent, BoundsLocation + BoundsExtent);
 }
 
+UHLODLayer* FWorldPartitionActorDesc::GetHLODLayer() const
+{
+	return HLODLayer.IsNone() ? nullptr : Cast<UHLODLayer>(FSoftObjectPath(HLODLayer).TryLoad());
+}
+
 bool FWorldPartitionActorDesc::IsLoaded() const
 {
 #if WITH_DEV_AUTOMATION_TESTS
diff --git a/Engine/Source/Runtime/Engine/Public/WorldPartition/WorldPartitionActorDesc.h b/Engine/Source/Runtime/Engine/Public/WorldPartition/WorldPartitionActorDesc.h
--- a/Engine/Source/Runtime/Engine/Public/WorldPartition/WorldPartitionActorDesc.h
+++ b/Engine/Source/Runtime/Engine/Public/WorldPartition/WorldPartitionActorDesc.h
@@ -58,6 +58,8 @@ public:
 	inline FName GetRuntimeGrid() const { return RuntimeGrid; }
 	inline bool GetActorIsEditorOnly() const { return bActorIsEditorOnly; }
 	inline bool GetLevelBoundsRelevant() const { return bLevelBoundsRelevant; }
+	inline bool GetActorIsHLODRelevant() const { return bActorIsHLODRelevant; }
+	inline class UHLODLayer* GetHLODLayer() const;
 	inline const TArray<FName>& GetDataLayers() const { return DataLayers; }
 	inline FName GetActorPackage() const { return ActorPackage; }
 	inline FName GetActorPath() const { return ActorPath; }
@@ -147,6 +149,8 @@ protected:
 	FName							RuntimeGrid;
 	bool							bActorIsEditorOnly;
 	bool							bLevelBoundsRelevant;
+	bool							bActorIsHLODRelevant;
+	FName							HLODLayer;
 	TArray<FName>					DataLayers;
 	TArray<FGuid>					References;
 	
