Build fix

diff --git a/Engine/Source/Runtime/Engine/Private/Components/StaticMeshComponent.cpp b/Engine/Source/Runtime/Engine/Private/Components/StaticMeshComponent.cpp
--- a/Engine/Source/Runtime/Engine/Private/Components/StaticMeshComponent.cpp
+++ b/Engine/Source/Runtime/Engine/Private/Components/StaticMeshComponent.cpp
@@ -2485,10 +2485,12 @@ bool UStaticMeshComponent::IsHLODRelevant() const
 		return false;
 	}
 
+#if WITH_EDITORONLY_DATA
 	if (!bEnableAutoLODGeneration)
 	{
 		return false;
 	}
+#endif
 
 	return true;
 }
