Always flush the transaction buffer when loading / unloading editor cells.
Stop refreshing the loading state of all editor cells from the data layers undo callback.

diff --git a/Engine/Source/Editor/DataLayerEditor/Private/DataLayer/DataLayerEditorSubsystem.cpp b/Engine/Source/Editor/DataLayerEditor/Private/DataLayer/DataLayerEditorSubsystem.cpp
--- a/Engine/Source/Editor/DataLayerEditor/Private/DataLayer/DataLayerEditorSubsystem.cpp
+++ b/Engine/Source/Editor/DataLayerEditor/Private/DataLayer/DataLayerEditorSubsystem.cpp
@@ -218,7 +218,6 @@ void UDataLayerEditorSubsystem::PostUndoRedo()
 {
 	DataLayerChanged.Broadcast(EDataLayerAction::Reset, NULL, NAME_None);
 	UpdateAllActorsVisibility(true, true);
-	RefreshWorldPartitionEditorCells();
 }
 
 ///////////////////////////////////////////////////////////////////////////////////////////////////////////////////////
diff --git a/Engine/Source/Runtime/Engine/Private/WorldPartition/WorldPartition.cpp b/Engine/Source/Runtime/Engine/Private/WorldPartition/WorldPartition.cpp
--- a/Engine/Source/Runtime/Engine/Private/WorldPartition/WorldPartition.cpp
+++ b/Engine/Source/Runtime/Engine/Private/WorldPartition/WorldPartition.cpp
@@ -99,6 +99,8 @@ struct FWorldPartionCellUpdateContext
 					});
 
 					GetMutableDefault<UWorldPartitionEditorPerProjectUserSettings>()->SetEditorGridLoadedCells(WorldPartition->GetWorld(), EditorGridLastLoadedCells);
+
+					GEditor->ResetTransaction(LOCTEXT("LoadingEditorCellsResetTrans", "Editor Cells Loading State Changed"));
 				}
 			}
 		}
@@ -623,11 +625,6 @@ bool UWorldPartition::UpdateEditorCells(TFunctionRef<bool(TArray<UWorldPartition
 		UpdateLoadingEditorCell(Cell, bIsCellShouldBeLoaded);
 	}
 
-	if (bIsUpdateUnloadingActors)
-	{
-		GEditor->ResetTransaction(LOCTEXT("UnloadingEditorCellsResetTrans", "Unloading Cells"));
-	}
-
 	return true;
 }
 
