Fix preventing negative frame index in AnimationData::RemoveKeys

diff --git a/Engine/Source/Runtime/Engine/Private/Animation/AnimSequenceHelpers.cpp b/Engine/Source/Runtime/Engine/Private/Animation/AnimSequenceHelpers.cpp
--- a/Engine/Source/Runtime/Engine/Private/Animation/AnimSequenceHelpers.cpp
+++ b/Engine/Source/Runtime/Engine/Private/Animation/AnimSequenceHelpers.cpp
@@ -886,7 +886,7 @@ void AnimationData::RemoveKeys(UAnimSequence* InSequence, int32 StartKeyIndex, i
 		const int32 NewNumberOfFrames = FMath::Max(NewNumberOfKeys - 1, 0);
 		const float NewSequenceLength = FrameRate.AsSeconds(NewNumberOfFrames);
 
-		const float StartFrameTime = FrameRate.AsSeconds(StartKeyIndex - 1);
+		const float StartFrameTime = FrameRate.AsSeconds(FMath::Max(StartKeyIndex - 1, 0));
 		const float RemovedTime = NumKeysToRemove * FrameRate.AsInterval();
 		const float EndTime = StartFrameTime + RemovedTime;
 
