Added HLOD info into FWorldPartitionActorDesc
* ActorIsHLODRelevant
* HLODLayer

diff --git a/Engine/Source/Runtime/Core/Public/UObject/UE5ReleaseStreamObjectVersion.h b/Engine/Source/Runtime/Core/Public/UObject/UE5ReleaseStreamObjectVersion.h
--- a/Engine/Source/Runtime/Core/Public/UObject/UE5ReleaseStreamObjectVersion.h
+++ b/Engine/Source/Runtime/Core/Public/UObject/UE5ReleaseStreamObjectVersion.h
@@ -14,6 +14,9 @@ struct CORE_API FUE5ReleaseStreamObjectVersion
 
 		// Added Lumen reflections to new reflection enum, changed defaults
 		ReflectionMethodEnum,
+		
+		// Serialize HLOD info in WorldPartitionActorDesc
+		WorldPartitionActorDescSerializeHLODInfo,
 
 		// -----<new versions can be added above this line>-------------------------------------------------
 		VersionPlusOne,
diff --git a/Engine/Source/Runtime/Engine/Private/WorldPartition/WorldPartitionActorDesc.cpp b/Engine/Source/Runtime/Engine/Private/WorldPartition/WorldPartitionActorDesc.cpp
--- a/Engine/Source/Runtime/Engine/Private/WorldPartition/WorldPartitionActorDesc.cpp
+++ b/Engine/Source/Runtime/Engine/Private/WorldPartition/WorldPartitionActorDesc.cpp
@@ -14,6 +14,7 @@
 #include "Engine/Level.h"
 #include "Editor/GroupActor.h"
 #include "UObject/UE5MainStreamObjectVersion.h"
+#include "UObject/UE5ReleaseStreamObjectVersion.h"
 #include "WorldPartition/WorldPartition.h"
 #include "WorldPartition/ActorDescContainer.h"
 #include "WorldPartition/DataLayer/DataLayer.h"
@@ -177,6 +178,7 @@ void FWorldPartitionActorDesc::Serialize(FArchive& Ar)
 	check(Ar.IsPersistent());
 
 	Ar.UsingCustomVersion(FUE5MainStreamObjectVersion::GUID);
+	Ar.UsingCustomVersion(FUE5ReleaseStreamObjectVersion::GUID);
 
 	Ar << Class << Guid << BoundsLocation << BoundsExtent << GridPlacement << RuntimeGrid << bActorIsEditorOnly << bLevelBoundsRelevant;
 	
@@ -203,7 +205,7 @@ void FWorldPartitionActorDesc::Serialize(FArchive& Ar)
 		Ar << ActorLabel;
 	}
 
-	if (Ar.CustomVer(FUE5MainStreamObjectVersion::GUID) >= FUE5MainStreamObjectVersion::WorldPartitionActorDescSerializeHLODInfo)
+	if (Ar.CustomVer(FUE5ReleaseStreamObjectVersion::GUID) >= FUE5ReleaseStreamObjectVersion::WorldPartitionActorDescSerializeHLODInfo)
 	{
 		Ar << bActorIsHLODRelevant;
 		Ar << HLODLayer;
