Added UActorDescContainer::GetActorDescChecked()

diff --git a/Engine/Source/Runtime/Engine/Private/WorldPartition/ActorDescContainer.cpp b/Engine/Source/Runtime/Engine/Private/WorldPartition/ActorDescContainer.cpp
--- a/Engine/Source/Runtime/Engine/Private/WorldPartition/ActorDescContainer.cpp
+++ b/Engine/Source/Runtime/Engine/Private/WorldPartition/ActorDescContainer.cpp
@@ -220,6 +220,18 @@ FWorldPartitionActorDesc* UActorDescContainer::GetActorDesc(const FGuid& Guid)
 	return ActorDesc ? (*ActorDesc)->Get() : nullptr;
 }
 
+const FWorldPartitionActorDesc& UActorDescContainer::GetActorDescChecked(const FGuid& Guid) const
+{
+	const TUniquePtr<FWorldPartitionActorDesc>* const ActorDesc = Actors.FindChecked(Guid);
+	return *ActorDesc->Get();
+}
+
+FWorldPartitionActorDesc& UActorDescContainer::GetActorDescChecked(const FGuid& Guid)
+{
+	TUniquePtr<FWorldPartitionActorDesc>* ActorDesc = Actors.FindChecked(Guid);
+	return *ActorDesc->Get();
+}
+
 void UActorDescContainer::RegisterDelegates()
 {
 	if (GEditor && !IsTemplate())
diff --git a/Engine/Source/Runtime/Engine/Public/WorldPartition/ActorDescContainer.h b/Engine/Source/Runtime/Engine/Public/WorldPartition/ActorDescContainer.h
--- a/Engine/Source/Runtime/Engine/Public/WorldPartition/ActorDescContainer.h
+++ b/Engine/Source/Runtime/Engine/Public/WorldPartition/ActorDescContainer.h
@@ -28,6 +28,9 @@ public:
 
 	FWorldPartitionActorDesc* GetActorDesc(const FGuid& Guid);
 	const FWorldPartitionActorDesc* GetActorDesc(const FGuid& Guid) const;
+
+	FWorldPartitionActorDesc& GetActorDescChecked(const FGuid& Guid);
+	const FWorldPartitionActorDesc& GetActorDescChecked(const FGuid& Guid) const;
 #endif
 
 protected:
