World Partition - HLOD Generation
* Separated creation of HLOD actors & the generation of the HLOD meshes
* Creation of HLOD actors now possible without loading any WP cells & actors
* Added dirtiness check to avoid rebuilding up-to-date HLODs
* Added arguments to WP HLOD Builder
** -CreateHLODs   Create, Update or Delete HLOD actors, but does not rebuild the actual geometry
** -BuildHLODs       Rebuild the geometry of HLOD actors
** -DeleteHLODs    Delete all HLOD actors

First step to allow HLOD generation to run on multiple machines in parallel

diff --git a/Engine/Source/Editor/UnrealEd/Private/WorldPartition/WorldPartitionBuilder.cpp b/Engine/Source/Editor/UnrealEd/Private/WorldPartition/WorldPartitionBuilder.cpp
--- a/Engine/Source/Editor/UnrealEd/Private/WorldPartition/WorldPartitionBuilder.cpp
+++ b/Engine/Source/Editor/UnrealEd/Private/WorldPartition/WorldPartitionBuilder.cpp
@@ -4,6 +4,44 @@
 
 #include "CoreMinimal.h"
 
+#include "DistanceFieldAtlas.h"
+#include "MeshCardRepresentation.h"
+#include "StaticMeshCompiler.h"
+
+
+DEFINE_LOG_CATEGORY_STATIC(LogWorldPartitionBuilder, Log, All);
+
+
 UWorldPartitionBuilder::UWorldPartitionBuilder(const FObjectInitializer& ObjectInitializer)
 	: Super(ObjectInitializer)
-{}
\ No newline at end of file
+{}
+
+bool UWorldPartitionBuilder::HasExceededMaxMemory() const
+{
+	const uint64 MemoryMinFreePhysical = 1024ll * 1024 * 1024;
+	const uint64 MemoryMaxUsedPhysical = 16384ll * 1024 * 1024l;
+	const FPlatformMemoryStats MemStats = FPlatformMemory::GetStats();
+	return (MemStats.AvailablePhysical < MemoryMinFreePhysical) || (MemStats.UsedPhysical >= MemoryMaxUsedPhysical);
+};
+
+void UWorldPartitionBuilder::DoCollectGarbage() const
+{
+	if (GDistanceFieldAsyncQueue)
+	{
+		GDistanceFieldAsyncQueue->BlockUntilAllBuildsComplete();
+	}
+
+	if (GCardRepresentationAsyncQueue)
+	{
+		GCardRepresentationAsyncQueue->BlockUntilAllBuildsComplete();
+	}
+
+	const FPlatformMemoryStats MemStatsBefore = FPlatformMemory::GetStats();
+	CollectGarbage(RF_NoFlags, true);
+	const FPlatformMemoryStats MemStatsAfter = FPlatformMemory::GetStats();
+
+	UE_LOG(LogWorldPartitionBuilder, Warning, TEXT("AvailablePhysical:%.2fGB AvailableVirtual %.2fGB"),
+		((int64)MemStatsAfter.AvailablePhysical - (int64)MemStatsBefore.AvailablePhysical) / (1024.0 * 1024.0 * 1024.0),
+		((int64)MemStatsAfter.AvailableVirtual - (int64)MemStatsBefore.AvailableVirtual) / (1024.0 * 1024.0 * 1024.0)
+	);
+};
diff --git a/Engine/Source/Editor/UnrealEd/Private/WorldPartition/WorldPartitionHLODsBuilder.cpp b/Engine/Source/Editor/UnrealEd/Private/WorldPartition/WorldPartitionHLODsBuilder.cpp
--- a/Engine/Source/Editor/UnrealEd/Private/WorldPartition/WorldPartitionHLODsBuilder.cpp
+++ b/Engine/Source/Editor/UnrealEd/Private/WorldPartition/WorldPartitionHLODsBuilder.cpp
@@ -6,6 +6,7 @@
 #include "EngineUtils.h"
 #include "Editor.h"
 #include "Logging/LogMacros.h"
+#include "FileHelpers.h"
 
 #include "WorldPartition/WorldPartition.h"
 #include "WorldPartition/WorldPartitionSubsystem.h"
@@ -51,6 +52,35 @@ public:
 		return PackageHelper.Delete(Package);
 	}
 
+	virtual bool Save(UPackage* Package) const override
+	{
+		// Checkout package
+		Package->MarkAsFullyLoaded();
+
+		if (!Checkout(Package))
+		{
+			UE_LOG(LogWorldPartitionHLODsBuilder, Error, TEXT("Error checking out package %s."), *Package->GetName());
+			return false;
+		}
+
+		// Save package
+		FString PackageFileName = GetFilename(Package);
+		if (!UPackage::SavePackage(Package, nullptr, RF_Standalone, *PackageFileName))
+		{
+			UE_LOG(LogWorldPartitionHLODsBuilder, Error, TEXT("Error saving package %s."), *Package->GetName());
+			return false;
+		}
+
+		// Add new package to source control
+		if (!Add(Package))
+		{
+			UE_LOG(LogWorldPartitionHLODsBuilder, Error, TEXT("Error adding package %s to source control."), *Package->GetName());
+			return false;
+		}
+
+		return true;
+	}
+
 private:
 	FPackageSourceControlHelper& PackageHelper;
 };
@@ -58,20 +88,132 @@ private:
 UWorldPartitionHLODsBuilder::UWorldPartitionHLODsBuilder(const FObjectInitializer& ObjectInitializer)
 	: Super(ObjectInitializer)
 {
+	bCreateHLODs = FParse::Param(FCommandLine::Get(), TEXT("CreateHLODs"));
+	bBuildHLODs = FParse::Param(FCommandLine::Get(), TEXT("BuildHLODs"));
+	bDeleteHLODs = FParse::Param(FCommandLine::Get(), TEXT("DeleteHLODs"));
 }
 
-bool UWorldPartitionHLODsBuilder::Run(UWorld* World, FPackageSourceControlHelper& PackageHelper)
+bool UWorldPartitionHLODsBuilder::RequiresCommandletRendering() const
 {
-	UWorldPartitionSubsystem* WorldPartitionSubsystem = World->GetSubsystem<UWorldPartitionSubsystem>();
-	check(WorldPartitionSubsystem);
+	return bBuildHLODs || !(bCreateHLODs || bDeleteHLODs);
+}
 
-	UWorldPartition* WorldPartition = World->GetWorldPartition();
+bool UWorldPartitionHLODsBuilder::Run(UWorld* World, FPackageSourceControlHelper& PackageHelper)
+{
+	WorldPartition = World->GetWorldPartition();
 	check(WorldPartition);
 
-	FSourceControlHelper SourceControlHelper(PackageHelper);
+	SourceControlHelper = new FSourceControlHelper(PackageHelper);
 
-	// Rebuild HLOD for the whole world
-	WorldPartition->GenerateHLOD(&SourceControlHelper);
+	bool bRet;
+	if (bCreateHLODs)
+	{
+		bRet = CreateHLODActors(true);
+	}
+	else if (bBuildHLODs)
+	{
+		bRet = BuildHLODActors();
+	}
+	else if (bDeleteHLODs)
+	{
+		bRet = DeleteHLODActors();
+	}
+	else
+	{
+		bRet = CreateHLODActors(false);
+	}
+
+	WorldPartition = nullptr;
+	delete SourceControlHelper;
 	
+	return bRet;
+}
+
+bool UWorldPartitionHLODsBuilder::CreateHLODActors(bool bCreateOnly)
+{
+	WorldPartition->GenerateHLOD(SourceControlHelper, bCreateOnly);
+
+	if (bCreateOnly)
+	{
+		UE_LOG(LogWorldPartitionHLODsBuilder, Display, TEXT("#### Created HLOD actors ####"));
+
+		int32 NumCreated = 0;
+		for (UActorDescContainer::TIterator<AWorldPartitionHLOD> HLODIterator(WorldPartition); HLODIterator; ++HLODIterator)
+		{
+			FWorldPartitionActorDesc* HLODActorDesc = *HLODIterator;
+			FString PackageName = HLODActorDesc->GetActorPackage().ToString();
+
+			UE_LOG(LogWorldPartitionHLODsBuilder, Display, TEXT("    [%d] %s"), NumCreated, *PackageName);
+
+			NumCreated++;
+		}
+
+		UE_LOG(LogWorldPartitionHLODsBuilder, Display, TEXT("#### Created %d HLOD actors ####"), NumCreated);
+	}
+
 	return true;
-}
\ No newline at end of file
+}
+
+bool UWorldPartitionHLODsBuilder::BuildHLODActors()
+{
+	int32 NumActors = 0;
+	for (UActorDescContainer::TIterator<AWorldPartitionHLOD> HLODIterator(WorldPartition); HLODIterator; ++HLODIterator)
+	{
+		NumActors++;
+	}
+
+	UE_LOG(LogWorldPartitionHLODsBuilder, Display, TEXT("#### Building %d HLOD actors ####"), NumActors);
+
+	int32 CurrentActor = 0;
+	for (UActorDescContainer::TIterator<AWorldPartitionHLOD> HLODIterator(WorldPartition); HLODIterator; ++HLODIterator)
+	{
+		FWorldPartitionReference ActorRef(WorldPartition, HLODIterator->GetGuid());
+		FWorldPartitionActorDesc* ActorDesc = ActorRef.Get();
+
+		AWorldPartitionHLOD* HLODActor = CastChecked<AWorldPartitionHLOD>(ActorDesc->GetActor());
+
+		HLODActor->BuildHLOD();
+
+		if (HLODActor->GetPackage()->IsDirty())
+		{
+			SourceControlHelper->Save(HLODActor->GetPackage());
+		}
+
+		if (HasExceededMaxMemory())
+		{
+			DoCollectGarbage();
+		}
+	}
+
+	UE_LOG(LogWorldPartitionHLODsBuilder, Display, TEXT("#### Built %d HLOD actors ####"), NumActors);
+
+	return true;
+}
+
+bool UWorldPartitionHLODsBuilder::DeleteHLODActors()
+{
+	UE_LOG(LogWorldPartitionHLODsBuilder, Display, TEXT("#### Deleting HLOD actors ####"));
+
+	int32 NumDeleted = 0;
+	for (UActorDescContainer::TIterator<AWorldPartitionHLOD> HLODIterator(WorldPartition); HLODIterator; ++HLODIterator)
+	{
+		FWorldPartitionActorDesc* HLODActorDesc = *HLODIterator;
+		FString PackageName = HLODActorDesc->GetActorPackage().ToString();
+
+		bool bDeleted = SourceControlHelper->Delete(PackageName);
+		if (bDeleted)
+		{
+			UE_LOG(LogWorldPartitionHLODsBuilder, Display, TEXT("    [%d] %s"), NumDeleted, *PackageName);
+		}
+		else
+		{
+			UE_LOG(LogWorldPartitionHLODsBuilder, Error, TEXT("Failed to delete %s, exiting..."), *PackageName);
+			return false;
+		}
+		NumDeleted++;
+	}
+
+	UE_LOG(LogWorldPartitionHLODsBuilder, Display, TEXT("#### Deleted %d HLOD actors ####"), NumDeleted);
+
+	return true;
+}
diff --git a/Engine/Source/Editor/UnrealEd/Public/WorldPartition/WorldPartitionBuilder.h b/Engine/Source/Editor/UnrealEd/Public/WorldPartition/WorldPartitionBuilder.h
--- a/Engine/Source/Editor/UnrealEd/Public/WorldPartition/WorldPartitionBuilder.h
+++ b/Engine/Source/Editor/UnrealEd/Public/WorldPartition/WorldPartitionBuilder.h
@@ -14,4 +14,8 @@ public:
 	virtual bool RequiresCommandletRendering() const PURE_VIRTUAL(UWorldPartitionBuilder::RequiresCommandletRendering, return false;);
 	virtual bool RequiresEntireWorldLoading() const PURE_VIRTUAL(UWorldPartitionBuilder::RequiresEntireWorldLoading, return false;);
 	virtual bool Run(UWorld* World, FPackageSourceControlHelper& PackageHelper) PURE_VIRTUAL(UWorldPartitionBuilder::Run, return false;);
+
+protected:
+	bool HasExceededMaxMemory() const;
+	void DoCollectGarbage() const;
 };
\ No newline at end of file
diff --git a/Engine/Source/Editor/UnrealEd/Public/WorldPartition/WorldPartitionHLODsBuilder.h b/Engine/Source/Editor/UnrealEd/Public/WorldPartition/WorldPartitionHLODsBuilder.h
--- a/Engine/Source/Editor/UnrealEd/Public/WorldPartition/WorldPartitionHLODsBuilder.h
+++ b/Engine/Source/Editor/UnrealEd/Public/WorldPartition/WorldPartitionHLODsBuilder.h
@@ -12,8 +12,21 @@ class UWorldPartitionHLODsBuilder : public UWorldPartitionBuilder
 
 public:
 	// UWorldPartitionBuilder interface begin
-	virtual bool RequiresCommandletRendering() const override { return true; }
+	virtual bool RequiresCommandletRendering() const override;
 	virtual bool RequiresEntireWorldLoading() const override { return false; }
 	virtual bool Run(UWorld* World, FPackageSourceControlHelper& PackageHelper) override;
 	// UWorldPartitionBuilder interface end
+
+protected:
+	bool CreateHLODActors(bool bCreateOnly);
+	bool BuildHLODActors();
+	bool DeleteHLODActors();
+
+private:
+	class UWorldPartition* WorldPartition;
+	class FSourceControlHelper* SourceControlHelper;
+
+	bool bCreateHLODs;
+	bool bBuildHLODs;
+	bool bDeleteHLODs;
 };
\ No newline at end of file
diff --git a/Engine/Source/Runtime/Engine/Private/WorldPartition/HLOD/HLODActor.cpp b/Engine/Source/Runtime/Engine/Private/WorldPartition/HLOD/HLODActor.cpp
--- a/Engine/Source/Runtime/Engine/Private/WorldPartition/HLOD/HLODActor.cpp
+++ b/Engine/Source/Runtime/Engine/Private/WorldPartition/HLOD/HLODActor.cpp
@@ -8,6 +8,7 @@
 #include "WorldPartition/WorldPartition.h"
 #include "WorldPartition/WorldPartitionActorDesc.h"
 #include "WorldPartition/HLOD/HLODActorDesc.h"
+#include "WorldPartition/HLOD/HLODBuilder.h"
 #endif
 
 AWorldPartitionHLOD::AWorldPartitionHLOD(const FObjectInitializer& ObjectInitializer)
@@ -17,6 +18,8 @@ AWorldPartitionHLOD::AWorldPartitionHLOD(const FObjectInitializer& ObjectInitial
 	SetActorEnableCollision(false);
 
 #if WITH_EDITORONLY_DATA
+	HLODHash = 0;
+	HLODBounds = FBox(EForceInit::ForceInit);
 	bListedInSceneOutliner = false;
 #endif
 }
@@ -192,21 +195,13 @@ void AWorldPartitionHLOD::SetHLODPrimitives(const TArray<UPrimitiveComponent*>&
 	}
 }
 
-void AWorldPartitionHLOD::SetSubActors(const TArray<const AActor*>& InSubActors)
+void AWorldPartitionHLOD::SetSubActors(const TArray<FGuid>& InSubActors)
 {
-	check(GetHLODComponent());
-
 	ResetLoadedSubActors();
-	SubActors.Empty();
 
-	for (const AActor* SubActor : InSubActors)
-	{
-		if (!LoadedSubActors.Contains(SubActor))
-		{
-			OnSubActorLoaded(*SubActor);
-			SubActors.Add(SubActor->GetActorGuid());
-		}
-	}
+	SubActors = InSubActors;
+
+	SetupLoadedSubActors();
 }
 
 const TArray<FGuid>& AWorldPartitionHLOD::GetSubActors() const
@@ -229,6 +224,45 @@ void AWorldPartitionHLOD::PostDuplicate(bool bDuplicateForPIE)
 	}
 }
 
+const FBox& AWorldPartitionHLOD::GetHLODBounds() const
+{
+	return HLODBounds;
+}
+
+void AWorldPartitionHLOD::SetHLODBounds(const FBox& InBounds)
+{
+	HLODBounds = InBounds;
+}
+
+void AWorldPartitionHLOD::GetActorBounds(bool bOnlyCollidingComponents, FVector& Origin, FVector& BoxExtent, bool bIncludeFromChildActors) const
+{
+	Super::GetActorBounds(bOnlyCollidingComponents, Origin, BoxExtent, bIncludeFromChildActors);
+
+	FBox Bounds = FBox(Origin - BoxExtent, Origin + BoxExtent);
+	Bounds += HLODBounds;
+	Bounds.GetCenterAndExtents(Origin, BoxExtent);
+}
+
+void AWorldPartitionHLOD::GetActorLocationBounds(bool bOnlyCollidingComponents, FVector& Origin, FVector& BoxExtent, bool bIncludeFromChildActors) const
+{
+	GetActorBounds(bOnlyCollidingComponents, Origin, BoxExtent, bIncludeFromChildActors);
+}
+
+uint32 AWorldPartitionHLOD::GetHLODHash() const
+{
+	return HLODHash;
+}
+
+void AWorldPartitionHLOD::BuildHLOD(bool bForceBuild)
+{
+	if (bForceBuild)
+	{
+		HLODHash = 0;
+	}
+
+	HLODHash = FHLODBuilderUtilities::BuildHLOD(this);
+}
+
 #endif // WITH_EDITOR
 
 
diff --git a/Engine/Source/Runtime/Engine/Private/WorldPartition/HLOD/HLODBuilder.cpp b/Engine/Source/Runtime/Engine/Private/WorldPartition/HLOD/HLODBuilder.cpp
--- a/Engine/Source/Runtime/Engine/Private/WorldPartition/HLOD/HLODBuilder.cpp
+++ b/Engine/Source/Runtime/Engine/Private/WorldPartition/HLOD/HLODBuilder.cpp
@@ -10,19 +10,27 @@
 #include "WorldPartition/HLOD/HLODLayer.h"
 #include "WorldPartition/HLOD/HLODActorDesc.h"
 
+#include "ISMPartition/ISMComponentDescriptor.h"
+
 #include "Algo/ForEach.h"
 #include "Algo/RemoveIf.h"
 #include "Algo/Transform.h"
 #include "Misc/HashBuilder.h"
+#include "Serialization/ArchiveCrc32.h"
 #include "Templates/UniquePtr.h"
 
 #include "IMeshMergeUtilities.h"
 #include "MeshMergeModule.h"
 #include "Modules/ModuleManager.h"
 #include "Engine/StaticMesh.h"
+#include "Engine/HLODProxy.h"
 #include "Materials/Material.h"
+#include "StaticMeshCompiler.h"
 
 #include "LevelInstance/LevelInstanceActor.h"
+#include "LevelInstance/LevelInstanceSubsystem.h"
+
+DEFINE_LOG_CATEGORY_STATIC(LogHLODBuilder, Log, All);
 
 /**
  * Base class for all HLODBuilders
@@ -32,89 +40,86 @@ class FHLODBuilder
 public:
 	virtual ~FHLODBuilder() {}
 
-	virtual const TCHAR* GetActorSuffix() const = 0;
-	virtual TArray<UPrimitiveComponent*> CreateComponents(AWorldPartitionHLOD* InHLODActor, const TArray<UPrimitiveComponent*>& InSubComponents) = 0;
+	virtual TArray<UPrimitiveComponent*> CreateComponents(AWorldPartitionHLOD* InHLODActor, const UHLODLayer* InHLODLayer, const TArray<UPrimitiveComponent*>& InSubComponents) = 0;
 
-	void Build(const TArray<const AActor*>& InSubActors)
+	void Build(AWorldPartitionHLOD* InHLODActor, const UHLODLayer* InHLODLayer, const TArray<const AActor*>& InSubActors)
 	{
-		TArray<UPrimitiveComponent*> SubComponents = GatherPrimitiveComponents(0, InSubActors);
+		TArray<UPrimitiveComponent*> SubComponents = GatherPrimitiveComponents(InSubActors);
 		if (SubComponents.IsEmpty())
 		{
 			return;
 		}
 
-		AWorldPartitionHLOD* HLODActor = nullptr;
-
-		// Compute HLODActor hash
-		uint64 CellHash = FHLODActorDesc::ComputeCellHash(HLODLayer->GetName(), Context->GridIndexX, Context->GridIndexY, Context->GridIndexZ, Context->DataLayersID);
+		TArray<UPrimitiveComponent*> HLODPrimitives = CreateComponents(InHLODActor, InHLODLayer, SubComponents);
+		HLODPrimitives.RemoveSwap(nullptr);
 
-		int32 HLODActorRefIndex = INDEX_NONE;
-		FWorldPartitionHandle HLODActorHandle;
-		if (Context->HLODActorDescs.RemoveAndCopyValue(CellHash, HLODActorHandle))
+		if (!HLODPrimitives.IsEmpty())
 		{
-			HLODActorRefIndex = Context->ActorReferences.Add(HLODActorHandle);
-			HLODActor = CastChecked<AWorldPartitionHLOD>(HLODActorHandle->GetActor());
+			InHLODActor->Modify();
+			InHLODActor->SetHLODPrimitives(HLODPrimitives);
 		}
+	}
 
-		if (!HLODActor)
-		{
-			FActorSpawnParameters SpawnParams;
-			SpawnParams.Name = *FString::Printf(TEXT("%s_%016llx_%s"), *HLODLayer->GetName(), CellHash, GetActorSuffix());
-			SpawnParams.NameMode = FActorSpawnParameters::ESpawnActorNameMode::Required_Fatal;
-			HLODActor = World->SpawnActor<AWorldPartitionHLOD>(SpawnParams);
-			HLODActor->SetActorLabel(CellName.ToString());
-		}
+	static bool LoadSubActors(AWorldPartitionHLOD* InHLODActor, TArray<const AActor*>& OutActors, TArray<FWorldPartitionReference>& OutActorReferences)
+	{
+		OutActors.Reserve(InHLODActor->GetSubActors().Num());
+		OutActorReferences.Reserve(InHLODActor->GetSubActors().Num());
 
-		TArray<UPrimitiveComponent*> HLODPrimitives = CreateComponents(HLODActor, SubComponents);
-		HLODPrimitives.RemoveSwap(nullptr);
+		UWorld* World = InHLODActor->GetWorld();
+		UWorldPartition* WorldPartition = World->GetWorldPartition();
+		check(WorldPartition);
 
-		if (!HLODPrimitives.IsEmpty())
-		{
-			HLODActor->Modify();
-			HLODActor->SetHLODPrimitives(HLODPrimitives);
-			HLODActor->SetSubActors(InSubActors);
-			HLODActor->SetRuntimeGrid(HLODLayer->GetRuntimeGrid(HLODLevel));
-			HLODActor->SetLODLevel(HLODLevel);
-			HLODActor->SetHLODLayer(HLODLayer->GetParentLayer().LoadSynchronous());
-			HLODActor->SetSubActorsHLODLayer(HLODLayer);
-			HLODActor->SetGridIndices(Context->GridIndexX, Context->GridIndexY, Context->GridIndexZ);
-		}
-		else
+		bool bIsDirty = false;
+
+		// Gather (and potentially load) actors
+		for (const FGuid& SubActorGuid : InHLODActor->GetSubActors())
 		{
-			if (HLODActorRefIndex != INDEX_NONE)
+			FWorldPartitionReference ActorRef(WorldPartition, SubActorGuid);
+
+			if (ActorRef.IsValid())
 			{
-				Context->HLODActorDescs.Add(CellHash, FWorldPartitionHandle(WorldPartition, HLODActor->GetActorGuid()));
-				Context->ActorReferences.RemoveAtSwap(HLODActorRefIndex);
+				AActor* LoadedActor = ActorRef.Get()->GetActor();
+
+				// Load level instances
+				if (ALevelInstance* LevelInstance = Cast<ALevelInstance>(LoadedActor))
+				{
+					// Wait for level instance loading
+					LevelInstance->GetLevelInstanceSubsystem()->BlockLoadLevelInstance(LevelInstance);
+				}
+
+				OutActors.Add(LoadedActor);
+				OutActorReferences.Add(MoveTemp(ActorRef));
 			}
 			else
 			{
-				World->DestroyActor(HLODActor);
-				HLODActor = nullptr;
+				bIsDirty = true;
 			}
 		}
 
-		if (HLODActor)
-		{
-			HLODActors.Add(HLODActor);
-		}
+		// Wait for compilation to finish
+		FStaticMeshCompilingManager::Get().FinishAllCompilation();
+
+		return !bIsDirty;
 	}
 
-	TArray<UPrimitiveComponent*> GatherPrimitiveComponents(uint32 InHLODLevel, const TArray<const AActor*>& InActors)
+	static TArray<UPrimitiveComponent*> GatherPrimitiveComponents(const TArray<const AActor*>& InActors)
 	{
 		TArray<UPrimitiveComponent*> PrimitiveComponents;
 
-		auto GatherPrimitivesFromActor = [&PrimitiveComponents, InHLODLevel](const AActor* Actor)
+		auto GatherPrimitivesFromActor = [&PrimitiveComponents](const AActor* Actor)
 		{
-			if (UHLODLayer::ShouldIncludeInHLOD(Actor, Actor->IsA<ALevelInstance>()))
+			const bool bIncludeTransientActors = Actor->IsA<ALevelInstance>();
+			for (UActorComponent* SubComponent : Actor->GetComponents())
 			{
-				for (UActorComponent* SubComponent : Actor->GetComponents())
+				if (SubComponent && SubComponent->IsHLODRelevant())
 				{
 					if (UPrimitiveComponent* PrimitiveComponent = Cast<UPrimitiveComponent>(SubComponent))
 					{
-						if (UHLODLayer::ShouldIncludeInHLOD(PrimitiveComponent, InHLODLevel))
-						{
-							PrimitiveComponents.Add(PrimitiveComponent);
-						}
+						PrimitiveComponents.Add(PrimitiveComponent);
+					}
+					else
+					{
+						UE_LOG(LogHLODBuilder, Warning, TEXT("Component \"%s\" is marked as HLOD-relevant but this type of component currently unsupported."), *SubComponent->GetFullName());
 					}
 				}
 			}
@@ -141,17 +146,80 @@ public:
 		return PrimitiveComponents;
 	}
 
-public:
-	UWorld*					World;
-	UWorldPartition*		WorldPartition;
-	const UHLODLayer*		HLODLayer;
-	uint32					HLODLevel;
-	FName					CellName;
-	FBox					CellBounds;
-	float					CellLoadingRange;
-	FHLODGenerationContext* Context;
+	static uint32 ComputeHLODHash(AWorldPartitionHLOD* InHLODActor, const TArray<const AActor*>& InActors)
+	{
+		FArchiveCrc32 Ar;
 
-	TArray<AWorldPartitionHLOD*> HLODActors;
+		// Base key, changing this will force a rebuild of all HLODs
+		FString HLODBaseKey = "5184EABE2BEF440DB5A461554A28A3E4";
+		Ar << HLODBaseKey;
+
+		// HLOD Layer
+		uint32 HLODLayerHash = InHLODActor->GetSubActorsHLODLayer()->GetCRC();
+		UE_LOG(LogHLODBuilder, Log, TEXT(" - HLODLayer (%s) = %x"), *InHLODActor->GetSubActorsHLODLayer()->GetName(), HLODLayerHash);
+		Ar << HLODLayerHash;
+
+		// We get the CRC of each component
+		TArray<uint32> ComponentsCRCs;
+		for (UPrimitiveComponent* Component : GatherPrimitiveComponents(InActors))
+		{
+			if (UStaticMeshComponent* StaticMeshComponent = Cast<UStaticMeshComponent>(Component))
+			{
+				uint32 ComponentCRC = 0;
+
+				UE_LOG(LogHLODBuilder, Log, TEXT(" - Component \'%s\' from actor \'%s\'"), *Component->GetName(), *Component->GetOwner()->GetName());
+
+				// CRC component
+				uint32 StaticMeshComponentCRC = UHLODProxy::GetCRC(StaticMeshComponent);
+				UE_LOG(LogHLODBuilder, Log, TEXT("     - StaticMeshComponent (%s) = %x"), *StaticMeshComponent->GetName(), StaticMeshComponentCRC);
+				ComponentCRC = HashCombine(ComponentCRC, StaticMeshComponentCRC);
+
+				if (UStaticMesh* StaticMesh = StaticMeshComponent->GetStaticMesh())
+				{
+					// CRC static mesh
+					int32 StaticMeshCRC = UHLODProxy::GetCRC(StaticMesh);
+					UE_LOG(LogHLODBuilder, Log, TEXT("     - StaticMesh (%s) = %x"), *StaticMesh->GetName(), StaticMeshCRC);
+					ComponentCRC = HashCombine(ComponentCRC, StaticMeshCRC);
+
+					// CRC materials
+					const int32 NumMaterials = StaticMeshComponent->GetNumMaterials();
+					for (int32 MaterialIndex = 0; MaterialIndex < NumMaterials; ++MaterialIndex)
+					{
+						UMaterialInterface* MaterialInterface = StaticMeshComponent->GetMaterial(MaterialIndex);
+						if (MaterialInterface)
+						{
+							if (MaterialInterface->GetName() == "M_Sidewalk_Master_Inst")
+							{
+								__debugbreak();
+							}
+							uint32 MaterialInterfaceCRC = UHLODProxy::GetCRC(MaterialInterface);
+							UE_LOG(LogHLODBuilder, Log, TEXT("     - MaterialInterface (%s) = %x"), *MaterialInterface->GetName(), MaterialInterfaceCRC);
+							ComponentCRC = HashCombine(ComponentCRC, MaterialInterfaceCRC);
+
+							TArray<UTexture*> Textures;
+							MaterialInterface->GetUsedTextures(Textures, EMaterialQualityLevel::High, true, ERHIFeatureLevel::SM5, true);
+							for (UTexture* Texture : Textures)
+							{
+								uint32 TextureCRC = UHLODProxy::GetCRC(Texture);
+								UE_LOG(LogHLODBuilder, Log, TEXT("     - Texture (%s) = %x"), *Texture->GetName(), TextureCRC);
+								ComponentCRC = HashCombine(ComponentCRC, TextureCRC);
+							}
+						}
+					}
+				}
+
+				ComponentsCRCs.Add(ComponentCRC);
+			}
+		}
+
+		// Sort the components CRCs to ensure the order of components won't have an impact on the final CRC
+		ComponentsCRCs.Sort();
+
+		// Append all components CRCs
+		Ar << ComponentsCRCs;
+			
+		return Ar.GetCrc();
+	}
 };
 
 
@@ -213,9 +281,7 @@ class FHLODBuilder_Instancing : public FHLODBuilder
 		uint32						Hash;
 	};
 
-	virtual const TCHAR* GetActorSuffix() const override { return TEXT("InstancedMeshes"); }
-
-	virtual TArray<UPrimitiveComponent*> CreateComponents(AWorldPartitionHLOD* InHLODActor, const TArray<UPrimitiveComponent*>& InSubComponents) override
+	virtual TArray<UPrimitiveComponent*> CreateComponents(AWorldPartitionHLOD* InHLODActor, const UHLODLayer* InHLODLayer, const TArray<UPrimitiveComponent*>& InSubComponents) override
 	{
 		TRACE_CPUPROFILER_EVENT_SCOPE(FHLODBuilder_Instancing::CreateComponents);
 
@@ -273,9 +339,7 @@ class FHLODBuilder_Instancing : public FHLODBuilder
  */
 class FHLODBuilder_MeshMerge : public FHLODBuilder
 {
-	virtual const TCHAR* GetActorSuffix() const override { return TEXT("MergedMesh"); }
-
-	virtual TArray<UPrimitiveComponent*> CreateComponents(AWorldPartitionHLOD* InHLODActor, const TArray<UPrimitiveComponent*>& InSubComponents) override
+	virtual TArray<UPrimitiveComponent*> CreateComponents(AWorldPartitionHLOD* InHLODActor, const UHLODLayer* InHLODLayer, const TArray<UPrimitiveComponent*>& InSubComponents) override
 	{
 		TRACE_CPUPROFILER_EVENT_SCOPE(FHLODBuilder_MeshMerge::CreateComponents);
 
@@ -283,7 +347,7 @@ class FHLODBuilder_MeshMerge : public FHLODBuilder
 		FVector MergedActorLocation;
 
 		const IMeshMergeUtilities& MeshMergeUtilities = FModuleManager::Get().LoadModuleChecked<IMeshMergeModule>("MeshMergeUtilities").GetUtilities();
-		MeshMergeUtilities.MergeComponentsToStaticMesh(InSubComponents, InHLODActor->GetWorld(), HLODLayer->GetMeshMergeSettings(), HLODLayer->GetHLODMaterial().LoadSynchronous(), InHLODActor->GetPackage(), CellName.ToString(), Assets, MergedActorLocation, 0.25f, false);
+		MeshMergeUtilities.MergeComponentsToStaticMesh(InSubComponents, InHLODActor->GetWorld(), InHLODLayer->GetMeshMergeSettings(), InHLODLayer->GetHLODMaterial().LoadSynchronous(), InHLODActor->GetPackage(), InHLODActor->GetActorLabel(), Assets, MergedActorLocation, 0.25f, false);
 
 		UStaticMeshComponent* Component = nullptr;
 		Algo::ForEach(Assets, [InHLODActor, &Component, &MergedActorLocation](UObject* Asset)
@@ -308,9 +372,7 @@ class FHLODBuilder_MeshMerge : public FHLODBuilder
  */
 class FHLODBuilder_MeshSimplify : public FHLODBuilder
 {
-	virtual const TCHAR* GetActorSuffix() const override { return TEXT("SimplifiedMesh"); }
-
-	virtual TArray<UPrimitiveComponent*> CreateComponents(AWorldPartitionHLOD* InHLODActor, const TArray<UPrimitiveComponent*>& InSubComponents) override
+	virtual TArray<UPrimitiveComponent*> CreateComponents(AWorldPartitionHLOD* InHLODActor, const UHLODLayer* InHLODLayer, const TArray<UPrimitiveComponent*>& InSubComponents) override
 	{
 		TRACE_CPUPROFILER_EVENT_SCOPE(FHLODBuilder_MeshSimplify::CreateComponents);
 
@@ -322,13 +384,13 @@ class FHLODBuilder_MeshSimplify : public FHLODBuilder
 		Algo::TransformIf(InSubComponents, StaticMeshComponents, [](UPrimitiveComponent* InPrimitiveComponent) { return InPrimitiveComponent->IsA<UStaticMeshComponent>(); }, [](UPrimitiveComponent* InPrimitiveComponent) { return Cast<UStaticMeshComponent>(InPrimitiveComponent); });
 
 		const IMeshMergeUtilities& MeshMergeUtilities = FModuleManager::Get().LoadModuleChecked<IMeshMergeModule>("MeshMergeUtilities").GetUtilities();
-		MeshMergeUtilities.CreateProxyMesh(StaticMeshComponents, HLODLayer->GetMeshSimplifySettings(), HLODLayer->GetHLODMaterial().LoadSynchronous(), InHLODActor->GetPackage(), CellName.ToString(), FGuid::NewGuid(), ProxyDelegate, true);
+		MeshMergeUtilities.CreateProxyMesh(StaticMeshComponents, InHLODLayer->GetMeshSimplifySettings(), InHLODLayer->GetHLODMaterial().LoadSynchronous(), InHLODActor->GetPackage(), InHLODActor->GetActorLabel(), FGuid::NewGuid(), ProxyDelegate, true);
 
 		UStaticMeshComponent* Component = nullptr;
 		Algo::ForEach(Assets, [InHLODActor, &Component](UObject* Asset)
 		{
 			Asset->ClearFlags(RF_Public | RF_Standalone);
-			Asset->Rename(nullptr, InHLODActor);
+			Asset->Rename(nullptr, InHLODActor, REN_ForceNoResetLoaders | REN_NonTransactional | REN_DoNotDirty);
 
 			if (Cast<UStaticMesh>(Asset))
 			{
@@ -341,11 +403,161 @@ class FHLODBuilder_MeshSimplify : public FHLODBuilder
 	}
 };
 
-TArray<AWorldPartitionHLOD*> FHLODBuilderUtilities::BuildHLODs(UWorldPartition* InWorldPartition, FHLODGenerationContext* InContext, FName InCellName, const FBox& InCellBounds, const UHLODLayer* InHLODLayer, uint32 InHLODLevel, const TArray<const AActor*>& InSubActors)
+TArray<AWorldPartitionHLOD*> FHLODBuilderUtilities::CreateHLODActors(FHLODCreationContext& InCreationContext, const FHLODCreationParams& InCreationParams, const TSet<FGuid>& InActors, const TArray<const UDataLayer*>& InDataLayers)
 {
-	TUniquePtr<FHLODBuilder> HLODBuilder = nullptr;
+	TMap<UHLODLayer*, TArray<FGuid>> HLODLayersActors;
+	for (const FGuid& ActorGuid : InActors)
+	{
+		FWorldPartitionActorDesc& ActorDesc = InCreationParams.WorldPartition->GetActorDescChecked(ActorGuid);
+		if (ActorDesc.GetActorIsHLODRelevant())
+		{
+			UHLODLayer* HLODLayer = UHLODLayer::GetHLODLayer(ActorDesc, InCreationParams.WorldPartition);
+			if (HLODLayer)
+			{
+				HLODLayersActors.FindOrAdd(HLODLayer).Add(ActorGuid);
+			}
+		}
+	}
 
-	EHLODLayerType HLODLevelType = InHLODLayer->GetLayerType();
+	TArray<AWorldPartitionHLOD*> HLODActors;
+	for (const auto& Pair : HLODLayersActors)
+	{
+		const UHLODLayer* HLODLayer = Pair.Key;
+		const TArray<FGuid> HLODLayerActors = Pair.Value;
+		check(!HLODLayerActors.IsEmpty());
+
+		// Compute HLODActor hash
+		uint64 CellHash = FHLODActorDesc::ComputeCellHash(HLODLayer->GetName(), InCreationParams.GridIndexX, InCreationParams.GridIndexY, InCreationParams.GridIndexZ, InCreationParams.DataLayersID);
+
+		AWorldPartitionHLOD* HLODActor = nullptr;
+		FWorldPartitionHandle HLODActorHandle;
+		if (InCreationContext.HLODActorDescs.RemoveAndCopyValue(CellHash, HLODActorHandle))
+		{
+			InCreationContext.ActorReferences.Add(HLODActorHandle);
+			HLODActor = CastChecked<AWorldPartitionHLOD>(HLODActorHandle->GetActor());
+		}
+
+		bool bNewActor = HLODActor == nullptr;
+		if (bNewActor)
+		{
+			FActorSpawnParameters SpawnParams;
+			SpawnParams.Name = *FString::Printf(TEXT("%s_%016llx"), *HLODLayer->GetName(), CellHash);
+			SpawnParams.NameMode = FActorSpawnParameters::ESpawnActorNameMode::Required_Fatal;
+			HLODActor = InCreationParams.WorldPartition->GetWorld()->SpawnActor<AWorldPartitionHLOD>(SpawnParams);
+
+			HLODActor->SetActorLabel(InCreationParams.CellName.ToString());
+			HLODActor->SetSubActorsHLODLayer(HLODLayer);
+			HLODActor->SetGridIndices(InCreationParams.GridIndexX, InCreationParams.GridIndexY, InCreationParams.GridIndexZ);
+
+			// Make sure the generated HLOD actor has the same data layers as the source actors
+			for (const UDataLayer* DataLayer : InDataLayers)
+			{
+				HLODActor->AddDataLayer(DataLayer);
+			}
+		}
+		else
+		{
+#if DO_CHECK
+			uint64 GridIndexX, GridIndexY, GridIndexZ;
+			HLODActor->GetGridIndices(GridIndexX, GridIndexY, GridIndexZ);
+			check(GridIndexX == InCreationParams.GridIndexX);
+			check(GridIndexY == InCreationParams.GridIndexY);
+			check(GridIndexZ == InCreationParams.GridIndexZ);
+			check(HLODActor->GetSubActorsHLODLayer() == HLODLayer);
+			check(FDataLayersID(HLODActor->GetDataLayerObjects()) == InCreationParams.DataLayersID);
+#endif
+		}
+
+		bool bIsDirty = false;
+
+		// Sub actors
+		{
+			bool bSubActorsChanged = HLODActor->GetSubActors().Num() != HLODLayerActors.Num();
+			if (!bSubActorsChanged)
+			{
+				TArray<FGuid> A = HLODActor->GetSubActors();
+				TArray<FGuid> B = HLODLayerActors;
+				A.Sort();
+				B.Sort();
+				bSubActorsChanged = A != B;
+			}
+
+			if (bSubActorsChanged)
+			{
+				HLODActor->SetSubActors(HLODLayerActors);
+				bIsDirty = true;
+			}
+		}
+
+		// Runtime grid
+		FName RuntimeGrid = HLODLayer->GetRuntimeGrid(InCreationParams.HLODLevel);
+		if (HLODActor->GetRuntimeGrid() != RuntimeGrid)
+		{
+			HLODActor->SetRuntimeGrid(RuntimeGrid);
+			bIsDirty = true;
+		}
+
+		// HLOD level
+		if (HLODActor->GetLODLevel() != InCreationParams.HLODLevel)
+		{
+			HLODActor->SetLODLevel(InCreationParams.HLODLevel);
+			bIsDirty = true;
+		}
+
+		// Parent HLOD layer
+		UHLODLayer* ParentHLODLayer = HLODLayer->GetParentLayer().LoadSynchronous();
+		if (HLODActor->GetHLODLayer() != ParentHLODLayer)
+		{
+			HLODActor->SetHLODLayer(ParentHLODLayer);
+			bIsDirty = true;
+		}
+
+		// Cell bounds
+		if (HLODActor->GetHLODBounds().Equals(InCreationParams.CellBounds))
+		{
+			HLODActor->SetHLODBounds(InCreationParams.CellBounds);
+			bIsDirty = true;
+		}
+
+		// If any change was performed, mark HLOD package as dirty
+		if (bIsDirty)
+		{
+			HLODActor->MarkPackageDirty();
+		}
+
+		HLODActors.Add(HLODActor);
+	}
+
+	return HLODActors;
+}
+
+uint32 FHLODBuilderUtilities::BuildHLOD(AWorldPartitionHLOD* InHLODActor)
+{
+	TArray<FWorldPartitionReference> ActorReferences;
+	TArray<const AActor*> SubActors;
+	
+	bool bIsDirty = !FHLODBuilder::LoadSubActors(InHLODActor, SubActors, ActorReferences);
+	if (bIsDirty)
+	{
+		UE_LOG(LogHLODBuilder, Warning, TEXT("HLOD actor \"%s\" needs to be rebuilt as it references actors that have been deleted."), *InHLODActor->GetActorLabel());
+	}
+
+	uint32 OldHLODHash = bIsDirty ? 0 : InHLODActor->GetHLODHash();
+	uint32 NewHLODHash = FHLODBuilder::ComputeHLODHash(InHLODActor, SubActors);
+
+	if (OldHLODHash == NewHLODHash)
+	{
+		return OldHLODHash;
+	}
+
+	// Hack to resave hashes
+	InHLODActor->MarkPackageDirty();
+	return NewHLODHash;
+
+	TUniquePtr<FHLODBuilder> HLODBuilder = nullptr;
+	
+	const UHLODLayer* HLODLayer = InHLODActor->GetSubActorsHLODLayer();
+	EHLODLayerType HLODLevelType = HLODLayer->GetLayerType();
 	switch (HLODLevelType)
 	{
 	case EHLODLayerType::Instancing:
@@ -364,24 +576,14 @@ TArray<AWorldPartitionHLOD*> FHLODBuilderUtilities::BuildHLODs(UWorldPartition*
 		checkf(false, TEXT("Unsupported type"));
 	}
 
-	TArray<AWorldPartitionHLOD*> HLODActors;
-
-	if (HLODBuilder)
+	if (ensure(HLODBuilder))
 	{
-		HLODBuilder->World = InWorldPartition->GetWorld();
-		HLODBuilder->WorldPartition = InWorldPartition;
-		HLODBuilder->HLODLayer = InHLODLayer;
-		HLODBuilder->HLODLevel = InHLODLevel;
-		HLODBuilder->CellName = InCellName;
-		HLODBuilder->CellBounds = InCellBounds;
-		HLODBuilder->Context = InContext;
-
-		HLODBuilder->Build(InSubActors);
-			
-		HLODActors = HLODBuilder->HLODActors;
+		HLODBuilder->Build(InHLODActor, HLODLayer, SubActors);
 	}
 
-	return HLODActors;
+	InHLODActor->MarkPackageDirty();
+
+	return NewHLODHash;
 }
 
 #endif
diff --git a/Engine/Source/Runtime/Engine/Private/WorldPartition/HLOD/HLODBuilder.h b/Engine/Source/Runtime/Engine/Private/WorldPartition/HLOD/HLODBuilder.h
--- a/Engine/Source/Runtime/Engine/Private/WorldPartition/HLOD/HLODBuilder.h
+++ b/Engine/Source/Runtime/Engine/Private/WorldPartition/HLOD/HLODBuilder.h
@@ -14,17 +14,26 @@ class UWorldPartition;
 class UHLODLayer;
 class AWorldPartitionHLOD;
 
-struct ENGINE_API FHLODGenerationContext
+struct ENGINE_API FHLODCreationContext
 {
 	TMap<uint64, FWorldPartitionHandle> HLODActorDescs;
 	TArray<FWorldPartitionReference> ActorReferences;
-	
+};
+
+struct ENGINE_API FHLODCreationParams
+{
+	UWorldPartition* WorldPartition;
+
 	// Everything needed to build the cell hash
 	int64 GridIndexX;
 	int64 GridIndexY;
 	int64 GridIndexZ;
 	FDataLayersID DataLayersID;
 	FName HLODLayerName;
+
+	FName CellName;
+	FBox  CellBounds;
+	uint32 HLODLevel;
 };
 
 /**
@@ -34,19 +43,22 @@ class FHLODBuilderUtilities
 {
 public:
 	/**
-	 * Build HLODs for a given cell of a WorldPartition level.
-	 * May spawn multiple AWorldPartitionHLOD actors depending on the HLODLayer settings.
+	 * Create HLOD actors for a given cell
+	 * 
+	 * @param	InCreationContext	HLOD creation context object
+	 * @param	InCreationParams	HLOD creation parameters object
+	 * @param	InActors			The actors for which we'll build an HLOD representation
+	 * @param	InDataLayers		The data layers to assign to the newly created HLOD actors
+	 */
+	static TArray<AWorldPartitionHLOD*> CreateHLODActors(FHLODCreationContext& InCreationContext, const FHLODCreationParams& InCreationParams, const TSet<FGuid>& InActors, const TArray<const UDataLayer*>& InDataLayers);
+
+	/**
+	 * Build HLOD for the specified AWorldPartitionHLOD actor.
 	 *
-	 * @param 	InWorldPartition	The WorldPartition for which we are building HLODs
-	 * @param	InContext			The HLODs generation context
-	 * @param 	InCellName			The name of the cell
-	 * @param 	InCellBounds		Bounds of the cell, will be assigned to the created HLOD actors.
-	 * @param 	InHLODLayer			The HLODLayer which will provide the HLOD building parameters.
-	 * @param	InHLODLevel			Level of HLOD (HLOD 0...N)
-	 * @param 	InSubActors			The actors from which we'll gather geometry to generate an HLOD mesh
-	 * @return The list of HLOD actors for this cell.
+	 * @param 	InHLODActor		The HLOD actor for which we'll build the HLOD
+	 * @return An hash that represent the content used to build this HLOD.
 	 */
-	static TArray<AWorldPartitionHLOD*> BuildHLODs(UWorldPartition* InWorldPartition, FHLODGenerationContext* InContext, FName InCellName, const FBox& InCellBounds, const UHLODLayer* InHLODLayer, uint32 InHLODLevel, const TArray<const AActor*>& InSubActors);
+	static uint32 BuildHLOD(AWorldPartitionHLOD* InHLODActor);
 };
 
 #endif
\ No newline at end of file
diff --git a/Engine/Source/Runtime/Engine/Private/WorldPartition/HLOD/HLODLayer.cpp b/Engine/Source/Runtime/Engine/Private/WorldPartition/HLOD/HLODLayer.cpp
--- a/Engine/Source/Runtime/Engine/Private/WorldPartition/HLOD/HLODLayer.cpp
+++ b/Engine/Source/Runtime/Engine/Private/WorldPartition/HLOD/HLODLayer.cpp
@@ -10,16 +10,17 @@
 #if WITH_EDITOR
 #include "Algo/Copy.h"
 #include "UObject/ConstructorHelpers.h"
-#include "Materials/MaterialInterface.h"
+#include "Materials/Material.h"
+#include "Engine/HLODProxy.h"
+#include "Serialization/ArchiveCrc32.h"
 
 #include "WorldPartition/WorldPartition.h"
 #include "WorldPartition/WorldPartitionActorDesc.h"
-#include "WorldPartition/HLOD/HLODBuilder.h"
-
-#include "LevelInstance/LevelInstanceActor.h"
-#include "LevelInstance/LevelInstanceSubsystem.h"
+#include "WorldPartition/DataLayer/DataLayerSubsystem.h"
 #endif
 
+DEFINE_LOG_CATEGORY_STATIC(LogHLODLayer, Log, All);
+
 UHLODLayer::UHLODLayer(const FObjectInitializer& ObjectInitializer)
 	: Super(ObjectInitializer)
 #if WITH_EDITORONLY_DATA
@@ -34,66 +35,44 @@ UHLODLayer::UHLODLayer(const FObjectInitializer& ObjectInitializer)
 
 #if WITH_EDITOR
 
-TArray<AWorldPartitionHLOD*> UHLODLayer::GenerateHLODForCell(UWorldPartition* InWorldPartition, FHLODGenerationContext* InContext, FName InCellName, const FBox& InCellBounds, uint32 InHLODLevel, const TArray<AActor*>& InCellActors)
+uint32 UHLODLayer::GetCRC() const
 {
-	TRACE_CPUPROFILER_EVENT_SCOPE(UHLODLayer::GenerateHLODForCell);
+	UHLODLayer& This = *const_cast<UHLODLayer*>(this);
 
-	TMap<UHLODLayer*, TArray<const AActor*>> HLODLayersActors;
-	for (AActor* Actor : InCellActors)
-	{
-		if (ShouldIncludeInHLOD(Actor))
-		{
-			UHLODLayer* HLODLayer = UHLODLayer::GetHLODLayer(Actor);
-			if (HLODLayer)
-			{
-				if (ALevelInstance* LevelInstance = Cast<ALevelInstance>(Actor))
-				{
-					// Wait for level instance loading
-					LevelInstance->GetLevelInstanceSubsystem()->BlockLoadLevelInstance(LevelInstance);
-				}
+	FArchiveCrc32 Ar;
 
-				HLODLayersActors.FindOrAdd(HLODLayer).Add(Actor);
-			}
-		}
-	}
+	Ar << This.LayerType;
+	UE_LOG(LogHLODLayer, VeryVerbose, TEXT(" - LayerType = %d"), Ar.GetCrc());
 
-	TArray<AWorldPartitionHLOD*> HLODActors;
-	for (const auto& HLODLayerActors : HLODLayersActors)
+	if (LayerType == EHLODLayerType::MeshMerge)
 	{
-		if (!HLODLayerActors.Value.IsEmpty())
-		{
-			HLODActors += FHLODBuilderUtilities::BuildHLODs(InWorldPartition, InContext, InCellName, InCellBounds, HLODLayerActors.Key, InHLODLevel, HLODLayerActors.Value);
-		}
+		Ar << This.MeshMergeSettings;
+		UE_LOG(LogHLODLayer, VeryVerbose, TEXT(" - MeshMergeSettings = %d"), Ar.GetCrc());
 	}
-	return HLODActors;
-}
-
-bool UHLODLayer::ShouldIncludeInHLOD(const UPrimitiveComponent* InComponent, int32 InLevelIndex)
-{
-	if (const UStaticMeshComponent* SMC = Cast<const UStaticMeshComponent>(InComponent))
+	else if (LayerType == EHLODLayerType::MeshSimplify)
 	{
-		if (!SMC->GetStaticMesh())
-		{
-			return false;
-		}
+		Ar << This.MeshSimplifySettings;
+		UE_LOG(LogHLODLayer, VeryVerbose, TEXT(" - MeshSimplifySettings = %d"), Ar.GetCrc());
 	}
 
-	if (InComponent->IsEditorOnly())
-	{
-		return false;
-	}
+	Ar << This.CellSize;
+	UE_LOG(LogHLODLayer, VeryVerbose, TEXT(" - CellSize = %d"), Ar.GetCrc());
 
-	if (InComponent->bHiddenInGame)
-	{
-		return false;
-	}
+	uint32 Hash = Ar.GetCrc();
 
-	if (!InComponent->ShouldGenerateAutoLOD(InLevelIndex))
+	const bool bUseHLODMaterial = LayerType == EHLODLayerType::MeshMerge || LayerType == EHLODLayerType::MeshSimplify;
+	if (bUseHLODMaterial && !HLODMaterial.IsNull())
 	{
-		return false;
+		UMaterialInterface* Material = HLODMaterial.LoadSynchronous();
+		if (Material)
+		{
+			uint32 MaterialCRC = UHLODProxy::GetCRC(Material);
+			UE_LOG(LogHLODLayer, VeryVerbose, TEXT(" - Material = %d"), MaterialCRC);
+			Hash = HashCombine(Hash, MaterialCRC);
+		}
 	}
 
-	return true;
+	return Hash;
 }
 
 UHLODLayer* UHLODLayer::GetHLODLayer(const AActor* InActor)
@@ -127,6 +106,42 @@ UHLODLayer* UHLODLayer::GetHLODLayer(const AActor* InActor)
 	return nullptr;
 }
 
+UHLODLayer* UHLODLayer::GetHLODLayer(const FWorldPartitionActorDesc& InActorDesc, const UWorldPartition* InWorldPartition)
+{
+	check(InWorldPartition);
+
+	if (UHLODLayer* HLODLayer = InActorDesc.GetHLODLayer())
+	{
+		return HLODLayer;
+	}
+
+	// Only fallback to the default HLODLayer for the first level of HLOD
+	bool bIsHLOD0 = !InActorDesc.GetActorClass()->IsChildOf<AWorldPartitionHLOD>();
+	if (bIsHLOD0)
+	{
+		if (UDataLayerSubsystem* DataLayerSubsystem = InWorldPartition->GetWorld()->GetSubsystem<UDataLayerSubsystem>())
+		{
+			if (const AWorldDataLayers* WorldDataLayers = AWorldDataLayers::Get(InWorldPartition->GetWorld()))
+			{
+				for (const FName& DataLayerName : InActorDesc.GetDataLayers())
+				{
+					const UDataLayer* DataLayer = WorldDataLayers->GetDataLayerFromName(DataLayerName);
+					UHLODLayer* HLODLayer = DataLayer ? DataLayer->GetDefaultHLODLayer() : nullptr;
+					if (HLODLayer)
+					{
+						return HLODLayer;
+					}
+				}
+			}
+		}
+
+		// Fallback to the world partition default HLOD layer
+		return InWorldPartition->DefaultHLODLayer;
+	}
+
+	return nullptr;
+}
+
 #endif // WITH_EDITOR
 
 #if WITH_EDITORONLY_DATA
@@ -142,50 +157,3 @@ FName UHLODLayer::GetRuntimeGrid(uint32 InHLODLevel) const
 }
 
 #endif // WITH_EDITOR
-
-bool UHLODLayer::ShouldIncludeInHLOD(const AActor* InActor, bool bInFromLevelInstance)
-{
-	if (!InActor)
-	{
-		return false;
-	}
-
-	if (InActor->IsHidden())
-	{
-		return false;
-	}
-
-	if (InActor->IsEditorOnly())
-	{
-		return false;
-	}
-
-	if (!bInFromLevelInstance && InActor->HasAnyFlags(RF_Transient))
-	{
-		return false;
-	}
-
-	if (InActor->IsTemplate())
-	{
-		return false;
-	}
-
-	if (InActor->IsPendingKill())
-	{
-		return false;
-	}
-
-	if (!InActor->bEnableAutoLODGeneration)
-	{
-		return false;
-	}
-
-	FVector Origin, Extent;
-	InActor->GetActorBounds(false, Origin, Extent);
-	if (Extent.SizeSquared() <= 0.1)
-	{
-		return false;
-	}
-
-	return true;
-}
diff --git a/Engine/Source/Runtime/Engine/Private/WorldPartition/RuntimeSpatialHash/RuntimeSpatialHashHLOD.cpp b/Engine/Source/Runtime/Engine/Private/WorldPartition/RuntimeSpatialHash/RuntimeSpatialHashHLOD.cpp
--- a/Engine/Source/Runtime/Engine/Private/WorldPartition/RuntimeSpatialHash/RuntimeSpatialHashHLOD.cpp
+++ b/Engine/Source/Runtime/Engine/Private/WorldPartition/RuntimeSpatialHash/RuntimeSpatialHashHLOD.cpp
@@ -33,31 +33,23 @@
 DEFINE_LOG_CATEGORY_STATIC(LogWorldPartitionRuntimeSpatialHashHLOD, Log, All);
 
 
-static void CheckoutAndSavePackage(UPackage* Package, ISourceControlHelper* SourceControlHelper)
+static void SavePackage(UPackage* Package, ISourceControlHelper* SourceControlHelper)
 {
-	// Checkout package
-	Package->MarkAsFullyLoaded();
-
-	if (SourceControlHelper && !SourceControlHelper->Checkout(Package))
+	if (SourceControlHelper)
 	{
-		UE_LOG(LogWorldPartitionRuntimeSpatialHashHLOD, Error, TEXT("Error checking out package %s."), *Package->GetName());
-		check(0);
+		SourceControlHelper->Save(Package);
 	}
-
-	// Save package
-	FPackagePath PackagePath = FPackagePath::FromPackageNameChecked(Package->GetName());
-	const FString PackageFileName = PackagePath.GetLocalFullPath();
-	if (!UPackage::SavePackage(Package, nullptr, RF_Standalone, *PackageFileName))
+	else
 	{
-		UE_LOG(LogWorldPartitionRuntimeSpatialHashHLOD, Error, TEXT("Error saving package %s."), *Package->GetName());
-		check(0);
-	}
+		Package->MarkAsFullyLoaded();
 
-	// Add new package to source control
-	if (SourceControlHelper && !SourceControlHelper->Add(Package))
-	{
-		UE_LOG(LogWorldPartitionRuntimeSpatialHashHLOD, Error, TEXT("Error adding package %s to source control."), *Package->GetName());
-		check(0);
+		FPackagePath PackagePath = FPackagePath::FromPackageNameChecked(Package->GetName());
+		const FString PackageFileName = PackagePath.GetLocalFullPath();
+		if (!UPackage::SavePackage(Package, nullptr, RF_Standalone, *PackageFileName))
+		{
+			UE_LOG(LogWorldPartitionRuntimeSpatialHashHLOD, Error, TEXT("Error saving package %s."), *Package->GetName());
+			check(0);
+		}
 	}
 }
 
@@ -100,7 +92,7 @@ static void DeletePackage(FWorldPartitionActorDesc* ActorDesc, ISourceControlHel
 	}
 }
 
-static TArray<FGuid> GenerateHLODsForGrid(UWorldPartition* WorldPartition, const FSpatialHashRuntimeGrid& RuntimeGrid, uint32 HLODLevel, FHLODGenerationContext& Context, ISourceControlHelper* SourceControlHelper, const TArray<FActorCluster>& ActorClusters)
+static TArray<FGuid> GenerateHLODsForGrid(UWorldPartition* WorldPartition, const FSpatialHashRuntimeGrid& RuntimeGrid, uint32 HLODLevel, FHLODCreationContext& Context, ISourceControlHelper* SourceControlHelper, bool bCreateActorsOnly, const TArray<FActorCluster>& ActorClusters)
 {
 	auto HasExceededMaxMemory = []()
 	{
@@ -174,6 +166,10 @@ static TArray<FGuid> GenerateHLODsForGrid(UWorldPartition* WorldPartition, const
 	{
 		const FSquare2DGridHelper::FGridLevel::FGridCell& GridCell = PartitionedActors.GetCell(CellCoord);
 
+		FBox2D CellBounds2D;
+		PartitionedActors.GetCellBounds(CellCoord, CellBounds2D);
+		FBox CellBounds = FBox(FVector(CellBounds2D.Min, WorldBounds.Min.Z), FVector(CellBounds2D.Max, WorldBounds.Max.Z));
+
 		for (const FSquare2DGridHelper::FGridLevel::FGridCellDataChunk& GridCellDataChunk : GridCell.GetDataChunks())
 		{
 			const bool bShouldGenerateHLODs = ShouldGenerateHLODs(GridCell, GridCellDataChunk);
@@ -181,46 +177,29 @@ static TArray<FGuid> GenerateHLODsForGrid(UWorldPartition* WorldPartition, const
 			{
 				SlowTask.EnterProgressFrame(1);
 
-				FBox2D CellBounds2D;
-				PartitionedActors.GetCellBounds(CellCoord, CellBounds2D);
-
 				FName CellName = UWorldPartitionRuntimeSpatialHash::GetCellName(WorldPartition, RuntimeGrid.GridName, CellCoord.Z, CellCoord.X, CellCoord.Y, GridCellDataChunk.GetDataLayersID());
-				FBox CellBounds = FBox(FVector(CellBounds2D.Min, WorldBounds.Min.Z), FVector(CellBounds2D.Max, WorldBounds.Max.Z));
 
-				Context.GridIndexX = CellCoord.X;
-				Context.GridIndexY = CellCoord.Y;
-				Context.GridIndexZ = CellCoord.Z;
-				Context.DataLayersID = GridCellDataChunk.GetDataLayersID();
+				TRACE_CPUPROFILER_EVENT_SCOPE_TEXT(*CellName.ToString());
 
-				// Gather (and potentially load) actors
-				TArray<AActor*> CellActors;
-				for (const FGuid& ActorGuid : GridCellDataChunk.GetActors())
-				{
-					FWorldPartitionReference& ActorRef = Context.ActorReferences.Emplace_GetRef(WorldPartition, ActorGuid);
-					FWorldPartitionActorDesc* ActorDesc = ActorRef.Get();
-					CellActors.Add(ActorDesc->GetActor());
-				}
+				UE_LOG(LogWorldPartitionRuntimeSpatialHashHLOD, Display, TEXT("[%d / %d] Creating/updating HLOD actor(s) for cell %s"), (int32)SlowTask.CompletedWork, (int32)SlowTask.TotalAmountOfWork, *CellName.ToString());
 
-				FStaticMeshCompilingManager::Get().FinishAllCompilation();
+				FHLODCreationParams CreationParams;
+				CreationParams.WorldPartition = WorldPartition;
+				CreationParams.GridIndexX = CellCoord.X;
+				CreationParams.GridIndexY = CellCoord.Y;
+				CreationParams.GridIndexZ = CellCoord.Z;
+				CreationParams.DataLayersID = GridCellDataChunk.GetDataLayersID();
+				CreationParams.CellName = CellName;
+				CreationParams.CellBounds = CellBounds;
+				CreationParams.HLODLevel = HLODLevel;
 
-				if (HasExceededMaxMemory())
-				{
-					DoCollectGarbage();
-				}
-
-				TArray<AWorldPartitionHLOD*> CellHLODActors = UHLODLayer::GenerateHLODForCell(WorldPartition, &Context, CellName, CellBounds, HLODLevel, CellActors);
+				TArray<AWorldPartitionHLOD*> CellHLODActors = FHLODBuilderUtilities::CreateHLODActors(Context, CreationParams, GridCellDataChunk.GetActors(), GridCellDataChunk.GetDataLayers());
 				if (!CellHLODActors.IsEmpty())
 				{
 					TArray<AWorldPartitionHLOD*> NewCellHLODActors;
 
 					for (AWorldPartitionHLOD* CellHLODActor : CellHLODActors)
 					{
-						// Make sure the generated HLOD actor has the same data layers as the source actors
-						for(const UDataLayer* DataLayer : GridCellDataChunk.GetDataLayers())
-						{
-							CellHLODActor->AddDataLayer(DataLayer);
-						}
-
 						FGuid ActorGuid = CellHLODActor->GetActorGuid();
 						GridHLODActors.Add(ActorGuid);
 
@@ -229,8 +208,26 @@ static TArray<FGuid> GenerateHLODsForGrid(UWorldPartition* WorldPartition, const
 						{
 							NewCellHLODActors.Add(CellHLODActor);
 						}
+					}
+
+					for (AWorldPartitionHLOD* CellHLODActor : CellHLODActors)
+					{
+						if (!bCreateActorsOnly)
+						{
+							FStaticMeshCompilingManager::Get().FinishAllCompilation();
+
+							if (HasExceededMaxMemory())
+							{
+								DoCollectGarbage();
+							}
+
+							CellHLODActor->BuildHLOD();
+						}
 
-						CheckoutAndSavePackage(CellHLODActorPackage, SourceControlHelper);
+						if (CellHLODActor->GetPackage()->IsDirty())
+						{
+							SavePackage(CellHLODActor->GetPackage(), SourceControlHelper);
+						}
 					}
 
 					// Manually tick the directory watcher and the asset registry to register newly created actors
@@ -251,7 +248,7 @@ static TArray<FGuid> GenerateHLODsForGrid(UWorldPartition* WorldPartition, const
 	});
 
 	// Need to collect garbage here since some HLOD actors have been marked pending kill when destroying them
-	// and they may be loaded when generating the nexty HLOD layer.
+	// and they may be loaded when generating the next HLOD layer.
 	DoCollectGarbage();
 
 	return GridHLODActors;
@@ -429,12 +426,12 @@ static void UpdateHLODGridsActors(UWorld* World, const TMap<FName, FSpatialHashR
 			HLODLevel = FMath::Clamp(HLODLevel + 2, 0U, LastLODColorationColorIdx);
 			GridActor->GridSettings.DebugColor = GEngine->HLODColorationColors[HLODLevel];
 
-			CheckoutAndSavePackage(GridActor->GetPackage(), SourceControlHelper);
+			SavePackage(GridActor->GetPackage(), SourceControlHelper);
 		}
 	}
 }
 
-bool UWorldPartitionRuntimeSpatialHash::GenerateHLOD(ISourceControlHelper* SourceControlHelper)
+bool UWorldPartitionRuntimeSpatialHash::GenerateHLOD(ISourceControlHelper* SourceControlHelper, bool bCreateActorsOnly)
 {
 	TRACE_CPUPROFILER_EVENT_SCOPE(UWorldPartitionRuntimeSpatialHash::GenerateHLOD);
 
@@ -473,8 +470,8 @@ bool UWorldPartitionRuntimeSpatialHash::GenerateHLOD(ISourceControlHelper* Sourc
 
 	UWorldPartition* WorldPartition = GetOuterUWorldPartition();
 
-	// Create HLODs generation context
-	FHLODGenerationContext Context;
+	// HLOD creation context
+	FHLODCreationContext Context;
 	TArray<FWorldPartitionHandle> InvalidHLODActors;
 
 	for (UActorDescContainer::TIterator<AWorldPartitionHLOD> HLODIterator(WorldPartition); HLODIterator; ++HLODIterator)
@@ -515,10 +512,10 @@ bool UWorldPartitionRuntimeSpatialHash::GenerateHLOD(ISourceControlHelper* Sourc
 	// Keep track of all valid HLOD actors, along with which runtime grid they live in
 	TMap<FName, TArray<FGuid>> GridsHLODActors;
 
-	auto GenerateHLODs = [&GridsHLODActors, WorldPartition, &GridsDepth, &Context, SourceControlHelper](const FSpatialHashRuntimeGrid& RuntimeGrid, uint32 HLODLevel, const TArray<FActorCluster>& ActorClusters)
+	auto GenerateHLODs = [&GridsHLODActors, WorldPartition, &GridsDepth, &Context, SourceControlHelper, bCreateActorsOnly](const FSpatialHashRuntimeGrid& RuntimeGrid, uint32 HLODLevel, const TArray<FActorCluster>& ActorClusters)
 	{
 		// Generate HLODs for this grid
-		TArray<FGuid> HLODActors = GenerateHLODsForGrid(WorldPartition, RuntimeGrid, HLODLevel, Context, SourceControlHelper, ActorClusters);
+		TArray<FGuid> HLODActors = GenerateHLODsForGrid(WorldPartition, RuntimeGrid, HLODLevel, Context, SourceControlHelper, bCreateActorsOnly, ActorClusters);
 
 		for (const FGuid& HLODActorGuid : HLODActors)
 		{
diff --git a/Engine/Source/Runtime/Engine/Private/WorldPartition/WorldPartition.cpp b/Engine/Source/Runtime/Engine/Private/WorldPartition/WorldPartition.cpp
--- a/Engine/Source/Runtime/Engine/Private/WorldPartition/WorldPartition.cpp
+++ b/Engine/Source/Runtime/Engine/Private/WorldPartition/WorldPartition.cpp
@@ -928,9 +928,9 @@ void UWorldPartition::FlushStreaming()
 	RuntimeHash->FlushStreaming();
 }
 
-void UWorldPartition::GenerateHLOD(ISourceControlHelper* SourceControlHelper)
+void UWorldPartition::GenerateHLOD(ISourceControlHelper* SourceControlHelper, bool bCreateActorsOnly)
 {
-	RuntimeHash->GenerateHLOD(SourceControlHelper);
+	RuntimeHash->GenerateHLOD(SourceControlHelper, bCreateActorsOnly);
 }
 
 void UWorldPartition::GenerateNavigationData()
diff --git a/Engine/Source/Runtime/Engine/Public/WorldPartition/HLOD/HLODActor.h b/Engine/Source/Runtime/Engine/Public/WorldPartition/HLOD/HLODActor.h
--- a/Engine/Source/Runtime/Engine/Public/WorldPartition/HLOD/HLODActor.h
+++ b/Engine/Source/Runtime/Engine/Public/WorldPartition/HLOD/HLODActor.h
@@ -16,8 +16,6 @@ class ENGINE_API AWorldPartitionHLOD : public AActor
 {
 	GENERATED_UCLASS_BODY()
 
-	friend class UHLODSubsystem;
-
 public:
 	void OnCellShown(FName InCellName);
 	void OnCellHidden(FName InCellName);
@@ -31,7 +29,7 @@ public:
 
 	void SetHLODPrimitives(const TArray<UPrimitiveComponent*>& InHLODPrimitives);
 
-	void SetSubActors(const TArray<const AActor*>& InSubActors);
+	void SetSubActors(const TArray<FGuid>& InSubActors);
 	const TArray<FGuid>& GetSubActors() const;
 
 	void SetSubActorsHLODLayer(const UHLODLayer* InSubActorsHLODLayer) { SubActorsHLODLayer = InSubActorsHLODLayer; }
@@ -52,6 +50,12 @@ public:
 	}
 
 	inline void SetLODLevel(uint32 InLODLevel) { LODLevel = InLODLevel; }
+
+	const FBox& GetHLODBounds() const;
+	void SetHLODBounds(const FBox& InBounds);
+
+	void BuildHLOD(bool bForceBuild = false);
+	uint32 GetHLODHash() const;
 #endif // WITH_EDITOR
 
 protected:
@@ -66,6 +70,9 @@ protected:
 
 	virtual void PostActorCreated() override;
 	virtual void PostDuplicate(bool bDuplicateForPIE) override;
+
+	virtual void GetActorBounds(bool bOnlyCollidingComponents, FVector& Origin, FVector& BoxExtent, bool bIncludeFromChildActors) const override;
+	virtual void GetActorLocationBounds(bool bOnlyCollidingComponents, FVector& Origin, FVector& BoxExtent, bool bIncludeFromChildActors) const override;
 #endif
 	//~ End AActor Interface.
 
@@ -97,6 +104,12 @@ private:
 	int64 GridIndexZ;
 
 	TSet<TWeakObjectPtr<const AActor>> LoadedSubActors;
+
+	UPROPERTY()
+	FBox HLODBounds;
+
+	UPROPERTY()
+	uint32 HLODHash;
 #endif
 
 	UPROPERTY(NonPIEDuplicateTransient, TextExportTransient, NonTransactional)
diff --git a/Engine/Source/Runtime/Engine/Public/WorldPartition/HLOD/HLODLayer.h b/Engine/Source/Runtime/Engine/Public/WorldPartition/HLOD/HLODLayer.h
--- a/Engine/Source/Runtime/Engine/Public/WorldPartition/HLOD/HLODLayer.h
+++ b/Engine/Source/Runtime/Engine/Public/WorldPartition/HLOD/HLODLayer.h
@@ -14,9 +14,7 @@
 #include "HLODLayer.generated.h"
 
 class UWorldPartition;
-class AWorldPartitionHLOD;
 class AActor;
-struct FHLODGenerationContext;
 
 UENUM()
 enum class EHLODLayerType : uint8
@@ -33,10 +31,10 @@ class ENGINE_API UHLODLayer : public UObject
 	
 #if WITH_EDITOR
 public:
-	static TArray<AWorldPartitionHLOD*> GenerateHLODForCell(UWorldPartition* InWorldPartition, FHLODGenerationContext* Context, FName InCellName, const FBox& InCellBounds, uint32 InLODLevel, const TArray<AActor*>& InCellActors);
-
 	static UHLODLayer* GetHLODLayer(const AActor* InActor);
-	static bool ShouldIncludeInHLOD(const UPrimitiveComponent* InComponent, int32 InLevelIndex);
+	static UHLODLayer* GetHLODLayer(const FWorldPartitionActorDesc& InActorDesc, const UWorldPartition* InWorldPartition);
+
+	uint32 GetCRC() const;
 #endif
 
 #if WITH_EDITORONLY_DATA
@@ -52,8 +50,6 @@ public:
 	static FName GetRuntimeGridName(uint32 InLODLevel, int32 InCellSize, float InLoadingRange);
 #endif
 
-	static bool ShouldIncludeInHLOD(const AActor* InActor, bool bInFromLevelInstance = false);
-
 #if WITH_EDITORONLY_DATA
 private:
 	/** Type of HLOD generation to use */
diff --git a/Engine/Source/Runtime/Engine/Public/WorldPartition/WorldPartition.h b/Engine/Source/Runtime/Engine/Public/WorldPartition/WorldPartition.h
--- a/Engine/Source/Runtime/Engine/Public/WorldPartition/WorldPartition.h
+++ b/Engine/Source/Runtime/Engine/Public/WorldPartition/WorldPartition.h
@@ -61,6 +61,7 @@ public:
 	virtual bool Add(UPackage* Package) const =0;
 	virtual bool Delete(const FString& PackageName) const =0;
 	virtual bool Delete(UPackage* Package) const =0;
+	virtual bool Save(UPackage* Package) const =0;
 };
 #endif
 
@@ -141,7 +142,7 @@ public:
 
 	FBox GetWorldBounds() const;
 	FBox GetEditorWorldBounds() const;
-	void GenerateHLOD(ISourceControlHelper* SourceControlHelper);
+	void GenerateHLOD(ISourceControlHelper* SourceControlHelper, bool bCreateActorsOnly);
 	void GenerateNavigationData();
 
 	// Debugging Methods
diff --git a/Engine/Source/Runtime/Engine/Public/WorldPartition/WorldPartitionRuntimeHash.h b/Engine/Source/Runtime/Engine/Public/WorldPartition/WorldPartitionRuntimeHash.h
--- a/Engine/Source/Runtime/Engine/Public/WorldPartition/WorldPartitionRuntimeHash.h
+++ b/Engine/Source/Runtime/Engine/Public/WorldPartition/WorldPartitionRuntimeHash.h
@@ -23,7 +23,7 @@ class ENGINE_API UWorldPartitionRuntimeHash : public UObject
 	virtual bool PopulateGeneratedPackageForCook(UPackage* InPackage, const FString& InPackageRelativePath, const FString& InPackageCookName) { return false; }
 	virtual void FinalizeGeneratedPackageForCook() {}
 	virtual void FlushStreaming() {}
-	virtual bool GenerateHLOD(ISourceControlHelper* SourceControlHelper) { return false; }
+	virtual bool GenerateHLOD(ISourceControlHelper* SourceControlHelper, bool bCreateActorsOnly) { return false; }
 	virtual bool GenerateNavigationData() { return false; }
 	virtual FName GetActorRuntimeGrid(const AActor* Actor) const { return NAME_None; }
 	virtual void DrawPreview() const {}
diff --git a/Engine/Source/Runtime/Engine/Public/WorldPartition/WorldPartitionRuntimeSpatialHash.h b/Engine/Source/Runtime/Engine/Public/WorldPartition/WorldPartitionRuntimeSpatialHash.h
--- a/Engine/Source/Runtime/Engine/Public/WorldPartition/WorldPartitionRuntimeSpatialHash.h
+++ b/Engine/Source/Runtime/Engine/Public/WorldPartition/WorldPartitionRuntimeSpatialHash.h
@@ -196,7 +196,7 @@ public:
 	virtual bool PopulateGeneratedPackageForCook(UPackage* InPackage, const FString& InPackageRelativePath, const FString& InPackageCookName) override;
 	virtual void FinalizeGeneratedPackageForCook() override;
 	virtual void FlushStreaming() override;
-	virtual bool GenerateHLOD(ISourceControlHelper* SourceControlHelper) override;
+	virtual bool GenerateHLOD(ISourceControlHelper* SourceControlHelper, bool bCreateActorsOnly) override;
 	virtual bool GenerateNavigationData() override;
 	virtual FName GetActorRuntimeGrid(const AActor* Actor) const override;
 	virtual void DrawPreview() const override;
