Fix crash when an actor descriptor has dependencies to a deleted actor.

diff --git a/Engine/Source/Runtime/Engine/Private/WorldPartition/RuntimeSpatialHash/RuntimeSpatialHashActorCluster.cpp b/Engine/Source/Runtime/Engine/Private/WorldPartition/RuntimeSpatialHash/RuntimeSpatialHashActorCluster.cpp
--- a/Engine/Source/Runtime/Engine/Private/WorldPartition/RuntimeSpatialHash/RuntimeSpatialHashActorCluster.cpp
+++ b/Engine/Source/Runtime/Engine/Private/WorldPartition/RuntimeSpatialHash/RuntimeSpatialHashActorCluster.cpp
@@ -102,34 +102,35 @@ void CreateActorCluster(const FWorldPartitionActorDesc* ActorDesc, EActorGridPla
 	{
 		for (const FGuid& ReferenceGuid : ActorDesc->GetReferences())
 		{
-			const FWorldPartitionActorDesc* ReferenceActorDesc = WorldPartition->GetActorDesc(ReferenceGuid);
-
-			// Don't include references to editor-only actors
-			if (!ReferenceActorDesc->GetActorIsEditorOnly())
+			if (const FWorldPartitionActorDesc* ReferenceActorDesc = WorldPartition->GetActorDesc(ReferenceGuid))
 			{
-				FActorCluster* ReferenceCluster = ActorToActorCluster.FindRef(ReferenceGuid);
-				if (ReferenceCluster)
+				// Don't include references to editor-only actors
+				if (!ReferenceActorDesc->GetActorIsEditorOnly())
 				{
-					if (ReferenceCluster != ActorCluster)
+					FActorCluster* ReferenceCluster = ActorToActorCluster.FindRef(ReferenceGuid);
+					if (ReferenceCluster)
 					{
-						// Merge reference cluster in Actor's cluster
-						ActorCluster->Add(*ReferenceCluster);
-						for (const FGuid& ReferenceClusterActorGuid : ReferenceCluster->Actors)
+						if (ReferenceCluster != ActorCluster)
 						{
-							ActorToActorCluster[ReferenceClusterActorGuid] = ActorCluster;
+							// Merge reference cluster in Actor's cluster
+							ActorCluster->Add(*ReferenceCluster);
+							for (const FGuid& ReferenceClusterActorGuid : ReferenceCluster->Actors)
+							{
+								ActorToActorCluster[ReferenceClusterActorGuid] = ActorCluster;
+							}
+							ActorClustersSet.Remove(ReferenceCluster);
+							delete ReferenceCluster;
 						}
-						ActorClustersSet.Remove(ReferenceCluster);
-						delete ReferenceCluster;
 					}
-				}
-				else
-				{
-					// Put Reference in Actor's cluster
-					ActorCluster->Add(FActorCluster(ReferenceActorDesc, GridPlacement, World));
-				}
+					else
+					{
+						// Put Reference in Actor's cluster
+						ActorCluster->Add(FActorCluster(ReferenceActorDesc, GridPlacement, World));
+					}
 
-				// Map its cluster
-				ActorToActorCluster.Add(ReferenceGuid, ActorCluster);
+					// Map its cluster
+					ActorToActorCluster.Add(ReferenceGuid, ActorCluster);
+				}
 			}
 		}
 	}
