Removed leftover debug code

diff --git a/Engine/Source/Runtime/Engine/Private/WorldPartition/HLOD/HLODBuilder.cpp b/Engine/Source/Runtime/Engine/Private/WorldPartition/HLOD/HLODBuilder.cpp
--- a/Engine/Source/Runtime/Engine/Private/WorldPartition/HLOD/HLODBuilder.cpp
+++ b/Engine/Source/Runtime/Engine/Private/WorldPartition/HLOD/HLODBuilder.cpp
@@ -188,10 +188,6 @@ public:
 						UMaterialInterface* MaterialInterface = StaticMeshComponent->GetMaterial(MaterialIndex);
 						if (MaterialInterface)
 						{
-							if (MaterialInterface->GetName() == "M_Sidewalk_Master_Inst")
-							{
-								__debugbreak();
-							}
 							uint32 MaterialInterfaceCRC = UHLODProxy::GetCRC(MaterialInterface);
 							UE_LOG(LogHLODBuilder, Log, TEXT("     - MaterialInterface (%s) = %x"), *MaterialInterface->GetName(), MaterialInterfaceCRC);
 							ComponentCRC = HashCombine(ComponentCRC, MaterialInterfaceCRC);
