Fix crash during mobile platform cooking by removing destroyed materials from the reverse lookup tables as soon as they are garbage collected

diff --git a/Engine/Source/Runtime/Engine/Private/Materials/MaterialInterface.cpp b/Engine/Source/Runtime/Engine/Private/Materials/MaterialInterface.cpp
--- a/Engine/Source/Runtime/Engine/Private/Materials/MaterialInterface.cpp
+++ b/Engine/Source/Runtime/Engine/Private/Materials/MaterialInterface.cpp
@@ -349,7 +349,7 @@ void UMaterialInterface::BeginDestroy()
 
 #if WITH_EDITOR
 	// The object cache needs to be notified when we're getting destroyed
-	FObjectCacheEventSink::NotifyReferencedTextureChanged_Concurrent(this);
+	FObjectCacheEventSink::NotifyMaterialDestroyed_Concurrent(this);
 #endif
 }
 
diff --git a/Engine/Source/Runtime/Engine/Private/ObjectCacheContext.cpp b/Engine/Source/Runtime/Engine/Private/ObjectCacheContext.cpp
--- a/Engine/Source/Runtime/Engine/Private/ObjectCacheContext.cpp
+++ b/Engine/Source/Runtime/Engine/Private/ObjectCacheContext.cpp
@@ -71,6 +71,27 @@ template <typename FromType, typename ToType>
 class FObjectReverseLookupCache
 {
 public:
+	void RemoveFrom(FromType* InFrom)
+	{
+		FWriteScopeLock ScopeLock(Lock);
+
+		TRawSet<ToType*>* ToObjects = FromToMapping.Find(InFrom);
+		if (ToObjects)
+		{
+			for (ToType* ToObject : *ToObjects)
+			{
+				TRawSet<FromType*>* FromObjects = ToFromMapping.Find(ToObject);
+				if (FromObjects)
+				{
+					UE_LOG(LogObjectCache, VeryVerbose, TEXT("Lookup mapping removed %s -> %s"), *InFrom->GetFullName(), *ToObject->GetFullName());
+					FromObjects->Remove(InFrom);
+				}
+			}
+
+			FromToMapping.Remove(InFrom);
+		}
+	}
+
 	void Update(ToType* InTo, const TArray<FromType*>& NewFromObjects)
 	{
 		FWriteScopeLock ScopeLock(Lock);
@@ -660,6 +681,21 @@ FObjectCacheContext& FObjectCacheContextScope::GetContext()
 
 #if WITH_EDITOR
 
+void FObjectCacheEventSink::NotifyMaterialDestroyed_Concurrent(UMaterialInterface* MaterialInterface)
+{
+	if (GetObjectReverseLookupMode() != EObjectReverseLookupMode::Temporary)
+	{
+		// If a material is destroyed, remove it from the cache so we don't return it anymore
+		// This can happen if a component doesn't update its render state after modifying its used materials
+		// (i.e. LandscapeComponent won't update it's render state when modifying MobileMaterialInterfaces during cook since they aren't meaningful for rendering)
+		// The net result is that the SceneProxy->UsedMaterialsForVerification will continue to hold an invalid material interface object but if it's not used for rendering, no harm will be done...
+		GMaterialToPrimitiveLookupCache.RemoveFrom(MaterialInterface);
+
+		// Remove any Material to Texture that might be present in the cache for this Material
+		GTextureToMaterialLookupCache.Update(MaterialInterface, {});
+	}
+}
+
 void FObjectCacheEventSink::NotifyUsedMaterialsChanged_Concurrent(const UPrimitiveComponent* PrimitiveComponent, const TArray<UMaterialInterface*>& UsedMaterials)
 {
 	if (GetObjectReverseLookupMode() != EObjectReverseLookupMode::Temporary)
diff --git a/Engine/Source/Runtime/Engine/Public/ObjectCacheEventSink.h b/Engine/Source/Runtime/Engine/Public/ObjectCacheEventSink.h
--- a/Engine/Source/Runtime/Engine/Public/ObjectCacheEventSink.h
+++ b/Engine/Source/Runtime/Engine/Public/ObjectCacheEventSink.h
@@ -11,6 +11,7 @@ struct FObjectCacheEventSink
 	static void NotifyRenderStateChanged_Concurrent(UActorComponent*);
 	static void NotifyReferencedTextureChanged_Concurrent(UMaterialInterface*);
 	static void NotifyStaticMeshChanged_Concurrent(UStaticMeshComponent*);
+	static void NotifyMaterialDestroyed_Concurrent(UMaterialInterface*);
 };
 
 #endif // #if WITH_EDITOR
\ No newline at end of file
