Optimized string usage in FShaderCompileJob::GetInputHash() avoiding intermediare FString allocations when writing to archive.

diff --git a/Engine/Source/Runtime/Engine/Private/ShaderCompiler/ShaderCompiler.cpp b/Engine/Source/Runtime/Engine/Private/ShaderCompiler/ShaderCompiler.cpp
--- a/Engine/Source/Runtime/Engine/Private/ShaderCompiler/ShaderCompiler.cpp
+++ b/Engine/Source/Runtime/Engine/Private/ShaderCompiler/ShaderCompiler.cpp
@@ -6316,14 +6316,16 @@ FSHAHash FShaderCompileJob::GetInputHash()
 
 	auto SerializeInputs = [this](FArchive& Archive)
 	{
+		check(Archive.IsSaving() && !Archive.IsLoading());
+
 		Archive << Input;
 		Archive << Input.Environment;
 		for (TMap<FString, FThreadSafeSharedStringPtr>::TConstIterator It(Input.Environment.IncludeVirtualPathToExternalContentsMap); It; ++It)
 		{
-			FString VirtualPath = It.Key();
-			Archive << VirtualPath;
+			const FString& VirtualPath = It.Key();
+			Archive << const_cast<FString&>(VirtualPath);
 			check(It.Value());
-			FString Contents = *It.Value();
+			FString& Contents = *It.Value();
 			Archive << Contents;
 		}
 
@@ -6332,10 +6334,10 @@ FSHAHash FShaderCompileJob::GetInputHash()
 			Archive << *Input.SharedEnvironment;
 			for (TMap<FString, FThreadSafeSharedStringPtr>::TConstIterator It(Input.SharedEnvironment->IncludeVirtualPathToExternalContentsMap); It; ++It)
 			{
-				FString VirtualPath = It.Key();
-				Archive << VirtualPath;
+				const FString& VirtualPath = It.Key();
+				Archive << const_cast<FString&>(VirtualPath);
 				check(It.Value());
-				FString Contents = *It.Value();
+				FString& Contents = *It.Value();
 				Archive << Contents;
 			}
 		}
