Added delegate in DataLayerSubsystem broadcasted when DataLayer activation state changed

diff --git a/Engine/Source/Runtime/Engine/Private/WorldPartition/DataLayer/DataLayerSubsystem.cpp b/Engine/Source/Runtime/Engine/Private/WorldPartition/DataLayer/DataLayerSubsystem.cpp
--- a/Engine/Source/Runtime/Engine/Private/WorldPartition/DataLayer/DataLayerSubsystem.cpp
+++ b/Engine/Source/Runtime/Engine/Private/WorldPartition/DataLayer/DataLayerSubsystem.cpp
@@ -106,6 +106,10 @@ void UDataLayerSubsystem::ActivateDataLayer(const UDataLayer* InDataLayer, bool
 		{
 			ActiveDataLayerNames.Remove(DataLayerName);
 		}
+		if (OnDataLayerActivationStateChanged.IsBound())
+		{
+			OnDataLayerActivationStateChanged.Broadcast(InDataLayer, bInActivate);
+		}
 	}
 }
 
diff --git a/Engine/Source/Runtime/Engine/Public/WorldPartition/DataLayer/DataLayer.h b/Engine/Source/Runtime/Engine/Public/WorldPartition/DataLayer/DataLayer.h
--- a/Engine/Source/Runtime/Engine/Public/WorldPartition/DataLayer/DataLayer.h
+++ b/Engine/Source/Runtime/Engine/Public/WorldPartition/DataLayer/DataLayer.h
@@ -10,7 +10,7 @@
 
 #include "DataLayer.generated.h"
 
-UCLASS(Config = Engine, PerObjectConfig, Within = WorldDataLayers)
+UCLASS(Config = Engine, PerObjectConfig, Within = WorldDataLayers, BlueprintType)
 class ENGINE_API UDataLayer : public UObject
 {
 	GENERATED_UCLASS_BODY()
@@ -32,9 +32,16 @@ public:
 	bool IsLocked() const { return bIsLocked; }
 #endif
 	
+	UFUNCTION(Category = DataLayer, BlueprintCallable)
 	FName GetDataLayerLabel() const  { return DataLayerLabel; }
+
+	UFUNCTION(Category = DataLayer, BlueprintCallable)
 	bool IsVisible() const { return bIsVisible; }
+
+	UFUNCTION(Category = DataLayer, BlueprintCallable)
 	bool IsDynamicallyLoaded() const { return bIsDynamicallyLoaded; }
+
+	UFUNCTION(Category = DataLayer, BlueprintCallable)
 	bool IsInitiallyActive() const { return IsDynamicallyLoaded() && bIsInitiallyActive; }	
 
 private:
diff --git a/Engine/Source/Runtime/Engine/Public/WorldPartition/DataLayer/DataLayerSubsystem.h b/Engine/Source/Runtime/Engine/Public/WorldPartition/DataLayer/DataLayerSubsystem.h
--- a/Engine/Source/Runtime/Engine/Public/WorldPartition/DataLayer/DataLayerSubsystem.h
+++ b/Engine/Source/Runtime/Engine/Public/WorldPartition/DataLayer/DataLayerSubsystem.h
@@ -12,6 +12,8 @@
  * UDataLayerSubsystem
  */
 
+DECLARE_DYNAMIC_MULTICAST_DELEGATE_TwoParams(FOnDataLayerActivationStateChanged, const UDataLayer*, DataLayer, bool, bIsActive);
+
 UCLASS()
 class ENGINE_API UDataLayerSubsystem : public UWorldSubsystem
 {
@@ -42,6 +44,9 @@ public:
 
 	UFUNCTION(BlueprintCallable, Category = DataLayers)
 	bool IsDataLayerActiveByLabel(const FName& InDataLayerName) const;
+
+	UPROPERTY(BlueprintAssignable)
+	FOnDataLayerActivationStateChanged OnDataLayerActivationStateChanged;
 	//~ End Blueprint callable functions
 
 	UDataLayer* GetDataLayerFromLabel(const FName& InDataLayerLabel) const;
