Fixed Strata default thickness to be 0.1millimeter as discussed before.

diff --git a/Engine/Shaders/Shared/StrataDefinitions.h b/Engine/Shaders/Shared/StrataDefinitions.h
--- a/Engine/Shaders/Shared/StrataDefinitions.h
+++ b/Engine/Shaders/Shared/StrataDefinitions.h
@@ -32,3 +32,6 @@
 
 #define STRATA_DATA_TILE_SIZE				8
 #define STRATA_DATA_TILE_SIZE_DIV_AS_SHIFT	3
+
+// The default thickness of a layer is considered to be 0.01 centimeter = 0.1 millimeter
+#define STRATA_LAYER_DEFAULT_THICKNESS_CM	0.01f
\ No newline at end of file
diff --git a/Engine/Source/Runtime/Engine/Private/Materials/HLSLMaterialTranslator.cpp b/Engine/Source/Runtime/Engine/Private/Materials/HLSLMaterialTranslator.cpp
--- a/Engine/Source/Runtime/Engine/Private/Materials/HLSLMaterialTranslator.cpp
+++ b/Engine/Source/Runtime/Engine/Private/Materials/HLSLMaterialTranslator.cpp
@@ -9184,11 +9184,14 @@ int32 FHLSLMaterialTranslator::StrataTransmittanceToMFP(int32 TransmittanceColor
 		return AddCodeChunk(MCT_Float3, 
 			// For the math to be valid, input to TransmittanceToMeanFreePath must be in meter.
 			// Then the output needs to be is converted to centimeters.
-			TEXT("(TransmittanceToMeanFreePath(%s, STRATA_SIMPLEVOLUME_THICKNESS_M) * STRATA_METER_TO_CENTIMETER)"), 
-			*GetParameterCode(TransmittanceColor));
+			TEXT("(TransmittanceToMeanFreePath(%s, %f * STRATA_CENTIMETER_TO_METER) * STRATA_METER_TO_CENTIMETER)"), 
+			*GetParameterCode(TransmittanceColor),
+			STRATA_LAYER_DEFAULT_THICKNESS_CM);
 		break;
 	case 1:
-		return AddInlinedCodeChunk(MCT_Float1, TEXT("STRATA_SIMPLEVOLUME_THICKNESS_CM")); // Thickness to be plugged into other nodes thickness input (in centimeters to match other nodes)
+		// Thickness to be plugged into other nodes thickness input.
+		// This matches the Slab node default using STRATA_LAYER_DEFAULT_THICKNESS_CM
+		return AddInlinedCodeChunk(MCT_Float1, TEXT("%f"), STRATA_LAYER_DEFAULT_THICKNESS_CM);
 		break;
 	}
 	return INDEX_NONE;
diff --git a/Engine/Source/Runtime/Engine/Private/Materials/MaterialExpressions.cpp b/Engine/Source/Runtime/Engine/Private/Materials/MaterialExpressions.cpp
--- a/Engine/Source/Runtime/Engine/Private/Materials/MaterialExpressions.cpp
+++ b/Engine/Source/Runtime/Engine/Private/Materials/MaterialExpressions.cpp
@@ -19778,7 +19778,7 @@ int32 UMaterialExpressionStrataSlabBSDF::Compile(class FMaterialCompiler* Compil
 		CompileWithDefaultFloat3(Compiler, EmissiveColor, 0.0f, 0.0f, 0.0f),
 		CompileWithDefaultFloat1(Compiler, Haziness, 0.0f),
 		CompileWithDefaultFloat1(Compiler, ThinFilmThickness, 0.0f),
-		CompileWithDefaultFloat1(Compiler, Thickness, 1.0f),	// thickness is in centimeters, it defaults to 1 centimeter
+		CompileWithDefaultFloat1(Compiler, Thickness, STRATA_LAYER_DEFAULT_THICKNESS_CM),
 		NormalCodeChunk,
 		TangentCodeChunk,
 		SharedNormalIndex);
