Address issues surfaced by initial lazy resolve testing:

1) Skinned mesh component shouldn't dereference skeletal mesh reference during destruction if there's a chance that is an unresolved reference.
2) Ensure that InterchangeManager doesn't use the *ByHash API without ensuring that it is using the same type hash function as in the TSet its operating on

diff --git a/Engine/Source/Runtime/Engine/Classes/Components/SkinnedMeshComponent.h b/Engine/Source/Runtime/Engine/Classes/Components/SkinnedMeshComponent.h
--- a/Engine/Source/Runtime/Engine/Classes/Components/SkinnedMeshComponent.h
+++ b/Engine/Source/Runtime/Engine/Classes/Components/SkinnedMeshComponent.h
@@ -639,6 +639,9 @@ protected:
 	/** External flag indicating that we may not be evaluated every frame */
 	uint8 bExternalEvaluationRateLimited:1;
 
+	/** Whether mip callbacks have been registered and need to be removed on destroy */
+	uint8 bMipLevelCallbackRegistered:1;
+
 #if !(UE_BUILD_SHIPPING || UE_BUILD_TEST)
 private:
 	/** Whether to draw this mesh's debug skeleton (regardless of showflags) */
diff --git a/Engine/Source/Runtime/Engine/Private/Components/SkinnedMeshComponent.cpp b/Engine/Source/Runtime/Engine/Private/Components/SkinnedMeshComponent.cpp
--- a/Engine/Source/Runtime/Engine/Private/Components/SkinnedMeshComponent.cpp
+++ b/Engine/Source/Runtime/Engine/Private/Components/SkinnedMeshComponent.cpp
@@ -367,6 +367,8 @@ USkinnedMeshComponent::USkinnedMeshComponent(const FObjectInitializer& ObjectIni
 	bExternalEvaluationRateLimited = false;
 	bExternalTickRateControlled = false;
 
+	bMipLevelCallbackRegistered = false;
+
 #if !(UE_BUILD_SHIPPING || UE_BUILD_TEST)
 	bDrawDebugSkeleton = false;
 #endif
@@ -3245,14 +3247,16 @@ void USkinnedMeshComponent::RegisterLODStreamingCallback(FLODStreamingCallback&&
 	if (SkeletalMesh)
 	{
 		SkeletalMesh->RegisterMipLevelChangeCallback(this, LODIdx, TimeoutSecs, bOnStreamIn, MoveTemp(Callback));
+		bMipLevelCallbackRegistered = true;
 	}
 }
 
 void USkinnedMeshComponent::BeginDestroy()
 {
-	if (SkeletalMesh)
+	if (SkeletalMesh && bMipLevelCallbackRegistered)
 	{
 		SkeletalMesh->RemoveMipLevelChangeCallback(this);
+		bMipLevelCallbackRegistered = false;
 	}
 
 	Super::BeginDestroy();
diff --git a/Engine/Source/Runtime/Interchange/Engine/Private/InterchangeManager.cpp b/Engine/Source/Runtime/Interchange/Engine/Private/InterchangeManager.cpp
--- a/Engine/Source/Runtime/Interchange/Engine/Private/InterchangeManager.cpp
+++ b/Engine/Source/Runtime/Interchange/Engine/Private/InterchangeManager.cpp
@@ -369,12 +369,13 @@ bool UInterchangeManager::RegisterTranslator(const UClass* TranslatorClass)
 		return false;
 	}
 
-	uint32 TranslatorClassHash = GetTypeHash(TranslatorClass);
-	if(RegisteredTranslatorsClass.ContainsByHash(TranslatorClassHash, TranslatorClass))
+	decltype(RegisteredTranslatorsClass)::ElementType TranslatorAsKey(TranslatorClass);
+	uint32 TranslatorClassHash = GetTypeHash(TranslatorAsKey);
+	if(RegisteredTranslatorsClass.ContainsByHash(TranslatorClassHash, TranslatorAsKey))
 	{
 		return true;
 	}
-	RegisteredTranslatorsClass.AddByHash(TranslatorClassHash, TranslatorClass);
+	RegisteredTranslatorsClass.AddByHash(TranslatorClassHash, TranslatorAsKey);
 
 
 	UInterchangeTranslatorBase* TranslatorToRegister = NewObject<UInterchangeTranslatorBase>(GetTransientPackage(), TranslatorClass, NAME_None);
