Change the default behaviour for console commands from autocomplete to search. This means that users will be able to type partial substrings from any portion of the command to find them instead of just the beginning of the command.
The old behaviour can still be used by setting console.searchmode.legacy to true.

diff --git a/Engine/Source/Runtime/Engine/Private/UserInterface/Console.cpp b/Engine/Source/Runtime/Engine/Private/UserInterface/Console.cpp
--- a/Engine/Source/Runtime/Engine/Private/UserInterface/Console.cpp
+++ b/Engine/Source/Runtime/Engine/Private/UserInterface/Console.cpp
@@ -59,6 +59,12 @@ static TAutoConsoleVariable<int32> CVarConsoleYPos(
 	TEXT("Console Y offset from bottom border \n"),
 	ECVF_Default);
 
+static TAutoConsoleVariable<bool> CVarConsoleLegacySearch(
+	TEXT("console.searchmode.legacy"),
+	false,
+	TEXT("Use the legacy search behaviour for console commands \n"),
+	ECVF_Default);
+
 
 namespace ConsoleDefs
 {
@@ -479,94 +485,91 @@ void UConsole::UpdateCompleteIndices()
 		BuildRuntimeAutoCompleteList(true);
 	}
 
-	// see if we should do a full search instead of normal autocomplete
-	static FString Space(" ");
-	static FString QuestionMark("?");
-	FString Left, Right;
-	if (TypedStr.Split(Space, &Left, &Right) ? Left.Equals(QuestionMark) : TypedStr.Equals(QuestionMark))
-	{
-		static FCheatTextFilter Filter(FCheatTextFilter::FItemToStringArray::CreateStatic(&CommandToStringArray));
-		Filter.SetRawFilterText(FText::FromString(Right));
+	AutoComplete.Reset();
 
+	if (CVarConsoleLegacySearch.GetValueOnAnyThread())
+	{
+		// use the old autocomplete behaviour
 		AutoCompleteIndex = 0;
 		AutoCompleteCursor = 0;
-		AutoComplete.Reset();
-
-		for (const FAutoCompleteCommand& Command : AutoCompleteList)
-		{
-			if (Filter.PassesFilter(Command))
-			{
-				AutoComplete.Add(Command);
-			}
-		}
 
-		AutoComplete.Sort();
-		return;
-	}
-
-	AutoCompleteIndex = 0;
-	AutoCompleteCursor = 0;
-	AutoComplete.Empty();
-	FAutoCompleteNode* Node = &AutoCompleteTree;
-	FString LowerTypedStr = TypedStr.ToLower();
-	int32 EndIdx = -1;
-	for (int32 Idx = 0; Idx < TypedStr.Len(); Idx++)
-	{
-		int32 Char = LowerTypedStr[Idx];
-		bool bFoundMatch = false;
-		int32 BranchCnt = 0;
-		for (int32 CharIdx = 0; CharIdx < Node->ChildNodes.Num(); CharIdx++)
-		{
-			BranchCnt += Node->ChildNodes[CharIdx]->ChildNodes.Num();
-			if (Node->ChildNodes[CharIdx]->IndexChar == Char)
-			{
-				bFoundMatch = true;
-				Node = Node->ChildNodes[CharIdx];
-				break;
-			}
-		}
-		if (!bFoundMatch)
+		FAutoCompleteNode* Node = &AutoCompleteTree;
+		FString LowerTypedStr = TypedStr.ToLower();
+		int32 EndIdx = -1;
+		for (int32 Idx = 0; Idx < TypedStr.Len(); Idx++)
 		{
-			if (!bAutoCompleteLocked && BranchCnt > 0)
+			int32 Char = LowerTypedStr[Idx];
+			bool bFoundMatch = false;
+			int32 BranchCnt = 0;
+			for (int32 CharIdx = 0; CharIdx < Node->ChildNodes.Num(); CharIdx++)
 			{
-				// we're off the grid!
-				return;
+				BranchCnt += Node->ChildNodes[CharIdx]->ChildNodes.Num();
+				if (Node->ChildNodes[CharIdx]->IndexChar == Char)
+				{
+					bFoundMatch = true;
+					Node = Node->ChildNodes[CharIdx];
+					break;
+				}
 			}
-			else
+			if (!bFoundMatch)
 			{
-				if (Idx < TypedStr.Len())
+				if (!bAutoCompleteLocked && BranchCnt > 0)
+				{
+					// we're off the grid!
+					return;
+				}
+				else
 				{
-					// if the first non-matching character is a space we might be adding parameters, stay on the last node we found so users can see the parameter info
-					if (TypedStr[Idx] == TCHAR(' '))
+					if (Idx < TypedStr.Len())
+					{
+						// if the first non-matching character is a space we might be adding parameters, stay on the last node we found so users can see the parameter info
+						if (TypedStr[Idx] == TCHAR(' '))
+						{
+							EndIdx = Idx;
+							break;
+						}
+						// there is more text behind the auto completed text, we don't need auto completion
+						return;
+					}
+					else
 					{
-						EndIdx = Idx;
 						break;
 					}
-					// there is more text behind the auto completed text, we don't need auto completion
-					return;
 				}
-				else
+			}
+		}
+		if (Node != &AutoCompleteTree)
+		{
+			const TArray<int32>& Leaf = Node->AutoCompleteListIndices;
+
+			for (uint32 i = 0, Num = (uint32)Leaf.Num(); i < Num; ++i)
+			{
+				// if we're adding parameters we want to make sure that we only display exact matches
+				// ie Typing "Foo 5" should still show info for "Foo" but not for "FooBar"
+				if (EndIdx < 0 || AutoCompleteList[Leaf[i]].Command.Len() == EndIdx)
 				{
-					break;
+					AutoComplete.Add(AutoCompleteList[Leaf[i]]);
 				}
 			}
+			AutoComplete.Sort();
 		}
 	}
-	if (Node != &AutoCompleteTree)
+	else if (!TypedStr.IsEmpty())
 	{
-		const TArray<int32>& Leaf = Node->AutoCompleteListIndices;
+		// search for any substring, not just the prefix
+		static FCheatTextFilter Filter(FCheatTextFilter::FItemToStringArray::CreateStatic(&CommandToStringArray));
+		Filter.SetRawFilterText(FText::FromString(TypedStr));
 
-		for (uint32 i = 0, Num = (uint32)Leaf.Num(); i < Num; ++i)
+		for (const FAutoCompleteCommand& Command : AutoCompleteList)
 		{
-			// if we're adding parameters we want to make sure that we only display exact matches
-			// ie Typing "Foo 5" should still show info for "Foo" but not for "FooBar"
-			if (EndIdx < 0 || AutoCompleteList[Leaf[i]].Command.Len() == EndIdx)
+			if (Filter.PassesFilter(Command))
 			{
-				AutoComplete.Add(AutoCompleteList[Leaf[i]]);
+				AutoComplete.Add(Command);
 			}
 		}
+
 		AutoComplete.Sort();
-	}
+	}	
 }
 
 void UConsole::SetAutoCompleteFromHistory()
