Add DebugContext for Chaos DDC plugin

diff --git a/Engine/Source/Runtime/Engine/Private/PhysicsEngine/Experimental/ChaosDerivedData.cpp b/Engine/Source/Runtime/Engine/Private/PhysicsEngine/Experimental/ChaosDerivedData.cpp
--- a/Engine/Source/Runtime/Engine/Private/PhysicsEngine/Experimental/ChaosDerivedData.cpp
+++ b/Engine/Source/Runtime/Engine/Private/PhysicsEngine/Experimental/ChaosDerivedData.cpp
@@ -59,6 +59,20 @@ const TCHAR* FChaosDerivedDataCooker::GetVersionString() const
 	return *Chaos::ChaosVersionString;
 }
 
+FString FChaosDerivedDataCooker::GetDebugContextString() const
+{
+	if (Setup)
+	{
+		UObject* Outer = Setup->GetOuter();
+		if (Outer)
+		{
+			return Outer->GetFullName();
+		}
+	}
+
+	return TEXT("Unknown Context");
+}
+
 FString FChaosDerivedDataCooker::GetPluginSpecificCacheKeySuffix() const
 {
 	FString SetupGeometryKey(TEXT("INVALID"));
diff --git a/Engine/Source/Runtime/Engine/Private/PhysicsEngine/Experimental/ChaosDerivedData.h b/Engine/Source/Runtime/Engine/Private/PhysicsEngine/Experimental/ChaosDerivedData.h
--- a/Engine/Source/Runtime/Engine/Private/PhysicsEngine/Experimental/ChaosDerivedData.h
+++ b/Engine/Source/Runtime/Engine/Private/PhysicsEngine/Experimental/ChaosDerivedData.h
@@ -50,6 +50,7 @@ public:
 	virtual FString GetPluginSpecificCacheKeySuffix() const override;
 	virtual bool IsBuildThreadsafe() const override;
 	virtual bool Build(TArray<uint8>& OutData) override;
+	virtual FString GetDebugContextString() const override;
 	//End FDerivedDataPluginInterface Interface
 
 	FChaosDerivedDataCooker(UBodySetup* InSetup, FName InFormat, bool bUseRefHolder = true);
