LevelInstance: Fix crash when trying to Edit Looping Level Instance

diff --git a/Engine/Source/Runtime/Engine/Private/LevelInstance/LevelInstanceEditorLevelStreaming.cpp b/Engine/Source/Runtime/Engine/Private/LevelInstance/LevelInstanceEditorLevelStreaming.cpp
--- a/Engine/Source/Runtime/Engine/Private/LevelInstance/LevelInstanceEditorLevelStreaming.cpp
+++ b/Engine/Source/Runtime/Engine/Private/LevelInstance/LevelInstanceEditorLevelStreaming.cpp
@@ -48,16 +48,20 @@ ULevelStreamingLevelInstanceEditor* ULevelStreamingLevelInstanceEditor::Load(ALe
 	UWorld* CurrentWorld = LevelInstanceActor->GetWorld();
 
 	TGuardValue<FLevelInstanceID> GuardEditLevelInstanceID(EditLevelInstanceID, LevelInstanceActor->GetLevelInstanceID());
-	ULevelStreamingLevelInstanceEditor* LevelStreaming = Cast<ULevelStreamingLevelInstanceEditor>(EditorLevelUtils::AddLevelToWorld(CurrentWorld, *LevelInstanceActor->GetWorldAssetPackage(), ULevelStreamingLevelInstanceEditor::StaticClass(), LevelInstanceActor->GetTransform()));
-	check(LevelStreaming);
-	check(LevelStreaming->LevelInstanceID == LevelInstanceActor->GetLevelInstanceID());
+	if (ULevelStreamingLevelInstanceEditor* LevelStreaming = Cast<ULevelStreamingLevelInstanceEditor>(EditorLevelUtils::AddLevelToWorld(CurrentWorld, *LevelInstanceActor->GetWorldAssetPackage(), ULevelStreamingLevelInstanceEditor::StaticClass(), LevelInstanceActor->GetTransform())))
+	{
+		check(LevelStreaming);
+		check(LevelStreaming->LevelInstanceID == LevelInstanceActor->GetLevelInstanceID());
+
+		GEngine->BlockTillLevelStreamingCompleted(LevelInstanceActor->GetWorld());
 
-	GEngine->BlockTillLevelStreamingCompleted(LevelInstanceActor->GetWorld());
-	
-	// Create special actor that will handle changing the pivot of this level
-	ALevelInstancePivot::Create(LevelInstanceActor, LevelStreaming);
+		// Create special actor that will handle changing the pivot of this level
+		ALevelInstancePivot::Create(LevelInstanceActor, LevelStreaming);
 
-	return LevelStreaming;
+		return LevelStreaming;
+	}
+
+	return nullptr;
 }
 
 void ULevelStreamingLevelInstanceEditor::Unload(ULevelStreamingLevelInstanceEditor* LevelStreaming)
diff --git a/Engine/Source/Runtime/Engine/Private/LevelInstance/LevelInstanceSubsystem.cpp b/Engine/Source/Runtime/Engine/Private/LevelInstance/LevelInstanceSubsystem.cpp
--- a/Engine/Source/Runtime/Engine/Private/LevelInstance/LevelInstanceSubsystem.cpp
+++ b/Engine/Source/Runtime/Engine/Private/LevelInstance/LevelInstanceSubsystem.cpp
@@ -1275,6 +1275,15 @@ bool ULevelInstanceSubsystem::CanEditLevelInstance(const ALevelInstance* LevelIn
 		return false;
 	}
 
+	if (GetWorld()->PersistentLevel->GetPackage()->GetName() == LevelInstanceActor->GetWorldAssetPackage())
+	{
+		if (OutReason)
+		{
+			*OutReason = LOCTEXT("CanEditLevelInstancePersistentLevel", "The Persistent level and the Level Instance are the same");
+		}
+		return false;
+	}
+
 	if (FLevelUtils::FindStreamingLevel(GetWorld(), *LevelInstanceActor->GetWorldAssetPackage()))
 	{
 		if (OutReason)
@@ -1390,6 +1399,12 @@ void ULevelInstanceSubsystem::EditLevelInstance(ALevelInstance* LevelInstanceAct
 		
 	// Load Edit LevelInstance level
 	ULevelStreamingLevelInstanceEditor* LevelStreaming = ULevelStreamingLevelInstanceEditor::Load(LevelInstanceActor);
+	if (!LevelStreaming)
+	{
+		LevelInstanceActor->LoadLevelInstance();
+		return;
+	}
+
 	FLevelInstanceEdit& LevelInstanceEdit = LevelInstanceEdits.Add(FName(*LevelInstanceActor->GetWorldAssetPackage()));
 	LevelInstanceEdit.LevelStreaming = LevelStreaming;
 		
