Workaround for out of bounds index in instanced static mesh setup

diff --git a/Engine/Source/Runtime/Engine/Private/InstancedStaticMesh.cpp b/Engine/Source/Runtime/Engine/Private/InstancedStaticMesh.cpp
--- a/Engine/Source/Runtime/Engine/Private/InstancedStaticMesh.cpp
+++ b/Engine/Source/Runtime/Engine/Private/InstancedStaticMesh.cpp
@@ -1082,7 +1082,7 @@ void FInstancedStaticMeshSceneProxy::SetupProxy(UInstancedStaticMeshComponent* I
 		for (int32 InInstanceIndex = 0; InInstanceIndex < Instances.Num(); ++InInstanceIndex)
 		{
 			int32 OutInstanceIndex = InInstanceIndex;
-			if (OutInstanceIndex < InstanceReorderTable.Num())
+			if (OutInstanceIndex < InstanceReorderTable.Num() && ensure(InstanceReorderTable[OutInstanceIndex] < Instances.Num()))
 			{
 				OutInstanceIndex = InstanceReorderTable[OutInstanceIndex];
 			}
