Crash fix in USkeleton::IsCompatible(USkeleton* InSkeleton) where the InSkeleton could be a nullptr.

#rnx

diff --git a/Engine/Source/Runtime/Engine/Private/Animation/Skeleton.cpp b/Engine/Source/Runtime/Engine/Private/Animation/Skeleton.cpp
--- a/Engine/Source/Runtime/Engine/Private/Animation/Skeleton.cpp
+++ b/Engine/Source/Runtime/Engine/Private/Animation/Skeleton.cpp
@@ -114,6 +114,11 @@ namespace VirtualBoneNameHelpers
 
 bool USkeleton::IsCompatible(const USkeleton* InSkeleton) const
 {
+	if (InSkeleton == nullptr)
+	{
+		return false;
+	}
+
 	if (InSkeleton == this)
 	{
 		return true;
