- Fixed world partition runtime spatial hash to garantee all cells of each grid level covers the world bounds
- Adapted runtime grid previewer to match new grid cell locations
- Allow 2D/3D debug draw of world partition runtime spatial hash while in SIE

diff --git a/Engine/Source/Editor/WorldPartitionEditor/Private/WorldPartition/SWorldPartitionEditorGrid2D.cpp b/Engine/Source/Editor/WorldPartitionEditor/Private/WorldPartition/SWorldPartitionEditorGrid2D.cpp
--- a/Engine/Source/Editor/WorldPartitionEditor/Private/WorldPartition/SWorldPartitionEditorGrid2D.cpp
+++ b/Engine/Source/Editor/WorldPartitionEditor/Private/WorldPartition/SWorldPartitionEditorGrid2D.cpp
@@ -623,7 +623,7 @@ int32 SWorldPartitionEditorGrid2D::OnPaint(const FPaintArgs& Args, const FGeomet
 		LayerId = PaintSoftwareCursor(AllottedGeometry, MyCullingRect, OutDrawElements, ++LayerId);
 		
 		// Draw a surrounding indicator when PIE is active
-		if (WorldPartition->IsSimulating())
+		if (UWorldPartition::IsSimulating())
 		{
 			FSlateDrawElement::MakeBox(
 				OutDrawElements,
diff --git a/Engine/Source/Runtime/Engine/Private/WorldPartition/RuntimeSpatialHash/RuntimeSpatialHashGridHelper.cpp b/Engine/Source/Runtime/Engine/Private/WorldPartition/RuntimeSpatialHash/RuntimeSpatialHashGridHelper.cpp
--- a/Engine/Source/Runtime/Engine/Private/WorldPartition/RuntimeSpatialHash/RuntimeSpatialHashGridHelper.cpp
+++ b/Engine/Source/Runtime/Engine/Private/WorldPartition/RuntimeSpatialHash/RuntimeSpatialHashGridHelper.cpp
@@ -4,26 +4,49 @@
 
 #include "ProfilingDebugging/ScopedTimers.h"
 
-FSquare2DGridHelper::FSquare2DGridHelper(int32 InNumLevels, const FVector& InOrigin, int32 InCellSize, int32 InGridSize)
-	: Origin(InOrigin)
+FSquare2DGridHelper::FSquare2DGridHelper(const FBox& InWorldBounds, const FVector& InOrigin, int32 InCellSize)
+	: WorldBounds(InWorldBounds)
+	, Origin(InOrigin)
 	, CellSize(InCellSize)
-	, GridSize(InGridSize)
 {
-	Levels.Reserve(InNumLevels);
+	// Compute Grid's size and level count based on World bounds
+	float WorldBoundsMaxExtent = 0.f;
+	if (WorldBounds.IsValid)
+	{
+		const FVector2D DistMin = FMath::Abs(FVector2D(WorldBounds.Min - Origin));
+		const FVector2D DistMax = FMath::Abs(FVector2D(WorldBounds.Max - Origin));
+		WorldBoundsMaxExtent = FMath::Max(DistMin.GetMax(), DistMax.GetMax());
+	}
+	int32 GridSize = 1;
+	int32 GridLevelCount = 1;
+	if (WorldBoundsMaxExtent > 0.f)
+	{
+		GridSize = 2.f * FMath::CeilToFloat(WorldBoundsMaxExtent / CellSize);
+		if (!FMath::IsPowerOfTwo(GridSize))
+		{
+			GridSize = FMath::Pow(2.f, FMath::CeilToFloat(FMath::Log2(static_cast<float>(GridSize))));
+		}
+		GridLevelCount = FMath::Log2(static_cast<float>(GridSize)) + 1;
+	}
+	else
+	{
+		UE_LOG(LogWorldPartitionRuntimeSpatialHash, Warning, TEXT("Invalid world bounds, grid partitioning will use a runtime grid with 1 cell."));
+	}
+
+	check(FMath::IsPowerOfTwo(GridSize));
 
+	Levels.Reserve(GridLevelCount);
 	int32 CurrentCellSize = CellSize;
 	int32 CurrentGridSize = GridSize;
-
-	const FVector2D BaseLevelOffset = 0.5 * FVector2D(CellSize, CellSize);
-	for (int32 Level = 0; Level < InNumLevels; ++Level)
+	for (int32 Level = 0; Level < GridLevelCount; ++Level)
 	{
-		// Add offset on origin based on level's cell size to break pattern of perfectly aligned cell edges at multiple level.
-		// This will prevent weird artefact during actor promotion.
-		// Apply base level offset so that first level isn't offset.
-		const FVector2D GridLevelOffset = 0.5f * FVector2D(CurrentCellSize, CurrentCellSize) - BaseLevelOffset;
-		const FVector2D LevelOrigin = FVector2D(InOrigin) + GridLevelOffset;
-
-		Levels.Emplace(LevelOrigin, CurrentCellSize, CurrentGridSize);
+		// Except for top level, adding 1 to CurrentGridSize (which is always a power of 2) breaks the pattern of perfectly aligned cell edges between grid level cells.
+		// This will prevent weird artefact during actor promotion when an actor is placed using its bounds and which overlaps multiple cells.
+		// In this situation, the algorithm will try to find a cell that encapsulates completely the actor's bounds by searching in the upper levels, until it finds one.
+		// Also note that, the default origin of each level will always be centered at the middle of the bounds of (level's cellsize * level's grid size).
+		int32 LevelGridSize = (Level == GridLevelCount-1) ? CurrentGridSize : CurrentGridSize + 1;
+		
+		Levels.Emplace(FVector2D(InOrigin), CurrentCellSize, LevelGridSize);
 
 		CurrentCellSize <<= 1;
 		CurrentGridSize >>= 1;
@@ -97,32 +120,7 @@ FSquare2DGridHelper GetGridHelper(const FBox& WorldBounds, int32 GridCellSize)
 {
 	// Default grid to a minimum of 1 level and 1 cell, for always loaded actors
 	FVector GridOrigin = FVector::ZeroVector;
-	int32 GridSize = 1;
-	int32 GridLevelCount = 1;
-
-	float WorldBoundsMaxExtent = 0.f;
-	// If World bounds is valid, compute Grid's size and level count based on it
-	if (WorldBounds.IsValid)
-	{
-		const FVector2D DistMin = FMath::Abs(FVector2D(WorldBounds.Min - GridOrigin));
-		const FVector2D DistMax = FMath::Abs(FVector2D(WorldBounds.Max - GridOrigin));
-		WorldBoundsMaxExtent = FMath::Max(DistMin.GetMax(), DistMax.GetMax());
-	}
-	if (WorldBoundsMaxExtent > 0.f)
-	{
-		GridSize = 2.f * FMath::CeilToFloat(WorldBoundsMaxExtent / GridCellSize);
-		if (!FMath::IsPowerOfTwo(GridSize))
-		{
-			GridSize = FMath::Pow(2.f, FMath::CeilToFloat(FMath::Log2(static_cast<float>(GridSize))));
-		}
-		GridLevelCount = FMath::Log2(static_cast<float>(GridSize)) + 1;
-	}
-	else
-	{
-		UE_LOG(LogWorldPartitionRuntimeSpatialHash, Warning, TEXT("Invalid world bounds, grid partitioning will use a runtime grid with 1 cell."));
-	}
-
-	return FSquare2DGridHelper(GridLevelCount, GridOrigin, GridCellSize, GridSize);
+	return FSquare2DGridHelper(WorldBounds, GridOrigin, GridCellSize);
 }
 
 FSquare2DGridHelper GetPartitionedActors(const UWorldPartition* WorldPartition, const FBox& WorldBounds, const FSpatialHashRuntimeGrid& Grid, const TArray<FActorCluster>& GridActors)
@@ -133,6 +131,16 @@ FSquare2DGridHelper GetPartitionedActors(const UWorldPartition* WorldPartition,
 	// Create the hierarchical grids for the game
 	//	
 	FSquare2DGridHelper PartitionedActors = GetGridHelper(WorldBounds, Grid.CellSize);
+	if (ensure(PartitionedActors.Levels.Num()) && WorldBounds.IsValid)
+	{
+		int32 IntersectingCellCount = 0;
+		FSquare2DGridHelper::FGridLevel& LastGridLevel = PartitionedActors.Levels.Last();
+		LastGridLevel.ForEachIntersectingCells(WorldBounds, [&IntersectingCellCount](const FIntVector2& Coords) { ++IntersectingCellCount; });
+		if (!ensure(IntersectingCellCount == 1))
+		{
+			UE_LOG(LogWorldPartitionRuntimeSpatialHash, Warning, TEXT("Can't find grid cell that encompasses world bounds."));
+		}
+	}
 
 	for (const FActorCluster& ActorCluster : GridActors)
 	{
@@ -155,7 +163,7 @@ FSquare2DGridHelper GetPartitionedActors(const UWorldPartition* WorldPartition,
 			check(ActorCluster.Actors.Num() == 1);
 			const FGuid& ActorGuid = *ActorCluster.Actors.CreateConstIterator();
 			const FWorldPartitionActorDesc& ActorDesc = *WorldPartition->GetActorDesc(ActorGuid);
-			if (PartitionedActors.GetLowestLevel().GetCellCoords(FVector2D(ActorDesc.GetOrigin()), CellCoords))
+			if (ensure(PartitionedActors.GetLowestLevel().GetCellCoords(FVector2D(ActorDesc.GetOrigin()), CellCoords)))
 			{
 				TArray<const UDataLayer*> ActorDataLayers;
 				if (const AWorldDataLayers* WorldDataLayers = AWorldDataLayers::Get(WorldPartition->GetWorld()))
@@ -174,6 +182,7 @@ FSquare2DGridHelper GetPartitionedActors(const UWorldPartition* WorldPartition,
 			{
 				GridPlacement = EActorGridPlacement::AlwaysLoaded;
 				bAlwaysLoadedPromotedOutOfGrid = true;
+				PartitionedActors.GetAlwaysLoadedCell().AddActors(ActorCluster.Actors, ActorCluster.DataLayers);
 			}
 			break;
 		}
@@ -195,10 +204,11 @@ FSquare2DGridHelper GetPartitionedActors(const UWorldPartition* WorldPartition,
 					break;
 				}
 			}
-			if (!bFoundCell)
+			if (!(ensure(bFoundCell)))
 			{
 				GridPlacement = EActorGridPlacement::AlwaysLoaded;
 				bAlwaysLoadedPromotedOutOfGrid = true;
+				PartitionedActors.GetAlwaysLoadedCell().AddActors(ActorCluster.Actors, ActorCluster.DataLayers);
 			}
 			break;
 		}
@@ -241,4 +251,4 @@ FSquare2DGridHelper GetPartitionedActors(const UWorldPartition* WorldPartition,
 	return PartitionedActors;
 }
 
-#endif // #if WITH_EDITOR
+#endif // #if WITH_EDITOR
\ No newline at end of file
diff --git a/Engine/Source/Runtime/Engine/Private/WorldPartition/RuntimeSpatialHash/RuntimeSpatialHashGridHelper.h b/Engine/Source/Runtime/Engine/Private/WorldPartition/RuntimeSpatialHash/RuntimeSpatialHashGridHelper.h
--- a/Engine/Source/Runtime/Engine/Private/WorldPartition/RuntimeSpatialHash/RuntimeSpatialHashGridHelper.h
+++ b/Engine/Source/Runtime/Engine/Private/WorldPartition/RuntimeSpatialHash/RuntimeSpatialHashGridHelper.h
@@ -344,7 +344,7 @@ struct FSquare2DGridHelper
 #endif
 	};
 
-	FSquare2DGridHelper(int32 InNumLevels, const FVector& InOrigin, int32 InCellSize, int32 InGridSize);
+	FSquare2DGridHelper(const FBox& InWorldBounds, const FVector& InOrigin, int32 InCellSize);
 
 #if WITH_EDITOR
 	// Returns the lowest grid level
@@ -397,9 +397,9 @@ struct FSquare2DGridHelper
 #endif
 
 public:
+	FBox WorldBounds;
 	FVector Origin;
 	int32 CellSize;
-	int32 GridSize;
 	TArray<FGridLevel> Levels;
 };
 
diff --git a/Engine/Source/Runtime/Engine/Private/WorldPartition/WorldPartition.cpp b/Engine/Source/Runtime/Engine/Private/WorldPartition/WorldPartition.cpp
--- a/Engine/Source/Runtime/Engine/Private/WorldPartition/WorldPartition.cpp
+++ b/Engine/Source/Runtime/Engine/Private/WorldPartition/WorldPartition.cpp
@@ -21,6 +21,7 @@
 #include "Editor.h"
 #include "FileHelpers.h"
 #include "HAL/FileManager.h"
+#include "LevelEditorViewport.h"
 #include "Misc/ScopedSlowTask.h"
 #include "Misc/ScopeExit.h"
 #include "ScopedTransaction.h"
@@ -465,12 +466,16 @@ const TArray<FWorldPartitionStreamingSource>& UWorldPartition::GetStreamingSourc
 	return EmptyStreamingSources;
 }
 
-#if WITH_EDITOR
-bool UWorldPartition::IsSimulating() const
+bool UWorldPartition::IsSimulating()
 {
-	return GEditor->bIsSimulatingInEditor || !!GEditor->PlayWorld;
+#if WITH_EDITOR
+	return GEditor && GEditor->bIsSimulatingInEditor && GCurrentLevelEditingViewportClient && GCurrentLevelEditingViewportClient->IsSimulateInEditorViewport();
+#else
+	return false;
+#endif
 }
 
+#if WITH_EDITOR
 void UWorldPartition::LoadEditorCells(const FBox& Box)
 {
 	TArray<UWorldPartitionEditorCell*> CellsToLoad;
@@ -862,21 +867,26 @@ class ULevel* UWorldPartition::GetPreferredLoadedLevelToAddToWorld() const
 	return nullptr;
 }
 
+bool UWorldPartition::CanDrawRuntimeHash() const
+{
+	return GetWorld()->IsGameWorld() || UWorldPartition::IsSimulating();
+}
+
 FVector2D UWorldPartition::GetDrawRuntimeHash2DDesiredFootprint(const FVector2D& CanvasSize)
 {
-	check(GetWorld()->IsGameWorld());
+	check(CanDrawRuntimeHash());
 	return GetStreamingPolicy()->GetDrawRuntimeHash2DDesiredFootprint(CanvasSize);
 }
 
 void UWorldPartition::DrawRuntimeHash2D(class UCanvas* Canvas, const FVector2D& PartitionCanvasOffset, const FVector2D& PartitionCanvasSize)
 {
-	check(GetWorld()->IsGameWorld());
+	check(CanDrawRuntimeHash());
 	GetStreamingPolicy()->DrawRuntimeHash2D(Canvas, PartitionCanvasOffset, PartitionCanvasSize);
 }
 
 void UWorldPartition::DrawRuntimeHash3D()
 {
-	check(GetWorld()->IsGameWorld());
+	check(CanDrawRuntimeHash());
 	GetStreamingPolicy()->DrawRuntimeHash3D();
 }
 
diff --git a/Engine/Source/Runtime/Engine/Private/WorldPartition/WorldPartitionRuntimeSpatialHash.cpp b/Engine/Source/Runtime/Engine/Private/WorldPartition/WorldPartitionRuntimeSpatialHash.cpp
--- a/Engine/Source/Runtime/Engine/Private/WorldPartition/WorldPartitionRuntimeSpatialHash.cpp
+++ b/Engine/Source/Runtime/Engine/Private/WorldPartition/WorldPartitionRuntimeSpatialHash.cpp
@@ -70,12 +70,10 @@ static FAutoConsoleVariableRef CVarRuntimeSpatialHashCellToSourceAngleContributi
 FSpatialHashStreamingGrid::FSpatialHashStreamingGrid()
 	: Origin(ForceInitToZero)
 	, CellSize(0)
-	, GridSize(0)
 	, LoadingRange(0.0f)
-#if WITH_EDITOR
 	, DebugColor(ForceInitToZero)
-	, OverrideLoadingRange(0)
-#endif
+	, WorldBounds(ForceInitToZero)
+	, OverrideLoadingRange(0.f)
 	, GridHelper(nullptr)
 {
 }
@@ -92,15 +90,14 @@ const FSquare2DGridHelper& FSpatialHashStreamingGrid::GetGridHelper() const
 {
 	if (!GridHelper)
 	{
-		GridHelper = new FSquare2DGridHelper(GridLevels.Num(), Origin, CellSize, GridSize);
-	}
-	else
-	{
-		check(GridHelper->Levels.Num() == GridLevels.Num());
-		check(GridHelper->Origin == Origin);
-		check(GridHelper->CellSize == CellSize);
-		check(GridHelper->GridSize == GridSize);
+		GridHelper = new FSquare2DGridHelper(WorldBounds, Origin, CellSize);
 	}
+
+	check(GridHelper->Levels.Num() == GridLevels.Num());
+	check(GridHelper->Origin == Origin);
+	check(GridHelper->CellSize == CellSize);
+	check(GridHelper->WorldBounds == WorldBounds);
+
 	return *GridHelper;
 }
 
@@ -323,6 +320,20 @@ void FSpatialHashStreamingGrid::Draw2D(UCanvas* Canvas, const TArray<FWorldParti
 		Box.SetColor(DebugColor);
 		Canvas->DrawItem(Box);
 	}
+
+	// Draw WorldBounds
+	FBox2D WorldScreenBounds = FBox2D(WorldToScreen(FVector2D(WorldBounds.Min)), WorldToScreen(FVector2D(WorldBounds.Max)));
+	if (!GridScreenBounds.IsInside(WorldScreenBounds))
+	{
+		WorldScreenBounds.Min.X = FMath::Clamp(WorldScreenBounds.Min.X, GridScreenBounds.Min.X, GridScreenBounds.Max.X);
+		WorldScreenBounds.Min.Y = FMath::Clamp(WorldScreenBounds.Min.Y, GridScreenBounds.Min.Y, GridScreenBounds.Max.Y);
+		WorldScreenBounds.Max.X = FMath::Clamp(WorldScreenBounds.Max.X, GridScreenBounds.Min.X, GridScreenBounds.Max.X);
+		WorldScreenBounds.Max.Y = FMath::Clamp(WorldScreenBounds.Max.Y, GridScreenBounds.Min.Y, GridScreenBounds.Max.Y);
+	}
+	FCanvasBoxItem Box(WorldScreenBounds.Min, WorldScreenBounds.GetSize());
+	Box.SetColor(FColor::Yellow);
+	Box.BlendMode = SE_BLEND_Translucent;
+	Canvas->DrawItem(Box);
 }
 
 // ------------------------------------------------------------------------------------------------
@@ -511,13 +522,13 @@ bool UWorldPartitionRuntimeSpatialHash::CreateStreamingGrid(const FSpatialHashRu
 
 	UWorldPartition* WorldPartition = GetOuterUWorldPartition();
 	const bool bIsMainWorldPartition = (WorldPartition->GetWorld() == WorldPartition->GetTypedOuter<UWorld>());
+	check(PartionedActors.WorldBounds == WorldPartition->GetWorldBounds());
 
-	check(FMath::IsPowerOfTwo(PartionedActors.GridSize));
 	FSpatialHashStreamingGrid& CurrentStreamingGrid = StreamingGrids.AddDefaulted_GetRef();
 	CurrentStreamingGrid.GridName = RuntimeGrid.GridName;
 	CurrentStreamingGrid.Origin = PartionedActors.Origin;
 	CurrentStreamingGrid.CellSize = PartionedActors.CellSize;
-	CurrentStreamingGrid.GridSize = PartionedActors.GridSize;
+	CurrentStreamingGrid.WorldBounds = PartionedActors.WorldBounds;
 	CurrentStreamingGrid.LoadingRange = RuntimeGrid.LoadingRange;
 	CurrentStreamingGrid.DebugColor = RuntimeGrid.DebugColor;
 
@@ -849,7 +860,10 @@ void UWorldPartitionRuntimeSpatialHash::Draw2D(UCanvas* Canvas, const TArray<FWo
 	for (const FSpatialHashStreamingGrid* StreamingGrid : SortedStreamingGrids)
 	{
 		// Display view sides based on extended grid loading range (minimum of N cells)
-		const float GridSideDistance = FMath::Max((2.f * StreamingGrid->GetLoadingRange() * GridViewLoadingRangeExtentRatio), StreamingGrid->CellSize * GridViewMinimumSizeInCellCount);
+		// Take into consideration GShowRuntimeSpatialHashGridLevel when using CellSize
+		int32 MinGridLevel = FMath::Clamp<int32>(GShowRuntimeSpatialHashGridLevel, 0, StreamingGrid->GridLevels.Num() - 1);
+		int32 CellSize = (1 << MinGridLevel) * StreamingGrid->CellSize;
+		const float GridSideDistance = FMath::Max((2.f * StreamingGrid->GetLoadingRange() * GridViewLoadingRangeExtentRatio), CellSize * GridViewMinimumSizeInCellCount);
 		FSphere AverageSphere(ForceInit);
 		for (const FWorldPartitionStreamingSource& Source : Sources)
 		{
diff --git a/Engine/Source/Runtime/Engine/Private/WorldPartition/WorldPartitionRuntimeSpatialHashGridPreviewer.cpp b/Engine/Source/Runtime/Engine/Private/WorldPartition/WorldPartitionRuntimeSpatialHashGridPreviewer.cpp
--- a/Engine/Source/Runtime/Engine/Private/WorldPartition/WorldPartitionRuntimeSpatialHashGridPreviewer.cpp
+++ b/Engine/Source/Runtime/Engine/Private/WorldPartition/WorldPartitionRuntimeSpatialHashGridPreviewer.cpp
@@ -62,6 +62,12 @@ void FWorldPartitionRuntimeSpatialHashGridPreviewer::Draw(UWorld* World, const T
 						CachedParameters.CellSize = Grid.CellSize;
 					}
 
+					if (CachedParameters.LoadingRange != Grid.LoadingRange)
+					{
+						MID->SetScalarParameterValue(*FString::Printf(TEXT("Grid%d_LoadingRange"), i), Grid.LoadingRange);
+						CachedParameters.LoadingRange = Grid.LoadingRange;
+					}
+
 					if (CachedParameters.GridColor != Grid.DebugColor)
 					{
 						MID->SetVectorParameterValue(*FString::Printf(TEXT("Grid%d_Color"), i), Grid.DebugColor);
diff --git a/Engine/Source/Runtime/Engine/Private/WorldPartition/WorldPartitionStreamingPolicy.cpp b/Engine/Source/Runtime/Engine/Private/WorldPartition/WorldPartitionStreamingPolicy.cpp
--- a/Engine/Source/Runtime/Engine/Private/WorldPartition/WorldPartitionStreamingPolicy.cpp
+++ b/Engine/Source/Runtime/Engine/Private/WorldPartition/WorldPartitionStreamingPolicy.cpp
@@ -47,7 +47,7 @@ void UWorldPartitionStreamingPolicy::UpdateStreamingSources()
 
 #if WITH_EDITOR
 	// We are in the SIE
-	if (GEditor->bIsSimulatingInEditor && GCurrentLevelEditingViewportClient->IsSimulateInEditorViewport())
+	if (UWorldPartition::IsSimulating())
 	{
 		// Transform to Local
 		const FVector ViewLocation = GCurrentLevelEditingViewportClient->GetViewLocation();
diff --git a/Engine/Source/Runtime/Engine/Private/WorldPartition/WorldPartitionSubsystem.cpp b/Engine/Source/Runtime/Engine/Private/WorldPartition/WorldPartitionSubsystem.cpp
--- a/Engine/Source/Runtime/Engine/Private/WorldPartition/WorldPartitionSubsystem.cpp
+++ b/Engine/Source/Runtime/Engine/Private/WorldPartition/WorldPartitionSubsystem.cpp
@@ -64,7 +64,7 @@ void UWorldPartitionSubsystem::PostInitialize()
 	{
 		MainPartition->Initialize(GetWorld(), FTransform::Identity);
 
-		if (GetWorld()->IsGameWorld() && (GetWorld()->GetNetMode() != NM_DedicatedServer))
+		if (MainPartition->CanDrawRuntimeHash() && (GetWorld()->GetNetMode() != NM_DedicatedServer))
 		{
 			DrawHandle = UDebugDrawService::Register(TEXT("Game"), FDebugDrawDelegate::CreateUObject(this, &UWorldPartitionSubsystem::Draw));
 		}
@@ -115,7 +115,7 @@ void UWorldPartitionSubsystem::Tick(float DeltaSeconds)
 	{
 		Partition->Tick(DeltaSeconds);
 		
-		if (GDrawRuntimeHash3D && GetWorld()->IsGameWorld())
+		if (GDrawRuntimeHash3D && Partition->CanDrawRuntimeHash())
 		{
 			Partition->DrawRuntimeHash3D();
 		}
diff --git a/Engine/Source/Runtime/Engine/Public/WorldPartition/WorldPartition.h b/Engine/Source/Runtime/Engine/Public/WorldPartition/WorldPartition.h
--- a/Engine/Source/Runtime/Engine/Public/WorldPartition/WorldPartition.h
+++ b/Engine/Source/Runtime/Engine/Public/WorldPartition/WorldPartition.h
@@ -121,7 +121,6 @@ public:
 	//~ End UObject Interface
 
 	FName GetWorldPartitionEditorName();
-	bool IsSimulating() const;
 
 	void LoadEditorCells(const FBox& Box);
 	void UnloadEditorCells(const FBox& Box);
@@ -151,6 +150,8 @@ public:
 #endif
 
 public:
+	static bool IsSimulating();
+
 	void Initialize(UWorld* World, const FTransform& InTransform);
 	bool IsInitialized() const;
 	virtual void Uninitialize() override;
@@ -163,6 +164,7 @@ public:
 	void UpdateStreamingState();
 	class ULevel* GetPreferredLoadedLevelToAddToWorld() const;
 
+	bool CanDrawRuntimeHash() const;
 	FVector2D GetDrawRuntimeHash2DDesiredFootprint(const FVector2D& CanvasSize);
 	void DrawRuntimeHash2D(class UCanvas* Canvas, const FVector2D& PartitionCanvasOffset, const FVector2D& PartitionCanvasSize);
 	void DrawRuntimeHash3D();
diff --git a/Engine/Source/Runtime/Engine/Public/WorldPartition/WorldPartitionRuntimeSpatialHash.h b/Engine/Source/Runtime/Engine/Public/WorldPartition/WorldPartitionRuntimeSpatialHash.h
--- a/Engine/Source/Runtime/Engine/Public/WorldPartition/WorldPartitionRuntimeSpatialHash.h
+++ b/Engine/Source/Runtime/Engine/Public/WorldPartition/WorldPartitionRuntimeSpatialHash.h
@@ -96,9 +96,6 @@ struct FSpatialHashStreamingGrid
 	UPROPERTY()
 	int32 CellSize;
 
-	UPROPERTY()
-	int32 GridSize;
-
 	UPROPERTY()
 	float LoadingRange;
 
@@ -108,6 +105,9 @@ struct FSpatialHashStreamingGrid
 	UPROPERTY()
 	TArray<FSpatialHashStreamingGridLevel> GridLevels;
 
+	UPROPERTY()
+	FBox WorldBounds;
+
 	float OverrideLoadingRange;
 
 	float GetLoadingRange() const
@@ -126,7 +126,6 @@ struct FSpatialHashStreamingGrid
 	void Draw3D(UWorld* World, const TArray<FWorldPartitionStreamingSource>& Sources, const FTransform& Transform) const;
 
 private:
-	mutable TMap<uint32, FVector> CachedLineTrace;
 	const FSquare2DGridHelper& GetGridHelper() const;
 	mutable FSquare2DGridHelper* GridHelper;
 };
diff --git a/Engine/Source/Runtime/Engine/Public/WorldPartition/WorldPartitionRuntimeSpatialHashGridPreviewer.h b/Engine/Source/Runtime/Engine/Public/WorldPartition/WorldPartitionRuntimeSpatialHashGridPreviewer.h
--- a/Engine/Source/Runtime/Engine/Public/WorldPartition/WorldPartitionRuntimeSpatialHashGridPreviewer.h
+++ b/Engine/Source/Runtime/Engine/Public/WorldPartition/WorldPartitionRuntimeSpatialHashGridPreviewer.h
@@ -35,11 +35,13 @@ struct ENGINE_API FWorldPartitionRuntimeSpatialHashGridPreviewer
 		FGridParametersCache()
 			: bEnabled(false)
 			, CellSize(0)
+			, LoadingRange(0.f)
 			, GridColor(0.f, 0.f, 0.f, 0.f)
 		{}
 
 		bool bEnabled;
 		int32 CellSize;
+		float LoadingRange;
 		FLinearColor GridColor;
 	};
 
