Make old behaviour for auto complete of console commands the default. Fix crash.

diff --git a/Engine/Source/Runtime/Engine/Private/UserInterface/Console.cpp b/Engine/Source/Runtime/Engine/Private/UserInterface/Console.cpp
--- a/Engine/Source/Runtime/Engine/Private/UserInterface/Console.cpp
+++ b/Engine/Source/Runtime/Engine/Private/UserInterface/Console.cpp
@@ -61,7 +61,7 @@ static TAutoConsoleVariable<int32> CVarConsoleYPos(
 
 static TAutoConsoleVariable<bool> CVarConsoleLegacySearch(
 	TEXT("console.searchmode.legacy"),
-	false,
+	true,
 	TEXT("Use the legacy search behaviour for console commands \n"),
 	ECVF_Default);
 
@@ -485,14 +485,13 @@ void UConsole::UpdateCompleteIndices()
 		BuildRuntimeAutoCompleteList(true);
 	}
 
-	AutoComplete.Reset();
+	AutoComplete.Empty();
+	AutoCompleteIndex = 0;
+	AutoCompleteCursor = 0;
 
 	if (CVarConsoleLegacySearch.GetValueOnAnyThread())
 	{
 		// use the old autocomplete behaviour
-		AutoCompleteIndex = 0;
-		AutoCompleteCursor = 0;
-
 		FAutoCompleteNode* Node = &AutoCompleteTree;
 		FString LowerTypedStr = TypedStr.ToLower();
 		int32 EndIdx = -1;
