Fixed Non-Unity compiler error in Android build

diff --git a/Engine/Source/Runtime/Engine/Public/PrecomputedVolumetricLightmap.h b/Engine/Source/Runtime/Engine/Public/PrecomputedVolumetricLightmap.h
--- a/Engine/Source/Runtime/Engine/Public/PrecomputedVolumetricLightmap.h
+++ b/Engine/Source/Runtime/Engine/Public/PrecomputedVolumetricLightmap.h
@@ -6,6 +6,7 @@
 
 #pragma once
 
+#include "UObject/WeakObjectPtr.h"
 #include "CoreMinimal.h"
 #include "Misc/Guid.h"
 #include "Math/SHMath.h"
