Don't duplicate the editor hash when duplicating partitioned worlds for PIE.

diff --git a/Engine/Source/Runtime/Engine/Private/WorldPartition/WorldPartition.cpp b/Engine/Source/Runtime/Engine/Private/WorldPartition/WorldPartition.cpp
--- a/Engine/Source/Runtime/Engine/Private/WorldPartition/WorldPartition.cpp
+++ b/Engine/Source/Runtime/Engine/Private/WorldPartition/WorldPartition.cpp
@@ -786,18 +786,6 @@ void UWorldPartition::UnhashActorDesc(FWorldPartitionActorDesc* ActorDesc)
 }
 #endif
 
-void UWorldPartition::Serialize(FArchive& Ar)
-{
-	Super::Serialize(Ar);
-
-#if WITH_EDITOR
-	if (Ar.GetPortFlags() & PPF_Duplicate)
-	{
-		Ar << EditorHash;
-	}
-#endif
-}
-
 void UWorldPartition::BeginDestroy()
 {
 	Super::BeginDestroy();
@@ -838,7 +826,10 @@ UWorldPartitionStreamingPolicy* UWorldPartition::GetStreamingPolicy() const
 void UWorldPartition::Tick(float DeltaSeconds)
 {
 #if WITH_EDITOR
-	EditorHash->Tick(DeltaSeconds);
+	if (EditorHash)
+	{
+		EditorHash->Tick(DeltaSeconds);
+	}
 
 	if (bForceGarbageCollection)
 	{
diff --git a/Engine/Source/Runtime/Engine/Public/WorldPartition/WorldPartition.h b/Engine/Source/Runtime/Engine/Public/WorldPartition/WorldPartition.h
--- a/Engine/Source/Runtime/Engine/Public/WorldPartition/WorldPartition.h
+++ b/Engine/Source/Runtime/Engine/Public/WorldPartition/WorldPartition.h
@@ -99,7 +99,6 @@ private:
 #endif
 
 	//~ Begin UObject Interface
-	virtual void Serialize(FArchive& Ar) override;
 	virtual void BeginDestroy() override;
 	//~ End UObject Interface
 
