Added clear coat to strata conversion.

diff --git a/Engine/Source/Runtime/Engine/Private/Materials/Material.cpp b/Engine/Source/Runtime/Engine/Private/Materials/Material.cpp
--- a/Engine/Source/Runtime/Engine/Private/Materials/Material.cpp
+++ b/Engine/Source/Runtime/Engine/Private/Materials/Material.cpp
@@ -81,6 +81,8 @@
 #include "Materials/MaterialExpressionStrata.h"
 #include "Materials/StrataMaterial.h"
 #include "Materials/MaterialExpressionThinTranslucentMaterialOutput.h"
+#include "Materials/MaterialExpressionClearCoatNormalCustomOutput.h"
+#include "Materials/MaterialExpressionConstant.h"
 
 #if WITH_EDITOR
 #include "Logging/TokenizedMessage.h"
@@ -3794,6 +3796,75 @@ void UMaterial::ConvertMaterialToStrataMaterial()
 				FrontMaterial.Connect(0, TopSlabBSDF);
 			}
 		}
+		else if (MaterialDomain == MD_Surface && ShadingModel == MSM_ClearCoat)
+		{
+			// Top slab BSDF as a simple Disney material
+			UMaterialExpressionStrataSlabBSDF* BottomSlabBSDF = NewObject<UMaterialExpressionStrataSlabBSDF>(this);
+			MoveConnectionTo(BaseColor, BottomSlabBSDF, 0);								// BaseColor
+			MoveConnectionTo(Metallic, BottomSlabBSDF, 2);								// Metallic
+			MoveConnectionTo(Specular, BottomSlabBSDF, 3);								// Specular
+			MoveConnectionTo(Roughness, BottomSlabBSDF, 4);								// Roughness
+			CopyConnectionTo(Anisotropy, BottomSlabBSDF, 5);							// Anisotropy
+			CopyConnectionTo(Normal, BottomSlabBSDF, 6);								// Normal
+			CopyConnectionTo(Tangent, BottomSlabBSDF, 7);								// Tangent
+
+			// Now weight the top base material by opacity.
+			UMaterialExpressionStrataSlabBSDF* TopSlabBSDF = NewObject<UMaterialExpressionStrataSlabBSDF>(this);
+			MoveConnectionTo(EmissiveColor, TopSlabBSDF, 10);							// Emissive
+			MoveConnectionTo(ClearCoatRoughness, TopSlabBSDF, 4);						// ClearCoatRoughness => Roughness
+			MoveConnectionTo(Anisotropy, TopSlabBSDF, 5);								// Anisotropy
+			MoveConnectionTo(Normal, TopSlabBSDF, 6);									// Normal
+			MoveConnectionTo(Tangent, TopSlabBSDF, 7);									// Tangent
+
+			//  The top layer has a hard coded specular value of 0.5 (F0 = 0.04)
+			UMaterialExpressionConstant* ConstantHalf = NewObject<UMaterialExpressionConstant>(this);
+			ConstantHalf->R = 0.5f;
+			TopSlabBSDF->GetInput(3)->Connect(0, ConstantHalf);							// BaseColor = 0 to only feature absorption, no scattering
+
+			// The original clear coat is a complex assemblage of arbitrary functions that do not always make sense. 
+			// To simplify things, we set the top slab BSDF as having a constant Grey scale transmittance. 
+			// As for the original, this is achieved with coverage so both transmittance and specular contribution vanishes
+			UMaterialExpressionConstant* ConstantZero = NewObject<UMaterialExpressionConstant>(this);
+			ConstantZero->R = 0.0f;
+			TopSlabBSDF->GetInput(0)->Connect(0, ConstantZero);							// BaseColor = 0 to only feature absorption, no scattering
+
+			// Now setup the mean free path with a hard coded transmittance of 0.75 when viewing the surface perpendicularly
+			UMaterialExpressionConstant* Constant075 = NewObject<UMaterialExpressionConstant>(this);
+			Constant075->R = 0.75f;
+			UMaterialExpressionStrataTransmittanceToMFP* TransToMDFP = NewObject<UMaterialExpressionStrataTransmittanceToMFP>(this);
+			TransToMDFP->GetInput(0)->Connect(0, Constant075);
+			TopSlabBSDF->GetInput(8)->Connect(0, TransToMDFP);							// MFP -> MFP
+			TopSlabBSDF->GetInput(13)->Connect(1, TransToMDFP);							// Thickness -> Thickness
+
+			// Connect the extra bottom normal if specified
+			TArray<class UMaterialExpressionCustomOutput*> CustomOutputExpressions;
+			GetAllCustomOutputExpressions(CustomOutputExpressions);
+			UMaterialExpressionClearCoatNormalCustomOutput* CCNormalCustomOutput = nullptr;
+			for (UMaterialExpressionCustomOutput* Expression : CustomOutputExpressions)
+			{
+				CCNormalCustomOutput = Cast<UMaterialExpressionClearCoatNormalCustomOutput>(Expression);
+				if (CCNormalCustomOutput)
+				{
+					break;
+				}
+			}
+			if (CCNormalCustomOutput)
+			{
+				MoveConnectionTo(*CCNormalCustomOutput->GetInput(0), BottomSlabBSDF, 6);	// ClearColorBottomNormal -> BottomSlabBSDF.Normal
+			}
+
+			// Now weight the top base material by ClearCoat
+			UMaterialExpressionStrataMultiply* TopSlabBSDFWithCoverage = NewObject<UMaterialExpressionStrataMultiply>(this);
+			TopSlabBSDFWithCoverage->GetInput(0)->Connect(0, TopSlabBSDF);				// TopSlabBSDF -> A
+			MoveConnectionTo(ClearCoat, TopSlabBSDFWithCoverage, 1);					// ClearCoat -> Weight
+
+			UMaterialExpressionStrataVerticalLayering* VerticalLayering = NewObject<UMaterialExpressionStrataVerticalLayering>(this);
+			VerticalLayering->GetInput(0)->Connect(0, TopSlabBSDFWithCoverage);			// Top -> Top
+			VerticalLayering->GetInput(1)->Connect(0, BottomSlabBSDF);					// Bottom -> Base
+
+			FrontMaterial.Connect(0, VerticalLayering);
+
+		}
 		else if (MaterialDomain == MD_Volume)
 		{
 			UMaterialExpressionStrataVolumetricFogCloudBSDF* VolBSDF = NewObject<UMaterialExpressionStrataVolumetricFogCloudBSDF>(this);
@@ -3807,9 +3878,13 @@ void UMaterial::ConvertMaterialToStrataMaterial()
 			FrontMaterial.Connect(0, VolBSDF);
 		}
 
-		// STRATA_TODO Other conversion: SSS, cloth, subsurface, Water, thin translucent, Clear Coat, etc. See EMaterialShadingModel
-
-		// STRATA_TODO unified translucency and coverage
+		// STRATA_TODO Other conversion, see EMaterialShadingModel: 
+		//	- MSM_Subsurface
+		//	- MSM_PreintegratedSkin
+		//	- MSM_SubsurfaceProfile
+		//	- MSM_TwoSidedFoliage
+		//	- MSM_Cloth
+		//	- MSM_Eye
 
 		// Now force the material to recompile and we use a hash of the original StateId.
 		// This is to avoid having different StateId each time we load the material and to not forever recompile it, i.e. use a cached version.
diff --git a/Engine/Source/Runtime/Engine/Private/Materials/MaterialExpressions.cpp b/Engine/Source/Runtime/Engine/Private/Materials/MaterialExpressions.cpp
--- a/Engine/Source/Runtime/Engine/Private/Materials/MaterialExpressions.cpp
+++ b/Engine/Source/Runtime/Engine/Private/Materials/MaterialExpressions.cpp
@@ -18164,7 +18164,12 @@ int32  UMaterialExpressionClearCoatNormalCustomOutput::Compile(class FMaterialCo
 	}
 	else
 	{
-		return CompilerError(Compiler, TEXT("Input missing"));
+		static const auto CVarStrata = IConsoleManager::Get().FindTConsoleVariableDataInt(TEXT("r.Strata"));
+		const bool bStrata = CVarStrata ? CVarStrata->GetValueOnAnyThread() > 0 : false;
+		if (!bStrata)
+		{
+			return CompilerError(Compiler, TEXT("Input missing"));
+		}
 	}
 	return INDEX_NONE;
 }
