WP HLOD - Remove debug code

diff --git a/Engine/Source/Runtime/Engine/Private/WorldPartition/HLOD/HLODBuilder.cpp b/Engine/Source/Runtime/Engine/Private/WorldPartition/HLOD/HLODBuilder.cpp
--- a/Engine/Source/Runtime/Engine/Private/WorldPartition/HLOD/HLODBuilder.cpp
+++ b/Engine/Source/Runtime/Engine/Private/WorldPartition/HLOD/HLODBuilder.cpp
@@ -546,10 +546,6 @@ uint32 FHLODBuilderUtilities::BuildHLOD(AWorldPartitionHLOD* InHLODActor)
 		return OldHLODHash;
 	}
 
-	// Hack to resave hashes
-	InHLODActor->MarkPackageDirty();
-	return NewHLODHash;
-
 	TUniquePtr<FHLODBuilder> HLODBuilder = nullptr;
 	
 	const UHLODLayer* HLODLayer = InHLODActor->GetSubActorsHLODLayer();
