Fixing compile error after removing Tessellation.

diff --git a/Engine/Source/Runtime/Engine/Private/Materials/MaterialShared.cpp b/Engine/Source/Runtime/Engine/Private/Materials/MaterialShared.cpp
--- a/Engine/Source/Runtime/Engine/Private/Materials/MaterialShared.cpp
+++ b/Engine/Source/Runtime/Engine/Private/Materials/MaterialShared.cpp
@@ -3135,13 +3135,6 @@ const FMaterial& FMaterialRenderProxy::GetIncompleteMaterialWithFallback(ERHIFea
 			Material = FallbackMaterialProxy->GetMaterialNoFallback(InFeatureLevel);
 		} while (!Material);
 	}
-	if (Material->GetTessellationMode() != MTM_NoTessellation)
-	{
-		// Incomplete tesselation shaders can result in a mismatch between rendering data and shaders used that cause fatal issues.
-		// Because of this, it's better for us to return the fallback even if we're asked for an incomplete material.
-		const FMaterialRenderProxy* Unused{};
-		return GetMaterialWithFallback(InFeatureLevel, Unused);
-	}
 	return *Material;
 }
 
