Fix Crash loading Skeletal mesh: Add missing UsingCustomVersion(FUE5ReleaseStreamObjectVersion::GUID)

diff --git a/Engine/Source/Runtime/Engine/Private/SkeletalMeshLODRenderData.cpp b/Engine/Source/Runtime/Engine/Private/SkeletalMeshLODRenderData.cpp
--- a/Engine/Source/Runtime/Engine/Private/SkeletalMeshLODRenderData.cpp
+++ b/Engine/Source/Runtime/Engine/Private/SkeletalMeshLODRenderData.cpp
@@ -596,6 +596,7 @@ private:
 
 void FSkeletalMeshLODRenderData::SerializeStreamedData(FArchive& Ar, USkeletalMesh* Owner, int32 LODIdx, uint8 ClassDataStripFlags, bool bNeedsCPUAccess, bool bForceKeepCPUResources)
 {
+	Ar.UsingCustomVersion(FUE5ReleaseStreamObjectVersion::GUID);
 	FStripDataFlags StripFlags(Ar, ClassDataStripFlags);
 
 	// TODO: A lot of data in a render section is needed during initialization but maybe some can still be streamed
@@ -651,6 +652,7 @@ void FSkeletalMeshLODRenderData::SerializeStreamedData(FArchive& Ar, USkeletalMe
 
 void FSkeletalMeshLODRenderData::SerializeAvailabilityInfo(FArchive& Ar, USkeletalMesh* Owner, int32 LODIdx, bool bAdjacencyDataStripped, bool bNeedsCPUAccess)
 {
+	Ar.UsingCustomVersion(FUE5ReleaseStreamObjectVersion::GUID);
 	MultiSizeIndexContainer.SerializeMetaData(Ar, bNeedsCPUAccess);
 	if (Ar.IsLoading() && Ar.CustomVer(FUE5ReleaseStreamObjectVersion::GUID) < FUE5ReleaseStreamObjectVersion::RemovingTessellation && !bAdjacencyDataStripped)
 	{
diff --git a/Engine/Source/Runtime/Engine/Private/StaticMesh.cpp b/Engine/Source/Runtime/Engine/Private/StaticMesh.cpp
--- a/Engine/Source/Runtime/Engine/Private/StaticMesh.cpp
+++ b/Engine/Source/Runtime/Engine/Private/StaticMesh.cpp
@@ -442,6 +442,8 @@ uint32 FStaticMeshLODResources::FStaticMeshBuffersSize::CalcBuffersSize() const
 
 void FStaticMeshLODResources::SerializeBuffers(FArchive& Ar, UStaticMesh* OwnerStaticMesh, uint8 InStripFlags, FStaticMeshBuffersSize& OutBuffersSize)
 {
+	Ar.UsingCustomVersion(FUE5ReleaseStreamObjectVersion::GUID);
+
 	bool bEnableDepthOnlyIndexBuffer = (CVarSupportDepthOnlyIndexBuffers.GetValueOnAnyThread() == 1);
 	bool bEnableReversedIndexBuffer = (CVarSupportReversedIndexBuffers.GetValueOnAnyThread() == 1);
 
@@ -563,6 +565,7 @@ void FStaticMeshLODResources::SerializeBuffers(FArchive& Ar, UStaticMesh* OwnerS
 
 void FStaticMeshLODResources::SerializeAvailabilityInfo(FArchive& Ar)
 {
+	Ar.UsingCustomVersion(FUE5ReleaseStreamObjectVersion::GUID);
 	bool bHasAdjacencyInfo = false;
 	const bool bEnableDepthOnlyIndexBuffer = !!CVarSupportDepthOnlyIndexBuffers.GetValueOnAnyThread();
 	const bool bEnableReversedIndexBuffer = !!CVarSupportReversedIndexBuffers.GetValueOnAnyThread();
