WorldPartition: Support embedding of child containers (Level Instances) into the main World Partition.

Current default behavior for Level Instance is to be embedded if it is using OFPA and Level Streaming if not.

diff --git a/Engine/Source/Editor/UnrealEd/Private/Commandlets/WorldPartitionConvertCommandlet.cpp b/Engine/Source/Editor/UnrealEd/Private/Commandlets/WorldPartitionConvertCommandlet.cpp
--- a/Engine/Source/Editor/UnrealEd/Private/Commandlets/WorldPartitionConvertCommandlet.cpp
+++ b/Engine/Source/Editor/UnrealEd/Private/Commandlets/WorldPartitionConvertCommandlet.cpp
@@ -1199,7 +1199,7 @@ int32 UWorldPartitionConvertCommandlet::Main(const FString& Params)
 			FActorSpawnParameters SpawnParams;
 			SpawnParams.OverrideLevel = MainLevel;
 			ALevelInstance* LevelInstanceActor = MainWorld->SpawnActor<ALevelInstance>(SpawnParams);
-
+			
 			FTransform LevelTransform;
 			if (SubLevelPackage->WorldTileInfo)
 			{
@@ -1210,6 +1210,7 @@ int32 UWorldPartitionConvertCommandlet::Main(const FString& Params)
 				LevelTransform = SubLevelStreaming->LevelTransform;
 			}
 
+			LevelInstanceActor->DesiredRuntimeBehavior = ELevelInstanceRuntimeBehavior::LevelStreaming;
 			LevelInstanceActor->SetActorTransform(LevelTransform);
 			LevelInstanceActor->SetWorldAsset(SubLevelWorld);
 		}
diff --git a/Engine/Source/Runtime/Core/Public/UObject/UE5ReleaseStreamObjectVersion.h b/Engine/Source/Runtime/Core/Public/UObject/UE5ReleaseStreamObjectVersion.h
--- a/Engine/Source/Runtime/Core/Public/UObject/UE5ReleaseStreamObjectVersion.h
+++ b/Engine/Source/Runtime/Core/Public/UObject/UE5ReleaseStreamObjectVersion.h
@@ -21,6 +21,9 @@ struct CORE_API FUE5ReleaseStreamObjectVersion
 		// Removing Tessellation from materials and meshes.
 		RemovingTessellation,
 
+		// LevelInstance serialize runtime behavior
+		LevelInstanceSerializeRuntimeBehavior,
+
 		// -----<new versions can be added above this line>-------------------------------------------------
 		VersionPlusOne,
 		LatestVersion = VersionPlusOne - 1
diff --git a/Engine/Source/Runtime/Engine/Private/Actor.cpp b/Engine/Source/Runtime/Engine/Private/Actor.cpp
--- a/Engine/Source/Runtime/Engine/Private/Actor.cpp
+++ b/Engine/Source/Runtime/Engine/Private/Actor.cpp
@@ -431,25 +431,22 @@ void AActor::GetAssetRegistryTags(TArray<FAssetRegistryTag>& OutTags) const
 
 	if (IsPackageExternal() && !IsChildActor())
 	{
-		if (UWorldPartition* WorldPartition = GetLevel()->GetWorldPartition())
-		{
-			TUniquePtr<FWorldPartitionActorDesc> ActorDesc(CreateActorDesc());
-			
-			const FString ActorMetaDataClass = GetParentNativeClass(GetClass())->GetName();
-			static FName NAME_ActorMetaDataClass(TEXT("ActorMetaDataClass"));
-			OutTags.Add(UObject::FAssetRegistryTag(NAME_ActorMetaDataClass, ActorMetaDataClass, UObject::FAssetRegistryTag::TT_Hidden));
+		TUniquePtr<FWorldPartitionActorDesc> ActorDesc(CreateActorDesc());
 
-			TArray<uint8> SerializedData;
-			ActorDesc->SerializeTo(SerializedData);
-			const FString ActorMetaData = FBase64::Encode(SerializedData);
-			static FName NAME_ActorMetaData(TEXT("ActorMetaData"));
-			OutTags.Add(UObject::FAssetRegistryTag(NAME_ActorMetaData, ActorMetaData, UObject::FAssetRegistryTag::TT_Hidden));
+		const FString ActorMetaDataClass = GetParentNativeClass(GetClass())->GetName();
+		static FName NAME_ActorMetaDataClass(TEXT("ActorMetaDataClass"));
+		OutTags.Add(UObject::FAssetRegistryTag(NAME_ActorMetaDataClass, ActorMetaDataClass, UObject::FAssetRegistryTag::TT_Hidden));
 
-			if (ActorLabel.Len() > 0)
-			{
-				static FName NAME_ActorLabel(TEXT("ActorLabel"));
-				OutTags.Add(UObject::FAssetRegistryTag(NAME_ActorLabel, ActorLabel, UObject::FAssetRegistryTag::TT_Hidden));
-			}
+		TArray<uint8> SerializedData;
+		ActorDesc->SerializeTo(SerializedData);
+		const FString ActorMetaData = FBase64::Encode(SerializedData);
+		static FName NAME_ActorMetaData(TEXT("ActorMetaData"));
+		OutTags.Add(UObject::FAssetRegistryTag(NAME_ActorMetaData, ActorMetaData, UObject::FAssetRegistryTag::TT_Hidden));
+
+		if (ActorLabel.Len() > 0)
+		{
+			static FName NAME_ActorLabel(TEXT("ActorLabel"));
+			OutTags.Add(UObject::FAssetRegistryTag(NAME_ActorLabel, ActorLabel, UObject::FAssetRegistryTag::TT_Hidden));
 		}
 	}
 }
diff --git a/Engine/Source/Runtime/Engine/Private/LevelInstance/LevelInstanceActor.cpp b/Engine/Source/Runtime/Engine/Private/LevelInstance/LevelInstanceActor.cpp
--- a/Engine/Source/Runtime/Engine/Private/LevelInstance/LevelInstanceActor.cpp
+++ b/Engine/Source/Runtime/Engine/Private/LevelInstance/LevelInstanceActor.cpp
@@ -30,6 +30,10 @@ ALevelInstance::ALevelInstance(const FObjectInitializer& ObjectInitializer)
 {
 	RootComponent = CreateDefaultSubobject<ULevelInstanceComponent>(TEXT("Root"));
 	RootComponent->Mobility = EComponentMobility::Static;
+
+#if WITH_EDITORONLY_DATA
+	DesiredRuntimeBehavior = ELevelInstanceRuntimeBehavior::Embedded;
+#endif
 }
 
 ULevelInstanceSubsystem* ALevelInstance::GetLevelInstanceSubsystem() const
@@ -464,6 +468,11 @@ bool ALevelInstance::CanEditChange(const FProperty* InProperty) const
 		return false;
 	}
 
+	if (InProperty->GetFName() == GET_MEMBER_NAME_CHECKED(ALevelInstance, DesiredRuntimeBehavior))
+	{
+		return GetLevel()->GetWorldPartition() != nullptr;
+	}
+
 	return true;
 }
 
diff --git a/Engine/Source/Runtime/Engine/Private/LevelInstance/Packed/PackedLevelInstanceActor.cpp b/Engine/Source/Runtime/Engine/Private/LevelInstance/Packed/PackedLevelInstanceActor.cpp
--- a/Engine/Source/Runtime/Engine/Private/LevelInstance/Packed/PackedLevelInstanceActor.cpp
+++ b/Engine/Source/Runtime/Engine/Private/LevelInstance/Packed/PackedLevelInstanceActor.cpp
@@ -14,6 +14,10 @@ APackedLevelInstance::APackedLevelInstance(const FObjectInitializer& ObjectIniti
 	, bChildChanged(false)
 #endif
 {
+#if WITH_EDITORONLY_DATA
+	// Packed Level Instances don't support level streaming or sub actors
+	DesiredRuntimeBehavior = ELevelInstanceRuntimeBehavior::None;
+#endif
 }
 
 bool APackedLevelInstance::SupportsLoading() const
@@ -127,6 +131,12 @@ bool APackedLevelInstance::CanEditChange(const FProperty* InProperty) const
 		return false;
 	}
 
+	// PackedLevelInstance doesn't support a runtime behavior
+	if (InProperty->GetFName() == GET_MEMBER_NAME_CHECKED(APackedLevelInstance, DesiredRuntimeBehavior))
+	{
+		return false;
+	}
+
 	// Disallow editing of the World if we are a BP instance
 	if (InProperty->GetFName() == GET_MEMBER_NAME_CHECKED(APackedLevelInstance, WorldAsset))
 	{
diff --git a/Engine/Source/Runtime/Engine/Private/WorldPartition/ActorDescContainer.cpp b/Engine/Source/Runtime/Engine/Private/WorldPartition/ActorDescContainer.cpp
--- a/Engine/Source/Runtime/Engine/Private/WorldPartition/ActorDescContainer.cpp
+++ b/Engine/Source/Runtime/Engine/Private/WorldPartition/ActorDescContainer.cpp
@@ -11,6 +11,7 @@
 
 UActorDescContainer::UActorDescContainer(const FObjectInitializer& ObjectInitializer)
 	: Super(ObjectInitializer)
+	, World(nullptr)
 #if WITH_EDITOR
 	, bContainerInitialized(false)
 	, bIgnoreAssetRegistryEvents(false)
@@ -18,8 +19,10 @@ UActorDescContainer::UActorDescContainer(const FObjectInitializer& ObjectInitial
 {
 }
 
-void UActorDescContainer::Initialize(FName InPackageName, bool bRegisterDelegates)
+void UActorDescContainer::Initialize(UWorld* InWorld, FName InPackageName, bool bRegisterDelegates)
 {
+	check(!World || World == InWorld);
+	World = InWorld;
 #if WITH_EDITOR
 	check(!bContainerInitialized);
 	ContainerPackageName = InPackageName;
@@ -63,6 +66,7 @@ void UActorDescContainer::Initialize(FName InPackageName, bool bRegisterDelegate
 
 void UActorDescContainer::Uninitialize()
 {
+	World = nullptr;
 #if WITH_EDITOR
 	if (bContainerInitialized)
 	{
@@ -77,6 +81,15 @@ void UActorDescContainer::Uninitialize()
 #endif
 }
 
+UWorld* UActorDescContainer::GetWorld() const
+{
+	if (World)
+	{
+		return World;
+	}
+	return Super::GetWorld();
+}
+
 void UActorDescContainer::BeginDestroy()
 {
 	Super::BeginDestroy();
diff --git a/Engine/Source/Runtime/Engine/Private/WorldPartition/HLOD/HLODBuilder.cpp b/Engine/Source/Runtime/Engine/Private/WorldPartition/HLOD/HLODBuilder.cpp
--- a/Engine/Source/Runtime/Engine/Private/WorldPartition/HLOD/HLODBuilder.cpp
+++ b/Engine/Source/Runtime/Engine/Private/WorldPartition/HLOD/HLODBuilder.cpp
@@ -399,18 +399,18 @@ class FHLODBuilder_MeshSimplify : public FHLODBuilder
 	}
 };
 
-TArray<AWorldPartitionHLOD*> FHLODBuilderUtilities::CreateHLODActors(FHLODCreationContext& InCreationContext, const FHLODCreationParams& InCreationParams, const TSet<FGuid>& InActors, const TArray<const UDataLayer*>& InDataLayers)
+TArray<AWorldPartitionHLOD*> FHLODBuilderUtilities::CreateHLODActors(FHLODCreationContext& InCreationContext, const FHLODCreationParams& InCreationParams, const TSet<FActorInstance>& InActors, const TArray<const UDataLayer*>& InDataLayers)
 {
 	TMap<UHLODLayer*, TArray<FGuid>> HLODLayersActors;
-	for (const FGuid& ActorGuid : InActors)
+	for (const FActorInstance& ActorInstance : InActors)
 	{
-		FWorldPartitionActorDesc& ActorDesc = InCreationParams.WorldPartition->GetActorDescChecked(ActorGuid);
+		FWorldPartitionActorDesc& ActorDesc = InCreationParams.WorldPartition->GetActorDescChecked(ActorInstance.Actor);
 		if (ActorDesc.GetActorIsHLODRelevant())
 		{
 			UHLODLayer* HLODLayer = UHLODLayer::GetHLODLayer(ActorDesc, InCreationParams.WorldPartition);
 			if (HLODLayer)
 			{
-				HLODLayersActors.FindOrAdd(HLODLayer).Add(ActorGuid);
+				HLODLayersActors.FindOrAdd(HLODLayer).Add(ActorInstance.Actor);
 			}
 		}
 	}
@@ -572,7 +572,7 @@ uint32 FHLODBuilderUtilities::BuildHLOD(AWorldPartitionHLOD* InHLODActor)
 	{
 		HLODBuilder->Build(InHLODActor, HLODLayer, SubActors);
 	}
-
+			
 	InHLODActor->MarkPackageDirty();
 
 	return NewHLODHash;
diff --git a/Engine/Source/Runtime/Engine/Private/WorldPartition/HLOD/HLODBuilder.h b/Engine/Source/Runtime/Engine/Private/WorldPartition/HLOD/HLODBuilder.h
--- a/Engine/Source/Runtime/Engine/Private/WorldPartition/HLOD/HLODBuilder.h
+++ b/Engine/Source/Runtime/Engine/Private/WorldPartition/HLOD/HLODBuilder.h
@@ -8,6 +8,7 @@
 #include "Containers/Array.h"
 #include "WorldPartition/WorldPartitionHandle.h"
 #include "WorldPartition/DataLayer/DataLayersID.h"
+#include "WorldPartition/WorldPartitionActorCluster.h"
 
 class AActor;
 class UWorldPartition;
@@ -50,7 +51,7 @@ public:
 	 * @param	InActors			The actors for which we'll build an HLOD representation
 	 * @param	InDataLayers		The data layers to assign to the newly created HLOD actors
 	 */
-	static TArray<AWorldPartitionHLOD*> CreateHLODActors(FHLODCreationContext& InCreationContext, const FHLODCreationParams& InCreationParams, const TSet<FGuid>& InActors, const TArray<const UDataLayer*>& InDataLayers);
+	static TArray<AWorldPartitionHLOD*> CreateHLODActors(FHLODCreationContext& InCreationContext, const FHLODCreationParams& InCreationParams, const TSet<FActorInstance>& InActors, const TArray<const UDataLayer*>& InDataLayers);
 
 	/**
 	 * Build HLOD for the specified AWorldPartitionHLOD actor.
diff --git a/Engine/Source/Runtime/Engine/Private/WorldPartition/LevelInstance/LevelInstanceActorDesc.cpp b/Engine/Source/Runtime/Engine/Private/WorldPartition/LevelInstance/LevelInstanceActorDesc.cpp
--- a/Engine/Source/Runtime/Engine/Private/WorldPartition/LevelInstance/LevelInstanceActorDesc.cpp
+++ b/Engine/Source/Runtime/Engine/Private/WorldPartition/LevelInstance/LevelInstanceActorDesc.cpp
@@ -1,10 +1,20 @@
 // Copyright Epic Games, Inc. All Rights Reserved.
 
 #include "WorldPartition/LevelInstance/LevelInstanceActorDesc.h"
+#include "WorldPartition/ActorDescContainer.h"
 
 #if WITH_EDITOR
+#include "Engine/World.h"
 #include "Engine/Level.h"
 #include "LevelInstance/LevelInstanceActor.h"
+#include "WorldPartition/WorldPartition.h"
+#include "UObject/UE5ReleaseStreamObjectVersion.h"
+
+FLevelInstanceActorDesc::FLevelInstanceActorDesc()
+	: DesiredRuntimeBehavior(ELevelInstanceRuntimeBehavior::Embedded)
+{
+
+}
 
 void FLevelInstanceActorDesc::Init(const AActor* InActor)
 {
@@ -12,15 +22,58 @@ void FLevelInstanceActorDesc::Init(const AActor* InActor)
 
 	const ALevelInstance* LevelInstanceActor = CastChecked<ALevelInstance>(InActor);
 	LevelPackage = *LevelInstanceActor->GetWorldAssetPackage();
-	LevelInstanceTransform = FTransform(InActor->GetActorRotation(), InActor->GetActorLocation());
+	LevelInstanceTransform = LevelInstanceActor->GetActorTransform();
+	DesiredRuntimeBehavior = LevelInstanceActor->GetDesiredRuntimeBehavior();
+
+	RegisterContainer();
+}
+
+void FLevelInstanceActorDesc::Init(UActorDescContainer* InContainer, const FWorldPartitionActorDescInitData& DescData)
+{
+	FWorldPartitionActorDesc::Init(InContainer, DescData);
+
+	RegisterContainer();
+}
+
+void FLevelInstanceActorDesc::RegisterContainer()
+{
+	check(!LevelInstanceContainer);
+	if (Container && (DesiredRuntimeBehavior == ELevelInstanceRuntimeBehavior::Embedded || DesiredRuntimeBehavior == ELevelInstanceRuntimeBehavior::Partitioned))
+	{
+		if (!LevelPackage.IsNone() && ULevel::GetIsLevelUsingExternalActorsFromPackage(LevelPackage))
+		{
+			UWorldPartition* WorldPartition = Container->GetWorld()->GetWorldPartition();
+			LevelInstanceContainer = WorldPartition->RegisterActorDescContainer(LevelPackage);
+		}
+	}
+}
+
+bool FLevelInstanceActorDesc::GetContainerInstance(const UActorDescContainer*& OutLevelContainer, FTransform& OutLevelTransform, EContainerClusterMode& OutClusterMode) const
+{
+	if (LevelInstanceContainer)
+	{
+		OutLevelContainer = LevelInstanceContainer;
+		OutLevelTransform = LevelInstanceTransform;
+		check(DesiredRuntimeBehavior == ELevelInstanceRuntimeBehavior::Embedded || DesiredRuntimeBehavior == ELevelInstanceRuntimeBehavior::Partitioned)
+		OutClusterMode = DesiredRuntimeBehavior == ELevelInstanceRuntimeBehavior::Embedded ? EContainerClusterMode::Embedded : EContainerClusterMode::Partitioned;
+		return true;
+	}
+
+	return false;
 }
 
 void FLevelInstanceActorDesc::Serialize(FArchive& Ar)
 {
 	FWorldPartitionActorDesc::Serialize(Ar);
+	Ar.UsingCustomVersion(FUE5ReleaseStreamObjectVersion::GUID);
 
 	Ar << LevelPackage << LevelInstanceTransform;
 
+	if (Ar.CustomVer(FUE5ReleaseStreamObjectVersion::GUID) >= FUE5ReleaseStreamObjectVersion::LevelInstanceSerializeRuntimeBehavior)
+	{
+		Ar << DesiredRuntimeBehavior;
+	}
+
 	if (Ar.IsLoading())
 	{
 		if (!LevelPackage.IsNone())
@@ -40,4 +93,10 @@ void FLevelInstanceActorDesc::Serialize(FArchive& Ar)
 		}
 	}
 }
+
+void FLevelInstanceActorDesc::AddReferencedObjects(FReferenceCollector& Collector)
+{
+	Collector.AddReferencedObject(LevelInstanceContainer);
+}
+
 #endif
diff --git a/Engine/Source/Runtime/Engine/Private/WorldPartition/RuntimeSpatialHash/RuntimeSpatialHashActorCluster.cpp b/Engine/Source/Runtime/Engine/Private/WorldPartition/RuntimeSpatialHash/RuntimeSpatialHashActorCluster.cpp
deleted file mode 100644
diff --git a/Engine/Source/Runtime/Engine/Private/WorldPartition/RuntimeSpatialHash/RuntimeSpatialHashActorCluster.h b/Engine/Source/Runtime/Engine/Private/WorldPartition/RuntimeSpatialHash/RuntimeSpatialHashActorCluster.h
deleted file mode 100644
diff --git a/Engine/Source/Runtime/Engine/Private/WorldPartition/RuntimeSpatialHash/RuntimeSpatialHashGridHelper.cpp b/Engine/Source/Runtime/Engine/Private/WorldPartition/RuntimeSpatialHash/RuntimeSpatialHashGridHelper.cpp
--- a/Engine/Source/Runtime/Engine/Private/WorldPartition/RuntimeSpatialHash/RuntimeSpatialHashGridHelper.cpp
+++ b/Engine/Source/Runtime/Engine/Private/WorldPartition/RuntimeSpatialHash/RuntimeSpatialHashGridHelper.cpp
@@ -90,7 +90,7 @@ void FSquare2DGridHelper::ValidateSingleActorReferer()
 {
 	UE_SCOPED_TIMER(TEXT("ValidateSingleActorReferer"), LogWorldPartitionRuntimeSpatialHash);
 
-	TSet<FGuid> ActorUsage;
+	TSet<FActorInstance> ActorUsage;
 	for (int32 Level = 0; Level < Levels.Num() - 1; Level++)
 	{
 		const int32 CurrentGridSize = Levels[Level].GridSize;
@@ -101,10 +101,10 @@ void FSquare2DGridHelper::ValidateSingleActorReferer()
 				FGridLevel::FGridCell& ThisCell = Levels[Level].Cells[y * CurrentGridSize + x];
 				for (const FGridLevel::FGridCellDataChunk& DataChunk : ThisCell.GetDataChunks())
 				{
-					for (const FGuid& ActorGuid : DataChunk.GetActors())
+					for (const FActorInstance& ActorInstance : DataChunk.GetActors())
 					{
 						bool bWasAlreadyInSet;
-						ActorUsage.Add(ActorGuid, &bWasAlreadyInSet);
+						ActorUsage.Add(ActorInstance, &bWasAlreadyInSet);
 						check(!bWasAlreadyInSet);
 					}
 				}
@@ -123,7 +123,7 @@ FSquare2DGridHelper GetGridHelper(const FBox& WorldBounds, int32 GridCellSize)
 	return FSquare2DGridHelper(WorldBounds, GridOrigin, GridCellSize);
 }
 
-FSquare2DGridHelper GetPartitionedActors(const UWorldPartition* WorldPartition, const FBox& WorldBounds, const FSpatialHashRuntimeGrid& Grid, const TArray<FActorCluster>& GridActors)
+FSquare2DGridHelper GetPartitionedActors(const UWorldPartition* WorldPartition, const FBox& WorldBounds, const FSpatialHashRuntimeGrid& Grid, const TArray<const FActorClusterInstance*>& GridActors)
 {
 	UE_SCOPED_TIMER(TEXT("GetPartitionedActors"), LogWorldPartitionRuntimeSpatialHash);
 
@@ -142,11 +142,11 @@ FSquare2DGridHelper GetPartitionedActors(const UWorldPartition* WorldPartition,
 		}
 	}
 
-	for (const FActorCluster& ActorCluster : GridActors)
+	for (const FActorClusterInstance* ClusterInstance : GridActors)
 	{
-		check(ActorCluster.Actors.Num() > 0);
-
-		EActorGridPlacement GridPlacement = ActorCluster.GridPlacement;
+		const FActorCluster* ActorCluster = ClusterInstance->Cluster;
+		check(ActorCluster && ActorCluster->Actors.Num() > 0);
+		EActorGridPlacement GridPlacement = ActorCluster->GridPlacement;
 		bool bAlwaysLoadedPromotedCluster = (GridPlacement == EActorGridPlacement::None);
 		bool bAlwaysLoadedPromotedOutOfGrid = false;
 
@@ -160,29 +160,18 @@ FSquare2DGridHelper GetPartitionedActors(const UWorldPartition* WorldPartition,
 		case EActorGridPlacement::Location:
 		{
 			FIntVector2 CellCoords;
-			check(ActorCluster.Actors.Num() == 1);
-			const FGuid& ActorGuid = *ActorCluster.Actors.CreateConstIterator();
-			const FWorldPartitionActorDesc& ActorDesc = *WorldPartition->GetActorDesc(ActorGuid);
-			if (ensure(PartitionedActors.GetLowestLevel().GetCellCoords(FVector2D(ActorDesc.GetOrigin()), CellCoords)))
+			check(ActorCluster->Actors.Num() == 1);
+			const FGuid& ActorGuid = *ActorCluster->Actors.CreateConstIterator();
+			FActorInstance ActorInstance(ActorGuid, ClusterInstance->ContainerInstance);
+			if (ensure(PartitionedActors.GetLowestLevel().GetCellCoords(FVector2D(ActorInstance.GetActorDesc()->GetOrigin()), CellCoords)))
 			{
-				TArray<const UDataLayer*> ActorDataLayers;
-				if (const AWorldDataLayers* WorldDataLayers = AWorldDataLayers::Get(WorldPartition->GetWorld()))
-				{
-					for (const FName& DataLayerName : ActorDesc.GetDataLayers())
-					{
-						if (const UDataLayer* DataLayer = WorldDataLayers->GetDataLayerFromName(DataLayerName))
-						{
-							ActorDataLayers.Add(DataLayer);
-						}
-					}
-				}
-				PartitionedActors.GetLowestLevel().GetCell(CellCoords).AddActor(ActorGuid, ActorDataLayers);
+				PartitionedActors.GetLowestLevel().GetCell(CellCoords).AddActor(MoveTemp(ActorInstance), ClusterInstance->DataLayers);
 			}
 			else
 			{
 				GridPlacement = EActorGridPlacement::AlwaysLoaded;
 				bAlwaysLoadedPromotedOutOfGrid = true;
-				PartitionedActors.GetAlwaysLoadedCell().AddActors(ActorCluster.Actors, ActorCluster.DataLayers);
+				PartitionedActors.GetAlwaysLoadedCell().AddActor(MoveTemp(ActorInstance), ClusterInstance->DataLayers);
 			}
 			break;
 		}
@@ -193,12 +182,12 @@ FSquare2DGridHelper GetPartitionedActors(const UWorldPartition* WorldPartition,
 			for (FSquare2DGridHelper::FGridLevel& GridLevel : PartitionedActors.Levels)
 			{
 				int32 IntersectingCellCount = 0;
-				GridLevel.ForEachIntersectingCellsBreakable(ActorCluster.Bounds, [&IntersectingCellCount](const FIntVector2& Coords) { return ++IntersectingCellCount <= 1; });
+				GridLevel.ForEachIntersectingCellsBreakable(ClusterInstance->Bounds, [&IntersectingCellCount](const FIntVector2& Coords) { return ++IntersectingCellCount <= 1; });
 				if (IntersectingCellCount == 1)
 				{
-					GridLevel.ForEachIntersectingCells(ActorCluster.Bounds, [&GridLevel, &ActorCluster](const FIntVector2& Coords)
+					GridLevel.ForEachIntersectingCells(ClusterInstance->Bounds, [&GridLevel, &ActorCluster, &ClusterInstance](const FIntVector2& Coords)
 						{
-							GridLevel.GetCell(Coords).AddActors(ActorCluster.Actors, ActorCluster.DataLayers);
+							GridLevel.GetCell(Coords).AddActors(ActorCluster->Actors, ClusterInstance->ContainerInstance, ClusterInstance->DataLayers);
 						});
 					bFoundCell = true;
 					break;
@@ -208,13 +197,13 @@ FSquare2DGridHelper GetPartitionedActors(const UWorldPartition* WorldPartition,
 			{
 				GridPlacement = EActorGridPlacement::AlwaysLoaded;
 				bAlwaysLoadedPromotedOutOfGrid = true;
-				PartitionedActors.GetAlwaysLoadedCell().AddActors(ActorCluster.Actors, ActorCluster.DataLayers);
+				PartitionedActors.GetAlwaysLoadedCell().AddActors(ActorCluster->Actors, ClusterInstance->ContainerInstance, ClusterInstance->DataLayers);
 			}
 			break;
 		}
 		case EActorGridPlacement::AlwaysLoaded:
 		{
-			PartitionedActors.GetAlwaysLoadedCell().AddActors(ActorCluster.Actors, ActorCluster.DataLayers);
+			PartitionedActors.GetAlwaysLoadedCell().AddActors(ActorCluster->Actors, ClusterInstance->ContainerInstance, ClusterInstance->DataLayers);
 			break;
 		}
 		default:
@@ -223,23 +212,24 @@ FSquare2DGridHelper GetPartitionedActors(const UWorldPartition* WorldPartition,
 
 		if (UE_LOG_ACTIVE(LogWorldPartitionRuntimeSpatialHash, Verbose))
 		{
-			if (ActorCluster.Actors.Num() > 1)
+			if (ActorCluster->Actors.Num() > 1)
 			{
 				static UEnum* ActorGridPlacementEnum = FindObject<UEnum>(ANY_PACKAGE, TEXT("EActorGridPlacement"));
 
 				UE_LOG(LogWorldPartitionRuntimeSpatialHash, Verbose, TEXT("Clustered %d actors (%s%s%s), generated shared BV of [%d x %d] (meters)"),
-					ActorCluster.Actors.Num(),
+					ActorCluster->Actors.Num(),
 					*ActorGridPlacementEnum->GetNameStringByValue((int64)GridPlacement),
 					bAlwaysLoadedPromotedCluster ? TEXT(":PromotedCluster") : TEXT(""),
 					bAlwaysLoadedPromotedOutOfGrid ? TEXT(":PromotedOutOfGrid") : TEXT(""),
-					(int)(0.01f * ActorCluster.Bounds.GetSize().X),
-					(int)(0.01f * ActorCluster.Bounds.GetSize().Y));
+					(int)(0.01f * ClusterInstance->Bounds.GetSize().X),
+					(int)(0.01f * ClusterInstance->Bounds.GetSize().Y));
 
-				for (const FGuid& ActorGuid : ActorCluster.Actors)
+				for (const FGuid& ActorGuid : ActorCluster->Actors)
 				{
-					const FWorldPartitionActorDesc& Desc = *WorldPartition->GetActorDesc(ActorGuid);
+					const FWorldPartitionActorDesc& Desc = *ClusterInstance->ContainerInstance->Container->GetActorDesc(ActorGuid);
 					UE_LOG(LogWorldPartitionRuntimeSpatialHash, Verbose, TEXT("   - Actor: %s (%s)"), *Desc.GetActorPath().ToString(), *ActorGuid.ToString(EGuidFormats::UniqueObjectGuid));
-					UE_LOG(LogWorldPartitionRuntimeSpatialHash, Verbose, TEXT("            %s"), *Desc.GetActorPackage().ToString());
+					UE_LOG(LogWorldPartitionRuntimeSpatialHash, Verbose, TEXT("         Package: %s"), *Desc.GetActorPackage().ToString());
+					UE_LOG(LogWorldPartitionRuntimeSpatialHash, Verbose, TEXT("         Container (%08x): %s"), ClusterInstance->ContainerInstance->ID, *ClusterInstance->ContainerInstance->Container->GetContainerPackage().ToString())
 				}
 			}
 		}
@@ -251,4 +241,4 @@ FSquare2DGridHelper GetPartitionedActors(const UWorldPartition* WorldPartition,
 	return PartitionedActors;
 }
 
-#endif // #if WITH_EDITOR
\ No newline at end of file
+#endif // #if WITH_EDITOR
diff --git a/Engine/Source/Runtime/Engine/Private/WorldPartition/RuntimeSpatialHash/RuntimeSpatialHashGridHelper.h b/Engine/Source/Runtime/Engine/Private/WorldPartition/RuntimeSpatialHash/RuntimeSpatialHashGridHelper.h
--- a/Engine/Source/Runtime/Engine/Private/WorldPartition/RuntimeSpatialHash/RuntimeSpatialHashGridHelper.h
+++ b/Engine/Source/Runtime/Engine/Private/WorldPartition/RuntimeSpatialHash/RuntimeSpatialHashGridHelper.h
@@ -7,7 +7,7 @@
 #include "ProfilingDebugging/ScopedTimers.h"
 
 #include "WorldPartition/WorldPartitionRuntimeSpatialHash.h"
-#include "WorldPartition/RuntimeSpatialHash/RuntimeSpatialHashActorCluster.h"
+#include "WorldPartition/WorldPartitionActorCluster.h"
 #include "WorldPartition/DataLayer/DataLayersID.h"
 
 /**
@@ -253,21 +253,21 @@ struct FSquare2DGridHelper
 				DataLayersID = FDataLayersID(DataLayers);
 			}
 
-			void AddActor(const FGuid& Actor) { Actors.Add(Actor); }
-			const TSet<FGuid>& GetActors() const { return Actors; }
+			void AddActor(FActorInstance ActorInstance) { Actors.Add(MoveTemp(ActorInstance)); }
+			const TSet<FActorInstance>& GetActors() const { return Actors; }
 			bool HasDataLayers() const { return !DataLayers.IsEmpty(); }
 			const TArray<const UDataLayer*>& GetDataLayers() const { return DataLayers; }
 			const FDataLayersID& GetDataLayersID() const { return DataLayersID; }
 
 		private:
-			TSet<FGuid> Actors;
+			TSet<FActorInstance> Actors;
 			TArray<const UDataLayer*> DataLayers;
 			FDataLayersID DataLayersID;
 		};
 
 		struct FGridCell
 		{
-			void AddActor(const FGuid& InActor, const TArray<const UDataLayer*>& InDataLayers)
+			void AddActor(FActorInstance ActorInstance, const TArray<const UDataLayer*>& InDataLayers)
 			{
 				FDataLayersID DataLayersID = FDataLayersID(InDataLayers);
 				FGridCellDataChunk* ActorDataChunk = Algo::FindByPredicate(DataChunks, [&](FGridCellDataChunk& InDataChunk) { return InDataChunk.GetDataLayersID() == DataLayersID; });
@@ -275,14 +275,14 @@ struct FSquare2DGridHelper
 				{
 					ActorDataChunk = &DataChunks.Emplace_GetRef(InDataLayers);
 				}
-				ActorDataChunk->AddActor(InActor);
+				ActorDataChunk->AddActor(MoveTemp(ActorInstance));
 			}
 
-			void AddActors(const TSet<FGuid>& InActors, const TArray<const UDataLayer*>& InDataLayers)
+			void AddActors(const TSet<FGuid>& InActors, const FActorContainerInstance* ContainerInstance, const TArray<const UDataLayer*>& InDataLayers)
 			{
 				for (const FGuid& Actor : InActors)
 				{
-					AddActor(Actor, InDataLayers);
+					AddActor(FActorInstance(Actor, ContainerInstance), InDataLayers);
 				}
 			}
 
@@ -405,5 +405,5 @@ public:
 
 #if WITH_EDITOR
 FSquare2DGridHelper GetGridHelper(const FBox& WorldBounds, int32 GridCellSize);
-FSquare2DGridHelper GetPartitionedActors(const UWorldPartition* WorldPartition, const FBox& WorldBounds, const FSpatialHashRuntimeGrid& Grid, const TArray<FActorCluster>& GridActors);
+FSquare2DGridHelper GetPartitionedActors(const UWorldPartition* WorldPartition, const FBox& WorldBounds, const FSpatialHashRuntimeGrid& Grid, const TArray<const FActorClusterInstance*>& GridActors);
 #endif // #if WITH_EDITOR
diff --git a/Engine/Source/Runtime/Engine/Private/WorldPartition/RuntimeSpatialHash/RuntimeSpatialHashHLOD.cpp b/Engine/Source/Runtime/Engine/Private/WorldPartition/RuntimeSpatialHash/RuntimeSpatialHashHLOD.cpp
--- a/Engine/Source/Runtime/Engine/Private/WorldPartition/RuntimeSpatialHash/RuntimeSpatialHashHLOD.cpp
+++ b/Engine/Source/Runtime/Engine/Private/WorldPartition/RuntimeSpatialHash/RuntimeSpatialHashHLOD.cpp
@@ -92,7 +92,9 @@ static void DeletePackage(FWorldPartitionActorDesc* ActorDesc, ISourceControlHel
 	}
 }
 
-static TArray<FGuid> GenerateHLODsForGrid(UWorldPartition* WorldPartition, const FSpatialHashRuntimeGrid& RuntimeGrid, uint32 HLODLevel, FHLODCreationContext& Context, ISourceControlHelper* SourceControlHelper, bool bCreateActorsOnly, const TArray<FActorCluster>& ActorClusters)
+PRAGMA_DISABLE_OPTIMIZATION
+
+static TArray<FGuid> GenerateHLODsForGrid(UWorldPartition* WorldPartition, const FSpatialHashRuntimeGrid& RuntimeGrid, uint32 HLODLevel, FHLODCreationContext& Context, ISourceControlHelper* SourceControlHelper, bool bCreateActorsOnly, const TArray<const FActorClusterInstance*>& ClusterInstances)
 {
 	auto HasExceededMaxMemory = []()
 	{
@@ -126,7 +128,7 @@ static TArray<FGuid> GenerateHLODsForGrid(UWorldPartition* WorldPartition, const
 
 	const FBox WorldBounds = WorldPartition->GetWorldBounds();
 
-	const FSquare2DGridHelper PartitionedActors = GetPartitionedActors(WorldPartition, WorldBounds, RuntimeGrid, ActorClusters);
+	const FSquare2DGridHelper PartitionedActors = GetPartitionedActors(WorldPartition, WorldBounds, RuntimeGrid, ClusterInstances);
 	const FSquare2DGridHelper::FGridLevel::FGridCell& AlwaysLoadedCell = PartitionedActors.GetAlwaysLoadedCell();
 
 	auto ShouldGenerateHLODs = [&AlwaysLoadedCell](const FSquare2DGridHelper::FGridLevel::FGridCell& GridCell, const FSquare2DGridHelper::FGridLevel::FGridCellDataChunk& GridCellDataChunk)
@@ -254,6 +256,8 @@ static TArray<FGuid> GenerateHLODsForGrid(UWorldPartition* WorldPartition, const
 	return GridHLODActors;
 }
 
+PRAGMA_ENABLE_OPTIMIZATION
+
 // Find all HLOD grids from the HLODLayer assets we have and build a dependency graph
 static void GatherHLODGrids(TMap<FName, FSpatialHashRuntimeGrid>& OutHLODGrids, TMap<FName, TSet<FName>>& OutHLODGridsGraph)
 {
@@ -489,33 +493,35 @@ bool UWorldPartitionRuntimeSpatialHash::GenerateHLOD(ISourceControlHelper* Sourc
 	}
 
 	// Create actor clusters - ignore HLOD actors
-	TArray<FActorCluster> ActorClusters = CreateActorClusters(WorldPartition, [](const FWorldPartitionActorDesc& ActorDesc)
-	{
-		return !ActorDesc.GetActorClass()->IsChildOf<AWorldPartitionHLOD>();
-	});
-
-	TArray<TArray<FActorCluster>> GridsClusters;
+	FActorClusterContext ClusterContext(WorldPartition, TOptional<FActorClusterContext::FFilterPredicate>([](const FWorldPartitionActorDesc& ActorDesc)
+		{
+			return !ActorDesc.GetActorClass()->IsChildOf<AWorldPartitionHLOD>();
+		}), /* bInIncludeChildContainers=*/ false);
+		
+	TArray<TArray<const FActorClusterInstance*>> GridsClusters;
 	GridsClusters.InsertDefaulted(0, Grids.Num());
 
-	for (FActorCluster& ActorCluster : ActorClusters)
+	for (const FActorClusterInstance& ClusterInstance : ClusterContext.GetClusterInstances())
 	{
-		int32* FoundIndex = GridsMapping.Find(ActorCluster.RuntimeGrid);
+		const FActorCluster* ActorCluster = ClusterInstance.Cluster;
+		check(ActorCluster);
+		int32* FoundIndex = GridsMapping.Find(ActorCluster->RuntimeGrid);
 		if (!FoundIndex)
 		{
-			UE_LOG(LogWorldPartitionRuntimeSpatialHashHLOD, Error, TEXT("Invalid partition grid '%s' referenced by actor cluster"), *ActorCluster.RuntimeGrid.ToString());
+			UE_LOG(LogWorldPartitionRuntimeSpatialHashHLOD, Error, TEXT("Invalid partition grid '%s' referenced by actor cluster"), *ActorCluster->RuntimeGrid.ToString());
 		}
 
 		int32 GridIndex = FoundIndex ? *FoundIndex : 0;
-		GridsClusters[GridIndex].Add(MoveTemp(ActorCluster));
+		GridsClusters[GridIndex].Add(&ClusterInstance);
 	}
 
 	// Keep track of all valid HLOD actors, along with which runtime grid they live in
 	TMap<FName, TArray<FGuid>> GridsHLODActors;
 
-	auto GenerateHLODs = [&GridsHLODActors, WorldPartition, &GridsDepth, &Context, SourceControlHelper, bCreateActorsOnly](const FSpatialHashRuntimeGrid& RuntimeGrid, uint32 HLODLevel, const TArray<FActorCluster>& ActorClusters)
+	auto GenerateHLODs = [&GridsHLODActors, WorldPartition, &GridsDepth, &Context, SourceControlHelper, bCreateActorsOnly](const FSpatialHashRuntimeGrid& RuntimeGrid, uint32 HLODLevel, const TArray<const FActorClusterInstance*>& ActorClusterInstances)
 	{
 		// Generate HLODs for this grid
-		TArray<FGuid> HLODActors = GenerateHLODsForGrid(WorldPartition, RuntimeGrid, HLODLevel, Context, SourceControlHelper, bCreateActorsOnly, ActorClusters);
+		TArray<FGuid> HLODActors = GenerateHLODsForGrid(WorldPartition, RuntimeGrid, HLODLevel, Context, SourceControlHelper, bCreateActorsOnly, ActorClusterInstances);
 
 		for (const FGuid& HLODActorGuid : HLODActors)
 		{
@@ -534,6 +540,7 @@ bool UWorldPartitionRuntimeSpatialHash::GenerateHLOD(ISourceControlHelper* Sourc
 	}
 
 	// Now, go on and create HLOD actors from HLOD grids (HLOD 1-N)
+	FActorContainerInstance MainContainerInstance(WorldPartition);
 	for (const FName HLODGridName : SortedGrids)
 	{
 		// No need to process empty grids
@@ -545,19 +552,25 @@ bool UWorldPartitionRuntimeSpatialHash::GenerateHLOD(ISourceControlHelper* Sourc
 
 		// Create one actor cluster per HLOD actor
 		TArray<FActorCluster> HLODActorClusters;
+		TArray<FActorClusterInstance> HLODActorClusterInstances;
+		TArray<const FActorClusterInstance*> HLODActorClusterInstancePtrs;
 		HLODActorClusters.Reserve(GridsHLODActors[HLODGridName].Num());
+		HLODActorClusterInstances.Reserve(GridsHLODActors[HLODGridName].Num());
+		HLODActorClusterInstancePtrs.Reserve(GridsHLODActors[HLODGridName].Num());
 
 		for (const FGuid& HLODActorGuid : GridsHLODActors[HLODGridName])
 		{
 			FWorldPartitionActorDesc* HLODActorDesc = WorldPartition->GetActorDesc(HLODActorGuid);
 			check(HLODActorDesc);
-
-			HLODActorClusters.Emplace(HLODActorDesc, HLODActorDesc->GetGridPlacement(), GetWorld());
+						
+			FActorCluster& NewHLODActorCluster = HLODActorClusters.Emplace_GetRef(HLODActorDesc, HLODActorDesc->GetGridPlacement(), WorldPartition);
+			FActorClusterInstance& NewHLODActorClusterInstance = HLODActorClusterInstances.Emplace_GetRef(&NewHLODActorCluster, &MainContainerInstance);
+			HLODActorClusterInstancePtrs.Add(&NewHLODActorClusterInstance);
 		};
 
 		// Generate HLODs for this grid - retrieve actors from our ValidHLODActors map
 		// We can't rely on actor descs for newly created HLOD actors
-		GenerateHLODs(HLODGrids[HLODGridName], GridsDepth[HLODGridName] + 1, HLODActorClusters);
+		GenerateHLODs(HLODGrids[HLODGridName], GridsDepth[HLODGridName] + 1, HLODActorClusterInstancePtrs);
 	}
 
 	auto DeleteHLODActor = [&SourceControlHelper](FWorldPartitionHandle ActorHandle)
diff --git a/Engine/Source/Runtime/Engine/Private/WorldPartition/WorldPartition.cpp b/Engine/Source/Runtime/Engine/Private/WorldPartition/WorldPartition.cpp
--- a/Engine/Source/Runtime/Engine/Private/WorldPartition/WorldPartition.cpp
+++ b/Engine/Source/Runtime/Engine/Private/WorldPartition/WorldPartition.cpp
@@ -119,7 +119,6 @@ UWorldPartition::FWorldPartitionChangedEvent UWorldPartition::WorldPartitionChan
 
 UWorldPartition::UWorldPartition(const FObjectInitializer& ObjectInitializer)
 	: Super(ObjectInitializer)
-	, World(nullptr)
 #if WITH_EDITOR
 	, EditorHash(nullptr)
 	, WorldPartitionEditor(nullptr)
@@ -229,7 +228,7 @@ void UWorldPartition::Initialize(UWorld* InWorld, const FTransform& InTransform)
 		}
 
 		const bool bRegisterDelegates = bEditorOnly && IsMainWorldPartition();
-		UActorDescContainer::Initialize(PackageName, bRegisterDelegates);
+		UActorDescContainer::Initialize(World, PackageName, bRegisterDelegates);
 		check(bContainerInitialized);
 
 		for (UActorDescContainer::TIterator<> ActorDescIterator(this); ActorDescIterator; ++ActorDescIterator)
@@ -318,15 +317,6 @@ void UWorldPartition::Initialize(UWorld* InWorld, const FTransform& InTransform)
 #endif
 }
 
-UWorld* UWorldPartition::GetWorld() const
-{
-	if (World)
-	{
-		return World;
-	}
-	return Super::GetWorld();
-}
-
 void UWorldPartition::RegisterStreamingSourceProvider(IWorldPartitionStreamingSourceProvider* StreamingSource)
 {
 	StreamingSourceProviders.Add(StreamingSource);
@@ -348,8 +338,6 @@ void UWorldPartition::WorldPartitionOnLevelRemovedFromWorld(ULevel* Level, UWorl
 
 void UWorldPartition::Uninitialize()
 {
-	Super::Uninitialize();
-
 	if (IsInitialized())
 	{
 		check(World);
@@ -363,6 +351,15 @@ void UWorldPartition::Uninitialize()
 		}
 
 #if WITH_EDITOR
+		for (auto& Pair : ActorDescContainers)
+		{
+			if (UActorDescContainer* Container = Pair.Value.Get())
+			{
+				Container->Uninitialize();
+			}
+		}
+		ActorDescContainers.Empty();
+
 		if (World->IsPlayInEditor())
 		{
 			CleanupForPIE();
@@ -393,6 +390,8 @@ void UWorldPartition::Uninitialize()
 
 		InitState = EWorldPartitionInitState::Uninitialized;
 	}
+
+	Super::Uninitialize();
 }
 
 void UWorldPartition::CleanupWorldPartition()
@@ -939,6 +938,24 @@ void UWorldPartition::GenerateNavigationData()
 	RuntimeHash->GenerateNavigationData();
 }
 
+const UActorDescContainer* UWorldPartition::RegisterActorDescContainer(FName PackageName)
+{
+	TWeakObjectPtr<UActorDescContainer>* ExistingContainerPtr = ActorDescContainers.Find(PackageName);
+	if (ExistingContainerPtr)
+	{
+		if (UActorDescContainer* LevelContainer = ExistingContainerPtr->Get())
+		{
+			return LevelContainer;
+		}
+	}
+		
+	UActorDescContainer* NewContainer = NewObject<UActorDescContainer>(GetTransientPackage());
+	NewContainer->Initialize(GetWorld(), PackageName, true);
+	ActorDescContainers.Add(PackageName, TWeakObjectPtr<UActorDescContainer>(NewContainer));
+
+	return NewContainer;
+}
+
 void UWorldPartition::DumpActorDescs(const FString& Path)
 {
 	if (FArchive* LogFile = IFileManager::Get().CreateFileWriter(*Path))
diff --git a/Engine/Source/Runtime/Engine/Private/WorldPartition/WorldPartitionActorCluster.cpp b/Engine/Source/Runtime/Engine/Private/WorldPartition/WorldPartitionActorCluster.cpp
new file mode 100644
--- /dev/null
+++ b/Engine/Source/Runtime/Engine/Private/WorldPartition/WorldPartitionActorCluster.cpp
@@ -0,0 +1,363 @@
+// Copyright Epic Games, Inc. All Rights Reserved.
+
+#include "WorldPartition/WorldPartitionActorCluster.h"
+
+#if WITH_EDITOR
+
+#include "Algo/Transform.h"
+#include "Engine/World.h"
+#include "WorldPartition/WorldPartition.h"
+#include "WorldPartition/ActorDescContainer.h"
+#include "ActorReferencesUtils.h"
+#include "Engine/LevelScriptBlueprint.h"
+#include "Misc/HashBuilder.h"
+
+DEFINE_LOG_CATEGORY(LogWorldPartitionActorCluster);
+
+template<class LayerNameContainer>
+TSet<const UDataLayer*> GetDataLayers(UWorld* InWorld, const LayerNameContainer& DataLayerNames)
+{
+	TSet<const UDataLayer*> DataLayers;
+	if (const AWorldDataLayers* WorldDataLayers = AWorldDataLayers::Get(InWorld))
+	{
+		for (const FName& DataLayerName : DataLayerNames)
+		{
+			if (const UDataLayer* DataLayer = WorldDataLayers->GetDataLayerFromName(DataLayerName))
+			{
+				if (DataLayer->IsDynamicallyLoaded())
+				{
+					DataLayers.Add(DataLayer);
+				}
+			}
+		}
+	}
+	return DataLayers;
+}
+
+FActorCluster::FActorCluster(const FWorldPartitionActorDesc* InActorDesc, EActorGridPlacement InGridPlacement, const UActorDescContainer* InContainer)
+	: GridPlacement(InGridPlacement)
+	, RuntimeGrid(InActorDesc->GetRuntimeGrid())
+	, Bounds(InActorDesc->GetBounds())
+{
+	check(GridPlacement != EActorGridPlacement::None);
+
+	Actors.Add(InActorDesc->GetGuid());
+	DataLayers = GetDataLayers(InContainer->GetWorld(), InActorDesc->GetDataLayers());
+	DataLayersID = FDataLayersID(DataLayers.Array());
+}
+
+void FActorCluster::Add(const FActorCluster& InActorCluster)
+{
+	// Merge Actors
+	Actors.Append(InActorCluster.Actors);
+
+	// Merge RuntimeGrid
+	RuntimeGrid = RuntimeGrid == InActorCluster.RuntimeGrid ? RuntimeGrid : NAME_None;
+
+	// Merge Bounds
+	Bounds += InActorCluster.Bounds;
+
+	// Merge GridPlacement
+	// If currently None, will always stay None
+	if (GridPlacement != EActorGridPlacement::None)
+	{
+		// If grid placement differs between the two clusters
+		if (GridPlacement != InActorCluster.GridPlacement)
+		{
+			// If one of the two cluster was always loaded, set to None
+			if (InActorCluster.GridPlacement == EActorGridPlacement::AlwaysLoaded || GridPlacement == EActorGridPlacement::AlwaysLoaded)
+			{
+				GridPlacement = EActorGridPlacement::None;
+			}
+			else
+			{
+				GridPlacement = InActorCluster.GridPlacement;
+			}
+		}
+
+		// If current placement is set to Location, that won't make sense when merging two clusters. Set to Bounds
+		if (GridPlacement == EActorGridPlacement::Location)
+		{
+			GridPlacement = EActorGridPlacement::Bounds;
+		}
+	}
+
+	// Merge DataLayers
+	if (DataLayersID != InActorCluster.DataLayersID)
+	{
+		for (const UDataLayer* DataLayer : InActorCluster.DataLayers)
+		{
+			check(DataLayer->IsDynamicallyLoaded());
+			DataLayers.Add(DataLayer);
+		}
+		DataLayersID = FDataLayersID(DataLayers.Array());
+	}
+}
+
+FActorClusterInstance::FActorClusterInstance(const FActorCluster* InCluster, const FActorContainerInstance* InContainerInstance)
+	: Cluster(InCluster)
+	, ContainerInstance(InContainerInstance)
+{
+	// Bounds based on cluster mode
+	Bounds = Cluster->Bounds;
+	if (ContainerInstance->ClusterMode == EContainerClusterMode::Embedded)
+	{
+		Bounds = ContainerInstance->Bounds;
+	}
+	Bounds = Bounds.TransformBy(ContainerInstance->Transform);
+	
+	TSet<const UDataLayer*> DataLayerSet;
+	DataLayerSet.Reserve(Cluster->DataLayers.Num() + ContainerInstance->DataLayers.Num());
+	// If the ContainerInstance is a WorldPartition we want the Cluster DataLayers to be propagated to the ClusterInstance
+	if (InContainerInstance->Container->IsA<UWorldPartition>())
+	{
+		DataLayerSet.Append(Cluster->DataLayers);
+	}
+	// We also want to propagate the ContainerInstance DataLayers to the ClusterInstance always
+	if (ContainerInstance->DataLayers.Num())
+	{
+		DataLayerSet.Append(ContainerInstance->DataLayers);	
+	}
+	DataLayers.Reserve(DataLayerSet.Num());
+	DataLayers.Append(DataLayerSet.Array());
+}
+
+FActorClusterContext::FActorClusterContext(const UWorldPartition* InWorldPartition, TOptional<FFilterPredicate> InFilterPredicate, bool bInIncludeChildContainers)
+	: WorldPartition(InWorldPartition)
+	, FilterPredicate(InFilterPredicate)
+	, bIncludeChildContainers(bInIncludeChildContainers)
+	, InstanceCountHint(0)
+{
+	CreateActorClusters();
+}
+
+FActorContainerInstance::FActorContainerInstance(const UActorDescContainer* InContainer)
+	: FActorContainerInstance(0, FTransform::Identity, FBox(ForceInit), TSet<FName>(), EContainerClusterMode::Partitioned, InContainer, TSet<FGuid>())
+{
+
+}
+
+FActorContainerInstance::FActorContainerInstance(uint32 InID, const FTransform& InTransform, const FBox& InBounds, const TSet<FName>& InDataLayers, EContainerClusterMode InClusterMode, const UActorDescContainer* InContainer, TSet<FGuid> InChildContainers)
+	: ID(InID)
+	, Transform(InTransform)
+	, Bounds(InBounds)
+	, ClusterMode(InClusterMode)
+	, Container(InContainer)
+	, ChildContainers(InChildContainers)
+{
+	DataLayers = GetDataLayers(InContainer->GetWorld(), InDataLayers);
+}
+
+FActorInstance::FActorInstance()
+	: ContainerInstance(nullptr)
+{
+
+}
+
+FActorInstance::FActorInstance(const FGuid& InActor, const FActorContainerInstance* InContainerInstance)
+	: Actor(InActor)
+	, ContainerInstance(InContainerInstance)
+{
+	check(ContainerInstance);
+}
+
+bool FActorInstance::ShouldStripFromStreaming() const
+{
+	// If this Actor instance is a Container itself we strip it
+	if (ContainerInstance->ChildContainers.Contains(Actor))
+	{
+		return true;
+	}
+
+	const FWorldPartitionActorDesc* ActorDesc = GetActorDesc();
+	return ActorDesc->GetActorIsEditorOnly();
+}
+
+FVector FActorInstance::GetOrigin() const
+{
+	return ContainerInstance->Transform.TransformPosition(GetActorDesc()->GetOrigin());
+}
+
+const FWorldPartitionActorDesc* FActorInstance::GetActorDesc() const
+{
+	return ContainerInstance->Container->GetActorDesc(Actor);
+}
+
+void CreateActorCluster(const FWorldPartitionActorDesc* ActorDesc, EActorGridPlacement GridPlacement, TMap<FGuid, FActorCluster*>& ActorToActorCluster, TSet<FActorCluster*>& ActorClustersSet, const UActorDescContainer* ActorDescContainer)
+{
+	UWorld* World = ActorDescContainer->GetWorld();
+	const FGuid& ActorGuid = ActorDesc->GetGuid();
+
+	FActorCluster* ActorCluster = ActorToActorCluster.FindRef(ActorGuid);
+	if (!ActorCluster)
+	{
+		ActorCluster = new FActorCluster(ActorDesc, GridPlacement, ActorDescContainer);
+		ActorClustersSet.Add(ActorCluster);
+		ActorToActorCluster.Add(ActorGuid, ActorCluster);
+	}
+
+	// Don't include references from editor-only actors
+	if (!ActorDesc->GetActorIsEditorOnly())
+	{
+		for (const FGuid& ReferenceGuid : ActorDesc->GetReferences())
+		{
+			if (const FWorldPartitionActorDesc* ReferenceActorDesc = ActorDescContainer->GetActorDesc(ReferenceGuid))
+			{
+				// Don't include references to editor-only actors
+				if (!ReferenceActorDesc->GetActorIsEditorOnly())
+				{
+					FActorCluster* ReferenceCluster = ActorToActorCluster.FindRef(ReferenceGuid);
+					if (ReferenceCluster)
+					{
+						if (ReferenceCluster != ActorCluster)
+						{
+							// Merge reference cluster in Actor's cluster
+							ActorCluster->Add(*ReferenceCluster);
+							for (const FGuid& ReferenceClusterActorGuid : ReferenceCluster->Actors)
+							{
+								ActorToActorCluster[ReferenceClusterActorGuid] = ActorCluster;
+							}
+							ActorClustersSet.Remove(ReferenceCluster);
+							delete ReferenceCluster;
+						}
+					}
+					else
+					{
+						// Put Reference in Actor's cluster
+						ActorCluster->Add(FActorCluster(ReferenceActorDesc, GridPlacement, ActorDescContainer));
+					}
+
+					// Map its cluster
+					ActorToActorCluster.Add(ReferenceGuid, ActorCluster);
+				}
+			}
+		}
+	}
+}
+
+void FActorClusterContext::CreateContainerInstanceRecursive(uint32 ID, const FTransform& Transform, EContainerClusterMode ClusterMode , const UActorDescContainer* Container, const TSet<FName>& DataLayers, FBox* ParentBounds)
+{
+	InstanceCountHint += Container->GetActorDescCount();
+		
+	TSet<FGuid> ChildContainers;
+	FBox Bounds(ForceInit);
+
+	for (UActorDescContainer::TConstIterator<> ActorDescIterator(Container); ActorDescIterator; ++ActorDescIterator)
+	{
+		const FWorldPartitionActorDesc* ActorDesc = *ActorDescIterator;
+		const UActorDescContainer* OutContainer = nullptr;
+		FTransform OutTransform;
+		EContainerClusterMode OutClusterMode;
+		if (bIncludeChildContainers && ActorDesc->GetContainerInstance(OutContainer, OutTransform, OutClusterMode))
+		{
+			// Add Child Container Guid so we can discard the actor later
+			ChildContainers.Add(ActorDesc->GetGuid());
+
+			FHashBuilder ChildContainerHashBuilder(ID);
+			ChildContainerHashBuilder << ActorDesc->GetGuid();
+			
+			TSet<FName> ChildDataLayers;
+			ChildDataLayers.Reserve(DataLayers.Num() + ActorDesc->GetDataLayers().Num());
+			// Only propagate ActorDesc DataLayers if we are a Root Container (WorldPartition)
+			if (Container->IsA<UWorldPartition>())
+			{
+				ChildDataLayers.Append(ActorDesc->GetDataLayers());
+			}
+			// Always inherite parent container DataLayers
+			ChildDataLayers.Append(DataLayers);
+			CreateContainerInstanceRecursive(ChildContainerHashBuilder.GetHash(), OutTransform * Transform, OutClusterMode, OutContainer, ChildDataLayers, &Bounds);
+		}
+		else
+		{
+			switch (ActorDesc->GetGridPlacement())
+			{
+				case EActorGridPlacement::Location:
+				{
+					FVector Location = ActorDescIterator->GetOrigin();
+					Bounds += FBox(Location, Location);
+				}
+				break;
+				case EActorGridPlacement::Bounds:
+				{
+					Bounds += ActorDescIterator->GetBounds();
+				}
+				break;
+			}
+		}
+	}
+	
+	if (ParentBounds)
+	{
+		*ParentBounds += Bounds;
+	}
+
+	UE_LOG(LogWorldPartitionActorCluster, Verbose, TEXT("ContainerInstance (%08x) Bounds (%s) Package (%s)"), ID, *Bounds.TransformBy(Transform).ToString(), *Container->GetContainerPackage().ToString());
+	ContainerInstances.Add(FActorContainerInstance(ID, Transform, Bounds, DataLayers, ClusterMode, Container, MoveTemp(ChildContainers)));
+}
+
+const TArray<FActorCluster>& FActorClusterContext::CreateActorClustersImpl(const FActorContainerInstance& ContainerInstance)
+{
+	if (TArray<FActorCluster>* FoundClusters = Clusters.Find(ContainerInstance.Container))
+	{
+		return *FoundClusters;
+	}
+
+	TMap<FGuid, FActorCluster*> ActorToActorCluster;
+	TSet<FActorCluster*> ActorClustersSet;
+	TArray<FActorCluster>& ActorClusters = Clusters.Add(ContainerInstance.Container);
+	
+	// Gather all references to external actors from the level script
+	TSet<AActor*> LevelScriptExternalActorReferences;
+	// Only applies to root container
+	if (ContainerInstance.Container->IsA<UWorldPartition>())
+	{
+		if (ULevelScriptBlueprint* LevelScriptBlueprint = ContainerInstance.Container->GetWorld()->PersistentLevel->GetLevelScriptBlueprint(true))
+		{
+			LevelScriptExternalActorReferences.Append(ActorsReferencesUtils::GetExternalActorReferences(LevelScriptBlueprint));
+		}
+	}
+		
+	for (UActorDescContainer::TConstIterator<> ActorDescIterator(ContainerInstance.Container); ActorDescIterator; ++ActorDescIterator)
+	{
+		const FWorldPartitionActorDesc* ActorDesc = *ActorDescIterator;
+		EActorGridPlacement GridPlacement = ActorDesc->GetGridPlacement();
+
+		// Check if the actor is loaded (potentially referenced by the level script)
+		if (AActor* Actor = ActorDesc->GetActor())
+		{
+			if (LevelScriptExternalActorReferences.Contains(Actor))
+			{
+				GridPlacement = EActorGridPlacement::AlwaysLoaded;
+			}
+		}
+
+		if (!FilterPredicate.IsSet() || FilterPredicate.GetValue()(*ActorDesc))
+		{
+			CreateActorCluster(ActorDesc, GridPlacement, ActorToActorCluster, ActorClustersSet, ContainerInstance.Container);
+		}
+	}
+				
+	ActorClusters.Reserve(ActorClustersSet.Num());
+	Algo::Transform(ActorClustersSet, ActorClusters, [](FActorCluster* ActorCluster) { return MoveTemp(*ActorCluster); });
+	for (FActorCluster* ActorCluster : ActorClustersSet) { delete ActorCluster; }
+
+	return ActorClusters;
+}
+
+void FActorClusterContext::CreateActorClusters()
+{
+	// First Instance is the main WorldPartition
+	CreateContainerInstanceRecursive(0, FTransform::Identity, EContainerClusterMode::Partitioned, WorldPartition, TSet<FName>());
+			
+	ClusterInstances.Reserve(InstanceCountHint);
+	for (const FActorContainerInstance& ContainerInstance : ContainerInstances)
+	{
+		const TArray<FActorCluster>& NewClusters = CreateActorClustersImpl(ContainerInstance);
+		for (const FActorCluster& Cluster : NewClusters)
+		{
+			ClusterInstances.Emplace(&Cluster, &ContainerInstance);
+		}
+	}
+}
+
+#endif // #if WITH_EDITOR
\ No newline at end of file
diff --git a/Engine/Source/Runtime/Engine/Private/WorldPartition/WorldPartitionActorDesc.cpp b/Engine/Source/Runtime/Engine/Private/WorldPartition/WorldPartitionActorDesc.cpp
--- a/Engine/Source/Runtime/Engine/Private/WorldPartition/WorldPartitionActorDesc.cpp
+++ b/Engine/Source/Runtime/Engine/Private/WorldPartition/WorldPartitionActorDesc.cpp
@@ -85,7 +85,6 @@ void FWorldPartitionActorDesc::Init(const AActor* InActor)
 	ActorLabel = *InActor->GetActorLabel(false);
 
 	Container = InActor->GetLevel()->GetWorldPartition();
-	check(Container);
 }
 
 void FWorldPartitionActorDesc::Init(UActorDescContainer* InContainer, const FWorldPartitionActorDescInitData& DescData)
diff --git a/Engine/Source/Runtime/Engine/Private/WorldPartition/WorldPartitionLevelHelper.cpp b/Engine/Source/Runtime/Engine/Private/WorldPartition/WorldPartitionLevelHelper.cpp
--- a/Engine/Source/Runtime/Engine/Private/WorldPartition/WorldPartitionLevelHelper.cpp
+++ b/Engine/Source/Runtime/Engine/Private/WorldPartition/WorldPartitionLevelHelper.cpp
@@ -5,11 +5,16 @@
  */
 
 #include "WorldPartition/WorldPartitionLevelHelper.h"
+#include "WorldPartition/WorldPartitionPackageCache.h"
 
 #if WITH_EDITOR
 
 #include "FileHelpers.h"
 #include "UnrealEngine.h"
+#include "WorldPartition/ActorDescContainer.h"
+#include "WorldPartition/WorldPartitionActorDesc.h"
+#include "LevelUtils.h"
+#include "Templates/SharedPointer.h"
 
  /**
   * Defaults World's initialization values for World Partition StreamingLevels
@@ -40,7 +45,7 @@ void FWorldPartitionLevelHelper::MoveExternalActorsToLevel(const TArray<FWorldPa
 	// Move all actors to Cell level
 	for (const FWorldPartitionRuntimeCellObjectMapping& PackageObjectMapping : InChildPackages)
 	{
-		AActor* Actor = FindObject<AActor>(nullptr, *PackageObjectMapping.Path.ToString());
+		AActor* Actor = FindObject<AActor>(nullptr, *PackageObjectMapping.LoadedPath.ToString());
 		if (ensure(Actor))
 		{
 			const bool bSameOuter = (InLevel == Actor->GetOuter());
@@ -107,4 +112,185 @@ ULevel* FWorldPartitionLevelHelper::CreateEmptyLevelForRuntimeCell(const UWorld*
 	return NewLevel;
 }
 
+bool FWorldPartitionLevelHelper::LoadActors(ULevel* InDestLevel, TArrayView<FWorldPartitionRuntimeCellObjectMapping> InActorPackages, FWorldPartitionPackageCache& InPackageCache, FWorldPartitionLevelHelper::FOnLoadActorsCompleted InCompletionCallback, bool bLoadForPIE, bool bLoadAsync /*=true*/, FLinkerInstancingContext* InOutInstancingContext /*=nullptr*/)
+{
+	UPackage* DestPackage = InDestLevel ? InDestLevel->GetPackage() : nullptr;
+	FString ShortLevelPackageName = DestPackage? FPackageName::GetShortName(DestPackage->GetFName()) : FString();
+
+	struct FLoadProgress
+	{
+		int32 NumPendingLoadRequests = 0;
+		int32 NumFailedLoadedRequests = 0;
+	};
+	TSharedPtr<FLoadProgress> LoadProgress = MakeShared<FLoadProgress>();
+
+	// Actors to load
+	TArray<FString> ActorPackageInstanceNames;
+	ActorPackageInstanceNames.Reserve(InActorPackages.Num());
+	TArray<FWorldPartitionRuntimeCellObjectMapping*> ActorPackages;
+	ActorPackages.Reserve(InActorPackages.Num());
+
+	// Levels to Load with actors to duplicate
+	TMap<uint32, TArray<FWorldPartitionRuntimeCellObjectMapping*>> PackagesToDuplicate;
+	check(!bLoadForPIE || InOutInstancingContext);
+
+	for (FWorldPartitionRuntimeCellObjectMapping& PackageObjectMapping : InActorPackages)
+	{
+		if (PackageObjectMapping.ContainerID == 0)
+		{
+			if (bLoadForPIE)
+			{
+				FString ObjectPath = PackageObjectMapping.Package.ToString();
+				FString ActorPackageName = FPackageName::ObjectPathToPackageName(ObjectPath);
+				FString ActorPackageInstanceName = FString::Printf(TEXT("%s_%s"), *ActorPackageName, *ShortLevelPackageName);
+
+				ActorPackageInstanceNames.Add(ActorPackageInstanceName);
+				InOutInstancingContext->AddMapping(FName(*ActorPackageName), FName(*ActorPackageInstanceName));
+			}
+			ActorPackages.Add(&PackageObjectMapping);
+		}
+		else
+		{
+			// When loading LevelInstances we duplicate the package 
+			PackagesToDuplicate.FindOrAdd(PackageObjectMapping.ContainerID).Add(&PackageObjectMapping);
+		}
+	}
+
+	LoadProgress->NumPendingLoadRequests = ActorPackages.Num() + PackagesToDuplicate.Num();
+
+	for (int32 ChildIndex = 0; ChildIndex < ActorPackages.Num(); ChildIndex++)
+	{
+		const FWorldPartitionRuntimeCellObjectMapping& PackageObjectMapping = *ActorPackages[ChildIndex];
+		const FString ActorPackageName = FPackageName::ObjectPathToPackageName(PackageObjectMapping.Package.ToString());
+		FName ActorName = *FPaths::GetExtension(PackageObjectMapping.Path.ToString());
+
+		FLoadPackageAsyncDelegate CompletionCallback = FLoadPackageAsyncDelegate::CreateLambda([LoadProgress, ActorName, bLoadForPIE, InDestLevel, InCompletionCallback](const FName& LoadedPackageName, UPackage* LoadedPackage, EAsyncLoadingResult::Type Result)
+			{
+				check(LoadProgress->NumPendingLoadRequests);
+				LoadProgress->NumPendingLoadRequests--;
+				if (LoadedPackage)
+				{
+					AActor* Actor = FindObject<AActor>(LoadedPackage, *ActorName.ToString());
+
+					check(Actor);
+					if (bLoadForPIE)
+					{
+						check(Actor->IsPackageExternal());
+						InDestLevel->Actors.Add(Actor);
+						check(Actor->GetLevel() == InDestLevel);
+					}
+								
+					UE_LOG(LogEngine, Verbose, TEXT(" ==> Loaded %s (remaining: %d)"), *Actor->GetFullName(), LoadProgress->NumPendingLoadRequests);
+				}
+				else
+				{
+					UE_LOG(LogEngine, Warning, TEXT("Failed to load %s"), *LoadedPackageName.ToString());
+					//@todo_ow: cumulate and process when NumPendingActorRequests == 0
+					LoadProgress->NumFailedLoadedRequests++;
+				}
+
+				if (!LoadProgress->NumPendingLoadRequests)
+				{
+					InCompletionCallback.Execute(!LoadProgress->NumFailedLoadedRequests);
+				}
+			});
+
+
+		FPackagePath PackagePath = FPackagePath::FromPackageNameChecked(*ActorPackageName);
+		if (bLoadForPIE)
+		{
+			::LoadPackageAsync(PackagePath, FName(*ActorPackageInstanceNames[ChildIndex]), CompletionCallback, nullptr, PKG_PlayInEditor, DestPackage->PIEInstanceID, 0, InOutInstancingContext);
+		}
+		else
+		{
+			::LoadPackageAsync(PackagePath, NAME_None, CompletionCallback, nullptr, PKG_None, INDEX_NONE, 0, InOutInstancingContext);
+		}
+	}
+
+	for (auto& Pair : PackagesToDuplicate)
+	{
+		FString PackageToLoadFrom = Pair.Value[0]->ContainerPackage.ToString();
+		FLoadPackageAsyncDelegate CompletionCallback = FLoadPackageAsyncDelegate::CreateLambda([LoadProgress, LevelInstanceID = Pair.Key, Mappings = MoveTemp(Pair.Value), InDestLevel, &InPackageCache, InCompletionCallback](const FName& LoadedPackageName, UPackage* LoadedPackage, EAsyncLoadingResult::Type Result)
+		{
+			check(LoadProgress->NumPendingLoadRequests);
+			LoadProgress->NumPendingLoadRequests--;
+
+			if (LoadedPackage)
+			{
+				FName DuplicatePackageName(*FString::Printf(TEXT("%s_%08x"), *LoadedPackage->GetName(), LevelInstanceID));
+				UPackage* DuplicatedPackage = InPackageCache.FindPackage(DuplicatePackageName);
+				UWorld* DuplicatedWorld = nullptr;
+
+				// Check that existing Package contains all actors (might be a previously used Dup package when loading unloaded cells)
+				if (DuplicatedPackage)
+				{
+					DuplicatedWorld = UWorld::FindWorldInPackage(DuplicatedPackage);
+					for (auto Mapping : Mappings)
+					{
+						FString ActorName = FPaths::GetExtension(Mapping->Path.ToString());
+						if (!FindObject<AActor>(DuplicatedWorld->PersistentLevel, *ActorName))
+						{
+							InPackageCache.TrashPackage(DuplicatedPackage);
+							DuplicatedPackage = nullptr;
+							DuplicatedWorld = nullptr;
+							break;
+						}
+					}
+				}
+
+				// Package wasn't found or was missing actors
+				if (!DuplicatedPackage)
+				{
+					DuplicatedPackage = InPackageCache.DuplicateWorldPackage(LoadedPackage, DuplicatePackageName);
+					check(DuplicatedPackage)
+						DuplicatedWorld = UWorld::FindWorldInPackage(DuplicatedPackage);
+				}
+
+				for (auto Mapping : Mappings)
+				{
+					FString ActorName = FPaths::GetExtension(Mapping->Path.ToString());
+					AActor* DuplicatedActor = FindObject<AActor>(DuplicatedWorld->PersistentLevel, *ActorName);
+					check(DuplicatedActor);
+
+					DuplicatedActor->Rename(*FString::Printf(TEXT("%s_%08x"), *DuplicatedActor->GetName(), Mapping->ContainerID), InDestLevel, REN_NonTransactional | REN_ForceNoResetLoaders | REN_DoNotDirty | REN_DontCreateRedirectors);
+					USceneComponent* RootComponent = DuplicatedActor->GetRootComponent();
+					
+					FLevelUtils::FApplyLevelTransformParams TransformParams(nullptr, Mapping->ContainerTransform);
+					TransformParams.Actor = DuplicatedActor;
+					TransformParams.bDoPostEditMove = false;
+					FLevelUtils::ApplyLevelTransform(TransformParams);
+
+					// Path to use when searching for this actor in MoveExternalActorsToLevel
+					Mapping->LoadedPath = *DuplicatedActor->GetPathName();
+					UE_LOG(LogEngine, Verbose, TEXT(" ==> Duplicated %s (remaining: %d)"), *DuplicatedActor->GetFullName(), LoadProgress->NumPendingLoadRequests);
+				}
+			}
+			else
+			{
+				UE_LOG(LogEngine, Warning, TEXT("Failed to load %s"), *LoadedPackageName.ToString());
+				LoadProgress->NumFailedLoadedRequests++;
+			}
+
+			if (!LoadProgress->NumPendingLoadRequests)
+			{
+				InCompletionCallback.Execute(!LoadProgress->NumFailedLoadedRequests);
+			}
+		});
+
+		// PostFix the Package so that we don't load over the original package
+		FName PackageDuplicateSeedName(*FString::Printf(TEXT("%s_DUP"), *PackageToLoadFrom));
+		
+		// Use Package Cache here because Level Instances can load the same levels across cells we want to avoid Loading the same package multiple times...
+		InPackageCache.LoadWorldPackageAsync(PackageDuplicateSeedName, *PackageToLoadFrom, CompletionCallback);
+	}
+
+	if (!bLoadAsync)
+	{
+		// Finish all async loading.
+		FlushAsyncLoading();
+	}
+
+	return (LoadProgress->NumPendingLoadRequests == 0);
+}
+
 #endif
\ No newline at end of file
diff --git a/Engine/Source/Runtime/Engine/Private/WorldPartition/WorldPartitionLevelStreamingDynamic.cpp b/Engine/Source/Runtime/Engine/Private/WorldPartition/WorldPartitionLevelStreamingDynamic.cpp
--- a/Engine/Source/Runtime/Engine/Private/WorldPartition/WorldPartitionLevelStreamingDynamic.cpp
+++ b/Engine/Source/Runtime/Engine/Private/WorldPartition/WorldPartitionLevelStreamingDynamic.cpp
@@ -19,6 +19,7 @@
 #include "WorldPartition/WorldPartitionStreamingPolicy.h"
 #include "WorldPartition/WorldPartitionLevelStreamingPolicy.h"
 #include "WorldPartition/WorldPartitionLevelHelper.h"
+#include "WorldPartition/WorldPartitionRuntimeHash.h"
 #endif
 
 #define LOCTEXT_NAMESPACE "World"
@@ -31,7 +32,7 @@ UWorldPartitionLevelStreamingDynamic::UWorldPartitionLevelStreamingDynamic(const
 	: Super(ObjectInitializer)
 #if WITH_EDITOR
 	, RuntimeLevel(nullptr)
-	, NumPendingActorRequests(0)
+	, bLoadRequestInProgress(false)
 #endif
 	, bIsActivated(false)
 	, bShouldBeAlwaysLoaded(false)
@@ -219,81 +220,39 @@ bool UWorldPartitionLevelStreamingDynamic::IssueLoadRequests()
 	check(ShouldBeVisible());
 	check(!HasLoadedLevel());
 	check(RuntimeLevel);
-	check(!NumPendingActorRequests);
+	check(!bLoadRequestInProgress);
 
 	TArray<FString> ActorPackageInstanceNames;
 	ActorPackageInstanceNames.Reserve(ChildPackages.Num());
 	UPackage* RuntimePackage = RuntimeLevel->GetPackage();
 
+	bLoadRequestInProgress = true;
 	FLinkerInstancingContext InstancingContext;
 	InstancingContext.AddMapping(OriginalLevelPackageName, RuntimePackage->GetFName());
-
-	FString ShortLevelPackageName = FPackageName::GetShortName(RuntimePackage->GetFName());
-	for (const FWorldPartitionRuntimeCellObjectMapping& PackageObjectMapping : ChildPackages)
-	{
-		FString ObjectPath = PackageObjectMapping.Package.ToString();
-		FString ActorPackageName = FPackageName::ObjectPathToPackageName(ObjectPath);
-		FString ActorPackageInstanceName = ActorPackageName + TEXT("_") + ShortLevelPackageName;
-		ActorPackageInstanceNames.Add(ActorPackageInstanceName);
-		
-		InstancingContext.AddMapping(FName(*ActorPackageName), FName(*ActorPackageInstanceName));
-	}
-
-	NumPendingActorRequests = ChildPackages.Num();
-	for (int32 ChildIndex=0; ChildIndex<ChildPackages.Num();  ChildIndex++)
-	{
-		const FWorldPartitionRuntimeCellObjectMapping& PackageObjectMapping = ChildPackages[ChildIndex];	
-		const FString ActorPackageName = FPackageName::ObjectPathToPackageName(PackageObjectMapping.Package.ToString());
-		FName ActorName = *FPaths::GetExtension(PackageObjectMapping.Path.ToString());
-
-		FLoadPackageAsyncDelegate CompletionCallback = FLoadPackageAsyncDelegate::CreateLambda([this, ActorName](const FName& PackageName, UPackage* Package, EAsyncLoadingResult::Type Result)
+	FWorldPartitionLevelHelper::FOnLoadActorsCompleted CompletionCallback = FWorldPartitionLevelHelper::FOnLoadActorsCompleted::CreateLambda([this](bool bSucceeded)
 		{
-			check(NumPendingActorRequests);
-			NumPendingActorRequests--;
-
-			if (Package)
-			{
-				AActor* Actor = FindObject<AActor>(Package, *ActorName.ToString());
-
-				check(Actor);
-				check(Actor->IsPackageExternal());
-				check(Actor->GetLevel() == RuntimeLevel);
-				
-				RuntimeLevel->Actors.Add(Actor);
-
-				UE_LOG(LogLevelStreaming, Verbose, TEXT(" ==> Loaded %s (remaining: %d)"), *Actor->GetFullName(), NumPendingActorRequests);
-			}
-			else if (Result == EAsyncLoadingResult::Canceled)
+			if (bSucceeded)
 			{
-				UE_LOG(LogLevelStreaming, Warning, TEXT("Async load cancelled for package '%s'"), *GetWorldAssetPackageName());
-				//@todo_ow: cumulate and process when NumPendingActorRequests == 0
+				check(bLoadRequestInProgress);
+				bLoadRequestInProgress = false;
+				FinalizeRuntimeLevel();
 			}
 			else
 			{
-				UE_LOG(LogLevelStreaming, Warning, TEXT("Failed to load package '%s'"), *GetWorldAssetPackageName());
-				//@todo_ow: cumulate and process when NumPendingActorRequests == 0
-			}
-
-			if (!NumPendingActorRequests)
-			{
-				FinalizeRuntimeLevel();
+				UE_LOG(LogLevelStreaming, Fatal, TEXT("UWorldPartitionLevelStreamingDynamic::IssueLoadRequests failed %s"), *GetWorldAssetPackageName());
 			}
 		});
 
-		FPackagePath PackagePath;
-		FPackagePath::TryFromMountedName(!ActorPackageName.IsEmpty() ? ActorPackageName : ActorPackageInstanceNames[ChildIndex], PackagePath);
-		FName PackageName(*ActorPackageInstanceNames[ChildIndex]);
-		LoadPackageAsync(PackagePath, PackageName, CompletionCallback, nullptr, PKG_PlayInEditor, RuntimePackage->PIEInstanceID, 0, &InstancingContext);
-	}
+	FWorldPartitionLevelHelper::LoadActors(RuntimeLevel, ChildPackages, PackageCache, CompletionCallback, /*bLoadForPie=*/true, /*bLoadAsync=*/true, &InstancingContext);
 
-	return !!NumPendingActorRequests;
+	return bLoadRequestInProgress;
 }
 
 void UWorldPartitionLevelStreamingDynamic::FinalizeRuntimeLevel()
 {
 	check(!HasLoadedLevel());
 	check(RuntimeLevel);
-	check(!NumPendingActorRequests);
+	check(!bLoadRequestInProgress);
 
 	// For RuntimeLevel's world NetGUID to be valid, make sure to flag bIsNameStableForNetworking so that IsNameStableForNetworking() returns true. (see FNetGUIDCache::SupportsObject)
 	RuntimeLevel->GetTypedOuter<UWorld>()->bIsNameStableForNetworking = true;
@@ -333,6 +292,8 @@ void UWorldPartitionLevelStreamingDynamic::OnCleanupLevel()
 {
 	if (RuntimeLevel)
 	{
+		PackageCache.UnloadPackages();
+
 		RuntimeLevel->OnCleanupLevel.Remove(OnCleanupLevelDelegateHandle);
 
 		// Clears RF_Standalone flag on objects in package (Metadata)
diff --git a/Engine/Source/Runtime/Engine/Private/WorldPartition/WorldPartitionPackageCache.cpp b/Engine/Source/Runtime/Engine/Private/WorldPartition/WorldPartitionPackageCache.cpp
new file mode 100644
--- /dev/null
+++ b/Engine/Source/Runtime/Engine/Private/WorldPartition/WorldPartitionPackageCache.cpp
@@ -0,0 +1,168 @@
+// Copyright Epic Games, Inc. All Rights Reserved.
+
+#include "WorldPartition/WorldPartitionPackageCache.h"
+
+#if WITH_EDITOR
+#include "UObject/UObjectHash.h"
+#include "Engine/World.h"
+#include "PackageTools.h"
+#include "Misc/PackagePath.h"
+
+FWorldPartitionPackageCache::FWorldPartitionPackageCache()
+{
+	
+}
+
+FWorldPartitionPackageCache::~FWorldPartitionPackageCache()
+{
+	check(!LoadingPackages.Num());
+	UnloadPackages();
+}
+
+void FWorldPartitionPackageCache::TrashPackage(UPackage* InPackage)
+{	
+	CachedPackages.Remove(InPackage->GetFName());
+	UnloadPackage(InPackage);
+
+	// Rename so it isn't found again
+	FName NewUniqueTrashName = MakeUniqueObjectName(nullptr, UPackage::StaticClass(), FName(*FString::Printf(TEXT("%s_Trashed"), *InPackage->GetName())));
+	InPackage->Rename(*NewUniqueTrashName.ToString(), nullptr, REN_DontCreateRedirectors | REN_ForceNoResetLoaders | REN_NonTransactional | REN_DoNotDirty);
+}
+
+void FWorldPartitionPackageCache::UnloadPackages()
+{
+	if (CachedPackages.Num())
+	{
+		for (const auto& Pair : CachedPackages)
+		{
+			if (UPackage* CachedPackage = Pair.Value.Get())
+			{
+				UnloadPackage(CachedPackage);
+			}
+		}
+				
+		CachedPackages.Empty();
+	}
+}
+
+void FWorldPartitionPackageCache::UnloadPackage(UPackage* InPackage)
+{
+	check(InPackage);
+	ForEachObjectWithPackage(InPackage, [](UObject* Object)
+		{
+			Object->ClearFlags(RF_Standalone);
+			return true;
+		}, false);
+
+	// World specific
+	if (UWorld* PackageWorld = UWorld::FindWorldInPackage(InPackage))
+	{
+		if (PackageWorld->PersistentLevel && PackageWorld->PersistentLevel->IsUsingExternalActors())
+		{
+			for (AActor* Actor : PackageWorld->PersistentLevel->Actors)
+			{
+				if (UPackage* ActorPackage = Actor ? Actor->GetExternalPackage() : nullptr)
+				{
+					ForEachObjectWithPackage(ActorPackage, [&](UObject* Object)
+						{
+							Object->ClearFlags(RF_Standalone);
+							return true;
+						}, false);
+				}
+			}
+		}
+	}
+}
+
+void FWorldPartitionPackageCache::LoadWorldPackageAsync(FName InPackageName, const TCHAR* InPackageToLoadFrom /*= nullptr*/, FLoadPackageAsyncDelegate InCompletionDelegate /*= FLoadPackageAsyncDelegate()*/, EPackageFlags InPackageFlags /*= PKG_None*/, int32 InPIEInstanceID /*= INDEX_NONE*/, int32 InPackagePriority /*= 0*/, const FLinkerInstancingContext* InInstancingContext /*=nullptr*/)
+{
+	LoadPackageAsyncInternal(InPackageName, InPackageToLoadFrom, InCompletionDelegate, InPackageFlags, InPIEInstanceID, InPackagePriority, InInstancingContext, /* bInWorldPackage=*/true);
+}
+
+void FWorldPartitionPackageCache::LoadPackageAsync(FName InPackageName, const TCHAR* InPackageToLoadFrom /*= nullptr*/, FLoadPackageAsyncDelegate InCompletionDelegate /*= FLoadPackageAsyncDelegate()*/, EPackageFlags InPackageFlags /*= PKG_None*/, int32 InPIEInstanceID /*= INDEX_NONE*/, int32 InPackagePriority /*= 0*/, const FLinkerInstancingContext* InInstancingContext /*=nullptr*/)
+{
+	LoadPackageAsyncInternal(InPackageName, InPackageToLoadFrom, InCompletionDelegate, InPackageFlags, InPIEInstanceID, InPackagePriority, InInstancingContext, /* bInWorldPackage=*/false);
+}
+
+void FWorldPartitionPackageCache::LoadPackageAsyncInternal(FName InPackageName, const TCHAR* InPackageToLoadFrom, FLoadPackageAsyncDelegate InCompletionDelegate, EPackageFlags InPackageFlags, int32 InPIEInstanceID, int32 InPackagePriority, const FLinkerInstancingContext* InInstancingContext, bool bInWorldPackage)
+{
+	if (UPackage* CachedPackage = FindPackage(InPackageName))
+	{
+		InCompletionDelegate.Execute(InPackageName, CachedPackage, EAsyncLoadingResult::Succeeded);
+		return;
+	}
+
+	// Find loading package
+	if (TArray<FLoadPackageAsyncDelegate>* CompletionDelegates = LoadingPackages.Find(InPackageName))
+	{
+		CompletionDelegates->Add(InCompletionDelegate);
+		return;
+	}
+
+	// Not found start loading
+	LoadingPackages.Add(InPackageName, { InCompletionDelegate });
+
+	FLoadPackageAsyncDelegate CompletionCallback = FLoadPackageAsyncDelegate::CreateLambda([this](const FName& LoadedPackageName, UPackage* LoadedPackage, EAsyncLoadingResult::Type Result)
+		{
+			if (Result == EAsyncLoadingResult::Succeeded)
+			{
+				CachedPackages.Add(LoadedPackageName, LoadedPackage);
+			}
+	
+			TArray<FLoadPackageAsyncDelegate> CompletionDelegates = LoadingPackages.FindAndRemoveChecked(LoadedPackageName);
+			for (FLoadPackageAsyncDelegate& CompletionDelegate : CompletionDelegates)
+			{
+				CompletionDelegate.Execute(LoadedPackageName, LoadedPackage, Result);
+			}
+		});
+
+	// This is to prevent the world to be initialized (when a World asset is added and its type is EWorldType::Inactive it gets initialized)
+	if (bInWorldPackage)
+	{
+		UWorld::WorldTypePreLoadMap.FindOrAdd(InPackageName) = EWorldType::Editor;
+	}
+	::LoadPackageAsync(FPackagePath::FromPackageNameChecked(InPackageToLoadFrom), InPackageName, CompletionCallback, nullptr, InPackageFlags, InPIEInstanceID, InPackagePriority, InInstancingContext);
+}
+
+UPackage* FWorldPartitionPackageCache::FindPackage(FName InPackageName)
+{
+	if (TWeakObjectPtr<UPackage>* CachedPackagePtr = CachedPackages.Find(InPackageName))
+	{
+		if (UPackage* CachedPackage = CachedPackagePtr->Get())
+		{
+			return CachedPackage;
+		}
+		CachedPackages.Remove(InPackageName);
+	}
+
+	// Might have been cached by other instance of a FWorldPartitionPackageCache. Happens when cooking WP Cells that each have their own FWorldPartitionPackageCache.
+	if(UPackage* Package = ::FindPackage(nullptr, *InPackageName.ToString()))
+	{
+		CachedPackages.Add(InPackageName, Package);
+		return Package;
+	}
+
+	return nullptr;
+}
+
+UPackage* FWorldPartitionPackageCache::DuplicateWorldPackage(UPackage* InPackage, FName InDuplicatePackageName)
+{
+	check(!CachedPackages.Contains(InDuplicatePackageName));
+
+	UPackage* DuplicatedPackage = nullptr;
+	if (UWorld* PackageWorld = UWorld::FindWorldInPackage(InPackage))
+	{
+		DuplicatedPackage = CreatePackage(*InDuplicatePackageName.ToString());
+		FObjectDuplicationParameters DuplicationParameters(PackageWorld, DuplicatedPackage);
+		DuplicationParameters.bAssignExternalPackages = false;
+		DuplicationParameters.DuplicateMode = EDuplicateMode::World;
+
+		UWorld* DuplicatedWorld = Cast<UWorld>(StaticDuplicateObjectEx(DuplicationParameters));
+		check(DuplicatedWorld);
+		CachedPackages.Add(InDuplicatePackageName, DuplicatedPackage);
+	}
+	
+	return DuplicatedPackage;
+}
+
+#endif
\ No newline at end of file
diff --git a/Engine/Source/Runtime/Engine/Private/WorldPartition/WorldPartitionRuntimeLevelStreamingCell.cpp b/Engine/Source/Runtime/Engine/Private/WorldPartition/WorldPartitionRuntimeLevelStreamingCell.cpp
--- a/Engine/Source/Runtime/Engine/Private/WorldPartition/WorldPartitionRuntimeLevelStreamingCell.cpp
+++ b/Engine/Source/Runtime/Engine/Private/WorldPartition/WorldPartitionRuntimeLevelStreamingCell.cpp
@@ -11,6 +11,7 @@
 #if WITH_EDITOR
 #include "WorldPartition/WorldPartitionLevelStreamingPolicy.h"
 #include "WorldPartition/WorldPartitionLevelHelper.h"
+#include "WorldPartition/WorldPartitionPackageCache.h"
 #include "Engine/LevelStreamingAlwaysLoaded.h"
 #endif
 
@@ -49,13 +50,11 @@ void UWorldPartitionRuntimeLevelStreamingCell::SetIsAlwaysLoaded(bool bInIsAlway
 
 #if WITH_EDITOR
 
-void UWorldPartitionRuntimeLevelStreamingCell::AddActorToCell(const FGuid& ActorGuid, FName Package, FName Path)
+void UWorldPartitionRuntimeLevelStreamingCell::AddActorToCell(const FGuid& InActorGuid, uint32 InContainerID, const FTransform& InContainerTransform, const UActorDescContainer* InContainer)
 {
-	Packages.Emplace(Package, Path);
-	if (!GetOuterUWorldPartition()->GetActorDesc(ActorGuid)->GetActorIsEditorOnly())
-	{
-		ActorsToLoadForCook.Add(ActorGuid);
-	}
+	const FWorldPartitionActorDesc* ActorDesc = InContainer->GetActorDesc(InActorGuid);
+	check(!ActorDesc->GetActorIsEditorOnly());
+	Packages.Emplace(ActorDesc->GetActorPackage(), ActorDesc->GetActorPath(), InContainerID, InContainerTransform, InContainer->GetContainerPackage());
 }
 
 ULevelStreaming* UWorldPartitionRuntimeLevelStreamingCell::CreateLevelStreaming(const FString& InPackageName) const
@@ -98,18 +97,10 @@ ULevelStreaming* UWorldPartitionRuntimeLevelStreamingCell::CreateLevelStreaming(
 
 void UWorldPartitionRuntimeLevelStreamingCell::LoadActorsForCook()
 {
-	// Load cell Actors
-	UWorldPartition* WorldPartition = GetOuterUWorldPartition();
-	for (const FGuid& ActorGuid : ActorsToLoadForCook)
-	{
-		const FWorldPartitionActorDesc* ActorDesc = WorldPartition->GetActorDesc(ActorGuid);
-		AActor* Actor = ActorDesc->GetActor();
-		if (!Actor)
-		{
-			Actor = ActorDesc->Load();
-			ensure(Actor);
-		}
-	}
+	FWorldPartitionPackageCache PackageCache;
+	FWorldPartitionLevelHelper::FOnLoadActorsCompleted CompletionCallback = FWorldPartitionLevelHelper::FOnLoadActorsCompleted::CreateLambda([](bool) {});
+	bool bResult = FWorldPartitionLevelHelper::LoadActors(nullptr, Packages, PackageCache, CompletionCallback, /*bLoadForPIE=*/false, /*bLoadAsync=*/false);
+	check(bResult);
 }
 
 void UWorldPartitionRuntimeLevelStreamingCell::MoveAlwaysLoadedContentToPersistentLevel()
@@ -138,10 +129,10 @@ bool UWorldPartitionRuntimeLevelStreamingCell::PopulateGeneratedPackageForCook(U
 		UWorld* OuterWorld = WorldPartition->GetTypedOuter<UWorld>();
 		ULevelStreaming* NewLevelStreaming = CreateLevelStreaming(InPackageCookName);
 		check(NewLevelStreaming)
-		
+
 		// Load cell Actors
 		LoadActorsForCook();
-
+		
 		LevelStreaming = Cast<UWorldPartitionLevelStreamingDynamic>(NewLevelStreaming);
 		ULevel* NewLevel = FWorldPartitionLevelHelper::CreateEmptyLevelForRuntimeCell(OuterWorld, NewLevelStreaming->GetWorldAsset().ToString(), InPackage);
 		check(NewLevel->GetPackage() == InPackage);
diff --git a/Engine/Source/Runtime/Engine/Private/WorldPartition/WorldPartitionRuntimeSpatialHash.cpp b/Engine/Source/Runtime/Engine/Private/WorldPartition/WorldPartitionRuntimeSpatialHash.cpp
--- a/Engine/Source/Runtime/Engine/Private/WorldPartition/WorldPartitionRuntimeSpatialHash.cpp
+++ b/Engine/Source/Runtime/Engine/Private/WorldPartition/WorldPartitionRuntimeSpatialHash.cpp
@@ -456,21 +456,22 @@ bool UWorldPartitionRuntimeSpatialHash::GenerateStreaming(EWorldPartitionStreami
 	}
 
 	// Create actor clusters
-	TArray<FActorCluster> ActorClusters = CreateActorClusters(WorldPartition);
+	FActorClusterContext Context(WorldPartition);
 
-	TArray<TArray<FActorCluster>> GridActors;
+	TArray<TArray<const FActorClusterInstance*>> GridActors;
 	GridActors.InsertDefaulted(0, AllGrids.Num());
 
-	for (FActorCluster& ActorCluster : ActorClusters)
+	for (const FActorClusterInstance& ClusterInstance : Context.GetClusterInstances())
 	{
-		int32* FoundIndex = GridsMapping.Find(ActorCluster.RuntimeGrid);
+		check(ClusterInstance.Cluster);
+		int32* FoundIndex = GridsMapping.Find(ClusterInstance.Cluster->RuntimeGrid);
 		if (!FoundIndex)
 		{
-			UE_LOG(LogWorldPartitionRuntimeSpatialHash, Error, TEXT("Invalid partition grid '%s' referenced by actor cluster"), *ActorCluster.RuntimeGrid.ToString());
+			UE_LOG(LogWorldPartitionRuntimeSpatialHash, Error, TEXT("Invalid partition grid '%s' referenced by actor cluster"), *ClusterInstance.Cluster->RuntimeGrid.ToString());
 		}
 
 		int32 GridIndex = FoundIndex ? *FoundIndex : 0;
-		GridActors[GridIndex].Add(MoveTemp(ActorCluster));
+		GridActors[GridIndex].Add(&ClusterInstance);
 	}
 	
 	const FBox WorldBounds = WorldPartition->GetWorldBounds();
@@ -535,7 +536,7 @@ bool UWorldPartitionRuntimeSpatialHash::CreateStreamingGrid(const FSpatialHashRu
 	// Move actors into the final streaming grids
 	CurrentStreamingGrid.GridLevels.Reserve(PartionedActors.Levels.Num());
 
-	TArray<FGuid> FilteredActors;
+	TArray<FActorInstance> FilteredActors;
 	int32 Level = INDEX_NONE;
 	for (const FSquare2DGridHelper::FGridLevel& TempLevel : PartionedActors.Levels)
 	{
@@ -559,13 +560,14 @@ bool UWorldPartitionRuntimeSpatialHash::CreateStreamingGrid(const FSpatialHashRu
 				FilteredActors.Reset(GridCellDataChunk.GetActors().Num());
 				if (GridCellDataChunk.GetActors().Num())
 				{
-					Algo::TransformIf(GridCellDataChunk.GetActors(), FilteredActors, [&WorldPartition](const FGuid& ActorGuid)
+					Algo::TransformIf(GridCellDataChunk.GetActors(), FilteredActors, [](const FActorInstance& ActorInstance)
 					{
-						const FWorldPartitionActorDesc* ActorDesc = WorldPartition->GetActorDesc(ActorGuid);
-						const bool bShouldStripActorFromStreaming = ActorDesc->GetActorIsEditorOnly();
-						UE_CLOG(bShouldStripActorFromStreaming, LogWorldPartitionRuntimeSpatialHash, Verbose, TEXT("Stripping Actor %s (%s) from streaming grid"), *(ActorDesc->GetActorPath().ToString()), *ActorGuid.ToString(EGuidFormats::UniqueObjectGuid));
+						const FWorldPartitionActorDesc* ActorDesc = ActorInstance.GetActorDesc();
+						const bool bShouldStripActorFromStreaming = ActorInstance.ShouldStripFromStreaming();
+						UE_CLOG(bShouldStripActorFromStreaming, LogWorldPartitionRuntimeSpatialHash, Verbose, TEXT("Stripping Actor %s (%s) from streaming grid (Container %08x)"), 
+							*(ActorDesc->GetActorPath().ToString()), *ActorInstance.Actor.ToString(EGuidFormats::UniqueObjectGuid), ActorInstance.ContainerInstance->ID);
 						return !bShouldStripActorFromStreaming;
-					}, [](const FGuid& ActorGuid) { return ActorGuid; });
+					}, [](const FActorInstance& ActorInstance) { return ActorInstance; });
 				}
 
 				// Cell cannot be treated as always loaded if it has data layers
@@ -578,9 +580,10 @@ bool UWorldPartitionRuntimeSpatialHash::CreateStreamingGrid(const FSpatialHashRu
 				// will unload actors that were not already loaded in the non PIE world.
 				if (bIsCellAlwaysLoaded && bIsMainWorldPartition && (Mode == EWorldPartitionStreamingMode::PIE))
 				{
-					for (const FGuid& ActorGuid : FilteredActors)
+					for (const FActorInstance& ActorInstance : FilteredActors)
 					{
-						AlwaysLoadedActorsForPIE.Emplace(WorldPartition, ActorGuid);
+						check(ActorInstance.ContainerInstance->Container == WorldPartition);
+						AlwaysLoadedActorsForPIE.Emplace(WorldPartition, ActorInstance.Actor);
 					}
 					continue;
 				}
@@ -601,21 +604,21 @@ bool UWorldPartitionRuntimeSpatialHash::CreateStreamingGrid(const FSpatialHashRu
 				verify(TempLevel.GetCellBounds(FIntVector2(CellCoordX, CellCoordY), Bounds));
 				StreamingCell->Position = FVector(Bounds.GetCenter(), 0.f);
 
-				UE_LOG(LogWorldPartitionRuntimeSpatialHash, Verbose, TEXT("Cell%s %s Actors = %d"), bIsCellAlwaysLoaded ? TEXT(" (AlwaysLoaded)") : TEXT(""), *StreamingCell->GetName(), FilteredActors.Num());
+				UE_LOG(LogWorldPartitionRuntimeSpatialHash, Verbose, TEXT("Cell%s %s Actors = %d Bounds (%s)"), bIsCellAlwaysLoaded ? TEXT(" (AlwaysLoaded)") : TEXT(""), *StreamingCell->GetName(), FilteredActors.Num(), *Bounds.ToString());
 
 				// Keep track of all AWorldPartitionHLOD actors referenced by this cell
 				TSet<FGuid> ReferencedHLODActors;
 
-				for (const FGuid& ActorGuid : FilteredActors)
+				for (const FActorInstance& ActorInstance : FilteredActors)
 				{
-					const FWorldPartitionActorDesc* ActorDesc = WorldPartition->GetActorDesc(ActorGuid);
-					FGuid ParentHLOD = CachedHLODParents.FindRef(ActorGuid);
+					const FWorldPartitionActorDesc* ActorDesc = ActorInstance.GetActorDesc();
+					FGuid ParentHLOD = CachedHLODParents.FindRef(ActorInstance.Actor);
 					if (ParentHLOD.IsValid())
 					{
 						ReferencedHLODActors.Add(ParentHLOD);
 					}
-					StreamingCell->AddActorToCell(ActorGuid, ActorDesc->GetActorPackage(), ActorDesc->GetActorPath());
-					UE_LOG(LogWorldPartitionRuntimeSpatialHash, Verbose, TEXT("  Actor : %s (%s) Origin(%s)"), *(ActorDesc->GetActorPath().ToString()), *ActorGuid.ToString(EGuidFormats::UniqueObjectGuid), *FVector2D(ActorDesc->GetOrigin()).ToString());
+					StreamingCell->AddActorToCell(ActorInstance.Actor, ActorInstance.ContainerInstance->ID, ActorInstance.ContainerInstance->Transform, ActorInstance.ContainerInstance->Container);
+					UE_LOG(LogWorldPartitionRuntimeSpatialHash, Verbose, TEXT("  Actor : %s (%s) (Container %08x) Origin(%s)"), *(ActorDesc->GetActorPath().ToString()), *ActorDesc->GetGuid().ToString(EGuidFormats::UniqueObjectGuid), ActorInstance.ContainerInstance->ID, *FVector2D(ActorInstance.GetOrigin()).ToString());
 				}
 
 				if (ReferencedHLODActors.Num() > 0)
diff --git a/Engine/Source/Runtime/Engine/Public/LevelInstance/LevelInstanceActor.h b/Engine/Source/Runtime/Engine/Public/LevelInstance/LevelInstanceActor.h
--- a/Engine/Source/Runtime/Engine/Public/LevelInstance/LevelInstanceActor.h
+++ b/Engine/Source/Runtime/Engine/Public/LevelInstance/LevelInstanceActor.h
@@ -7,8 +7,18 @@
 #include "GameFramework/Actor.h"
 #include "Containers/Array.h"
 #include "LevelInstance/LevelInstanceTypes.h"
+#include "WorldPartition/WorldPartitionActorDesc.h"
 #include "LevelInstanceActor.generated.h"
 
+UENUM()
+enum class ELevelInstanceRuntimeBehavior : uint8
+{
+	None UMETA(Hidden),
+	Embedded UMETA(ToolTip="Level Instance is discarded at runtime and all its actors are part of the same runtime World Partition cell"),
+	Partitioned UMETA(Hidden), 
+	LevelStreaming UMETA(ToolTip="Level Instance is streamed like a regular level"),
+};
+
 UCLASS()
 class ENGINE_API ALevelInstance : public AActor
 {
@@ -18,8 +28,14 @@ protected:
 	/** Level LevelInstance */
 	UPROPERTY(EditAnywhere, Category = LevelInstance, meta = (NoCreate))
 	TSoftObjectPtr<UWorld> WorldAsset;
+
 public:
 
+#if WITH_EDITORONLY_DATA
+	UPROPERTY(EditAnywhere, Category = Runtime)
+	ELevelInstanceRuntimeBehavior DesiredRuntimeBehavior;
+#endif
+
 	void LoadLevelInstance();
 	void UnloadLevelInstance();
 	virtual bool SupportsLoading() const;
@@ -57,6 +73,8 @@ public:
 	virtual FBox GetComponentsBoundingBox(bool bNonColliding = false, bool bIncludeFromChildActors = false) const override;
 	virtual void GetActorLocationBounds(bool bOnlyCollidingComponents, FVector& Origin, FVector& BoxExtent, bool bIncludeFromChildActors = false) const override;
 	virtual bool IsLockLocation() const override;
+	virtual ELevelInstanceRuntimeBehavior GetDesiredRuntimeBehavior() const { return DesiredRuntimeBehavior; }
+	virtual bool IsHLODRelevant() const override { return true; }
 
 	bool CanEdit(FText* OutReason = nullptr) const;
 	bool CanCommit(FText* OutReason = nullptr) const;
diff --git a/Engine/Source/Runtime/Engine/Public/WorldPartition/ActorDescContainer.h b/Engine/Source/Runtime/Engine/Public/WorldPartition/ActorDescContainer.h
--- a/Engine/Source/Runtime/Engine/Public/WorldPartition/ActorDescContainer.h
+++ b/Engine/Source/Runtime/Engine/Public/WorldPartition/ActorDescContainer.h
@@ -17,9 +17,11 @@ class ENGINE_API UActorDescContainer : public UObject
 	friend struct FWorldPartitionHandleUtils;
 
 public:
-	void Initialize(FName InPackageName, bool bRegisterDelegates);
+	void Initialize(UWorld* World, FName InPackageName, bool bRegisterDelegates);
 	virtual void Uninitialize();
 	
+	virtual UWorld* GetWorld() const override;
+
 #if WITH_EDITOR
 	// Asset registry events
 	virtual void OnAssetAdded(const FAssetData& InAssetData);
@@ -31,8 +33,13 @@ public:
 
 	FWorldPartitionActorDesc& GetActorDescChecked(const FGuid& Guid);
 	const FWorldPartitionActorDesc& GetActorDescChecked(const FGuid& Guid) const;
+
+	int32 GetActorDescCount() const { return Actors.Num(); }
+	FName GetContainerPackage() const { return ContainerPackageName; }
 #endif
 
+	UPROPERTY(Transient)
+	TObjectPtr<UWorld> World;
 protected:
 	//~ Begin UObject Interface
 	virtual void BeginDestroy() override;
diff --git a/Engine/Source/Runtime/Engine/Public/WorldPartition/LevelInstance/LevelInstanceActorDesc.h b/Engine/Source/Runtime/Engine/Public/WorldPartition/LevelInstance/LevelInstanceActorDesc.h
--- a/Engine/Source/Runtime/Engine/Public/WorldPartition/LevelInstance/LevelInstanceActorDesc.h
+++ b/Engine/Source/Runtime/Engine/Public/WorldPartition/LevelInstance/LevelInstanceActorDesc.h
@@ -4,28 +4,38 @@
 
 #include "CoreMinimal.h"
 #include "UObject/ObjectMacros.h"
+#include "UObject/GCObject.h"
 #include "Containers/Map.h"
 #include "WorldPartition/WorldPartitionActorDesc.h"
 
 class ULevelInstanceSubsystem;
-class UWorldPartition;
+class UActorDescContainer;
+enum class ELevelInstanceRuntimeBehavior : uint8;
 
 /**
  * ActorDesc for Actors that are part of a LevelInstanceActor Level.
  */
-class ENGINE_API FLevelInstanceActorDesc : public FWorldPartitionActorDesc
+class ENGINE_API FLevelInstanceActorDesc : public FWorldPartitionActorDesc, public FGCObject
 {
 #if WITH_EDITOR
-	friend class FLevelInstanceActorDescFactory;
+	friend class ALevelInstance;
 
 public:
 	inline FName GetLevelPackage() const { return LevelPackage; }
 
+	virtual bool GetContainerInstance(const UActorDescContainer*& OutLevelContainer, FTransform& OutLevelTransform, EContainerClusterMode& OutClusterMode) const override;
 protected:
+	FLevelInstanceActorDesc();
 	virtual void Init(const AActor* InActor) override;
+	virtual void Init(UActorDescContainer* InContainer, const FWorldPartitionActorDescInitData& DescData) override;
 	virtual void Serialize(FArchive& Ar) override;
+	virtual void AddReferencedObjects(FReferenceCollector& Collector);
+	virtual void RegisterContainer();
 
 	FName LevelPackage;
 	FTransform LevelInstanceTransform;
+	ELevelInstanceRuntimeBehavior DesiredRuntimeBehavior;
+
+	const UActorDescContainer* LevelInstanceContainer = nullptr;
 #endif
 };
diff --git a/Engine/Source/Runtime/Engine/Public/WorldPartition/WorldPartition.h b/Engine/Source/Runtime/Engine/Public/WorldPartition/WorldPartition.h
--- a/Engine/Source/Runtime/Engine/Public/WorldPartition/WorldPartition.h
+++ b/Engine/Source/Runtime/Engine/Public/WorldPartition/WorldPartition.h
@@ -143,6 +143,8 @@ public:
 	void GenerateHLOD(ISourceControlHelper* SourceControlHelper, bool bCreateActorsOnly);
 	void GenerateNavigationData();
 
+	const UActorDescContainer* RegisterActorDescContainer(FName PackageName);
+
 	// Debugging Methods
 	void DrawRuntimeHashPreview();
 	void DumpActorDescs(const FString& Path);
@@ -168,13 +170,9 @@ public:
 	void DrawRuntimeHash2D(class UCanvas* Canvas, const FVector2D& PartitionCanvasOffset, const FVector2D& PartitionCanvasSize);
 	void DrawRuntimeHash3D();
 
-	virtual UWorld* GetWorld() const override;
-
 	void RegisterStreamingSourceProvider(IWorldPartitionStreamingSourceProvider* StreamingSource);
 	bool UnregisterStreamingSourceProvider(IWorldPartitionStreamingSourceProvider* StreamingSource);
 
-	UPROPERTY(Transient)
-	TObjectPtr<UWorld> World;
 
 #if WITH_EDITORONLY_DATA
 	UPROPERTY(DuplicateTransient)
@@ -197,6 +195,8 @@ public:
 	TObjectPtr<class UHLODLayer> DefaultHLODLayer;
 
 	TArray<FWorldPartitionReference> LoadedSubobjects;
+
+	TMap<FName, TWeakObjectPtr<UActorDescContainer>> ActorDescContainers;
 #endif
 
 private:
diff --git a/Engine/Source/Runtime/Engine/Public/WorldPartition/WorldPartitionActorCluster.h b/Engine/Source/Runtime/Engine/Public/WorldPartition/WorldPartitionActorCluster.h
new file mode 100644
--- /dev/null
+++ b/Engine/Source/Runtime/Engine/Public/WorldPartition/WorldPartitionActorCluster.h
@@ -0,0 +1,131 @@
+// Copyright Epic Games, Inc. All Rights Reserved.
+#pragma once
+
+#if WITH_EDITOR
+
+#include "WorldPartition/DataLayer/DataLayersID.h"
+#include "Misc/HashBuilder.h"
+
+class UActorDescContainer;
+class UWorldPartition;
+class FWorldPartitionActorDesc;
+enum class EContainerClusterMode : uint8;
+
+DECLARE_LOG_CATEGORY_EXTERN(LogWorldPartitionActorCluster, Log, All);
+
+/**
+ * List of actors bound together based on clustering rules. (mainly object references)
+ */
+struct FActorCluster
+{
+	TSet<FGuid>					Actors;
+	EActorGridPlacement			GridPlacement;
+	FName						RuntimeGrid;
+	FBox						Bounds;
+	TSet<const UDataLayer*>		DataLayers;
+	FDataLayersID				DataLayersID;
+
+	FActorCluster(const FWorldPartitionActorDesc* InActorDesc, EActorGridPlacement InGridPlacement, const UActorDescContainer* InContainer);
+	void Add(const FActorCluster& InActorCluster);
+
+};
+
+struct FActorClusterInstance;
+
+/**
+ * Instance of a container (level) with specific instance properties (transform, id, datalayers) 
+ */
+struct FActorContainerInstance
+{
+	FActorContainerInstance(const UActorDescContainer* InContainer);
+	FActorContainerInstance(uint32 InID, const FTransform& InTransform, const FBox& InBounds, const TSet<FName>& InDataLayers, EContainerClusterMode InClusterMode, const UActorDescContainer* InContainer, TSet<FGuid> InChildContainers);
+	
+	uint32						ID;
+	FTransform					Transform;
+	FBox						Bounds;
+	EContainerClusterMode		ClusterMode;
+	const UActorDescContainer*	Container;
+	TSet<FGuid>					ChildContainers;
+	TSet<const UDataLayer*>		DataLayers;
+};
+
+/**
+ *  Represents an instanced actor cluster. It is a combination of the clustering of actors (link between actors agnostic of instance) and a container instance (transform, instance data layers).
+ */
+struct FActorClusterInstance
+{
+	FActorClusterInstance(const FActorCluster* InCluster, const FActorContainerInstance* InContainerInstance);
+
+	FBox							Bounds;
+	const FActorCluster*			Cluster;
+	const FActorContainerInstance*	ContainerInstance;
+	TArray<const UDataLayer*>		DataLayers;
+};
+
+/**
+ * Represents one actor and his owning container instance 
+ */
+struct FActorInstance
+{
+	FGuid Actor;
+	const FActorContainerInstance* ContainerInstance;
+
+	FActorInstance();
+	FActorInstance(const FGuid& InActor, const FActorContainerInstance* InContainerInstance);
+
+	FVector GetOrigin() const;
+
+	friend uint32 GetTypeHash(const FActorInstance& InActorInstance)
+	{
+		FHashBuilder HashBuilder;
+		HashBuilder << InActorInstance.Actor;
+		HashBuilder << InActorInstance.ContainerInstance->ID;
+		return HashBuilder.GetHash();
+	}
+
+	friend bool operator==(const FActorInstance& A, const FActorInstance& B)
+	{
+		return A.Actor == B.Actor && A.ContainerInstance == B.ContainerInstance;
+	}
+
+	bool ShouldStripFromStreaming() const;
+
+	const FWorldPartitionActorDesc* GetActorDesc() const;
+};
+
+/**
+ * Class used to generate the actor clustering. 
+ */
+class FActorClusterContext
+{
+public:
+	typedef TFunctionRef<bool(const FWorldPartitionActorDesc&)> FFilterPredicate;
+
+	/**
+	 * Create the actor clusters from the root World Partition. Optionally filtering some actors and including child Containers. 
+	 */
+	FActorClusterContext(const UWorldPartition* InWorldPartition, TOptional<FFilterPredicate> InFilterPredicate = TOptional<FFilterPredicate>(), bool bInIncludeChildContainers = true);
+				
+	/**
+	 * Returns the list of cluster instances of this context. 
+	 */
+	const TArray<FActorClusterInstance>& GetClusterInstances() const { return ClusterInstances; }
+private:
+	void CreateActorClusters();
+	void CreateContainerInstanceRecursive(uint32 ID, const FTransform& Transform, EContainerClusterMode ClusterMode, const UActorDescContainer* ActorDescContainer, const TSet<FName>& DataLayers, FBox* ParentBounds = nullptr);
+	const TArray<FActorCluster>& CreateActorClustersImpl(const FActorContainerInstance& ContainerInstance);
+	
+	// Init data
+	const UWorldPartition* WorldPartition;
+	TOptional<FFilterPredicate> FilterPredicate;
+	bool bIncludeChildContainers;
+
+	// Generated data
+	TMap<const UActorDescContainer*, TArray<FActorCluster>> Clusters;
+	TArray<FActorContainerInstance> ContainerInstances;
+	TArray<FActorClusterInstance> ClusterInstances;
+	
+	int32 InstanceCountHint;
+};
+
+#endif // #if WITH_EDITOR
diff --git a/Engine/Source/Runtime/Engine/Public/WorldPartition/WorldPartitionActorDesc.h b/Engine/Source/Runtime/Engine/Public/WorldPartition/WorldPartitionActorDesc.h
--- a/Engine/Source/Runtime/Engine/Public/WorldPartition/WorldPartitionActorDesc.h
+++ b/Engine/Source/Runtime/Engine/Public/WorldPartition/WorldPartitionActorDesc.h
@@ -20,6 +20,12 @@ struct FWorldPartitionActorDescInitData
 };
 
 class UActorDescContainer;
+
+enum class EContainerClusterMode : uint8
+{
+	Partitioned, // Per Actor Partitioning
+	Embedded // Per Container Partitioning: Every actor of the container are going to be clustered together
+};
 #endif
 
 #if WITH_DEV_AUTOMATION_TESTS
@@ -66,6 +72,8 @@ public:
 	inline FName GetActorLabel() const { return ActorLabel; }
 	FBox GetBounds() const;
 
+	virtual bool GetContainerInstance(const UActorDescContainer*& OutLevelContainer, FTransform& OutLevelTransform, EContainerClusterMode& OutClusterMode) const { return false; }
+
 protected:
 	inline uint32 IncSoftRefCount() const
 	{
@@ -116,7 +124,7 @@ public:
 	void UnregisterActor();
 
 	virtual void Init(const AActor* InActor);
-	void Init(UActorDescContainer* InContainer, const FWorldPartitionActorDescInitData& DescData);
+	virtual void Init(UActorDescContainer* InContainer, const FWorldPartitionActorDescInitData& DescData);
 
 	void SerializeTo(TArray<uint8>& OutData);
 
diff --git a/Engine/Source/Runtime/Engine/Public/WorldPartition/WorldPartitionLevelHelper.h b/Engine/Source/Runtime/Engine/Public/WorldPartition/WorldPartitionLevelHelper.h
--- a/Engine/Source/Runtime/Engine/Public/WorldPartition/WorldPartitionLevelHelper.h
+++ b/Engine/Source/Runtime/Engine/Public/WorldPartition/WorldPartitionLevelHelper.h
@@ -14,12 +14,19 @@
 #include "Engine/World.h"
 #include "WorldPartition/WorldPartitionRuntimeCell.h"
 
+class FWorldPartitionPackageCache;
+
 class FWorldPartitionLevelHelper
 {
 public:
 	static ULevel* CreateEmptyLevelForRuntimeCell(const UWorld* InWorld, const FString& InWorldAssetName, UPackage* DestPackage = nullptr);
 	static void MoveExternalActorsToLevel(const TArray<FWorldPartitionRuntimeCellObjectMapping>& InChildPackages, ULevel* InLevel);
+	
+	DECLARE_DELEGATE_OneParam(FOnLoadActorsCompleted, bool);
+
+	static bool LoadActors(ULevel* InDestLevel, TArrayView<FWorldPartitionRuntimeCellObjectMapping> InActorPackages, FWorldPartitionPackageCache& InPackageCache, FOnLoadActorsCompleted InCompletionCallback, bool bLoadForPie, bool bLoadAsync = true, FLinkerInstancingContext* InOutInstancingContext = nullptr);
 private:
 	static UWorld::InitializationValues GetWorldInitializationValues();
 };
+
 #endif
\ No newline at end of file
diff --git a/Engine/Source/Runtime/Engine/Public/WorldPartition/WorldPartitionLevelStreamingDynamic.h b/Engine/Source/Runtime/Engine/Public/WorldPartition/WorldPartitionLevelStreamingDynamic.h
--- a/Engine/Source/Runtime/Engine/Public/WorldPartition/WorldPartitionLevelStreamingDynamic.h
+++ b/Engine/Source/Runtime/Engine/Public/WorldPartition/WorldPartitionLevelStreamingDynamic.h
@@ -16,6 +16,7 @@
 #include "WorldPartition/WorldPartition.h"
 #if WITH_EDITOR
 #include "WorldPartition/WorldPartitionRuntimeLevelStreamingCell.h"
+#include "WorldPartition/WorldPartitionPackageCache.h"
 #endif
 #include "WorldPartitionLevelStreamingDynamic.generated.h"
 
@@ -52,7 +53,8 @@ private:
 #endif
 
 	FDelegateHandle OnCleanupLevelDelegateHandle;
-	uint32 NumPendingActorRequests;
+	bool bLoadRequestInProgress;
+	FWorldPartitionPackageCache PackageCache;
 #endif
 
 	bool bIsActivated;
diff --git a/Engine/Source/Runtime/Engine/Public/WorldPartition/WorldPartitionPackageCache.h b/Engine/Source/Runtime/Engine/Public/WorldPartition/WorldPartitionPackageCache.h
new file mode 100644
--- /dev/null
+++ b/Engine/Source/Runtime/Engine/Public/WorldPartition/WorldPartitionPackageCache.h
@@ -0,0 +1,36 @@
+// Copyright Epic Games, Inc. All Rights Reserved.
+
+#pragma once
+
+#if WITH_EDITOR
+
+#include "CoreMinimal.h"
+#include "UObject/ObjectMacros.h"
+#include "UObject/UObjectGlobals.h"
+#include "UObject/WeakObjectPtr.h"
+#include "Misc/Guid.h"
+
+class FLinkerInstancingContext;
+
+class FWorldPartitionPackageCache
+{
+public:
+	FWorldPartitionPackageCache();
+	~FWorldPartitionPackageCache();
+
+	void LoadWorldPackageAsync(FName InPackageName, const TCHAR* InPackageToLoadFrom = nullptr, FLoadPackageAsyncDelegate InCompletionDelegate = FLoadPackageAsyncDelegate(), EPackageFlags InPackageFlags = PKG_None, int32 InPIEInstanceID = INDEX_NONE, int32 InPackagePriority = 0, const FLinkerInstancingContext* InInstancingContext = nullptr);
+	void LoadPackageAsync(FName InPackageName, const TCHAR* InPackageToLoadFrom = nullptr, FLoadPackageAsyncDelegate InCompletionDelegate = FLoadPackageAsyncDelegate(), EPackageFlags InPackageFlags = PKG_None, int32 InPIEInstanceID = INDEX_NONE, int32 InPackagePriority = 0, const FLinkerInstancingContext* InInstancingContext = nullptr);
+	UPackage* FindPackage(FName InPackageName);
+	UPackage* DuplicateWorldPackage(UPackage* InPackage, FName InPackageName);
+	void UnloadPackages();
+	void TrashPackage(UPackage* InPackage);
+	
+private:
+	void LoadPackageAsyncInternal(FName InPackageName, const TCHAR* InPackageToLoadFrom, FLoadPackageAsyncDelegate InCompletionDelegate, EPackageFlags InPackageFlags, int32 InPIEInstanceID, int32 InPackagePriority, const FLinkerInstancingContext* InInstancingContext, bool bInWorldPackage);
+	void UnloadPackage(UPackage* InPackage);
+
+	TMap<FName, TWeakObjectPtr<UPackage>> CachedPackages;
+	TMap<FName, TArray<FLoadPackageAsyncDelegate>> LoadingPackages;
+};
+
+#endif
\ No newline at end of file
diff --git a/Engine/Source/Runtime/Engine/Public/WorldPartition/WorldPartitionRuntimeCell.h b/Engine/Source/Runtime/Engine/Public/WorldPartition/WorldPartitionRuntimeCell.h
--- a/Engine/Source/Runtime/Engine/Public/WorldPartition/WorldPartitionRuntimeCell.h
+++ b/Engine/Source/Runtime/Engine/Public/WorldPartition/WorldPartitionRuntimeCell.h
@@ -12,15 +12,28 @@ struct FWorldPartitionRuntimeCellObjectMapping
 	GENERATED_USTRUCT_BODY()
 
 	FWorldPartitionRuntimeCellObjectMapping()
+#if WITH_EDITORONLY_DATA
 		: Package(NAME_None)
 		, Path(NAME_None)
+		, ContainerID(0)
+		, ContainerTransform(FTransform::Identity)
+		, ContainerPackage(NAME_None)
+		, LoadedPath(NAME_None)
+#endif
 	{}
 
-	FWorldPartitionRuntimeCellObjectMapping(FName InPackage, FName InPath)
+	FWorldPartitionRuntimeCellObjectMapping(FName InPackage, FName InPath, uint32 InContainerID, const FTransform& InContainerTransform, FName InContainerPackage)
+#if WITH_EDITORONLY_DATA
 		: Package(InPackage)
 		, Path(InPath)
+		, ContainerID(InContainerID)
+		, ContainerTransform(InContainerTransform)
+		, ContainerPackage(InContainerPackage)
+		, LoadedPath(InPath)
+#endif
 	{}
 
+#if WITH_EDITORONLY_DATA
 	/** 
 	 * The name of the package to load to resolve on disk (can contain a single actor or a data chunk)
 	 */
@@ -32,6 +45,36 @@ struct FWorldPartitionRuntimeCellObjectMapping
 	 */
 	UPROPERTY()
 	FName Path;
+
+	/**
+	 * ID of the owning container instance
+	 */
+	UPROPERTY()
+	uint32 ContainerID;
+
+	/** 
+	 * Transform of the owning container instance
+	 */
+	UPROPERTY()
+	FTransform ContainerTransform;
+		
+	/**
+	 * Package of the owning container instance
+	 */
+	UPROPERTY()
+	FName ContainerPackage;
+
+	/**
+	* Loaded actor path (when cooking or pie)
+	* 
+	* Depending on if the actor was part of a container instance or the main partition this will be the path
+	* of the loaded or duplicated actor before it is moved into its runtime cell.
+	* 
+	* If the actor was part of the world partition this path should match the Path property.
+	*/
+	UPROPERTY()
+	FName LoadedPath;
+#endif
 };
 
 /**
@@ -43,6 +86,8 @@ class UWorldPartitionRuntimeCellData : public UObject
 	GENERATED_UCLASS_BODY()
 };
 
+class UActorDescContainer;
+
 /**
  * Represents a PIE/Game streaming cell which points to external actor/data chunk packages
  */
@@ -60,7 +105,7 @@ class UWorldPartitionRuntimeCell : public UObject
 #if WITH_EDITOR
 	void SetDataLayers(const TArray<const UDataLayer*> InDataLayers);
 	void AddCellData(const UWorldPartitionRuntimeCellData* InCellData);
-	virtual void AddActorToCell(const FGuid& ActorGuid, FName Package, FName Path) PURE_VIRTUAL(UWorldPartitionRuntimeCell::AddActorToCell,);
+	virtual void AddActorToCell(const FGuid& InActorGuid, uint32 InContainerID, const FTransform& InContainerTransform, const UActorDescContainer* InContainer) PURE_VIRTUAL(UWorldPartitionRuntimeCell::AddActorToCell,);
 	virtual int32 GetActorCount() const PURE_VIRTUAL(UWorldPartitionRuntimeCell::GetActorCount, return 0;);
 
 	// Cook methods
diff --git a/Engine/Source/Runtime/Engine/Public/WorldPartition/WorldPartitionRuntimeLevelStreamingCell.h b/Engine/Source/Runtime/Engine/Public/WorldPartition/WorldPartitionRuntimeLevelStreamingCell.h
--- a/Engine/Source/Runtime/Engine/Public/WorldPartition/WorldPartitionRuntimeLevelStreamingCell.h
+++ b/Engine/Source/Runtime/Engine/Public/WorldPartition/WorldPartitionRuntimeLevelStreamingCell.h
@@ -17,7 +17,7 @@ class UWorldPartitionRuntimeLevelStreamingCell : public UWorldPartitionRuntimeSp
 	virtual FLinearColor GetDebugColor() const override;
 
 #if WITH_EDITOR
-	virtual void AddActorToCell(const FGuid& ActorGuid, FName Package, FName Path) override;
+	virtual void AddActorToCell(const FGuid& InActorGuid, uint32 InContainerID, const FTransform& InContainerTransform, const UActorDescContainer* InContainer) override;
 	virtual int32 GetActorCount() const override;
 	const TArray<FWorldPartitionRuntimeCellObjectMapping>& GetPackages() const { return Packages; }
 
@@ -46,9 +46,6 @@ private:
 #if WITH_EDITORONLY_DATA
 	UPROPERTY()
 	TArray<FWorldPartitionRuntimeCellObjectMapping> Packages;
-
-	// Actors to be loaded for cook
-	TArray<FGuid> ActorsToLoadForCook;
 #endif
 
 	UPROPERTY()
