WorldPartition: Fix Failed check on Actors with no metadata (invalid actor desc)

diff --git a/Engine/Source/Runtime/Engine/Private/WorldPartition/ActorDescContainer.cpp b/Engine/Source/Runtime/Engine/Private/WorldPartition/ActorDescContainer.cpp
--- a/Engine/Source/Runtime/Engine/Private/WorldPartition/ActorDescContainer.cpp
+++ b/Engine/Source/Runtime/Engine/Private/WorldPartition/ActorDescContainer.cpp
@@ -50,9 +50,10 @@ void UActorDescContainer::Initialize(UWorld* InWorld, FName InPackageName, bool
 	for (const FAssetData& Asset : Assets)
 	{
 		TUniquePtr<FWorldPartitionActorDesc>* NewActorDesc = new(ActorDescList) TUniquePtr<FWorldPartitionActorDesc>(GetActorDescriptor(Asset));
-		check(NewActorDesc->IsValid());
-				
-		Actors.Add((*NewActorDesc)->GetGuid(), NewActorDesc);
+		if (NewActorDesc->IsValid())
+		{
+			Actors.Add((*NewActorDesc)->GetGuid(), NewActorDesc);
+		}
 	}
 
 	if (bRegisterDelegates)
