ShaderCompiler: do not print cache stats if the cache isn't enabled.

diff --git a/Engine/Source/Runtime/Engine/Private/ShaderCompiler/ShaderCompiler.cpp b/Engine/Source/Runtime/Engine/Private/ShaderCompiler/ShaderCompiler.cpp
--- a/Engine/Source/Runtime/Engine/Private/ShaderCompiler/ShaderCompiler.cpp
+++ b/Engine/Source/Runtime/Engine/Private/ShaderCompiler/ShaderCompiler.cpp
@@ -556,7 +556,11 @@ void FShaderCompileJobCollection::AddToCacheAndProcessPending(FShaderCommonCompi
 void FShaderCompileJobCollection::LogCachingStats(bool bForceLogIgnoringTimeInverval)
 {
 	static double LastTimeStatsPrinted = FPlatformTime::Seconds();
-	if (!bForceLogIgnoringTimeInverval && GShaderCompilerCacheStatsPrintoutInterval > 0 && FPlatformTime::Seconds() - LastTimeStatsPrinted < GShaderCompilerCacheStatsPrintoutInterval)
+	// do not print if
+	//  - job cache is disabled
+	//  - not enough time passed since the previous time and we're not forced to print
+	if (!ShaderCompiler::IsJobCacheEnabled() || 
+		(!bForceLogIgnoringTimeInverval && GShaderCompilerCacheStatsPrintoutInterval > 0 && FPlatformTime::Seconds() - LastTimeStatsPrinted < GShaderCompilerCacheStatsPrintoutInterval))
 	{
 		return;
 	}
