Implemented GetBounds for static mesh instance elements

#rnx

diff --git a/Engine/Source/Runtime/Engine/Private/Elements/SMInstance/SMInstanceElementWorldInterface.cpp b/Engine/Source/Runtime/Engine/Private/Elements/SMInstance/SMInstanceElementWorldInterface.cpp
--- a/Engine/Source/Runtime/Engine/Private/Elements/SMInstance/SMInstanceElementWorldInterface.cpp
+++ b/Engine/Source/Runtime/Engine/Private/Elements/SMInstance/SMInstanceElementWorldInterface.cpp
@@ -7,6 +7,8 @@
 #include "Elements/Framework/EngineElementsLibrary.h"
 #include "Elements/Framework/TypedElementSelectionSet.h"
 
+#include "Engine/StaticMesh.h"
+
 bool USMInstanceElementWorldInterface::CanEditElement(const FTypedElementHandle& InElementHandle)
 {
 	const FSMInstanceId SMInstance = SMInstanceElementDataUtil::GetSMInstanceFromHandle(InElementHandle);
@@ -42,7 +44,20 @@ bool USMInstanceElementWorldInterface::GetBounds(const FTypedElementHandle& InEl
 {
 	if (const FSMInstanceId SMInstance = SMInstanceElementDataUtil::GetSMInstanceFromHandle(InElementHandle))
 	{
-		// TODO: ISM instance bounds?
+		if (const UStaticMesh* StaticMesh = SMInstance.ISMComponent->GetStaticMesh())
+		{
+			OutBounds = StaticMesh->GetBounds();
+		}
+		else
+		{
+			OutBounds = FBoxSphereBounds();
+		}
+
+		FTransform InstanceTransform;
+		SMInstance.ISMComponent->GetInstanceTransform(SMInstance.InstanceIndex, InstanceTransform, /*bWorldSpace*/true);
+
+		OutBounds = OutBounds.TransformBy(InstanceTransform);
+		return true;
 	}
 
 	return false;
