OFPA: Added OnActorPackingModeChanged event

diff --git a/Engine/Source/Runtime/Engine/Classes/GameFramework/Actor.h b/Engine/Source/Runtime/Engine/Classes/GameFramework/Actor.h
--- a/Engine/Source/Runtime/Engine/Classes/GameFramework/Actor.h
+++ b/Engine/Source/Runtime/Engine/Classes/GameFramework/Actor.h
@@ -93,6 +93,10 @@ DECLARE_DYNAMIC_MULTICAST_SPARSE_DELEGATE_TwoParams(FActorEndPlaySignature, AAct
 
 DECLARE_DELEGATE_SixParams(FMakeNoiseDelegate, AActor*, float /*Loudness*/, class APawn*, const FVector&, float /*MaxRange*/, FName /*Tag*/);
 
+#if WITH_EDITOR
+DECLARE_EVENT_TwoParams(AActor, FActorOnPackagingModeChanged, AActor*, bool /* bExternal */);
+#endif
+
 #if !UE_BUILD_SHIPPING
 DECLARE_DELEGATE_RetVal_ThreeParams(bool, FOnProcessEvent, AActor*, UFunction*, void*);
 #endif
@@ -806,6 +810,8 @@ public:
 	 */
 	void SetPackageExternal(bool bExternal, bool bShouldDirty = true);
 
+	FActorOnPackagingModeChanged OnPackagingModeChanged;
+
 	/**
 	 * Determine how this actor should be placed in the partition (meant to be called on the default object).
 	 / @return EActorGridPlacement::None if the placement should be on a per-instance basis, otherwise the default value wins.
diff --git a/Engine/Source/Runtime/Engine/Private/ActorEditor.cpp b/Engine/Source/Runtime/Engine/Private/ActorEditor.cpp
--- a/Engine/Source/Runtime/Engine/Private/ActorEditor.cpp
+++ b/Engine/Source/Runtime/Engine/Private/ActorEditor.cpp
@@ -884,6 +884,8 @@ void AActor::SetPackageExternal(bool bExternal, bool bShouldDirty)
 			ActorComponent->SetPackageExternal(bExternal, bShouldDirty);
 		}
 	}
+
+	OnPackagingModeChanged.Broadcast(this, bExternal);
 	
 	// Mark the new actor package dirty
 	MarkPackageDirty();
