Print more information about material DDC (UE-43174)

(Edigrating 15491171, 15501722).

diff --git a/Engine/Source/Runtime/Engine/Private/Materials/MaterialShader.cpp b/Engine/Source/Runtime/Engine/Private/Materials/MaterialShader.cpp
--- a/Engine/Source/Runtime/Engine/Private/Materials/MaterialShader.cpp
+++ b/Engine/Source/Runtime/Engine/Private/Materials/MaterialShader.cpp
@@ -1082,7 +1082,7 @@ void FMaterialShaderMap::GetAllOutdatedTypes(TArray<const FShaderType*>& Outdate
 #endif
 }
 
-void FMaterialShaderMap::LoadFromDerivedDataCache(const FMaterial* Material, const FMaterialShaderMapId& ShaderMapId, EShaderPlatform InPlatform, const ITargetPlatform* TargetPlatform, TRefCountPtr<FMaterialShaderMap>& InOutShaderMap)
+void FMaterialShaderMap::LoadFromDerivedDataCache(const FMaterial* Material, const FMaterialShaderMapId& ShaderMapId, EShaderPlatform InPlatform, const ITargetPlatform* TargetPlatform, TRefCountPtr<FMaterialShaderMap>& InOutShaderMap, FString& OutDDCKeyDesc)
 {
 	if (InOutShaderMap != NULL)
 	{
@@ -1100,6 +1100,14 @@ void FMaterialShaderMap::LoadFromDerivedDataCache(const FMaterial* Material, con
 
 			TArray<uint8> CachedData;
 			const FString DataKey = GetMaterialShaderMapKeyString(ShaderMapId, InPlatform, TargetPlatform);
+			{
+				FSHAHash DDCKeyHash;
+				FSHA1 HashState;
+				HashState.UpdateWithString(*DataKey, DataKey.Len());
+				HashState.Final();
+				HashState.GetHash(&DDCKeyHash.Hash[0]);
+				OutDDCKeyDesc = DDCKeyHash.ToString();
+			}
 
 			bool CheckCache = true;
 
diff --git a/Engine/Source/Runtime/Engine/Private/Materials/MaterialShared.cpp b/Engine/Source/Runtime/Engine/Private/Materials/MaterialShared.cpp
--- a/Engine/Source/Runtime/Engine/Private/Materials/MaterialShared.cpp
+++ b/Engine/Source/Runtime/Engine/Private/Materials/MaterialShared.cpp
@@ -1484,6 +1484,7 @@ bool FMaterialResource::ShouldApplyCloudFogging() const { return Material->bAppl
 bool FMaterialResource::IsSky() const { return Material->bIsSky; }
 bool FMaterialResource::ComputeFogPerPixel() const {return Material->bComputeFogPerPixel;}
 FString FMaterialResource::GetFriendlyName() const { return *GetNameSafe(Material); } //avoid using the material instance name here, we want materials that share a shadermap to also share a friendly name.
+FString FMaterialResource::GetAssetName() const { return MaterialInstance ? *GetNameSafe(MaterialInstance) : *GetNameSafe(Material); }
 
 uint32 FMaterialResource::GetDecalBlendMode() const
 {
@@ -2015,6 +2016,9 @@ bool FMaterial::CacheShaders(const FMaterialShaderMapId& ShaderMapId, EShaderPla
 {
 	TRACE_CPUPROFILER_EVENT_SCOPE(FMaterial::CacheShaders);
 	UE_CLOG(!ShaderMapId.IsValid(), LogMaterial, Warning, TEXT("Invalid shader map ID caching shaders for '%s', will use default material."), *GetFriendlyName());
+#if WITH_EDITOR
+	FString DDCKeyHash;
+#endif // WITH_EDITOR
 
 	// If we loaded this material with inline shaders, use what was loaded (GameThreadShaderMap) instead of looking in the DDC
 	if (bContainsInlineShaders)
@@ -2058,7 +2062,7 @@ bool FMaterial::CacheShaders(const FMaterialShaderMapId& ShaderMapId, EShaderPla
 		if (!ShaderMap && !FPlatformProperties::RequiresCookedData())
 		{
 			TRefCountPtr<FMaterialShaderMap> LoadedShaderMap;
-			FMaterialShaderMap::LoadFromDerivedDataCache(this, ShaderMapId, Platform, TargetPlatform, LoadedShaderMap);
+			FMaterialShaderMap::LoadFromDerivedDataCache(this, ShaderMapId, Platform, TargetPlatform, LoadedShaderMap, DDCKeyHash);
 			ShaderMap = LoadedShaderMap;
 		}
 
@@ -2126,7 +2130,23 @@ bool FMaterial::CacheShaders(const FMaterialShaderMapId& ShaderMapId, EShaderPla
 			{
 				ShaderMapCondition = TEXT("Missing");
 			}
-			UE_LOG(LogMaterial, Display, TEXT("%s cached shader map for material %s, compiling. %s"),ShaderMapCondition,*GetFriendlyName(), IsSpecialEngineMaterial() ? TEXT("Is special engine material.") : TEXT("") );
+#if WITH_EDITOR
+			UE_LOG(LogMaterial, Display, TEXT("%s cached shadermap for %s in %s, %s, %s (DDC key hash: %s), compiling. %s"),
+				ShaderMapCondition,
+				*GetAssetName(),
+				*LexToString(Platform),
+				*LexToString(ShaderMapId.QualityLevel),
+				*LexToString(ShaderMapId.FeatureLevel),
+				*DDCKeyHash,
+				IsSpecialEngineMaterial() ? TEXT("Is special engine material.") : TEXT("") 
+			);
+#else
+			UE_LOG(LogMaterial, Display, TEXT("%s cached shader map for material %, compiling. %s"),
+				ShaderMapCondition,
+				*GetAssetName(),
+				IsSpecialEngineMaterial() ? TEXT("Is special engine material.") : TEXT("")
+			);
+#endif
 
 #if WITH_EDITORONLY_DATA
 			FStaticParameterSet StaticParameterSet;
diff --git a/Engine/Source/Runtime/Engine/Private/SceneView.cpp b/Engine/Source/Runtime/Engine/Private/SceneView.cpp
--- a/Engine/Source/Runtime/Engine/Private/SceneView.cpp
+++ b/Engine/Source/Runtime/Engine/Private/SceneView.cpp
@@ -390,6 +390,26 @@ TLinkedList<FSceneViewStateReference*>*& FSceneViewStateReference::GetSceneViewS
 	return List;
 }
 
+FString LexToString(EMaterialQualityLevel::Type QualityLevel)
+{
+	switch (QualityLevel)
+	{
+		case EMaterialQualityLevel::Low:
+			return TEXT("Low");
+		case EMaterialQualityLevel::High:
+			return TEXT("High");
+		case EMaterialQualityLevel::Medium:
+			return TEXT("Medium");
+		case EMaterialQualityLevel::Epic:
+			return TEXT("Epic");
+		case EMaterialQualityLevel::Num:
+			return TEXT("Default");
+		default:
+			break;
+	}
+	return TEXT("UnknownFeatureLevel");
+}
+
 /**
  * Utility function to create the inverse depth projection transform to be used
  * by the shader system.
diff --git a/Engine/Source/Runtime/Engine/Public/MaterialShared.h b/Engine/Source/Runtime/Engine/Public/MaterialShared.h
--- a/Engine/Source/Runtime/Engine/Public/MaterialShared.h
+++ b/Engine/Source/Runtime/Engine/Public/MaterialShared.h
@@ -1072,8 +1072,9 @@ public:
 	/** 
 	 * Attempts to load the shader map for the given material from the Derived Data Cache.
 	 * If InOutShaderMap is valid, attempts to load the individual missing shaders instead.
+	 * Returns (via OutDDCKeyDesc parameter) a helpful string to debug the DDC key and parameters
 	 */
-	static void LoadFromDerivedDataCache(const FMaterial* Material, const FMaterialShaderMapId& ShaderMapId, EShaderPlatform Platform, const ITargetPlatform* TargetPlatform, TRefCountPtr<FMaterialShaderMap>& InOutShaderMap);
+	static void LoadFromDerivedDataCache(const FMaterial* Material, const FMaterialShaderMapId& ShaderMapId, EShaderPlatform Platform, const ITargetPlatform* TargetPlatform, TRefCountPtr<FMaterialShaderMap>& InOutShaderMap, FString& OutDDCKeyDesc);
 #endif
 
 	FMaterialShaderMap();
@@ -1749,6 +1750,8 @@ public:
 	virtual bool ComputeFogPerPixel() const { return false; }
 	virtual bool IsSky() const { return false; }
 	virtual FString GetFriendlyName() const = 0;
+	/** Similar to GetFriendlyName, but but avoids historical behavior of the former, returning the exact asset name for material instances instead of just the material. */
+	virtual FString GetAssetName() const { return GetFriendlyName(); }
 	virtual bool HasVertexPositionOffsetConnected() const { return false; }
 	virtual bool HasPixelDepthOffsetConnected() const { return false; }
 	virtual bool HasMaterialAttributesConnected() const { return false; }
@@ -2615,6 +2618,7 @@ public:
 	ENGINE_API virtual bool IsDitherMasked() const override;
 	ENGINE_API virtual bool AllowNegativeEmissiveColor() const override;
 	ENGINE_API virtual FString GetFriendlyName() const override;
+	ENGINE_API virtual FString GetAssetName() const override;
 	ENGINE_API virtual bool RequiresSynchronousCompilation() const override;
 	ENGINE_API virtual bool IsDefaultMaterial() const override;
 	ENGINE_API virtual int32 GetNumCustomizedUVs() const override;
diff --git a/Engine/Source/Runtime/Engine/Public/SceneTypes.h b/Engine/Source/Runtime/Engine/Public/SceneTypes.h
--- a/Engine/Source/Runtime/Engine/Public/SceneTypes.h
+++ b/Engine/Source/Runtime/Engine/Public/SceneTypes.h
@@ -143,6 +143,8 @@ namespace EMaterialQualityLevel
 	};
 }
 
+ENGINE_API FString LexToString(EMaterialQualityLevel::Type QualityLevel);
+
 //
 //	EMaterialProperty
 //
diff --git a/Engine/Source/Runtime/RHI/Private/RHI.cpp b/Engine/Source/Runtime/RHI/Private/RHI.cpp
--- a/Engine/Source/Runtime/RHI/Private/RHI.cpp
+++ b/Engine/Source/Runtime/RHI/Private/RHI.cpp
@@ -1579,6 +1579,12 @@ FString LexToString(EShaderPlatform Platform, bool bError)
 	}
 }
 
+FString LexToString(EShaderPlatform Platform)
+{
+	bool bError = true;
+	return LexToString(Platform, bError);
+}
+
 void LexFromString(EShaderPlatform& Value, const TCHAR* String)
 {
 	Value = EShaderPlatform::SP_NumPlatforms;
@@ -1593,6 +1599,23 @@ void LexFromString(EShaderPlatform& Value, const TCHAR* String)
 	}
 }
 
+FString LexToString(ERHIFeatureLevel::Type Level)
+{
+	switch (Level)
+	{
+		case ERHIFeatureLevel::ES2_REMOVED:
+			return TEXT("ES2_REMOVED");
+		case ERHIFeatureLevel::ES3_1:
+			return TEXT("ES3_1");
+		case ERHIFeatureLevel::SM4_REMOVED:
+			return TEXT("SM4_REMOVED");
+		case ERHIFeatureLevel::SM5:
+			return TEXT("SM5");
+		default:
+			break;
+	}
+	return TEXT("UnknownFeatureLevel");
+}
 
 const FName LANGUAGE_D3D("D3D");
 const FName LANGUAGE_Metal("Metal");
diff --git a/Engine/Source/Runtime/RHI/Public/RHI.h b/Engine/Source/Runtime/RHI/Public/RHI.h
--- a/Engine/Source/Runtime/RHI/Public/RHI.h
+++ b/Engine/Source/Runtime/RHI/Public/RHI.h
@@ -715,6 +715,11 @@ inline EShaderPlatform GetFeatureLevelShaderPlatform(const FStaticFeatureLevel I
 	return GShaderPlatformForFeatureLevel[InFeatureLevel];
 }
 
+/** Stringifies EShaderPlatform */
+extern RHI_API FString LexToString(EShaderPlatform Platform);
+
+/** Stringifies ERHIFeatureLevel */
+extern RHI_API FString LexToString(ERHIFeatureLevel::Type Level);
 
 /** Finds a corresponding ERHIShadingPath::Type given an FName, or returns false if one could not be found. */
 extern RHI_API bool GetShadingPathFromName(FName Name, ERHIShadingPath::Type& OutShadingPath);
