Don't try to be cute with ICH templates, instead just move a conflicting object out of the way.

diff --git a/Engine/Source/Runtime/Engine/Classes/Engine/InheritableComponentHandler.h b/Engine/Source/Runtime/Engine/Classes/Engine/InheritableComponentHandler.h
--- a/Engine/Source/Runtime/Engine/Classes/Engine/InheritableComponentHandler.h
+++ b/Engine/Source/Runtime/Engine/Classes/Engine/InheritableComponentHandler.h
@@ -160,7 +160,7 @@ private:
 	const FComponentOverrideRecord* FindRecord(const FComponentKey& Key) const;
 
 	/** Helper method used to assist with fixing up component template names at load time. */
-	void FixComponentTemplateName(UActorComponent* ComponentTemplate, const FString& NewName);
+	void FixComponentTemplateName(UActorComponent* ComponentTemplate, FName NewName);
 	
 	/** All component records */
 	UPROPERTY()
diff --git a/Engine/Source/Runtime/Engine/Private/InheritableComponentHandler.cpp b/Engine/Source/Runtime/Engine/Private/InheritableComponentHandler.cpp
--- a/Engine/Source/Runtime/Engine/Private/InheritableComponentHandler.cpp
+++ b/Engine/Source/Runtime/Engine/Private/InheritableComponentHandler.cpp
@@ -50,17 +50,17 @@ void UInheritableComponentHandler::PostLoad()
 				// Note: If the key maps to an SCS node, the node's variable GUID will be used for the lookup instead of the name below (that's only used for UCS keys).
 				if (UActorComponent* OriginalTemplate = Record.ComponentKey.GetOriginalTemplate(Record.ComponentTemplate->GetFName()))
 				{
-					FString ExpectedTemplateName = OriginalTemplate->GetName();
+					FName ExpectedTemplateName = OriginalTemplate->GetFName();
 					if (USCS_Node* SCSNode = Record.ComponentKey.FindSCSNode())
 					{
 						// We append a prefix onto SCS default scene root node overrides. This is done to ensure that the override template does not collide with our owner's own SCS default scene root node template.
 						if (SCSNode == SCSNode->GetSCS()->GetDefaultSceneRootNode())
 						{
-							ExpectedTemplateName = SCSDefaultSceneRootOverrideNamePrefix + ExpectedTemplateName;
+							ExpectedTemplateName = FName(SCSDefaultSceneRootOverrideNamePrefix + ExpectedTemplateName.ToString());
 						}
 					}
 
-					if (ExpectedTemplateName != Record.ComponentTemplate->GetName())
+					if (ExpectedTemplateName != Record.ComponentTemplate->GetFName())
 					{
 						FixComponentTemplateName(Record.ComponentTemplate, ExpectedTemplateName);
 					}
@@ -504,33 +504,18 @@ const FComponentOverrideRecord* UInheritableComponentHandler::FindRecord(const F
 	return nullptr;
 }
 
-void UInheritableComponentHandler::FixComponentTemplateName(UActorComponent* ComponentTemplate, const FString& NewName)
+void UInheritableComponentHandler::FixComponentTemplateName(UActorComponent* ComponentTemplate, const FName NewName)
 {
-	// Look for a collision with the template we're trying to rename here. It's possible that names were swapped on the
-	// original component template objects that were inherited from the associated Blueprint's parent class, for example.
-	FComponentOverrideRecord* MatchingRecord = Records.FindByPredicate([ComponentTemplate, &NewName](FComponentOverrideRecord& Record)
-	{
-		if (Record.ComponentTemplate && Record.ComponentTemplate != ComponentTemplate && Record.ComponentTemplate->GetName() == NewName)
-		{
-			const FName OverrideTemplateName = ComponentTemplate->GetFName();
-			const UActorComponent* OriginalTemplate = Record.ComponentKey.GetOriginalTemplate(OverrideTemplateName);
-			return ensureMsgf(OriginalTemplate && OriginalTemplate->GetFName() != OverrideTemplateName,
-				TEXT("Found a collision with an existing override record, but its associated template object is either invalid or already matches its inherited template's name (%s). This is unexpected."), *NewName);
-		}
-
-		return false;
-	});
-
 	// If we found a collision, temporarily rename the associated template object to something unique so that it no longer
 	// collides with the one we're trying to correct here. This will be fixed up when we later encounter this record during
 	// PostLoad() validation and see that it still doesn't match its original template name.
-	if (MatchingRecord)
+	if (UObject* ExistingObject = (UObject*)FindObjectWithOuter(ComponentTemplate->GetOuter(), nullptr, NewName))
 	{
-		MatchingRecord->ComponentTemplate->Rename(nullptr, nullptr, REN_DontCreateRedirectors | REN_ForceNoResetLoaders);
+		ExistingObject->Rename(nullptr, nullptr, REN_DontCreateRedirectors | REN_ForceNoResetLoaders);
 	}
 
 	// Now that we're sure there are no collisions with other records, we can safely rename this one to its new name.
-	ComponentTemplate->Rename(*NewName, nullptr, REN_DontCreateRedirectors | REN_ForceNoResetLoaders);
+	ComponentTemplate->Rename(*NewName.ToString(), nullptr, REN_DontCreateRedirectors | REN_ForceNoResetLoaders);
 }
 
 // FComponentOverrideRecord
