Move definitition of ISM element data guard to build.cs file so it can be accessed without having to include the SMInstanceElementId header.

diff --git a/Engine/Source/Runtime/Engine/Engine.Build.cs b/Engine/Source/Runtime/Engine/Engine.Build.cs
--- a/Engine/Source/Runtime/Engine/Engine.Build.cs
+++ b/Engine/Source/Runtime/Engine/Engine.Build.cs
@@ -273,6 +273,8 @@ public class Engine : ModuleRules
 			}
 		}
 
+		PublicDefinitions.Add("UE_ENABLE_SMINSTANCE_ELEMENTS=0");
+
 		DynamicallyLoadedModuleNames.AddRange(
 			new string[] {
 				"NullNetworkReplayStreaming",
diff --git a/Engine/Source/Runtime/Engine/Public/Elements/SMInstance/SMInstanceElementId.h b/Engine/Source/Runtime/Engine/Public/Elements/SMInstance/SMInstanceElementId.h
--- a/Engine/Source/Runtime/Engine/Public/Elements/SMInstance/SMInstanceElementId.h
+++ b/Engine/Source/Runtime/Engine/Public/Elements/SMInstance/SMInstanceElementId.h
@@ -7,8 +7,6 @@
 #include "InstancedStaticMeshDelegates.h"
 #include "SMInstanceElementId.generated.h"
 
-#define UE_ENABLE_SMINSTANCE_ELEMENTS (0)
-
 class FSMInstanceElementIdMap;
 class UInstancedStaticMeshComponent;
 class USMInstanceElementIdMapTransactor;
