Relocate asset compilation shutdown to avoid crash when potentially queueing distancefield during async static mesh build
Avoid static mesh queuing any more work during shutdown to improve exit speed

diff --git a/Engine/Source/Runtime/Engine/Private/StaticMesh.cpp b/Engine/Source/Runtime/Engine/Private/StaticMesh.cpp
--- a/Engine/Source/Runtime/Engine/Private/StaticMesh.cpp
+++ b/Engine/Source/Runtime/Engine/Private/StaticMesh.cpp
@@ -2778,6 +2778,12 @@ void FStaticMeshRenderData::Cache(const ITargetPlatform* TargetPlatform, UStatic
 		}
 	}
 
+	// If the engine is exiting and is waiting on us to finish this pending task, exit asap and don't start any other work.
+	if (IsEngineExitRequested())
+	{
+		return;
+	}
+
 	static const auto CVar = IConsoleManager::Get().FindTConsoleVariableDataInt(TEXT("r.GenerateMeshDistanceFields"));
 
 	if (CVar->GetValueOnAnyThread(true) != 0 || Owner->bGenerateMeshDistanceField)
