Fix Crash when editing LevelInstance

diff --git a/Engine/Source/Runtime/Engine/Private/LevelInstance/LevelInstanceActor.cpp b/Engine/Source/Runtime/Engine/Private/LevelInstance/LevelInstanceActor.cpp
--- a/Engine/Source/Runtime/Engine/Private/LevelInstance/LevelInstanceActor.cpp
+++ b/Engine/Source/Runtime/Engine/Private/LevelInstance/LevelInstanceActor.cpp
@@ -470,7 +470,7 @@ bool ALevelInstance::CanEditChange(const FProperty* InProperty) const
 
 	if (InProperty->GetFName() == GET_MEMBER_NAME_CHECKED(ALevelInstance, DesiredRuntimeBehavior))
 	{
-		return GetLevel()->GetWorldPartition() != nullptr;
+		return GetLevel() && GetLevel()->GetWorldPartition() != nullptr;
 	}
 
 	return true;
