LevelInstance: PIE Fix for PackedLevelInstance that were stripped by error

diff --git a/Engine/Source/Runtime/Engine/Private/WorldPartition/LevelInstance/LevelInstanceActorDesc.cpp b/Engine/Source/Runtime/Engine/Private/WorldPartition/LevelInstance/LevelInstanceActorDesc.cpp
--- a/Engine/Source/Runtime/Engine/Private/WorldPartition/LevelInstance/LevelInstanceActorDesc.cpp
+++ b/Engine/Source/Runtime/Engine/Private/WorldPartition/LevelInstance/LevelInstanceActorDesc.cpp
@@ -30,8 +30,11 @@ void FLevelInstanceActorDesc::Init(const AActor* InActor)
 
 void FLevelInstanceActorDesc::Init(UActorDescContainer* InContainer, const FWorldPartitionActorDescInitData& DescData)
 {
-	FWorldPartitionActorDesc::Init(InContainer, DescData);
+	ALevelInstance* CDO = DescData.NativeClass->GetDefaultObject<ALevelInstance>();
+	DesiredRuntimeBehavior = CDO->GetDefaultRuntimeBehavior();
 
+	FWorldPartitionActorDesc::Init(InContainer, DescData);
+		
 	RegisterContainer();
 }
 
diff --git a/Engine/Source/Runtime/Engine/Public/LevelInstance/LevelInstanceActor.h b/Engine/Source/Runtime/Engine/Public/LevelInstance/LevelInstanceActor.h
--- a/Engine/Source/Runtime/Engine/Public/LevelInstance/LevelInstanceActor.h
+++ b/Engine/Source/Runtime/Engine/Public/LevelInstance/LevelInstanceActor.h
@@ -74,6 +74,7 @@ public:
 	virtual void GetActorLocationBounds(bool bOnlyCollidingComponents, FVector& Origin, FVector& BoxExtent, bool bIncludeFromChildActors = false) const override;
 	virtual bool IsLockLocation() const override;
 	virtual ELevelInstanceRuntimeBehavior GetDesiredRuntimeBehavior() const { return DesiredRuntimeBehavior; }
+	virtual ELevelInstanceRuntimeBehavior GetDefaultRuntimeBehavior() const { return ELevelInstanceRuntimeBehavior::Embedded; }
 	virtual bool IsHLODRelevant() const override { return true; }
 
 	bool CanEdit(FText* OutReason = nullptr) const;
diff --git a/Engine/Source/Runtime/Engine/Public/LevelInstance/Packed/PackedLevelInstanceActor.h b/Engine/Source/Runtime/Engine/Public/LevelInstance/Packed/PackedLevelInstanceActor.h
--- a/Engine/Source/Runtime/Engine/Public/LevelInstance/Packed/PackedLevelInstanceActor.h
+++ b/Engine/Source/Runtime/Engine/Public/LevelInstance/Packed/PackedLevelInstanceActor.h
@@ -49,6 +49,8 @@ public:
 	void DestroyPackedComponents();
 	void GetPackedComponents(TArray<UActorComponent*>& OutPackedComponents) const;
 
+	virtual ELevelInstanceRuntimeBehavior GetDefaultRuntimeBehavior() const override { return ELevelInstanceRuntimeBehavior::None; }
+
 	template<class T>
 	T* AddPackedComponent()
 	{
