Add auto conversion for SSS profile, Pre-integrated SSS, and SSS shading models.

diff --git a/Engine/Source/Runtime/Engine/Private/Materials/Material.cpp b/Engine/Source/Runtime/Engine/Private/Materials/Material.cpp
--- a/Engine/Source/Runtime/Engine/Private/Materials/Material.cpp
+++ b/Engine/Source/Runtime/Engine/Private/Materials/Material.cpp
@@ -49,6 +49,9 @@
 #include "Materials/MaterialFunctionInstance.h"
 #include "Materials/MaterialExpressionMaterialFunctionCall.h"
 #include "Materials/MaterialExpressionSingleLayerWaterMaterialOutput.h"
+#include "Materials/MaterialExpressionOneMinus.h"
+#include "Materials/MaterialExpressionMultiply.h"
+#include "Materials/MaterialExpressionSaturate.h"
 #include "Engine/Font.h"
 #include "SceneManagement.h"
 #include "Materials/MaterialUniformExpressions.h"
@@ -3653,7 +3656,6 @@ void UMaterial::ConvertMaterialToStrataMaterial()
 		{
 			UMaterialExpressionStrataSlabBSDF* SlabBSDF = NewObject<UMaterialExpressionStrataSlabBSDF>(this);
 
-
 			MoveConnectionTo(BaseColor, SlabBSDF, 0);		// BaseColor
 			MoveConnectionTo(Metallic, SlabBSDF, 2);		// Metallic
 			MoveConnectionTo(Specular, SlabBSDF, 3);		// Specular
@@ -3669,7 +3671,6 @@ void UMaterial::ConvertMaterialToStrataMaterial()
 			// WorldPositionOffset can remain on the end point node
 			// STRATA_TODO WorldDisplacement and TessellationMultiplier are going to be removed and should be ignored?
 			// STRATA_TODO ShadingModel
-			// STRATA_TODO SubsurfaceColor
 			// STRATA_TODO ClearCoat
 			// STRATA_TODO ClearCoatRoughness
 			// STRATA_TODO AmbientOcclusion should be removed and put on the BSDF node
@@ -3679,6 +3680,67 @@ void UMaterial::ConvertMaterialToStrataMaterial()
 
 			FrontMaterial.Connect(0, SlabBSDF);
 		}
+		else if (MaterialDomain == MD_Surface && (ShadingModel == MSM_Subsurface || ShadingModel == MSM_SubsurfaceProfile || ShadingModel == MSM_PreintegratedSkin))
+		{
+			UMaterialExpressionStrataSlabBSDF* SlabBSDF = NewObject<UMaterialExpressionStrataSlabBSDF>(this);
+
+			// Common connections
+			MoveConnectionTo(Metallic, SlabBSDF, 2);		// Metallic
+			MoveConnectionTo(Specular, SlabBSDF, 3);		// Specular
+			MoveConnectionTo(Roughness, SlabBSDF, 4);		// Roughness
+			if (Anisotropy.IsConnected())
+			{
+				MoveConnectionTo(Anisotropy, SlabBSDF, 5);	// Anisotropy
+			}
+			MoveConnectionTo(Normal, SlabBSDF, 6);			// Normal
+			MoveConnectionTo(Tangent, SlabBSDF, 7);			// Tangent
+			
+			// SSS connections
+			if (ShadingModel == MSM_SubsurfaceProfile)
+			{
+				MoveConnectionTo(BaseColor, SlabBSDF, 0);		// BaseColor
+				SlabBSDF->SubsurfaceProfile = SubsurfaceProfile;// SSS Profile
+				MoveConnectionTo(Opacity, SlabBSDF, 9);			// SSSDMFPScale
+			}
+			else if (ShadingModel == MSM_Subsurface || ShadingModel == MSM_PreintegratedSkin)
+			{	
+				// Insure there is no profile, as this would take priority otherwise
+				SlabBSDF->SubsurfaceProfile = nullptr;
+				ShadingModel = MSM_SubsurfaceProfile;
+
+				// Subsurface model has a subsurface color and a base color which are additive through normal warpping. 
+				// We emulate that by adding them together
+				UMaterialExpressionAdd* AddBaseColor = NewObject<UMaterialExpressionAdd>(this);
+				MoveConnectionTo(BaseColor, AddBaseColor, 0);		// A
+				CopyConnectionTo(SubsurfaceColor, AddBaseColor, 1);	// B
+				SlabBSDF->BaseColor.Connect(0, AddBaseColor);
+
+				// For Subsurface and Pre-integrated subsurface opacity acts as lerp between SSS (0) and Opaque(1), 
+				// so we need to (saturate and) inverse with a OneMinus node.
+				UMaterialExpressionSaturate* Saturate = NewObject<UMaterialExpressionSaturate>(this);
+				MoveConnectionTo(Opacity, Saturate, 0);
+
+				UMaterialExpressionOneMinus* OneMinusOpacity = NewObject<UMaterialExpressionOneMinus>(this);
+				OneMinusOpacity->Input.Connect(0, Saturate);
+				
+				UMaterialExpressionMultiply* MulOpacityWithSubsurfaceColor = NewObject<UMaterialExpressionMultiply>(this);
+				MoveConnectionTo(SubsurfaceColor, MulOpacityWithSubsurfaceColor, 0);	// A
+				MulOpacityWithSubsurfaceColor->B.Connect(0, OneMinusOpacity);			// B
+				
+				// Subsurface color is converted into a MFP by taking the subsurface color and multipling it with the 1-Opacity 
+				// and a constant scaling for matching roughly the behavior
+				UMaterialExpressionConstant* MFPConstantScale = NewObject<UMaterialExpressionConstant>(this);
+				MFPConstantScale->R = 2.f;
+
+				UMaterialExpressionMultiply* MulScaleMFP = NewObject<UMaterialExpressionMultiply>(this);
+				MulScaleMFP->A.Connect(0, MulOpacityWithSubsurfaceColor);
+				MulScaleMFP->B.Connect(0, MFPConstantScale);
+
+				SlabBSDF->SSSDMFP.Connect(0, MulScaleMFP);
+			}
+
+			FrontMaterial.Connect(0, SlabBSDF);
+		}
 		else if (MaterialDomain == MD_Surface && ShadingModel == MSM_Hair)
 		{
 			UMaterialExpressionStrataHairBSDF* HairBSDF = NewObject<UMaterialExpressionStrataHairBSDF>(this);
@@ -3889,6 +3951,9 @@ void UMaterial::ConvertMaterialToStrataMaterial()
 		StateId.C = HashBuffer[2];
 		StateId.D = HashBuffer[3];
 	}
+
+	// For rebuild the shading mode since we have change it
+	RebuildShadingModelField();
 #endif
 }
 
