Add support for async texture of VT lightmaps which was previously simply crashing

diff --git a/Engine/Source/Runtime/Engine/Classes/Engine/MapBuildDataRegistry.h b/Engine/Source/Runtime/Engine/Classes/Engine/MapBuildDataRegistry.h
--- a/Engine/Source/Runtime/Engine/Classes/Engine/MapBuildDataRegistry.h
+++ b/Engine/Source/Runtime/Engine/Classes/Engine/MapBuildDataRegistry.h
@@ -20,6 +20,7 @@
 
 class FPrecomputedLightVolumeData;
 class FPrecomputedVolumetricLightmapData;
+struct FAssetCompileData;
 
 struct ENGINE_API FPerInstanceLightmapData
 {
@@ -376,6 +377,9 @@ public:
 	*/
 	ENGINE_API void HandleLegacyEncodedCubemapData();
 private:
+#if WITH_EDITOR
+	void HandleAssetPostCompileEvent(const TArray<FAssetCompileData>& CompiledAssets);
+#endif
 
 	ENGINE_API void ReleaseResources(const TSet<FGuid>* ResourcesToKeep = nullptr);
 	ENGINE_API void EmptyLevelData(const TSet<FGuid>* ResourcesToKeep = nullptr);
diff --git a/Engine/Source/Runtime/Engine/Private/LightMap.cpp b/Engine/Source/Runtime/Engine/Private/LightMap.cpp
--- a/Engine/Source/Runtime/Engine/Private/LightMap.cpp
+++ b/Engine/Source/Runtime/Engine/Private/LightMap.cpp
@@ -3464,6 +3464,15 @@ IAllocatedVirtualTexture* FLightmapResourceCluster::AcquireAllocatedVT() const
 	int32 LightmapIndex = bHighQuality ? 0 : 1;
 
 	const ULightMapVirtualTexture2D* VirtualTexture = Input.LightMapVirtualTextures[LightmapIndex];
+	
+#if WITH_EDITOR
+	// Compilation is still pending, this function will be called back once compilation finishes.
+	if (VirtualTexture->IsCompiling())
+	{
+		return nullptr;
+	}
+#endif
+
 	if (!AllocatedVT && VirtualTexture && VirtualTexture->Resource)
 	{
 		check(VirtualTexture->VirtualTextureStreaming);
diff --git a/Engine/Source/Runtime/Engine/Private/LightmapUniformShaderParameters.cpp b/Engine/Source/Runtime/Engine/Private/LightmapUniformShaderParameters.cpp
--- a/Engine/Source/Runtime/Engine/Private/LightmapUniformShaderParameters.cpp
+++ b/Engine/Source/Runtime/Engine/Private/LightmapUniformShaderParameters.cpp
@@ -101,13 +101,14 @@ void GetPrecomputedLightingParameters(
 			const FLightmapResourceCluster* ResourceCluster = LCI->GetResourceCluster();
 			check(ResourceCluster);
 			IAllocatedVirtualTexture* AllocatedVT = ResourceCluster->AcquireAllocatedVT();
-			check(AllocatedVT);
-
-			AllocatedVT->GetPackedPageTableUniform(&Parameters.LightmapVTPackedPageTableUniform[0]);
-			NumLightmapVTLayers = AllocatedVT->GetNumTextureLayers();
-			for (uint32 LayerIndex = 0u; LayerIndex < NumLightmapVTLayers; ++LayerIndex)
+			if (AllocatedVT)
 			{
-				AllocatedVT->GetPackedUniform(&Parameters.LightmapVTPackedUniform[LayerIndex], LayerIndex);
+				AllocatedVT->GetPackedPageTableUniform(&Parameters.LightmapVTPackedPageTableUniform[0]);
+				NumLightmapVTLayers = AllocatedVT->GetNumTextureLayers();
+				for (uint32 LayerIndex = 0u; LayerIndex < NumLightmapVTLayers; ++LayerIndex)
+				{
+					AllocatedVT->GetPackedUniform(&Parameters.LightmapVTPackedUniform[LayerIndex], LayerIndex);
+				}
 			}
 		}
 		else
diff --git a/Engine/Source/Runtime/Engine/Private/MapBuildData.cpp b/Engine/Source/Runtime/Engine/Private/MapBuildData.cpp
--- a/Engine/Source/Runtime/Engine/Private/MapBuildData.cpp
+++ b/Engine/Source/Runtime/Engine/Private/MapBuildData.cpp
@@ -25,6 +25,10 @@ MapBuildData.cpp
 #include "Components/ReflectionCaptureComponent.h"
 #include "Interfaces/ITargetPlatform.h"
 #if WITH_EDITOR
+#include "LandscapeComponent.h"
+#include "Components/StaticMeshComponent.h"
+#include "VT/LightmapVirtualTexture.h"
+#include "AssetCompilingManager.h"
 #include "Factories/TextureFactory.h"
 #endif
 #include "Engine/TextureCube.h"
@@ -361,7 +365,81 @@ UMapBuildDataRegistry::UMapBuildDataRegistry(const FObjectInitializer& ObjectIni
 {
 	LevelLightingQuality = Quality_MAX;
 	bSetupResourceClusters = false;
+
+#if WITH_EDITOR
+	FAssetCompilingManager::Get().OnAssetPostCompileEvent().AddUObject(this, &ThisClass::HandleAssetPostCompileEvent);
+#endif
+}
+
+#if WITH_EDITOR
+void UMapBuildDataRegistry::HandleAssetPostCompileEvent(const TArray<FAssetCompileData>& CompiledAssets)
+{
+	TSet<FLightmapResourceCluster*> ClustersToUpdate;
+	for (const FAssetCompileData& CompileData : CompiledAssets)
+	{
+		if (ULightMapVirtualTexture2D* LightMapVirtualTexture2D = Cast<ULightMapVirtualTexture2D>(CompileData.Asset.Get()))
+		{
+			// If our lightmap clusters are affected by the virtual textures that just finished compiling, 
+			// we need to update their uniform buffer.
+			for (FLightmapResourceCluster& Cluster : LightmapResourceClusters)
+			{
+				if (Cluster.Input.LightMapVirtualTextures[0] == LightMapVirtualTexture2D ||
+					Cluster.Input.LightMapVirtualTextures[1] == LightMapVirtualTexture2D)
+				{
+					ClustersToUpdate.Add(&Cluster);
+				}
+			}
+		}
+	}
+
+	if (ClustersToUpdate.Num())
+	{
+		TRACE_CPUPROFILER_EVENT_SCOPE(UMapBuildDataRegistry::HandleAssetPostCompileEvent);
+
+		for (FLightmapResourceCluster* Cluster : ClustersToUpdate)
+		{
+			ENQUEUE_RENDER_COMMAND(UpdateClusterUniformBuffer)(
+				[Cluster](FRHICommandList& RHICmdList)
+				{
+					Cluster->UpdateUniformBuffer_RenderThread();
+				});
+		}
+
+		for (TObjectIterator<ULandscapeComponent> It; It; ++It)
+		{
+			if (It->IsRenderStateCreated() && It->SceneProxy != nullptr)
+			{
+				if (FMeshMapBuildData* BuildData = MeshBuildData.Find(It->MapBuildDataId))
+				{
+					if (ClustersToUpdate.Contains(BuildData->ResourceCluster))
+					{
+						It->MarkRenderStateDirty();
+					}
+				}
+			}
+		}
+
+		for (TObjectIterator<UStaticMeshComponent> It; It; ++It)
+		{
+			if (It->IsRenderStateCreated() && It->SceneProxy != nullptr)
+			{
+				for (const FStaticMeshComponentLODInfo& LODInfo : It->LODData)
+				{
+					if (FMeshMapBuildData* BuildData = MeshBuildData.Find(LODInfo.MapBuildDataId))
+					{
+						if (ClustersToUpdate.Contains(BuildData->ResourceCluster))
+						{
+							It->MarkRenderStateDirty();
+							break;
+						}
+					}
+				}
+			}
+		}
+	}
+
 }
+#endif // #if WITH_EDITOR
 
 void UMapBuildDataRegistry::Serialize(FArchive& Ar)
 {
