Added UTickableWorldSubsystem as a base class for world sub systems that need to be ticked along with their world.
By default:
* It implements FTickableGameObject::GetTickableGameObjectWorld to return the subsystem's world
* It prevents the subsystem from ticking as soon as the subsystem is deinitialized (has required custom solutions in the past to avoid crashes when UWaterSubsystem and UAutoDestroySubsystem get ticked after their world is pending destroy)
* It prevents the subsystem's CDO from ever ticking

diff --git a/Engine/Plugins/Experimental/Water/Source/Runtime/Private/WaterSubsystem.cpp b/Engine/Plugins/Experimental/Water/Source/Runtime/Private/WaterSubsystem.cpp
--- a/Engine/Plugins/Experimental/Water/Source/Runtime/Private/WaterSubsystem.cpp
+++ b/Engine/Plugins/Experimental/Water/Source/Runtime/Private/WaterSubsystem.cpp
@@ -145,7 +145,6 @@ UWaterSubsystem::UWaterSubsystem()
 	PrevWorldTimeSeconds = 0.f;
 	bUnderWaterForAudio = false;
 	bPauseWaveTime = false;
-	bInitialized = false;
 
 	struct FConstructorStatics
 	{
@@ -186,7 +185,8 @@ FWaterBodyManager* UWaterSubsystem::GetWaterBodyManager(UWorld* InWorld)
 
 void UWaterSubsystem::Tick(float DeltaTime)
 {
-	check(bInitialized); // Ticking should have been disabled for an uninitialized subsystem
+	Super::Tick(DeltaTime);
+
 	check(GetWorld() != nullptr);
 	if (FreezeWaves == 0 && bPauseWaveTime == false)
 	{
@@ -236,6 +236,8 @@ bool UWaterSubsystem::DoesSupportWorldType(EWorldType::Type WorldType) const
 
 void UWaterSubsystem::Initialize(FSubsystemCollectionBase& Collection)
 {
+	Super::Initialize(Collection);
+
 	UWorld* World = GetWorld();
 	check(World != nullptr);
 
@@ -265,8 +267,6 @@ void UWaterSubsystem::Initialize(FSubsystemCollectionBase& Collection)
 	}
 	UCollisionProfile::Get()->OnLoadProfileConfig.AddUObject(this, &UWaterSubsystem::OnLoadProfileConfig);
 	AddWaterCollisionProfile();
-
-	bInitialized = true;
 }
 
 void UWaterSubsystem::Deinitialize()
@@ -290,7 +290,7 @@ void UWaterSubsystem::Deinitialize()
 	GetDefault<UWaterRuntimeSettings>()->OnSettingsChange.RemoveAll(this);
 #endif //WITH_EDITOR
 
-	bInitialized = false;
+	Super::Deinitialize();
 }
 
 void UWaterSubsystem::ApplyRuntimeSettings(const UWaterRuntimeSettings* Settings, EPropertyChangeType::Type ChangeType)
diff --git a/Engine/Plugins/Experimental/Water/Source/Runtime/Public/WaterSubsystem.h b/Engine/Plugins/Experimental/Water/Source/Runtime/Public/WaterSubsystem.h
--- a/Engine/Plugins/Experimental/Water/Source/Runtime/Public/WaterSubsystem.h
+++ b/Engine/Plugins/Experimental/Water/Source/Runtime/Public/WaterSubsystem.h
@@ -6,7 +6,6 @@
 #include "UObject/ObjectMacros.h"
 #include "Subsystems/WorldSubsystem.h"
 #include "Engine/EngineTypes.h"
-#include "Engine/Public/Tickable.h"
 #include "Interfaces/Interface_PostProcessVolume.h"
 #include "WaterBodyManager.h"
 #include "WaterSubsystem.generated.h"
@@ -58,7 +57,7 @@ struct FUnderwaterPostProcessVolume : public IInterface_PostProcessVolume
  * This is the API used to get information about water at runtime
  */
 UCLASS(BlueprintType, Transient)
-class WATER_API UWaterSubsystem : public UWorldSubsystem, public FTickableGameObject
+class WATER_API UWaterSubsystem : public UTickableWorldSubsystem
 {
 	GENERATED_BODY()
 
@@ -66,9 +65,8 @@ public:
 	UWaterSubsystem();
 
 	// FTickableGameObject implementation Begin
-	virtual UWorld* GetTickableGameObjectWorld() const override { return GetWorld(); }
-	virtual bool IsTickable() const override { return bInitialized; }
-	virtual bool IsTickableInEditor() const override { return bInitialized; }
+	virtual bool IsTickable() const override { return true; }
+	virtual bool IsTickableInEditor() const override { return true; }
 	virtual void Tick(float DeltaTime) override;
 	virtual TStatId GetStatId() const override;
 	// FTickableGameObject implementation End
@@ -227,7 +225,6 @@ private:
 	bool bUsingOverrideWorldTimeSeconds;
 	bool bUnderWaterForAudio;
 	bool bPauseWaveTime;
-	bool bInitialized;
 
 	/** The parameter collection asset that holds the global parameters that are updated by this actor */
 	UPROPERTY()
diff --git a/Engine/Source/Runtime/AudioAnalyzer/Classes/AudioAnalyzerSubsystem.h b/Engine/Source/Runtime/AudioAnalyzer/Classes/AudioAnalyzerSubsystem.h
--- a/Engine/Source/Runtime/AudioAnalyzer/Classes/AudioAnalyzerSubsystem.h
+++ b/Engine/Source/Runtime/AudioAnalyzer/Classes/AudioAnalyzerSubsystem.h
@@ -5,7 +5,6 @@
 #include "CoreMinimal.h"
 #include "AudioAnalyzer.h"
 #include "Subsystems/WorldSubsystem.h"
-#include "Tickable.h"
 #include "AudioAnalyzerSubsystem.generated.h"
 
 class UWorld;
@@ -14,7 +13,7 @@ class UWorld;
 * Class manages querying analysis results from various audio analyzers.
 */
 UCLASS(DisplayName = "Audio Analysis")
-class AUDIOANALYZER_API UAudioAnalyzerSubsystem : public UWorldSubsystem, public FTickableGameObject
+class AUDIOANALYZER_API UAudioAnalyzerSubsystem : public UTickableWorldSubsystem
 {
 	GENERATED_BODY()
 
@@ -25,7 +24,7 @@ public:
 	//~ Begin FTickableGameObject Interface
 	virtual void Tick(float DeltaTime) override;
 	virtual bool IsTickable() const override;
-	virtual bool IsTickableInEditor() const { return true; }
+	virtual bool IsTickableInEditor() const override { return true; }
 	virtual TStatId GetStatId() const override;
 	//~ End FTickableGameObject Interface
 
diff --git a/Engine/Source/Runtime/AudioAnalyzer/Private/AudioAnalyzerSubsystem.cpp b/Engine/Source/Runtime/AudioAnalyzer/Private/AudioAnalyzerSubsystem.cpp
--- a/Engine/Source/Runtime/AudioAnalyzer/Private/AudioAnalyzerSubsystem.cpp
+++ b/Engine/Source/Runtime/AudioAnalyzer/Private/AudioAnalyzerSubsystem.cpp
@@ -28,6 +28,8 @@ TStatId UAudioAnalyzerSubsystem::GetStatId() const
 
 void UAudioAnalyzerSubsystem::Tick(float DeltaTime)
 {
+	Super::Tick(DeltaTime);
+
 	// Loop through all analyzers and if they're ready to analyze, do it
 	for (UAudioAnalyzer* Analyzer : AudioAnalyzers)
 	{
diff --git a/Engine/Source/Runtime/AudioMixer/Private/Quartz/QuartzSubsystem.cpp b/Engine/Source/Runtime/AudioMixer/Private/Quartz/QuartzSubsystem.cpp
--- a/Engine/Source/Runtime/AudioMixer/Private/Quartz/QuartzSubsystem.cpp
+++ b/Engine/Source/Runtime/AudioMixer/Private/Quartz/QuartzSubsystem.cpp
@@ -5,6 +5,7 @@
 #include "Quartz/QuartzMetronome.h"
 #include "Quartz/AudioMixerClockManager.h"
 #include "Sound/QuartzQuantizationUtilities.h"
+#include "Stats/Stats.h"
 
 #include "AudioDevice.h"
 #include "AudioMixerDevice.h"
@@ -66,6 +67,8 @@ UQuartzSubsystem::~UQuartzSubsystem()
 
 void UQuartzSubsystem::Tick(float DeltaTime)
 {
+	Super::Tick(DeltaTime);
+
 	const int32 NumSubscribers = QuartzTickSubscribers.Num();
 
 	if (DisableQuartzCvar)
@@ -137,7 +140,7 @@ bool UQuartzSubsystem::IsTickable() const
 
 TStatId UQuartzSubsystem::GetStatId() const
 {
-	return Super::GetStatID();
+	RETURN_QUICK_DECLARE_CYCLE_STAT(UQuartzSubsystem, STATGROUP_Tickables);
 }
 
 
diff --git a/Engine/Source/Runtime/AudioMixer/Public/Quartz/QuartzSubsystem.h b/Engine/Source/Runtime/AudioMixer/Public/Quartz/QuartzSubsystem.h
--- a/Engine/Source/Runtime/AudioMixer/Public/Quartz/QuartzSubsystem.h
+++ b/Engine/Source/Runtime/AudioMixer/Public/Quartz/QuartzSubsystem.h
@@ -8,7 +8,6 @@
 #include "UObject/Class.h"
 #include "Subsystems/WorldSubsystem.h"
 #include "Sound/QuartzQuantizationUtilities.h"
-#include "Tickable.h"
 
 #include "QuartzSubsystem.generated.h"
 
@@ -24,7 +23,7 @@ using MetronomeCommandQueuePtr = TSharedPtr<Audio::FShareableQuartzCommandQueue,
 
 
 UCLASS(DisplayName = "Quartz")
-class AUDIOMIXER_API UQuartzSubsystem : public UWorldSubsystem, public FQuartLatencyTracker, public FTickableGameObject
+class AUDIOMIXER_API UQuartzSubsystem : public UTickableWorldSubsystem, public FQuartLatencyTracker
 {
 	GENERATED_BODY()
 
diff --git a/Engine/Source/Runtime/Engine/Classes/Engine/AutoDestroySubsystem.h b/Engine/Source/Runtime/Engine/Classes/Engine/AutoDestroySubsystem.h
--- a/Engine/Source/Runtime/Engine/Classes/Engine/AutoDestroySubsystem.h
+++ b/Engine/Source/Runtime/Engine/Classes/Engine/AutoDestroySubsystem.h
@@ -3,7 +3,6 @@
 #pragma once
 
 #include "Subsystems/WorldSubsystem.h"
-#include "Tickable.h"
 #include "Stats/Stats.h"
 #include "Engine/EngineTypes.h"
 
@@ -18,7 +17,7 @@ struct FLatentActionManager;
 * get properly destroyed, as well as decouple this behavior from AActor::Tick
 */
 UCLASS()
-class UAutoDestroySubsystem : public UWorldSubsystem, public FTickableGameObject
+class UAutoDestroySubsystem : public UTickableWorldSubsystem
 {
 	GENERATED_BODY()
 
diff --git a/Engine/Source/Runtime/Engine/Private/Engine/AutoDestroySubsystem.cpp b/Engine/Source/Runtime/Engine/Private/Engine/AutoDestroySubsystem.cpp
--- a/Engine/Source/Runtime/Engine/Private/Engine/AutoDestroySubsystem.cpp
+++ b/Engine/Source/Runtime/Engine/Private/Engine/AutoDestroySubsystem.cpp
@@ -19,6 +19,8 @@ void UAutoDestroySubsystem::Deinitialize()
 		}
 	}
 	ActorsToPoll.Empty();
+
+	Super::Deinitialize();
 }
 
 bool UAutoDestroySubsystem::RegisterActor(AActor* ActorToRegister)
@@ -86,16 +88,9 @@ ETickableTickType UAutoDestroySubsystem::GetTickableTickType() const
 
 void UAutoDestroySubsystem::Tick(float DeltaTime)
 {
-	// Only tick if we know that the outer world is not being destroyed
-	// otherwise we will be added to the StaticTickables list with a dangling 
-	// world pointer (UE-96602)
-	UObject* Outer = GetOuter();
-	if (!Outer || Outer->HasAnyFlags(RF_BeginDestroyed))
-	{
-		return;
-	}
+	Super::Tick(DeltaTime);
 
-	UWorld* MyWorld = Cast<UWorld>(Outer);
+	UWorld* MyWorld = GetWorld();
 	
 	if (MyWorld && ActorsToPoll.Num())
 	{
diff --git a/Engine/Source/Runtime/Engine/Private/Subsystems/WorldSubsystem.cpp b/Engine/Source/Runtime/Engine/Private/Subsystems/WorldSubsystem.cpp
--- a/Engine/Source/Runtime/Engine/Private/Subsystems/WorldSubsystem.cpp
+++ b/Engine/Source/Runtime/Engine/Private/Subsystems/WorldSubsystem.cpp
@@ -3,6 +3,8 @@
 #include "Subsystems/WorldSubsystem.h"
 #include "Engine/World.h"
 
+// ----------------------------------------------------------------------------------
+
 UWorldSubsystem::UWorldSubsystem()
 	: USubsystem()
 {
@@ -29,4 +31,43 @@ bool UWorldSubsystem::ShouldCreateSubsystem(UObject* Outer) const
 bool UWorldSubsystem::DoesSupportWorldType(EWorldType::Type WorldType) const
 {
 	return WorldType == EWorldType::Game || WorldType == EWorldType::Editor || WorldType == EWorldType::PIE;
-}
\ No newline at end of file
+}
+
+
+// ----------------------------------------------------------------------------------
+
+UTickableWorldSubsystem::UTickableWorldSubsystem()
+{
+
+}
+
+ETickableTickType UTickableWorldSubsystem::GetTickableTickType() const 
+{ 
+	// By default (if the child class doesn't override GetTickableTickType), don't let CDOs ever tick: 
+	return IsTemplate() ? ETickableTickType::Never : FTickableGameObject::GetTickableTickType(); 
+}
+
+bool UTickableWorldSubsystem::IsAllowedToTick() const
+{
+	// No matter what IsTickable says, don't let CDOs or uninitialized world subsystems tick :
+	// Note: even if GetTickableTickType was overridden by the child class and returns something else than ETickableTickType::Never for CDOs, 
+	//  it's probably a mistake, so by default, don't allow ticking. If the child class really intends its CDO to tick, he can always override IsAllowedToTick...
+	return !IsTemplate() && bInitialized;
+}
+
+void UTickableWorldSubsystem::Tick(float DeltaTime)
+{
+	checkf(IsInitialized(), TEXT("Ticking should have been disabled for an uninitialized subsystem : remember to call IsInitialized in the subsystem's IsTickable, IsTickableInEditor and/or IsTickableWhenPaused implementation"));
+}
+
+void UTickableWorldSubsystem::Initialize(FSubsystemCollectionBase& Collection)
+{
+	check(!bInitialized);
+	bInitialized = true;
+}
+
+void UTickableWorldSubsystem::Deinitialize()
+{
+	check(bInitialized);
+	bInitialized = false;
+}
diff --git a/Engine/Source/Runtime/Engine/Private/Tickable.cpp b/Engine/Source/Runtime/Engine/Private/Tickable.cpp
--- a/Engine/Source/Runtime/Engine/Private/Tickable.cpp
+++ b/Engine/Source/Runtime/Engine/Private/Tickable.cpp
@@ -136,7 +136,9 @@ void FTickableGameObject::TickObjects(UWorld* World, const int32 InTickType, con
 			if (FTickableGameObject* TickableObject = static_cast<FTickableGameObject*>(TickableEntry.TickableObject))
 			{
 				// If it is tickable and in this world
-				if (((TickableEntry.TickType == ETickableTickType::Always) || TickableObject->IsTickable()) && (TickableObject->GetTickableGameObjectWorld() == World))
+				if (((TickableEntry.TickType == ETickableTickType::Always) || TickableObject->IsTickable()) 
+					&& (TickableObject->GetTickableGameObjectWorld() == World)
+					&& TickableObject->IsAllowedToTick())
 				{
 					const bool bIsGameWorld = InTickType == LEVELTICK_All || (World && World->IsGameWorld());
 					// If we are in editor and it is editor tickable, always tick
diff --git a/Engine/Source/Runtime/Engine/Private/WorldPartition/WorldPartitionSubsystem.cpp b/Engine/Source/Runtime/Engine/Private/WorldPartition/WorldPartitionSubsystem.cpp
--- a/Engine/Source/Runtime/Engine/Private/WorldPartition/WorldPartitionSubsystem.cpp
+++ b/Engine/Source/Runtime/Engine/Private/WorldPartition/WorldPartitionSubsystem.cpp
@@ -89,7 +89,7 @@ void UWorldPartitionSubsystem::Deinitialize()
 		check(!RegisteredWorldPartitions.Contains(WorldPartition));
 	}
 
-	UWorldSubsystem::Deinitialize();
+	Super::Deinitialize();
 }
 
 void UWorldPartitionSubsystem::RegisterWorldPartition(UWorldPartition* WorldPartition)
@@ -111,6 +111,8 @@ void UWorldPartitionSubsystem::UnregisterWorldPartition(UWorldPartition* WorldPa
 
 void UWorldPartitionSubsystem::Tick(float DeltaSeconds)
 {
+	Super::Tick(DeltaSeconds);
+
 	for (UWorldPartition* Partition : RegisteredWorldPartitions)
 	{
 		Partition->Tick(DeltaSeconds);
@@ -129,16 +131,6 @@ void UWorldPartitionSubsystem::Tick(float DeltaSeconds)
 	}
 }
 
-bool UWorldPartitionSubsystem::IsTickableInEditor() const
-{
-	return true;
-}
-
-UWorld* UWorldPartitionSubsystem::GetTickableGameObjectWorld() const
-{
-	return GetWorld();
-}
-
 ETickableTickType UWorldPartitionSubsystem::GetTickableTickType() const
 {
 	return IsTemplate() ? ETickableTickType::Never : ETickableTickType::Always;
diff --git a/Engine/Source/Runtime/Engine/Public/Subsystems/WorldSubsystem.h b/Engine/Source/Runtime/Engine/Public/Subsystems/WorldSubsystem.h
--- a/Engine/Source/Runtime/Engine/Public/Subsystems/WorldSubsystem.h
+++ b/Engine/Source/Runtime/Engine/Public/Subsystems/WorldSubsystem.h
@@ -4,6 +4,7 @@
 
 #include "Subsystem.h"
 #include "Engine/EngineTypes.h"
+#include "Engine/Public/Tickable.h"
 #include "WorldSubsystem.generated.h"
 
 /**
@@ -32,3 +33,36 @@ public:
 protected:
 	virtual bool DoesSupportWorldType(EWorldType::Type WorldType) const;
 };
+
+/**
+ * UTickableWorldSubsystem
+ * Base class for auto instanced and initialized systems that share the lifetime of a UWorld and are ticking along with it
+
+ */
+
+UCLASS(Abstract)
+class ENGINE_API UTickableWorldSubsystem : public UWorldSubsystem, public FTickableGameObject
+{
+	GENERATED_BODY()
+
+public:
+	UTickableWorldSubsystem();
+
+	// FTickableGameObject implementation Begin
+	virtual UWorld* GetTickableGameObjectWorld() const override { return GetWorld(); }
+	virtual ETickableTickType GetTickableTickType() const override;
+	virtual bool IsAllowedToTick() const override final;
+	virtual void Tick(float DeltaTime) override;
+	virtual TStatId GetStatId() const override PURE_VIRTUAL(UTickableWorldSubsystem::GetStatId, return TStatId(););
+	// FTickableGameObject implementation End
+
+	// USubsystem implementation Begin
+	virtual void Initialize(FSubsystemCollectionBase& Collection) override;
+	virtual void Deinitialize() override;
+	// USubsystem implementation End
+
+	bool IsInitialized() const { return bInitialized; }
+
+private:
+	bool bInitialized = false;
+};
\ No newline at end of file
diff --git a/Engine/Source/Runtime/Engine/Public/Tickable.h b/Engine/Source/Runtime/Engine/Public/Tickable.h
--- a/Engine/Source/Runtime/Engine/Public/Tickable.h
+++ b/Engine/Source/Runtime/Engine/Public/Tickable.h
@@ -79,6 +79,17 @@ public:
 	 */
 	virtual bool IsTickable() const { return true; }
 
+	/**
+	 * Virtual that can be overloaded by the inheriting class. It is
+	 * used to determine whether an object is allowed to tick. This is different than IsTickable
+	 * for keeping existing tickable objects working while providing an extra-safety for classes that decide to 
+	 * implement it (e.g. tickable world subsystems that have been deinitialized are not allowed to tick, regardless 
+	 * of what IsTickable says)
+	 *
+	 * @return	true if object is allowed to be ticked, false otherwise.
+	 */
+	virtual bool IsAllowedToTick() const { return true; }
+
 	/** return the stat id to use for this tickable **/
 	virtual TStatId GetStatId() const = 0;
 };
diff --git a/Engine/Source/Runtime/Engine/Public/WorldPartition/WorldPartitionSubsystem.h b/Engine/Source/Runtime/Engine/Public/WorldPartition/WorldPartitionSubsystem.h
--- a/Engine/Source/Runtime/Engine/Public/WorldPartition/WorldPartitionSubsystem.h
+++ b/Engine/Source/Runtime/Engine/Public/WorldPartition/WorldPartitionSubsystem.h
@@ -4,7 +4,6 @@
 
 #include "Subsystems/WorldSubsystem.h"
 #include "Templates/SubclassOf.h"
-#include "Tickable.h"
 #include "WorldPartitionSubsystem.generated.h"
 
 class UWorldPartition;
@@ -16,7 +15,7 @@ class FWorldPartitionActorDesc;
  */
 
 UCLASS()
-class ENGINE_API UWorldPartitionSubsystem : public UWorldSubsystem, public FTickableGameObject
+class ENGINE_API UWorldPartitionSubsystem : public UTickableWorldSubsystem
 {
 	GENERATED_BODY()
 
@@ -35,8 +34,7 @@ public:
 
 	//~ Begin FTickableGameObject
 	virtual void Tick(float DeltaSeconds) override;
-	virtual bool IsTickableInEditor() const override;
-	virtual UWorld* GetTickableGameObjectWorld() const override;
+	virtual bool IsTickableInEditor() const override { return true; }
 	virtual ETickableTickType GetTickableTickType() const override;
 	virtual TStatId GetStatId() const override;
 	//~End FTickableGameObject
diff --git a/Engine/Source/Runtime/Landscape/Private/LandscapeSubsystem.cpp b/Engine/Source/Runtime/Landscape/Private/LandscapeSubsystem.cpp
--- a/Engine/Source/Runtime/Landscape/Private/LandscapeSubsystem.cpp
+++ b/Engine/Source/Runtime/Landscape/Private/LandscapeSubsystem.cpp
@@ -94,6 +94,8 @@ void ULandscapeSubsystem::Tick(float DeltaTime)
 	CSV_SCOPED_TIMING_STAT_EXCLUSIVE(Landscape);
 	LLM_SCOPE(ELLMTag::Landscape);
 
+	Super::Tick(DeltaTime);
+
 	UWorld* World = GetWorld();
 
 	static TArray<FVector> OldCameras;
diff --git a/Engine/Source/Runtime/Landscape/Public/LandscapeSubsystem.h b/Engine/Source/Runtime/Landscape/Public/LandscapeSubsystem.h
--- a/Engine/Source/Runtime/Landscape/Public/LandscapeSubsystem.h
+++ b/Engine/Source/Runtime/Landscape/Public/LandscapeSubsystem.h
@@ -5,14 +5,13 @@
 #include "CoreMinimal.h"
 #include "UObject/ObjectMacros.h"
 #include "Subsystems/WorldSubsystem.h"
-#include "Tickable.h"
 #include "LandscapeSubsystem.generated.h"
 
 class ALandscapeProxy;
 class ULandscapeInfo;
 
 UCLASS(MinimalAPI)
-class ULandscapeSubsystem : public UWorldSubsystem, public FTickableGameObject
+class ULandscapeSubsystem : public UTickableWorldSubsystem
 {
 	GENERATED_BODY()
 
