Add cloth material conversion for Strata.

diff --git a/Engine/Source/Runtime/Engine/Private/Materials/Material.cpp b/Engine/Source/Runtime/Engine/Private/Materials/Material.cpp
--- a/Engine/Source/Runtime/Engine/Private/Materials/Material.cpp
+++ b/Engine/Source/Runtime/Engine/Private/Materials/Material.cpp
@@ -52,6 +52,8 @@
 #include "Materials/MaterialExpressionOneMinus.h"
 #include "Materials/MaterialExpressionMultiply.h"
 #include "Materials/MaterialExpressionSaturate.h"
+#include "Materials/MaterialExpressionConstant3Vector.h"
+#include "Materials/MaterialExpressionPower.h"
 #include "Engine/Font.h"
 #include "SceneManagement.h"
 #include "Materials/MaterialUniformExpressions.h"
@@ -3741,6 +3743,58 @@ void UMaterial::ConvertMaterialToStrataMaterial()
 
 			FrontMaterial.Connect(0, SlabBSDF);
 		}
+		else if (MaterialDomain == MD_Surface && ShadingModel == MSM_Cloth)
+		{
+			// Helpers
+			UMaterialExpressionConstant3Vector* ConstantZero3 = NewObject<UMaterialExpressionConstant3Vector>(this);
+			UMaterialExpressionConstant* ConstantZero1 = NewObject<UMaterialExpressionConstant>(this);
+			UMaterialExpressionPower* Power = NewObject<UMaterialExpressionPower>(this);
+			ConstantZero3->Constant = FLinearColor::Black;
+			ConstantZero1->R = 0.f;
+			Power->ConstExponent = 2.f;
+
+			// Sheen model affect only the specular component of a material, and lerp the slab specular with the Sheen specular. This is how the legacy cloth shading model works. In order to reproduce that, since we don't have fine grain control over the BSDF we:
+			// 1. Create a sheen BSDF
+			// 2. Create a slab with only specular value
+			// 3. Lerp the two slab based on the Fuzz amount
+			// 4. Create a slab with only diffuse value
+			// 5. Add the Diffuse slab and the previous other slab
+			UMaterialExpressionStrataSlabBSDF* DiffuseSlabBSDF = NewObject<UMaterialExpressionStrataSlabBSDF>(this);
+			UMaterialExpressionStrataSlabBSDF* SpecularSlabBSDF = NewObject<UMaterialExpressionStrataSlabBSDF>(this);
+			UMaterialExpressionStrataSheenBSDF* SheenBSDF = NewObject<UMaterialExpressionStrataSheenBSDF>(this);
+			
+			// Diffuse slab
+			MoveConnectionTo(BaseColor, DiffuseSlabBSDF, 0);		// BaseColor
+			DiffuseSlabBSDF->Metallic.Connect(0, ConstantZero1);	// Metallic
+			DiffuseSlabBSDF->Specular.Connect(0, ConstantZero1);	// Specular
+			CopyConnectionTo(Roughness, DiffuseSlabBSDF, 4);		// Roughness
+			CopyConnectionTo(Normal,	DiffuseSlabBSDF, 6);		// Normal
+
+			// Specular slab
+			SpecularSlabBSDF->BaseColor.Connect(0, ConstantZero3);	// BaseColor
+			SpecularSlabBSDF->Metallic.Connect(0, ConstantZero1);	// Metallic
+			CopyConnectionTo(Specular,	SpecularSlabBSDF, 3);		// Specular
+			CopyConnectionTo(Roughness, SpecularSlabBSDF, 4);		// Roughness
+			CopyConnectionTo(Normal,	SpecularSlabBSDF, 6);		// Normal
+
+			// Sheen 
+			MoveConnectionTo(SubsurfaceColor, Power, 0);			// Square the base color to match the legacy code
+			SheenBSDF->BaseColor.Connect(0, Power);					// BaseColor
+			CopyConnectionTo(Roughness, SheenBSDF, 1);				// Roughness
+			CopyConnectionTo(Normal, SheenBSDF, 2);					// Normal
+			
+			UMaterialExpressionStrataHorizontalMixing* Mix = NewObject<UMaterialExpressionStrataHorizontalMixing>(this);
+			UMaterialExpressionStrataAdd* Add = NewObject<UMaterialExpressionStrataAdd>(this);
+
+			Mix->Background.Connect(0, SpecularSlabBSDF);			// Background
+			Mix->Foreground.Connect(0, SheenBSDF);					// Foreground
+			MoveConnectionTo(ClearCoat, Mix, 2);					// Mix
+
+			Add->A.Connect(0, Mix);									// A
+			Add->B.Connect(0, DiffuseSlabBSDF);						// B
+
+			FrontMaterial.Connect(0, Add);
+		}
 		else if (MaterialDomain == MD_Surface && ShadingModel == MSM_Hair)
 		{
 			UMaterialExpressionStrataHairBSDF* HairBSDF = NewObject<UMaterialExpressionStrataHairBSDF>(this);
