Remove Pragma submitted by error

diff --git a/Engine/Source/Runtime/Engine/Private/WorldPartition/RuntimeSpatialHash/RuntimeSpatialHashHLOD.cpp b/Engine/Source/Runtime/Engine/Private/WorldPartition/RuntimeSpatialHash/RuntimeSpatialHashHLOD.cpp
--- a/Engine/Source/Runtime/Engine/Private/WorldPartition/RuntimeSpatialHash/RuntimeSpatialHashHLOD.cpp
+++ b/Engine/Source/Runtime/Engine/Private/WorldPartition/RuntimeSpatialHash/RuntimeSpatialHashHLOD.cpp
@@ -92,8 +92,6 @@ static void DeletePackage(FWorldPartitionActorDesc* ActorDesc, ISourceControlHel
 	}
 }
 
-PRAGMA_DISABLE_OPTIMIZATION
-
 static TArray<FGuid> GenerateHLODsForGrid(UWorldPartition* WorldPartition, const FSpatialHashRuntimeGrid& RuntimeGrid, uint32 HLODLevel, FHLODCreationContext& Context, ISourceControlHelper* SourceControlHelper, bool bCreateActorsOnly, const TArray<const FActorClusterInstance*>& ClusterInstances)
 {
 	auto HasExceededMaxMemory = []()
@@ -256,8 +254,6 @@ static TArray<FGuid> GenerateHLODsForGrid(UWorldPartition* WorldPartition, const
 	return GridHLODActors;
 }
 
-PRAGMA_ENABLE_OPTIMIZATION
-
 // Find all HLOD grids from the HLODLayer assets we have and build a dependency graph
 static void GatherHLODGrids(TMap<FName, FSpatialHashRuntimeGrid>& OutHLODGrids, TMap<FName, TSet<FName>>& OutHLODGridsGraph)
 {
