HLODBuilder - Reduce log verbosity

diff --git a/Engine/Source/Runtime/Engine/Private/WorldPartition/HLOD/HLODBuilder.cpp b/Engine/Source/Runtime/Engine/Private/WorldPartition/HLOD/HLODBuilder.cpp
--- a/Engine/Source/Runtime/Engine/Private/WorldPartition/HLOD/HLODBuilder.cpp
+++ b/Engine/Source/Runtime/Engine/Private/WorldPartition/HLOD/HLODBuilder.cpp
@@ -156,7 +156,7 @@ public:
 
 		// HLOD Layer
 		uint32 HLODLayerHash = InHLODActor->GetSubActorsHLODLayer()->GetCRC();
-		UE_LOG(LogHLODBuilder, Log, TEXT(" - HLODLayer (%s) = %x"), *InHLODActor->GetSubActorsHLODLayer()->GetName(), HLODLayerHash);
+		UE_LOG(LogHLODBuilder, VeryVerbose, TEXT(" - HLODLayer (%s) = %x"), *InHLODActor->GetSubActorsHLODLayer()->GetName(), HLODLayerHash);
 		Ar << HLODLayerHash;
 
 		// We get the CRC of each component
@@ -167,18 +167,18 @@ public:
 			{
 				uint32 ComponentCRC = 0;
 
-				UE_LOG(LogHLODBuilder, Log, TEXT(" - Component \'%s\' from actor \'%s\'"), *Component->GetName(), *Component->GetOwner()->GetName());
+				UE_LOG(LogHLODBuilder, VeryVerbose, TEXT(" - Component \'%s\' from actor \'%s\'"), *Component->GetName(), *Component->GetOwner()->GetName());
 
 				// CRC component
 				uint32 StaticMeshComponentCRC = UHLODProxy::GetCRC(StaticMeshComponent);
-				UE_LOG(LogHLODBuilder, Log, TEXT("     - StaticMeshComponent (%s) = %x"), *StaticMeshComponent->GetName(), StaticMeshComponentCRC);
+				UE_LOG(LogHLODBuilder, VeryVerbose, TEXT("     - StaticMeshComponent (%s) = %x"), *StaticMeshComponent->GetName(), StaticMeshComponentCRC);
 				ComponentCRC = HashCombine(ComponentCRC, StaticMeshComponentCRC);
 
 				if (UStaticMesh* StaticMesh = StaticMeshComponent->GetStaticMesh())
 				{
 					// CRC static mesh
 					int32 StaticMeshCRC = UHLODProxy::GetCRC(StaticMesh);
-					UE_LOG(LogHLODBuilder, Log, TEXT("     - StaticMesh (%s) = %x"), *StaticMesh->GetName(), StaticMeshCRC);
+					UE_LOG(LogHLODBuilder, VeryVerbose, TEXT("     - StaticMesh (%s) = %x"), *StaticMesh->GetName(), StaticMeshCRC);
 					ComponentCRC = HashCombine(ComponentCRC, StaticMeshCRC);
 
 					// CRC materials
@@ -189,7 +189,7 @@ public:
 						if (MaterialInterface)
 						{
 							uint32 MaterialInterfaceCRC = UHLODProxy::GetCRC(MaterialInterface);
-							UE_LOG(LogHLODBuilder, Log, TEXT("     - MaterialInterface (%s) = %x"), *MaterialInterface->GetName(), MaterialInterfaceCRC);
+							UE_LOG(LogHLODBuilder, VeryVerbose, TEXT("     - MaterialInterface (%s) = %x"), *MaterialInterface->GetName(), MaterialInterfaceCRC);
 							ComponentCRC = HashCombine(ComponentCRC, MaterialInterfaceCRC);
 
 							TArray<UTexture*> Textures;
@@ -197,7 +197,7 @@ public:
 							for (UTexture* Texture : Textures)
 							{
 								uint32 TextureCRC = UHLODProxy::GetCRC(Texture);
-								UE_LOG(LogHLODBuilder, Log, TEXT("     - Texture (%s) = %x"), *Texture->GetName(), TextureCRC);
+								UE_LOG(LogHLODBuilder, VeryVerbose, TEXT("     - Texture (%s) = %x"), *Texture->GetName(), TextureCRC);
 								ComponentCRC = HashCombine(ComponentCRC, TextureCRC);
 							}
 						}
