Removed ensure after diagnosing the issue further and concluding it is not achieving anything.

diff --git a/Engine/Source/Runtime/Engine/Private/InstancedStaticMesh.cpp b/Engine/Source/Runtime/Engine/Private/InstancedStaticMesh.cpp
--- a/Engine/Source/Runtime/Engine/Private/InstancedStaticMesh.cpp
+++ b/Engine/Source/Runtime/Engine/Private/InstancedStaticMesh.cpp
@@ -1082,7 +1082,9 @@ void FInstancedStaticMeshSceneProxy::SetupProxy(UInstancedStaticMeshComponent* I
 		for (int32 InInstanceIndex = 0; InInstanceIndex < Instances.Num(); ++InInstanceIndex)
 		{
 			int32 OutInstanceIndex = InInstanceIndex;
-			if (OutInstanceIndex < InstanceReorderTable.Num() && ensure(InstanceReorderTable[OutInstanceIndex] < Instances.Num()))
+			// GPUCULL_TODO: After deleting instances in a HISM the InstanceReorderTable often contains nonsense, this is then corrected
+			// by the async build, which re-creates the proxy in a nearby future frame. All this should be removed in favour of GPU-side culling.
+			if (OutInstanceIndex < InstanceReorderTable.Num() && InstanceReorderTable[OutInstanceIndex] < Instances.Num())
 			{
 				OutInstanceIndex = InstanceReorderTable[OutInstanceIndex];
 			}
