Update non-strata SSS material scaling constant for MFP to have better match.

diff --git a/Engine/Source/Runtime/Engine/Private/Materials/Material.cpp b/Engine/Source/Runtime/Engine/Private/Materials/Material.cpp
--- a/Engine/Source/Runtime/Engine/Private/Materials/Material.cpp
+++ b/Engine/Source/Runtime/Engine/Private/Materials/Material.cpp
@@ -3732,7 +3732,7 @@ void UMaterial::ConvertMaterialToStrataMaterial()
 				// Subsurface color is converted into a MFP by taking the subsurface color and multipling it with the 1-Opacity 
 				// and a constant scaling for matching roughly the behavior
 				UMaterialExpressionConstant* MFPConstantScale = NewObject<UMaterialExpressionConstant>(this);
-				MFPConstantScale->R = 2.f;
+				MFPConstantScale->R = 0.5f;
 
 				UMaterialExpressionMultiply* MulScaleMFP = NewObject<UMaterialExpressionMultiply>(this);
 				MulScaleMFP->A.Connect(0, MulOpacityWithSubsurfaceColor);
