Material conversion: default lit now converts MSM_DefaultLit + comment clean up.

diff --git a/Engine/Source/Runtime/Engine/Private/Materials/Material.cpp b/Engine/Source/Runtime/Engine/Private/Materials/Material.cpp
--- a/Engine/Source/Runtime/Engine/Private/Materials/Material.cpp
+++ b/Engine/Source/Runtime/Engine/Private/Materials/Material.cpp
@@ -3648,6 +3648,19 @@ void UMaterial::ConvertMaterialToStrataMaterial()
 			}
 		};
 
+		// STRATA_TODO for mateiral conversion
+		//  - WorldPositionOffset can remain on the end point node
+		//  - WorldDisplacement and TessellationMultiplier are going to be removed and should be ignored?
+		//  - ShadingModel
+		//  - Refraction
+		//  - PixelDepthOffset
+		//	- MSM_Subsurface
+		//	- MSM_PreintegratedSkin
+		//	- MSM_SubsurfaceProfile
+		//	- MSM_TwoSidedFoliage
+		//	- MSM_Cloth
+		//	- MSM_Eye
+
 		if (MaterialDomain == MD_Surface && ShadingModel == MSM_Unlit)
 		{
 			UMaterialExpressionStrataUnlitBSDF* UnlitBSDF = NewObject<UMaterialExpressionStrataUnlitBSDF>(this);
@@ -3662,23 +3675,10 @@ void UMaterial::ConvertMaterialToStrataMaterial()
 			MoveConnectionTo(Metallic, SlabBSDF, 2);		// Metallic
 			MoveConnectionTo(Specular, SlabBSDF, 3);		// Specular
 			MoveConnectionTo(Roughness, SlabBSDF, 4);		// Roughness
-			if (Anisotropy.IsConnected())
-			{
-				MoveConnectionTo(Anisotropy, SlabBSDF, 5);	// Anisotropy
-			}
+			MoveConnectionTo(Anisotropy, SlabBSDF, 5);		// Anisotropy
 			MoveConnectionTo(Normal, SlabBSDF, 6);			// Normal
 			MoveConnectionTo(Tangent, SlabBSDF, 7);			// Tangent
-			// Opacity can remain on the end point node
-			// Opacity mask can remain on the end point node
-			// WorldPositionOffset can remain on the end point node
-			// STRATA_TODO WorldDisplacement and TessellationMultiplier are going to be removed and should be ignored?
-			// STRATA_TODO ShadingModel
-			// STRATA_TODO ClearCoat
-			// STRATA_TODO ClearCoatRoughness
-			// STRATA_TODO AmbientOcclusion should be removed and put on the BSDF node
-			// STRATA_TODO Refraction
-			// STRATA_TODO PixelDepthOffset
-			// STRATA_TODO EmissiveColor
+			MoveConnectionTo(EmissiveColor, SlabBSDF, 10);	// Emissive
 
 			FrontMaterial.Connect(0, SlabBSDF);
 		}
@@ -3988,14 +3988,6 @@ void UMaterial::ConvertMaterialToStrataMaterial()
 			FrontMaterial.Connect(0, VolBSDF);
 		}
 
-		// STRATA_TODO Other conversion, see EMaterialShadingModel: 
-		//	- MSM_Subsurface
-		//	- MSM_PreintegratedSkin
-		//	- MSM_SubsurfaceProfile
-		//	- MSM_TwoSidedFoliage
-		//	- MSM_Cloth
-		//	- MSM_Eye
-
 		// Now force the material to recompile and we use a hash of the original StateId.
 		// This is to avoid having different StateId each time we load the material and to not forever recompile it, i.e. use a cached version.
 		uint32 HashBuffer[5];
