Add a console variable for setting mobile anti-aliasing method, remove the r.MobileMSAA and get the MSAA sample count from r.MSAACount.

diff --git a/Engine/Plugins/Runtime/GoogleVR/GoogleVRHMD/Source/GoogleVRHMD/Private/GoogleVRHMDRender.cpp b/Engine/Plugins/Runtime/GoogleVR/GoogleVRHMD/Source/GoogleVRHMD/Private/GoogleVRHMDRender.cpp
--- a/Engine/Plugins/Runtime/GoogleVR/GoogleVRHMD/Source/GoogleVRHMD/Private/GoogleVRHMDRender.cpp
+++ b/Engine/Plugins/Runtime/GoogleVR/GoogleVRHMD/Source/GoogleVRHMD/Private/GoogleVRHMDRender.cpp
@@ -433,12 +433,11 @@ namespace {
 	int32 GetMobileMSAASampleSetting()
 	{
 		static const int32 MaxMSAASamplesSupported = FOpenGL::GetMaxMSAASamplesTileMem();
-		static const auto CVarMobileMSAA = IConsoleManager::Get().FindTConsoleVariableDataInt(TEXT("r.MobileMSAA"));
-		static const int32 CVarMobileMSAAValue = CVarMobileMSAA->GetValueOnRenderThread();
+		static const int32 CVarMobileMSAAValue = GetDefaultMSAACount(ERHIFeatureLevel::ES3_1, GDynamicRHI->RHIGetPlatformTextureMaxSampleCount());
 		static const int32 MobileMSAAValue = FMath::Min(CVarMobileMSAAValue, MaxMSAASamplesSupported);
 		if (MobileMSAAValue != CVarMobileMSAAValue)
 		{
-			UE_LOG(LogHMD, Warning, TEXT("r.MobileMSAA is set to %i but we are using %i due to hardware support limitations."), CVarMobileMSAAValue, MobileMSAAValue);
+			UE_LOG(LogHMD, Warning, TEXT("MobileMSAA is set to %i but we are using %i due to hardware support limitations."), CVarMobileMSAAValue, MobileMSAAValue);
 		}
 		return MobileMSAAValue;
 	}
diff --git a/Engine/Plugins/Runtime/Oculus/OculusVR/Source/OculusEditor/Private/OculusToolWidget.cpp b/Engine/Plugins/Runtime/Oculus/OculusVR/Source/OculusEditor/Private/OculusToolWidget.cpp
--- a/Engine/Plugins/Runtime/Oculus/OculusVR/Source/OculusEditor/Private/OculusToolWidget.cpp
+++ b/Engine/Plugins/Runtime/Oculus/OculusVR/Source/OculusEditor/Private/OculusToolWidget.cpp
@@ -705,15 +705,17 @@ EVisibility SOculusToolWidget::MobileMultiViewVisibility(FName tag) const
 FReply SOculusToolWidget::MobileMSAAEnable(bool text)
 {
 	URendererSettings* Settings = GetMutableDefault<URendererSettings>();
-	Settings->MobileMSAASampleCount = EMobileMSAASampleCount::Four;
-	Settings->UpdateSinglePropertyInConfigFile(Settings->GetClass()->FindPropertyByName(GET_MEMBER_NAME_CHECKED(URendererSettings, MobileMSAASampleCount)), Settings->GetDefaultConfigFilename());
+	Settings->MobileAntiAliasing = EAntiAliasingMethod::AAM_MSAA;
+	Settings->UpdateSinglePropertyInConfigFile(Settings->GetClass()->FindPropertyByName(GET_MEMBER_NAME_CHECKED(URendererSettings, MobileAntiAliasing)), Settings->GetDefaultConfigFilename());
+	Settings->MSAASampleCount = ECompositingSampleCount::Four;
+	Settings->UpdateSinglePropertyInConfigFile(Settings->GetClass()->FindPropertyByName(GET_MEMBER_NAME_CHECKED(URendererSettings, MSAASampleCount)), Settings->GetDefaultConfigFilename());
 	return FReply::Handled();
 }
 
 EVisibility SOculusToolWidget::MobileMSAAVisibility(FName tag) const
 {
 	URendererSettings* Settings = GetMutableDefault<URendererSettings>();
-	const bool bMobileMSAAValid = Settings->MobileMSAASampleCount == EMobileMSAASampleCount::Four;
+	const bool bMobileMSAAValid = Settings->MobileAntiAliasing == EAntiAliasingMethod::AAM_MSAA && Settings->MSAASampleCount == ECompositingSampleCount::Four;
 
 	return bMobileMSAAValid ? EVisibility::Collapsed : EVisibility::Visible;
 }
diff --git a/Engine/Plugins/Runtime/Oculus/OculusVR/Source/OculusHMD/Private/OculusHMD_Layer.cpp b/Engine/Plugins/Runtime/Oculus/OculusVR/Source/OculusHMD/Private/OculusHMD_Layer.cpp
--- a/Engine/Plugins/Runtime/Oculus/OculusVR/Source/OculusHMD/Private/OculusHMD_Layer.cpp
+++ b/Engine/Plugins/Runtime/Oculus/OculusVR/Source/OculusHMD/Private/OculusHMD_Layer.cpp
@@ -549,8 +549,7 @@ void FLayer::Initialize_RenderThread(const FSettings* Settings, FCustomPresent*
 			uint32 NumSamplesTileMem = 1;
 			if (OvrpLayerDesc.Shape == ovrpShape_EyeFov)
 			{
-				static const auto CVarMobileMSAA = IConsoleManager::Get().FindTConsoleVariableDataInt(TEXT("r.MobileMSAA"));
-				NumSamplesTileMem = (CVarMobileMSAA ? CVarMobileMSAA->GetValueOnAnyThread() : 1);
+				NumSamplesTileMem = GetDefaultMSAACount(ERHIFeatureLevel::ES3_1, GDynamicRHI->RHIGetPlatformTextureMaxSampleCount());
 			}
 
 			ERHIResourceType ResourceType;			
diff --git a/Engine/Source/Runtime/Apple/MetalRHI/Private/MetalRHI.cpp b/Engine/Source/Runtime/Apple/MetalRHI/Private/MetalRHI.cpp
--- a/Engine/Source/Runtime/Apple/MetalRHI/Private/MetalRHI.cpp
+++ b/Engine/Source/Runtime/Apple/MetalRHI/Private/MetalRHI.cpp
@@ -1213,9 +1213,9 @@ void* FMetalDynamicRHI::RHIGetNativeInstance()
 
 uint16 FMetalDynamicRHI::RHIGetPlatformTextureMaxSampleCount()
 {
-	TArray<EMobileMSAASampleCount::Type> SamplesArray{ EMobileMSAASampleCount::Type::One, EMobileMSAASampleCount::Type::Two, EMobileMSAASampleCount::Type::Four, EMobileMSAASampleCount::Type::Eight };
+	TArray<ECompositingSampleCount::Type> SamplesArray{ ECompositingSampleCount::Type::One, ECompositingSampleCount::Type::Two, ECompositingSampleCount::Type::Four, ECompositingSampleCount::Type::Eight };
 
-	uint16 PlatformMaxSampleCount = EMobileMSAASampleCount::Type::One;
+	uint16 PlatformMaxSampleCount = ECompositingSampleCount::Type::One;
 	for (auto sampleIt = SamplesArray.CreateConstIterator(); sampleIt; ++sampleIt)
 	{
 		int sample = *sampleIt;
diff --git a/Engine/Source/Runtime/Core/Private/HAL/ConsoleManager.cpp b/Engine/Source/Runtime/Core/Private/HAL/ConsoleManager.cpp
--- a/Engine/Source/Runtime/Core/Private/HAL/ConsoleManager.cpp
+++ b/Engine/Source/Runtime/Core/Private/HAL/ConsoleManager.cpp
@@ -2811,3 +2811,13 @@ static TAutoConsoleVariable<int32> CVarSkyLightingQuality(
 	TEXT("  1: on\n"),
 	ECVF_Scalability | ECVF_RenderThreadSafe
 );
+
+static TAutoConsoleVariable<int32> CVarDefaultAntiAliasing(
+	TEXT("r.Mobile.AntiAliasing"),
+	2,
+	TEXT("Mobile default AntiAliasingMethod\n")
+	TEXT(" 0: off (no anti-aliasing)\n")
+	TEXT(" 1: FXAA (default, faster than TemporalAA but much more shimmering for non static cases)\n")
+	TEXT(" 2: TemporalAA\n")
+	TEXT(" 3: MSAA"),
+	ECVF_RenderThreadSafe);
\ No newline at end of file
diff --git a/Engine/Source/Runtime/Engine/Classes/Engine/RendererSettings.h b/Engine/Source/Runtime/Engine/Classes/Engine/RendererSettings.h
--- a/Engine/Source/Runtime/Engine/Classes/Engine/RendererSettings.h
+++ b/Engine/Source/Runtime/Engine/Classes/Engine/RendererSettings.h
@@ -47,22 +47,6 @@ namespace ECompositingSampleCount
 }
 
 
-/**
-* Enumerates available mobile MSAA sample counts.
-*/
-UENUM()
-namespace EMobileMSAASampleCount
-{
-	enum Type
-	{
-		One = 1 UMETA(DisplayName = "No MSAA"),
-		Two = 2 UMETA(DisplayName = "2x MSAA"),
-		Four = 4 UMETA(DisplayName = "4x MSAA"),
-		Eight = 8 UMETA(DisplayName = "8x MSAA"),
-	};
-}
-
-
 /**
  * Enumerates available options for custom depth.
  */
@@ -184,9 +168,9 @@ class ENGINE_API URendererSettings : public UDeveloperSettings
 		int32 MaxMobileCascades;
 
 	UPROPERTY(config, EditAnywhere, Category = Mobile, meta = (
-		ConsoleVariable = "r.MobileMSAA", DisplayName = "Mobile MSAA",
-		ToolTip = "Multi-sample anti-aliasing setting to use on mobile. MSAA is currently supported using Metal on iOS, and on Android devices with the required support using ES 2 or ES 3.1.\nIf MSAA is not available, the current default AA method will be used."))
-	TEnumAsByte<EMobileMSAASampleCount::Type> MobileMSAASampleCount;
+		ConsoleVariable = "r.Mobile.AntiAliasing", DisplayName = "Mobile Anti-Aliasing Method",
+		ToolTip = "The mobile default anti-aliasing method."))
+	TEnumAsByte<EAntiAliasingMethod> MobileAntiAliasing;
 
 	UPROPERTY(config, EditAnywhere, Category = Mobile, meta=(
 		ConsoleVariable="r.Mobile.AllowDitheredLODTransition", DisplayName="Allow Dithered LOD Transition",
@@ -520,6 +504,11 @@ class ENGINE_API URendererSettings : public UDeveloperSettings
 		ToolTip = "Which anti-aliasing mode is used by default"))
 	TEnumAsByte<EAntiAliasingMethod> DefaultFeatureAntiAliasing;
 
+	UPROPERTY(config, EditAnywhere, Category = DefaultSettings, meta = (
+		ConsoleVariable = "r.MSAACount", DisplayName = "MSAA Sample Count",
+		ToolTip = "Default number of samples for MSAA."))
+	TEnumAsByte<ECompositingSampleCount::Type> MSAASampleCount;
+
 	UPROPERTY(config, EditAnywhere, Category = DefaultSettings, meta = (
 		ConsoleVariable = "r.DefaultFeature.LightUnits", DisplayName = "Light Units",
 		ToolTip = "Which units to use for newly placed point, spot and rect lights"))
diff --git a/Engine/Source/Runtime/Engine/Classes/Engine/Scene.h b/Engine/Source/Runtime/Engine/Classes/Engine/Scene.h
--- a/Engine/Source/Runtime/Engine/Classes/Engine/Scene.h
+++ b/Engine/Source/Runtime/Engine/Classes/Engine/Scene.h
@@ -30,18 +30,6 @@ enum EDepthOfFieldMethod
 	DOFM_MAX,
 };
 
-/** Used by rendering project settings. */
-UENUM()
-enum EAntiAliasingMethod
-{
-	AAM_None UMETA(DisplayName="None"),
-	AAM_FXAA UMETA(DisplayName="FXAA"),
-	AAM_TemporalAA UMETA(DisplayName="TemporalAA"),
-	/** Only supported with forward shading.  MSAA sample count is controlled by r.MSAACount. */
-	AAM_MSAA UMETA(DisplayName="MSAA"),
-	AAM_MAX,
-};
-
 /** Used by FPostProcessSettings Auto Exposure */
 UENUM()
 enum EAutoExposureMethod
diff --git a/Engine/Source/Runtime/Engine/Private/SceneUtils.cpp b/Engine/Source/Runtime/Engine/Private/SceneUtils.cpp
--- a/Engine/Source/Runtime/Engine/Private/SceneUtils.cpp
+++ b/Engine/Source/Runtime/Engine/Private/SceneUtils.cpp
@@ -1,6 +1,11 @@
 // Copyright Epic Games, Inc. All Rights Reserved.
 
 #include "SceneUtils.h"
+#include "RHI.h"
+#include "Engine\Scene.h"
+#include "RenderUtils.h"
+
+DEFINE_LOG_CATEGORY(LogSceneUtils);
 
 bool IsMobileHDR()
 {
@@ -27,3 +32,138 @@ ENGINE_API bool IsMobileColorsRGB()
 
 	return !IsMobileHDR() && bMobileUseHWsRGBEncoding;
 }
+
+ENGINE_API EAntiAliasingMethod GetDefaultAntiAliasingMethod(const FStaticFeatureLevel InFeatureLevel)
+{
+	EAntiAliasingMethod AntiAliasingMethod = EAntiAliasingMethod::AAM_None;
+
+	if (InFeatureLevel == ERHIFeatureLevel::ES3_1)
+	{
+		static auto* MobileAntiAliasingCvar = IConsoleManager::Get().FindTConsoleVariableDataInt(TEXT("r.Mobile.AntiAliasing"));
+		AntiAliasingMethod = (EAntiAliasingMethod)MobileAntiAliasingCvar->GetValueOnAnyThread();
+
+		static auto* MobileHDRCvar = IConsoleManager::Get().FindTConsoleVariableDataInt(TEXT("r.MobileHDR"));
+		// Disable antialiasing in GammaLDR mode to avoid jittering.
+		if (MobileHDRCvar->GetValueOnAnyThread() == 0 && AntiAliasingMethod != EAntiAliasingMethod::AAM_MSAA)
+		{
+			AntiAliasingMethod = EAntiAliasingMethod::AAM_None;
+		}
+	}
+	else
+	{
+		static auto* DefaultAntiAliasingCvar = IConsoleManager::Get().FindTConsoleVariableDataInt(TEXT("r.DefaultFeature.AntiAliasing"));
+		AntiAliasingMethod = (EAntiAliasingMethod)DefaultAntiAliasingCvar->GetValueOnAnyThread();
+	}
+
+	static auto* MSAACountCVar = IConsoleManager::Get().FindTConsoleVariableDataInt(TEXT("r.MSAACount"));
+
+	static auto* PostProcessAAQualityCVar = IConsoleManager::Get().FindTConsoleVariableDataInt(TEXT("r.PostProcessAAQuality"));
+
+	EShaderPlatform ShaderPlatform = GetFeatureLevelShaderPlatform(InFeatureLevel);
+
+	if (AntiAliasingMethod == EAntiAliasingMethod::AAM_MSAA)
+	{
+		if (MSAACountCVar->GetValueOnAnyThread() <= 0)
+		{
+			// Fallback to temporal AA so we can easily toggle methods with r.MSAACount
+			AntiAliasingMethod = EAntiAliasingMethod::AAM_TemporalAA;
+		}
+		else if ((InFeatureLevel >= ERHIFeatureLevel::SM5 && !IsForwardShadingEnabled(ShaderPlatform))
+			|| (IsMobilePlatform(ShaderPlatform) && IsMobileDeferredShadingEnabled(ShaderPlatform)))
+		{
+			// Disable antialiasing for deferred renderer
+			AntiAliasingMethod = EAntiAliasingMethod::AAM_None;
+		}
+	}
+
+	int32 Quality = FMath::Clamp(PostProcessAAQualityCVar->GetValueOnAnyThread(), 0, 6);
+
+	if (Quality <= 0)
+	{
+		AntiAliasingMethod = EAntiAliasingMethod::AAM_None;
+	}
+
+	if (AntiAliasingMethod == AAM_TemporalAA && Quality < 3)
+	{
+		AntiAliasingMethod = AAM_FXAA;
+	}
+
+	return AntiAliasingMethod;
+}
+
+ENGINE_API uint32 GetDefaultMSAACount(const FStaticFeatureLevel InFeatureLevel, uint32 PlatformMaxSampleCount/* = 8*/)
+{
+	uint32 NumSamples = 1;
+
+	EAntiAliasingMethod AntiAliasingMethod = GetDefaultAntiAliasingMethod(InFeatureLevel);
+
+	if (AntiAliasingMethod == EAntiAliasingMethod::AAM_MSAA)
+	{
+		static auto* MSAACountCVar = IConsoleManager::Get().FindTConsoleVariableDataInt(TEXT("r.MSAACount"));
+
+		EShaderPlatform ShaderPlatform = GetFeatureLevelShaderPlatform(InFeatureLevel);
+
+		if (IsForwardShadingEnabled(ShaderPlatform) || (IsMobilePlatform(ShaderPlatform) && !IsMobileDeferredShadingEnabled(ShaderPlatform)))
+		{
+			NumSamples = FMath::Max(1, MSAACountCVar->GetValueOnAnyThread());
+
+
+			NumSamples = FMath::Min(NumSamples, PlatformMaxSampleCount);
+
+			if (NumSamples != 1 && NumSamples != 2 && NumSamples != 4 && NumSamples != 8)
+			{
+				UE_LOG(LogSceneUtils, Warning, TEXT("Requested %d samples for MSAA, but this is not supported; falling back to 1 sample"), NumSamples);
+				NumSamples = 1;
+			}
+		}
+
+		if (NumSamples > 1)
+		{
+			bool bRendererSupportMSAA = false;
+
+			FString FailedReason = TEXT("");
+
+			bool bRHISupportsMSAA = RHISupportsMSAA(ShaderPlatform);
+
+			if (InFeatureLevel == ERHIFeatureLevel::ES3_1)
+			{
+				// Disable MSAA if we are using mobile pixel projected reflection, since we have to resolve the SceneColor and SceneDepth after opaque base pass
+				bool bMobilePixelProjectedReflection = IsUsingMobilePixelProjectedReflection(ShaderPlatform);
+
+				// Disable MSAA if we are using mobile ambient occlusion, since we have to resolve the SceneColor and SceneDepth after opaque base pass
+				bool bMobileAmbientOcclusion = IsUsingMobileAmbientOcclusion(ShaderPlatform);
+
+				bRendererSupportMSAA = bRHISupportsMSAA && !bMobilePixelProjectedReflection && !bMobileAmbientOcclusion;
+
+				if (!bRendererSupportMSAA)
+				{
+					FailedReason = FString::Printf(TEXT("RHISupportsMSAA %d, MobilePixelProjectedReflection %d, MobileAmbientOcclusion %d"), bRHISupportsMSAA ? 1 : 0, bMobilePixelProjectedReflection ? 1 : 0, bMobileAmbientOcclusion ? 1 : 0);
+				}
+			}
+			else
+			{
+				bRendererSupportMSAA = bRHISupportsMSAA;
+
+				if (!bRendererSupportMSAA)
+				{
+					FailedReason = FString::Printf(TEXT("RHISupportsMSAA %d"), bRHISupportsMSAA ? 1 : 0);
+				}
+			}
+
+			if (!bRendererSupportMSAA)
+			{
+				NumSamples = 1;
+
+				static bool bWarned = false;
+
+				if (!bWarned)
+				{
+					bWarned = true;
+					UE_LOG(LogSceneUtils, Log, TEXT("Requested %d samples for MSAA, but the platform doesn't support MSAA, failed reason : %s"), NumSamples, *FailedReason);
+				}
+			}
+		}
+	}
+
+	return NumSamples;
+}
diff --git a/Engine/Source/Runtime/Engine/Private/SceneView.cpp b/Engine/Source/Runtime/Engine/Private/SceneView.cpp
--- a/Engine/Source/Runtime/Engine/Private/SceneView.cpp
+++ b/Engine/Source/Runtime/Engine/Private/SceneView.cpp
@@ -879,54 +879,25 @@ bool FSceneView::VerifyMembersChecks() const
 void FSceneView::SetupAntiAliasingMethod()
 {
 	{
-		int32 Value = CVarDefaultAntiAliasing.GetValueOnAnyThread();
+		int32 Value = GetDefaultAntiAliasingMethod(FeatureLevel);
 		if (Value >= 0 && Value < AAM_MAX)
 		{
 			AntiAliasingMethod = (EAntiAliasingMethod)Value;
 		}
 	}
 
-	static const auto CVarMobileMSAA = IConsoleManager::Get().FindTConsoleVariableDataInt(TEXT("r.MobileMSAA"));
-	if (FeatureLevel <= ERHIFeatureLevel::ES3_1 && CVarMobileMSAA ? CVarMobileMSAA->GetValueOnAnyThread() > 1 : false)
-	{
-		// Using mobile MSAA, disable other AA methods.
-		AntiAliasingMethod = AAM_None;
-
-		// Turn off various features which won't work with mobile MSAA.
-		//FinalPostProcessSettings.DepthOfFieldScale = 0.0f;
-	}
-
 	if (Family)
 	{
-		static IConsoleVariable* CVarMSAACount = IConsoleManager::Get().FindConsoleVariable(TEXT("r.MSAACount"));
-
-		const EShaderPlatform ShaderPlatform = GetFeatureLevelShaderPlatform(FeatureLevel);
-
-		if (AntiAliasingMethod == AAM_MSAA && IsForwardShadingEnabled(ShaderPlatform) && CVarMSAACount->GetInt() <= 0)
-		{
-			// Fallback to temporal AA so we can easily toggle methods with r.MSAACount
-			AntiAliasingMethod = AAM_TemporalAA;
-		}
-
-		static const auto PostProcessAAQualityCVar = IConsoleManager::Get().FindTConsoleVariableDataInt(TEXT("r.PostProcessAAQuality"));
-		static auto* MobileHDRCvar = IConsoleManager::Get().FindTConsoleVariableDataInt(TEXT("r.MobileHDR"));
-		static auto* MobileMSAACvar = IConsoleManager::Get().FindTConsoleVariableDataInt(TEXT("r.MobileMSAA"));
-		uint32 MobileMSAAValue = MobileMSAACvar->GetValueOnAnyThread();
-
-		int32 Quality = FMath::Clamp(PostProcessAAQualityCVar->GetValueOnAnyThread(), 0, 6);
 		const bool bWillApplyTemporalAA = Family->EngineShowFlags.PostProcessing || bIsPlanarReflection;
 
-		if (!bWillApplyTemporalAA || !Family->EngineShowFlags.AntiAliasing || Quality <= 0
-			// Disable antialiasing in GammaLDR mode to avoid jittering.
-			|| (FeatureLevel <= ERHIFeatureLevel::ES3_1 && MobileHDRCvar->GetValueOnAnyThread() == 0)
-			|| (FeatureLevel <= ERHIFeatureLevel::ES3_1 && (MobileMSAAValue > 1)))
+		if (!bWillApplyTemporalAA || !Family->EngineShowFlags.AntiAliasing)
 		{
 			AntiAliasingMethod = AAM_None;
 		}
 
 		if (AntiAliasingMethod == AAM_TemporalAA)
 		{
-			if (!Family->EngineShowFlags.TemporalAA || !Family->bRealtimeUpdate || Quality < 3)
+			if (!Family->EngineShowFlags.TemporalAA || !Family->bRealtimeUpdate)
 			{
 				AntiAliasingMethod = AAM_FXAA;
 			}
@@ -2616,6 +2587,20 @@ FSceneViewFamily::FSceneViewFamily(const ConstructionValues& CVS)
 		bRequireMultiView = (GSupportsMobileMultiView || GRHISupportsArrayIndexFromAnyShader) && bUsingMobileRenderer && bSkipPostprocessing && (MobileMultiViewCVar && MobileMultiViewCVar->GetValueOnAnyThread() != 0);
 	}
 #endif
+
+	// Check if the translucency are allowed to be rendered after DOF, if not, translucency after DOF will be rendered in standard translucency.
+	{
+		bAllowTranslucencyAfterDOF = CVarAllowTranslucencyAfterDOF.GetValueOnRenderThread() != 0
+			&& EngineShowFlags.PostProcessing // Used for reflection captures.
+			&& !UseDebugViewPS()
+			&& EngineShowFlags.SeparateTranslucency;
+
+		if (GetFeatureLevel() == ERHIFeatureLevel::ES3_1)
+		{
+			const bool bMobileMSAA = GetDefaultMSAACount(ERHIFeatureLevel::ES3_1) > 1;
+			bAllowTranslucencyAfterDOF &= (IsMobileHDR() && !bMobileMSAA); // on <= ES3_1 separate translucency requires HDR on and MSAA off
+		}
+	}
 }
 
 FSceneViewFamily::~FSceneViewFamily()
@@ -2676,20 +2661,6 @@ bool FSceneViewFamily::SupportsScreenPercentage() const
 	return true;
 }
 
-bool FSceneViewFamily::AllowTranslucencyAfterDOF() const
-{
-	static IConsoleVariable* CVarMobileMSAA = IConsoleManager::Get().FindConsoleVariable(TEXT("r.MobileMSAA"));
-	const bool bMobileMSAA = CVarMobileMSAA ? (CVarMobileMSAA->GetInt() > 1) : false;
-
-	return CVarAllowTranslucencyAfterDOF.GetValueOnRenderThread() != 0
-		&& (GetFeatureLevel() > ERHIFeatureLevel::ES3_1 || (IsMobileHDR() && !bMobileMSAA)) // on <= ES3_1 separate translucency requires HDR on and MSAA off
-	&& EngineShowFlags.PostProcessing // Used for reflection captures.
-	&& !UseDebugViewPS()
-	&& EngineShowFlags.SeparateTranslucency;
-	// If not, translucency after DOF will be rendered in standard translucency.
-}
-
-
 FSceneViewFamilyContext::~FSceneViewFamilyContext()
 {
 	// Cleanup the views allocated for this view family.
diff --git a/Engine/Source/Runtime/Engine/Public/SceneUtils.h b/Engine/Source/Runtime/Engine/Public/SceneUtils.h
--- a/Engine/Source/Runtime/Engine/Public/SceneUtils.h
+++ b/Engine/Source/Runtime/Engine/Public/SceneUtils.h
@@ -11,6 +11,9 @@
 #pragma once
 
 #include "ProfilingDebugging/RealtimeGPUProfiler.h"
+#include "RHIDefinitions.h"
+
+ENGINE_API DECLARE_LOG_CATEGORY_EXTERN(LogSceneUtils, Log, All);
 
 enum class EShadingPath
 {
@@ -26,6 +29,18 @@ enum class EMobileHDRMode
 	EnabledFloat16,
 };
 
+/** Used by rendering project settings. */
+UENUM()
+enum EAntiAliasingMethod
+{
+	AAM_None UMETA(DisplayName = "None"),
+	AAM_FXAA UMETA(DisplayName = "FXAA"),
+	AAM_TemporalAA UMETA(DisplayName = "TemporalAA"),
+	/** Only supported with forward shading.  MSAA sample count is controlled by r.MSAACount. */
+	AAM_MSAA UMETA(DisplayName = "MSAA"),
+	AAM_MAX,
+};
+
 /** True if HDR is enabled for the mobile renderer. */
 ENGINE_API bool IsMobileHDR();
 
@@ -33,6 +48,10 @@ ENGINE_API EMobileHDRMode GetMobileHDRMode();
 
 ENGINE_API bool IsMobileColorsRGB();
 
+ENGINE_API EAntiAliasingMethod GetDefaultAntiAliasingMethod(const FStaticFeatureLevel InFeatureLevel);
+
+ENGINE_API uint32 GetDefaultMSAACount(const FStaticFeatureLevel InFeatureLevel, uint32 PlatformMaxSampleCount = 8);
+
 // Callback for calling one action (typical use case: delay a clear until it's actually needed)
 class FDelayedRendererAction
 {
diff --git a/Engine/Source/Runtime/Engine/Public/SceneView.h b/Engine/Source/Runtime/Engine/Public/SceneView.h
--- a/Engine/Source/Runtime/Engine/Public/SceneView.h
+++ b/Engine/Source/Runtime/Engine/Public/SceneView.h
@@ -1738,7 +1738,7 @@ public:
 	/** Returns whether the screen percentage show flag is supported or not for this view family. */
 	bool SupportsScreenPercentage() const;
 
-	bool AllowTranslucencyAfterDOF() const;
+	FORCEINLINE bool AllowTranslucencyAfterDOF() const { return bAllowTranslucencyAfterDOF; }
 
 	/* Returns the maximum FSceneViewScreenPercentageConfig::PrimaryResolutionFraction. */
 	FORCEINLINE float GetPrimaryResolutionFractionUpperBound() const
@@ -1811,6 +1811,9 @@ private:
 
 	const ITemporalUpscaler* TemporalUpscalerInterface;
 
+	/** whether the translucency are allowed to render after DOF, if not they will be rendered in standard translucency. */
+	bool bAllowTranslucencyAfterDOF;
+
 	// Only FSceneRenderer can copy a view family.
 	FSceneViewFamily(const FSceneViewFamily&) = default;
 
diff --git a/Engine/Source/Runtime/Renderer/Private/MobileShadingRenderer.cpp b/Engine/Source/Runtime/Renderer/Private/MobileShadingRenderer.cpp
--- a/Engine/Source/Runtime/Renderer/Private/MobileShadingRenderer.cpp
+++ b/Engine/Source/Runtime/Renderer/Private/MobileShadingRenderer.cpp
@@ -215,8 +215,7 @@ FMobileSceneRenderer::FMobileSceneRenderer(const FSceneViewFamily* InViewFamily,
 		}
 	}
 
-	static const auto CVarMobileMSAA = IConsoleManager::Get().FindTConsoleVariableDataInt(TEXT("r.MobileMSAA"));
-	NumMSAASamples = (CVarMobileMSAA ? CVarMobileMSAA->GetValueOnAnyThread() : 1);
+	NumMSAASamples = GetDefaultMSAACount(ERHIFeatureLevel::ES3_1);
 }
 
 class FMobileDirLightShaderParamsRenderResource : public FRenderResource
@@ -416,11 +415,6 @@ void FMobileSceneRenderer::InitViews(FRDGBuilder& GraphBuilder, FSceneTexturesCo
 
 	// Finalize and set the scene textures config.
 	SceneTexturesConfig.bKeepDepthContent = bKeepDepthContent;
-
-	if (!SupportsMSAA())
-	{
-		SceneTexturesConfig.NumSamples = 1;
-	}
 	
 	FSceneTexturesConfig::Set(SceneTexturesConfig);
 
@@ -1718,9 +1712,4 @@ void FMobileSceneRenderer::UpdateMovablePointLightUniformBufferAndShadowInfo()
 			}
 		}
 	}
-}
-
-bool FMobileSceneRenderer::SupportsMSAA() const
-{
-	return !(IsUsingMobilePixelProjectedReflection(ShaderPlatform) || IsUsingMobileAmbientOcclusion(ShaderPlatform) || ShouldRenderVelocities() || bDeferredShading);
-}
+}
\ No newline at end of file
diff --git a/Engine/Source/Runtime/Renderer/Private/PostProcess/PostProcessing.cpp b/Engine/Source/Runtime/Renderer/Private/PostProcess/PostProcessing.cpp
--- a/Engine/Source/Runtime/Renderer/Private/PostProcess/PostProcessing.cpp
+++ b/Engine/Source/Runtime/Renderer/Private/PostProcess/PostProcessing.cpp
@@ -1414,10 +1414,9 @@ void AddMobilePostProcessingPasses(FRDGBuilder& GraphBuilder, const FViewInfo& V
 	const EAutoExposureMethod AutoExposureMethod = GetAutoExposureMethod(View);
 	const bool bUseEyeAdaptation = IsMobileEyeAdaptationEnabled(View);
 
-	static const auto CVarMobileMSAA = IConsoleManager::Get().FindTConsoleVariableDataInt(TEXT("r.MobileMSAA"));
-		
+	
 	//The input scene color has been encoded to non-linear space and needs to decode somewhere if MSAA enabled on Metal platform
-	bool bMetalMSAAHDRDecode = GSupportsShaderFramebufferFetch && IsMetalMobilePlatform(View.GetShaderPlatform()) && CVarMobileMSAA && CVarMobileMSAA->GetValueOnRenderThread() > 1;
+	bool bMetalMSAAHDRDecode = GSupportsShaderFramebufferFetch && IsMetalMobilePlatform(View.GetShaderPlatform()) && GetDefaultMSAACount(ERHIFeatureLevel::ES3_1) > 1;
 
 	// add the passes we want to add to the graph (commenting a line means the pass is not inserted into the graph) ---------
 
diff --git a/Engine/Source/Runtime/Renderer/Private/SceneRendering.h b/Engine/Source/Runtime/Renderer/Private/SceneRendering.h
--- a/Engine/Source/Runtime/Renderer/Private/SceneRendering.h
+++ b/Engine/Source/Runtime/Renderer/Private/SceneRendering.h
@@ -1792,7 +1792,6 @@ public:
 	virtual void Render(FRDGBuilder& GraphBuilder) = 0;
 	virtual void RenderHitProxies(FRDGBuilder& GraphBuilder) {}
 	virtual bool ShouldRenderVelocities() const { return false; }
-	virtual bool SupportsMSAA() const { return true; }
 
 	/** Creates a scene renderer based on the current feature level. */
 	static FSceneRenderer* CreateSceneRenderer(const FSceneViewFamily* InViewFamily, FHitProxyConsumer* HitProxyConsumer);
@@ -2145,8 +2144,6 @@ public:
 
 	virtual bool ShouldRenderVelocities() const override;
 
-	virtual bool SupportsMSAA() const override;
-
 	void RenderInverseOpacity(FRDGBuilder& GraphBuilder, const FViewInfo& View);
 
 protected:
diff --git a/Engine/Source/Runtime/Renderer/Private/SceneTextures.cpp b/Engine/Source/Runtime/Renderer/Private/SceneTextures.cpp
--- a/Engine/Source/Runtime/Renderer/Private/SceneTextures.cpp
+++ b/Engine/Source/Runtime/Renderer/Private/SceneTextures.cpp
@@ -49,17 +49,6 @@ static TAutoConsoleVariable<int32> CVarMSAACount(
 	ECVF_RenderThreadSafe | ECVF_Scalability
 	);
 
-static TAutoConsoleVariable<int32> CVarMobileMSAA(
-	TEXT("r.MobileMSAA"),
-	1,
-	TEXT("Use MSAA instead of Temporal AA on mobile:\n")
-	TEXT("1: Use Temporal AA (MSAA disabled)\n")
-	TEXT("2: Use 2x MSAA (Temporal AA disabled)\n")
-	TEXT("4: Use 4x MSAA (Temporal AA disabled)\n")
-	TEXT("8: Use 8x MSAA (Temporal AA disabled)"),
-	ECVF_RenderThreadSafe | ECVF_Scalability
-	);
-
 static TAutoConsoleVariable<int32> CVarGBufferFormat(
 	TEXT("r.GBufferFormat"),
 	1,
@@ -191,71 +180,6 @@ static EPixelFormat GetSceneColorFormat(const FSceneViewFamily& ViewFamily)
 	return Format;
 }
 
-static uint32 GetSceneTextureNumSamples(ERHIFeatureLevel::Type InFeatureLevel)
-{
-	uint32 NumSamples = 1;
-
-	if (InFeatureLevel >= ERHIFeatureLevel::SM5)
-	{
-		static IConsoleVariable* CVarDefaultAntiAliasing = IConsoleManager::Get().FindConsoleVariable(TEXT("r.DefaultFeature.AntiAliasing"));
-		EAntiAliasingMethod Method = (EAntiAliasingMethod)CVarDefaultAntiAliasing->GetInt();
-
-		if (IsForwardShadingEnabled(GetFeatureLevelShaderPlatform(InFeatureLevel)) && Method == AAM_MSAA)
-		{
-			NumSamples = FMath::Max(1, CVarMSAACount.GetValueOnRenderThread());
-
-			if (NumSamples != 1 && NumSamples != 2 && NumSamples != 4 && NumSamples != 8)
-			{
-				UE_LOG(LogRenderer, Warning, TEXT("Requested %d samples for MSAA, but this is not supported; falling back to 1 sample"), NumSamples);
-				NumSamples = 1;
-			}
-		}
-	}
-	else
-	{
-		NumSamples = CVarMobileMSAA.GetValueOnRenderThread();
-
-		static uint32 PlatformMaxSampleCount = GDynamicRHI->RHIGetPlatformTextureMaxSampleCount();
-		NumSamples = FMath::Min(NumSamples, PlatformMaxSampleCount);
-
-		if (NumSamples != 1 && NumSamples != 2 && NumSamples != 4 && NumSamples != 8)
-		{
-			UE_LOG(LogRenderer, Warning, TEXT("Requested %d samples for MSAA, but this is not supported; falling back to 1 sample"), NumSamples);
-			NumSamples = 1;
-		}
-
-		// Disable MSAA if we are using mobile pixel projected reflection, since we have to resolve the SceneColor and SceneDepth after opaque base pass
-		// Disable MSAA if we are using mobile ambient occlusion, since we have to resolve the SceneColor and SceneDepth after opaque base pass
-		if (NumSamples > 1 && (IsUsingMobilePixelProjectedReflection(GetFeatureLevelShaderPlatform(InFeatureLevel)) || IsUsingMobileAmbientOcclusion(GetFeatureLevelShaderPlatform(InFeatureLevel))))
-		{
-			NumSamples = 1;
-
-			static bool bWarned = false;
-
-			if (!bWarned)
-			{
-				bWarned = true;
-				UE_LOG(LogRenderer, Log, TEXT("Requested %d samples for MSAA, but using pixel projected reflection should disable MSAA"), NumSamples);
-			}
-		}
-	}
-
-	if (NumSamples > 1 && !RHISupportsMSAA(GShaderPlatformForFeatureLevel[InFeatureLevel]))
-	{
-		NumSamples = 1;
-
-		static bool bWarned = false;
-
-		if (!bWarned)
-		{
-			bWarned = true;
-			UE_LOG(LogRenderer, Log, TEXT("MSAA requested but the platform doesn't support MSAA, falling back to Temporal AA"));
-		}
-	}
-
-	return NumSamples;
-}
-
 static uint32 GetEditorPrimitiveNumSamples(ERHIFeatureLevel::Type FeatureLevel)
 {
 	uint32 SampleCount = 1;
@@ -582,7 +506,7 @@ FSceneTexturesConfig FSceneTexturesConfig::Create(const FSceneViewFamily& ViewFa
 	Config.ShadingPath = FSceneInterface::GetShadingPath(Config.FeatureLevel);
 	Config.ShaderPlatform = GetFeatureLevelShaderPlatform(Config.FeatureLevel);
 	Config.Extent = FSceneTextureExtentState::Get().Compute(ViewFamily);
-	Config.NumSamples = GetSceneTextureNumSamples(Config.FeatureLevel);
+	Config.NumSamples = GetDefaultMSAACount(Config.FeatureLevel, GDynamicRHI->RHIGetPlatformTextureMaxSampleCount());
 	Config.EditorPrimitiveNumSamples = GetEditorPrimitiveNumSamples(Config.FeatureLevel);
 	Config.ColorFormat = PF_Unknown;
 	Config.ColorClearValue = FClearValueBinding::Black;
