Fix potential null dereference found by static analysis

diff --git a/Engine/Source/Runtime/Engine/Private/LightMap.cpp b/Engine/Source/Runtime/Engine/Private/LightMap.cpp
--- a/Engine/Source/Runtime/Engine/Private/LightMap.cpp
+++ b/Engine/Source/Runtime/Engine/Private/LightMap.cpp
@@ -3467,7 +3467,7 @@ IAllocatedVirtualTexture* FLightmapResourceCluster::AcquireAllocatedVT() const
 	
 #if WITH_EDITOR
 	// Compilation is still pending, this function will be called back once compilation finishes.
-	if (VirtualTexture->IsCompiling())
+	if (VirtualTexture && VirtualTexture->IsCompiling())
 	{
 		return nullptr;
 	}
