Fix Static Analysis (add check)

diff --git a/Engine/Source/Runtime/Engine/Private/WorldPartition/WorldPartitionLevelHelper.cpp b/Engine/Source/Runtime/Engine/Private/WorldPartition/WorldPartitionLevelHelper.cpp
--- a/Engine/Source/Runtime/Engine/Private/WorldPartition/WorldPartitionLevelHelper.cpp
+++ b/Engine/Source/Runtime/Engine/Private/WorldPartition/WorldPartitionLevelHelper.cpp
@@ -199,6 +199,7 @@ bool FWorldPartitionLevelHelper::LoadActors(ULevel* InDestLevel, TArrayView<FWor
 		FPackagePath PackagePath = FPackagePath::FromPackageNameChecked(*ActorPackageName);
 		if (bLoadForPIE)
 		{
+			check(DestPackage);
 			::LoadPackageAsync(PackagePath, FName(*ActorPackageInstanceNames[ChildIndex]), CompletionCallback, nullptr, PKG_PlayInEditor, DestPackage->PIEInstanceID, 0, InOutInstancingContext);
 		}
 		else
