Added support for re-instancing the ISM components referenced Static Mesh Instance elements

#rnx

diff --git a/Engine/Source/Runtime/Engine/Classes/Components/InstancedStaticMeshComponent.h b/Engine/Source/Runtime/Engine/Classes/Components/InstancedStaticMeshComponent.h
--- a/Engine/Source/Runtime/Engine/Classes/Components/InstancedStaticMeshComponent.h
+++ b/Engine/Source/Runtime/Engine/Classes/Components/InstancedStaticMeshComponent.h
@@ -277,6 +277,10 @@ class ENGINE_API UInstancedStaticMeshComponent : public UStaticMeshComponent
 	UFUNCTION(BlueprintCallable, Category = "Components|InstancedStaticMesh")
 	int32 GetInstanceCount() const;
 
+	/** Does the given index map to a valid instance in this component? */
+	UFUNCTION(BlueprintCallable, Category = "Components|InstancedStaticMesh")
+	bool IsValidInstance(int32 InstanceIndex) const;
+
 	/** Sets the fading start and culling end distances for this component. */
 	UFUNCTION(BlueprintCallable, Category = "Components|InstancedStaticMesh")
 	void SetCullDistances(int32 StartCullDistance, int32 EndCullDistance);
diff --git a/Engine/Source/Runtime/Engine/Private/Elements/Framework/EngineElementsLibrary.cpp b/Engine/Source/Runtime/Engine/Private/Elements/Framework/EngineElementsLibrary.cpp
--- a/Engine/Source/Runtime/Engine/Private/Elements/Framework/EngineElementsLibrary.cpp
+++ b/Engine/Source/Runtime/Engine/Private/Elements/Framework/EngineElementsLibrary.cpp
@@ -74,28 +74,41 @@ FTypedElementHandle AcquireEditorTypedElementHandle(const ObjectClass* Object, T
 	return FTypedElementHandle();
 }
 
-template <typename ObjectClass, typename ElementDataType>
-void ReplaceEditorTypedElementHandles(TArray<FTypedElementHandle>& OutUpdatedElements, const TMap<UObject*, UObject*>& ReplacementObjects, TTypedElementOwnerStore<ElementDataType, const ObjectClass*>& ElementOwnerStore, TFunctionRef<void(const ObjectClass*, TTypedElementOwner<ElementDataType>&)> UpdateElement, TFunctionRef<void(const ObjectClass*, TTypedElementOwner<ElementDataType>&)> DestroyElement)
+template <typename ObjectClass>
+TArray<TTuple<const ObjectClass*, const ObjectClass*>> CalculatePotentialObjectReplacements(const TMap<UObject*, UObject*>& ReplacementObjects)
 {
+	TArray<TTuple<const ObjectClass*, const ObjectClass*>> PotentialObjectReplacements;
+
 	for (const TTuple<UObject*, UObject*>& ReplacementObjectPair : ReplacementObjects)
 	{
-		const ObjectClass* OldObject = Cast<ObjectClass>(ReplacementObjectPair.Key);
-		if (!OldObject)
+		if (const ObjectClass* OldObject = Cast<ObjectClass>(ReplacementObjectPair.Key))
 		{
-			// We can't do anything for objects of the incorrect type...
-			continue;
+			const ObjectClass* NewObject = Cast<ObjectClass>(ReplacementObjectPair.Value);
+			PotentialObjectReplacements.Add(MakeTuple(OldObject, NewObject));
 		}
+	}
+
+	return PotentialObjectReplacements;
+}
+
+template <typename ElementDataType, typename KeyDataType>
+void ReplaceEditorTypedElementHandles(TArray<FTypedElementHandle>& OutUpdatedElements, const TArray<TTuple<KeyDataType, KeyDataType>>& ReplacementKeys, TTypedElementOwnerStore<ElementDataType, KeyDataType>& ElementOwnerStore, TFunctionRef<void(KeyDataType, TTypedElementOwner<ElementDataType>&)> UpdateElement, TFunctionRef<void(KeyDataType, TTypedElementOwner<ElementDataType>&)> DestroyElement)
+{
+	for (const TTuple<KeyDataType, KeyDataType>& ReplacementKeyPair : ReplacementKeys)
+	{
+		const KeyDataType& OldKey = ReplacementKeyPair.Key;
+		const KeyDataType& NewKey = ReplacementKeyPair.Value;
 
 		// We only need to attempt replacement if we actually have an element for the old object...
-		if (TTypedElementOwner<ElementDataType> OldEditorElement = ElementOwnerStore.UnregisterElementOwner(OldObject))
+		if (TTypedElementOwner<ElementDataType> OldEditorElement = ElementOwnerStore.UnregisterElementOwner(OldKey))
 		{
 			if (OldEditorElement.Private_GetInternalData()->GetRefCount() > 1)
 			{
-				if (const ObjectClass* NewObject = Cast<ObjectClass>(ReplacementObjectPair.Value))
+				if (NewKey)
 				{
 					// The old element has external references, so we're going to destroy the new element (if any) and re-point the old element to the new object...
 					// Note: This requires that the new element has no external references - if both old and new elements have external references then we'll need to support redirection at the element level!
-					if (TTypedElementOwner<ElementDataType> NewEditorElement = ElementOwnerStore.UnregisterElementOwner(NewObject))
+					if (TTypedElementOwner<ElementDataType> NewEditorElement = ElementOwnerStore.UnregisterElementOwner(NewKey))
 					{
 						if (NewEditorElement.Private_GetInternalData()->GetRefCount() > 1)
 						{
@@ -106,25 +119,25 @@ void ReplaceEditorTypedElementHandles(TArray<FTypedElementHandle>& OutUpdatedEle
 						}
 
 						checkf(NewEditorElement.Private_GetInternalData()->GetRefCount() <= 1, TEXT("The old and new element both have external references! Replacing these will require support for redirection at the element level!"));
-						DestroyElement(NewObject, NewEditorElement);
+						DestroyElement(NewKey, NewEditorElement);
 					}
 
-					UpdateElement(NewObject, OldEditorElement);
+					UpdateElement(NewKey, OldEditorElement);
 					OutUpdatedElements.Emplace(OldEditorElement.AcquireHandle());
-					ElementOwnerStore.RegisterElementOwner(NewObject, MoveTemp(OldEditorElement));
+					ElementOwnerStore.RegisterElementOwner(NewKey, MoveTemp(OldEditorElement));
 				}
 				else
 				{
 					// The object has been redirected to null, so try and clear any external references to the old element...
 					TTuple<FTypedElementHandle, FTypedElementHandle> ElementRedirect = MakeTuple(OldEditorElement.AcquireHandle(), FTypedElementHandle());
 					UTypedElementRegistry::GetInstance()->OnElementReplaced().Broadcast(MakeArrayView(&ElementRedirect, 1));
-					DestroyElement(OldObject, OldEditorElement);
+					DestroyElement(OldKey, OldEditorElement);
 				}
 			}
 			else
 			{
 				// The old element has no external references, so we can just destroy it...
-				DestroyElement(OldObject, OldEditorElement);
+				DestroyElement(OldKey, OldEditorElement);
 			}
 		}
 	}
@@ -162,20 +175,67 @@ void UEngineElementsLibrary::OnObjectsReplaced(const TMap<UObject*, UObject*>& I
 {
 	TArray<FTypedElementHandle> UpdatedElements;
 
-	EngineElementsLibraryUtil::ReplaceEditorTypedElementHandles<UObject, FObjectElementData>(UpdatedElements, InReplacementObjects, GObjectElementOwnerStore, [](const UObject* InObject, TTypedElementOwner<FObjectElementData>& InOutObjectElement)
 	{
-		InOutObjectElement.GetDataChecked().Object = const_cast<UObject*>(InObject);
-	}, &UEngineElementsLibrary::DestroyObjectElement);
+		const TArray<TTuple<const UObject*, const UObject*>> PotentialObjectReplacements = EngineElementsLibraryUtil::CalculatePotentialObjectReplacements<UObject>(InReplacementObjects);
+		EngineElementsLibraryUtil::ReplaceEditorTypedElementHandles<FObjectElementData, const UObject*>(UpdatedElements, PotentialObjectReplacements, GObjectElementOwnerStore, [](const UObject* InObject, TTypedElementOwner<FObjectElementData>& InOutObjectElement)
+		{
+			InOutObjectElement.GetDataChecked().Object = const_cast<UObject*>(InObject);
+		}, &UEngineElementsLibrary::DestroyObjectElement);
+	}
+
+	{
+		const TArray<TTuple<const AActor*, const AActor*>> PotentialActorReplacements = EngineElementsLibraryUtil::CalculatePotentialObjectReplacements<AActor>(InReplacementObjects);
+		EngineElementsLibraryUtil::ReplaceEditorTypedElementHandles<FActorElementData, const AActor*>(UpdatedElements, PotentialActorReplacements, GActorElementOwnerStore, [](const AActor* InActor, TTypedElementOwner<FActorElementData>& InOutActorElement)
+		{
+			InOutActorElement.GetDataChecked().Actor = const_cast<AActor*>(InActor);
+		}, &UEngineElementsLibrary::DestroyActorElement);
+	}
 
-	EngineElementsLibraryUtil::ReplaceEditorTypedElementHandles<AActor, FActorElementData>(UpdatedElements, InReplacementObjects, GActorElementOwnerStore, [](const AActor* InActor, TTypedElementOwner<FActorElementData>& InOutActorElement)
 	{
-		InOutActorElement.GetDataChecked().Actor = const_cast<AActor*>(InActor);
-	}, &UEngineElementsLibrary::DestroyActorElement);
+		const TArray<TTuple<const UActorComponent*, const UActorComponent*>> PotentialComponentReplacements = EngineElementsLibraryUtil::CalculatePotentialObjectReplacements<UActorComponent>(InReplacementObjects);
+		EngineElementsLibraryUtil::ReplaceEditorTypedElementHandles<FComponentElementData, const UActorComponent*>(UpdatedElements, PotentialComponentReplacements, GComponentElementOwnerStore, [](const UActorComponent* InComponent, TTypedElementOwner<FComponentElementData>& InOutComponentElement)
+		{
+			InOutComponentElement.GetDataChecked().Component = const_cast<UActorComponent*>(InComponent);
+		}, &UEngineElementsLibrary::DestroyComponentElement);
+	}
 
-	EngineElementsLibraryUtil::ReplaceEditorTypedElementHandles<UActorComponent, FComponentElementData>(UpdatedElements, InReplacementObjects, GComponentElementOwnerStore, [](const UActorComponent* InComponent, TTypedElementOwner<FComponentElementData>& InOutComponentElement)
 	{
-		InOutComponentElement.GetDataChecked().Component = const_cast<UActorComponent*>(InComponent);
-	}, &UEngineElementsLibrary::DestroyComponentElement);
+		TArray<TTuple<FSMInstanceElementId, FSMInstanceElementId>> PotentialSMInstanceReplacements;
+
+		FSMInstanceElementIdMap& SMInstanceElementIdMap = FSMInstanceElementIdMap::Get();
+		for (const TTuple<UObject*, UObject*>& ReplacementObjectPair : InReplacementObjects)
+		{
+			if (UInstancedStaticMeshComponent* OldISMComponent = Cast<UInstancedStaticMeshComponent>(ReplacementObjectPair.Key))
+			{
+				const TArray<FSMInstanceElementId> OldSMInstanceElementIds = SMInstanceElementIdMap.GetSMInstanceElementIdsForComponent(OldISMComponent);
+				PotentialSMInstanceReplacements.Reserve(PotentialSMInstanceReplacements.Num() + OldSMInstanceElementIds.Num());
+
+				if (UInstancedStaticMeshComponent* NewISMComponent = Cast<UInstancedStaticMeshComponent>(ReplacementObjectPair.Value))
+				{
+					// Attempt to ensure that the old IDs are re-used on the new component
+					// This is required so that in-memory stored references (eg, undo/redo) map correctly when using the new component instance
+					SMInstanceElementIdMap.OnComponentReplaced(OldISMComponent, NewISMComponent);
+
+					for (const FSMInstanceElementId& OldSMInstanceElementId : OldSMInstanceElementIds)
+					{
+						PotentialSMInstanceReplacements.Add(MakeTuple(OldSMInstanceElementId, FSMInstanceElementId{ NewISMComponent, OldSMInstanceElementId.InstanceId }));
+					}
+				}
+				else
+				{
+					for (const FSMInstanceElementId& OldSMInstanceElementId : OldSMInstanceElementIds)
+					{
+						PotentialSMInstanceReplacements.Add(MakeTuple(OldSMInstanceElementId, FSMInstanceElementId()));
+					}
+				}
+			}
+		}
+
+		EngineElementsLibraryUtil::ReplaceEditorTypedElementHandles<FSMInstanceElementData, FSMInstanceElementId>(UpdatedElements, PotentialSMInstanceReplacements, GSMInstanceElementOwnerStore, [](FSMInstanceElementId InSMInstanceElementId, TTypedElementOwner<FSMInstanceElementData>& InOutSMInstanceElement)
+		{
+			InOutSMInstanceElement.GetDataChecked().InstanceElementId = InSMInstanceElementId;
+		}, &UEngineElementsLibrary::DestroySMInstanceElement);
+	}
 
 	if (UpdatedElements.Num() > 0)
 	{
diff --git a/Engine/Source/Runtime/Engine/Private/Elements/SMInstance/SMInstanceElementId.cpp b/Engine/Source/Runtime/Engine/Private/Elements/SMInstance/SMInstanceElementId.cpp
--- a/Engine/Source/Runtime/Engine/Private/Elements/SMInstance/SMInstanceElementId.cpp
+++ b/Engine/Source/Runtime/Engine/Private/Elements/SMInstance/SMInstanceElementId.cpp
@@ -92,6 +92,57 @@ FSMInstanceElementId FSMInstanceElementIdMap::GetSMInstanceElementIdFromSMInstan
 	return FSMInstanceElementId();
 }
 
+TArray<FSMInstanceElementId> FSMInstanceElementIdMap::GetSMInstanceElementIdsForComponent(UInstancedStaticMeshComponent* InComponent) const
+{
+	TArray<FSMInstanceElementId> SMInstanceElementIds;
+
+	{
+		FScopeLock Lock(&ISMComponentsCS);
+
+		if (TSharedPtr<const FSMInstanceElementIdMapEntry> ISMEntry = ISMComponents.FindRef(InComponent))
+		{
+			SMInstanceElementIds.Reserve(ISMEntry->InstanceIdToIndexMap.Num());
+			for (const TTuple<uint64, int32>& InstanceIdToIndexPair : ISMEntry->InstanceIdToIndexMap)
+			{
+				SMInstanceElementIds.Add(FSMInstanceElementId{ InComponent, InstanceIdToIndexPair.Key });
+			}
+		}
+	}
+
+	return SMInstanceElementIds;
+}
+
+void FSMInstanceElementIdMap::OnComponentReplaced(UInstancedStaticMeshComponent* InOldComponent, UInstancedStaticMeshComponent* InNewComponent)
+{
+	FScopeLock Lock(&ISMComponentsCS);
+
+	if (TSharedPtr<const FSMInstanceElementIdMapEntry> OldISMEntry = ISMComponents.FindRef(InOldComponent))
+	{
+		TSharedPtr<FSMInstanceElementIdMapEntry> NewISMEntry = ISMComponents.FindRef(InNewComponent);
+		if (NewISMEntry)
+		{
+			ClearInstanceData_NoLock(InNewComponent, *NewISMEntry);
+		}
+		else
+		{
+			NewISMEntry = ISMComponents.Add(InNewComponent, MakeShared<FSMInstanceElementIdMapEntry>(this, InNewComponent));
+		}
+
+		NewISMEntry->InstanceIndexToIdMap.Reserve(OldISMEntry->InstanceIndexToIdMap.Num());
+		NewISMEntry->InstanceIdToIndexMap.Reserve(OldISMEntry->InstanceIdToIndexMap.Num());
+		NewISMEntry->NextInstanceId = OldISMEntry->NextInstanceId;
+
+		for (const TTuple<uint64, int32>& OldInstanceIdToIndexPair : OldISMEntry->InstanceIdToIndexMap)
+		{
+			if (InNewComponent->IsValidInstance(OldInstanceIdToIndexPair.Value))
+			{
+				NewISMEntry->InstanceIndexToIdMap.Add(OldInstanceIdToIndexPair.Value, OldInstanceIdToIndexPair.Key);
+				NewISMEntry->InstanceIdToIndexMap.Add(OldInstanceIdToIndexPair.Key, OldInstanceIdToIndexPair.Value);
+			}
+		}
+	}
+}
+
 void FSMInstanceElementIdMap::SerializeIdMappings(FSMInstanceElementIdMapEntry* InEntry, FArchive& Ar)
 {
 	FScopeLock Lock(&ISMComponentsCS);
@@ -196,21 +247,7 @@ void FSMInstanceElementIdMap::OnInstanceIndexUpdated(UInstancedStaticMeshCompone
 
 		case FInstancedStaticMeshDelegates::EInstanceIndexUpdateType::Cleared:
 			{
-				// Take the current instance IDs, so that we can find notify the elements have been unmapped
-				TMap<uint64, int32> PreviousInstanceIdToIndexMap = MoveTemp(ISMEntry->InstanceIdToIndexMap);
-
-				// Note: We can't remove things from the ISMComponents map until we get a 'Destroyed' notification for the owner ISM component, as the 
-				// per-component transactors may be referenced by the undo buffer, so the transactor needs to persist in order for redo to work correctly
-				ISMEntry->InstanceIndexToIdMap.Reset();
-				ISMEntry->InstanceIdToIndexMap.Reset();
-				
-				if (OnInstanceRemovedDelegate.IsBound())
-				{
-					for (const TTuple<uint64, int32>& PreviousInstanceIdToIndexPair : PreviousInstanceIdToIndexMap)
-					{
-						OnInstanceRemovedDelegate.Broadcast(FSMInstanceElementId{ InComponent, PreviousInstanceIdToIndexPair.Key }, PreviousInstanceIdToIndexPair.Value);
-					}
-				}
+				ClearInstanceData_NoLock(InComponent, *ISMEntry);
 			}
 			break;
 
@@ -239,6 +276,25 @@ void FSMInstanceElementIdMap::OnInstanceIndexUpdated(UInstancedStaticMeshCompone
 	}
 }
 
+void FSMInstanceElementIdMap::ClearInstanceData_NoLock(UInstancedStaticMeshComponent* InComponent, FSMInstanceElementIdMapEntry& InEntry)
+{
+	// Take the current instance IDs, so that we can find notify the elements have been unmapped
+	TMap<uint64, int32> PreviousInstanceIdToIndexMap = MoveTemp(InEntry.InstanceIdToIndexMap);
+
+	// Note: We can't remove things from the ISMComponents map until we get a 'Destroyed' notification for the owner ISM component, as the 
+	// per-component transactors may be referenced by the undo buffer, so the transactor needs to persist in order for redo to work correctly
+	InEntry.InstanceIndexToIdMap.Reset();
+	InEntry.InstanceIdToIndexMap.Reset();
+
+	if (OnInstanceRemovedDelegate.IsBound())
+	{
+		for (const TTuple<uint64, int32>& PreviousInstanceIdToIndexPair : PreviousInstanceIdToIndexMap)
+		{
+			OnInstanceRemovedDelegate.Broadcast(FSMInstanceElementId{ InComponent, PreviousInstanceIdToIndexPair.Key }, PreviousInstanceIdToIndexPair.Value);
+		}
+	}
+}
+
 #if WITH_EDITOR
 void FSMInstanceElementIdMap::OnObjectModified(UObject* InObject)
 {
diff --git a/Engine/Source/Runtime/Engine/Private/InstancedStaticMesh.cpp b/Engine/Source/Runtime/Engine/Private/InstancedStaticMesh.cpp
--- a/Engine/Source/Runtime/Engine/Private/InstancedStaticMesh.cpp
+++ b/Engine/Source/Runtime/Engine/Private/InstancedStaticMesh.cpp
@@ -3023,6 +3023,11 @@ int32 UInstancedStaticMeshComponent::GetInstanceCount() const
 	return PerInstanceSMData.Num();
 }
 
+bool UInstancedStaticMeshComponent::IsValidInstance(int32 InstanceIndex) const
+{
+	return PerInstanceSMData.IsValidIndex(InstanceIndex);
+}
+
 void UInstancedStaticMeshComponent::SetCullDistances(int32 StartCullDistance, int32 EndCullDistance)
 {
 	if (InstanceStartCullDistance != StartCullDistance || InstanceEndCullDistance != EndCullDistance)
diff --git a/Engine/Source/Runtime/Engine/Public/Elements/SMInstance/SMInstanceElementId.h b/Engine/Source/Runtime/Engine/Public/Elements/SMInstance/SMInstanceElementId.h
--- a/Engine/Source/Runtime/Engine/Public/Elements/SMInstance/SMInstanceElementId.h
+++ b/Engine/Source/Runtime/Engine/Public/Elements/SMInstance/SMInstanceElementId.h
@@ -138,6 +138,16 @@ public:
 	 */
 	FSMInstanceElementId GetSMInstanceElementIdFromSMInstanceId(const FSMInstanceId& InSMInstanceId, const bool bAllowCreate = true);
 
+	/**
+	 * Given an ISM component, get all FSMInstanceElementId values that are currently mapped for it.
+	 */
+	TArray<FSMInstanceElementId> GetSMInstanceElementIdsForComponent(UInstancedStaticMeshComponent* InComponent) const;
+
+	/**
+	 * Called when one ISM component is replaced with another, and attempts to copy the ID mapping of the old component to the new.
+	 */
+	void OnComponentReplaced(UInstancedStaticMeshComponent* InOldComponent, UInstancedStaticMeshComponent* InNewComponent);
+
 public:
 	//~ FGCObject interface
 	virtual void AddReferencedObjects(FReferenceCollector& Collector) override;
@@ -152,6 +162,8 @@ public:
 private:
 	void OnInstanceIndexUpdated(UInstancedStaticMeshComponent* InComponent, TArrayView<const FInstancedStaticMeshDelegates::FInstanceIndexUpdateData> InIndexUpdates);
 
+	void ClearInstanceData_NoLock(UInstancedStaticMeshComponent* InComponent, FSMInstanceElementIdMapEntry& InEntry);
+
 #if WITH_EDITOR
 	void OnObjectModified(UObject* InObject);
 #endif	// WITH_EDITOR
