WorldPartition: Move registration/unregistration WP specific code into virtual

diff --git a/Engine/Source/Runtime/Engine/Private/WorldPartition/WorldPartition.cpp b/Engine/Source/Runtime/Engine/Private/WorldPartition/WorldPartition.cpp
--- a/Engine/Source/Runtime/Engine/Private/WorldPartition/WorldPartition.cpp
+++ b/Engine/Source/Runtime/Engine/Private/WorldPartition/WorldPartition.cpp
@@ -18,6 +18,7 @@
 #include "ProfilingDebugging/ScopedTimers.h"
 
 #if WITH_EDITOR
+#include "LevelUtils.h"
 #include "Editor.h"
 #include "FileHelpers.h"
 #include "HAL/FileManager.h"
@@ -768,6 +769,35 @@ void UWorldPartition::OnActorDescUpdated(const TUniquePtr<FWorldPartitionActorDe
 	}
 }
 
+void UWorldPartition::ApplyActorTransform(AActor* Actor, const FTransform& InTransform)
+{
+	if (!InTransform.Equals(FTransform::Identity))
+	{
+		FLevelUtils::FApplyLevelTransformParams TransformParams(Actor->GetLevel(), InTransform);
+		TransformParams.Actor = Actor;
+		TransformParams.bDoPostEditMove = true;
+		FLevelUtils::ApplyLevelTransform(TransformParams);
+	}
+}
+
+void UWorldPartition::OnActorDescRegistered(const FWorldPartitionActorDesc& ActorDesc) 
+{
+	AActor* Actor = ActorDesc.GetActor();
+	check(Actor);
+	ApplyActorTransform(Actor, InstanceTransform);
+	Actor->GetLevel()->AddLoadedActor(Actor);
+}
+
+void UWorldPartition::OnActorDescUnregistered(const FWorldPartitionActorDesc& ActorDesc) 
+{
+	AActor* Actor = ActorDesc.GetActor();
+	if (!Actor->IsPendingKill())
+	{
+		Actor->GetLevel()->RemoveLoadedActor(Actor);
+		ApplyActorTransform(Actor, InstanceTransform.Inverse());
+	}
+}
+
 void UWorldPartition::HashActorDesc(FWorldPartitionActorDesc* ActorDesc)
 {
 	check(ActorDesc);
diff --git a/Engine/Source/Runtime/Engine/Private/WorldPartition/WorldPartitionActorDesc.cpp b/Engine/Source/Runtime/Engine/Private/WorldPartition/WorldPartitionActorDesc.cpp
--- a/Engine/Source/Runtime/Engine/Private/WorldPartition/WorldPartitionActorDesc.cpp
+++ b/Engine/Source/Runtime/Engine/Private/WorldPartition/WorldPartitionActorDesc.cpp
@@ -3,7 +3,6 @@
 #include "WorldPartition/WorldPartitionActorDesc.h"
 
 #if WITH_EDITOR
-#include "LevelUtils.h"
 #include "Misc/HashBuilder.h"
 #include "UObject/LinkerInstancingContext.h"
 #include "UObject/UObjectHash.h"
@@ -152,21 +151,6 @@ void FWorldPartitionActorDesc::TransformInstance(const FString& From, const FStr
 	}
 }
 
-void FWorldPartitionActorDesc::ApplyActorTransform(const FTransform& InTransform)
-{
-	AActor* Actor = GetActor();
-	check(Actor);
-
-	if (!InTransform.Equals(FTransform::Identity))
-	{
-		FLevelUtils::FApplyLevelTransformParams TransformParams(Actor->GetLevel(), InTransform);
-		TransformParams.Actor = Actor;
-		TransformParams.bDoPostEditMove = true;
-		FLevelUtils::ApplyLevelTransform(TransformParams);
-	}
-}
-
-
 FString FWorldPartitionActorDesc::ToString() const
 {
 	return FString::Printf(TEXT("Guid:%s Class:%s Name:%s"), *Guid.ToString(), *Class.ToString(), *FPaths::GetExtension(ActorPath.ToString()));
@@ -310,8 +294,7 @@ void FWorldPartitionActorDesc::RegisterActor()
 {
 	if (AActor* Actor = GetActor())
 	{
-		ApplyActorTransform(CastChecked<UWorldPartition>(Container)->InstanceTransform);
-		Actor->GetLevel()->AddLoadedActor(Actor);
+		Container->OnActorDescRegistered(*this);
 	}
 }
 
@@ -319,11 +302,7 @@ void FWorldPartitionActorDesc::UnregisterActor()
 {
 	if (AActor* Actor = GetActor())
 	{
-		if (!Actor->IsPendingKill())
-		{
-			Actor->GetLevel()->RemoveLoadedActor(Actor);
-			ApplyActorTransform(CastChecked<UWorldPartition>(Container)->InstanceTransform.Inverse());
-		}
+		Container->OnActorDescUnregistered(*this);
 	}
 }
 #endif
\ No newline at end of file
diff --git a/Engine/Source/Runtime/Engine/Public/WorldPartition/ActorDescContainer.h b/Engine/Source/Runtime/Engine/Public/WorldPartition/ActorDescContainer.h
--- a/Engine/Source/Runtime/Engine/Public/WorldPartition/ActorDescContainer.h
+++ b/Engine/Source/Runtime/Engine/Public/WorldPartition/ActorDescContainer.h
@@ -15,6 +15,7 @@ class ENGINE_API UActorDescContainer : public UObject
 	GENERATED_UCLASS_BODY()
 
 	friend struct FWorldPartitionHandleUtils;
+	friend class FWorldPartitionActorDesc;
 
 public:
 	void Initialize(UWorld* World, FName InPackageName, bool bRegisterDelegates);
@@ -56,6 +57,9 @@ protected:
 	virtual void OnActorDescUpdating(const TUniquePtr<FWorldPartitionActorDesc>& ActorDesc) {}
 	virtual void OnActorDescUpdated(const TUniquePtr<FWorldPartitionActorDesc>& ActorDesc) {}
 
+	virtual void OnActorDescRegistered(const FWorldPartitionActorDesc&) {}
+	virtual void OnActorDescUnregistered(const FWorldPartitionActorDesc&) {}
+
 	bool ShouldHandleAssetEvent(const FAssetData& InAssetData);
 
 	bool bContainerInitialized;
diff --git a/Engine/Source/Runtime/Engine/Public/WorldPartition/WorldPartition.h b/Engine/Source/Runtime/Engine/Public/WorldPartition/WorldPartition.h
--- a/Engine/Source/Runtime/Engine/Public/WorldPartition/WorldPartition.h
+++ b/Engine/Source/Runtime/Engine/Public/WorldPartition/WorldPartition.h
@@ -95,6 +95,9 @@ private:
 	virtual void OnActorDescUpdating(const TUniquePtr<FWorldPartitionActorDesc>& ActorDesc) override;
 	virtual void OnActorDescUpdated(const TUniquePtr<FWorldPartitionActorDesc>& ActorDesc) override;
 
+	virtual void OnActorDescRegistered(const FWorldPartitionActorDesc&) override;
+	virtual void OnActorDescUnregistered(const FWorldPartitionActorDesc&) override;
+
 	bool ShouldHandleAssetEvent(const FAssetData& InAssetData);
 #endif
 
@@ -224,6 +227,8 @@ private:
 	void UnhashActorDesc(FWorldPartitionActorDesc* ActorDesc);
 	bool ShouldActorBeLoaded(const FWorldPartitionActorDesc* ActorDesc) const;
 	bool UpdateEditorCells(TFunctionRef<bool(TArray<UWorldPartitionEditorCell*>&)> GetCellsToProcess, bool bIsCellShouldBeLoaded);
+
+	void ApplyActorTransform(AActor* Actor, const FTransform& InTransform);
 #endif
 
 	UWorldPartitionStreamingPolicy* GetStreamingPolicy() const;
diff --git a/Engine/Source/Runtime/Engine/Public/WorldPartition/WorldPartitionActorDesc.h b/Engine/Source/Runtime/Engine/Public/WorldPartition/WorldPartitionActorDesc.h
--- a/Engine/Source/Runtime/Engine/Public/WorldPartition/WorldPartitionActorDesc.h
+++ b/Engine/Source/Runtime/Engine/Public/WorldPartition/WorldPartitionActorDesc.h
@@ -132,7 +132,6 @@ protected:
 	FWorldPartitionActorDesc();
 
 	void TransformInstance(const FString& From, const FString& To, const FTransform& Transform);
-	void ApplyActorTransform(const FTransform& InTransform);
 
 	inline void TransferRefCounts(const FWorldPartitionActorDesc* From) const
 	{
