Fix notify states not inserting transactions into the undo history when edited

diff --git a/Engine/Source/Runtime/Engine/Private/Animation/AnimNotifyState_TimedParticleEffect.cpp b/Engine/Source/Runtime/Engine/Private/Animation/AnimNotifyState_TimedParticleEffect.cpp
--- a/Engine/Source/Runtime/Engine/Private/Animation/AnimNotifyState_TimedParticleEffect.cpp
+++ b/Engine/Source/Runtime/Engine/Private/Animation/AnimNotifyState_TimedParticleEffect.cpp
@@ -105,6 +105,8 @@ FString UAnimNotifyState_TimedParticleEffect::GetNotifyName_Implementation() con
 #if WITH_EDITOR
 void UAnimNotifyState_TimedParticleEffect::PreEditChange(FProperty* PropertyAboutToChange)
 {
+	Super::PreEditChange(PropertyAboutToChange);
+
 	if(PropertyAboutToChange)
 	{
 		if(PropertyAboutToChange->GetName() == GET_MEMBER_NAME_STRING_CHECKED(UAnimNotifyState_TimedParticleEffect, PSTemplate) && PSTemplate != NULL)
