WorldPartition: Fix failed check on always loaded LevelInstance sub actors

diff --git a/Engine/Source/Runtime/Engine/Private/WorldPartition/WorldPartitionRuntimeSpatialHash.cpp b/Engine/Source/Runtime/Engine/Private/WorldPartition/WorldPartitionRuntimeSpatialHash.cpp
--- a/Engine/Source/Runtime/Engine/Private/WorldPartition/WorldPartitionRuntimeSpatialHash.cpp
+++ b/Engine/Source/Runtime/Engine/Private/WorldPartition/WorldPartitionRuntimeSpatialHash.cpp
@@ -556,36 +556,39 @@ bool UWorldPartitionRuntimeSpatialHash::CreateStreamingGrid(const FSpatialHashRu
 
 			for (const FSquare2DGridHelper::FGridLevel::FGridCellDataChunk& GridCellDataChunk : TempCell.GetDataChunks())
 			{
+				// Cell cannot be treated as always loaded if it has data layers
+				const bool bIsCellAlwaysLoaded = (&TempCell == &PartionedActors.GetAlwaysLoadedCell()) && !GridCellDataChunk.HasDataLayers();
+				
 				FilteredActors.SetNum(0, false);
 				FilteredActors.Reset(GridCellDataChunk.GetActors().Num());
 				if (GridCellDataChunk.GetActors().Num())
 				{
-					Algo::TransformIf(GridCellDataChunk.GetActors(), FilteredActors, [](const FActorInstance& ActorInstance)
+					for (const FActorInstance& ActorInstance : GridCellDataChunk.GetActors())
 					{
 						const FWorldPartitionActorDesc* ActorDesc = ActorInstance.GetActorDesc();
-						const bool bShouldStripActorFromStreaming = ActorInstance.ShouldStripFromStreaming();
-						UE_CLOG(bShouldStripActorFromStreaming, LogWorldPartitionRuntimeSpatialHash, Verbose, TEXT("Stripping Actor %s (%s) from streaming grid (Container %08x)"), 
-							*(ActorDesc->GetActorPath().ToString()), *ActorInstance.Actor.ToString(EGuidFormats::UniqueObjectGuid), ActorInstance.ContainerInstance->ID);
-						return !bShouldStripActorFromStreaming;
-					}, [](const FActorInstance& ActorInstance) { return ActorInstance; });
-				}
-
-				// Cell cannot be treated as always loaded if it has data layers
-				const bool bIsCellAlwaysLoaded = (&TempCell == &PartionedActors.GetAlwaysLoadedCell()) && !GridCellDataChunk.HasDataLayers();
+						if (ActorInstance.ShouldStripFromStreaming())
+						{
+							UE_LOG(LogWorldPartitionRuntimeSpatialHash, Verbose, TEXT("Stripping Actor %s (%s) from streaming grid (Container %08x)"),
+								*(ActorDesc->GetActorPath().ToString()), *ActorInstance.Actor.ToString(EGuidFormats::UniqueObjectGuid), ActorInstance.ContainerInstance->ID);
+							continue;
+						}
 
-				// In PIE, Always loaded cell is not generated. Instead, always loaded actors will be added to AlwaysLoadedActorsForPIE.
-				// This will trigger loading/registration of these actors in the PersistentLevel (if not already loaded).
-				// Then, duplication of world for PIE will duplicate only these actors. 
-				// When stopping PIE, WorldPartition will release these FWorldPartitionReferences which 
-				// will unload actors that were not already loaded in the non PIE world.
-				if (bIsCellAlwaysLoaded && bIsMainWorldPartition && (Mode == EWorldPartitionStreamingMode::PIE))
-				{
-					for (const FActorInstance& ActorInstance : FilteredActors)
-					{
-						check(ActorInstance.ContainerInstance->Container == WorldPartition);
-						AlwaysLoadedActorsForPIE.Emplace(WorldPartition, ActorInstance.Actor);
+						// In PIE, Always loaded cell is not generated. Instead, always loaded actors will be added to AlwaysLoadedActorsForPIE.
+						// This will trigger loading/registration of these actors in the PersistentLevel (if not already loaded).
+						// Then, duplication of world for PIE will duplicate only these actors. 
+						// When stopping PIE, WorldPartition will release these FWorldPartitionReferences which 
+						// will unload actors that were not already loaded in the non PIE world.
+						if (bIsCellAlwaysLoaded && bIsMainWorldPartition && (Mode == EWorldPartitionStreamingMode::PIE))
+						{
+							if (ActorInstance.ContainerInstance->Container == WorldPartition)
+							{
+								AlwaysLoadedActorsForPIE.Emplace(WorldPartition, ActorInstance.Actor);
+								continue;
+							}
+						}
+																
+						FilteredActors.Add(ActorInstance);
 					}
-					continue;
 				}
 
 				if (!FilteredActors.Num())
