Fix Non-Unity Build: remove inline

diff --git a/Engine/Source/Runtime/Engine/Public/WorldPartition/WorldPartitionActorDesc.h b/Engine/Source/Runtime/Engine/Public/WorldPartition/WorldPartitionActorDesc.h
--- a/Engine/Source/Runtime/Engine/Public/WorldPartition/WorldPartitionActorDesc.h
+++ b/Engine/Source/Runtime/Engine/Public/WorldPartition/WorldPartitionActorDesc.h
@@ -65,7 +65,7 @@ public:
 	inline bool GetActorIsEditorOnly() const { return bActorIsEditorOnly; }
 	inline bool GetLevelBoundsRelevant() const { return bLevelBoundsRelevant; }
 	inline bool GetActorIsHLODRelevant() const { return bActorIsHLODRelevant; }
-	inline class UHLODLayer* GetHLODLayer() const;
+	class UHLODLayer* GetHLODLayer() const;
 	inline const TArray<FName>& GetDataLayers() const { return DataLayers; }
 	inline FName GetActorPackage() const { return ActorPackage; }
 	inline FName GetActorPath() const { return ActorPath; }
