Do not iniitialize WorldPartition when running WorldPartitionConvertCommandlet.

diff --git a/Engine/Source/Runtime/Engine/Private/WorldPartition/WorldPartitionSubsystem.cpp b/Engine/Source/Runtime/Engine/Private/WorldPartition/WorldPartitionSubsystem.cpp
--- a/Engine/Source/Runtime/Engine/Private/WorldPartition/WorldPartitionSubsystem.cpp
+++ b/Engine/Source/Runtime/Engine/Private/WorldPartition/WorldPartitionSubsystem.cpp
@@ -60,6 +60,16 @@ void UWorldPartitionSubsystem::PostInitialize()
 {
 	Super::PostInitialize();
 
+#if WITH_EDITOR
+	static UClass* WorldPartitionConvertCommandletClass = FindObject<UClass>(ANY_PACKAGE, TEXT("WorldPartitionConvertCommandlet"), true);
+	check(WorldPartitionConvertCommandletClass);
+	const bool bIsRunningWorldPartitionConvertCommandlet = GetRunningCommandletClass() && GetRunningCommandletClass()->IsChildOf(WorldPartitionConvertCommandletClass);
+	if (bIsRunningWorldPartitionConvertCommandlet)
+	{
+		return;
+	}
+#endif
+
 	if (UWorldPartition* MainPartition = GetMainWorldPartition())
 	{
 		MainPartition->Initialize(GetWorld(), FTransform::Identity);
