Fix curve weights out of bounds access in FInertializationPoseDiff::InitFrom.

This is caused by the PoseSnapshots.Swap(i, i + 1) call in Evaluate_AnyThread. When the snapshots were memory swapped their respective curve LUT pointers in FBlendedHeapCurve weren't getting fixed up. The inertialization curve snapshot is now wrapped in a structure with move assignment that fixes up the LUT and the snapshot swap operation now uses move assignment instead of bitwise memory swapping.

diff --git a/Engine/Source/Runtime/Engine/Classes/Animation/AnimNode_Inertialization.h b/Engine/Source/Runtime/Engine/Classes/Animation/AnimNode_Inertialization.h
--- a/Engine/Source/Runtime/Engine/Classes/Animation/AnimNode_Inertialization.h
+++ b/Engine/Source/Runtime/Engine/Classes/Animation/AnimNode_Inertialization.h
@@ -64,6 +64,53 @@ enum class EInertializationSpace : uint8
 	WorldRotation	// Inertialize rotation only in world space (to conceal discontinuities in actor orientation)
 };
 
+struct ENGINE_API FInertializationCurve
+{
+	FBlendedHeapCurve BlendedCurve;
+	TArray<uint16> CurveUIDToArrayIndexLUT;
+
+	FInertializationCurve() = default;
+
+	FInertializationCurve(const FInertializationCurve& Other)
+	{
+		*this = Other;
+	}
+
+	FInertializationCurve(FInertializationCurve&& Other)
+	{
+		*this = MoveTemp(Other);
+	}
+
+	FInertializationCurve& operator=(const FInertializationCurve& Other)
+	{
+		BlendedCurve.CopyFrom(Other.BlendedCurve);
+		BlendedCurve.UIDToArrayIndexLUT = &CurveUIDToArrayIndexLUT;
+		CurveUIDToArrayIndexLUT = Other.CurveUIDToArrayIndexLUT;
+		return *this;
+	}
+
+	FInertializationCurve& operator=(FInertializationCurve&& Other)
+	{
+		BlendedCurve.MoveFrom(Other.BlendedCurve);
+		BlendedCurve.UIDToArrayIndexLUT = &CurveUIDToArrayIndexLUT;
+		CurveUIDToArrayIndexLUT = MoveTemp(Other.CurveUIDToArrayIndexLUT);
+		return *this;
+	}
+
+	template <typename OtherAllocator>
+	void InitFrom(const FBaseBlendedCurve<OtherAllocator>& Other)
+	{
+		CurveUIDToArrayIndexLUT.Reset();
+
+		BlendedCurve.CopyFrom(Other);
+		BlendedCurve.UIDToArrayIndexLUT = &CurveUIDToArrayIndexLUT;
+
+		if (Other.UIDToArrayIndexLUT)
+		{
+			CurveUIDToArrayIndexLUT = *Other.UIDToArrayIndexLUT;
+		}
+	}
+};
 
 USTRUCT()
 struct FInertializationPose
@@ -84,8 +131,8 @@ struct FInertializationPose
 
 	// Snapshot of active curves
 	// 
-	FBlendedHeapCurve Curves;
-	TArray<uint16> CurveUIDToArrayIndexLUT;
+	FInertializationCurve Curves;
+
 
 	FName AttachParentName;
 	float DeltaTime;
@@ -96,10 +143,20 @@ struct FInertializationPose
 		, DeltaTime(0.0f)
 	{
 	}
+	
+	FInertializationPose(const FInertializationPose&) = default;
+	FInertializationPose(FInertializationPose&&) = default;
+	FInertializationPose& operator=(const FInertializationPose&) = default;
+	FInertializationPose& operator=(FInertializationPose&&) = default;
 
 	void InitFrom(const FCompactPose& Pose, const FBlendedCurve& InCurves, const FTransform& InComponentTransform, const FName& InAttachParentName, float InDeltaTime);
 };
 
+template <>
+struct TUseBitwiseSwap<FInertializationPose>
+{
+	enum { Value = false };
+};
 
 USTRUCT()
 struct FInertializationBoneDiff
diff --git a/Engine/Source/Runtime/Engine/Private/Animation/AnimNode_Inertialization.cpp b/Engine/Source/Runtime/Engine/Private/Animation/AnimNode_Inertialization.cpp
--- a/Engine/Source/Runtime/Engine/Private/Animation/AnimNode_Inertialization.cpp
+++ b/Engine/Source/Runtime/Engine/Private/Animation/AnimNode_Inertialization.cpp
@@ -307,7 +307,11 @@ void FAnimNode_Inertialization::Evaluate_AnyThread(FPoseContext& Output)
 		// Bubble the old poses forward in the buffer (using swaps to avoid allocations and copies)
 		for (int32 i = 0; i < INERTIALIZATION_MAX_POSE_SNAPSHOTS - 1; ++i)
 		{
-			PoseSnapshots.Swap(i, i + 1);
+			Swap(PoseSnapshots[i], PoseSnapshots[i+1]);
+
+			// Verify the swap operation used move assignment to properly fix up curve LUT pointers
+			checkSlow(PoseSnapshots[i].Curves.BlendedCurve.UIDToArrayIndexLUT == &PoseSnapshots[i].Curves.CurveUIDToArrayIndexLUT);
+			checkSlow(PoseSnapshots[i+1].Curves.BlendedCurve.UIDToArrayIndexLUT == &PoseSnapshots[i+1].Curves.CurveUIDToArrayIndexLUT);
 		}
 
 		// Overwrite the (now irrelevant) pose in the last slot with the new post snapshot
@@ -428,15 +432,7 @@ void FInertializationPose::InitFrom(const FCompactPose& Pose, const FBlendedCurv
 		}
 	}
 
-	// Copy curves and make a new copy of the curve lookup table as it may change from frame to frame
-	Curves.CopyFrom(InCurves);
-	CurveUIDToArrayIndexLUT.Reset();
-	if (InCurves.UIDToArrayIndexLUT)
-	{
-		CurveUIDToArrayIndexLUT = *InCurves.UIDToArrayIndexLUT;
-		Curves.UIDToArrayIndexLUT = &CurveUIDToArrayIndexLUT;
-	}
-
+	Curves.InitFrom(InCurves);
 	ComponentTransform = InComponentTransform;
 	AttachParentName = InAttachParentName;
 	DeltaTime = InDeltaTime;
@@ -627,7 +623,7 @@ void FInertializationPoseDiff::InitFrom(const FCompactPose& Pose, const FBlended
 	for (int32 CurveUID = 0; CurveUID != CurveNum; ++CurveUID)
 	{
 		const int32 CurrIdx = Curves.GetArrayIndexByUID(CurveUID);
-		const int32 Prev1Idx = Prev1.Curves.GetArrayIndexByUID(CurveUID);
+		const int32 Prev1Idx = Prev1.Curves.BlendedCurve.GetArrayIndexByUID(CurveUID);
 		if (CurrIdx == INDEX_NONE || Prev1Idx == INDEX_NONE)
 		{
 			// CurveDiff is zeroed
@@ -641,17 +637,17 @@ void FInertializationPoseDiff::InitFrom(const FCompactPose& Pose, const FBlended
 		}
 
 		const float CurrWeight = Curves.CurveWeights[CurrIdx];
-		const float Prev1Weight = Prev1.Curves.CurveWeights[Prev1Idx];
+		const float Prev1Weight = Prev1.Curves.BlendedCurve.CurveWeights[Prev1Idx];
 		FInertializationCurveDiff& CurveDiff = CurveDiffs[CurveUID];
 
 		// Note we intentionally ignore FBlendedCurve::ValidCurveWeights. We want to ease in/out when only one
 		// curve is valid, and we'll compute a zero delta and derivative when both are invalid.
 		CurveDiff.Delta = Prev1Weight - CurrWeight;
 
-		const int32 Prev2Idx = Prev2.Curves.GetArrayIndexByUID(CurveUID);
+		const int32 Prev2Idx = Prev2.Curves.BlendedCurve.GetArrayIndexByUID(CurveUID);
 		if (Prev2Idx != INDEX_NONE && Prev1.DeltaTime > KINDA_SMALL_NUMBER)
 		{
-			const float Prev2Weight = Prev2.Curves.CurveWeights[Prev2Idx];
+			const float Prev2Weight = Prev2.Curves.BlendedCurve.CurveWeights[Prev2Idx];
 			CurveDiff.Derivative = (Prev1Weight - Prev2Weight) / Prev1.DeltaTime;
 		}
 	}
