LevelInstance: Build menu is now silent (doesn't pop a save dialog on each BP)

diff --git a/Engine/Source/Runtime/Engine/Private/LevelInstance/LevelInstanceSubsystem.cpp b/Engine/Source/Runtime/Engine/Private/LevelInstance/LevelInstanceSubsystem.cpp
--- a/Engine/Source/Runtime/Engine/Private/LevelInstance/LevelInstanceSubsystem.cpp
+++ b/Engine/Source/Runtime/Engine/Private/LevelInstance/LevelInstanceSubsystem.cpp
@@ -435,9 +435,10 @@ void ULevelInstanceSubsystem::PackLevelInstances()
 	}
 
 	TSharedPtr<FPackedLevelInstanceBuilder> Builder = FPackedLevelInstanceBuilder::CreateDefaultBuilder();
+	const bool bSilentUpdate = true;
 	for (UBlueprint* Blueprint : BlueprintsToUpdate)
 	{
-		Builder->UpdateBlueprint(Blueprint);
+		Builder->UpdateBlueprint(Blueprint, bSilentUpdate);
 		UpdateProgress();
 	}
 }
diff --git a/Engine/Source/Runtime/Engine/Private/LevelInstance/Packed/PackedLevelInstanceBuilder.cpp b/Engine/Source/Runtime/Engine/Private/LevelInstance/Packed/PackedLevelInstanceBuilder.cpp
--- a/Engine/Source/Runtime/Engine/Private/LevelInstance/Packed/PackedLevelInstanceBuilder.cpp
+++ b/Engine/Source/Runtime/Engine/Private/LevelInstance/Packed/PackedLevelInstanceBuilder.cpp
@@ -158,12 +158,12 @@ TSharedPtr<FPackedLevelInstanceBuilder> FPackedLevelInstanceBuilder::CreateDefau
 	return Builder;
 }
 
-void FPackedLevelInstanceBuilder::PackActor(APackedLevelInstance* InPackedLevelInstance)
+bool FPackedLevelInstanceBuilder::PackActor(APackedLevelInstance* InPackedLevelInstance)
 {
-	PackActor(InPackedLevelInstance, InPackedLevelInstance);
+	return PackActor(InPackedLevelInstance, InPackedLevelInstance);
 }
 
-void FPackedLevelInstanceBuilder::PackActor(APackedLevelInstance* InPackedLevelInstance, ALevelInstance* InLevelInstanceToPack)
+bool FPackedLevelInstanceBuilder::PackActor(APackedLevelInstance* InPackedLevelInstance, ALevelInstance* InLevelInstanceToPack)
 {
 	FMessageLog LevelInstanceLog("LevelInstance");
 	LevelInstanceLog.Info(FText::Format(LOCTEXT("PackingStarted", "Packing of '{0}' started..."), FText::FromString(InPackedLevelInstance->GetWorldAssetPackage())));
@@ -176,6 +176,12 @@ void FPackedLevelInstanceBuilder::PackActor(APackedLevelInstance* InPackedLevelI
 	check(LevelInstanceSubystem);
 	
 	ULevel* SourceLevel = LevelInstanceSubystem->GetLevelInstanceLevel(InLevelInstanceToPack);
+	if (!SourceLevel)
+	{
+		LevelInstanceLog.Error(FText::Format(LOCTEXT("FailedPackingNoLevel", "Packing of '{0}' failed"), FText::FromString(InPackedLevelInstance->GetWorldAssetPackage())));
+		return false;
+	}
+
 	ULevelStreaming* SourceLevelStreaming = ULevelStreaming::FindStreamingLevel(SourceLevel);
 	AWorldSettings* WorldSettings = SourceLevel->GetWorldSettings();
 	
@@ -202,6 +208,7 @@ void FPackedLevelInstanceBuilder::PackActor(APackedLevelInstance* InPackedLevelI
 	}
 			
 	Context.Report(LevelInstanceLog);
+	return true;
 }
 
 void FPackedLevelInstanceBuilderContext::Report(FMessageLog& LevelInstanceLog) const
@@ -299,47 +306,49 @@ ALevelInstance* FPackedLevelInstanceBuilder::CreateTransientLevelInstanceForPack
 	return LevelInstance;
 }
 
-void FPackedLevelInstanceBuilder::PackActor(APackedLevelInstance* InActor, TSoftObjectPtr<UWorld> InWorldAsset)
+bool FPackedLevelInstanceBuilder::PackActor(APackedLevelInstance* InActor, TSoftObjectPtr<UWorld> InWorldAsset)
 {
 	ALevelInstance* TransientLevelInstance = CreateTransientLevelInstanceForPacking(InWorldAsset, InActor->GetActorLocation(), InActor->GetActorRotation());
 
-	PackActor(InActor, TransientLevelInstance);
+	bool bResult = PackActor(InActor, TransientLevelInstance);
 
 	TransientLevelInstance->GetWorld()->DestroyActor(TransientLevelInstance);
+
+	return bResult;
 }
 
-void FPackedLevelInstanceBuilder::UpdateBlueprint(UBlueprint* Blueprint)
+void FPackedLevelInstanceBuilder::UpdateBlueprint(UBlueprint* Blueprint, bool bInSilentUpdate)
 {
 	APackedLevelInstance* CDO = CastChecked<APackedLevelInstance>(Blueprint->GeneratedClass->GetDefaultObject());
 	check(CDO);
 
-	CreateOrUpdateBlueprint(CDO->GetWorldAsset(), Blueprint);
+	CreateOrUpdateBlueprint(CDO->GetWorldAsset(), Blueprint, bInSilentUpdate);
 }
 
-bool FPackedLevelInstanceBuilder::CreateOrUpdateBlueprint(TSoftObjectPtr<UWorld> InWorldAsset, TSoftObjectPtr<UBlueprint> InBlueprintAsset)
+bool FPackedLevelInstanceBuilder::CreateOrUpdateBlueprint(TSoftObjectPtr<UWorld> InWorldAsset, TSoftObjectPtr<UBlueprint> InBlueprintAsset, bool bInSilentUpdate)
 {
 	bool bResult = true;
 	
 	ALevelInstance* TransientLevelInstance = CreateTransientLevelInstanceForPacking(InWorldAsset, FVector::ZeroVector, FRotator::ZeroRotator);
 	
-	bResult = CreateOrUpdateBlueprintFromUnpacked(TransientLevelInstance, InBlueprintAsset);
+	bResult = CreateOrUpdateBlueprintFromUnpacked(TransientLevelInstance, InBlueprintAsset, bInSilentUpdate);
 
 	TransientLevelInstance->GetWorld()->DestroyActor(TransientLevelInstance);
 
 	return bResult;
 }
 
-bool FPackedLevelInstanceBuilder::CreateOrUpdateBlueprint(ALevelInstance* InLevelInstance, TSoftObjectPtr<UBlueprint> InBlueprintAsset)
+bool FPackedLevelInstanceBuilder::CreateOrUpdateBlueprint(ALevelInstance* InLevelInstance, TSoftObjectPtr<UBlueprint> InBlueprintAsset, bool bInSilentUpdate)
 {
 	if (APackedLevelInstance* PackedLevelInstance = Cast<APackedLevelInstance>(InLevelInstance))
 	{
-		return CreateOrUpdateBlueprintFromPacked(PackedLevelInstance, InBlueprintAsset);
+		return CreateOrUpdateBlueprintFromPacked(PackedLevelInstance, InBlueprintAsset, bInSilentUpdate);
 	}
 	
-	return CreateOrUpdateBlueprintFromUnpacked(InLevelInstance, InBlueprintAsset);
+	return CreateOrUpdateBlueprintFromUnpacked(InLevelInstance, InBlueprintAsset, bInSilentUpdate);
 }
 
-bool FPackedLevelInstanceBuilder::CreateOrUpdateBlueprintFromUnpacked(ALevelInstance* InActor, TSoftObjectPtr<UBlueprint> InBlueprintAsset)
+bool FPackedLevelInstanceBuilder::CreateOrUpdateBlueprintFromUnpacked(ALevelInstance* InActor, TSoftObjectPtr<UBlueprint> InBlueprintAsset, bool bInSilentUpdate)
 {
 	bool bResult = true;
 	
@@ -355,18 +364,26 @@ bool FPackedLevelInstanceBuilder::CreateOrUpdateBlueprintFromUnpacked(ALevelInst
 
 	APackedLevelInstance* PackedLevelInstance = World->SpawnActor<APackedLevelInstance>(InActor->GetActorLocation(), InActor->GetActorRotation(), SpawnParams);
 	PackedLevelInstance->SetWorldAsset(InActor->GetWorldAsset());
+	ON_SCOPE_EXIT
+	{
+		if (PackedLevelInstance)
+		{
+			InActor->GetWorld()->DestroyActor(PackedLevelInstance);
+		}
+	};
 
-	PackActor(PackedLevelInstance, InActor);
+	if (!PackActor(PackedLevelInstance, InActor))
+	{
+		return false;
+	}
 
 	PackedLevelInstance->BlueprintAsset = InBlueprintAsset;
-	bResult &= CreateOrUpdateBlueprintFromPacked(PackedLevelInstance, PackedLevelInstance->BlueprintAsset);
-
-	InActor->GetWorld()->DestroyActor(PackedLevelInstance);
+	bResult &= CreateOrUpdateBlueprintFromPacked(PackedLevelInstance, PackedLevelInstance->BlueprintAsset, bInSilentUpdate);
 
 	return bResult;
 }
 
-bool FPackedLevelInstanceBuilder::CreateOrUpdateBlueprintFromPacked(APackedLevelInstance* InActor, TSoftObjectPtr<UBlueprint> InBlueprintAsset)
+bool FPackedLevelInstanceBuilder::CreateOrUpdateBlueprintFromPacked(APackedLevelInstance* InActor, TSoftObjectPtr<UBlueprint> InBlueprintAsset, bool bInSilentUpdate)
 {
 	UBlueprint* BP = nullptr;
 	if (!InBlueprintAsset.IsNull())
@@ -428,12 +445,17 @@ bool FPackedLevelInstanceBuilder::CreateOrUpdateBlueprintFromPacked(APackedLevel
 	// Synchronous compile
 	FKismetEditorUtilities::CompileBlueprint(BP, EBlueprintCompileOptions::SkipGarbageCollection);
 
-	const bool bCheckDirty = false;
-	const bool bPromptToSave = true;
-	TArray<UPackage*> OutFailedPackages;
-	FEditorFileUtils::PromptForCheckoutAndSave({ BP->GetPackage() }, bCheckDirty, bPromptToSave, &OutFailedPackages);
-		
-	return !OutFailedPackages.Num();
+	if (!bInSilentUpdate)
+	{
+		const bool bCheckDirty = false;
+		const bool bPromptToSave = true;
+		TArray<UPackage*> OutFailedPackages;
+		FEditorFileUtils::PromptForCheckoutAndSave({ BP->GetPackage() }, bCheckDirty, bPromptToSave, &OutFailedPackages);
+
+		return !OutFailedPackages.Num();
+	}
+
+	return true;
 }
 #endif
 
diff --git a/Engine/Source/Runtime/Engine/Public/LevelInstance/Packed/PackedLevelInstanceBuilder.h b/Engine/Source/Runtime/Engine/Public/LevelInstance/Packed/PackedLevelInstanceBuilder.h
--- a/Engine/Source/Runtime/Engine/Public/LevelInstance/Packed/PackedLevelInstanceBuilder.h
+++ b/Engine/Source/Runtime/Engine/Public/LevelInstance/Packed/PackedLevelInstanceBuilder.h
@@ -29,18 +29,18 @@ public:
 	static TSharedPtr<FPackedLevelInstanceBuilder> CreateDefaultBuilder();
 	
 	/* Packs InPackedLevelInstance using InLevelInstanceToPack as its source level actor */
-	void PackActor(APackedLevelInstance* InPackedLevelInstance, ALevelInstance* InLevelInstanceToPack);
+	bool PackActor(APackedLevelInstance* InPackedLevelInstance, ALevelInstance* InLevelInstanceToPack);
 	/* Packs InPackedLevelInstance using itself as the source level actor */
-	void PackActor(APackedLevelInstance* InPackedLevelInstance);
+	bool PackActor(APackedLevelInstance* InPackedLevelInstance);
 	/* Packs InPackedLevelInstance using InWorldAsset as the source level */
-	void PackActor(APackedLevelInstance* InPackedLevelInstance, TSoftObjectPtr<UWorld> InWorldAsset);
+	bool PackActor(APackedLevelInstance* InPackedLevelInstance, TSoftObjectPtr<UWorld> InWorldAsset);
 
 	/* Creates/Updates a APackedLevelInstance Blueprint from InLevelInstance (will overwrite existing asset or show a dialog if InBlueprintAsset is null) */
-	bool CreateOrUpdateBlueprint(ALevelInstance* InLevelInstance, TSoftObjectPtr<UBlueprint> InBlueprintAsset);
+	bool CreateOrUpdateBlueprint(ALevelInstance* InLevelInstance, TSoftObjectPtr<UBlueprint> InBlueprintAsset, bool bSilentUpdate = false);
 	/* Creates/Updates a APackedLeveInstance Blueprint from InWorldAsset (will overwrite existing asset or show a dialog if InBlueprintAsset is null) */
-	bool CreateOrUpdateBlueprint(TSoftObjectPtr<UWorld> InWorldAsset, TSoftObjectPtr<UBlueprint> InBlueprintAsset);
+	bool CreateOrUpdateBlueprint(TSoftObjectPtr<UWorld> InWorldAsset, TSoftObjectPtr<UBlueprint> InBlueprintAsset, bool bSilentUpdate = false);
 	/* Update existing Blueprint */
-	void UpdateBlueprint(UBlueprint* Blueprint);
+	void UpdateBlueprint(UBlueprint* Blueprint, bool bSilentUpdate = false);
 
 	static const FString& GetPackedBPSuffix();
 	/* Creates a new APackedLevelInstance Blueprint using InPackagePath/InAssetName as hint for path. Prompts the user to input the final asset name. */
@@ -50,9 +50,9 @@ public:
 	
 private:
 	/* Create/Updates a APackedLevelInstance Blueprint from InPackedActor (will overwrite existing asset or show a dialog if InBlueprintAsset is null) */
-	bool CreateOrUpdateBlueprintFromPacked(APackedLevelInstance* InPackedActor, TSoftObjectPtr<UBlueprint> InBlueprintAsset);
+	bool CreateOrUpdateBlueprintFromPacked(APackedLevelInstance* InPackedActor, TSoftObjectPtr<UBlueprint> InBlueprintAsset, bool bInSilentUpdate);
 	/* Creates/Updates a APackedLevelInstance Blueprint from InLevelInstance (will overwrite existing asset or show a dialog if InBlueprintAsset is null) */
-	bool CreateOrUpdateBlueprintFromUnpacked(ALevelInstance* InPackedActor, TSoftObjectPtr<UBlueprint> InBlueprintAsset);
+	bool CreateOrUpdateBlueprintFromUnpacked(ALevelInstance* InPackedActor, TSoftObjectPtr<UBlueprint> InBlueprintAsset, bool bInSilentUpdate);
 	/* Creates and loads a ALevelInstance so it can be used for packing */
 	ALevelInstance* CreateTransientLevelInstanceForPacking(TSoftObjectPtr<UWorld> InWorldAsset, const FVector& InLocation, const FRotator& InRotator);
 
