Fix ensure on exit due to FPrecomputedVolumetricLightmap

diff --git a/Engine/Source/Runtime/Engine/Private/PrecomputedVolumetricLightmap.cpp b/Engine/Source/Runtime/Engine/Private/PrecomputedVolumetricLightmap.cpp
--- a/Engine/Source/Runtime/Engine/Private/PrecomputedVolumetricLightmap.cpp
+++ b/Engine/Source/Runtime/Engine/Private/PrecomputedVolumetricLightmap.cpp
@@ -727,7 +727,10 @@ void FPrecomputedVolumetricLightmap::RemoveFromScene(FSceneInterface* Scene)
 		bAddedToScene = false;
 
 		// Certain paths in the editor (namely, ReloadPackages and ForceDelete) will GC the registry before the UWorld destruction (which destructs FScene)
-		ensureMsgf(SourceRegistry.IsValid(), TEXT("UMapBuildDataRegistry is garbage collected before an FPrecomputedVolumetricLightmap is removed from the scene. Is there a missing ReleaseRenderingResources() call?"));
+		ensureMsgf(
+			SourceRegistry.IsValid() // either SourceRegistry is valid
+			|| (!Scene->GetWorld() || Scene->GetWorld()->IsPendingKillOrUnreachable()) // or the world we're in is going away (usually during shutdown)
+			, TEXT("UMapBuildDataRegistry is garbage collected before an FPrecomputedVolumetricLightmap is removed from the scene. Is there a missing ReleaseRenderingResources() call?"));
 
 		// While that can be explained as missing ReleaseRenderingResources() calls, this fail-safe guard is added here
 		if (SourceRegistry.IsValid())
