Fix a bug of default  AA method for mobile.

diff --git a/Engine/Source/Runtime/Engine/Private/SceneView.cpp b/Engine/Source/Runtime/Engine/Private/SceneView.cpp
--- a/Engine/Source/Runtime/Engine/Private/SceneView.cpp
+++ b/Engine/Source/Runtime/Engine/Private/SceneView.cpp
@@ -2590,7 +2590,7 @@ FSceneViewFamily::FSceneViewFamily(const ConstructionValues& CVS)
 
 	// Check if the translucency are allowed to be rendered after DOF, if not, translucency after DOF will be rendered in standard translucency.
 	{
-		bAllowTranslucencyAfterDOF = CVarAllowTranslucencyAfterDOF.GetValueOnRenderThread() != 0
+		bAllowTranslucencyAfterDOF = CVarAllowTranslucencyAfterDOF.GetValueOnAnyThread() != 0
 			&& EngineShowFlags.PostProcessing // Used for reflection captures.
 			&& !UseDebugViewPS()
 			&& EngineShowFlags.SeparateTranslucency;
