PostProcessVolume + Skylight now respect the different editor Hidden flags (Outliner, DataLayers, Level)

diff --git a/Engine/Source/Runtime/Engine/Classes/Components/PostProcessComponent.h b/Engine/Source/Runtime/Engine/Classes/Components/PostProcessComponent.h
--- a/Engine/Source/Runtime/Engine/Classes/Components/PostProcessComponent.h
+++ b/Engine/Source/Runtime/Engine/Classes/Components/PostProcessComponent.h
@@ -54,7 +54,7 @@ class UPostProcessComponent : public USceneComponent, public IInterface_PostProc
 	ENGINE_API virtual FPostProcessVolumeProperties GetProperties() const override
 	{
 		FPostProcessVolumeProperties Ret;
-		Ret.bIsEnabled = bEnabled != 0;
+		Ret.bIsEnabled = bEnabled != 0 && ShouldRender();
 		Ret.bIsUnbound = bUnbound != 0 || Cast<UShapeComponent>(GetAttachParent()) == nullptr;
 		Ret.BlendRadius = BlendRadius;
 		Ret.BlendWeight = BlendWeight;
diff --git a/Engine/Source/Runtime/Engine/Classes/Engine/PostProcessVolume.h b/Engine/Source/Runtime/Engine/Classes/Engine/PostProcessVolume.h
--- a/Engine/Source/Runtime/Engine/Classes/Engine/PostProcessVolume.h
+++ b/Engine/Source/Runtime/Engine/Classes/Engine/PostProcessVolume.h
@@ -55,7 +55,7 @@ class ENGINE_API APostProcessVolume : public AVolume, public IInterface_PostProc
 	virtual FPostProcessVolumeProperties GetProperties() const override
 	{
 		FPostProcessVolumeProperties Ret;
-		Ret.bIsEnabled = bEnabled != 0;
+		Ret.bIsEnabled = IsPPVEnabled();
 		Ret.bIsUnbound = bUnbound != 0;
 		Ret.BlendRadius = BlendRadius;
 		Ret.BlendWeight = BlendWeight;
@@ -76,6 +76,8 @@ class ENGINE_API APostProcessVolume : public AVolume, public IInterface_PostProc
 protected:
 	virtual void PostRegisterAllComponents() override;
 	//~ End AActor Interface
+
+	bool IsPPVEnabled() const;
 public:
 	
 	//~ Begin UObject Interface
diff --git a/Engine/Source/Runtime/Engine/Private/Components/SkyLightComponent.cpp b/Engine/Source/Runtime/Engine/Private/Components/SkyLightComponent.cpp
--- a/Engine/Source/Runtime/Engine/Private/Components/SkyLightComponent.cpp
+++ b/Engine/Source/Runtime/Engine/Private/Components/SkyLightComponent.cpp
@@ -371,19 +371,9 @@ void USkyLightComponent::CreateRenderState_Concurrent(FRegisterComponentContext*
 {
 	Super::CreateRenderState_Concurrent(Context);
 
-	bool bHidden = false;
-#if WITH_EDITORONLY_DATA
-	bHidden = GetOwner() ? GetOwner()->bHiddenEdLevel : false;
-#endif // WITH_EDITORONLY_DATA
-
-	if(!ShouldComponentAddToScene())
-	{
-		bHidden = true;
-	}
-
 	const bool bIsValid = SourceType != SLS_SpecifiedCubemap || Cubemap != NULL || IsRealTimeCaptureEnabled();
 
-	if (bAffectsWorld && GetVisibleFlag() && !bHidden && bIsValid)
+	if (bAffectsWorld && bIsValid && ShouldComponentAddToScene() && ShouldRender())
 	{
 		// Create the light's scene proxy.
 		SceneProxy = CreateSceneProxy();
@@ -1063,7 +1053,7 @@ ASkyLight::ASkyLight(const FObjectInitializer& ObjectInitializer)
 {
 	LightComponent = CreateDefaultSubobject<USkyLightComponent>(TEXT("SkyLightComponent0"));
 	RootComponent = LightComponent;
-
+	SetHidden(false);
 #if WITH_EDITORONLY_DATA
 	// Null out the sprite. The Skylight components sprite is the one we use.
 	if (GetSpriteComponent())
diff --git a/Engine/Source/Runtime/Engine/Private/PostProcessVolume.cpp b/Engine/Source/Runtime/Engine/Private/PostProcessVolume.cpp
--- a/Engine/Source/Runtime/Engine/Private/PostProcessVolume.cpp
+++ b/Engine/Source/Runtime/Engine/Private/PostProcessVolume.cpp
@@ -18,6 +18,20 @@ APostProcessVolume::APostProcessVolume(const FObjectInitializer& ObjectInitializ
 	BlendWeight = 1.0f;
 }
 
+bool APostProcessVolume::IsPPVEnabled() const
+{
+#if WITH_EDITOR
+	const bool bShowInEditor = GIsEditor ? !IsHiddenEd() : false;
+	const bool bInGameWorld = GetWorld() && GetWorld()->UsesGameHiddenFlags();
+
+	// bEnabled is the only thing we check in a game world. In the editor we also check the editor hidden flags
+	// In a GameWorld we can't use the IsHidden() flag of the Actor because the APostProcessVolume is always Hidden by subclassing of ABrush
+	return bEnabled != 0 && (bInGameWorld || (!bInGameWorld && bShowInEditor));
+#else
+	return bEnabled != 0;
+#endif
+}
+
 bool APostProcessVolume::EncompassesPoint(FVector Point, float SphereRadius/*=0.f*/, float* OutDistanceToPoint)
 {
 	// Redirect IInterface_PostProcessVolume's non-const pure virtual EncompassesPoint virtual in to AVolume's non-virtual const EncompassesPoint
