Fix for Android compiler error

diff --git a/Engine/Source/Runtime/Engine/Public/SceneUtils.h b/Engine/Source/Runtime/Engine/Public/SceneUtils.h
--- a/Engine/Source/Runtime/Engine/Public/SceneUtils.h
+++ b/Engine/Source/Runtime/Engine/Public/SceneUtils.h
@@ -10,6 +10,7 @@
 // Each event type will be displayed using the defined color
 #pragma once
 
+#include "UObject/ObjectMacros.h"
 #include "ProfilingDebugging/RealtimeGPUProfiler.h"
 #include "RHIDefinitions.h"
 
