WorldPartitionActorDescView
- Added FWorldPartitionActorDescView that can override some properties during the streaming generation process, currently used for automatic grid placement automatic changes based on bounds.
- Replaced all accesses to WorldPartitionActorDesc with WorldPartitionActorDescView during the streaming generation process.

diff --git a/Engine/Source/Runtime/Engine/Private/WorldPartition/RuntimeSpatialHash/RuntimeSpatialHashGridHelper.cpp b/Engine/Source/Runtime/Engine/Private/WorldPartition/RuntimeSpatialHash/RuntimeSpatialHashGridHelper.cpp
--- a/Engine/Source/Runtime/Engine/Private/WorldPartition/RuntimeSpatialHash/RuntimeSpatialHashGridHelper.cpp
+++ b/Engine/Source/Runtime/Engine/Private/WorldPartition/RuntimeSpatialHash/RuntimeSpatialHashGridHelper.cpp
@@ -163,7 +163,7 @@ FSquare2DGridHelper GetPartitionedActors(const UWorldPartition* WorldPartition,
 			check(ActorCluster->Actors.Num() == 1);
 			const FGuid& ActorGuid = *ActorCluster->Actors.CreateConstIterator();
 			FActorInstance ActorInstance(ActorGuid, ClusterInstance->ContainerInstance);
-			if (ensure(PartitionedActors.GetLowestLevel().GetCellCoords(FVector2D(ActorInstance.GetActorDesc()->GetOrigin()), CellCoords)))
+			if (ensure(PartitionedActors.GetLowestLevel().GetCellCoords(FVector2D(ActorInstance.GetActorDescView().GetOrigin()), CellCoords)))
 			{
 				PartitionedActors.GetLowestLevel().GetCell(CellCoords).AddActor(MoveTemp(ActorInstance), ClusterInstance->DataLayers);
 			}
@@ -226,9 +226,9 @@ FSquare2DGridHelper GetPartitionedActors(const UWorldPartition* WorldPartition,
 
 				for (const FGuid& ActorGuid : ActorCluster->Actors)
 				{
-					const FWorldPartitionActorDesc& Desc = *ClusterInstance->ContainerInstance->Container->GetActorDesc(ActorGuid);
-					UE_LOG(LogWorldPartitionRuntimeSpatialHash, Verbose, TEXT("   - Actor: %s (%s)"), *Desc.GetActorPath().ToString(), *ActorGuid.ToString(EGuidFormats::UniqueObjectGuid));
-					UE_LOG(LogWorldPartitionRuntimeSpatialHash, Verbose, TEXT("         Package: %s"), *Desc.GetActorPackage().ToString());
+					const FWorldPartitionActorDescView& ActorDescView = ClusterInstance->ContainerInstance->GetActorDescView(ActorGuid);
+					UE_LOG(LogWorldPartitionRuntimeSpatialHash, Verbose, TEXT("   - Actor: %s (%s)"), *ActorDescView.GetActorPath().ToString(), *ActorGuid.ToString(EGuidFormats::UniqueObjectGuid));
+					UE_LOG(LogWorldPartitionRuntimeSpatialHash, Verbose, TEXT("         Package: %s"), *ActorDescView.GetActorPackage().ToString());
 					UE_LOG(LogWorldPartitionRuntimeSpatialHash, Verbose, TEXT("         Container (%08x): %s"), ClusterInstance->ContainerInstance->ID, *ClusterInstance->ContainerInstance->Container->GetContainerPackage().ToString())
 				}
 			}
diff --git a/Engine/Source/Runtime/Engine/Private/WorldPartition/RuntimeSpatialHash/RuntimeSpatialHashHLOD.cpp b/Engine/Source/Runtime/Engine/Private/WorldPartition/RuntimeSpatialHash/RuntimeSpatialHashHLOD.cpp
--- a/Engine/Source/Runtime/Engine/Private/WorldPartition/RuntimeSpatialHash/RuntimeSpatialHashHLOD.cpp
+++ b/Engine/Source/Runtime/Engine/Private/WorldPartition/RuntimeSpatialHash/RuntimeSpatialHashHLOD.cpp
@@ -489,9 +489,9 @@ bool UWorldPartitionRuntimeSpatialHash::GenerateHLOD(ISourceControlHelper* Sourc
 	}
 
 	// Create actor clusters - ignore HLOD actors
-	FActorClusterContext ClusterContext(WorldPartition, TOptional<FActorClusterContext::FFilterPredicate>([](const FWorldPartitionActorDesc& ActorDesc)
+	FActorClusterContext ClusterContext(WorldPartition, this, TOptional<FActorClusterContext::FFilterPredicate>([](const FWorldPartitionActorDescView& ActorDescView)
 		{
-			return !ActorDesc.GetActorClass()->IsChildOf<AWorldPartitionHLOD>();
+			return !ActorDescView.GetActorClass()->IsChildOf<AWorldPartitionHLOD>();
 		}), /* bInIncludeChildContainers=*/ false);
 		
 	TArray<TArray<const FActorClusterInstance*>> GridsClusters;
@@ -536,7 +536,7 @@ bool UWorldPartitionRuntimeSpatialHash::GenerateHLOD(ISourceControlHelper* Sourc
 	}
 
 	// Now, go on and create HLOD actors from HLOD grids (HLOD 1-N)
-	FActorContainerInstance MainContainerInstance(WorldPartition);
+	FActorContainerInstance& MainContainerInstance = *ClusterContext.GetClusterInstance(WorldPartition);
 	for (const FName HLODGridName : SortedGrids)
 	{
 		// No need to process empty grids
@@ -556,10 +556,8 @@ bool UWorldPartitionRuntimeSpatialHash::GenerateHLOD(ISourceControlHelper* Sourc
 
 		for (const FGuid& HLODActorGuid : GridsHLODActors[HLODGridName])
 		{
-			FWorldPartitionActorDesc* HLODActorDesc = WorldPartition->GetActorDesc(HLODActorGuid);
-			check(HLODActorDesc);
-						
-			FActorCluster& NewHLODActorCluster = HLODActorClusters.Emplace_GetRef(HLODActorDesc, HLODActorDesc->GetGridPlacement(), WorldPartition);
+			const FWorldPartitionActorDescView& HLODActorDescView = MainContainerInstance.GetActorDescView(HLODActorGuid);
+			FActorCluster& NewHLODActorCluster = HLODActorClusters.Emplace_GetRef(WorldPartition->GetWorld(), HLODActorDescView, HLODActorDescView.GetGridPlacement());
 			FActorClusterInstance& NewHLODActorClusterInstance = HLODActorClusterInstances.Emplace_GetRef(&NewHLODActorCluster, &MainContainerInstance);
 			HLODActorClusterInstancePtrs.Add(&NewHLODActorClusterInstance);
 		};
diff --git a/Engine/Source/Runtime/Engine/Private/WorldPartition/WorldPartitionActorCluster.cpp b/Engine/Source/Runtime/Engine/Private/WorldPartition/WorldPartitionActorCluster.cpp
--- a/Engine/Source/Runtime/Engine/Private/WorldPartition/WorldPartitionActorCluster.cpp
+++ b/Engine/Source/Runtime/Engine/Private/WorldPartition/WorldPartitionActorCluster.cpp
@@ -8,8 +8,8 @@
 #include "Engine/World.h"
 #include "WorldPartition/WorldPartition.h"
 #include "WorldPartition/ActorDescContainer.h"
-#include "ActorReferencesUtils.h"
-#include "Engine/LevelScriptBlueprint.h"
+#include "WorldPartition/WorldPartitionActorDescView.h"
+#include "WorldPartition/WorldPartitionRuntimeHash.h"
 #include "Misc/HashBuilder.h"
 
 DEFINE_LOG_CATEGORY(LogWorldPartitionActorCluster);
@@ -34,15 +34,15 @@ TSet<const UDataLayer*> GetDataLayers(UWorld* InWorld, const LayerNameContainer&
 	return DataLayers;
 }
 
-FActorCluster::FActorCluster(const FWorldPartitionActorDesc* InActorDesc, EActorGridPlacement InGridPlacement, const UActorDescContainer* InContainer)
+FActorCluster::FActorCluster(UWorld* InWorld, const FWorldPartitionActorDescView& InActorDescView, EActorGridPlacement InGridPlacement)
 	: GridPlacement(InGridPlacement)
-	, RuntimeGrid(InActorDesc->GetRuntimeGrid())
-	, Bounds(InActorDesc->GetBounds())
+	, RuntimeGrid(InActorDescView.GetRuntimeGrid())
+	, Bounds(InActorDescView.GetBounds())
 {
 	check(GridPlacement != EActorGridPlacement::None);
 
-	Actors.Add(InActorDesc->GetGuid());
-	DataLayers = GetDataLayers(InContainer->GetWorld(), InActorDesc->GetDataLayers());
+	Actors.Add(InActorDescView.GetGuid());
+	DataLayers = GetDataLayers(InWorld, InActorDescView.GetDataLayers());
 	DataLayersID = FDataLayersID(DataLayers.Array());
 }
 
@@ -122,8 +122,9 @@ FActorClusterInstance::FActorClusterInstance(const FActorCluster* InCluster, con
 	DataLayers.Append(DataLayerSet.Array());
 }
 
-FActorClusterContext::FActorClusterContext(const UWorldPartition* InWorldPartition, TOptional<FFilterPredicate> InFilterPredicate, bool bInIncludeChildContainers)
+FActorClusterContext::FActorClusterContext(const UWorldPartition* InWorldPartition, const UWorldPartitionRuntimeHash* InRuntimeHash, TOptional<FFilterPredicate> InFilterPredicate, bool bInIncludeChildContainers)
 	: WorldPartition(InWorldPartition)
+	, RuntimeHash(InRuntimeHash)
 	, FilterPredicate(InFilterPredicate)
 	, bIncludeChildContainers(bInIncludeChildContainers)
 	, InstanceCountHint(0)
@@ -131,23 +132,27 @@ FActorClusterContext::FActorClusterContext(const UWorldPartition* InWorldPartiti
 	CreateActorClusters();
 }
 
-FActorContainerInstance::FActorContainerInstance(const UActorDescContainer* InContainer)
-	: FActorContainerInstance(0, FTransform::Identity, FBox(ForceInit), TSet<FName>(), EContainerClusterMode::Partitioned, InContainer, TSet<FGuid>())
-{
-
-}
+FActorContainerInstance::FActorContainerInstance(const UActorDescContainer* InContainer, TMap<FGuid, FWorldPartitionActorDescView> InActorDescViewMap)
+	: FActorContainerInstance(0, FTransform::Identity, FBox(ForceInit), TSet<FName>(), EContainerClusterMode::Partitioned, InContainer, TSet<FGuid>(), InActorDescViewMap)
+{}
 
-FActorContainerInstance::FActorContainerInstance(uint32 InID, const FTransform& InTransform, const FBox& InBounds, const TSet<FName>& InDataLayers, EContainerClusterMode InClusterMode, const UActorDescContainer* InContainer, TSet<FGuid> InChildContainers)
+FActorContainerInstance::FActorContainerInstance(uint32 InID, const FTransform& InTransform, const FBox& InBounds, const TSet<FName>& InDataLayers, EContainerClusterMode InClusterMode, const UActorDescContainer* InContainer, TSet<FGuid> InChildContainers, TMap<FGuid, FWorldPartitionActorDescView> InActorDescViewMap)
 	: ID(InID)
 	, Transform(InTransform)
 	, Bounds(InBounds)
 	, ClusterMode(InClusterMode)
 	, Container(InContainer)
 	, ChildContainers(InChildContainers)
+	, ActorDescViewMap(InActorDescViewMap)
 {
 	DataLayers = GetDataLayers(InContainer->GetWorld(), InDataLayers);
 }
 
+const FWorldPartitionActorDescView& FActorContainerInstance::GetActorDescView(const FGuid& InGuid) const
+{
+	return ActorDescViewMap.FindChecked(InGuid);
+}
+
 FActorInstance::FActorInstance()
 	: ContainerInstance(nullptr)
 {
@@ -169,99 +174,102 @@ bool FActorInstance::ShouldStripFromStreaming() const
 		return true;
 	}
 
-	const FWorldPartitionActorDesc* ActorDesc = GetActorDesc();
-	return ActorDesc->GetActorIsEditorOnly();
+	const FWorldPartitionActorDescView ActorDescView = GetActorDescView();
+	return ActorDescView.GetActorIsEditorOnly();
 }
 
 FVector FActorInstance::GetOrigin() const
 {
-	return ContainerInstance->Transform.TransformPosition(GetActorDesc()->GetOrigin());
+	return ContainerInstance->Transform.TransformPosition(GetActorDescView().GetOrigin());
 }
 
-const FWorldPartitionActorDesc* FActorInstance::GetActorDesc() const
+const FWorldPartitionActorDescView& FActorInstance::GetActorDescView() const
 {
-	return ContainerInstance->Container->GetActorDesc(Actor);
+	return ContainerInstance->GetActorDescView(Actor);
 }
 
-void CreateActorCluster(const FWorldPartitionActorDesc* ActorDesc, EActorGridPlacement GridPlacement, TMap<FGuid, FActorCluster*>& ActorToActorCluster, TSet<FActorCluster*>& ActorClustersSet, const UActorDescContainer* ActorDescContainer)
+void CreateActorCluster(const FWorldPartitionActorDescView& ActorDescView, EActorGridPlacement GridPlacement, TMap<FGuid, FActorCluster*>& ActorToActorCluster, TSet<FActorCluster*>& ActorClustersSet, const FActorContainerInstance& ContainerInstance)
 {
+	const UActorDescContainer* ActorDescContainer = ContainerInstance.Container;
 	UWorld* World = ActorDescContainer->GetWorld();
-	const FGuid& ActorGuid = ActorDesc->GetGuid();
+	const FGuid& ActorGuid = ActorDescView.GetGuid();
 
 	FActorCluster* ActorCluster = ActorToActorCluster.FindRef(ActorGuid);
 	if (!ActorCluster)
 	{
-		ActorCluster = new FActorCluster(ActorDesc, GridPlacement, ActorDescContainer);
+		ActorCluster = new FActorCluster(World, ActorDescView, GridPlacement);
 		ActorClustersSet.Add(ActorCluster);
 		ActorToActorCluster.Add(ActorGuid, ActorCluster);
 	}
 
 	// Don't include references from editor-only actors
-	if (!ActorDesc->GetActorIsEditorOnly())
+	if (!ActorDescView.GetActorIsEditorOnly())
 	{
-		for (const FGuid& ReferenceGuid : ActorDesc->GetReferences())
+		for (const FGuid& ReferenceGuid : ActorDescView.GetReferences())
 		{
-			if (const FWorldPartitionActorDesc* ReferenceActorDesc = ActorDescContainer->GetActorDesc(ReferenceGuid))
+			const FWorldPartitionActorDescView& ReferenceActorDescView = ContainerInstance.ActorDescViewMap.FindChecked(ReferenceGuid);
+
+			// Don't include references to editor-only actors
+			if (!ReferenceActorDescView.GetActorIsEditorOnly())
 			{
-				// Don't include references to editor-only actors
-				if (!ReferenceActorDesc->GetActorIsEditorOnly())
+				FActorCluster* ReferenceCluster = ActorToActorCluster.FindRef(ReferenceGuid);
+				if (ReferenceCluster)
 				{
-					FActorCluster* ReferenceCluster = ActorToActorCluster.FindRef(ReferenceGuid);
-					if (ReferenceCluster)
+					if (ReferenceCluster != ActorCluster)
 					{
-						if (ReferenceCluster != ActorCluster)
+						// Merge reference cluster in Actor's cluster
+						ActorCluster->Add(*ReferenceCluster);
+						for (const FGuid& ReferenceClusterActorGuid : ReferenceCluster->Actors)
 						{
-							// Merge reference cluster in Actor's cluster
-							ActorCluster->Add(*ReferenceCluster);
-							for (const FGuid& ReferenceClusterActorGuid : ReferenceCluster->Actors)
-							{
-								ActorToActorCluster[ReferenceClusterActorGuid] = ActorCluster;
-							}
-							ActorClustersSet.Remove(ReferenceCluster);
-							delete ReferenceCluster;
+							ActorToActorCluster[ReferenceClusterActorGuid] = ActorCluster;
 						}
+						ActorClustersSet.Remove(ReferenceCluster);
+						delete ReferenceCluster;
 					}
-					else
-					{
-						// Put Reference in Actor's cluster
-						ActorCluster->Add(FActorCluster(ReferenceActorDesc, GridPlacement, ActorDescContainer));
-					}
-
-					// Map its cluster
-					ActorToActorCluster.Add(ReferenceGuid, ActorCluster);
 				}
+				else
+				{
+					// Put Reference in Actor's cluster
+					ActorCluster->Add(FActorCluster(World, ReferenceActorDescView, GridPlacement));
+				}
+
+				// Map its cluster
+				ActorToActorCluster.Add(ReferenceGuid, ActorCluster);
 			}
 		}
 	}
 }
 
-void FActorClusterContext::CreateContainerInstanceRecursive(uint32 ID, const FTransform& Transform, EContainerClusterMode ClusterMode , const UActorDescContainer* Container, const TSet<FName>& DataLayers, FBox* ParentBounds)
+void FActorClusterContext::CreateContainerInstanceRecursive(uint32 ID, const FTransform& Transform, EContainerClusterMode ClusterMode, const UActorDescContainer* Container, const TSet<FName>& DataLayers, FBox* ParentBounds)
 {
 	InstanceCountHint += Container->GetActorDescCount();
 		
 	TSet<FGuid> ChildContainers;
 	FBox Bounds(ForceInit);
 
-	for (UActorDescContainer::TConstIterator<> ActorDescIterator(Container); ActorDescIterator; ++ActorDescIterator)
+	TMap<FGuid, FWorldPartitionActorDescView> ActorDescViewMap;
+	RuntimeHash->CreateActorDescViewMap(Container, ActorDescViewMap);
+
+	for (auto& ActorDesViewPair : ActorDescViewMap)
 	{
-		const FWorldPartitionActorDesc* ActorDesc = *ActorDescIterator;
+		FWorldPartitionActorDescView& ActorDescView = ActorDesViewPair.Value;
 		const UActorDescContainer* OutContainer = nullptr;
 		FTransform OutTransform;
 		EContainerClusterMode OutClusterMode;
-		if (bIncludeChildContainers && ActorDesc->GetContainerInstance(OutContainer, OutTransform, OutClusterMode))
+		if (bIncludeChildContainers && ActorDescView.GetContainerInstance(OutContainer, OutTransform, OutClusterMode))
 		{
 			// Add Child Container Guid so we can discard the actor later
-			ChildContainers.Add(ActorDesc->GetGuid());
+			ChildContainers.Add(ActorDescView.GetGuid());
 
 			FHashBuilder ChildContainerHashBuilder(ID);
-			ChildContainerHashBuilder << ActorDesc->GetGuid();
+			ChildContainerHashBuilder << ActorDescView.GetGuid();
 			
 			TSet<FName> ChildDataLayers;
-			ChildDataLayers.Reserve(DataLayers.Num() + ActorDesc->GetDataLayers().Num());
+			ChildDataLayers.Reserve(DataLayers.Num() + ActorDescView.GetDataLayers().Num());
 			// Only propagate ActorDesc DataLayers if we are a Root Container (WorldPartition)
 			if (Container->IsA<UWorldPartition>())
 			{
-				ChildDataLayers.Append(ActorDesc->GetDataLayers());
+				ChildDataLayers.Append(ActorDescView.GetDataLayers());
 			}
 			// Always inherite parent container DataLayers
 			ChildDataLayers.Append(DataLayers);
@@ -269,17 +277,17 @@ void FActorClusterContext::CreateContainerInstanceRecursive(uint32 ID, const FTr
 		}
 		else
 		{
-			switch (ActorDesc->GetGridPlacement())
+			switch (ActorDescView.GetGridPlacement())
 			{
 				case EActorGridPlacement::Location:
 				{
-					FVector Location = ActorDescIterator->GetOrigin();
+					FVector Location = ActorDescView.GetOrigin();
 					Bounds += FBox(Location, Location);
 				}
 				break;
 				case EActorGridPlacement::Bounds:
 				{
-					Bounds += ActorDescIterator->GetBounds();
+					Bounds += ActorDescView.GetBounds();
 				}
 				break;
 			}
@@ -292,7 +300,20 @@ void FActorClusterContext::CreateContainerInstanceRecursive(uint32 ID, const FTr
 	}
 
 	UE_LOG(LogWorldPartitionActorCluster, Verbose, TEXT("ContainerInstance (%08x) Bounds (%s) Package (%s)"), ID, *Bounds.TransformBy(Transform).ToString(), *Container->GetContainerPackage().ToString());
-	ContainerInstances.Add(FActorContainerInstance(ID, Transform, Bounds, DataLayers, ClusterMode, Container, MoveTemp(ChildContainers)));
+	ContainerInstances.Add(FActorContainerInstance(ID, Transform, Bounds, DataLayers, ClusterMode, Container, MoveTemp(ChildContainers), MoveTemp(ActorDescViewMap)));
+}
+
+FActorContainerInstance* FActorClusterContext::GetClusterInstance(const UActorDescContainer* InContainer)
+{
+	for (FActorContainerInstance& ContainerInstance : ContainerInstances)
+	{
+		if (ContainerInstance.Container == InContainer)
+		{
+			return &ContainerInstance;
+		}
+	}
+
+	return nullptr;
 }
 
 const TArray<FActorCluster>& FActorClusterContext::CreateActorClustersImpl(const FActorContainerInstance& ContainerInstance)
@@ -306,34 +327,14 @@ const TArray<FActorCluster>& FActorClusterContext::CreateActorClustersImpl(const
 	TSet<FActorCluster*> ActorClustersSet;
 	TArray<FActorCluster>& ActorClusters = Clusters.Add(ContainerInstance.Container);
 	
-	// Gather all references to external actors from the level script
-	TSet<AActor*> LevelScriptExternalActorReferences;
-	// Only applies to root container
-	if (ContainerInstance.Container->IsA<UWorldPartition>())
-	{
-		if (ULevelScriptBlueprint* LevelScriptBlueprint = ContainerInstance.Container->GetWorld()->PersistentLevel->GetLevelScriptBlueprint(true))
-		{
-			LevelScriptExternalActorReferences.Append(ActorsReferencesUtils::GetExternalActorReferences(LevelScriptBlueprint));
-		}
-	}
-		
-	for (UActorDescContainer::TConstIterator<> ActorDescIterator(ContainerInstance.Container); ActorDescIterator; ++ActorDescIterator)
+	for (auto& ActorDescViewPair : ContainerInstance.ActorDescViewMap)
 	{
-		const FWorldPartitionActorDesc* ActorDesc = *ActorDescIterator;
-		EActorGridPlacement GridPlacement = ActorDesc->GetGridPlacement();
-
-		// Check if the actor is loaded (potentially referenced by the level script)
-		if (AActor* Actor = ActorDesc->GetActor())
-		{
-			if (LevelScriptExternalActorReferences.Contains(Actor))
-			{
-				GridPlacement = EActorGridPlacement::AlwaysLoaded;
-			}
-		}
+		const FWorldPartitionActorDescView& ActorDescView = ActorDescViewPair.Value;
+		EActorGridPlacement GridPlacement = ActorDescView.GetGridPlacement();
 
-		if (!FilterPredicate.IsSet() || FilterPredicate.GetValue()(*ActorDesc))
+		if (!FilterPredicate.IsSet() || FilterPredicate.GetValue()(ActorDescView))
 		{
-			CreateActorCluster(ActorDesc, GridPlacement, ActorToActorCluster, ActorClustersSet, ContainerInstance.Container);
+			CreateActorCluster(ActorDescView, GridPlacement, ActorToActorCluster, ActorClustersSet, ContainerInstance);
 		}
 	}
 				
diff --git a/Engine/Source/Runtime/Engine/Private/WorldPartition/WorldPartitionActorDescView.cpp b/Engine/Source/Runtime/Engine/Private/WorldPartition/WorldPartitionActorDescView.cpp
new file mode 100644
--- /dev/null
+++ b/Engine/Source/Runtime/Engine/Private/WorldPartition/WorldPartitionActorDescView.cpp
@@ -0,0 +1,114 @@
+// Copyright Epic Games, Inc. All Rights Reserved.
+
+#include "WorldPartition/WorldPartitionActorDescView.h"
+
+#if WITH_EDITOR
+#include "WorldPartition/WorldPartitionActorDesc.h"
+
+FWorldPartitionActorDescView::FWorldPartitionActorDescView()
+	: ActorDesc(nullptr)
+	, EffectiveGridPlacement(EActorGridPlacement::None)
+{}
+
+FWorldPartitionActorDescView::FWorldPartitionActorDescView(const FWorldPartitionActorDesc* InActorDesc)
+	: ActorDesc(InActorDesc)
+	, EffectiveGridPlacement(InActorDesc->GetGridPlacement())
+{}
+
+const FGuid& FWorldPartitionActorDescView::GetGuid() const
+{
+	return ActorDesc->GetGuid();
+}
+
+FName FWorldPartitionActorDescView::GetClass() const
+{
+	return ActorDesc->GetClass();
+}
+
+UClass* FWorldPartitionActorDescView::GetActorClass() const
+{
+	return ActorDesc->GetActorClass();
+}
+
+FVector FWorldPartitionActorDescView::GetOrigin() const
+{
+	return ActorDesc->GetOrigin();
+}
+
+EActorGridPlacement FWorldPartitionActorDescView::GetGridPlacement() const
+{
+	return EffectiveGridPlacement;
+}
+
+FName FWorldPartitionActorDescView::GetRuntimeGrid() const
+{
+	return ActorDesc->GetRuntimeGrid();
+}
+
+bool FWorldPartitionActorDescView::GetActorIsEditorOnly() const
+{
+	return ActorDesc->GetActorIsEditorOnly();
+}
+
+bool FWorldPartitionActorDescView::GetLevelBoundsRelevant() const
+{
+	return ActorDesc->GetLevelBoundsRelevant();
+}
+
+bool FWorldPartitionActorDescView::GetActorIsHLODRelevant() const
+{
+	return ActorDesc->GetActorIsHLODRelevant();
+}
+
+UHLODLayer* FWorldPartitionActorDescView::GetHLODLayer() const
+{
+	return ActorDesc->GetHLODLayer();
+}
+
+const TArray<FName>& FWorldPartitionActorDescView::GetDataLayers() const
+{
+	return ActorDesc->GetDataLayers();
+}
+
+FName FWorldPartitionActorDescView::GetActorPackage() const
+{
+	return ActorDesc->GetActorPackage();
+}
+
+FName FWorldPartitionActorDescView::GetActorPath() const
+{
+	return ActorDesc->GetActorPath();
+}
+
+FName FWorldPartitionActorDescView::GetActorLabel() const
+{
+	return ActorDesc->GetActorLabel();
+}
+
+FBox FWorldPartitionActorDescView::GetBounds() const
+{
+	return ActorDesc->GetBounds();
+}
+
+const TArray<FGuid>& FWorldPartitionActorDescView::GetReferences() const
+{
+	return ActorDesc->GetReferences();
+}
+
+bool FWorldPartitionActorDescView::GetContainerInstance(const UActorDescContainer*& OutLevelContainer, FTransform& OutLevelTransform, EContainerClusterMode& OutClusterMode) const
+{
+	return ActorDesc->GetContainerInstance(OutLevelContainer, OutLevelTransform, OutClusterMode);
+}
+
+FGuid FWorldPartitionActorDescView::GetHLODParent() const
+{
+	return HLODParent;
+}
+
+void FWorldPartitionActorDescView::SetHLODParent(const FGuid& InHLODParent)
+{
+	check(!HLODParent.IsValid());
+	check(InHLODParent.IsValid());
+	HLODParent = InHLODParent;
+}
+#endif
\ No newline at end of file
diff --git a/Engine/Source/Runtime/Engine/Private/WorldPartition/WorldPartitionRuntimeHash.cpp b/Engine/Source/Runtime/Engine/Private/WorldPartition/WorldPartitionRuntimeHash.cpp
--- a/Engine/Source/Runtime/Engine/Private/WorldPartition/WorldPartitionRuntimeHash.cpp
+++ b/Engine/Source/Runtime/Engine/Private/WorldPartition/WorldPartitionRuntimeHash.cpp
@@ -1,6 +1,10 @@
 // Copyright Epic Games, Inc. All Rights Reserved.
 
 #include "WorldPartition/WorldPartitionRuntimeHash.h"
+#include "WorldPartition/HLOD/HLODActor.h"
+#include "WorldPartition/HLOD/HLODActorDesc.h"
+#include "Engine/LevelScriptBlueprint.h"
+#include "ActorReferencesUtils.h"
 
 UWorldPartitionRuntimeHash::UWorldPartitionRuntimeHash(const FObjectInitializer& ObjectInitializer)
 : Super(ObjectInitializer)
@@ -34,4 +38,45 @@ void UWorldPartitionRuntimeHash::ForceExternalActorLevelReference(bool bForceExt
 	}
 }
 
+void UWorldPartitionRuntimeHash::CreateActorDescViewMap(const UActorDescContainer* Container, TMap<FGuid, FWorldPartitionActorDescView>& OutActorDescViewMap) const
+{
+	// Build the actor desc view map
+	OutActorDescViewMap.Empty();
+	for (UActorDescContainer::TConstIterator<> ActorDescIt(Container); ActorDescIt; ++ActorDescIt)
+	{
+		OutActorDescViewMap.Emplace(ActorDescIt->GetGuid(), *ActorDescIt);
+	}
+
+	// Set HLOD parents into actor desc views
+	for (UActorDescContainer::TConstIterator<AWorldPartitionHLOD> HLODIterator(Container); HLODIterator; ++HLODIterator)
+	{
+		for (const FGuid& SubActor : HLODIterator->GetSubActors())
+		{
+			if (FWorldPartitionActorDescView* SubActorDescView = OutActorDescViewMap.Find(SubActor))
+			{
+				SubActorDescView->SetHLODParent(HLODIterator->GetGuid());
+			}
+		}
+	}
+
+	// Gather all references to external actors from the level script and make them always loaded
+	if (ULevelScriptBlueprint* LevelScriptBlueprint = Container->GetWorld()->PersistentLevel->GetLevelScriptBlueprint(true))
+	{
+		TArray<AActor*> LevelScriptExternalActorReferences = ActorsReferencesUtils::GetExternalActorReferences(LevelScriptBlueprint);
+
+		for (AActor* Actor : LevelScriptExternalActorReferences)
+		{
+			ChangeActorDescViewGridPlacement(OutActorDescViewMap.FindChecked(Actor->GetActorGuid()), EActorGridPlacement::AlwaysLoaded);
+		}
+	}
+}
+
+void UWorldPartitionRuntimeHash::ChangeActorDescViewGridPlacement(FWorldPartitionActorDescView& ActorDescView, EActorGridPlacement GridPlacement) const
+{
+	if (ActorDescView.EffectiveGridPlacement != GridPlacement)
+	{
+		ActorDescView.EffectiveGridPlacement = GridPlacement;
+		UE_LOG(LogWorldPartition, Warning, TEXT("Actor '%s' grid placement changed to %d"), *ActorDescView.GetActorLabel().ToString(), (int32)GridPlacement);
+	}
+}
 #endif
diff --git a/Engine/Source/Runtime/Engine/Private/WorldPartition/WorldPartitionRuntimeSpatialHash.cpp b/Engine/Source/Runtime/Engine/Private/WorldPartition/WorldPartitionRuntimeSpatialHash.cpp
--- a/Engine/Source/Runtime/Engine/Private/WorldPartition/WorldPartitionRuntimeSpatialHash.cpp
+++ b/Engine/Source/Runtime/Engine/Private/WorldPartition/WorldPartitionRuntimeSpatialHash.cpp
@@ -6,6 +6,7 @@
 #include "WorldPartition/WorldPartitionRuntimeSpatialHash.h"
 #include "WorldPartition/WorldPartitionRuntimeSpatialHashCell.h"
 #include "WorldPartition/WorldPartitionActorDesc.h"
+#include "WorldPartition/WorldPartitionActorDescView.h"
 #include "WorldPartition/WorldPartitionStreamingPolicy.h"
 #include "WorldPartition/WorldPartitionHandle.h"
 
@@ -18,6 +19,7 @@
 #include "WorldPartition/DataLayer/WorldDataLayers.h"
 #include "WorldPartition/DataLayer/DataLayer.h"
 
+#include "GameFramework/Actor.h"
 #include "GameFramework/WorldSettings.h"
 #include "ProfilingDebugging/ScopedTimers.h"
 #include "Engine/World.h"
@@ -443,9 +445,6 @@ bool UWorldPartitionRuntimeSpatialHash::GenerateStreaming(EWorldPartitionStreami
 
 	check(!StreamingGrids.Num());
 
-	// Build a map of Actor GUID -> HLODActor GUID once instead of having to recompute for every streaming grid we create
-	CacheHLODParents();
-
 	TMap<FName, int32> GridsMapping;
 	GridsMapping.Add(NAME_None, 0);
 	for (int32 i = 0; i < AllGrids.Num(); i++)
@@ -456,7 +455,7 @@ bool UWorldPartitionRuntimeSpatialHash::GenerateStreaming(EWorldPartitionStreami
 	}
 
 	// Create actor clusters
-	FActorClusterContext Context(WorldPartition);
+	FActorClusterContext Context(WorldPartition, this);
 
 	TArray<TArray<const FActorClusterInstance*>> GridActors;
 	GridActors.InsertDefaulted(0, AllGrids.Num());
@@ -474,7 +473,7 @@ bool UWorldPartitionRuntimeSpatialHash::GenerateStreaming(EWorldPartitionStreami
 		GridActors[GridIndex].Add(&ClusterInstance);
 	}
 	
-	const FBox WorldBounds = WorldPartition->GetWorldBounds();
+	const FBox WorldBounds = Context.GetClusterInstance(WorldPartition)->Bounds;
 	for (int32 GridIndex=0; GridIndex < AllGrids.Num(); GridIndex++)
 	{
 		const FSpatialHashRuntimeGrid& Grid = AllGrids[GridIndex];
@@ -502,17 +501,74 @@ FName UWorldPartitionRuntimeSpatialHash::GetCellName(FName InGridName, int32 InL
 	return UWorldPartitionRuntimeSpatialHash::GetCellName(WorldPartition, InGridName, InLevel, InCellX, InCellY, InDataLayerID);
 }
 
-void UWorldPartitionRuntimeSpatialHash::CacheHLODParents()
+void UWorldPartitionRuntimeSpatialHash::CreateActorDescViewMap(const UActorDescContainer* Container, TMap<FGuid, FWorldPartitionActorDescView>& OutActorDescViewMap) const
 {
-	CachedHLODParents.Reset();
+	Super::CreateActorDescViewMap(Container, OutActorDescViewMap);
 
-	UWorldPartition* WorldPartition = GetOuterUWorldPartition();
+	// Compute the initial world bounds.
+	// This bound will change after we update actor descriptor views grid placement.
+	auto GetWorldBounds = [](const TMap<FGuid, FWorldPartitionActorDescView> ActorDescViewMap)
+	{
+		FBox WorldBounds(ForceInit);
+		for (auto& ActorDescViewPair : ActorDescViewMap)
+		{
+			const FWorldPartitionActorDescView& ActorDescView = ActorDescViewPair.Value;
+			switch (ActorDescView.GetGridPlacement())
+			{
+				case EActorGridPlacement::Location:
+				{
+					FVector Location = ActorDescView.GetOrigin();
+					WorldBounds += FBox(Location, Location);
+				}
+				break;
+				case EActorGridPlacement::Bounds:
+				{
+					WorldBounds += ActorDescView.GetBounds();
+				}
+				break;
+			}
+		}
+		return WorldBounds;
+	};
+
+ 	const FBox WorldBounds = GetWorldBounds(OutActorDescViewMap);
+
+	TMap<FName, int32> GridsMapping;
+	GridsMapping.Add(NAME_None, 0);
+	for (int32 i = 0; i < Grids.Num(); i++)
+	{
+		const FSpatialHashRuntimeGrid& Grid = Grids[i];
+		check(!GridsMapping.Contains(Grid.GridName));
+		GridsMapping.Add(Grid.GridName, i);
+	}
+
+	const FBox2D WorldBounds2D = FBox2D(FVector2D(WorldBounds.Min), FVector2D(WorldBounds.Max));
+	const float WorldBoundsArea = WorldBounds2D.GetArea();
 
-	for (UActorDescContainer::TIterator<AWorldPartitionHLOD> HLODIterator(WorldPartition); HLODIterator; ++HLODIterator)
+	for (auto& ActorDescViewPair : OutActorDescViewMap)
 	{
-		for (const auto& SubActor : HLODIterator->GetSubActors())
+		FWorldPartitionActorDescView& ActorDescView = ActorDescViewPair.Value;
+
+		if (!ActorDescView.GetActorIsEditorOnly())
 		{
-			CachedHLODParents.Emplace(SubActor, HLODIterator->GetGuid());
+			if (ActorDescView.GetGridPlacement() == EActorGridPlacement::Bounds)
+			{
+				const int32* GridIndexPtr = GridsMapping.Find(ActorDescView.GetRuntimeGrid());
+				const int32 GridIndex = GridIndexPtr ? *GridIndexPtr : 0;
+				const FSpatialHashRuntimeGrid& RuntimeGrid = Grids[GridIndex];
+				const float CellArea = RuntimeGrid.CellSize * RuntimeGrid.CellSize;
+				const FBox2D ActorBounds2D = FBox2D(FVector2D(ActorDescView.GetBounds().Min), FVector2D(ActorDescView.GetBounds().Max));
+				const float ActorBoundsArea = ActorBounds2D.GetArea();
+
+				if (ActorBoundsArea < CellArea / 16.0f)
+				{
+					ChangeActorDescViewGridPlacement(ActorDescView, EActorGridPlacement::Location);
+				}
+				else if (ActorBoundsArea > WorldBoundsArea / 2.0f)
+				{
+					ChangeActorDescViewGridPlacement(ActorDescView, EActorGridPlacement::AlwaysLoaded);
+				}
+			}
 		}
 	}
 }
@@ -523,7 +579,6 @@ bool UWorldPartitionRuntimeSpatialHash::CreateStreamingGrid(const FSpatialHashRu
 
 	UWorldPartition* WorldPartition = GetOuterUWorldPartition();
 	const bool bIsMainWorldPartition = (WorldPartition->GetWorld() == WorldPartition->GetTypedOuter<UWorld>());
-	check(PartionedActors.WorldBounds == WorldPartition->GetWorldBounds());
 
 	FSpatialHashStreamingGrid& CurrentStreamingGrid = StreamingGrids.AddDefaulted_GetRef();
 	CurrentStreamingGrid.GridName = RuntimeGrid.GridName;
@@ -565,11 +620,11 @@ bool UWorldPartitionRuntimeSpatialHash::CreateStreamingGrid(const FSpatialHashRu
 				{
 					for (const FActorInstance& ActorInstance : GridCellDataChunk.GetActors())
 					{
-						const FWorldPartitionActorDesc* ActorDesc = ActorInstance.GetActorDesc();
 						if (ActorInstance.ShouldStripFromStreaming())
 						{
+							const FWorldPartitionActorDescView& ActorDescView = ActorInstance.GetActorDescView();
 							UE_LOG(LogWorldPartitionRuntimeSpatialHash, Verbose, TEXT("Stripping Actor %s (%s) from streaming grid (Container %08x)"),
-								*(ActorDesc->GetActorPath().ToString()), *ActorInstance.Actor.ToString(EGuidFormats::UniqueObjectGuid), ActorInstance.ContainerInstance->ID);
+								*(ActorDescView.GetActorPath().ToString()), *ActorInstance.Actor.ToString(EGuidFormats::UniqueObjectGuid), ActorInstance.ContainerInstance->ID);
 							continue;
 						}
 
@@ -614,14 +669,13 @@ bool UWorldPartitionRuntimeSpatialHash::CreateStreamingGrid(const FSpatialHashRu
 
 				for (const FActorInstance& ActorInstance : FilteredActors)
 				{
-					const FWorldPartitionActorDesc* ActorDesc = ActorInstance.GetActorDesc();
-					FGuid ParentHLOD = CachedHLODParents.FindRef(ActorInstance.Actor);
-					if (ParentHLOD.IsValid())
+					const FWorldPartitionActorDescView& ActorDescView = ActorInstance.GetActorDescView();
+					if (ActorDescView.GetHLODParent().IsValid())
 					{
-						ReferencedHLODActors.Add(ParentHLOD);
+						ReferencedHLODActors.Add(ActorDescView.GetHLODParent());
 					}
 					StreamingCell->AddActorToCell(ActorInstance.Actor, ActorInstance.ContainerInstance->ID, ActorInstance.ContainerInstance->Transform, ActorInstance.ContainerInstance->Container);
-					UE_LOG(LogWorldPartitionRuntimeSpatialHash, Verbose, TEXT("  Actor : %s (%s) (Container %08x) Origin(%s)"), *(ActorDesc->GetActorPath().ToString()), *ActorDesc->GetGuid().ToString(EGuidFormats::UniqueObjectGuid), ActorInstance.ContainerInstance->ID, *FVector2D(ActorInstance.GetOrigin()).ToString());
+					UE_LOG(LogWorldPartitionRuntimeSpatialHash, Verbose, TEXT("  Actor : %s (%s) (Container %08x) Origin(%s)"), *(ActorDescView.GetActorPath().ToString()), *ActorDescView.GetGuid().ToString(EGuidFormats::UniqueObjectGuid), ActorInstance.ContainerInstance->ID, *FVector2D(ActorInstance.GetOrigin()).ToString());
 				}
 
 				if (ReferencedHLODActors.Num() > 0)
diff --git a/Engine/Source/Runtime/Engine/Public/WorldPartition/WorldPartitionActorCluster.h b/Engine/Source/Runtime/Engine/Public/WorldPartition/WorldPartitionActorCluster.h
--- a/Engine/Source/Runtime/Engine/Public/WorldPartition/WorldPartitionActorCluster.h
+++ b/Engine/Source/Runtime/Engine/Public/WorldPartition/WorldPartitionActorCluster.h
@@ -4,11 +4,13 @@
 #if WITH_EDITOR
 
 #include "WorldPartition/DataLayer/DataLayersID.h"
+#include "WorldPartition/WorldPartitionActorDescView.h"
 #include "Misc/HashBuilder.h"
 
 class UActorDescContainer;
 class UWorldPartition;
 class FWorldPartitionActorDesc;
+class UWorldPartitionRuntimeHash;
 enum class EContainerClusterMode : uint8;
 
 DECLARE_LOG_CATEGORY_EXTERN(LogWorldPartitionActorCluster, Log, All);
@@ -25,7 +27,7 @@ struct FActorCluster
 	TSet<const UDataLayer*>		DataLayers;
 	FDataLayersID				DataLayersID;
 
-	FActorCluster(const FWorldPartitionActorDesc* InActorDesc, EActorGridPlacement InGridPlacement, const UActorDescContainer* InContainer);
+	FActorCluster(UWorld* InWorld, const FWorldPartitionActorDescView& InActorDescView, EActorGridPlacement InGridPlacement);
 	void Add(const FActorCluster& InActorCluster);
 
 };
@@ -37,8 +39,8 @@ struct FActorClusterInstance;
  */
 struct FActorContainerInstance
 {
-	FActorContainerInstance(const UActorDescContainer* InContainer);
-	FActorContainerInstance(uint32 InID, const FTransform& InTransform, const FBox& InBounds, const TSet<FName>& InDataLayers, EContainerClusterMode InClusterMode, const UActorDescContainer* InContainer, TSet<FGuid> InChildContainers);
+	FActorContainerInstance(const UActorDescContainer* InContainer, TMap<FGuid, FWorldPartitionActorDescView> InActorDescViewMap);
+	FActorContainerInstance(uint32 InID, const FTransform& InTransform, const FBox& InBounds, const TSet<FName>& InDataLayers, EContainerClusterMode InClusterMode, const UActorDescContainer* InContainer, TSet<FGuid> InChildContainers, TMap<FGuid, FWorldPartitionActorDescView> InActorDescViewMap);
 	
 	uint32						ID;
 	FTransform					Transform;
@@ -46,7 +48,10 @@ struct FActorContainerInstance
 	EContainerClusterMode		ClusterMode;
 	const UActorDescContainer*	Container;
 	TSet<FGuid>					ChildContainers;
+	TMap<FGuid, FWorldPartitionActorDescView> ActorDescViewMap;
 	TSet<const UDataLayer*>		DataLayers;
+
+	const FWorldPartitionActorDescView& GetActorDescView(const FGuid& InGuid) const;
 };
 
 /**
@@ -90,7 +95,7 @@ struct FActorInstance
 
 	bool ShouldStripFromStreaming() const;
 
-	const FWorldPartitionActorDesc* GetActorDesc() const;
+	const FWorldPartitionActorDescView& GetActorDescView() const;
 };
 
 /**
@@ -99,17 +104,20 @@ struct FActorInstance
 class FActorClusterContext
 {
 public:
-	typedef TFunctionRef<bool(const FWorldPartitionActorDesc&)> FFilterPredicate;
+	typedef TFunctionRef<bool(const FWorldPartitionActorDescView&)> FFilterPredicate;
 
 	/**
 	 * Create the actor clusters from the root World Partition. Optionally filtering some actors and including child Containers. 
 	 */
-	FActorClusterContext(const UWorldPartition* InWorldPartition, TOptional<FFilterPredicate> InFilterPredicate = TOptional<FFilterPredicate>(), bool bInIncludeChildContainers = true);
+	FActorClusterContext(const UWorldPartition* InWorldPartition, const UWorldPartitionRuntimeHash* InRuntimeHash, TOptional<FFilterPredicate> InFilterPredicate = TOptional<FFilterPredicate>(), bool bInIncludeChildContainers = true);
 				
 	/**
 	 * Returns the list of cluster instances of this context. 
 	 */
 	const TArray<FActorClusterInstance>& GetClusterInstances() const { return ClusterInstances; }
+	
+	FActorContainerInstance* GetClusterInstance(const UActorDescContainer* InContainer);
+
 private:
 	void CreateActorClusters();
 	void CreateContainerInstanceRecursive(uint32 ID, const FTransform& Transform, EContainerClusterMode ClusterMode, const UActorDescContainer* ActorDescContainer, const TSet<FName>& DataLayers, FBox* ParentBounds = nullptr);
@@ -117,6 +125,7 @@ private:
 	
 	// Init data
 	const UWorldPartition* WorldPartition;
+	const UWorldPartitionRuntimeHash* RuntimeHash;
 	TOptional<FFilterPredicate> FilterPredicate;
 	bool bIncludeChildContainers;
 
diff --git a/Engine/Source/Runtime/Engine/Public/WorldPartition/WorldPartitionActorDescView.h b/Engine/Source/Runtime/Engine/Public/WorldPartition/WorldPartitionActorDescView.h
new file mode 100644
--- /dev/null
+++ b/Engine/Source/Runtime/Engine/Public/WorldPartition/WorldPartitionActorDescView.h
@@ -0,0 +1,50 @@
+// Copyright Epic Games, Inc. All Rights Reserved.
+#pragma once
+
+#include "CoreMinimal.h"
+
+#if WITH_EDITOR
+class UHLODLayer;
+class FWorldPartitionActorDesc;
+enum class EContainerClusterMode : uint8;
+
+/**
+ * A view on top of an actor desc, used to cache information that can be (potentially) different than the actor desc
+ * itself due to streaming generation logic, etc.
+ */
+class ENGINE_API FWorldPartitionActorDescView
+{
+	friend class UWorldPartitionRuntimeHash;
+
+public:
+	FWorldPartitionActorDescView();
+	FWorldPartitionActorDescView(const FWorldPartitionActorDesc* InActorDesc);
+
+	const FGuid& GetGuid() const;
+	inline FName GetClass() const;
+	inline UClass* GetActorClass() const;
+	inline FVector GetOrigin() const;
+	EActorGridPlacement GetGridPlacement() const;
+	FName GetRuntimeGrid() const;
+	bool GetActorIsEditorOnly() const;
+	bool GetLevelBoundsRelevant() const;
+	bool GetActorIsHLODRelevant() const;
+	UHLODLayer* GetHLODLayer() const;
+	const TArray<FName>& GetDataLayers() const;
+	FName GetActorPackage() const;
+	FName GetActorPath() const;
+	FName GetActorLabel() const;
+	FBox GetBounds() const;
+	const TArray<FGuid>& GetReferences() const;
+
+	bool GetContainerInstance(const UActorDescContainer*& OutLevelContainer, FTransform& OutLevelTransform, EContainerClusterMode& OutClusterMode) const;
+
+	FGuid GetHLODParent() const;
+	void SetHLODParent(const FGuid& InHLODParent);
+
+protected:
+	const FWorldPartitionActorDesc* ActorDesc;
+	EActorGridPlacement EffectiveGridPlacement;
+	FGuid HLODParent;
+};
+#endif
\ No newline at end of file
diff --git a/Engine/Source/Runtime/Engine/Public/WorldPartition/WorldPartitionRuntimeHash.h b/Engine/Source/Runtime/Engine/Public/WorldPartition/WorldPartitionRuntimeHash.h
--- a/Engine/Source/Runtime/Engine/Public/WorldPartition/WorldPartitionRuntimeHash.h
+++ b/Engine/Source/Runtime/Engine/Public/WorldPartition/WorldPartitionRuntimeHash.h
@@ -6,6 +6,7 @@
 #include "UObject/Object.h"
 #include "WorldPartition.h"
 #include "WorldPartitionStreamingPolicy.h"
+#include "WorldPartitionActorDescView.h"
 #include "WorldPartition/WorldPartitionHandle.h"
 #include "WorldPartitionRuntimeHash.generated.h"
 
@@ -16,6 +17,8 @@ class ENGINE_API UWorldPartitionRuntimeHash : public UObject
 {
 	GENERATED_UCLASS_BODY()
 
+	friend class FActorClusterContext;
+
 #if WITH_EDITOR
 	virtual void SetDefaultValues() {}
 	virtual void ImportFromWorldComposition(class UWorldComposition* WorldComposition) {}
@@ -48,6 +51,12 @@ class ENGINE_API UWorldPartitionRuntimeHash : public UObject
 	virtual void Draw2D(class UCanvas* Canvas, const TArray<FWorldPartitionStreamingSource>& Sources, const FVector2D& PartitionCanvasOffset, const FVector2D& PartitionCanvasSize) const {}
 	virtual void Draw3D(const TArray<FWorldPartitionStreamingSource>& Sources) const {}
 
+protected:
+#if WITH_EDITOR
+	virtual void CreateActorDescViewMap(const UActorDescContainer* Container, TMap<FGuid, FWorldPartitionActorDescView>& OutActorDescViewMap) const;
+	void ChangeActorDescViewGridPlacement(FWorldPartitionActorDescView& ActorDescView, EActorGridPlacement GridPlacement) const;
+#endif
+
 private:
 #if WITH_EDITOR
 	void ForceExternalActorLevelReference(bool bForceExternalActorLevelReferenceForPIE);
diff --git a/Engine/Source/Runtime/Engine/Public/WorldPartition/WorldPartitionRuntimeSpatialHash.h b/Engine/Source/Runtime/Engine/Public/WorldPartition/WorldPartitionRuntimeSpatialHash.h
--- a/Engine/Source/Runtime/Engine/Public/WorldPartition/WorldPartitionRuntimeSpatialHash.h
+++ b/Engine/Source/Runtime/Engine/Public/WorldPartition/WorldPartitionRuntimeSpatialHash.h
@@ -202,6 +202,8 @@ public:
 
 	FName GetCellName(FName InGridName, int32 InLevel, int32 InCellX, int32 InCellY, const FDataLayersID& InDataLayerID) const;
 	static FName GetCellName(UWorldPartition* WorldPartition, FName InGridName, int32 InLevel, int32 InCellX, int32 InCellY, const FDataLayersID& InDataLayerID);
+
+	virtual void CreateActorDescViewMap(const UActorDescContainer* Container, TMap<FGuid, FWorldPartitionActorDescView>& OutActorDescViewMap) const override;
 #endif
 
 	// streaming interface
@@ -228,7 +230,6 @@ private:
 	mutable FWorldPartitionRuntimeSpatialHashGridPreviewer GridPreviewer;
 
 	TMap<FString, UWorldPartitionRuntimeCell*> PackagesToGenerateForCook;
-	TMap<FGuid, FGuid> CachedHLODParents;
 #endif
 
 	/** 
@@ -246,7 +247,6 @@ private:
 	void GetStreamingCells(const FVector& Position, const FSpatialHashStreamingGrid& StreamingGrid, TSet<const UWorldPartitionRuntimeCell*>& Cells) const;
 #if WITH_EDITOR
 	bool CreateStreamingGrid(const FSpatialHashRuntimeGrid& RuntimeGrid, const FSquare2DGridHelper& PartionedActors, EWorldPartitionStreamingMode Mode, UWorldPartitionStreamingPolicy* StreamingPolicy, TArray<FString>* OutPackagesToGenerate = nullptr);
-	void CacheHLODParents();
 #endif
 
 	friend class UWorldPartitionSubsystem;
