Expose FActorDataLayer Name to BP and added BP function to compare UDataLayer with FActorDataLayer.

diff --git a/Engine/Source/Runtime/Engine/Public/WorldPartition/DataLayer/ActorDataLayer.h b/Engine/Source/Runtime/Engine/Public/WorldPartition/DataLayer/ActorDataLayer.h
--- a/Engine/Source/Runtime/Engine/Public/WorldPartition/DataLayer/ActorDataLayer.h
+++ b/Engine/Source/Runtime/Engine/Public/WorldPartition/DataLayer/ActorDataLayer.h
@@ -22,6 +22,6 @@ struct ENGINE_API FActorDataLayer
 	FORCEINLINE bool operator<(const FActorDataLayer& Other) const { return Name.FastLess(Other.Name); }
 
 	/** The name of this layer */
-	UPROPERTY(VisibleAnywhere, Category = DataLayer)
+	UPROPERTY(BlueprintReadOnly, VisibleAnywhere, Category = DataLayer)
 	FName Name;
 };
\ No newline at end of file
diff --git a/Engine/Source/Runtime/Engine/Public/WorldPartition/DataLayer/DataLayer.h b/Engine/Source/Runtime/Engine/Public/WorldPartition/DataLayer/DataLayer.h
--- a/Engine/Source/Runtime/Engine/Public/WorldPartition/DataLayer/DataLayer.h
+++ b/Engine/Source/Runtime/Engine/Public/WorldPartition/DataLayer/DataLayer.h
@@ -7,6 +7,7 @@
 #include "UObject/UObjectGlobals.h"
 #include "UObject/Object.h"
 #include "WorldDataLayers.h"
+#include "ActorDataLayer.h"
 
 #include "DataLayer.generated.h"
 
@@ -31,6 +32,9 @@ public:
 
 	bool IsLocked() const { return bIsLocked; }
 #endif
+
+	UFUNCTION(Category = DataLayer, BlueprintCallable)
+	bool Equals(const FActorDataLayer& ActorDataLayer) const { return ActorDataLayer.Name == GetFName(); }
 	
 	UFUNCTION(Category = DataLayer, BlueprintCallable)
 	FName GetDataLayerLabel() const  { return DataLayerLabel; }
