Add parameter structure layout hash to the material shader DDC (UE-109435).

diff --git a/Engine/Source/Runtime/Engine/Private/Materials/MaterialShader.cpp b/Engine/Source/Runtime/Engine/Private/Materials/MaterialShader.cpp
--- a/Engine/Source/Runtime/Engine/Private/Materials/MaterialShader.cpp
+++ b/Engine/Source/Runtime/Engine/Private/Materials/MaterialShader.cpp
@@ -767,6 +767,11 @@ void FMaterialShaderMapId::AppendKeyString(FString& KeyString) const
 		KeyString += ShaderType->GetName();
 		KeyString += ShaderTypeDependency.SourceHash.ToString();
 
+		if (const FShaderParametersMetadata* ParameterStructMetadata = ShaderType->GetRootParametersMetadata())
+		{
+			KeyString += FString::Printf(TEXT("%08x"), ParameterStructMetadata->GetLayoutHash());
+		}
+
 		const FSHAHash LayoutHash = Freeze::HashLayout(ShaderType->GetLayout(), LayoutParams);
 		KeyString += LayoutHash.ToString();
 
@@ -793,6 +798,11 @@ void FMaterialShaderMapId::AppendKeyString(FString& KeyString) const
 
 		for (const FShaderType* ShaderType : ShaderPipelineType->GetStages())
 		{
+			if (const FShaderParametersMetadata* ParameterStructMetadata = ShaderType->GetRootParametersMetadata())
+			{
+				KeyString += FString::Printf(TEXT("%08x"), ParameterStructMetadata->GetLayoutHash());
+			}
+
 			const TMap<const TCHAR*, FCachedUniformBufferDeclaration>& ReferencedUniformBufferStructsCache = ShaderType->GetReferencedUniformBufferStructsCache();
 
 			// Gather referenced uniform buffers
