Fix Static Analysis: remove if that shouldn't have been added in latest change

diff --git a/Engine/Source/Runtime/Engine/Private/LevelInstance/Packed/PackedLevelInstanceBuilder.cpp b/Engine/Source/Runtime/Engine/Private/LevelInstance/Packed/PackedLevelInstanceBuilder.cpp
--- a/Engine/Source/Runtime/Engine/Private/LevelInstance/Packed/PackedLevelInstanceBuilder.cpp
+++ b/Engine/Source/Runtime/Engine/Private/LevelInstance/Packed/PackedLevelInstanceBuilder.cpp
@@ -362,14 +362,11 @@ bool FPackedLevelInstanceBuilder::CreateOrUpdateBlueprintFromUnpacked(ALevelInst
 	UWorld* World = InActor->GetWorld();
 	SpawnParams.OverrideLevel = World->PersistentLevel;
 
-	APackedLevelInstance* PackedLevelInstance = World->SpawnActor<APackedLevelInstance>(InActor->GetActorLocation(), InActor->GetActorRotation(), SpawnParams);
+	APackedLevelInstance* PackedLevelInstance = World->SpawnActor<APackedLevelInstance>(InActor->GetActorLocation(), InActor->GetActorRotation(), SpawnParams);	
 	PackedLevelInstance->SetWorldAsset(InActor->GetWorldAsset());
 	ON_SCOPE_EXIT
 	{
-		if (PackedLevelInstance)
-		{
-			InActor->GetWorld()->DestroyActor(PackedLevelInstance);
-		}
+		InActor->GetWorld()->DestroyActor(PackedLevelInstance);
 	};
 
 	if (!PackActor(PackedLevelInstance, InActor))
