Fix crash for non-external actors referenced by the level script.

diff --git a/Engine/Source/Runtime/Engine/Private/WorldPartition/WorldPartitionRuntimeHash.cpp b/Engine/Source/Runtime/Engine/Private/WorldPartition/WorldPartitionRuntimeHash.cpp
--- a/Engine/Source/Runtime/Engine/Private/WorldPartition/WorldPartitionRuntimeHash.cpp
+++ b/Engine/Source/Runtime/Engine/Private/WorldPartition/WorldPartitionRuntimeHash.cpp
@@ -66,7 +66,10 @@ void UWorldPartitionRuntimeHash::CreateActorDescViewMap(const UActorDescContaine
 
 		for (AActor* Actor : LevelScriptExternalActorReferences)
 		{
-			ChangeActorDescViewGridPlacement(OutActorDescViewMap.FindChecked(Actor->GetActorGuid()), EActorGridPlacement::AlwaysLoaded);
+			if (FWorldPartitionActorDescView* ActorDescView = OutActorDescViewMap.Find(Actor->GetActorGuid()))
+			{
+				ChangeActorDescViewGridPlacement(*ActorDescView, EActorGridPlacement::AlwaysLoaded);
+			}
 		}
 	}
 }
