Make sure the MyBoneName element of the HitResults is available as a pin in the Blueprint "On Compmonent Hit" event

diff --git a/Engine/Source/Runtime/Engine/Classes/Kismet/GameplayStatics.h b/Engine/Source/Runtime/Engine/Classes/Kismet/GameplayStatics.h
--- a/Engine/Source/Runtime/Engine/Classes/Kismet/GameplayStatics.h
+++ b/Engine/Source/Runtime/Engine/Classes/Kismet/GameplayStatics.h
@@ -920,12 +920,13 @@ public:
 	 * @param HitActor		Actor hit by the trace.
 	 * @param HitComponent	PrimitiveComponent hit by the trace.
 	 * @param HitBoneName	Name of the bone hit (valid only if we hit a skeletal mesh).
+	 * @param BoneName		Name of the trace bone hit (valid only if we hit a skeletal mesh).
 	 * @param HitItem		Primitive-specific data recording which item in the primitive was hit
 	 * @param ElementIndex	If colliding with a primitive with multiple parts, index of the part that was hit.
 	 * @param FaceIndex		If colliding with trimesh or landscape, index of face that was hit.
 	 */
 	UFUNCTION(BlueprintPure, Category = "Collision", meta=(NativeBreakFunc, AdvancedDisplay="3"))
-	static void BreakHitResult(const struct FHitResult& Hit, bool& bBlockingHit, bool& bInitialOverlap, float& Time, float& Distance, FVector& Location, FVector& ImpactPoint, FVector& Normal, FVector& ImpactNormal, class UPhysicalMaterial*& PhysMat, class AActor*& HitActor, class UPrimitiveComponent*& HitComponent, FName& HitBoneName, int32& HitItem, int32& ElementIndex, int32& FaceIndex, FVector& TraceStart, FVector& TraceEnd);
+	static void BreakHitResult(const struct FHitResult& Hit, bool& bBlockingHit, bool& bInitialOverlap, float& Time, float& Distance, FVector& Location, FVector& ImpactPoint, FVector& Normal, FVector& ImpactNormal, class UPhysicalMaterial*& PhysMat, class AActor*& HitActor, class UPrimitiveComponent*& HitComponent, FName& HitBoneName, FName& BoneName, int32& HitItem, int32& ElementIndex, int32& FaceIndex, FVector& TraceStart, FVector& TraceEnd);
 
 	/** 
 	 *	Create a HitResult struct
@@ -942,12 +943,13 @@ public:
 	 * @param HitActor		Actor hit by the trace.
 	 * @param HitComponent	PrimitiveComponent hit by the trace.
 	 * @param HitBoneName	Name of the bone hit (valid only if we hit a skeletal mesh).
+	 * @param BoneName		Name of the trace bone hit (valid only if we hit a skeletal mesh).
 	 * @param HitItem		Primitive-specific data recording which item in the primitive was hit
 	 * @param ElementIndex	If colliding with a primitive with multiple parts, index of the part that was hit.
 	 * @param FaceIndex		If colliding with trimesh or landscape, index of face that was hit.
 	 */
 	UFUNCTION(BlueprintPure, Category = "Collision", meta = (NativeMakeFunc, AdvancedDisplay="2", Normal="0,0,1", ImpactNormal="0,0,1"))
-	static FHitResult MakeHitResult(bool bBlockingHit, bool bInitialOverlap, float Time, float Distance, FVector Location, FVector ImpactPoint, FVector Normal, FVector ImpactNormal, class UPhysicalMaterial* PhysMat, class AActor* HitActor, class UPrimitiveComponent* HitComponent, FName HitBoneName, int32 HitItem, int32 ElementIndex, int32 FaceIndex, FVector TraceStart, FVector TraceEnd);
+	static FHitResult MakeHitResult(bool bBlockingHit, bool bInitialOverlap, float Time, float Distance, FVector Location, FVector ImpactPoint, FVector Normal, FVector ImpactNormal, class UPhysicalMaterial* PhysMat, class AActor* HitActor, class UPrimitiveComponent* HitComponent, FName HitBoneName, FName BoneName, int32 HitItem, int32 ElementIndex, int32 FaceIndex, FVector TraceStart, FVector TraceEnd);
 
 
 	/** Returns the EPhysicalSurface type of the given Hit. 
diff --git a/Engine/Source/Runtime/Engine/Private/GameplayStatics.cpp b/Engine/Source/Runtime/Engine/Private/GameplayStatics.cpp
--- a/Engine/Source/Runtime/Engine/Private/GameplayStatics.cpp
+++ b/Engine/Source/Runtime/Engine/Private/GameplayStatics.cpp
@@ -1177,7 +1177,7 @@ UParticleSystemComponent* UGameplayStatics::SpawnEmitterAttached(UParticleSystem
 }
 
 // FRED_TODO: propagate hit object handles further instead of converting to an actor here
-void UGameplayStatics::BreakHitResult(const FHitResult& Hit, bool& bBlockingHit, bool& bInitialOverlap, float& Time, float& Distance, FVector& Location, FVector& ImpactPoint, FVector& Normal, FVector& ImpactNormal, UPhysicalMaterial*& PhysMat, AActor*& HitActor, UPrimitiveComponent*& HitComponent, FName& HitBoneName, int32& HitItem, int32& ElementIndex, int32& FaceIndex, FVector& TraceStart, FVector& TraceEnd)
+void UGameplayStatics::BreakHitResult(const FHitResult& Hit, bool& bBlockingHit, bool& bInitialOverlap, float& Time, float& Distance, FVector& Location, FVector& ImpactPoint, FVector& Normal, FVector& ImpactNormal, UPhysicalMaterial*& PhysMat, AActor*& HitActor, UPrimitiveComponent*& HitComponent, FName& HitBoneName, FName& BoneName, int32& HitItem, int32& ElementIndex, int32& FaceIndex, FVector& TraceStart, FVector& TraceEnd)
 {
 	SCOPE_CYCLE_COUNTER(STAT_BreakHitResult);
 	bBlockingHit = Hit.bBlockingHit;
@@ -1192,6 +1192,7 @@ void UGameplayStatics::BreakHitResult(const FHitResult& Hit, bool& bBlockingHit,
 	HitActor = Hit.GetHitObjectHandle().FetchActor();
 	HitComponent = Hit.GetComponent();
 	HitBoneName = Hit.BoneName;
+	BoneName = Hit.MyBoneName;
 	HitItem = Hit.Item;
 	ElementIndex = Hit.ElementIndex;
 	TraceStart = Hit.TraceStart;
@@ -1199,7 +1200,7 @@ void UGameplayStatics::BreakHitResult(const FHitResult& Hit, bool& bBlockingHit,
 	FaceIndex = Hit.FaceIndex;
 }
 
-FHitResult UGameplayStatics::MakeHitResult(bool bBlockingHit, bool bInitialOverlap, float Time, float Distance, FVector Location, FVector ImpactPoint, FVector Normal, FVector ImpactNormal, class UPhysicalMaterial* PhysMat, class AActor* HitActor, class UPrimitiveComponent* HitComponent, FName HitBoneName, int32 HitItem, int32 ElementIndex, int32 FaceIndex, FVector TraceStart, FVector TraceEnd)
+FHitResult UGameplayStatics::MakeHitResult(bool bBlockingHit, bool bInitialOverlap, float Time, float Distance, FVector Location, FVector ImpactPoint, FVector Normal, FVector ImpactNormal, class UPhysicalMaterial* PhysMat, class AActor* HitActor, class UPrimitiveComponent* HitComponent, FName HitBoneName, FName BoneName, int32 HitItem, int32 ElementIndex, int32 FaceIndex, FVector TraceStart, FVector TraceEnd)
 {
 	SCOPE_CYCLE_COUNTER(STAT_MakeHitResult);
 	FHitResult Hit;
@@ -1215,6 +1216,7 @@ FHitResult UGameplayStatics::MakeHitResult(bool bBlockingHit, bool bInitialOverl
 	Hit.HitObjectHandle = FActorInstanceHandle(HitActor);
 	Hit.Component = HitComponent;
 	Hit.BoneName = HitBoneName;
+	Hit.MyBoneName = BoneName;
 	Hit.Item = HitItem;
 	Hit.ElementIndex = ElementIndex;
 	Hit.TraceStart = TraceStart;
