Changes for allocating StaticMesh subobjects in a more thread-safe way.
Always construct a UMeshDescriptionBase as a default subobject of the MeshDescriptionBulkData object, which can be recycled each time a UMeshDescriptionBase is required for that LOD.
Put GC scope guards around NewObject<> allocations to avoid clashes with GC when allocated on a thread.
Also identified other static mesh subobjects which can be potentially created in cooking commandlets when GC is taking place.

diff --git a/Engine/Source/Runtime/Engine/Classes/Engine/StaticMeshSourceData.h b/Engine/Source/Runtime/Engine/Classes/Engine/StaticMeshSourceData.h
--- a/Engine/Source/Runtime/Engine/Classes/Engine/StaticMeshSourceData.h
+++ b/Engine/Source/Runtime/Engine/Classes/Engine/StaticMeshSourceData.h
@@ -24,8 +24,7 @@ class ENGINE_API UStaticMeshDescriptionBulkData : public UMeshDescriptionBaseBul
 	GENERATED_BODY()
 
 public:
-	virtual UClass* GetMeshDescriptionType() const override;
-
+	UStaticMeshDescriptionBulkData();
 };
 
 
diff --git a/Engine/Source/Runtime/Engine/Private/StaticMesh.cpp b/Engine/Source/Runtime/Engine/Private/StaticMesh.cpp
--- a/Engine/Source/Runtime/Engine/Private/StaticMesh.cpp
+++ b/Engine/Source/Runtime/Engine/Private/StaticMesh.cpp
@@ -6258,7 +6258,11 @@ void UStaticMesh::CreateBodySetup()
 {
 	if (GetBodySetup() == nullptr)
 	{
-		UBodySetup* NewBodySetup = NewObject<UBodySetup>(this);
+		UBodySetup* NewBodySetup = nullptr;
+		{
+			FGCScopeGuard Scope;
+			NewBodySetup = NewObject<UBodySetup>(this);
+		}
 		NewBodySetup->DefaultInstance.SetCollisionProfileName(UCollisionProfile::BlockAll_ProfileName);
 		NewBodySetup->bSupportUVsAndFaceRemap = bSupportPhysicalMaterialMasks;
 		SetBodySetup(NewBodySetup);
@@ -6271,7 +6275,12 @@ void UStaticMesh::CreateNavCollision(const bool bIsUpdate)
 	{
 		if (GetNavCollision() == nullptr)
 		{
-			SetNavCollision(UNavCollisionBase::ConstructNew(*this));
+			UNavCollisionBase* NewNavCollisionBase = nullptr;
+			{
+				FGCScopeGuard Scope;
+				NewNavCollisionBase = UNavCollisionBase::ConstructNew(*this);
+			}
+			SetNavCollision(NewNavCollisionBase);
 		}
 
 		if (GetNavCollision())
diff --git a/Engine/Source/Runtime/Engine/Private/StaticMeshSourceData.cpp b/Engine/Source/Runtime/Engine/Private/StaticMeshSourceData.cpp
--- a/Engine/Source/Runtime/Engine/Private/StaticMeshSourceData.cpp
+++ b/Engine/Source/Runtime/Engine/Private/StaticMeshSourceData.cpp
@@ -12,9 +12,10 @@
 #endif
 
 
-UClass* UStaticMeshDescriptionBulkData::GetMeshDescriptionType() const
+UStaticMeshDescriptionBulkData::UStaticMeshDescriptionBulkData()
 {
-	return UStaticMeshDescription::StaticClass();
+	const bool bTransient = true;
+	PreallocatedMeshDescription = CreateDefaultSubobject<UStaticMeshDescription>(TEXT("MeshDescription"), bTransient);
 }
 
 
@@ -76,10 +77,14 @@ FStaticMeshSourceModel& FStaticMeshSourceModel::operator=(FStaticMeshSourceModel
 void FStaticMeshSourceModel::CreateSubObjects(UStaticMesh* InOwner)
 {
 #if WITH_EDITORONLY_DATA
-	ensure(IsInGameThread());
 	if (StaticMeshDescriptionBulkData == nullptr)
 	{
+		// Currently there is a restriction on creating UObjects on a thread while garbage collection is taking place,
+		// so use a GCScopeGuard to prevent it happening at the same time.
+		FGCScopeGuard Guard;
 		StaticMeshDescriptionBulkData = NewObject<UStaticMeshDescriptionBulkData>(InOwner, NAME_None, RF_Transactional);
+
+		// OK to clear the async flag immediately here, because the UObject* has already been assigned directly to a UPROPERTY.
 		StaticMeshDescriptionBulkData->AtomicallyClearInternalFlags(EInternalObjectFlags::Async);
 	}
 
@@ -294,6 +299,8 @@ void FStaticMeshSourceModel::SerializeBulkData(FArchive& Ar, UObject* Owner)
 		// If this was a legacy asset, or is being created for the first time, create a bulkdata UObject wrapper
 		if (StaticMeshDescriptionBulkData == nullptr)
 		{
+			// For now, we assume that serialization will only occur on the game thread
+			ensure(IsInGameThread());
 			StaticMeshDescriptionBulkData = NewObject<UStaticMeshDescriptionBulkData>(Owner, NAME_None, RF_Transactional);
 		}
 	}
diff --git a/Engine/Source/Runtime/MeshDescription/Private/MeshDescriptionBaseBulkData.cpp b/Engine/Source/Runtime/MeshDescription/Private/MeshDescriptionBaseBulkData.cpp
--- a/Engine/Source/Runtime/MeshDescription/Private/MeshDescriptionBaseBulkData.cpp
+++ b/Engine/Source/Runtime/MeshDescription/Private/MeshDescriptionBaseBulkData.cpp
@@ -58,12 +58,6 @@ bool UMeshDescriptionBaseBulkData::NeedsLoadForEditorGame() const
 }
 
 
-UClass* UMeshDescriptionBaseBulkData::GetMeshDescriptionType() const
-{
-	return UMeshDescriptionBase::StaticClass();
-}
-
-
 #if WITH_EDITORONLY_DATA
 
 void UMeshDescriptionBaseBulkData::Empty()
@@ -76,11 +70,8 @@ UMeshDescriptionBase* UMeshDescriptionBaseBulkData::CreateMeshDescription()
 {
 	if (MeshDescription == nullptr)
 	{
-		// Assign the new UObject directly to a UPROPERTY.
-		// If we were run from a different thread, we can immediately clear the Async internal flag, now that GC refs are set up correctly.
-		UClass* MeshDescriptionType = GetMeshDescriptionType();
-		MeshDescription = NewObject<UMeshDescriptionBase>(this, MeshDescriptionType, NAME_None, RF_Transient);
-		MeshDescription->AtomicallyClearInternalFlags(EInternalObjectFlags::Async);
+		check(PreallocatedMeshDescription);
+		MeshDescription = PreallocatedMeshDescription;
 	}
 
 	// Reset the mesh description and register its attributes
diff --git a/Engine/Source/Runtime/MeshDescription/Public/MeshDescriptionBaseBulkData.h b/Engine/Source/Runtime/MeshDescription/Public/MeshDescriptionBaseBulkData.h
--- a/Engine/Source/Runtime/MeshDescription/Public/MeshDescriptionBaseBulkData.h
+++ b/Engine/Source/Runtime/MeshDescription/Public/MeshDescriptionBaseBulkData.h
@@ -29,7 +29,6 @@ public:
 	virtual bool NeedsLoadForClient() const override;
 	virtual bool NeedsLoadForServer() const override;
 	virtual bool NeedsLoadForEditorGame() const override;
-	virtual UClass* GetMeshDescriptionType() const;
 
 #if WITH_EDITORONLY_DATA
 
@@ -48,6 +47,9 @@ public:
 protected:
 	TPimplPtr<FMeshDescriptionBulkData> BulkData;
 
+	UPROPERTY(Transient)
+	UMeshDescriptionBase* PreallocatedMeshDescription;
+
 	UPROPERTY(Transient)
 	UMeshDescriptionBase* MeshDescription;
 
