Fixed non-editor build.

diff --git a/Engine/Source/Runtime/Engine/Private/StaticMeshSourceData.cpp b/Engine/Source/Runtime/Engine/Private/StaticMeshSourceData.cpp
--- a/Engine/Source/Runtime/Engine/Private/StaticMeshSourceData.cpp
+++ b/Engine/Source/Runtime/Engine/Private/StaticMeshSourceData.cpp
@@ -14,8 +14,10 @@
 
 UStaticMeshDescriptionBulkData::UStaticMeshDescriptionBulkData()
 {
+#if WITH_EDITORONLY_DATA
 	const bool bTransient = true;
 	PreallocatedMeshDescription = CreateDefaultSubobject<UStaticMeshDescription>(TEXT("MeshDescription"), bTransient);
+#endif
 }
 
 
