Early WIP of refactored material HLSL translator/generator.
 Only mildly functional currently, my goal here is to get feedback on the high level direction.  The new code is disabled by default; it needs to be opted-into per material (and the UI for that is hidden by default behind a CVAR), so it shouldn't affect any existing workflows.
Basic outline:
* UMaterialExpression::Compile is replaced by GenerateHLSLStatement/GenerateHLSLExpression
* GenerateHLSLExpression returns a HLSLTree::FExpression derived object.  This is broadly similar to FShaderCodeChunk, except FExpression stores pointers to dependent expressions, and only generates the final HLSL string on demand (EmitHLSL method).  FExpression will have many derived classes to represent various HLSL expresisons like Add, Parameter, Constant, etc.
* GenerateHLSLStatement returns a HLSLTree::FStatement.  Statements represent chunk of HLSL code rather than a value.  'if(...)' is a statement for example.
* UMaterialExpression will normally override only 1 of GenerateHLSLStatement/GenerateHLSLExpression, although both are possible (for loop may override both, since the for loop index is an expression for example)
* FMaterialHLSLGenerator is the 'context' object passed to GenerateHLSLStatement/GenerateHLSLExpression.  This is kind of like FHLSLMaterialTranslator.  It's responsible for tracking which scope statements/expressions belong to (this can evolve as a given expression/statement is accessed from multiple UMaterialExpressions)
* HLSLTree.h contains the implementation of Expression and FStatement.  I feel like it's possible this could be useful outside of the material system (for example could be shared with Niagara).  Will be tricky to partition/organize the code in such a way to make this work though.
* HLSLTreeCommon.h contains derived expression/statement types.  The intent is to support plugins that define additional derived types.
* MaterialExpressionHLSL.cpp contains the implementations of GenerateHLSLStatement/GenerateHLSLExpression I've added so far.  Eventually all expression types will need to be converted over (this will be a big job).  Once the new system is done, we can remove all the old Compile() functions in MaterialExpressions.cpp
* I haven't given any real thought to Strata or analytic derivatives yet, but I don't forsee any serious problem adapting them to this new system

diff --git a/Engine/Shaders/Private/MaterialTemplate.ush b/Engine/Shaders/Private/MaterialTemplate.ush
--- a/Engine/Shaders/Private/MaterialTemplate.ush
+++ b/Engine/Shaders/Private/MaterialTemplate.ush
@@ -1517,27 +1517,42 @@ MaterialFloat4 ProcessMaterialLinearColorTextureLookup(MaterialFloat4 TextureVal
 	return TextureValue;
 }
 
-MaterialFloat ProcessMaterialGreyscaleTextureLookup(MaterialFloat TextureValue)
+MaterialFloat4 ProcessMaterialGreyscaleTextureLookup(MaterialFloat4 TextureValue)
 {
+	// Sampling a greyscale texture in D3D9 gives: (G,G,G)
+	// Sampling a greyscale texture in D3D11 gives: (G,0,0)
+	// This replication reproduces the D3D9 behavior in all cases.
+	MaterialFloat GrayValue = TextureValue.r;
 #if (COMPILER_GLSL_ES3_1 || VULKAN_PROFILE) // OpenGLES3.1, Vulkan3.1 do not support sRGB sampling from R8
 	#if MOBILE_EMULATION
 	if( ResolvedView.MobilePreviewMode > 0.5f )
 	{
 		// undo HW srgb->lin
-		TextureValue = pow(TextureValue, 1.0f/2.2f); // TODO: replace with a more accurate lin -> sRGB conversion.
+		GrayValue = pow(GrayValue, 1.0f/2.2f); // TODO: replace with a more accurate lin -> sRGB conversion.
 	}
 	#endif
 	// sRGB read approximation (in highp if possible)
-	float LinValue = TextureValue;
+	float LinValue = GrayValue;
 	LinValue *= LinValue;
-	return MaterialFloat(LinValue);
+	return MaterialFloat4(LinValue.rrrr);
 #endif 
-	return TextureValue;
+	return GrayValue.rrrr;
 }
 
-MaterialFloat ProcessMaterialLinearGreyscaleTextureLookup(MaterialFloat TextureValue)
+MaterialFloat4 ProcessMaterialLinearGreyscaleTextureLookup(MaterialFloat4 TextureValue)
 {
-	return TextureValue;
+	// Sampling a greyscale texture in D3D9 gives: (G,G,G)
+	// Sampling a greyscale texture in D3D11 gives: (G,0,0)
+	// This replication reproduces the D3D9 behavior in all cases.
+	return TextureValue.rrrr;
+}
+
+MaterialFloat4 ProcessMaterialAlphaTextureLookup(MaterialFloat4 TextureValue)
+{
+	// Sampling a single channel texture in D3D9 gives: (G,G,G)
+	// Sampling a single channel texture in D3D11 gives: (G,0,0)
+	// This replication reproduces the D3D9 behavior in all cases.
+	return TextureValue.rrrr;
 }
 
 /** Accesses a shared material sampler or falls back if independent samplers are not supported. */
diff --git a/Engine/Source/Editor/MaterialEditor/Private/MaterialEditorDetailCustomization.cpp b/Engine/Source/Editor/MaterialEditor/Private/MaterialEditorDetailCustomization.cpp
--- a/Engine/Source/Editor/MaterialEditor/Private/MaterialEditorDetailCustomization.cpp
+++ b/Engine/Source/Editor/MaterialEditor/Private/MaterialEditorDetailCustomization.cpp
@@ -737,6 +737,7 @@ TSharedRef<class IDetailCustomization> FMaterialDetailCustomization::MakeInstanc
 void FMaterialDetailCustomization::CustomizeDetails( IDetailLayoutBuilder& DetailLayout )
 {
 	static const auto CVarMaterialEnableControlFlow = IConsoleManager::Get().FindTConsoleVariableDataInt(TEXT("r.MaterialEnableControlFlow"));
+	static const auto CVarMaterialEnableNewHLSLGenerator = IConsoleManager::Get().FindTConsoleVariableDataInt(TEXT("r.MaterialEnableNewHLSLGenerator"));
 
 	TArray<TWeakObjectPtr<UObject> > Objects;
 	DetailLayout.GetObjectsBeingCustomized( Objects );
@@ -788,6 +789,11 @@ void FMaterialDetailCustomization::CustomizeDetails( IDetailLayoutBuilder& Detai
 				DetailLayout.HideProperty(PropertyHandle);
 			}
 
+			if (PropertyName == GET_MEMBER_NAME_CHECKED(UMaterial, bEnableNewHLSLGenerator) && !CVarMaterialEnableNewHLSLGenerator->GetValueOnAnyThread())
+			{
+				DetailLayout.HideProperty(PropertyHandle);
+			}
+
 #if WITH_EDITORONLY_DATA
 			// Hide the shading model 
 			if (!bIsShadingModelFromMaterialExpression && PropertyName == GET_MEMBER_NAME_CHECKED(UMaterial, UsedShadingModels))
diff --git a/Engine/Source/Editor/MaterialEditor/Private/MaterialEditorUtilities.cpp b/Engine/Source/Editor/MaterialEditor/Private/MaterialEditorUtilities.cpp
--- a/Engine/Source/Editor/MaterialEditor/Private/MaterialEditorUtilities.cpp
+++ b/Engine/Source/Editor/MaterialEditor/Private/MaterialEditorUtilities.cpp
@@ -27,6 +27,7 @@
 #include "Materials/MaterialExpressionCustomOutput.h"
 #include "Materials/MaterialExpressionMaterialAttributeLayers.h"
 #include "Materials/MaterialExpressionRerouteBase.h"
+#include "Materials/MaterialExpressionExecBegin.h"
 
 #include "Toolkits/ToolkitManager.h"
 #include "MaterialEditor.h"
@@ -239,13 +240,20 @@ void FMaterialEditorUtilities::GetVisibleMaterialParameters(const UMaterial* Mat
 	TArray<FGetVisibleMaterialParametersFunctionState*> FunctionStack;
 	FunctionStack.Push(FunctionState.Get());
 
-	for(uint32 i = 0; i < MP_MAX; ++i)
+	if (Material->IsCompiledWithExecutionFlow())
 	{
-		FExpressionInput* ExpressionInput = ((UMaterial *)Material)->GetExpressionInputForProperty((EMaterialProperty)i);
-
-		if(ExpressionInput)
+		GetVisibleMaterialParametersFromExpression(FMaterialExpressionKey(Material->ExpressionExecBegin, INDEX_NONE), MaterialInstance, VisibleExpressions, FunctionStack);
+	}
+	else
+	{
+		for (uint32 i = 0; i < MP_MAX; ++i)
 		{
-			GetVisibleMaterialParametersFromExpression(FMaterialExpressionKey(ExpressionInput->Expression, ExpressionInput->OutputIndex), MaterialInstance, VisibleExpressions, FunctionStack);
+			FExpressionInput* ExpressionInput = ((UMaterial*)Material)->GetExpressionInputForProperty((EMaterialProperty)i);
+
+			if (ExpressionInput)
+			{
+				GetVisibleMaterialParametersFromExpression(FMaterialExpressionKey(ExpressionInput->Expression, ExpressionInput->OutputIndex), MaterialInstance, VisibleExpressions, FunctionStack);
+			}
 		}
 	}
 
@@ -584,6 +592,13 @@ void FMaterialEditorUtilities::GetVisibleMaterialParametersFromExpression(
 				FExpressionInput* Input = ExpressionInputs[ExpressionInputIndex];
 				GetVisibleMaterialParametersFromExpression(FMaterialExpressionKey(Input->Expression, Input->OutputIndex), MaterialInstance, VisibleExpressions, FunctionStack);
 			}
+
+			TArray<FExpressionExecOutputEntry> ExpressionExecOutputs;
+			MaterialExpressionKey.Expression->GetExecOutputs(ExpressionExecOutputs);
+			for (const FExpressionExecOutputEntry& Entry : ExpressionExecOutputs)
+			{
+				GetVisibleMaterialParametersFromExpression(FMaterialExpressionKey(Entry.Output->Expression, INDEX_NONE), MaterialInstance, VisibleExpressions, FunctionStack);
+			}
 		}
 	}
 
diff --git a/Engine/Source/Runtime/Engine/Classes/Engine/EngineTypes.h b/Engine/Source/Runtime/Engine/Classes/Engine/EngineTypes.h
--- a/Engine/Source/Runtime/Engine/Classes/Engine/EngineTypes.h
+++ b/Engine/Source/Runtime/Engine/Classes/Engine/EngineTypes.h
@@ -262,6 +262,25 @@ enum ESamplerSourceMode
 	SSM_Clamp_WorldGroupSettings UMETA(DisplayName="Shared: Clamp")
 };
 
+/** defines how MipValue is used */
+UENUM()
+enum ETextureMipValueMode
+{
+	/* Use hardware computed sample's mip level with automatic anisotropic filtering support. */
+	TMVM_None UMETA(DisplayName="None (use computed mip level)"),
+
+	/* Explicitly compute the sample's mip level. Disables anisotropic filtering. */
+	TMVM_MipLevel UMETA(DisplayName="MipLevel (absolute, 0 is full resolution)"),
+	
+	/* Bias the hardware computed sample's mip level. Disables anisotropic filtering. */
+	TMVM_MipBias UMETA(DisplayName="MipBias (relative to the computed mip level)"),
+	
+	/* Explicitly compute the sample's DDX and DDY for anisotropic filtering. */
+	TMVM_Derivative UMETA(DisplayName="Derivative (explicit derivative to compute mip level)"),
+
+	TMVM_MAX,
+};
+
 /** Describes how to handle lighting of translucent objets */
 UENUM()
 enum ETranslucencyLightingMode
diff --git a/Engine/Source/Runtime/Engine/Classes/Materials/Material.h b/Engine/Source/Runtime/Engine/Classes/Materials/Material.h
--- a/Engine/Source/Runtime/Engine/Classes/Materials/Material.h
+++ b/Engine/Source/Runtime/Engine/Classes/Materials/Material.h
@@ -850,6 +850,9 @@ public:
 	UPROPERTY(EditAnywhere, Category = Material)
 	uint8 bEnableExecWire : 1;
 
+	UPROPERTY(EditAnywhere, Category = Material)
+	uint8 bEnableNewHLSLGenerator : 1;
+
 	/** when true, the material casts ray tracing shadows. */
 	UPROPERTY(EditAnywhere, Category = Material)
 	uint8 bCastRayTracedShadows : 1;
@@ -1073,6 +1076,8 @@ public:
 
 	ENGINE_API bool IsCompiledWithExecutionFlow() const;
 
+	ENGINE_API bool IsUsingNewHLSLGenerator() const;
+
 #if WITH_EDITOR
 	/** Allows material properties to be compiled with the option of being overridden by the material attributes input. */
 	ENGINE_API virtual int32 CompilePropertyEx( class FMaterialCompiler* Compiler, const FGuid& AttributeID ) override;
diff --git a/Engine/Source/Runtime/Engine/Classes/Materials/MaterialExpression.h b/Engine/Source/Runtime/Engine/Classes/Materials/MaterialExpression.h
--- a/Engine/Source/Runtime/Engine/Classes/Materials/MaterialExpression.h
+++ b/Engine/Source/Runtime/Engine/Classes/Materials/MaterialExpression.h
@@ -18,6 +18,19 @@ class UTexture;
 struct FPropertyChangedEvent;
 struct FMaterialShadingModelField;
 
+class FMaterialHLSLGenerator;
+
+namespace UE
+{
+namespace HLSLTree
+{
+class FScope;
+class FStatement;
+class FExpression;
+class FTextureParameterDeclaration;
+}
+}
+
 //@warning: FExpressionInput is mirrored in MaterialShared.h and manually "subclassed" in Material.h (FMaterialInput)
 #if !CPP      //noexport struct
 USTRUCT(noexport)
@@ -105,7 +118,14 @@ struct FExpressionExecOutput
 	GENERATED_BODY()
 
 #if WITH_EDITOR
-	ENGINE_API int32 Compile(class FMaterialCompiler* Compiler);
+	ENGINE_API int32 Compile(class FMaterialCompiler* Compiler) const;
+
+	/** Returns the statement for the expression connected to this input */ 
+	ENGINE_API UE::HLSLTree::FStatement* AcquireHLSLStatement(FMaterialHLSLGenerator& Generator, UE::HLSLTree::FScope& Scope) const;
+
+	/** Creates a new scope, and populates it with the expression connected to this input */
+	ENGINE_API UE::HLSLTree::FScope* NewScopeWithStatement(FMaterialHLSLGenerator& Generator, UE::HLSLTree::FScope& Scope) const;
+	ENGINE_API UE::HLSLTree::FScope* NewLinkedScopeWithStatement(FMaterialHLSLGenerator& Generator, UE::HLSLTree::FScope& Scope) const;
 #endif // WITH_EDITOR
 
 	UPROPERTY()
@@ -118,6 +138,12 @@ struct FExpressionExecOutputEntry
 	FExpressionExecOutput* Output = nullptr;
 };
 
+enum class EMaterialGenerateHLSLStatus : uint8
+{
+	Success,
+	Error,
+};
+
 UCLASS(abstract, BlueprintType, hidecategories=Object)
 class ENGINE_API UMaterialExpression : public UObject
 {
@@ -253,7 +279,18 @@ class ENGINE_API UMaterialExpression : public UObject
 	 */	
 	virtual int32 Compile(class FMaterialCompiler* Compiler, int32 OutputIndex) { return INDEX_NONE; }
 	virtual int32 CompilePreview(class FMaterialCompiler* Compiler, int32 OutputIndex) { return Compile(Compiler, OutputIndex); }
-#endif
+
+	/**
+	 * A given UMaterial implementation should implement at least one of these methods in order to generate HLSL code
+	 * It's valid to implement more than one, if the expression can be used in multiple ways.
+	 * For example, a for-loop expression might generate a statement for the execution input, but also generate an expression to access the loop index
+	 * These methods replace the Compile() method; once we switch over to the new system, Compile() will be removed
+	 */
+	virtual EMaterialGenerateHLSLStatus GenerateHLSLStatement(FMaterialHLSLGenerator& Generator, UE::HLSLTree::FScope& Scope, UE::HLSLTree::FStatement*& OutStatement);
+	virtual EMaterialGenerateHLSLStatus GenerateHLSLExpression(FMaterialHLSLGenerator& Generator, UE::HLSLTree::FScope& Scope, int32 OutputIndex, UE::HLSLTree::FExpression* &OutExpression);
+	virtual EMaterialGenerateHLSLStatus GenerateHLSLTexture(FMaterialHLSLGenerator& Generator, UE::HLSLTree::FScope& Scope, int32 OutputIndex, UE::HLSLTree::FTextureParameterDeclaration*& OutTexture);
+
+#endif // WITH_EDITOR
 
 	/**
 	* Fill the array with all textures dependence that should trig a recompile of the material.
diff --git a/Engine/Source/Runtime/Engine/Classes/Materials/MaterialExpressionAdd.h b/Engine/Source/Runtime/Engine/Classes/Materials/MaterialExpressionAdd.h
--- a/Engine/Source/Runtime/Engine/Classes/Materials/MaterialExpressionAdd.h
+++ b/Engine/Source/Runtime/Engine/Classes/Materials/MaterialExpressionAdd.h
@@ -34,6 +34,8 @@ class UMaterialExpressionAdd : public UMaterialExpression
 	virtual int32 Compile(class FMaterialCompiler* Compiler, int32 OutputIndex) override;
 	virtual void GetCaption(TArray<FString>& OutCaptions) const override;
 	virtual FText GetKeywords() const override {return FText::FromString(TEXT("+"));}
+
+	virtual EMaterialGenerateHLSLStatus GenerateHLSLExpression(FMaterialHLSLGenerator& Generator, UE::HLSLTree::FScope& Scope, int32 OutputIndex, UE::HLSLTree::FExpression*& OutExpression) override;
 #endif // WITH_EDITOR
 	//~ End UMaterialExpression Interface
 
diff --git a/Engine/Source/Runtime/Engine/Classes/Materials/MaterialExpressionConstant.h b/Engine/Source/Runtime/Engine/Classes/Materials/MaterialExpressionConstant.h
--- a/Engine/Source/Runtime/Engine/Classes/Materials/MaterialExpressionConstant.h
+++ b/Engine/Source/Runtime/Engine/Classes/Materials/MaterialExpressionConstant.h
@@ -22,6 +22,7 @@ class UMaterialExpressionConstant : public UMaterialExpression
 	virtual int32 Compile(class FMaterialCompiler* Compiler, int32 OutputIndex) override;
 	virtual void GetCaption(TArray<FString>& OutCaptions) const override;
 	virtual FString GetDescription() const override;
+	virtual EMaterialGenerateHLSLStatus GenerateHLSLExpression(FMaterialHLSLGenerator& Generator, UE::HLSLTree::FScope& Scope, int32 OutputIndex, UE::HLSLTree::FExpression*& OutExpression) override;
 #endif // WITH_EDITOR
 
 	//~ End UMaterialExpression Interface
diff --git a/Engine/Source/Runtime/Engine/Classes/Materials/MaterialExpressionConstant2Vector.h b/Engine/Source/Runtime/Engine/Classes/Materials/MaterialExpressionConstant2Vector.h
--- a/Engine/Source/Runtime/Engine/Classes/Materials/MaterialExpressionConstant2Vector.h
+++ b/Engine/Source/Runtime/Engine/Classes/Materials/MaterialExpressionConstant2Vector.h
@@ -26,6 +26,7 @@ class UMaterialExpressionConstant2Vector : public UMaterialExpression
 	virtual void GetCaption(TArray<FString>& OutCaptions) const override;
 	virtual FString GetDescription() const override;
 	virtual uint32 GetOutputType(int32 OutputIndex) override {return MCT_Float2;}
+	virtual EMaterialGenerateHLSLStatus GenerateHLSLExpression(FMaterialHLSLGenerator& Generator, UE::HLSLTree::FScope& Scope, int32 OutputIndex, UE::HLSLTree::FExpression*& OutExpression) override;
 #endif // WITH_EDITOR
 	//~ End UMaterialExpression Interface
 };
diff --git a/Engine/Source/Runtime/Engine/Classes/Materials/MaterialExpressionConstant3Vector.h b/Engine/Source/Runtime/Engine/Classes/Materials/MaterialExpressionConstant3Vector.h
--- a/Engine/Source/Runtime/Engine/Classes/Materials/MaterialExpressionConstant3Vector.h
+++ b/Engine/Source/Runtime/Engine/Classes/Materials/MaterialExpressionConstant3Vector.h
@@ -25,5 +25,6 @@ public:
 
 	virtual FString GetDescription() const override;
 	virtual uint32 GetOutputType(int32 OutputIndex) override {return MCT_Float3;}
+	virtual EMaterialGenerateHLSLStatus GenerateHLSLExpression(FMaterialHLSLGenerator& Generator, UE::HLSLTree::FScope& Scope, int32 OutputIndex, UE::HLSLTree::FExpression*& OutExpression) override;
 #endif // WITH_EDITOR
 };
diff --git a/Engine/Source/Runtime/Engine/Classes/Materials/MaterialExpressionConstant4Vector.h b/Engine/Source/Runtime/Engine/Classes/Materials/MaterialExpressionConstant4Vector.h
--- a/Engine/Source/Runtime/Engine/Classes/Materials/MaterialExpressionConstant4Vector.h
+++ b/Engine/Source/Runtime/Engine/Classes/Materials/MaterialExpressionConstant4Vector.h
@@ -22,6 +22,7 @@ class UMaterialExpressionConstant4Vector : public UMaterialExpression
 	virtual void GetCaption(TArray<FString>& OutCaptions) const override;
 	virtual FString GetDescription() const override;
 	virtual uint32 GetOutputType(int32 OutputIndex) override {return MCT_Float4;}
+	virtual EMaterialGenerateHLSLStatus GenerateHLSLExpression(FMaterialHLSLGenerator& Generator, UE::HLSLTree::FScope& Scope, int32 OutputIndex, UE::HLSLTree::FExpression*& OutExpression) override;
 #endif // WITH_EDITOR
 	//~ End UMaterialExpression Interface
 };
diff --git a/Engine/Source/Runtime/Engine/Classes/Materials/MaterialExpressionExecBegin.h b/Engine/Source/Runtime/Engine/Classes/Materials/MaterialExpressionExecBegin.h
--- a/Engine/Source/Runtime/Engine/Classes/Materials/MaterialExpressionExecBegin.h
+++ b/Engine/Source/Runtime/Engine/Classes/Materials/MaterialExpressionExecBegin.h
@@ -22,6 +22,8 @@ class UMaterialExpressionExecBegin : public UMaterialExpression
 	virtual int32 Compile(class FMaterialCompiler* Compiler, int32 OutputIndex) override;
 	virtual void GetCaption(TArray<FString>& OutCaptions) const override;
 	virtual bool CanUserDeleteExpression() const override { return false; }
+
+	virtual EMaterialGenerateHLSLStatus GenerateHLSLStatement(FMaterialHLSLGenerator& Generator, UE::HLSLTree::FScope& Scope, UE::HLSLTree::FStatement*& OutStatement) override;
 #endif
 	//~ End UMaterialExpression Interface
 };
diff --git a/Engine/Source/Runtime/Engine/Classes/Materials/MaterialExpressionForLoop.h b/Engine/Source/Runtime/Engine/Classes/Materials/MaterialExpressionForLoop.h
--- a/Engine/Source/Runtime/Engine/Classes/Materials/MaterialExpressionForLoop.h
+++ b/Engine/Source/Runtime/Engine/Classes/Materials/MaterialExpressionForLoop.h
@@ -36,6 +36,8 @@ class UMaterialExpressionForLoop : public UMaterialExpression
 	virtual uint32 GetInputType(int32 InputIndex) override;
 	virtual uint32 GetOutputType(int32 InputIndex) override;
 	virtual bool HasExecInput() override { return true; }
+
+	virtual EMaterialGenerateHLSLStatus GenerateHLSLStatement(FMaterialHLSLGenerator& Generator, UE::HLSLTree::FScope& Scope, UE::HLSLTree::FStatement*& OutStatement) override;
 #endif
 	//~ End UMaterialExpression Interface
 };
diff --git a/Engine/Source/Runtime/Engine/Classes/Materials/MaterialExpressionGetLocal.h b/Engine/Source/Runtime/Engine/Classes/Materials/MaterialExpressionGetLocal.h
--- a/Engine/Source/Runtime/Engine/Classes/Materials/MaterialExpressionGetLocal.h
+++ b/Engine/Source/Runtime/Engine/Classes/Materials/MaterialExpressionGetLocal.h
@@ -9,7 +9,8 @@
 #include "Materials/MaterialExpression.h"
 #include "MaterialExpressionGetLocal.generated.h"
 
-UCLASS(collapsecategories, hidecategories = Object, MinimalAPI, meta = (MaterialControlFlow))
+UCLASS(collapsecategories, hidecategories = Object, MinimalAPI, meta = (MaterialControlFlow)
+)
 class UMaterialExpressionGetLocal : public UMaterialExpression
 {
 	GENERATED_UCLASS_BODY()
@@ -22,6 +23,7 @@ class UMaterialExpressionGetLocal : public UMaterialExpression
 	virtual int32 Compile(class FMaterialCompiler* Compiler, int32 OutputIndex) override;
 	virtual void GetCaption(TArray<FString>& OutCaptions) const override;
 	virtual uint32 GetOutputType(int32 InputIndex) override;
+	virtual EMaterialGenerateHLSLStatus GenerateHLSLExpression(FMaterialHLSLGenerator& Generator, UE::HLSLTree::FScope& Scope, int32 OutputIndex, UE::HLSLTree::FExpression*& OutExpression) override;
 #endif
 	//~ End UMaterialExpression Interface
 };
diff --git a/Engine/Source/Runtime/Engine/Classes/Materials/MaterialExpressionIfThenElse.h b/Engine/Source/Runtime/Engine/Classes/Materials/MaterialExpressionIfThenElse.h
--- a/Engine/Source/Runtime/Engine/Classes/Materials/MaterialExpressionIfThenElse.h
+++ b/Engine/Source/Runtime/Engine/Classes/Materials/MaterialExpressionIfThenElse.h
@@ -29,6 +29,8 @@ class UMaterialExpressionIfThenElse : public UMaterialExpression
 	virtual void GetCaption(TArray<FString>& OutCaptions) const override;
 	virtual uint32 GetInputType(int32 InputIndex) override;
 	virtual bool HasExecInput() override { return true; }
+
+	virtual EMaterialGenerateHLSLStatus GenerateHLSLStatement(FMaterialHLSLGenerator& Generator, UE::HLSLTree::FScope& Scope, UE::HLSLTree::FStatement*& OutStatement) override;
 #endif
 	//~ End UMaterialExpression Interface
 };
diff --git a/Engine/Source/Runtime/Engine/Classes/Materials/MaterialExpressionReturnMaterialAttributes.h b/Engine/Source/Runtime/Engine/Classes/Materials/MaterialExpressionReturnMaterialAttributes.h
--- a/Engine/Source/Runtime/Engine/Classes/Materials/MaterialExpressionReturnMaterialAttributes.h
+++ b/Engine/Source/Runtime/Engine/Classes/Materials/MaterialExpressionReturnMaterialAttributes.h
@@ -25,6 +25,8 @@ class UMaterialExpressionReturnMaterialAttributes : public UMaterialExpression
 	virtual FExpressionInput* GetInput(int32 InputIndex) override;
 	virtual uint32 GetInputType(int32 InputIndex) override;
 	virtual bool HasExecInput() override { return true; }
+
+	virtual EMaterialGenerateHLSLStatus GenerateHLSLStatement(FMaterialHLSLGenerator& Generator, UE::HLSLTree::FScope& Scope, UE::HLSLTree::FStatement*& OutStatement) override;
 #endif
 	//~ End UMaterialExpression Interface
 };
diff --git a/Engine/Source/Runtime/Engine/Classes/Materials/MaterialExpressionScalarParameter.h b/Engine/Source/Runtime/Engine/Classes/Materials/MaterialExpressionScalarParameter.h
--- a/Engine/Source/Runtime/Engine/Classes/Materials/MaterialExpressionScalarParameter.h
+++ b/Engine/Source/Runtime/Engine/Classes/Materials/MaterialExpressionScalarParameter.h
@@ -43,6 +43,7 @@ class UMaterialExpressionScalarParameter : public UMaterialExpressionParameter
 #if WITH_EDITOR
 	virtual int32 Compile(class FMaterialCompiler* Compiler, int32 OutputIndex) override;
 	virtual void GetCaption(TArray<FString>& OutCaptions) const override;
+	EMaterialGenerateHLSLStatus GenerateHLSLExpression(FMaterialHLSLGenerator& Generator, UE::HLSLTree::FScope& Scope, int32 OutputIndex, UE::HLSLTree::FExpression*& OutExpression) override;
 #endif
 	//~ End UMaterialExpression Interface
 
diff --git a/Engine/Source/Runtime/Engine/Classes/Materials/MaterialExpressionSetLocal.h b/Engine/Source/Runtime/Engine/Classes/Materials/MaterialExpressionSetLocal.h
--- a/Engine/Source/Runtime/Engine/Classes/Materials/MaterialExpressionSetLocal.h
+++ b/Engine/Source/Runtime/Engine/Classes/Materials/MaterialExpressionSetLocal.h
@@ -29,6 +29,7 @@ class UMaterialExpressionSetLocal : public UMaterialExpression
 	virtual void GetCaption(TArray<FString>& OutCaptions) const override;
 	virtual uint32 GetInputType(int32 InputIndex) override;
 	virtual bool HasExecInput() override { return true; }
+	virtual EMaterialGenerateHLSLStatus GenerateHLSLStatement(FMaterialHLSLGenerator& Generator, UE::HLSLTree::FScope& Scope, UE::HLSLTree::FStatement*& OutStatement) override;
 #endif
 	//~ End UMaterialExpression Interface
 };
diff --git a/Engine/Source/Runtime/Engine/Classes/Materials/MaterialExpressionSetMaterialAttributes.h b/Engine/Source/Runtime/Engine/Classes/Materials/MaterialExpressionSetMaterialAttributes.h
--- a/Engine/Source/Runtime/Engine/Classes/Materials/MaterialExpressionSetMaterialAttributes.h
+++ b/Engine/Source/Runtime/Engine/Classes/Materials/MaterialExpressionSetMaterialAttributes.h
@@ -43,6 +43,8 @@ class UMaterialExpressionSetMaterialAttributes : public UMaterialExpression
 	virtual bool IsInputConnectionRequired(int32 InputIndex) const override {return true;}
 	virtual bool IsResultMaterialAttributes(int32 OutputIndex) override {return true;}
 	virtual void GetExpressionToolTip(TArray<FString>& OutToolTip) override;
+
+	virtual EMaterialGenerateHLSLStatus GenerateHLSLExpression(FMaterialHLSLGenerator& Generator, UE::HLSLTree::FScope& Scope, int32 OutputIndex, UE::HLSLTree::FExpression*& OutExpression) override;
 #endif
 	//~ End UMaterialExpression Interface
 };
diff --git a/Engine/Source/Runtime/Engine/Classes/Materials/MaterialExpressionTextureBase.h b/Engine/Source/Runtime/Engine/Classes/Materials/MaterialExpressionTextureBase.h
--- a/Engine/Source/Runtime/Engine/Classes/Materials/MaterialExpressionTextureBase.h
+++ b/Engine/Source/Runtime/Engine/Classes/Materials/MaterialExpressionTextureBase.h
@@ -66,5 +66,19 @@ class ENGINE_API UMaterialExpressionTextureBase : public UMaterialExpression
 	 * @returns the default sampler type for the specified texture.
 	 */
 	static EMaterialSamplerType GetSamplerTypeForTexture( const UTexture* Texture, bool ForceNoVT = false );
+
+	/**
+	 * Verify that the texture and sampler type. Generates a compiler waring if
+	 * they do not.
+	 * @param Texture - The texture to verify. A nullptr texture is considered valid!
+	 * @param SamplerType - The sampler type to verify.
+	 * @param OutErrorMessage - If 'false' is returned, will contain a message describing the error
+	 */
+	static bool VerifySamplerType(
+		ERHIFeatureLevel::Type FeatureLevel,
+		const ITargetPlatform* TargetPlatform,
+		const UTexture* Texture,
+		EMaterialSamplerType SamplerType,
+		FString& OutErrorMessage);
 #endif // WITH_EDITOR
 };
diff --git a/Engine/Source/Runtime/Engine/Classes/Materials/MaterialExpressionTextureObject.h b/Engine/Source/Runtime/Engine/Classes/Materials/MaterialExpressionTextureObject.h
--- a/Engine/Source/Runtime/Engine/Classes/Materials/MaterialExpressionTextureObject.h
+++ b/Engine/Source/Runtime/Engine/Classes/Materials/MaterialExpressionTextureObject.h
@@ -28,6 +28,7 @@ class UMaterialExpressionTextureObject : public UMaterialExpressionTextureBase
 	virtual int32 CompilePreview(class FMaterialCompiler* Compiler, int32 OutputIndex) override;
 	virtual void GetCaption(TArray<FString>& OutCaptions) const override;
 	virtual uint32 GetOutputType(int32 OutputIndex) override;
+	virtual EMaterialGenerateHLSLStatus GenerateHLSLTexture(FMaterialHLSLGenerator& Generator, UE::HLSLTree::FScope& Scope, int32 OutputIndex, UE::HLSLTree::FTextureParameterDeclaration*& OutTexture) override;
 	//~ End UMaterialExpression Interface
 #endif // WITH_EDITOR
 };
diff --git a/Engine/Source/Runtime/Engine/Classes/Materials/MaterialExpressionTextureObjectParameter.h b/Engine/Source/Runtime/Engine/Classes/Materials/MaterialExpressionTextureObjectParameter.h
--- a/Engine/Source/Runtime/Engine/Classes/Materials/MaterialExpressionTextureObjectParameter.h
+++ b/Engine/Source/Runtime/Engine/Classes/Materials/MaterialExpressionTextureObjectParameter.h
@@ -30,6 +30,7 @@ class UMaterialExpressionTextureObjectParameter : public UMaterialExpressionText
 	virtual int32 CompilePreview(class FMaterialCompiler* Compiler, int32 OutputIndex) override;
 	virtual uint32 GetOutputType(int32 InputIndex) override {return MCT_Texture;}
 	virtual const TArray<FExpressionInput*> GetInputs() override;
+	virtual EMaterialGenerateHLSLStatus GenerateHLSLTexture(FMaterialHLSLGenerator& Generator, UE::HLSLTree::FScope& Scope, int32 OutputIndex, UE::HLSLTree::FTextureParameterDeclaration*& OutTexture) override;
 #endif
 	//~ End UMaterialExpression Interface
 };
diff --git a/Engine/Source/Runtime/Engine/Classes/Materials/MaterialExpressionTextureSample.h b/Engine/Source/Runtime/Engine/Classes/Materials/MaterialExpressionTextureSample.h
--- a/Engine/Source/Runtime/Engine/Classes/Materials/MaterialExpressionTextureSample.h
+++ b/Engine/Source/Runtime/Engine/Classes/Materials/MaterialExpressionTextureSample.h
@@ -11,25 +11,6 @@
 
 struct FPropertyChangedEvent;
 
-/** defines how MipValue is used */
-UENUM()
-enum ETextureMipValueMode
-{
-	/* Use hardware computed sample's mip level with automatic anisotropic filtering support. */
-	TMVM_None UMETA(DisplayName="None (use computed mip level)"),
-
-	/* Explicitly compute the sample's mip level. Disables anisotropic filtering. */
-	TMVM_MipLevel UMETA(DisplayName="MipLevel (absolute, 0 is full resolution)"),
-	
-	/* Bias the hardware computed sample's mip level. Disables anisotropic filtering. */
-	TMVM_MipBias UMETA(DisplayName="MipBias (relative to the computed mip level)"),
-	
-	/* Explicitly compute the sample's DDX and DDY for anisotropic filtering. */
-	TMVM_Derivative UMETA(DisplayName="Derivative (explicit derivative to compute mip level)"),
-
-	TMVM_MAX,
-};
-
 UCLASS(collapsecategories, hidecategories=Object)
 class ENGINE_API UMaterialExpressionTextureSample : public UMaterialExpressionTextureBase
 {
@@ -82,6 +63,10 @@ protected:
 	// Inherited parameter expressions can hide unused input pin
 	uint8 bShowTextureInputPin : 1;
 
+#if WITH_EDITOR
+	EMaterialGenerateHLSLStatus GenerateHLSLExpressionBase(FMaterialHLSLGenerator& Generator, UE::HLSLTree::FScope& Scope, UE::HLSLTree::FTextureParameterDeclaration* TextureDeclaration, UE::HLSLTree::FExpression*& OutExpression);
+#endif
+
 public:
 
 	/** only used if Coordinates is not hooked up */
@@ -114,6 +99,7 @@ public:
 	virtual bool MatchesSearchQuery( const TCHAR* SearchQuery ) override;
 	virtual uint32 GetInputType(int32 InputIndex) override;
 	virtual bool CanIgnoreOutputIndex() { return true; }
+	virtual EMaterialGenerateHLSLStatus GenerateHLSLExpression(FMaterialHLSLGenerator& Generator, UE::HLSLTree::FScope& Scope, int32 OutputIndex, UE::HLSLTree::FExpression*& OutExpression) override;
 #endif // WITH_EDITOR
 	//~ End UMaterialExpression Interface
 
diff --git a/Engine/Source/Runtime/Engine/Classes/Materials/MaterialExpressionTextureSampleParameter.h b/Engine/Source/Runtime/Engine/Classes/Materials/MaterialExpressionTextureSampleParameter.h
--- a/Engine/Source/Runtime/Engine/Classes/Materials/MaterialExpressionTextureSampleParameter.h
+++ b/Engine/Source/Runtime/Engine/Classes/Materials/MaterialExpressionTextureSampleParameter.h
@@ -51,6 +51,7 @@ class ENGINE_API UMaterialExpressionTextureSampleParameter : public UMaterialExp
 	virtual void SetParameterName(const FName& Name) override { ParameterName = Name; }
 	virtual void ValidateParameterName(const bool bAllowDuplicateName) override;
 	virtual void SetValueToMatchingExpression(UMaterialExpression* OtherExpression) override;
+	virtual EMaterialGenerateHLSLStatus GenerateHLSLExpression(FMaterialHLSLGenerator& Generator, UE::HLSLTree::FScope& Scope, int32 OutputIndex, UE::HLSLTree::FExpression*& OutExpression) override;
 #endif
 	//~ End UMaterialExpression Interface
 
diff --git a/Engine/Source/Runtime/Engine/Classes/Materials/MaterialExpressionVectorParameter.h b/Engine/Source/Runtime/Engine/Classes/Materials/MaterialExpressionVectorParameter.h
--- a/Engine/Source/Runtime/Engine/Classes/Materials/MaterialExpressionVectorParameter.h
+++ b/Engine/Source/Runtime/Engine/Classes/Materials/MaterialExpressionVectorParameter.h
@@ -32,6 +32,7 @@ class UMaterialExpressionVectorParameter : public UMaterialExpressionParameter
 #if WITH_EDITOR
 	virtual int32 Compile(class FMaterialCompiler* Compiler, int32 OutputIndex) override;
 	virtual void GetCaption(TArray<FString>& OutCaptions) const override;
+	EMaterialGenerateHLSLStatus GenerateHLSLExpression(FMaterialHLSLGenerator& Generator, UE::HLSLTree::FScope& Scope, int32 OutputIndex, UE::HLSLTree::FExpression*& OutExpression) override;
 #endif
 	//~ End UMaterialExpression Interface
 
diff --git a/Engine/Source/Runtime/Engine/Private/HLSLTree/HLSLTree.cpp b/Engine/Source/Runtime/Engine/Private/HLSLTree/HLSLTree.cpp
new file mode 100644
--- /dev/null
+++ b/Engine/Source/Runtime/Engine/Private/HLSLTree/HLSLTree.cpp
@@ -0,0 +1,609 @@
+// Copyright Epic Games, Inc. All Rights Reserved.
+#include "HLSLTree/HLSLTree.h"
+#include "Misc/StringBuilder.h"
+#include "Misc/MemStack.h"
+#include "MaterialShared.h" // TODO - split preshader out into its own module
+
+class FLocalHLSLCodeWriter : public UE::HLSLTree::FCodeWriter
+{
+public:
+	FLocalHLSLCodeWriter() : UE::HLSLTree::FCodeWriter(&LocalStringBuilder) {}
+
+	TStringBuilder<2048> LocalStringBuilder;
+};
+
+UE::HLSLTree::FExpressionTypeDescription UE::HLSLTree::GetExpressionTypeDescription(EExpressionType Type)
+{
+	switch (Type)
+	{
+	case EExpressionType::Void: return FExpressionTypeDescription(TEXT("void"), EExpressionComponentType::Void, 0);
+	case EExpressionType::Float1: return FExpressionTypeDescription(TEXT("float"), EExpressionComponentType::Float, 1);
+	case EExpressionType::Float2: return FExpressionTypeDescription(TEXT("float2"), EExpressionComponentType::Float, 2);
+	case EExpressionType::Float3: return FExpressionTypeDescription(TEXT("float3"), EExpressionComponentType::Float, 3);
+	case EExpressionType::Float4: return FExpressionTypeDescription(TEXT("float4"), EExpressionComponentType::Float, 4);
+	case EExpressionType::MaterialAttributes: return FExpressionTypeDescription(TEXT("FMaterialAttributes"), EExpressionComponentType::MaterialAttributes, 0);
+	default: checkNoEntry(); return FExpressionTypeDescription();
+	}
+}
+
+UE::HLSLTree::EExpressionType UE::HLSLTree::MakeExpressionType(EExpressionComponentType ComponentType, int32 NumComponents)
+{
+	switch (ComponentType)
+	{
+	case EExpressionComponentType::Void:
+		check(NumComponents == 0);
+		return EExpressionType::Void;
+	case EExpressionComponentType::MaterialAttributes:
+		check(NumComponents == 0);
+		return EExpressionType::MaterialAttributes;
+	case EExpressionComponentType::Float:
+		switch (NumComponents)
+		{
+		case 1: return EExpressionType::Float1;
+		case 2: return EExpressionType::Float2;
+		case 3: return EExpressionType::Float3;
+		case 4: return EExpressionType::Float4;
+		default: break;
+		}
+	default:
+		break;
+	}
+
+	checkNoEntry();
+	return EExpressionType::Void;
+}
+
+UE::HLSLTree::EExpressionType UE::HLSLTree::MakeExpressionType(EExpressionType BaseType, int32 NumComponents)
+{
+	return MakeExpressionType(GetExpressionTypeDescription(BaseType).ComponentType, NumComponents);
+}
+
+UE::HLSLTree::EExpressionType UE::HLSLTree::MakeArithmeticResultType(EExpressionType Lhs, EExpressionType Rhs, FString& OutErrorMessage)
+{
+	const FExpressionTypeDescription LhsDesc = GetExpressionTypeDescription(Lhs);
+	const FExpressionTypeDescription RhsDesc = GetExpressionTypeDescription(Rhs);
+	// Types with 0 components are non-arithmetic
+	if (LhsDesc.NumComponents > 0 && RhsDesc.NumComponents > 0)
+	{
+		if (Lhs == Rhs)
+		{
+			return Lhs;
+		}
+		if (LhsDesc.ComponentType == RhsDesc.ComponentType)
+		{
+			if (LhsDesc.NumComponents == 1 || RhsDesc.NumComponents == 1)
+			{
+				// single component type is valid to combine with other type
+				return MakeExpressionType(LhsDesc.ComponentType, FMath::Max(LhsDesc.NumComponents, RhsDesc.NumComponents));
+			}
+		}
+		OutErrorMessage = FString::Printf(TEXT("Arithmetic between types %s and %s are undefined"), LhsDesc.Name, RhsDesc.Name);
+	}
+	else
+	{
+		OutErrorMessage = FString::Printf(TEXT("Attempting to perform arithmetic on non-numeric types: %s %s"), LhsDesc.Name, RhsDesc.Name);
+	}
+
+	return EExpressionType::Void;
+}
+
+UE::HLSLTree::EExpressionEvaluationType UE::HLSLTree::CombineEvaluationTypes(UE::HLSLTree::EExpressionEvaluationType Lhs, UE::HLSLTree::EExpressionEvaluationType Rhs)
+{
+	if (Lhs == EExpressionEvaluationType::Constant && Rhs == EExpressionEvaluationType::Constant)
+	{
+		// 2 constants make a constant
+		return EExpressionEvaluationType::Constant;
+	}
+	else if (Lhs == EExpressionEvaluationType::Shader || Rhs == EExpressionEvaluationType::Shader)
+	{
+		// If either requires shader, shader is required
+		return EExpressionEvaluationType::Shader;
+	}
+	// Any combination of constants/preshader can make a preshader
+	return EExpressionEvaluationType::Preshader;
+}
+
+template <typename FormatType, typename... ArgTypes>
+const TCHAR* AllocateStringf(FMemStackBase& Allocator, const FormatType& Format, ArgTypes... Args)
+{
+	TCHAR Buffer[1024];
+	const int32 Length = FCString::Snprintf(Buffer, 1024, Format, Forward<ArgTypes>(Args)...);
+	check(Length > 0);
+	TCHAR* Result = New<TCHAR>(Allocator, Length + 1);
+	FMemory::Memcpy(Result, Buffer, Length * sizeof(TCHAR));
+	Result[Length] = 0;
+	return Result;
+}
+
+const TCHAR* AllocateString(FMemStackBase& Allocator, const FStringBuilderBase& StringBuilder)
+{
+	const int32 Length = StringBuilder.Len();
+	TCHAR* Result = New<TCHAR>(Allocator, Length + 1);
+	FMemory::Memcpy(Result, StringBuilder.GetData(), Length * sizeof(TCHAR));
+	Result[Length] = 0;
+	return Result;
+}
+
+UE::HLSLTree::FCodeWriter* UE::HLSLTree::FCodeWriter::Create(FMemStackBase& Allocator)
+{
+	static const int32 InitialBufferSize = 4 * 1024;
+	TCHAR* Buffer = New<TCHAR>(Allocator, InitialBufferSize);
+	FStringBuilderBase* LocalStringBuilder = new(Allocator) TStringBuilderBase<TCHAR>(Buffer, InitialBufferSize);
+	return new(Allocator) FCodeWriter(LocalStringBuilder);
+}
+
+void UE::HLSLTree::FCodeWriter::IncreaseIndent()
+{
+	++IndentLevel;
+}
+
+void UE::HLSLTree::FCodeWriter::DecreaseIndent()
+{
+	check(IndentLevel > 0);
+	--IndentLevel;
+}
+
+void UE::HLSLTree::FCodeWriter::WriteIndent()
+{
+	for (int32 i = 0; i < IndentLevel; ++i)
+	{
+		StringBuilder->Append(TCHAR(' '));
+	}
+}
+
+FLinearColor UE::HLSLTree::FConstant::ToLinearColor() const
+{
+	switch (Type)
+	{
+	case EExpressionType::Float1: return FLinearColor(Float[0], 0.0f, 0.0f, 0.0f); break;
+	case EExpressionType::Float2: return FLinearColor(Float[0], Float[1], 0.0f, 0.0f); break;
+	case EExpressionType::Float3: return FLinearColor(Float[0], Float[1], Float[2], 0.0f); break;
+	case EExpressionType::Float4: return FLinearColor(Float[0], Float[1], Float[2], Float[3]); break;
+	default: checkNoEntry(); return FLinearColor(ForceInitToZero);
+	}
+}
+
+void UE::HLSLTree::FConstant::EmitHLSL(UE::HLSLTree::FCodeWriter& Writer) const
+{
+	switch (Type)
+	{
+	case EExpressionType::Float1:
+		Writer.Writef(TEXT("float(%0.8f)"), Float[0]);
+		break;
+	case EExpressionType::Float2:
+		Writer.Writef(TEXT("float2(%0.8f, %0.8f)"), Float[0], Float[1]);
+		break;
+	case EExpressionType::Float3:
+		Writer.Writef(TEXT("float3(%0.8f, %0.8f, %0.8f)"), Float[0], Float[1], Float[2]);
+		break;
+	case EExpressionType::Float4:
+		Writer.Writef(TEXT("float4(%0.8f, %0.8f, %0.8f, %0.8f)"), Float[0], Float[1], Float[2], Float[3]);
+		break;
+	default:
+		checkNoEntry();
+		break;
+	}
+}
+
+UE::HLSLTree::FEmitContext::FScopeEntry* UE::HLSLTree::FEmitContext::FindScope(UE::HLSLTree::FScope* Scope)
+{
+	for (int32 Index = ScopeStack.Num() - 1; Index >= 0; --Index)
+	{
+		FScopeEntry& Entry = ScopeStack[Index];
+		if (Entry.Scope == Scope)
+		{
+			return &Entry;
+		}
+	}
+	return nullptr;
+}
+
+const TCHAR* UE::HLSLTree::FEmitContext::AcquireHLSLReference(UE::HLSLTree::FLocalDeclaration* Declaration)
+{
+	check(Declaration);
+	check(Declaration->Type != EExpressionType::Void);
+
+	FDeclarationEntry& Entry = DeclarationMap.FindOrAdd(Declaration);
+	if (!Entry.Definition)
+	{
+		FScopeEntry* ScopeEntry = FindScope(Declaration->ParentScope);
+		check(ScopeEntry);
+
+		Entry.Definition = AllocateStringf(*Allocator, TEXT("Local%d"), NumExpressionLocals++);
+		const FExpressionTypeDescription& TypeDesc = GetExpressionTypeDescription(Declaration->Type);
+		ScopeEntry->ExpressionCodeWriter->WriteLinef(TEXT("%s %s = (%s)0.0f;"),
+			TypeDesc.Name,
+			Entry.Definition,
+			TypeDesc.Name);
+	}
+
+	return Entry.Definition;
+}
+
+const TCHAR* UE::HLSLTree::FEmitContext::AcquireHLSLReference(UE::HLSLTree::FExpression* Expression)
+{
+	check(Expression);
+	check(Expression->Type != EExpressionType::Void);
+
+	FDeclarationEntry& Entry = DeclarationMap.FindOrAdd(Expression);
+	if (!Entry.Definition)
+	{
+		bool bInline = false;
+		FLocalHLSLCodeWriter LocalWriter;
+		if (Expression->EvaluationType == EExpressionEvaluationType::Shader)
+		{
+			bInline = Expression->IsInline();
+			Expression->EmitHLSL(*this, LocalWriter);
+		}
+		else
+		{
+			if (Expression->EvaluationType == EExpressionEvaluationType::Constant)
+			{
+				// Evaluate the constant preshader, and write out its value
+				FMaterialPreshaderData Preshader;
+				Expression->EmitPreshader(*this, Preshader);
+
+				FMaterialRenderContext RenderContext(nullptr, *Material, nullptr);
+				FLinearColor ConstantValue;
+				Preshader.Evaluate(nullptr, RenderContext, ConstantValue);
+
+				const FConstant Constant(Expression->Type, ConstantValue);
+				Constant.EmitHLSL(LocalWriter);
+			}
+			else
+			{
+				// Non-constant preshader, store it on the material
+				check(Expression->EvaluationType == EExpressionEvaluationType::Preshader);
+				FUniformExpressionSet& UniformExpressionSet = MaterialCompilationOutput->UniformExpressionSet;
+
+				FMaterialUniformPreshaderHeader* PreshaderHeader = nullptr;
+				if (Expression->Type == EExpressionType::Float1)
+				{
+					const static TCHAR IndexToMask[] = { 'x', 'y', 'z', 'w' };
+					const int32 ScalarInputIndex = UniformExpressionSet.UniformScalarPreshaders.Num();
+					LocalWriter.Writef(TEXT("Material.ScalarExpressions[%u].%c"), ScalarInputIndex / 4, IndexToMask[ScalarInputIndex % 4]);
+
+					PreshaderHeader = &UniformExpressionSet.UniformScalarPreshaders.AddDefaulted_GetRef();
+				}
+				else
+				{
+					const TCHAR* Mask = TEXT("");
+					switch (Expression->Type)
+					{
+					case EExpressionType::Float1: Mask = TEXT(".r"); break;
+					case EExpressionType::Float2: Mask = TEXT(".rg"); break;
+					case EExpressionType::Float3: Mask = TEXT(".rgb"); break;
+					case EExpressionType::Float4: break;
+					default: checkNoEntry(); break;
+					};
+					const int32 VectorInputIndex = UniformExpressionSet.UniformVectorPreshaders.Num();
+					LocalWriter.Writef(TEXT("Material.VectorExpressions[%u]%s"), VectorInputIndex, Mask);
+
+					PreshaderHeader = &UniformExpressionSet.UniformVectorPreshaders.AddDefaulted_GetRef();
+				}
+				PreshaderHeader->OpcodeOffset = UniformExpressionSet.UniformPreshaderData.Num();
+				Expression->EmitPreshader(*this, UniformExpressionSet.UniformPreshaderData);
+				PreshaderHeader->OpcodeSize = UniformExpressionSet.UniformPreshaderData.Num() - PreshaderHeader->OpcodeOffset;
+			}
+			bInline = true;
+		}
+
+		if (bInline)
+		{
+			Entry.Definition = AllocateString(*Allocator, LocalWriter.GetStringBuilder());
+		}
+		else
+		{
+			FScopeEntry* ScopeEntry = FindScope(Expression->ParentScope);
+			check(ScopeEntry);
+
+			const FExpressionTypeDescription& TypeDesc = GetExpressionTypeDescription(Expression->Type);
+			Entry.Definition = AllocateStringf(*Allocator, TEXT("Local%d"), NumExpressionLocals++);
+			ScopeEntry->ExpressionCodeWriter->WriteLinef(TEXT("%s %s = %s;"),
+				TypeDesc.Name,
+				Entry.Definition,
+				LocalWriter.GetStringBuilder().ToString());
+		}
+	}
+
+	return Entry.Definition;
+}
+
+void UE::HLSLTree::FNodeVisitor::VisitNode(UE::HLSLTree::FNode* Node)
+{
+	if (Node)
+	{
+		Node->Visit(*this);
+	}
+}
+
+UE::HLSLTree::ENodeVisitResult UE::HLSLTree::FStatement::Visit(UE::HLSLTree::FNodeVisitor& Visitor)
+{
+	return Visitor.OnStatement(*this);
+}
+
+UE::HLSLTree::ENodeVisitResult UE::HLSLTree::FExpression::Visit(UE::HLSLTree::FNodeVisitor& Visitor)
+{
+	return Visitor.OnExpression(*this);
+}
+
+void UE::HLSLTree::FExpression::EmitHLSL(UE::HLSLTree::FEmitContext& Context, UE::HLSLTree::FCodeWriter& Writer) const
+{
+	check(false);
+}
+
+void UE::HLSLTree::FExpression::EmitPreshader(UE::HLSLTree::FEmitContext& Context, FMaterialPreshaderData& OutPreshader) const
+{
+	check(false);
+}
+
+UE::HLSLTree::ENodeVisitResult UE::HLSLTree::FLocalDeclaration::Visit(UE::HLSLTree::FNodeVisitor& Visitor)
+{
+	return Visitor.OnLocalDeclaration(*this);
+}
+
+UE::HLSLTree::ENodeVisitResult UE::HLSLTree::FParameterDeclaration::Visit(UE::HLSLTree::FNodeVisitor& Visitor)
+{
+	return Visitor.OnParameterDeclaration(*this);
+}
+
+UE::HLSLTree::ENodeVisitResult UE::HLSLTree::FTextureParameterDeclaration::Visit(UE::HLSLTree::FNodeVisitor& Visitor)
+{
+	return Visitor.OnTextureParameterDeclaration(*this);
+}
+
+UE::HLSLTree::ENodeVisitResult UE::HLSLTree::FScope::Visit(UE::HLSLTree::FNodeVisitor& Visitor)
+{
+	const ENodeVisitResult Result = Visitor.OnScope(*this);
+	if (ShouldVisitDependentNodes(Result))
+	{
+		FStatement* Statement = FirstStatement;
+		while (Statement)
+		{
+			Visitor.VisitNode(Statement);
+			Statement = Statement->NextStatement;
+		}
+	}
+	return Result;
+}
+
+void UE::HLSLTree::FScope::EmitHLSL(UE::HLSLTree::FEmitContext& Context, UE::HLSLTree::FCodeWriter& Writer) const
+{
+	const int32 ScopeIndentLevel = Writer.IndentLevel + 1;
+
+	FLocalHLSLCodeWriter DeclarationCodeWriter;
+	DeclarationCodeWriter.IndentLevel = ScopeIndentLevel;
+
+	FEmitContext::FScopeEntry& ScopeEntry = Context.ScopeStack.AddDefaulted_GetRef();
+	ScopeEntry.Scope = this;
+	ScopeEntry.ExpressionCodeWriter = &DeclarationCodeWriter;
+
+	Writer.WriteLine(TEXT("{"));
+
+	FLocalHLSLCodeWriter StatementCodeWriter;
+	StatementCodeWriter.IndentLevel = ScopeIndentLevel;
+	const FStatement* Statement = FirstStatement;
+	while (Statement)
+	{
+		Statement->EmitHLSL(Context, StatementCodeWriter);
+
+		// First write any expressions needed by the statement, then the statement itself
+		if (DeclarationCodeWriter.StringBuilder->Len() > 0)
+		{
+			Writer.StringBuilder->Append(DeclarationCodeWriter.StringBuilder->GetData(), DeclarationCodeWriter.StringBuilder->Len());
+		}
+		Writer.StringBuilder->Append(StatementCodeWriter.StringBuilder->GetData(), StatementCodeWriter.StringBuilder->Len());
+
+		DeclarationCodeWriter.StringBuilder->Reset();
+		StatementCodeWriter.StringBuilder->Reset();
+
+		Statement = Statement->NextStatement;
+	}
+
+	Writer.WriteLine(TEXT("}"));
+
+	Context.ScopeStack.Pop(false);
+}
+
+void UE::HLSLTree::FScope::AddDeclaration(UE::HLSLTree::FLocalDeclaration* Declaration)
+{
+	check(!Declaration->ParentScope);
+	Declaration->ParentScope = this;
+}
+
+void UE::HLSLTree::FScope::AddExpression(UE::HLSLTree::FExpression* Expression)
+{
+	check(!Expression->ParentScope);
+	Expression->ParentScope = this;
+}
+
+void UE::HLSLTree::FScope::UseNode(UE::HLSLTree::FNode* Node)
+{
+	FScope* Scope0 = this;
+	FScope* Scope1 = Node->ParentScope;
+	if (Scope1)
+	{
+		while (Scope0 != Scope1)
+		{
+			if (Scope0->NestedLevel > Scope1->NestedLevel)
+			{
+				check(Scope0->ParentScope);
+				Scope0 = Scope0->ParentScope;
+			}
+			else
+			{
+				check(Scope1->ParentScope);
+				Scope1 = Scope1->ParentScope;
+			}
+		}
+	}
+	Node->ParentScope = Scope0;
+}
+
+void UE::HLSLTree::FScope::UseDeclaration(UE::HLSLTree::FLocalDeclaration* Declaration)
+{
+	UseNode(Declaration);
+}
+
+namespace UE
+{
+namespace HLSLTree
+{
+class FNodeVisitor_MoveToScope : public FNodeVisitor
+{
+public:
+	explicit FNodeVisitor_MoveToScope(FScope* InScope) : Scope(InScope) {}
+
+	virtual ENodeVisitResult OnScope(FScope& InScope) override
+	{
+		InScope.ParentScope = Scope;
+		// Don't want to recurse into any child scopes
+		return ENodeVisitResult::SkipDependentNodes;
+	}
+
+	virtual ENodeVisitResult OnExpression(FExpression& InExpression) override
+	{
+		Scope->UseNode(&InExpression);
+		return ENodeVisitResult::VisitDependentNodes;
+	}
+
+	virtual ENodeVisitResult OnLocalDeclaration(FLocalDeclaration& InDeclaration) override
+	{
+		Scope->UseNode(&InDeclaration);
+		return ENodeVisitResult::VisitDependentNodes;
+	}
+
+	FScope* Scope;
+};
+} // namespace HLSLTree
+} // namespace UE
+
+void UE::HLSLTree::FScope::UseExpression(UE::HLSLTree::FExpression* Expression)
+{
+	// Need to move all dependent expressions/declarations into this scope
+	FNodeVisitor_MoveToScope Visitor(this);
+	Visitor.VisitNode(Expression);
+}
+
+void UE::HLSLTree::FScope::AddStatement(UE::HLSLTree::FStatement* Statement)
+{
+	check(!Statement->ParentScope);
+	check(!Statement->NextStatement);
+
+	Statement->ParentScope = this;
+	if (!FirstStatement)
+	{
+		check(!LastStatement);
+		FirstStatement = Statement;
+		LastStatement = Statement;
+	}
+	else
+	{
+		check(LastStatement);
+		LastStatement->NextStatement = Statement;
+	}
+}
+
+bool UE::HLSLTree::FScope::TryMoveStatement(UE::HLSLTree::FStatement* Statement)
+{
+	bool bResult = false;
+	FScope* PrevScope = Statement->ParentScope;
+	if (!PrevScope)
+	{
+		// Statement does not currently belong to any scope, just add it
+		AddStatement(Statement);
+		bResult = true;
+	}
+	else if (LinkedScope && ParentScope == PrevScope)
+	{
+		// Already moved
+		bResult = true;
+	}
+	else if(LinkedScope == PrevScope)
+	{
+		// Statement belongs to our linked scope, so move it into our parent scope
+		check(ParentScope->FirstStatement);
+		check(ParentScope->LastStatement);
+
+		// First need to unlink the statement from its previous scope
+		FStatement* PrevStatement = PrevScope->FirstStatement;
+		bool bFoundInPrevScope = false;
+		while (PrevStatement)
+		{
+			FStatement* NextStatement = PrevStatement->NextStatement;
+			if (NextStatement == Statement)
+			{
+				PrevScope->LastStatement = PrevStatement;
+				PrevStatement->NextStatement = nullptr;
+				bFoundInPrevScope = true;
+				break;
+			}
+			PrevStatement = NextStatement;
+		}
+		check(bFoundInPrevScope);
+
+		// Now link it into the new scope
+		FNodeVisitor_MoveToScope Visitor(ParentScope);
+		ParentScope->LastStatement->NextStatement = Statement;
+		FStatement* StatementToLink = Statement;
+		while (StatementToLink)
+		{
+			StatementToLink->ParentScope = ParentScope;
+			ParentScope->LastStatement = StatementToLink;
+			Visitor.VisitNode(StatementToLink);
+			StatementToLink = StatementToLink->NextStatement;
+		}
+		bResult = true;
+	}
+	return bResult;
+}
+
+UE::HLSLTree::FTree* UE::HLSLTree::FTree::Create(FMemStackBase& Allocator)
+{
+	FTree* Tree = new(Allocator) FTree();
+	Tree->Allocator = &Allocator;
+	Tree->RootScope = Tree->NewNode<FScope>();
+	return Tree;
+}
+
+void UE::HLSLTree::FTree::EmitHLSL(UE::HLSLTree::FEmitContext& Context, UE::HLSLTree::FCodeWriter& Writer) const
+{
+	RootScope->EmitHLSL(Context, Writer);
+}
+
+UE::HLSLTree::FScope* UE::HLSLTree::FTree::NewScope(UE::HLSLTree::FScope& Scope)
+{
+	FScope* NewScope = NewNode<FScope>();
+	NewScope->ParentScope = &Scope;
+	NewScope->NestedLevel = Scope.NestedLevel + 1;
+	return NewScope;
+}
+
+UE::HLSLTree::FScope* UE::HLSLTree::FTree::NewLinkedScope(UE::HLSLTree::FScope& Scope)
+{
+	FScope* NewScope = NewNode<FScope>();
+	NewScope->ParentScope = Scope.ParentScope;
+	NewScope->LinkedScope = &Scope;
+	NewScope->NestedLevel = Scope.NestedLevel;
+	return NewScope;
+}
+
+UE::HLSLTree::FLocalDeclaration* UE::HLSLTree::FTree::NewLocalDeclaration(UE::HLSLTree::FScope& Scope, UE::HLSLTree::EExpressionType Type, const FName& Name)
+{
+	FLocalDeclaration* Declaration = NewNode<FLocalDeclaration>(Name, Type);
+	Scope.AddDeclaration(Declaration);
+	return Declaration;
+}
+
+UE::HLSLTree::FParameterDeclaration* UE::HLSLTree::FTree::NewParameterDeclaration(UE::HLSLTree::FScope& Scope, const FName& Name, const UE::HLSLTree::FConstant& DefaultValue)
+{
+	FParameterDeclaration* Declaration = NewNode<FParameterDeclaration>(Name, DefaultValue);
+	Declaration->ParentScope = &Scope;
+	return Declaration;
+}
+UE::HLSLTree::FTextureParameterDeclaration* UE::HLSLTree::FTree::NewTextureParameterDeclaration(UE::HLSLTree::FScope& Scope, const FName& Name, const UE::HLSLTree::FTextureDescription& DefaultValue)
+{
+	FTextureParameterDeclaration* Declaration = NewNode<FTextureParameterDeclaration>(Name, DefaultValue);
+	Declaration->ParentScope = &Scope;
+	return Declaration;
+}
diff --git a/Engine/Source/Runtime/Engine/Private/HLSLTree/HLSLTreeCommon.cpp b/Engine/Source/Runtime/Engine/Private/HLSLTree/HLSLTreeCommon.cpp
new file mode 100644
--- /dev/null
+++ b/Engine/Source/Runtime/Engine/Private/HLSLTree/HLSLTreeCommon.cpp
@@ -0,0 +1,391 @@
+// Copyright Epic Games, Inc. All Rights Reserved.
+#include "HLSLTree/HLSLTreeCommon.h"
+#include "Misc/StringBuilder.h"
+#include "MaterialShared.h"
+#include "Engine/Texture.h"
+
+UE::HLSLTree::FSwizzleParameters::FSwizzleParameters(int8 InR, int8 InG, int8 InB, int8 InA) : NumComponents(0)
+{
+	ComponentIndex[0] = InR;
+	ComponentIndex[1] = InG;
+	ComponentIndex[2] = InB;
+	ComponentIndex[3] = InA;
+
+	if (InA >= 0)
+	{
+		check(InA <= 3);
+		++NumComponents;
+		check(InB >= 0);
+	}
+	if (InB >= 0)
+	{
+		check(InB <= 3);
+		++NumComponents;
+		check(InG >= 0);
+	}
+
+	if (InG >= 0)
+	{
+		check(InG <= 3);
+		++NumComponents;
+	}
+
+	// At least one proper index
+	check(InR >= 0 && InR <= 3);
+	++NumComponents;
+}
+
+UE::HLSLTree::FSwizzleParameters UE::HLSLTree::MakeSwizzleMask(bool bInR, bool bInG, bool bInB, bool bInA)
+{
+	int8 ComponentIndex[4] = { INDEX_NONE, INDEX_NONE, INDEX_NONE, INDEX_NONE };
+	int8 CurrentComponent = 0;
+	if (bInR)
+	{
+		ComponentIndex[CurrentComponent++] = 0;
+	}
+	if (bInG)
+	{
+		ComponentIndex[CurrentComponent++] = 1;
+	}
+	if (bInB)
+	{
+		ComponentIndex[CurrentComponent++] = 2;
+	}
+	if (bInA)
+	{
+		ComponentIndex[CurrentComponent++] = 3;
+	}
+	return FSwizzleParameters(ComponentIndex[0], ComponentIndex[1], ComponentIndex[2], ComponentIndex[3]);
+}
+
+void UE::HLSLTree::FExpressionConstant::EmitPreshader(UE::HLSLTree::FEmitContext& Context, FMaterialPreshaderData& OutPreshader) const
+{
+	OutPreshader.WriteOpcode(EMaterialPreshaderOpcode::Constant);
+	OutPreshader.Write(Value.ToLinearColor());
+}
+
+void UE::HLSLTree::FExpressionConstant::EmitHLSL(UE::HLSLTree::FEmitContext& Context, UE::HLSLTree::FCodeWriter& Writer) const
+{
+	Value.EmitHLSL(Writer);
+}
+
+void UE::HLSLTree::FExpressionLocalVariable::EmitHLSL(UE::HLSLTree::FEmitContext& Context, UE::HLSLTree::FCodeWriter& Writer) const
+{
+	const TCHAR* DeclarationRef = Context.AcquireHLSLReference(Declaration);
+	Writer.Writef(TEXT("%s"), DeclarationRef);
+}
+
+void UE::HLSLTree::FExpressionParameter::EmitPreshader(UE::HLSLTree::FEmitContext& Context, FMaterialPreshaderData& OutPreshader) const
+{
+	EMaterialPreshaderOpcode Opcode = EMaterialPreshaderOpcode::Nop;
+	int32 ParameterIndex = INDEX_NONE;
+	if (Declaration->DefaultValue.Type == UE::HLSLTree::EExpressionType::Float1)
+	{
+		ParameterIndex = Context.MaterialCompilationOutput->UniformExpressionSet.FindOrAddScalarParameter(Declaration->Name, Declaration->DefaultValue.Float[0]);
+		Opcode = EMaterialPreshaderOpcode::ScalarParameter;
+	}
+	else
+	{
+		ParameterIndex = Context.MaterialCompilationOutput->UniformExpressionSet.FindOrAddVectorParameter(Declaration->Name, Declaration->DefaultValue.ToLinearColor());
+		Opcode = EMaterialPreshaderOpcode::VectorParameter;
+	}
+
+	check(ParameterIndex >= 0 && ParameterIndex <= 0xffff);
+	OutPreshader.WriteOpcode(Opcode);
+	OutPreshader.Write((uint16)ParameterIndex);
+}
+
+void UE::HLSLTree::FExpressionExternalInput::EmitHLSL(UE::HLSLTree::FEmitContext& Context, UE::HLSLTree::FCodeWriter& Writer) const
+{
+	const int32 TypeIndex = (int32)InputType;
+	const int32 TexCoordIndex = TypeIndex - (int32)UE::HLSLTree::EExternalInputType::TexCoord0;
+	if (TexCoordIndex >= 0 && TexCoordIndex < 8)
+	{
+		Context.NumTexCoords = FMath::Max(Context.NumTexCoords, TexCoordIndex + 1);
+		Writer.Writef(TEXT("Parameters.TexCoords[%u].xy"), TexCoordIndex);
+	}
+	else
+	{
+		checkNoEntry();
+	}
+}
+
+void UE::HLSLTree::FExpressionTextureSample::EmitHLSL(UE::HLSLTree::FEmitContext& Context, UE::HLSLTree::FCodeWriter& Writer) const
+{
+	const UE::HLSLTree::FTextureDescription& Desc = Declaration->Description;
+	const EMaterialValueType MaterialValueType = Desc.Texture->GetMaterialType();
+	EMaterialTextureParameterType TextureType = EMaterialTextureParameterType::Count;
+	bool bVirtualTexture = false;
+	const TCHAR* TextureTypeName = nullptr;
+	const TCHAR* SampleFunctionName = nullptr;
+	switch (MaterialValueType)
+	{
+	case MCT_Texture2D:
+		TextureType = EMaterialTextureParameterType::Standard2D;
+		TextureTypeName = TEXT("Texture2D");
+		SampleFunctionName = TEXT("Texture2DSample");
+		break;
+	case MCT_TextureCube:
+		TextureType = EMaterialTextureParameterType::Cube;
+		TextureTypeName = TEXT("TextureCube");
+		SampleFunctionName = TEXT("TextureCubeSample");
+		break;
+	case MCT_Texture2DArray:
+		TextureType = EMaterialTextureParameterType::Array2D;
+		TextureTypeName = TEXT("Texture2DArray");
+		SampleFunctionName = TEXT("Texture2DArraySample");
+		break;
+	case MCT_VolumeTexture:
+		TextureType = EMaterialTextureParameterType::Volume;
+		TextureTypeName = TEXT("VolumeTexture");
+		SampleFunctionName = TEXT("Texture3DSample");
+		break;
+	case MCT_TextureExternal:
+		// TODO
+		TextureTypeName = TEXT("ExternalTexture");
+		SampleFunctionName = TEXT("TextureExternalSample");
+		break; 
+	case MCT_TextureVirtual:
+		// TODO
+		TextureType = EMaterialTextureParameterType::Virtual;
+		SampleFunctionName = TEXT("TextureVirtualSample");
+		bVirtualTexture = true;
+		break;
+	default:
+		checkNoEntry();
+		break;
+	}
+
+	const TCHAR* SamplerTypeFunction = TEXT("");
+	switch (Desc.SamplerType)
+	{
+	case SAMPLERTYPE_External:
+		SamplerTypeFunction = TEXT("ProcessMaterialExternalTextureLookup");
+		break;
+	case SAMPLERTYPE_Color:
+		SamplerTypeFunction = TEXT("ProcessMaterialColorTextureLookup");
+		break;
+	case SAMPLERTYPE_VirtualColor:
+		// has a mobile specific workaround
+		SamplerTypeFunction = TEXT("ProcessMaterialVirtualColorTextureLookup");
+		break;
+	case SAMPLERTYPE_LinearColor:
+	case SAMPLERTYPE_VirtualLinearColor:
+		SamplerTypeFunction = TEXT("ProcessMaterialLinearColorTextureLookup");
+		break;
+	case SAMPLERTYPE_Alpha:
+	case SAMPLERTYPE_VirtualAlpha:
+	case SAMPLERTYPE_DistanceFieldFont:
+		SamplerTypeFunction = TEXT("ProcessMaterialAlphaTextureLookup");
+		break;
+	case SAMPLERTYPE_Grayscale:
+	case SAMPLERTYPE_VirtualGrayscale:
+		SamplerTypeFunction = TEXT("ProcessMaterialGreyscaleTextureLookup");
+		break;
+	case SAMPLERTYPE_LinearGrayscale:
+	case SAMPLERTYPE_VirtualLinearGrayscale:
+		SamplerTypeFunction = TEXT("ProcessMaterialLinearGreyscaleTextureLookup");
+		break;
+	case SAMPLERTYPE_Normal:
+	case SAMPLERTYPE_VirtualNormal:
+		// Normal maps need to be unpacked in the pixel shader.
+		SamplerTypeFunction = TEXT("UnpackNormalMap");
+		break;
+	case SAMPLERTYPE_Masks:
+	case SAMPLERTYPE_VirtualMasks:
+	case SAMPLERTYPE_Data:
+		SamplerTypeFunction = TEXT("");
+		break;
+	default:
+		check(0);
+		break;
+	}
+
+	FMaterialTextureParameterInfo TextureParameterInfo;
+	TextureParameterInfo.ParameterInfo = Declaration->Name;
+	TextureParameterInfo.TextureIndex = Context.Material->GetReferencedTextures().Find(Desc.Texture);
+	TextureParameterInfo.SamplerSource = SamplerSource;
+	check(TextureParameterInfo.TextureIndex != INDEX_NONE);
+
+	const int32 ParameterIndex = Context.MaterialCompilationOutput->UniformExpressionSet.FindOrAddTextureParameter(TextureType, TextureParameterInfo);
+	const FString TextureName = FString::Printf(TEXT("Material.%s_%u"), TextureTypeName, ParameterIndex);
+
+	FString SamplerStateCode;
+	bool AutomaticViewMipBias = false; // TODO
+	bool RequiresManualViewMipBias = AutomaticViewMipBias;
+
+	if (!bVirtualTexture) //VT does not have explict samplers (and always requires manual view mip bias)
+	{
+		if (SamplerSource == SSM_FromTextureAsset)
+		{
+			SamplerStateCode = FString::Printf(TEXT("%sSampler"), *TextureName);
+		}
+		else if (SamplerSource == SSM_Wrap_WorldGroupSettings)
+		{
+			// Use the shared sampler to save sampler slots
+			const TCHAR* SharedSamplerName = AutomaticViewMipBias ? TEXT("View.MaterialTextureBilinearWrapedSampler") : TEXT("Material.Wrap_WorldGroupSettings");
+			SamplerStateCode = FString::Printf(TEXT("GetMaterialSharedSampler(%sSampler,%s)"), *TextureName, SharedSamplerName);
+			RequiresManualViewMipBias = false;
+		}
+		else if (SamplerSource == SSM_Clamp_WorldGroupSettings)
+		{
+			// Use the shared sampler to save sampler slots
+			const TCHAR* SharedSamplerName = AutomaticViewMipBias ? TEXT("View.MaterialTextureBilinearClampedSampler") : TEXT("Material.Clamp_WorldGroupSettings");
+			SamplerStateCode = FString::Printf(TEXT("GetMaterialSharedSampler(%sSampler,%s)"), *TextureName, SharedSamplerName);
+			RequiresManualViewMipBias = false;
+		}
+	}
+
+	const TCHAR* TexCoordRef = Context.AcquireHLSLReference(TexCoordExpression);
+	Writer.Writef(TEXT("%s(%s(%s, %s, %s))"), SamplerTypeFunction, SampleFunctionName, *TextureName, *SamplerStateCode, TexCoordRef);
+}
+
+void UE::HLSLTree::FExpressionDefaultMaterialAttributes::EmitHLSL(UE::HLSLTree::FEmitContext& Context, UE::HLSLTree::FCodeWriter& Writer) const
+{
+	Writer.Write(TEXT("DefaultMaterialAttributes"));
+}
+
+void UE::HLSLTree::FExpressionSetMaterialAttribute::EmitHLSL(UE::HLSLTree::FEmitContext& Context, UE::HLSLTree::FCodeWriter& Writer) const
+{
+	const EMaterialProperty Property = FMaterialAttributeDefinitionMap::GetProperty(AttributeID);
+	const FString PropertyName = FMaterialAttributeDefinitionMap::GetAttributeName(AttributeID);
+	const EMaterialValueType PropertyType = FMaterialAttributeDefinitionMap::GetValueType(AttributeID);
+	const EShaderFrequency Frequency = FMaterialAttributeDefinitionMap::GetShaderFrequency(AttributeID);
+
+	const TCHAR* AttributresRef = Context.AcquireHLSLReference(AttributesExpression);
+	const TCHAR* ValueRef = Context.AcquireHLSLReference(ValueExpression);
+
+	Writer.Writef(TEXT("FMaterialAttributes_Set%s(%s, %s)"), *PropertyName, AttributresRef, ValueRef);
+}
+
+void UE::HLSLTree::FExpressionBinaryOp::EmitPreshader(UE::HLSLTree::FEmitContext& Context, FMaterialPreshaderData& OutPreshader) const
+{
+	EMaterialPreshaderOpcode Opcode = EMaterialPreshaderOpcode::Nop;
+	switch (Op)
+	{
+	case UE::HLSLTree::EBinaryOp::Add: Opcode = EMaterialPreshaderOpcode::Add; break;
+	default: checkNoEntry(); break;
+	}
+
+	Lhs->EmitPreshader(Context, OutPreshader);
+	Rhs->EmitPreshader(Context, OutPreshader);
+	OutPreshader.WriteOpcode(Opcode);
+}
+
+void UE::HLSLTree::FExpressionBinaryOp::EmitHLSL(UE::HLSLTree::FEmitContext& Context, UE::HLSLTree::FCodeWriter& Writer) const
+{
+	const TCHAR* LhsRef = Context.AcquireHLSLReference(Lhs);
+	const TCHAR* RhsRef = Context.AcquireHLSLReference(Rhs);
+	Writer.Writef(TEXT("(%s + %s)"), LhsRef, RhsRef);
+}
+
+void UE::HLSLTree::FExpressionSwizzle::EmitPreshader(FEmitContext& Context, FMaterialPreshaderData& OutPreshader) const
+{
+	Input->EmitPreshader(Context, OutPreshader);
+	OutPreshader.WriteOpcode(EMaterialPreshaderOpcode::ComponentSwizzle)
+		.Write((uint8)Parameters.NumComponents)
+		.Write((uint8)Parameters.ComponentIndex[0])
+		.Write((uint8)Parameters.ComponentIndex[1])
+		.Write((uint8)Parameters.ComponentIndex[2])
+		.Write((uint8)Parameters.ComponentIndex[3]);
+}
+
+void UE::HLSLTree::FExpressionSwizzle::EmitHLSL(FEmitContext& Context, FCodeWriter& Writer) const
+{
+	static const TCHAR ComponentName[] = { 'x', 'y', 'z', 'w' };
+	TCHAR Swizzle[5] = TEXT("");
+	for (int32 i = 0; i < Parameters.NumComponents; ++i)
+	{
+		const int32 ComponentIndex = Parameters.ComponentIndex[i];
+		check(ComponentIndex >= 0 && ComponentIndex < 4);
+		Swizzle[i] = ComponentName[ComponentIndex];
+	}
+
+	const TCHAR* InputRef = Context.AcquireHLSLReference(Input);
+	Writer.Writef(TEXT("%s.%s"), InputRef, Swizzle);
+}
+
+void UE::HLSLTree::FExpressionCast::EmitHLSL(FEmitContext& Context, FCodeWriter& Writer) const
+{
+	const EExpressionType InputType = Input->Type;
+	const TCHAR* InputRef = Context.AcquireHLSLReference(Input);
+	if (InputType == Type)
+	{
+		// No cast required (nop)
+		Writer.Writef(TEXT("%s"), InputRef);
+	}
+
+	const FExpressionTypeDescription OutputTypeDesc = GetExpressionTypeDescription(Type);
+	const FExpressionTypeDescription InputTypeDesc = GetExpressionTypeDescription(InputType);
+	if (InputTypeDesc.NumComponents == OutputTypeDesc.NumComponents)
+	{
+		// Cast between different underlying types with the same number of components, use a C-style cast
+		Writer.Writef(TEXT("((%s)%s)"), OutputTypeDesc.Name, InputRef);
+	}
+	else if (OutputTypeDesc.NumComponents < InputTypeDesc.NumComponents)
+	{
+		static const TCHAR* SwizzleName[] = { TEXT("x"), TEXT("xy"), TEXT("xyz"), TEXT("xyzw") };
+		// Masking off some input components, use a swizzle
+		if (OutputTypeDesc.ComponentType == InputTypeDesc.ComponentType)
+		{
+			// underlying types are the same, no C-style cast needed
+			Writer.Writef(TEXT("%s.%s"), InputRef, SwizzleName[OutputTypeDesc.NumComponents]);
+		}
+		else
+		{
+			// Different underlying types, add a C-style cast as well
+			Writer.Writef(TEXT("((%s)%s.%s)"), OutputTypeDesc.Name, InputRef, SwizzleName[OutputTypeDesc.NumComponents]);
+		}
+	}
+	else
+	{
+		// Padding out components
+		if (InputTypeDesc.NumComponents == 1 && EnumHasAllFlags(Flags, ECastFlags::ReplicateScalar))
+		{
+			// Replicating a scalar value, use a swizzle
+			static const TCHAR* ScalarSwizzleName[] = { TEXT("x"), TEXT("xx"), TEXT("xxx"), TEXT("xxxx") };
+			Writer.Writef(TEXT("%s.%s"), InputRef, ScalarSwizzleName[OutputTypeDesc.NumComponents]);
+		}
+		else
+		{
+			// Padding with 0s, use a constructor
+			Writer.Writef(TEXT("%s(%s"), OutputTypeDesc.Name, InputRef);
+			for (int32 i = InputTypeDesc.NumComponents; i < OutputTypeDesc.NumComponents; ++i)
+			{
+				// TODO - worth making these 0s the appropriate type? 0u, 0.0f, etc
+				Writer.Write(TEXT(", 0"));
+			}
+			Writer.Write(TEXT(")"));
+		}
+	}
+}
+
+void UE::HLSLTree::FStatementReturn::EmitHLSL(UE::HLSLTree::FEmitContext& Context, UE::HLSLTree::FCodeWriter& Writer) const
+{
+	const TCHAR* Ref = Context.AcquireHLSLReference(Expression);
+	Writer.WriteLinef(TEXT("return %s;"), Ref);
+}
+
+void UE::HLSLTree::FStatementSetLocalVariable::EmitHLSL(UE::HLSLTree::FEmitContext& Context, UE::HLSLTree::FCodeWriter& Writer) const
+{
+	const TCHAR* DeclarationRef = Context.AcquireHLSLReference(Declaration);
+	const TCHAR* ExpressionRef = Context.AcquireHLSLReference(Expression);
+	Writer.WriteLinef(TEXT("%s = %s;"), DeclarationRef, ExpressionRef);
+}
+
+void UE::HLSLTree::FStatementIf::EmitHLSL(UE::HLSLTree::FEmitContext& Context, UE::HLSLTree::FCodeWriter& Writer) const
+{
+	const TCHAR* ConditionRef = Context.AcquireHLSLReference(ConditionExpression);
+	Writer.WriteLinef(TEXT("if (%s)"), ConditionRef);
+	ThenScope->EmitHLSL(Context, Writer);
+	if (ElseScope)
+	{
+		Writer.WriteLine(TEXT("else"));
+		ElseScope->EmitHLSL(Context, Writer);
+	}
+}
+
+void UE::HLSLTree::FStatementFor::EmitHLSL(UE::HLSLTree::FEmitContext& Context, UE::HLSLTree::FCodeWriter& Writer) const
+{
+	
+}
diff --git a/Engine/Source/Runtime/Engine/Private/Materials/HLSLMaterialTranslator.cpp b/Engine/Source/Runtime/Engine/Private/Materials/HLSLMaterialTranslator.cpp
--- a/Engine/Source/Runtime/Engine/Private/Materials/HLSLMaterialTranslator.cpp
+++ b/Engine/Source/Runtime/Engine/Private/Materials/HLSLMaterialTranslator.cpp
@@ -7023,26 +7023,17 @@ static FString ApplySamplerType(const FString& InSampleCode, EMaterialSamplerTyp
 	case SAMPLERTYPE_Alpha:
 	case SAMPLERTYPE_VirtualAlpha:
 	case SAMPLERTYPE_DistanceFieldFont:
-		// Sampling a single channel texture in D3D9 gives: (G,G,G)
-		// Sampling a single channel texture in D3D11 gives: (G,0,0)
-		// This replication reproduces the D3D9 behavior in all cases.
-		DstSampleCode = FString::Printf( TEXT("(%s).rrrr"), *InSampleCode );
+		DstSampleCode = FString::Printf( TEXT("ProcessMaterialAlphaTextureLookup(%s)"), *InSampleCode );
 		break;
 
 	case SAMPLERTYPE_Grayscale:
 	case SAMPLERTYPE_VirtualGrayscale:
-		// Sampling a greyscale texture in D3D9 gives: (G,G,G)
-		// Sampling a greyscale texture in D3D11 gives: (G,0,0)
-		// This replication reproduces the D3D9 behavior in all cases.
-		DstSampleCode = FString::Printf( TEXT("ProcessMaterialGreyscaleTextureLookup((%s).r).rrrr"), *InSampleCode );
+		DstSampleCode = FString::Printf( TEXT("ProcessMaterialGreyscaleTextureLookup(%s)"), *InSampleCode );
 		break;
 
 	case SAMPLERTYPE_LinearGrayscale:
 	case SAMPLERTYPE_VirtualLinearGrayscale:
-		// Sampling a greyscale texture in D3D9 gives: (G,G,G)
-		// Sampling a greyscale texture in D3D11 gives: (G,0,0)
-		// This replication reproduces the D3D9 behavior in all cases.
-		DstSampleCode = FString::Printf(TEXT("ProcessMaterialLinearGreyscaleTextureLookup((%s).r).rrrr"), *InSampleCode);
+		DstSampleCode = FString::Printf(TEXT("ProcessMaterialLinearGreyscaleTextureLookup(%s)"), *InSampleCode);
 		break;
 
 	case SAMPLERTYPE_Normal:
diff --git a/Engine/Source/Runtime/Engine/Private/Materials/Material.cpp b/Engine/Source/Runtime/Engine/Private/Materials/Material.cpp
--- a/Engine/Source/Runtime/Engine/Private/Materials/Material.cpp
+++ b/Engine/Source/Runtime/Engine/Private/Materials/Material.cpp
@@ -136,6 +136,12 @@ static TAutoConsoleVariable<int32> CVarMaterialEnableControlFlow(
 	TEXT("Allows experemental control flow to be used in the material editor.\n"),
 	ECVF_RenderThreadSafe | ECVF_ReadOnly);
 
+static TAutoConsoleVariable<int32> CVarMaterialEnableNewHLSLGenerator(
+	TEXT("r.MaterialEnableNewHLSLGenerator"),
+	0,
+	TEXT("Enables the new (WIP) material HLSL generator.\n"),
+	ECVF_RenderThreadSafe | ECVF_ReadOnly);
+
 #if WITH_EDITOR
 const FMaterialsWithDirtyUsageFlags FMaterialsWithDirtyUsageFlags::DefaultAnnotation;
 
@@ -4328,6 +4334,15 @@ bool UMaterial::IsCompiledWithExecutionFlow() const
 	return false;
 }
 
+bool UMaterial::IsUsingNewHLSLGenerator() const
+{
+	if (bEnableNewHLSLGenerator)
+	{
+		return CVarMaterialEnableNewHLSLGenerator.GetValueOnAnyThread() != 0;
+	}
+	return false;
+}
+
 #if WITH_EDITOR
 void UMaterial::BeginCacheForCookedPlatformData( const ITargetPlatform *TargetPlatform )
 {
diff --git a/Engine/Source/Runtime/Engine/Private/Materials/MaterialExpressionHLSL.cpp b/Engine/Source/Runtime/Engine/Private/Materials/MaterialExpressionHLSL.cpp
new file mode 100644
--- /dev/null
+++ b/Engine/Source/Runtime/Engine/Private/Materials/MaterialExpressionHLSL.cpp
@@ -0,0 +1,309 @@
+// Copyright Epic Games, Inc. All Rights Reserved.
+#include "CoreMinimal.h"
+
+#if WITH_EDITOR
+
+#include "Materials/MaterialExpression.h"
+#include "MaterialHLSLGenerator.h"
+#include "HLSLTree/HLSLTree.h"
+#include "HLSLTree/HLSLTreeCommon.h"
+
+#include "Materials/MaterialExpressionExecBegin.h"
+#include "Materials/MaterialExpressionReturnMaterialAttributes.h"
+#include "Materials/MaterialExpressionConstant.h"
+#include "Materials/MaterialExpressionConstant2Vector.h"
+#include "Materials/MaterialExpressionConstant3Vector.h"
+#include "Materials/MaterialExpressionConstant4Vector.h"
+#include "Materials/MaterialExpressionVectorParameter.h"
+#include "Materials/MaterialExpressionScalarParameter.h"
+#include "Materials/MaterialExpressionTextureSample.h"
+#include "Materials/MaterialExpressionTextureSampleParameter.h"
+#include "Materials/MaterialExpressionTextureObject.h"
+#include "Materials/MaterialExpressionTextureObjectParameter.h"
+#include "Materials/MaterialExpressionGetLocal.h"
+#include "Materials/MaterialExpressionAdd.h"
+#include "Materials/MaterialExpressionSetMaterialAttributes.h"
+#include "Materials/MaterialExpressionSetLocal.h"
+#include "Materials/MaterialExpressionIfThenElse.h"
+#include "Materials/MaterialExpressionForLoop.h"
+
+EMaterialGenerateHLSLStatus UMaterialExpression::GenerateHLSLExpression(FMaterialHLSLGenerator& Generator, UE::HLSLTree::FScope& Scope, int32 OutputIndex, UE::HLSLTree::FExpression*& OutExpression)
+{
+	return Generator.Error(TEXT("Node does not support expressions"));
+}
+
+EMaterialGenerateHLSLStatus UMaterialExpression::GenerateHLSLStatement(FMaterialHLSLGenerator& Generator, UE::HLSLTree::FScope& Scope, UE::HLSLTree::FStatement*& OutStatement)
+{
+	return Generator.Error(TEXT("Node does not support statements"));
+}
+
+EMaterialGenerateHLSLStatus UMaterialExpression::GenerateHLSLTexture(FMaterialHLSLGenerator& Generator, UE::HLSLTree::FScope& Scope, int32 OutputIndex, UE::HLSLTree::FTextureParameterDeclaration*& OutTexture)
+{
+	return Generator.Error(TEXT("Node does not support textures"));
+}
+
+EMaterialGenerateHLSLStatus UMaterialExpressionConstant::GenerateHLSLExpression(FMaterialHLSLGenerator& Generator, UE::HLSLTree::FScope& Scope, int32 OutputIndex, UE::HLSLTree::FExpression*& OutExpression)
+{
+	OutExpression = Generator.NewConstant(Scope, R);
+	return EMaterialGenerateHLSLStatus::Success;
+}
+
+EMaterialGenerateHLSLStatus UMaterialExpressionConstant2Vector::GenerateHLSLExpression(FMaterialHLSLGenerator& Generator, UE::HLSLTree::FScope& Scope, int32 OutputIndex, UE::HLSLTree::FExpression*& OutExpression)
+{
+	OutExpression = Generator.NewConstant(Scope, UE::HLSLTree::FConstant(R, G));
+	return EMaterialGenerateHLSLStatus::Success;
+}
+
+EMaterialGenerateHLSLStatus UMaterialExpressionConstant3Vector::GenerateHLSLExpression(FMaterialHLSLGenerator& Generator, UE::HLSLTree::FScope& Scope, int32 OutputIndex, UE::HLSLTree::FExpression*& OutExpression)
+{
+	OutExpression = Generator.NewConstant(Scope, UE::HLSLTree::FConstant(Constant.R, Constant.G, Constant.B));
+	return EMaterialGenerateHLSLStatus::Success;
+}
+
+EMaterialGenerateHLSLStatus UMaterialExpressionConstant4Vector::GenerateHLSLExpression(FMaterialHLSLGenerator& Generator, UE::HLSLTree::FScope& Scope, int32 OutputIndex, UE::HLSLTree::FExpression*& OutExpression)
+{
+	OutExpression = Generator.NewConstant(Scope, UE::HLSLTree::FConstant(Constant.R, Constant.G, Constant.B, Constant.A));
+	return EMaterialGenerateHLSLStatus::Success;
+}
+
+EMaterialGenerateHLSLStatus UMaterialExpressionGetLocal::GenerateHLSLExpression(FMaterialHLSLGenerator& Generator, UE::HLSLTree::FScope& Scope, int32 OutputIndex, UE::HLSLTree::FExpression*& OutExpression)
+{
+	UE::HLSLTree::FLocalDeclaration* Declaration = Generator.AcquireLocalDeclaration(Scope, UE::HLSLTree::EExpressionType::Float3, LocalName);
+	OutExpression = Generator.GetTree().NewExpression<UE::HLSLTree::FExpressionLocalVariable>(Scope, Declaration);
+	return EMaterialGenerateHLSLStatus::Success;
+}
+
+EMaterialGenerateHLSLStatus UMaterialExpressionVectorParameter::GenerateHLSLExpression(FMaterialHLSLGenerator& Generator, UE::HLSLTree::FScope& Scope, int32 OutputIndex, UE::HLSLTree::FExpression*& OutExpression)
+{
+	UE::HLSLTree::FParameterDeclaration* Declaration = Generator.AcquireParameterDeclaration(Scope, ParameterName, UE::HLSLTree::FConstant(UE::HLSLTree::EExpressionType::Float4, DefaultValue));
+	OutExpression = Generator.GetTree().NewExpression<UE::HLSLTree::FExpressionParameter>(Scope, Declaration);
+	return EMaterialGenerateHLSLStatus::Success;
+}
+
+EMaterialGenerateHLSLStatus UMaterialExpressionScalarParameter::GenerateHLSLExpression(FMaterialHLSLGenerator& Generator, UE::HLSLTree::FScope& Scope, int32 OutputIndex, UE::HLSLTree::FExpression*& OutExpression)
+{
+	UE::HLSLTree::FParameterDeclaration* Declaration = Generator.AcquireParameterDeclaration(Scope, ParameterName, DefaultValue);
+	OutExpression = Generator.GetTree().NewExpression<UE::HLSLTree::FExpressionParameter>(Scope, Declaration);
+	return EMaterialGenerateHLSLStatus::Success;
+}
+
+EMaterialGenerateHLSLStatus UMaterialExpressionTextureObject::GenerateHLSLTexture(FMaterialHLSLGenerator& Generator, UE::HLSLTree::FScope& Scope, int32 OutputIndex, UE::HLSLTree::FTextureParameterDeclaration*& OutTexture)
+{
+	const UE::HLSLTree::FTextureDescription TextureDesc(Texture, SamplerType);
+	OutTexture = Generator.AcquireTextureDeclaration(Scope, TextureDesc);
+	return EMaterialGenerateHLSLStatus::Success;
+}
+
+EMaterialGenerateHLSLStatus UMaterialExpressionTextureObjectParameter::GenerateHLSLTexture(FMaterialHLSLGenerator& Generator, UE::HLSLTree::FScope& Scope, int32 OutputIndex, UE::HLSLTree::FTextureParameterDeclaration*& OutTexture)
+{
+	const UE::HLSLTree::FTextureDescription TextureDesc(Texture, SamplerType);
+	OutTexture = Generator.AcquireTextureParameterDeclaration(Scope, ParameterName, TextureDesc);
+	return EMaterialGenerateHLSLStatus::Success;
+}
+
+EMaterialGenerateHLSLStatus UMaterialExpressionTextureSample::GenerateHLSLExpressionBase(FMaterialHLSLGenerator& Generator, UE::HLSLTree::FScope& Scope, UE::HLSLTree::FTextureParameterDeclaration* TextureDeclaration, UE::HLSLTree::FExpression*& OutExpression)
+{
+	if (!TextureDeclaration)
+	{
+		return Generator.Error(TEXT("Missing input texture"));
+	}
+
+	UE::HLSLTree::FExpression* TexCoordExpression = Coordinates.GetTracedInput().Expression ? Coordinates.AcquireHLSLExpression(Generator, Scope) : Generator.NewTexCoord(Scope, ConstCoordinate);
+	UE::HLSLTree::FExpressionTextureSample* ExpressionTextureSample = Generator.GetTree().NewExpression<UE::HLSLTree::FExpressionTextureSample>(Scope, TextureDeclaration, TexCoordExpression);
+	ExpressionTextureSample->SamplerSource = SamplerSource;
+	ExpressionTextureSample->MipValueMode = MipValueMode;
+	OutExpression = ExpressionTextureSample;
+	return EMaterialGenerateHLSLStatus::Success;
+}
+
+EMaterialGenerateHLSLStatus UMaterialExpressionTextureSample::GenerateHLSLExpression(FMaterialHLSLGenerator& Generator, UE::HLSLTree::FScope& Scope, int32 OutputIndex, UE::HLSLTree::FExpression*& OutExpression)
+{
+	UE::HLSLTree::FTextureParameterDeclaration* TextureDeclaration = nullptr;
+	if (TextureObject.GetTracedInput().Expression)
+	{
+		TextureDeclaration = TextureObject.AcquireHLSLTexture(Generator, Scope);
+	}
+	else if (Texture)
+	{
+		const UE::HLSLTree::FTextureDescription TextureDesc(Texture, SamplerType);
+		TextureDeclaration = Generator.AcquireTextureDeclaration(Scope, TextureDesc);
+	}
+
+	return GenerateHLSLExpressionBase(Generator, Scope, TextureDeclaration, OutExpression);
+}
+
+EMaterialGenerateHLSLStatus UMaterialExpressionTextureSampleParameter::GenerateHLSLExpression(FMaterialHLSLGenerator& Generator, UE::HLSLTree::FScope& Scope, int32 OutputIndex, UE::HLSLTree::FExpression*& OutExpression)
+{
+	UE::HLSLTree::FTextureParameterDeclaration* TextureDeclaration = nullptr;
+	if (ParameterName.IsNone() && TextureObject.GetTracedInput().Expression)
+	{
+		TextureDeclaration = TextureObject.AcquireHLSLTexture(Generator, Scope);
+	}
+	else if (Texture)
+	{
+		const UE::HLSLTree::FTextureDescription TextureDesc(Texture, SamplerType);
+		TextureDeclaration = Generator.AcquireTextureParameterDeclaration(Scope, ParameterName, TextureDesc);
+	}
+
+	return GenerateHLSLExpressionBase(Generator, Scope, TextureDeclaration, OutExpression);
+}
+
+EMaterialGenerateHLSLStatus UMaterialExpressionAdd::GenerateHLSLExpression(FMaterialHLSLGenerator& Generator, UE::HLSLTree::FScope& Scope, int32 OutputIndex, UE::HLSLTree::FExpression*& OutExpression)
+{
+	UE::HLSLTree::FExpression* Lhs = A.GetTracedInput().Expression ? A.AcquireHLSLExpression(Generator, Scope) : Generator.NewConstant(Scope, ConstA);
+	UE::HLSLTree::FExpression* Rhs = B.GetTracedInput().Expression ? B.AcquireHLSLExpression(Generator, Scope) : Generator.NewConstant(Scope, ConstB);
+	if (!Lhs || !Rhs)
+	{
+		return EMaterialGenerateHLSLStatus::Error;
+	}
+
+	FString ErrorMessage;
+	const UE::HLSLTree::EExpressionType ResultType = UE::HLSLTree::MakeArithmeticResultType(Lhs->Type, Rhs->Type, ErrorMessage);
+	if (ResultType == UE::HLSLTree::EExpressionType::Void)
+	{
+		return Generator.Errorf(TEXT("%s"), *ErrorMessage);
+	}
+
+	OutExpression = Generator.GetTree().NewExpression<UE::HLSLTree::FExpressionBinaryOp>(Scope, ResultType, UE::HLSLTree::EBinaryOp::Add, Lhs, Rhs);
+	return EMaterialGenerateHLSLStatus::Success;
+}
+
+EMaterialGenerateHLSLStatus UMaterialExpressionSetMaterialAttributes::GenerateHLSLExpression(FMaterialHLSLGenerator& Generator, UE::HLSLTree::FScope& Scope, int32 OutputIndex, UE::HLSLTree::FExpression*& OutExpression)
+{
+	UE::HLSLTree::FExpression* AttributesExpression = nullptr;
+	if (Inputs[0].GetTracedInput().Expression)
+	{
+		AttributesExpression = Inputs[0].AcquireHLSLExpression(Generator, Scope);
+	}
+	else
+	{
+		AttributesExpression = Generator.GetTree().NewExpression<UE::HLSLTree::FExpressionDefaultMaterialAttributes>(Scope);
+	}
+
+	for (int32 PinIndex = 0; PinIndex < AttributeSetTypes.Num(); ++PinIndex)
+	{
+		const FExpressionInput& AttributeInput = Inputs[PinIndex + 1];
+		if (AttributeInput.GetTracedInput().Expression)
+		{
+			const FGuid& AttributeID = AttributeSetTypes[PinIndex];
+			// Only compile code to set attributes of the current shader frequency
+			const EShaderFrequency AttributeFrequency = FMaterialAttributeDefinitionMap::GetShaderFrequency(AttributeID);
+			//if (AttributeFrequency == Compiler->GetCurrentShaderFrequency())
+			{
+				UE::HLSLTree::FExpression* ValueExpression = AttributeInput.AcquireHLSLExpression(Generator, Scope);
+				if (ValueExpression)
+				{
+					UE::HLSLTree::FExpressionSetMaterialAttribute* SetAttributeExpression = Generator.GetTree().NewExpression<UE::HLSLTree::FExpressionSetMaterialAttribute>(Scope);
+					SetAttributeExpression->AttributeID = AttributeID;
+					SetAttributeExpression->AttributesExpression = AttributesExpression;
+					SetAttributeExpression->ValueExpression = ValueExpression;
+					AttributesExpression = SetAttributeExpression;
+				}
+			}
+		}
+	}
+
+	OutExpression = AttributesExpression;
+	return EMaterialGenerateHLSLStatus::Success;
+}
+
+EMaterialGenerateHLSLStatus UMaterialExpressionExecBegin::GenerateHLSLStatement(FMaterialHLSLGenerator& Generator, UE::HLSLTree::FScope& Scope, UE::HLSLTree::FStatement*& OutStatement)
+{
+	OutStatement = Exec.AcquireHLSLStatement(Generator, Scope);
+	return EMaterialGenerateHLSLStatus::Success;
+}
+
+EMaterialGenerateHLSLStatus UMaterialExpressionSetLocal::GenerateHLSLStatement(FMaterialHLSLGenerator& Generator, UE::HLSLTree::FScope& Scope, UE::HLSLTree::FStatement*& OutStatement)
+{
+	UE::HLSLTree::FExpression* ValueExpression = Value.AcquireHLSLExpressionWithCast(Generator, Scope, UE::HLSLTree::EExpressionType::Float3);
+	if (!ValueExpression)
+	{
+		return Generator.Error(TEXT("Missing value connection"));
+	}
+
+	UE::HLSLTree::FLocalDeclaration* Declaration = Generator.AcquireLocalDeclaration(Scope, UE::HLSLTree::EExpressionType::Float3, LocalName);
+	if (!Declaration)
+	{
+		return EMaterialGenerateHLSLStatus::Error;
+	}
+
+	UE::HLSLTree::FStatementSetLocalVariable* Statement = Generator.GetTree().NewStatement<UE::HLSLTree::FStatementSetLocalVariable>(Scope);
+	Statement->Declaration = Declaration;
+	Statement->Expression = ValueExpression;
+
+	Exec.AcquireHLSLStatement(Generator, Scope);
+
+	OutStatement = Statement;
+	return EMaterialGenerateHLSLStatus::Success;
+}
+
+EMaterialGenerateHLSLStatus UMaterialExpressionReturnMaterialAttributes::GenerateHLSLStatement(FMaterialHLSLGenerator& Generator, UE::HLSLTree::FScope& Scope, UE::HLSLTree::FStatement*& OutStatement)
+{
+	UE::HLSLTree::FExpression* AttributesExpression = MaterialAttributes.AcquireHLSLExpression(Generator, Scope);
+	if (!AttributesExpression)
+	{
+		return Generator.Error(TEXT("Missing attribute connection"));
+	}
+
+	UE::HLSLTree::FStatementReturn* ReturnStatement = Generator.GetTree().NewStatement<UE::HLSLTree::FStatementReturn>(Scope);
+	ReturnStatement->Expression = AttributesExpression;
+
+	OutStatement = ReturnStatement;
+	return EMaterialGenerateHLSLStatus::Success;
+}
+
+EMaterialGenerateHLSLStatus UMaterialExpressionIfThenElse::GenerateHLSLStatement(FMaterialHLSLGenerator& Generator, UE::HLSLTree::FScope& Scope, UE::HLSLTree::FStatement*& OutStatement)
+{
+	UE::HLSLTree::FExpression* ConditionExpression = Condition.AcquireHLSLExpression(Generator, Scope);
+	if (!ConditionExpression)
+	{
+		return Generator.Error(TEXT("Missing condition connection"));
+	}
+
+	UE::HLSLTree::FScope* ThenScope = Then.NewScopeWithStatement(Generator, Scope);
+	if (!ThenScope)
+	{
+		return Generator.Error(TEXT("Missing Then connection"));
+	}
+
+	UE::HLSLTree::FStatementIf* IfStatement = Generator.GetTree().NewStatement<UE::HLSLTree::FStatementIf>(Scope);
+	IfStatement->ConditionExpression = ConditionExpression;
+	IfStatement->ThenScope = ThenScope;
+	IfStatement->ElseScope = Else.NewLinkedScopeWithStatement(Generator, *ThenScope);
+
+	OutStatement = IfStatement;
+	return EMaterialGenerateHLSLStatus::Success;
+}
+
+EMaterialGenerateHLSLStatus UMaterialExpressionForLoop::GenerateHLSLStatement(FMaterialHLSLGenerator& Generator, UE::HLSLTree::FScope& Scope, UE::HLSLTree::FStatement*& OutStatement)
+{
+	UE::HLSLTree::FExpression* StartExpression = StartIndex.AcquireHLSLExpression(Generator, Scope);
+	if (!StartExpression)
+	{
+		return Generator.Error(TEXT("Missing StartIndex connection"));
+	}
+
+	UE::HLSLTree::FExpression* EndExpression = EndIndex.AcquireHLSLExpression(Generator, Scope);
+	if (!EndExpression)
+	{
+		return Generator.Error(TEXT("Missing EndIndex connection"));
+	}
+
+	UE::HLSLTree::FScope* LoopScope = LoopBody.NewScopeWithStatement(Generator, Scope);
+	if (!LoopScope)
+	{
+		return Generator.Error(TEXT("Missing LoopBody connection"));
+	}
+
+	UE::HLSLTree::FStatementFor* ForStatement = Generator.GetTree().NewStatement<UE::HLSLTree::FStatementFor>(Scope);
+	ForStatement->StartExpression = StartExpression;
+	ForStatement->EndExpression = EndExpression;
+	ForStatement->LoopScope = LoopScope;
+
+	Completed.AcquireHLSLStatement(Generator, Scope);
+
+	OutStatement = ForStatement;
+	return EMaterialGenerateHLSLStatus::Success;
+}
+
+#endif // WITH_EDITOR
diff --git a/Engine/Source/Runtime/Engine/Private/Materials/MaterialExpressions.cpp b/Engine/Source/Runtime/Engine/Private/Materials/MaterialExpressions.cpp
--- a/Engine/Source/Runtime/Engine/Private/Materials/MaterialExpressions.cpp
+++ b/Engine/Source/Runtime/Engine/Private/Materials/MaterialExpressions.cpp
@@ -1788,28 +1788,18 @@ FName UMaterialExpressionTextureSample::GetInputName(int32 InputIndex) const
 }
 #undef IF_INPUT_RETURN
 
-/**
- * Verify that the texture and sampler type. Generates a compiler waring if 
- * they do not.
- * @param Compiler - The material compiler to which errors will be reported.
- * @param ExpressionDesc - Description of the expression verifying the sampler type.
- * @param Texture - The texture to verify. A nullptr texture is considered valid!
- * @param SamplerType - The sampler type to verify.
- */
-static bool VerifySamplerType(
-	FMaterialCompiler* Compiler,
-	const TCHAR* ExpressionDesc,
+bool UMaterialExpressionTextureBase::VerifySamplerType(
+	ERHIFeatureLevel::Type FeatureLevel,
+	const ITargetPlatform* TargetPlatform,
 	const UTexture* Texture,
-	EMaterialSamplerType SamplerType )
+	EMaterialSamplerType SamplerType,
+	FString& OutErrorMessage)
 {
-	check( Compiler );
-	check( ExpressionDesc );
-
 	if ( Texture )
 	{
 		EMaterialSamplerType CorrectSamplerType = UMaterialExpressionTextureBase::GetSamplerTypeForTexture( Texture );
 		bool bIsVirtualTextured = IsVirtualSamplerType(SamplerType);
-		if (bIsVirtualTextured && UseVirtualTexturing(Compiler->GetFeatureLevel(), Compiler->GetTargetPlatform()) == false)
+		if (bIsVirtualTextured && !UseVirtualTexturing(FeatureLevel, TargetPlatform))
 		{
 			SamplerType = UMaterialExpressionTextureBase::GetSamplerTypeForTexture(Texture, !bIsVirtualTextured);
 		}
@@ -1821,8 +1811,7 @@ static bool VerifySamplerType(
 			FString SamplerTypeDisplayName = SamplerTypeEnum->GetDisplayNameTextByValue(SamplerType).ToString();
 			FString TextureTypeDisplayName = SamplerTypeEnum->GetDisplayNameTextByValue(CorrectSamplerType).ToString();
 
-			Compiler->Errorf( TEXT("%s> Sampler type is %s, should be %s for %s"),
-				ExpressionDesc,
+			OutErrorMessage = FString::Printf(TEXT("Sampler type is %s, should be %s for %s"),
 				*SamplerTypeDisplayName,
 				*TextureTypeDisplayName,
 				*Texture->GetPathName() );
@@ -1835,8 +1824,7 @@ static bool VerifySamplerType(
 
 			FString SamplerTypeDisplayName = SamplerTypeEnum->GetDisplayNameTextByValue(SamplerType).ToString();
 
-			Compiler->Errorf( TEXT("%s> To use '%s' as sampler type, SRGB must be disabled for %s"),
-				ExpressionDesc,
+			OutErrorMessage = FString::Printf(TEXT("%s> To use '%s' as sampler type, SRGB must be disabled for %s"),
 				*SamplerTypeDisplayName,
 				*Texture->GetPathName() );
 			return false;
@@ -1900,7 +1888,8 @@ int32 UMaterialExpressionTextureSample::Compile(class FMaterialCompiler* Compile
 			}
 		}
 
-		if (EffectiveTexture && VerifySamplerType(Compiler, (Desc.Len() > 0 ? *Desc : TEXT("TextureSample")), EffectiveTexture, EffectiveSamplerType))
+		FString SamplerTypeError;
+		if (EffectiveTexture && VerifySamplerType(Compiler->GetFeatureLevel(), Compiler->GetTargetPlatform(), EffectiveTexture, EffectiveSamplerType, SamplerTypeError))
 		{
 			if (TextureCodeIndex != INDEX_NONE)
 			{
@@ -1942,7 +1931,7 @@ int32 UMaterialExpressionTextureSample::Compile(class FMaterialCompiler* Compile
 		else
 		{
 			// TextureObject.Expression is responsible for generating the error message, since it had a nullptr texture value
-			return INDEX_NONE;
+			return Compiler->Errorf(TEXT("%s"), *SamplerTypeError);
 		}
 	}
 	else
@@ -2823,9 +2812,10 @@ int32 UMaterialExpressionTextureSampleParameter::Compile(class FMaterialCompiler
 		return CompilerError(Compiler, *ErrorMessage);
 	}
 
-	if (!VerifySamplerType(Compiler, (Desc.Len() > 0 ? *Desc : TEXT("TextureSampleParameter")), Texture, SamplerType))
+	FString SamplerTypeError;
+	if (!VerifySamplerType(Compiler->GetFeatureLevel(), Compiler->GetTargetPlatform(), Texture, SamplerType, SamplerTypeError))
 	{
-		return INDEX_NONE;
+		return Compiler->Errorf(TEXT("%s"), *SamplerTypeError);
 	}
 
 	if (!ParameterName.IsValid() || ParameterName.IsNone())
@@ -3593,9 +3583,10 @@ int32 UMaterialExpressionTextureSampleParameterSubUV::Compile(class FMaterialCom
 		return CompilerError(Compiler, *ErrorMessage);
 	}
 
-	if (!VerifySamplerType(Compiler, (Desc.Len() > 0 ? *Desc : TEXT("TextureSampleParameterSubUV")), Texture, SamplerType))
+	FString SamplerTypeError;
+	if (!VerifySamplerType(Compiler->GetFeatureLevel(), Compiler->GetTargetPlatform(), Texture, SamplerType, SamplerTypeError))
 	{
-		return INDEX_NONE;
+		return Compiler->Errorf(TEXT("%s"), *SamplerTypeError);
 	}
 
 	int32 TextureCodeIndex = Compiler->TextureParameter(ParameterName, Texture, SamplerType);
@@ -9204,9 +9195,10 @@ int32 UMaterialExpressionParticleSubUV::Compile(class FMaterialCompiler* Compile
 {
 	if (Texture)
 	{
-		if (!VerifySamplerType(Compiler, (Desc.Len() > 0 ? *Desc : TEXT("ParticleSubUV")), Texture, SamplerType))
+		FString SamplerTypeError;
+		if (!VerifySamplerType(Compiler->GetFeatureLevel(), Compiler->GetTargetPlatform(), Texture, SamplerType, SamplerTypeError))
 		{
-			return INDEX_NONE;
+			return Compiler->Errorf(TEXT("%s"), *SamplerTypeError);
 		}
 		int32 TextureCodeIndex = Compiler->Texture(Texture, SamplerType);
 		return ParticleSubUV(Compiler, TextureCodeIndex, Texture, SamplerType, Coordinates, bBlend);
@@ -11039,9 +11031,10 @@ int32 UMaterialExpressionFontSample::Compile(class FMaterialCompiler* Compiler,
 			ExpectedSamplerType = Texture->SRGB ? SAMPLERTYPE_Color : SAMPLERTYPE_LinearColor;
 		}
 
-		if (!VerifySamplerType(Compiler, (Desc.Len() > 0 ? *Desc : TEXT("FontSample")), Texture, ExpectedSamplerType))
+		FString SamplerTypeError;
+		if (!UMaterialExpressionTextureBase::VerifySamplerType(Compiler->GetFeatureLevel(), Compiler->GetTargetPlatform(), Texture, ExpectedSamplerType, SamplerTypeError))
 		{
-			return INDEX_NONE;
+			return Compiler->Errorf(TEXT("%s"), *SamplerTypeError);
 		}
 
 		int32 TextureCodeIndex = Compiler->Texture(Texture, ExpectedSamplerType);
@@ -11144,10 +11137,12 @@ int32 UMaterialExpressionFontSampleParameter::Compile(class FMaterialCompiler* C
 			ExpectedSamplerType = Texture->SRGB ? SAMPLERTYPE_Color : SAMPLERTYPE_LinearColor;
 		}
 
-		if (!VerifySamplerType(Compiler, (Desc.Len() > 0 ? *Desc : TEXT("FontSampleParameter")), Texture, ExpectedSamplerType))
+		FString SamplerTypeError;
+		if (!UMaterialExpressionTextureBase::VerifySamplerType(Compiler->GetFeatureLevel(), Compiler->GetTargetPlatform(), Texture, ExpectedSamplerType, SamplerTypeError))
 		{
-			return INDEX_NONE;
+			return Compiler->Errorf(TEXT("%s"), *SamplerTypeError);
 		}
+
 		int32 TextureCodeIndex = Compiler->TextureParameter(ParameterName,Texture, ExpectedSamplerType);
 		Result = Compiler->TextureSample(
 			TextureCodeIndex,
@@ -17194,9 +17189,10 @@ int32 UMaterialExpressionAntialiasedTextureMask::Compile(class FMaterialCompiler
 		TextureCodeIndex = Compiler->TextureParameter(ParameterName, Texture, SamplerType);
 	}
 
-	if (!VerifySamplerType(Compiler, (Desc.Len() > 0 ? *Desc : TEXT("AntialiasedTextureMask")), Texture, SamplerType))
+	FString SamplerTypeError;
+	if (!VerifySamplerType(Compiler->GetFeatureLevel(), Compiler->GetTargetPlatform(), Texture, SamplerType, SamplerTypeError))
 	{
-		return INDEX_NONE;
+		return Compiler->Errorf(TEXT("%s"), *SamplerTypeError);
 	}
 
 	return Compiler->AntialiasedTextureMask(TextureCodeIndex,ArgCoord,Threshold,Channel);
diff --git a/Engine/Source/Runtime/Engine/Private/Materials/MaterialHLSLGenerator.cpp b/Engine/Source/Runtime/Engine/Private/Materials/MaterialHLSLGenerator.cpp
new file mode 100644
--- /dev/null
+++ b/Engine/Source/Runtime/Engine/Private/Materials/MaterialHLSLGenerator.cpp
@@ -0,0 +1,902 @@
+// Copyright Epic Games, Inc. All Rights Reserved.
+#include "MaterialHLSLGenerator.h"
+
+#if WITH_EDITOR
+
+#include "Materials/MaterialExpression.h"
+#include "Materials/MaterialExpressionTextureBase.h"
+#include "Materials/MaterialExpressionExecBegin.h"
+#include "Materials/MaterialExpressionMaterialFunctionCall.h"
+#include "Materials/MaterialExpressionFunctionInput.h"
+#include "Materials/MaterialExpressionFunctionOutput.h"
+#include "Materials/MaterialExpressionVolumetricAdvancedMaterialOutput.h"
+#include "Materials/Material.h"
+#include "ShaderCore.h"
+#include "HLSLTree/HLSLTree.h"
+#include "HLSLTree/HLSLTreeCommon.h"
+#include "Containers/LazyPrintf.h"
+
+static bool SharedPixelProperties[CompiledMP_MAX];
+
+static const TCHAR* HLSLTypeString(EMaterialValueType Type)
+{
+	switch (Type)
+	{
+	case MCT_Float1:				return TEXT("MaterialFloat");
+	case MCT_Float2:				return TEXT("MaterialFloat2");
+	case MCT_Float3:				return TEXT("MaterialFloat3");
+	case MCT_Float4:				return TEXT("MaterialFloat4");
+	case MCT_Float:					return TEXT("MaterialFloat");
+	case MCT_Texture2D:				return TEXT("texture2D");
+	case MCT_TextureCube:			return TEXT("textureCube");
+	case MCT_Texture2DArray:		return TEXT("texture2DArray");
+	case MCT_VolumeTexture:			return TEXT("volumeTexture");
+	case MCT_StaticBool:			return TEXT("static bool");
+	case MCT_MaterialAttributes:	return TEXT("FMaterialAttributes");
+	case MCT_TextureExternal:		return TEXT("TextureExternal");
+	case MCT_TextureVirtual:		return TEXT("TextureVirtual");
+	case MCT_VTPageTableResult:		return TEXT("VTPageTableResult");
+	case MCT_ShadingModel:			return TEXT("uint");
+	case MCT_Strata:				return TEXT("FStrataData");
+	default:						return TEXT("unknown");
+	};
+}
+
+FMaterialHLSLGenerator::FMaterialHLSLGenerator(EShaderPlatform InShaderPlatform,
+	ERHIFeatureLevel::Type InFeatureLevel,
+	const ITargetPlatform* InTargetPlatform)
+	: ShaderPlatform(InShaderPlatform)
+	, FeatureLevel(InFeatureLevel)
+	, TargetPlatform(InTargetPlatform)
+{
+	FMemory::Memzero(SharedPixelProperties);
+	SharedPixelProperties[MP_Normal] = true;
+	SharedPixelProperties[MP_Tangent] = true;
+	SharedPixelProperties[MP_EmissiveColor] = true;
+	SharedPixelProperties[MP_Opacity] = true;
+	SharedPixelProperties[MP_OpacityMask] = true;
+	SharedPixelProperties[MP_BaseColor] = true;
+	SharedPixelProperties[MP_Metallic] = true;
+	SharedPixelProperties[MP_Specular] = true;
+	SharedPixelProperties[MP_Roughness] = true;
+	SharedPixelProperties[MP_Anisotropy] = true;
+	SharedPixelProperties[MP_AmbientOcclusion] = true;
+	SharedPixelProperties[MP_Refraction] = true;
+	SharedPixelProperties[MP_PixelDepthOffset] = true;
+	SharedPixelProperties[MP_SubsurfaceColor] = true;
+	SharedPixelProperties[MP_ShadingModel] = true;
+	SharedPixelProperties[MP_FrontMaterial] = true;
+}
+
+static FString GenerateMaterialTemplateHLSL(EShaderPlatform ShaderPlatform,
+	FMaterial& Material,
+	const UE::HLSLTree::FEmitContext& EmitContext,
+	const TCHAR* PixelShaderCode)
+{
+	FString MaterialTemplateSource;
+	LoadShaderSourceFileChecked(TEXT("/Engine/Private/MaterialTemplate.ush"), ShaderPlatform, MaterialTemplateSource);
+
+	// Find the string index of the '#line' statement in MaterialTemplate.usf
+	const int32 LineIndex = MaterialTemplateSource.Find(TEXT("#line"), ESearchCase::CaseSensitive);
+	check(LineIndex != INDEX_NONE);
+
+	// Count line endings before the '#line' statement
+	int32 MaterialTemplateLineNumber = INDEX_NONE;
+	int32 StartPosition = LineIndex + 1;
+	do
+	{
+		MaterialTemplateLineNumber++;
+		// Using \n instead of LINE_TERMINATOR as not all of the lines are terminated consistently
+		// Subtract one from the last found line ending index to make sure we skip over it
+		StartPosition = MaterialTemplateSource.Find(TEXT("\n"), ESearchCase::CaseSensitive, ESearchDir::FromEnd, StartPosition - 1);
+	} while (StartPosition != INDEX_NONE);
+	check(MaterialTemplateLineNumber != INDEX_NONE);
+	// At this point MaterialTemplateLineNumber is one less than the line number of the '#line' statement
+	// For some reason we have to add 2 more to the #line value to get correct error line numbers from D3DXCompileShader
+	MaterialTemplateLineNumber += 3;
+
+	FLazyPrintf LazyPrintf(*MaterialTemplateSource);
+
+	const uint32 NumUserVertexTexCoords = EmitContext.NumTexCoords;
+	const uint32 NumUserTexCoords = EmitContext.NumTexCoords;
+	const uint32 NumCustomVectors = 0u;
+	const uint32 NumTexCoordVectors = EmitContext.NumTexCoords;
+
+	LazyPrintf.PushParam(*FString::Printf(TEXT("%u"), NumUserVertexTexCoords));
+	LazyPrintf.PushParam(*FString::Printf(TEXT("%u"), NumUserTexCoords));
+	LazyPrintf.PushParam(*FString::Printf(TEXT("%u"), NumCustomVectors));
+	LazyPrintf.PushParam(*FString::Printf(TEXT("%u"), NumTexCoordVectors));
+
+	FString VertexInterpolatorsOffsetsDefinition;
+	LazyPrintf.PushParam(*VertexInterpolatorsOffsetsDefinition);
+
+	FString MaterialAttributesDeclaration;
+	FString MaterialAttributesUtilities;
+	FString MaterialAttributesDefault;
+
+	const EMaterialShadingModel DefaultShadingModel = Material.GetShadingModels().GetFirstShadingModel();
+
+	const TArray<FGuid>& OrderedVisibleAttributes = FMaterialAttributeDefinitionMap::GetOrderedVisibleAttributeList();
+	for (const FGuid& AttributeID : OrderedVisibleAttributes)
+	{
+		const FString PropertyName = FMaterialAttributeDefinitionMap::GetAttributeName(AttributeID);
+		const EMaterialValueType PropertyType = FMaterialAttributeDefinitionMap::GetValueType(AttributeID);
+		const TCHAR* HLSLType = nullptr;
+
+		switch (PropertyType)
+		{
+		case MCT_Float1: case MCT_Float: HLSLType = TEXT("float"); break;
+		case MCT_Float2: HLSLType = TEXT("float2"); break;
+		case MCT_Float3: HLSLType = TEXT("float3"); break;
+		case MCT_Float4: HLSLType = TEXT("float4"); break;
+		case MCT_ShadingModel: HLSLType = TEXT("uint"); break;
+		case MCT_Strata: HLSLType = TEXT("FStrataData"); break;
+		default: break;
+		}
+
+		if (HLSLType)
+		{
+			const FVector4& DefaultValue = FMaterialAttributeDefinitionMap::GetDefaultValue(AttributeID);
+
+			MaterialAttributesDeclaration += FString::Printf(TEXT("\t%s %s;") LINE_TERMINATOR, HLSLType, *PropertyName);
+
+			// Chainable method to set the attribute
+			MaterialAttributesUtilities += FString::Printf(TEXT("FMaterialAttributes FMaterialAttributes_Set%s(FMaterialAttributes InAttributes, %s InValue) { InAttributes.%s = InValue; return InAttributes; }") LINE_TERMINATOR,
+				*PropertyName, HLSLType, *PropertyName);
+
+			// Standard value type
+			switch (PropertyType)
+			{
+			case MCT_Float: case MCT_Float1: MaterialAttributesDefault += FString::Printf(TEXT("\tResult.%s = %0.8f;") LINE_TERMINATOR, *PropertyName, DefaultValue.X); break;
+			case MCT_Float2: MaterialAttributesDefault += FString::Printf(TEXT("\tResult.%s = MaterialFloat2(%0.8f,%0.8f);") LINE_TERMINATOR, *PropertyName, DefaultValue.X, DefaultValue.Y); break;
+			case MCT_Float3: MaterialAttributesDefault += FString::Printf(TEXT("\tResult.%s = MaterialFloat3(%0.8f,%0.8f,%0.8f);") LINE_TERMINATOR, *PropertyName, DefaultValue.X, DefaultValue.Y, DefaultValue.Z); break;
+			case MCT_Float4: MaterialAttributesDefault += FString::Printf(TEXT("\tResult.%s = MaterialFloat4(%0.8f,%0.8f,%0.8f,%0.8f);") LINE_TERMINATOR, *PropertyName, DefaultValue.X, DefaultValue.Y, DefaultValue.Z, DefaultValue.W); break;
+			case MCT_ShadingModel: MaterialAttributesDefault += FString::Printf(TEXT("\tResult.%s = %d;") LINE_TERMINATOR, *PropertyName, (int32)DefaultShadingModel); break;
+			case MCT_Strata: MaterialAttributesDefault += FString::Printf(TEXT("\tResult.%s = GetInitialisedStrataData();") LINE_TERMINATOR, *PropertyName); break; // TODO
+			default: checkNoEntry(); break;
+			}
+		}
+	}
+
+	LazyPrintf.PushParam(*MaterialAttributesDeclaration);
+	LazyPrintf.PushParam(*MaterialAttributesUtilities);
+
+	// Stores the shared shader results member declarations
+	// PixelMembersDeclaration should be the same for all variations, but might change in the future. There are cases where work is shared
+	// between the pixel and vertex shader, but with Nanite all work has to be moved into the pixel shader, which means we will want
+	// different inputs. But for now, we are keeping them the same.
+	FString PixelMembersDeclaration;
+
+	for (int32 PropertyIndex = 0; PropertyIndex < MP_MAX; ++PropertyIndex)
+	{
+		// Skip non-shared properties
+		if (!SharedPixelProperties[PropertyIndex])
+		{
+			continue;
+		}
+
+		const EMaterialProperty Property = (EMaterialProperty)PropertyIndex;
+		check(FMaterialAttributeDefinitionMap::GetShaderFrequency(Property) == SF_Pixel);
+		// Special case MP_SubsurfaceColor as the actual property is a combination of the color and the profile but we don't want to expose the profile
+		const FString PropertyName = Property == MP_SubsurfaceColor ? "Subsurface" : FMaterialAttributeDefinitionMap::GetAttributeName(Property);
+		check(PropertyName.Len() > 0);
+		const EMaterialValueType Type = Property == MP_SubsurfaceColor ? MCT_Float4 : FMaterialAttributeDefinitionMap::GetValueType(Property);
+
+		PixelMembersDeclaration += FString::Printf(TEXT("\t%s %s;\n"), HLSLTypeString(Type), *PropertyName);
+	}
+
+	//for (int32 Index = 0; Index < CompiledPDV_MAX; Index++)
+	//{
+	//	GetSharedInputsMaterialCode(PixelMembersDeclaration[Index], NormalAssignment[Index], PixelMembersSetupAndAssignments[Index], (ECompiledPartialDerivativeVariation)Index);
+	//}
+
+	LazyPrintf.PushParam(*PixelMembersDeclaration);
+
+	{
+		FString DerivativeHelpers;// = DerivativeAutogen.GenerateUsedFunctions(*this);
+		FString DerivativeHelpersAndResources;// = DerivativeHelpers + ResourcesString;
+		//LazyPrintf.PushParam(*ResourcesString);
+		LazyPrintf.PushParam(*DerivativeHelpersAndResources);
+	}
+
+	// Anything used bye the GenerationFunctionCode() like WorldPositionOffset shouldn't be using texures, right?
+	// Let those use the standard finite differences textures, since they should be the same. If we actually want
+	// those to handle texture reads properly, we'll have to make extra versions.
+	ECompiledPartialDerivativeVariation BaseDerivativeVariation = CompiledPDV_FiniteDifferences;
+
+	//if (bCompileForComputeShader)
+	//{
+	//	LazyPrintf.PushParam(*GenerateFunctionCode(CompiledMP_EmissiveColorCS, BaseDerivativeVariation));
+	//}
+	//else
+	{
+		LazyPrintf.PushParam(TEXT("return 0"));
+	}
+
+	LazyPrintf.PushParam(*FString::Printf(TEXT("return %.5f"), Material.GetTranslucencyDirectionalLightingIntensity()));
+
+	LazyPrintf.PushParam(*FString::Printf(TEXT("return %.5f"), Material.GetTranslucentShadowDensityScale()));
+	LazyPrintf.PushParam(*FString::Printf(TEXT("return %.5f"), Material.GetTranslucentSelfShadowDensityScale()));
+	LazyPrintf.PushParam(*FString::Printf(TEXT("return %.5f"), Material.GetTranslucentSelfShadowSecondDensityScale()));
+	LazyPrintf.PushParam(*FString::Printf(TEXT("return %.5f"), Material.GetTranslucentSelfShadowSecondOpacity()));
+	LazyPrintf.PushParam(*FString::Printf(TEXT("return %.5f"), Material.GetTranslucentBackscatteringExponent()));
+
+	{
+		FLinearColor Extinction = Material.GetTranslucentMultipleScatteringExtinction();
+
+		LazyPrintf.PushParam(*FString::Printf(TEXT("return MaterialFloat3(%.5f, %.5f, %.5f)"), Extinction.R, Extinction.G, Extinction.B));
+	}
+
+	LazyPrintf.PushParam(*FString::Printf(TEXT("return %.5f"), Material.GetOpacityMaskClipValue()));
+
+	LazyPrintf.PushParam(TEXT("return Parameters.MaterialVertexAttributes.WorldPositionOffset"));
+	LazyPrintf.PushParam(TEXT("return 0.0f"));
+	LazyPrintf.PushParam(TEXT("return 0.0f"));
+	LazyPrintf.PushParam(TEXT("return 0.0f"));
+	LazyPrintf.PushParam(TEXT("return 0.0f"));
+	LazyPrintf.PushParam(TEXT("return 0.0f"));
+
+	// Print custom texture coordinate assignments, should be fine with regular derivatives
+	FString CustomUVAssignments;
+
+	int32 LastProperty = -1;
+	for (uint32 CustomUVIndex = 0; CustomUVIndex < NumUserTexCoords; CustomUVIndex++)
+	{
+		//if (bEnableExecutionFlow)
+		{
+			const FString AttributeName = FMaterialAttributeDefinitionMap::GetAttributeName((EMaterialProperty)(MP_CustomizedUVs0 + CustomUVIndex));
+			CustomUVAssignments += FString::Printf(TEXT("\tOutTexCoords[%u] = Parameters.MaterialVertexAttributes.%s;") LINE_TERMINATOR, CustomUVIndex, *AttributeName);
+		}
+		/*else
+		{
+			if (CustomUVIndex == 0)
+			{
+				CustomUVAssignments += DerivativeVariations[BaseDerivativeVariation].TranslatedCodeChunkDefinitions[MP_CustomizedUVs0 + CustomUVIndex];
+			}
+
+			if (DerivativeVariations[BaseDerivativeVariation].TranslatedCodeChunkDefinitions[MP_CustomizedUVs0 + CustomUVIndex].Len() > 0)
+			{
+				if (LastProperty >= 0)
+				{
+					check(DerivativeVariations[BaseDerivativeVariation].TranslatedCodeChunkDefinitions[LastProperty].Len() == DerivativeVariations[BaseDerivativeVariation].TranslatedCodeChunkDefinitions[MP_CustomizedUVs0 + CustomUVIndex].Len());
+				}
+				LastProperty = MP_CustomizedUVs0 + CustomUVIndex;
+			}
+			CustomUVAssignments += FString::Printf(TEXT("\tOutTexCoords[%u] = %s;") LINE_TERMINATOR, CustomUVIndex, *DerivativeVariations[BaseDerivativeVariation].TranslatedCodeChunks[MP_CustomizedUVs0 + CustomUVIndex]);
+		}*/
+	}
+
+	LazyPrintf.PushParam(*CustomUVAssignments);
+
+	// Print custom vertex shader interpolator assignments
+	FString CustomInterpolatorAssignments;
+
+	/*for (UMaterialExpressionVertexInterpolator* Interpolator : CustomVertexInterpolators)
+	{
+		if (Interpolator->InterpolatorOffset != INDEX_NONE)
+		{
+			check(Interpolator->InterpolatorIndex != INDEX_NONE);
+			check(Interpolator->InterpolatedType & MCT_Float);
+
+			const EMaterialValueType Type = Interpolator->InterpolatedType == MCT_Float ? MCT_Float1 : Interpolator->InterpolatedType;
+			const TCHAR* Swizzle[2] = { TEXT("x"), TEXT("y") };
+			const int32 Offset = Interpolator->InterpolatorOffset;
+			const int32 Index = Interpolator->InterpolatorIndex;
+
+			// Note: We reference the UV define directly to avoid having to pre-accumulate UV counts before property translation
+			CustomInterpolatorAssignments += FString::Printf(TEXT("\tOutTexCoords[VERTEX_INTERPOLATOR_%i_TEXCOORDS_X].%s = VertexInterpolator%i(Parameters).x;") LINE_TERMINATOR, Index, Swizzle[Offset % 2], Index);
+
+			if (Type >= MCT_Float2)
+			{
+				CustomInterpolatorAssignments += FString::Printf(TEXT("\tOutTexCoords[VERTEX_INTERPOLATOR_%i_TEXCOORDS_Y].%s = VertexInterpolator%i(Parameters).y;") LINE_TERMINATOR, Index, Swizzle[(Offset + 1) % 2], Index);
+
+				if (Type >= MCT_Float3)
+				{
+					CustomInterpolatorAssignments += FString::Printf(TEXT("\tOutTexCoords[VERTEX_INTERPOLATOR_%i_TEXCOORDS_Z].%s = VertexInterpolator%i(Parameters).z;") LINE_TERMINATOR, Index, Swizzle[(Offset + 2) % 2], Index);
+
+					if (Type == MCT_Float4)
+					{
+						CustomInterpolatorAssignments += FString::Printf(TEXT("\tOutTexCoords[VERTEX_INTERPOLATOR_%i_TEXCOORDS_W].%s = VertexInterpolator%i(Parameters).w;") LINE_TERMINATOR, Index, Swizzle[(Offset + 3) % 2], Index);
+					}
+				}
+			}
+		}
+	}*/
+
+	LazyPrintf.PushParam(*CustomInterpolatorAssignments);
+
+	LazyPrintf.PushParam(*MaterialAttributesDefault);
+
+	//if (bEnableExecutionFlow)
+	{
+		FString EvaluateVertexCode;
+
+		// Set default texcoords in the VS
+		for (uint32 TexCoordIndex = 0; TexCoordIndex < NumUserVertexTexCoords; ++TexCoordIndex)
+		{
+			const FString AttributeName = FMaterialAttributeDefinitionMap::GetAttributeName((EMaterialProperty)(MP_CustomizedUVs0 + TexCoordIndex));
+
+			EvaluateVertexCode += FString::Printf(TEXT("\tDefaultMaterialAttributes.%s = Parameters.TexCoords[%d];") LINE_TERMINATOR, *AttributeName, TexCoordIndex);
+		}
+
+		//EvaluateVertexCode += TranslatedAttributesCodeChunks[SF_Vertex];
+
+		LazyPrintf.PushParam(*EvaluateVertexCode);
+		LazyPrintf.PushParam(PixelShaderCode);
+
+		FString EvaluateMaterialAttributesCode = TEXT("    FMaterialAttributes MaterialAttributes = EvaluatePixelMaterialAttributes(Parameters);" LINE_TERMINATOR);
+
+		for (int32 PropertyIndex = 0; PropertyIndex < MP_MAX; ++PropertyIndex)
+		{
+			// Skip non-shared properties
+			if (!SharedPixelProperties[PropertyIndex])
+			{
+				continue;
+			}
+
+			const EMaterialProperty Property = (EMaterialProperty)PropertyIndex;
+			check(FMaterialAttributeDefinitionMap::GetShaderFrequency(Property) == SF_Pixel);
+			// Special case MP_SubsurfaceColor as the actual property is a combination of the color and the profile but we don't want to expose the profile
+			const FString PropertyName = FMaterialAttributeDefinitionMap::GetAttributeName(Property);
+
+			if (PropertyIndex == MP_SubsurfaceColor)
+			{
+				// TODO - properly handle subsurface profile
+				EvaluateMaterialAttributesCode += FString::Printf("    PixelMaterialInputs.Subsurface = float4(MaterialAttributes.%s, 0.0f);" LINE_TERMINATOR, *PropertyName);
+			}
+			else
+			{
+				EvaluateMaterialAttributesCode += FString::Printf("    PixelMaterialInputs.%s = MaterialAttributes.%s;" LINE_TERMINATOR, *PropertyName, *PropertyName);
+			}
+		}
+
+		// TODO - deriv
+		for (int32 Iter = 0; Iter < CompiledPDV_MAX; Iter++)
+		{
+			LazyPrintf.PushParam(*EvaluateMaterialAttributesCode);
+			LazyPrintf.PushParam(TEXT(""));
+			LazyPrintf.PushParam(TEXT(""));
+		}
+	}
+
+	LazyPrintf.PushParam(*FString::Printf(TEXT("%u"), MaterialTemplateLineNumber));
+
+	return LazyPrintf.GetResultString();
+}
+
+static void GetMaterialEnvironment(EShaderPlatform InPlatform,
+	const FMaterial& InMaterial,
+	const FMaterialCompilationOutput& MaterialCompilationOutput,
+	FShaderCompilerEnvironment& OutEnvironment)
+{
+	bool bMaterialRequestsDualSourceBlending = false;
+
+	if (false)//bNeedsParticlePosition || Material->ShouldGenerateSphericalParticleNormals() || bUsesSphericalParticleOpacity)
+	{
+		OutEnvironment.SetDefine(TEXT("NEEDS_PARTICLE_POSITION"), 1);
+	}
+
+	if (false)//bNeedsParticleVelocity || Material->IsUsedWithNiagaraMeshParticles())
+	{
+		OutEnvironment.SetDefine(TEXT("NEEDS_PARTICLE_VELOCITY"), 1);
+	}
+
+	uint32 DynamicParticleParameterMask = 0u;
+	if (DynamicParticleParameterMask)
+	{
+		OutEnvironment.SetDefine(TEXT("USE_DYNAMIC_PARAMETERS"), 1);
+		OutEnvironment.SetDefine(TEXT("DYNAMIC_PARAMETERS_MASK"), DynamicParticleParameterMask);
+	}
+
+	if (false)//bNeedsParticleTime)
+	{
+		OutEnvironment.SetDefine(TEXT("NEEDS_PARTICLE_TIME"), 1);
+	}
+
+	if (false)//bUsesParticleMotionBlur)
+	{
+		OutEnvironment.SetDefine(TEXT("USES_PARTICLE_MOTION_BLUR"), 1);
+	}
+
+	if (false)//bNeedsParticleRandom)
+	{
+		OutEnvironment.SetDefine(TEXT("NEEDS_PARTICLE_RANDOM"), 1);
+	}
+
+	if (false)//bUsesSphericalParticleOpacity)
+	{
+		OutEnvironment.SetDefine(TEXT("SPHERICAL_PARTICLE_OPACITY"), TEXT("1"));
+	}
+
+	if (false)//bUsesParticleSubUVs)
+	{
+		OutEnvironment.SetDefine(TEXT("USE_PARTICLE_SUBUVS"), TEXT("1"));
+	}
+
+	if (false)//bUsesLightmapUVs)
+	{
+		OutEnvironment.SetDefine(TEXT("LIGHTMAP_UV_ACCESS"), TEXT("1"));
+	}
+
+	if (false)//bUsesAOMaterialMask)
+	{
+		OutEnvironment.SetDefine(TEXT("USES_AO_MATERIAL_MASK"), TEXT("1"));
+	}
+
+	if (false)//bUsesSpeedTree)
+	{
+		OutEnvironment.SetDefine(TEXT("USES_SPEEDTREE"), TEXT("1"));
+	}
+
+	if (false)//bNeedsWorldPositionExcludingShaderOffsets)
+	{
+		OutEnvironment.SetDefine(TEXT("NEEDS_WORLD_POSITION_EXCLUDING_SHADER_OFFSETS"), TEXT("1"));
+	}
+
+	if (false)//bNeedsParticleSize)
+	{
+		OutEnvironment.SetDefine(TEXT("NEEDS_PARTICLE_SIZE"), TEXT("1"));
+	}
+
+	if (MaterialCompilationOutput.bNeedsSceneTextures)
+	{
+		OutEnvironment.SetDefine(TEXT("NEEDS_SCENE_TEXTURES"), TEXT("1"));
+	}
+	if (MaterialCompilationOutput.bUsesEyeAdaptation)
+	{
+		OutEnvironment.SetDefine(TEXT("USES_EYE_ADAPTATION"), TEXT("1"));
+	}
+
+	if (MaterialCompilationOutput.bHasRuntimeVirtualTextureOutputNode)
+	{
+		OutEnvironment.SetDefine(TEXT("VIRTUAL_TEXTURE_OUTPUT"), 1);
+	}
+
+	OutEnvironment.SetDefine(TEXT("USES_PER_INSTANCE_CUSTOM_DATA"), false);// bUsesPerInstanceCustomData&& Material->IsUsedWithInstancedStaticMeshes());
+
+	// @todo MetalMRT: Remove this hack and implement proper atmospheric-fog solution for Metal MRT...
+	OutEnvironment.SetDefine(TEXT("MATERIAL_ATMOSPHERIC_FOG"), false);// !IsMetalMRTPlatform(InPlatform) ? bUsesAtmosphericFog : 0);
+	OutEnvironment.SetDefine(TEXT("MATERIAL_SKY_ATMOSPHERE"), false);// bUsesSkyAtmosphere);
+	OutEnvironment.SetDefine(TEXT("INTERPOLATE_VERTEX_COLOR"), false);// bUsesVertexColor);
+	OutEnvironment.SetDefine(TEXT("NEEDS_PARTICLE_COLOR"), false);// bUsesParticleColor);
+	OutEnvironment.SetDefine(TEXT("NEEDS_PARTICLE_LOCAL_TO_WORLD"), false);// bUsesParticleLocalToWorld);
+	OutEnvironment.SetDefine(TEXT("NEEDS_PARTICLE_WORLD_TO_LOCAL"), false);// bUsesParticleWorldToLocal);
+	OutEnvironment.SetDefine(TEXT("USES_TRANSFORM_VECTOR"), false);// bUsesTransformVector);
+	OutEnvironment.SetDefine(TEXT("WANT_PIXEL_DEPTH_OFFSET"), false);// bUsesPixelDepthOffset);
+	if (IsMetalPlatform(InPlatform))
+	{
+		OutEnvironment.SetDefine(TEXT("USES_WORLD_POSITION_OFFSET"), false);// bUsesWorldPositionOffset);
+	}
+	OutEnvironment.SetDefine(TEXT("USES_EMISSIVE_COLOR"), false);// bUsesEmissiveColor);
+	// Distortion uses tangent space transform 
+	OutEnvironment.SetDefine(TEXT("USES_DISTORTION"), InMaterial.IsDistorted());
+
+	OutEnvironment.SetDefine(TEXT("MATERIAL_ENABLE_TRANSLUCENCY_FOGGING"), InMaterial.ShouldApplyFogging());
+	OutEnvironment.SetDefine(TEXT("MATERIAL_ENABLE_TRANSLUCENCY_CLOUD_FOGGING"), InMaterial.ShouldApplyCloudFogging());
+	OutEnvironment.SetDefine(TEXT("MATERIAL_IS_SKY"), InMaterial.IsSky());
+	OutEnvironment.SetDefine(TEXT("MATERIAL_COMPUTE_FOG_PER_PIXEL"), InMaterial.ComputeFogPerPixel());
+	OutEnvironment.SetDefine(TEXT("MATERIAL_FULLY_ROUGH"), false);// bIsFullyRough || InMaterial.IsFullyRough());
+	OutEnvironment.SetDefine(TEXT("MATERIAL_USES_ANISOTROPY"), false);// bUsesAnisotropy);
+
+	// Count the number of VTStacks (each stack will allocate a feedback slot)
+	OutEnvironment.SetDefine(TEXT("NUM_VIRTUALTEXTURE_SAMPLES"), 0);// VTStacks.Num());
+
+	// Setup defines to map each VT stack to either 1 or 2 page table textures, depending on how many layers it uses
+	/*for (int i = 0; i < VTStacks.Num(); ++i)
+	{
+		const FMaterialVirtualTextureStack& Stack = MaterialCompilationOutput.UniformExpressionSet.VTStacks[i];
+		FString PageTableValue = FString::Printf(TEXT("Material.VirtualTexturePageTable0_%d"), i);
+		if (Stack.GetNumLayers() > 4u)
+		{
+			PageTableValue += FString::Printf(TEXT(", Material.VirtualTexturePageTable1_%d"), i);
+		}
+		if (VTStacks[i].bAdaptive)
+		{
+			PageTableValue += FString::Printf(TEXT(", Material.VirtualTexturePageTableIndirection_%d"), i);
+		}
+		OutEnvironment.SetDefine(*FString::Printf(TEXT("VIRTUALTEXTURE_PAGETABLE_%d"), i), *PageTableValue);
+	}*/
+
+	/*for (int32 CollectionIndex = 0; CollectionIndex < ParameterCollections.Num(); CollectionIndex++)
+	{
+		// Add uniform buffer declarations for any parameter collections referenced
+		const FString CollectionName = FString::Printf(TEXT("MaterialCollection%u"), CollectionIndex);
+		// This can potentially become an issue for MaterialCollection Uniform Buffers if they ever get non-numeric resources (eg Textures), as
+		// OutEnvironment.ResourceTableMap has a map by name, and the N ParameterCollection Uniform Buffers ALL are names "MaterialCollection"
+		// (and the hlsl cbuffers are named MaterialCollection0, etc, so the names don't match the layout)
+		FShaderUniformBufferParameter::ModifyCompilationEnvironment(*CollectionName, ParameterCollections[CollectionIndex]->GetUniformBufferStruct(), InPlatform, OutEnvironment);
+	}*/
+	OutEnvironment.SetDefine(TEXT("IS_MATERIAL_SHADER"), true);
+
+	// Set all the shading models for this material here 
+	FMaterialShadingModelField ShadingModels = InMaterial.GetShadingModels();
+
+	// If the material gets its shading model from the material expressions, then we use the result from the compilation (assuming it's valid).
+	// This result will potentially be tighter than what GetShadingModels() returns, because it only picks up the shading models from the expressions that get compiled for a specific feature level and quality level
+	// For example, the material might have shading models behind static switches. GetShadingModels() will return both the true and the false paths from that switch, whereas the shading model field from the compilation will only contain the actual shading model selected 
+	/*if (InMaterial.IsShadingModelFromMaterialExpression() && ShadingModelsFromCompilation.IsValid())
+	{
+		// Shading models fetched from the compilation of the expression graph
+		ShadingModels = ShadingModelsFromCompilation;
+	}*/
+
+	ensure(ShadingModels.IsValid());
+
+	if (ShadingModels.IsLit())
+	{
+		int32 NumSetMaterials = 0;
+		if (ShadingModels.HasShadingModel(MSM_DefaultLit))
+		{
+			OutEnvironment.SetDefine(TEXT("MATERIAL_SHADINGMODEL_DEFAULT_LIT"), TEXT("1"));
+			NumSetMaterials++;
+		}
+		if (ShadingModels.HasShadingModel(MSM_Subsurface))
+		{
+			OutEnvironment.SetDefine(TEXT("MATERIAL_SHADINGMODEL_SUBSURFACE"), TEXT("1"));
+			NumSetMaterials++;
+		}
+		if (ShadingModels.HasShadingModel(MSM_PreintegratedSkin))
+		{
+			OutEnvironment.SetDefine(TEXT("MATERIAL_SHADINGMODEL_PREINTEGRATED_SKIN"), TEXT("1"));
+			NumSetMaterials++;
+		}
+		if (ShadingModels.HasShadingModel(MSM_SubsurfaceProfile))
+		{
+			OutEnvironment.SetDefine(TEXT("MATERIAL_SHADINGMODEL_SUBSURFACE_PROFILE"), TEXT("1"));
+			NumSetMaterials++;
+		}
+		if (ShadingModels.HasShadingModel(MSM_ClearCoat))
+		{
+			OutEnvironment.SetDefine(TEXT("MATERIAL_SHADINGMODEL_CLEAR_COAT"), TEXT("1"));
+			NumSetMaterials++;
+		}
+		if (ShadingModels.HasShadingModel(MSM_TwoSidedFoliage))
+		{
+			OutEnvironment.SetDefine(TEXT("MATERIAL_SHADINGMODEL_TWOSIDED_FOLIAGE"), TEXT("1"));
+			NumSetMaterials++;
+		}
+		if (ShadingModels.HasShadingModel(MSM_Hair))
+		{
+			OutEnvironment.SetDefine(TEXT("MATERIAL_SHADINGMODEL_HAIR"), TEXT("1"));
+			NumSetMaterials++;
+		}
+		if (ShadingModels.HasShadingModel(MSM_Cloth))
+		{
+			OutEnvironment.SetDefine(TEXT("MATERIAL_SHADINGMODEL_CLOTH"), TEXT("1"));
+			NumSetMaterials++;
+		}
+		if (ShadingModels.HasShadingModel(MSM_Eye))
+		{
+			OutEnvironment.SetDefine(TEXT("MATERIAL_SHADINGMODEL_EYE"), TEXT("1"));
+			NumSetMaterials++;
+		}
+		if (ShadingModels.HasShadingModel(MSM_SingleLayerWater))
+		{
+			OutEnvironment.SetDefine(TEXT("MATERIAL_SHADINGMODEL_SINGLELAYERWATER"), TEXT("1"));
+			NumSetMaterials++;
+		}
+		if (ShadingModels.HasShadingModel(MSM_ThinTranslucent))
+		{
+			OutEnvironment.SetDefine(TEXT("MATERIAL_SHADINGMODEL_THIN_TRANSLUCENT"), TEXT("1"));
+			NumSetMaterials++;
+
+			bMaterialRequestsDualSourceBlending = true;
+		}
+
+		if (ShadingModels.HasShadingModel(MSM_SingleLayerWater) &&
+			(IsSwitchPlatform(InPlatform) || IsVulkanMobileSM5Platform(InPlatform) || FDataDrivenShaderPlatformInfo::GetRequiresDisableForwardLocalLights(InPlatform)))
+		{
+			OutEnvironment.SetDefine(TEXT("DISABLE_FORWARD_LOCAL_LIGHTS"), TEXT("1"));
+		}
+
+		// This is to have switch use the simple single layer water shading similar to mobile: no dynamic lights, only sun and sky, no distortion, no colored transmittance on background, no custom depth read.
+		const bool bSingleLayerWaterUsesSimpleShading = (IsSwitchPlatform(InPlatform) || IsVulkanMobileSM5Platform(InPlatform)) && IsForwardShadingEnabled(InPlatform);
+		if (ShadingModels.HasShadingModel(MSM_SingleLayerWater) && bSingleLayerWaterUsesSimpleShading)
+		{
+			OutEnvironment.SetDefine(TEXT("SINGLE_LAYER_WATER_SIMPLE_FORWARD"), TEXT("1"));
+		}
+
+		if (NumSetMaterials == 1)
+		{
+			OutEnvironment.SetDefine(TEXT("MATERIAL_SINGLE_SHADINGMODEL"), TEXT("1"));
+		}
+
+		ensure(NumSetMaterials != 0);
+		if (NumSetMaterials == 0)
+		{
+			// Should not really end up here
+			UE_LOG(LogMaterial, Warning, TEXT("Unknown material shading model(s). Setting to MSM_DefaultLit"));
+			OutEnvironment.SetDefine(TEXT("MATERIAL_SHADINGMODEL_DEFAULT_LIT"), TEXT("1"));
+		}
+	}
+	else
+	{
+		// Unlit shading model can only exist by itself
+		OutEnvironment.SetDefine(TEXT("MATERIAL_SINGLE_SHADINGMODEL"), TEXT("1"));
+		OutEnvironment.SetDefine(TEXT("MATERIAL_SHADINGMODEL_UNLIT"), TEXT("1"));
+	}
+
+	if (InMaterial.GetMaterialDomain() == MD_Volume) // && Material->HasN)
+	{
+		TArray<const UMaterialExpressionVolumetricAdvancedMaterialOutput*> VolumetricAdvancedExpressions;
+		InMaterial.GetMaterialInterface()->GetMaterial()->GetAllExpressionsOfType(VolumetricAdvancedExpressions);
+		if (VolumetricAdvancedExpressions.Num() > 0)
+		{
+			if (VolumetricAdvancedExpressions.Num() > 1)
+			{
+				UE_LOG(LogMaterial, Fatal, TEXT("Only a single UMaterialExpressionVolumetricAdvancedMaterialOutput node is supported."));
+			}
+
+			OutEnvironment.SetDefine(TEXT("MATERIAL_VOLUMETRIC_ADVANCED"), TEXT("1"));
+
+			const UMaterialExpressionVolumetricAdvancedMaterialOutput* VolumetricAdvancedNode = VolumetricAdvancedExpressions[0];
+			if (VolumetricAdvancedNode->GetEvaluatePhaseOncePerSample())
+			{
+				OutEnvironment.SetDefine(TEXT("MATERIAL_VOLUMETRIC_ADVANCED_PHASE_PERSAMPLE"), TEXT("1"));
+			}
+			else
+			{
+				OutEnvironment.SetDefine(TEXT("MATERIAL_VOLUMETRIC_ADVANCED_PHASE_PERPIXEL"), TEXT("1"));
+			}
+
+			OutEnvironment.SetDefine(TEXT("MATERIAL_VOLUMETRIC_ADVANCED_GRAYSCALE_MATERIAL"), VolumetricAdvancedNode->bGrayScaleMaterial ? TEXT("1") : TEXT("0"));
+			OutEnvironment.SetDefine(TEXT("MATERIAL_VOLUMETRIC_ADVANCED_RAYMARCH_VOLUME_SHADOW"), VolumetricAdvancedNode->bRayMarchVolumeShadow ? TEXT("1") : TEXT("0"));
+
+			OutEnvironment.SetDefine(TEXT("MATERIAL_VOLUMETRIC_ADVANCED_MULTISCATTERING_OCTAVE_COUNT"), VolumetricAdvancedNode->GetMultiScatteringApproximationOctaveCount());
+
+			OutEnvironment.SetDefine(TEXT("MATERIAL_VOLUMETRIC_ADVANCED_CONSERVATIVE_DENSITY"),
+				VolumetricAdvancedNode->ConservativeDensity.IsConnected() ? TEXT("1") : TEXT("0"));
+
+			OutEnvironment.SetDefine(TEXT("MATERIAL_VOLUMETRIC_ADVANCED_OVERRIDE_AMBIENT_OCCLUSION"),
+				InMaterial.HasAmbientOcclusionConnected() ? TEXT("1") : TEXT("0"));
+
+			OutEnvironment.SetDefine(TEXT("MATERIAL_VOLUMETRIC_ADVANCED_GROUND_CONTRIBUTION"),
+				VolumetricAdvancedNode->bGroundContribution ? TEXT("1") : TEXT("0"));
+		}
+	}
+
+	OutEnvironment.SetDefine(TEXT("MATERIAL_IS_STRATA"), false);// bMaterialIsStrata ? TEXT("1") : TEXT("0"));
+
+	// STRATA_TODO do not request DualSourceBlending if gray scale transmittance is selected.
+	// bMaterialRequestsDualSourceBlending this base on limited set of blend mode: Opaque, Masked, TransmittanceCoverage, TransmittanceColored;
+	//bMaterialRequestsDualSourceBlending |= bMaterialIsStrata;
+
+	// if duals source blending (colored transmittance) is not supported on a platform, it will fall back to standard alpha blending (grey scale transmittance)
+	OutEnvironment.SetDefine(TEXT("DUAL_SOURCE_COLOR_BLENDING_ENABLED"), false);// bMaterialRequestsDualSourceBlending&& Material->IsDualBlendingEnabled(Platform) ? TEXT("1") : TEXT("0"));
+
+	// Translate() is called before getting here so we can create related define
+	/*if (StrataMaterialAnalysis.RequestedBSDFCount > 0)
+	{
+		OutEnvironment.SetDefine(TEXT("STRATA_CLAMPED_LAYER_COUNT"), StrataMaterialAnalysis.ClampedLayerCount);
+	}*/
+}
+
+
+bool FMaterialHLSLGenerator::GenerateHLSL(FMaterial& InMaterial,
+	FMaterialCompilationOutput& OutCompilationOutput,
+	TRefCountPtr<FSharedShaderCompilerEnvironment>& OutMaterialEnvironment)
+{
+	UMaterialInterface* MaterialInterface = InMaterial.GetMaterialInterface();
+	UMaterial* BaseMaterial = MaterialInterface->GetMaterial();
+
+	FMemMark MemMark(Allocator);
+
+	HLSLTree = UE::HLSLTree::FTree::Create(Allocator);
+
+	UMaterialExpression* BaseExpression = BaseMaterial->ExpressionExecBegin;
+	UE::HLSLTree::FStatement* RootStatement = AcquireStatement(HLSLTree->GetRootScope(), BaseExpression);
+
+	UE::HLSLTree::FCodeWriter* CodeWriter = UE::HLSLTree::FCodeWriter::Create(Allocator);
+	
+	UE::HLSLTree::FEmitContext EmitContext;
+	EmitContext.Allocator = &Allocator;
+	EmitContext.Material = &InMaterial;
+	EmitContext.MaterialCompilationOutput = &OutCompilationOutput;
+	HLSLTree->EmitHLSL(EmitContext, *CodeWriter);
+
+	InMaterial.CompileErrors = MoveTemp(CompileErrors);
+	InMaterial.ErrorExpressions = MoveTemp(ErrorExpressions);
+
+	const FString MaterialTemplateSource = GenerateMaterialTemplateHLSL(ShaderPlatform, InMaterial, EmitContext, CodeWriter->StringBuilder->ToString());
+
+	OutMaterialEnvironment = new FSharedShaderCompilerEnvironment();
+	OutMaterialEnvironment->TargetPlatform = TargetPlatform;
+	GetMaterialEnvironment(ShaderPlatform, InMaterial, OutCompilationOutput, *OutMaterialEnvironment);
+	OutMaterialEnvironment->IncludeVirtualPathToContentsMap.Add(TEXT("/Engine/Generated/Material.ush"), MaterialTemplateSource);
+
+	return true;
+}
+
+EMaterialGenerateHLSLStatus FMaterialHLSLGenerator::Error(const FString& Message)
+{
+	UMaterialExpression* ExpressionToError = nullptr;
+	FString ErrorString;
+
+	if (ExpressionStack.Num() > 0)
+	{
+		UMaterialExpression* ErrorExpression = ExpressionStack.Last().Expression;
+		check(ErrorExpression);
+
+		if (ErrorExpression->GetClass() != UMaterialExpressionMaterialFunctionCall::StaticClass()
+			&& ErrorExpression->GetClass() != UMaterialExpressionFunctionInput::StaticClass()
+			&& ErrorExpression->GetClass() != UMaterialExpressionFunctionOutput::StaticClass())
+		{
+			// Add the expression currently being compiled to ErrorExpressions so we can draw it differently
+			ExpressionToError = ErrorExpression;
+
+			const int32 ChopCount = FCString::Strlen(TEXT("MaterialExpression"));
+			const FString ErrorClassName = ErrorExpression->GetClass()->GetName();
+
+			// Add the node type to the error message
+			ErrorString += FString(TEXT("(Node ")) + ErrorClassName.Right(ErrorClassName.Len() - ChopCount) + TEXT(") ");
+		}
+	}
+
+	ErrorString += Message;
+
+	// Standard error handling, immediately append one-off errors and signal failure
+	CompileErrors.AddUnique(ErrorString);
+
+	if (ExpressionToError)
+	{
+		ErrorExpressions.Add(ExpressionToError);
+		ExpressionToError->LastErrorText = Message;
+	}
+
+	return EMaterialGenerateHLSLStatus::Error;
+}
+
+UE::HLSLTree::FExpressionConstant* FMaterialHLSLGenerator::NewConstant(UE::HLSLTree::FScope& Scope, const UE::HLSLTree::FConstant& Value)
+{
+	UE::HLSLTree::FExpressionConstant* Expression = HLSLTree->NewExpression<UE::HLSLTree::FExpressionConstant>(Scope, Value);
+	Expression->bInline = true;
+	return Expression;
+}
+
+UE::HLSLTree::FExpressionExternalInput* FMaterialHLSLGenerator::NewTexCoord(UE::HLSLTree::FScope& Scope, int32 Index)
+{
+	UE::HLSLTree::FExpressionExternalInput* Expression = HLSLTree->NewExpression<UE::HLSLTree::FExpressionExternalInput>(Scope, UE::HLSLTree::MakeInputTexCoord(Index));
+	Expression->bInline = true;
+	return Expression;
+}
+
+UE::HLSLTree::FExpressionSwizzle* FMaterialHLSLGenerator::NewSwizzle(UE::HLSLTree::FScope& Scope, const UE::HLSLTree::FSwizzleParameters& Params, UE::HLSLTree::FExpression* Input)
+{
+	UE::HLSLTree::FExpressionSwizzle* Expression = HLSLTree->NewExpression<UE::HLSLTree::FExpressionSwizzle>(Scope, Params, Input);
+	Expression->bInline = true;
+	return Expression;
+}
+
+UE::HLSLTree::FExpression* FMaterialHLSLGenerator::NewCast(UE::HLSLTree::FScope& Scope, UE::HLSLTree::EExpressionType Type, UE::HLSLTree::FExpression* Input, UE::HLSLTree::ECastFlags Flags)
+{
+	// Only generate the cast expression if the types don't match
+	if (Input && Type != Input->Type)
+	{
+		UE::HLSLTree::FExpressionCast* Expression = HLSLTree->NewExpression<UE::HLSLTree::FExpressionCast>(Scope, Type, Input, Flags);
+		Expression->bInline = true;
+		return Expression;
+	}
+	return Input;
+}
+
+UE::HLSLTree::FLocalDeclaration* FMaterialHLSLGenerator::AcquireLocalDeclaration(UE::HLSLTree::FScope& Scope, UE::HLSLTree::EExpressionType Type, const FName& Name)
+{
+	UE::HLSLTree::FLocalDeclaration*& Declaration = LocalDeclarationMap.FindOrAdd(Name);
+	if (!Declaration)
+	{
+		Declaration = HLSLTree->NewLocalDeclaration(Scope, Type, Name);
+	}
+	else
+	{
+		if (Declaration->Type == Type)
+		{
+			Scope.UseDeclaration(Declaration);
+		}
+		else
+		{
+			check(false);
+			return nullptr;
+		}
+	}
+	return Declaration;
+}
+
+UE::HLSLTree::FParameterDeclaration* FMaterialHLSLGenerator::AcquireParameterDeclaration(UE::HLSLTree::FScope& Scope, const FName& Name, const UE::HLSLTree::FConstant& DefaultValue)
+{
+	UE::HLSLTree::FParameterDeclaration*& Declaration = ParameterDeclarationMap.FindOrAdd(Name);
+	if (!Declaration)
+	{
+		Declaration = HLSLTree->NewParameterDeclaration(Scope, Name, DefaultValue);
+	}
+	return Declaration;
+}
+
+UE::HLSLTree::FTextureParameterDeclaration* FMaterialHLSLGenerator::AcquireTextureDeclaration(UE::HLSLTree::FScope& Scope, const UE::HLSLTree::FTextureDescription& Value)
+{
+	FString SamplerTypeError;
+	if (!UMaterialExpressionTextureBase::VerifySamplerType(FeatureLevel, TargetPlatform, Value.Texture, Value.SamplerType, SamplerTypeError))
+	{
+		Errorf(TEXT("%s"), *SamplerTypeError);
+		return nullptr;
+	}
+
+	UE::HLSLTree::FTextureParameterDeclaration*& Declaration = TextureDeclarationMap.FindOrAdd(Value);
+	if (!Declaration)
+	{
+		Declaration = HLSLTree->NewTextureParameterDeclaration(Scope, FName(), Value);
+	}
+	return Declaration;
+}
+
+UE::HLSLTree::FTextureParameterDeclaration* FMaterialHLSLGenerator::AcquireTextureParameterDeclaration(UE::HLSLTree::FScope& Scope, const FName& Name, const UE::HLSLTree::FTextureDescription& DefaultValue)
+{
+	FString SamplerTypeError;
+	if (!UMaterialExpressionTextureBase::VerifySamplerType(FeatureLevel, TargetPlatform, DefaultValue.Texture, DefaultValue.SamplerType, SamplerTypeError))
+	{
+		Errorf(TEXT("%s"), *SamplerTypeError);
+		return nullptr;
+	}
+
+	UE::HLSLTree::FTextureParameterDeclaration*& Declaration = TextureParameterDeclarationMap.FindOrAdd(Name);
+	if (!Declaration)
+	{
+		Declaration = HLSLTree->NewTextureParameterDeclaration(Scope, Name, DefaultValue);
+	}
+	return Declaration;
+}
+
+UE::HLSLTree::FExpression* FMaterialHLSLGenerator::AcquireExpression(UE::HLSLTree::FScope& Scope, UMaterialExpression* MaterialExpression, int32 OutputIndex)
+{
+	const FExpressionKey Key(MaterialExpression, OutputIndex);
+	UE::HLSLTree::FExpression** PrevExpression = ExpressionMap.Find(Key);
+	UE::HLSLTree::FExpression* Expression = nullptr;
+	if (!PrevExpression)
+	{
+		ExpressionStack.Add(Key);
+		const EMaterialGenerateHLSLStatus Status = MaterialExpression->GenerateHLSLExpression(*this, Scope, OutputIndex, Expression);
+		verify(ExpressionStack.Pop() == Key);
+		ExpressionMap.Add(Key, Expression);
+	}
+	else
+	{
+		Expression = *PrevExpression;
+		Scope.UseExpression(Expression);
+	}
+	return Expression;
+}
+
+UE::HLSLTree::FTextureParameterDeclaration* FMaterialHLSLGenerator::AcquireTextureDeclaration(UE::HLSLTree::FScope& Scope, UMaterialExpression* MaterialExpression, int32 OutputIndex)
+{
+	// No need to cache at this level, TextureDeclarations are cached at a lower level, as they're generated by UMaterialExpression
+	UE::HLSLTree::FTextureParameterDeclaration* TextureDeclaration = nullptr;
+	const EMaterialGenerateHLSLStatus Status = MaterialExpression->GenerateHLSLTexture(*this, Scope, OutputIndex, TextureDeclaration);
+	return TextureDeclaration;
+}
+
+UE::HLSLTree::FStatement* FMaterialHLSLGenerator::AcquireStatement(UE::HLSLTree::FScope& Scope, UMaterialExpression* MaterialExpression)
+{
+	UE::HLSLTree::FStatement** PrevStatement = StatementMap.Find(MaterialExpression);
+	UE::HLSLTree::FStatement* Statement = nullptr;
+	if (!PrevStatement)
+	{
+		const FExpressionKey Key(MaterialExpression);
+		ExpressionStack.Add(Key);
+		const EMaterialGenerateHLSLStatus Status = MaterialExpression->GenerateHLSLStatement(*this, Scope, Statement);
+		verify(ExpressionStack.Pop() == Key);
+		StatementMap.Add(MaterialExpression, Statement);
+	}
+	else
+	{
+		Statement = *PrevStatement;
+		if (!Scope.TryMoveStatement(Statement))
+		{
+			// Could not move existing statement to the given scope
+			Error(TEXT("Invalid control flow"));
+			return nullptr;
+		}
+	}
+
+	return Statement;
+}
+
+#endif // WITH_EDITOR
diff --git a/Engine/Source/Runtime/Engine/Private/Materials/MaterialShader.cpp b/Engine/Source/Runtime/Engine/Private/Materials/MaterialShader.cpp
--- a/Engine/Source/Runtime/Engine/Private/Materials/MaterialShader.cpp
+++ b/Engine/Source/Runtime/Engine/Private/Materials/MaterialShader.cpp
@@ -292,6 +292,8 @@ void FMaterialShaderMapId::Serialize(FArchive& Ar, bool bLoadedByCookedMaterial)
 	Ar.UsingCustomVersion(FEditorObjectVersion::GUID);
 	Ar.UsingCustomVersion(FReleaseObjectVersion::GUID);
 
+	const bool bIsLegacyPackage = Ar.UE4Ver() < VER_UE4_PURGED_FMATERIAL_COMPILE_OUTPUTS;
+
 	// Ensure saved content is correct
 	check(!Ar.IsSaving() || IsContentValid());
 
@@ -314,7 +316,7 @@ void FMaterialShaderMapId::Serialize(FArchive& Ar, bool bLoadedByCookedMaterial)
 	}
 #endif
 
-	if (Ar.UE4Ver() >= VER_UE4_PURGED_FMATERIAL_COMPILE_OUTPUTS)
+	if (!bIsLegacyPackage)
 	{
 		Ar << (int32&)QualityLevel;
 		Ar << (int32&)FeatureLevel;
@@ -360,13 +362,13 @@ void FMaterialShaderMapId::Serialize(FArchive& Ar, bool bLoadedByCookedMaterial)
 		}
 
 		Ar << ShaderTypeDependencies;
-		if (Ar.UE4Ver() >= VER_UE4_PURGED_FMATERIAL_COMPILE_OUTPUTS)
+		if (!bIsLegacyPackage)
 		{
 			Ar << ShaderPipelineTypeDependencies;
 		}
 		Ar << VertexFactoryTypeDependencies;
 
-		if (Ar.UE4Ver() >= VER_UE4_PURGED_FMATERIAL_COMPILE_OUTPUTS)
+		if (!bIsLegacyPackage)
 		{
 			Ar << TextureReferencesHash;
 		}
@@ -380,6 +382,15 @@ void FMaterialShaderMapId::Serialize(FArchive& Ar, bool bLoadedByCookedMaterial)
 		{
 			Ar << BasePropertyOverridesHash;
 		}
+
+		if (!bIsLegacyPackage)
+		{
+			Ar << bUsingNewHLSLGenerator;
+		}
+		else
+		{
+			bUsingNewHLSLGenerator = false;
+		}
 	}
 	else
 	{
@@ -398,10 +409,7 @@ void FMaterialShaderMapId::Serialize(FArchive& Ar, bool bLoadedByCookedMaterial)
 	checkf(CookedShaderMapIdHash != FSHAHash(), TEXT("Loaded an invalid cooked shadermap id hash"));
 #endif // WITH_EDITOR
 
-	// Based on the comment above, ShaderMapId would only be embedded into packages for archives < VER_UE4_PURGED_FMATERIAL_COMPILE_OUTPUTS (~2013 era).
-	// Since backward compatibility only works for the packages, we can assume that any archive with a version newer than that will have
-	// the LayoutParams serialized. The other option (old DDC) should be prevented by us having mutated MATERIALSHADERMAP_DERIVEDDATA_VER
-	if (Ar.UE4Ver() >= VER_UE4_PURGED_FMATERIAL_COMPILE_OUTPUTS)
+	if (!bIsLegacyPackage)
 	{
 		Ar << LayoutParams;
 	}
@@ -473,6 +481,8 @@ void FMaterialShaderMapId::GetMaterialHash(FSHAHash& OutHash) const
 
 	HashState.Update((const uint8*)&BasePropertyOverridesHash, sizeof(BasePropertyOverridesHash));
 
+	HashState.Update((const uint8*)&bUsingNewHLSLGenerator, sizeof(bUsingNewHLSLGenerator));
+
 	HashState.Final();
 	HashState.GetHash(&OutHash.Hash[0]);
 }
@@ -495,6 +505,11 @@ bool FMaterialShaderMapId::operator==(const FMaterialShaderMapId& ReferenceSet)
 		return false;
 	}
 
+	if (bUsingNewHLSLGenerator != ReferenceSet.bUsingNewHLSLGenerator)
+	{
+		return false;
+	}
+
 	if (!IsCookedId())
 	{
 		if (Usage != ReferenceSet.Usage
@@ -857,6 +872,11 @@ void FMaterialShaderMapId::AppendKeyString(FString& KeyString) const
 	KeyString += BytesToHex(&TextureReferencesHash.Hash[0], sizeof(TextureReferencesHash.Hash));
 
 	KeyString += BytesToHex(&BasePropertyOverridesHash.Hash[0], sizeof(BasePropertyOverridesHash.Hash));
+
+	if (bUsingNewHLSLGenerator)
+	{
+		KeyString += TEXT("_NewHLSL");
+	}
 }
 #endif // WITH_EDITOR
 
diff --git a/Engine/Source/Runtime/Engine/Private/Materials/MaterialShared.cpp b/Engine/Source/Runtime/Engine/Private/Materials/MaterialShared.cpp
--- a/Engine/Source/Runtime/Engine/Private/Materials/MaterialShared.cpp
+++ b/Engine/Source/Runtime/Engine/Private/Materials/MaterialShared.cpp
@@ -48,6 +48,8 @@
 #include "Misc/ConfigCacheIni.h"
 #if WITH_EDITOR
 #include "Rendering/StaticLightingSystemInterface.h"
+#include "MaterialHLSLGenerator.h"
+#include "HLSLTree/HLSLTreeCommon.h"
 #endif
 #if WITH_ODSC
 #include "ODSC/ODSCManager.h"
@@ -258,6 +260,41 @@ int32 FExpressionInput::Compile(class FMaterialCompiler* Compiler)
 		return INDEX_NONE;
 }
 
+UE::HLSLTree::FExpression* FExpressionInput::AcquireHLSLExpression(FMaterialHLSLGenerator& Generator, UE::HLSLTree::FScope& Scope) const
+{
+	UE::HLSLTree::FExpression* Result = nullptr;
+	if (Expression)
+	{
+		Expression->ValidateState();
+		Result = Generator.AcquireExpression(Scope, Expression, OutputIndex);
+		if (Mask && Result)
+		{
+			const UE::HLSLTree::FSwizzleParameters SwizzleParams = UE::HLSLTree::MakeSwizzleMask(!!MaskR, !!MaskG, !!MaskB, !!MaskA);
+			Result = Generator.NewSwizzle(Scope, SwizzleParams, Result);
+		}
+	}
+	
+	return Result;
+}
+
+UE::HLSLTree::FExpression* FExpressionInput::AcquireHLSLExpressionWithCast(FMaterialHLSLGenerator& Generator, UE::HLSLTree::FScope& Scope, UE::HLSLTree::EExpressionType Type) const
+{
+	UE::HLSLTree::FExpression* Result = AcquireHLSLExpression(Generator, Scope);
+	return Generator.NewCast(Scope, Type, Result);
+}
+
+UE::HLSLTree::FTextureParameterDeclaration* FExpressionInput::AcquireHLSLTexture(FMaterialHLSLGenerator& Generator, UE::HLSLTree::FScope& Scope) const
+{
+	UE::HLSLTree::FTextureParameterDeclaration* Result = nullptr;
+	if (Expression)
+	{
+		Expression->ValidateState();
+		Result = Generator.AcquireTextureDeclaration(Scope, Expression, OutputIndex);
+	}
+
+	return Result;
+}
+
 void FExpressionInput::Connect( int32 InOutputIndex, class UMaterialExpression* InExpression )
 {
 	InExpression->ConnectExpression(this, InOutputIndex);
@@ -277,7 +314,7 @@ FExpressionInput FExpressionInput::GetTracedInput() const
 }
 
 #if WITH_EDITOR
-int32 FExpressionExecOutput::Compile(class FMaterialCompiler* Compiler)
+int32 FExpressionExecOutput::Compile(class FMaterialCompiler* Compiler) const
 {
 	if (Expression)
 	{
@@ -285,6 +322,45 @@ int32 FExpressionExecOutput::Compile(class FMaterialCompiler* Compiler)
 	}
 	return INDEX_NONE;
 }
+
+UE::HLSLTree::FStatement* FExpressionExecOutput::AcquireHLSLStatement(FMaterialHLSLGenerator& Generator, UE::HLSLTree::FScope& Scope) const
+{
+	UE::HLSLTree::FStatement* Result = nullptr;
+	if (Expression)
+	{
+		Expression->ValidateState();
+		Result = Generator.AcquireStatement(Scope, Expression);
+	}
+
+	return Result;
+}
+
+UE::HLSLTree::FScope* FExpressionExecOutput::NewScopeWithStatement(FMaterialHLSLGenerator& Generator, UE::HLSLTree::FScope& Scope) const
+{
+	UE::HLSLTree::FScope* Result = nullptr;
+	if (Expression)
+	{
+		Expression->ValidateState();
+		Result = Generator.GetTree().NewScope(Scope); // Create a new scope for the statement
+		Generator.AcquireStatement(*Result, Expression);
+	}
+
+	return Result;
+}
+
+UE::HLSLTree::FScope* FExpressionExecOutput::NewLinkedScopeWithStatement(FMaterialHLSLGenerator& Generator, UE::HLSLTree::FScope& Scope) const
+{
+	UE::HLSLTree::FScope* Result = nullptr;
+	if (Expression)
+	{
+		Expression->ValidateState();
+		Result = Generator.GetTree().NewLinkedScope(Scope); // Create a new scope for the statement
+		Generator.AcquireStatement(*Result, Expression);
+	}
+
+	return Result;
+}
+
 #endif // WITH_EDITOR
 
 /** Native serialize for FMaterialExpression struct */
@@ -561,6 +637,7 @@ void FMaterial::GetShaderMapId(EShaderPlatform Platform, const ITargetPlatform*
 		GetDependentShaderAndVFTypes(Platform, ShaderTypes, ShaderPipelineTypes, VFTypes);
 
 		OutId.Usage = GetShaderMapUsage();
+		OutId.bUsingNewHLSLGenerator = IsUsingNewHLSLGenerator();
 		OutId.BaseMaterialId = GetMaterialId();
 		OutId.QualityLevel = GetQualityLevel();
 		OutId.FeatureLevel = GetFeatureLevel();
@@ -1681,6 +1758,15 @@ bool FMaterialResource::IsCompiledWithExecutionFlow() const
 	}
 	return false;
 }
+
+bool FMaterialResource::IsUsingNewHLSLGenerator() const
+{
+	if (Material)
+	{
+		return Material->IsUsingNewHLSLGenerator();
+	}
+	return false;
+}
 #endif
 
 FString FMaterialResource::GetFullPath() const
@@ -2189,6 +2275,51 @@ bool FMaterial::CacheShaders(const FMaterialShaderMapId& ShaderMapId, EShaderPla
 	return bShaderMapValid;
 }
 
+bool FMaterial::Translate_Legacy(const FMaterialShaderMapId& ShaderMapId,
+	const FStaticParameterSet& InStaticParameters,
+	EShaderPlatform InPlatform,
+	const ITargetPlatform* InTargetPlatform,
+	FMaterialCompilationOutput& OutCompilationOutput,
+	TRefCountPtr<FSharedShaderCompilerEnvironment>& OutMaterialEnvironment)
+{
+#if WITH_EDITORONLY_DATA
+	FHLSLMaterialTranslator MaterialTranslator(this, OutCompilationOutput, InStaticParameters, InPlatform, GetQualityLevel(), ShaderMapId.FeatureLevel, InTargetPlatform);
+	const bool bSuccess = MaterialTranslator.Translate();
+	if (bSuccess)
+	{
+		// Create a shader compiler environment for the material that will be shared by all jobs from this material
+		OutMaterialEnvironment = new FSharedShaderCompilerEnvironment();
+		OutMaterialEnvironment->TargetPlatform = InTargetPlatform;
+		MaterialTranslator.GetMaterialEnvironment(InPlatform, *OutMaterialEnvironment);
+		const FString MaterialShaderCode = MaterialTranslator.GetMaterialShaderCode();
+
+		OutMaterialEnvironment->IncludeVirtualPathToContentsMap.Add(TEXT("/Engine/Generated/Material.ush"), MaterialShaderCode);
+	}
+	return bSuccess;
+#else
+	checkNoEntry();
+	return false;
+#endif
+}
+
+bool FMaterial::Translate_New(const FMaterialShaderMapId& ShaderMapId,
+	const FStaticParameterSet& InStaticParameters,
+	EShaderPlatform InPlatform,
+	const ITargetPlatform* InTargetPlatform,
+	FMaterialCompilationOutput& OutCompilationOutput,
+	TRefCountPtr<FSharedShaderCompilerEnvironment>& OutMaterialEnvironment)
+{
+#if WITH_EDITOR
+	FMaterialHLSLGenerator Generator(InPlatform, ShaderMapId.FeatureLevel, InTargetPlatform);
+	return Generator.GenerateHLSL(*this,
+		OutCompilationOutput,
+		OutMaterialEnvironment);
+#else
+	checkNoEntry();
+	return false;
+#endif
+}
+
 /**
 * Compiles this material for Platform
 *
@@ -2216,21 +2347,21 @@ bool FMaterial::BeginCompileShaderMap(
 #if WITH_EDITOR
 	NewShaderMap->AssociateWithAsset(GetAssetPath());
 #endif
+
 	// Generate the material shader code.
 	FMaterialCompilationOutput NewCompilationOutput;
-	FHLSLMaterialTranslator MaterialTranslator(this, NewCompilationOutput, StaticParameterSet, Platform,GetQualityLevel(), ShaderMapId.FeatureLevel, TargetPlatform);
-	bSuccess = MaterialTranslator.Translate();
+	TRefCountPtr<FSharedShaderCompilerEnvironment> MaterialEnvironment;
+	if (ShaderMapId.bUsingNewHLSLGenerator)
+	{
+		bSuccess = Translate_New(ShaderMapId, StaticParameterSet, Platform, TargetPlatform, NewCompilationOutput, MaterialEnvironment);
+	}
+	else
+	{
+		bSuccess = Translate_Legacy(ShaderMapId, StaticParameterSet, Platform, TargetPlatform, NewCompilationOutput, MaterialEnvironment);
+	}
 
 	if(bSuccess)
 	{
-		// Create a shader compiler environment for the material that will be shared by all jobs from this material
-		TRefCountPtr<FSharedShaderCompilerEnvironment> MaterialEnvironment = new FSharedShaderCompilerEnvironment();
-		MaterialEnvironment->TargetPlatform = TargetPlatform;
-		MaterialTranslator.GetMaterialEnvironment(Platform, *MaterialEnvironment);
-		const FString MaterialShaderCode = MaterialTranslator.GetMaterialShaderCode();
-
-		MaterialEnvironment->IncludeVirtualPathToContentsMap.Add(TEXT("/Engine/Generated/Material.ush"), MaterialShaderCode);
-
 		FShaderCompileUtilities::GenerateBrdfHeaders((EShaderPlatform)Platform);
 		FShaderCompileUtilities::ApplyDerivedDefines(*MaterialEnvironment, nullptr, (EShaderPlatform)Platform);
 
diff --git a/Engine/Source/Runtime/Engine/Private/Materials/MaterialUniformExpressions.cpp b/Engine/Source/Runtime/Engine/Private/Materials/MaterialUniformExpressions.cpp
--- a/Engine/Source/Runtime/Engine/Private/Materials/MaterialUniformExpressions.cpp
+++ b/Engine/Source/Runtime/Engine/Private/Materials/MaterialUniformExpressions.cpp
@@ -1024,14 +1024,18 @@ static void EvaluatePreshader(const FUniformExpressionSet* UniformExpressionSet,
 	}
 }
 
+void FMaterialPreshaderData::Evaluate(FUniformExpressionSet* UniformExpressionSet, const struct FMaterialRenderContext& Context, FLinearColor& OutValue)
+{
+	FPreshaderStack Stack;
+	FPreshaderDataContext PreshaderContext(*this);
+	EvaluatePreshader(UniformExpressionSet, Context, Stack, PreshaderContext, OutValue);
+}
+
 void FMaterialUniformExpression::GetNumberValue(const struct FMaterialRenderContext& Context, FLinearColor& OutValue) const
 {
 	FMaterialPreshaderData PreshaderData;
 	WriteNumberOpcodes(PreshaderData);
-
-	FPreshaderStack Stack;
-	FPreshaderDataContext PreshaderContext(PreshaderData);
-	EvaluatePreshader(nullptr, Context, Stack, PreshaderContext, OutValue);
+	PreshaderData.Evaluate(nullptr, Context, OutValue);
 }
 
 const FMaterialVectorParameterInfo* FUniformExpressionSet::FindVectorParameter(const FHashedMaterialParameterInfo& ParameterInfo) const
@@ -1058,6 +1062,53 @@ const FMaterialScalarParameterInfo* FUniformExpressionSet::FindScalarParameter(c
 	return nullptr;
 }
 
+int32 FUniformExpressionSet::FindOrAddScalarParameter(const FHashedMaterialParameterInfo& ParameterInfo, float DefaultValue)
+{
+	for (int32 i = 0; i < UniformScalarParameters.Num(); ++i)
+	{
+		if (UniformScalarParameters[i].ParameterInfo == ParameterInfo)
+		{
+			return i;
+		}
+	}
+
+	const int32 Index = UniformScalarParameters.Num();
+	FMaterialScalarParameterInfo& Parameter = UniformScalarParameters.AddDefaulted_GetRef();
+	Parameter.ParameterInfo = ParameterInfo;
+	Parameter.DefaultValue = DefaultValue;
+	return Index;
+}
+
+int32 FUniformExpressionSet::FindOrAddVectorParameter(const FHashedMaterialParameterInfo& ParameterInfo, const FLinearColor& DefaultValue)
+{
+	for (int32 i = 0; i < UniformVectorParameters.Num(); ++i)
+	{
+		if (UniformVectorParameters[i].ParameterInfo == ParameterInfo)
+		{
+			return i;
+		}
+	}
+
+	const int32 Index = UniformVectorParameters.Num();
+	FMaterialVectorParameterInfo& Parameter = UniformVectorParameters.AddDefaulted_GetRef();
+	Parameter.ParameterInfo = ParameterInfo;
+	Parameter.DefaultValue = DefaultValue;
+	return Index;
+}
+
+int32 FUniformExpressionSet::FindOrAddTextureParameter(EMaterialTextureParameterType Type, const FMaterialTextureParameterInfo& Info)
+{
+	for (int32 i = 0; i < UniformTextureParameters[(int32)Type].Num(); ++i)
+	{
+		if (UniformTextureParameters[(int32)Type][i] == Info)
+		{
+			return i;
+		}
+	}
+
+	return UniformTextureParameters[(int32)Type].Add(Info);
+}
+
 void FUniformExpressionSet::GetGameThreadTextureValue(EMaterialTextureParameterType Type, int32 Index, const UMaterialInterface* MaterialInterface, const FMaterial& Material, UTexture*& OutValue, bool bAllowOverride) const
 {
 	check(IsInGameThread());
diff --git a/Engine/Source/Runtime/Engine/Private/ShaderDerivedDataVersion.h b/Engine/Source/Runtime/Engine/Private/ShaderDerivedDataVersion.h
--- a/Engine/Source/Runtime/Engine/Private/ShaderDerivedDataVersion.h
+++ b/Engine/Source/Runtime/Engine/Private/ShaderDerivedDataVersion.h
@@ -11,4 +11,4 @@
 // In case of merge conflicts with DDC versions, you MUST generate a new GUID and set this new
 // guid as version
 #define GLOBALSHADERMAP_DERIVEDDATA_VER			TEXT("CE4A26A5071E4C9596A658A9AC69B29B")
-#define MATERIALSHADERMAP_DERIVEDDATA_VER		TEXT("9147101451524D13ACE72CCECEC5F230")
\ No newline at end of file
+#define MATERIALSHADERMAP_DERIVEDDATA_VER		TEXT("FCCBAA5CC66144528A3A2E4B4A29881B}")
\ No newline at end of file
diff --git a/Engine/Source/Runtime/Engine/Public/HLSLTree/HLSLTree.h b/Engine/Source/Runtime/Engine/Public/HLSLTree/HLSLTree.h
new file mode 100644
--- /dev/null
+++ b/Engine/Source/Runtime/Engine/Public/HLSLTree/HLSLTree.h
@@ -0,0 +1,323 @@
+// Copyright Epic Games, Inc. All Rights Reserved.
+#pragma once
+
+#include "CoreMinimal.h"
+#include "Engine/EngineTypes.h"
+#include "Misc/StringBuilder.h"
+#include "Misc/MemStack.h"
+#include "HLSLTree/HLSLTreeTypes.h"
+
+class FMaterial;
+class FMaterialCompilationOutput;
+class FMaterialPreshaderData;
+
+/**
+ * The HLSLTree module contains classes to build an HLSL AST (abstract syntax tree)
+ * This allows C++ to procedurally define an HLSL program.  The structure of the tree is designed to be flexible, to facilitate incremental generation from a material node graph
+ * Once the tree is complete, HLSL source code may be generated
+ */
+namespace UE
+{
+namespace HLSLTree
+{
+
+/** Allows building a string incrementally, with some additional features to support code generation, such as indent handling */
+class FCodeWriter
+{
+public:
+	static FCodeWriter* Create(FMemStackBase& Allocator);
+
+	const FStringBuilderBase& GetStringBuilder() const { return *StringBuilder; }
+
+	void IncreaseIndent();
+	void DecreaseIndent();
+
+	template <typename TextType>
+	inline void Write(const TextType& Text)
+	{
+		StringBuilder->Append(Text);
+	}
+
+	template <typename FormatType, typename... ArgTypes>
+	inline void Writef(const FormatType& Format, ArgTypes... Args)
+	{
+		StringBuilder->Appendf(Format, Args...);
+	}
+
+	template <typename TextType>
+	inline void WriteLine(const TextType& Text)
+	{
+		WriteIndent();
+		StringBuilder->Append(Text);
+		StringBuilder->Append('\n');
+	}
+
+	template <typename FormatType, typename... ArgTypes>
+	inline void WriteLinef(const FormatType& Format, ArgTypes... Args)
+	{
+		WriteIndent();
+		StringBuilder->Appendf(Format, Args...);
+		StringBuilder->Append('\n');
+	}
+
+	explicit FCodeWriter(FStringBuilderBase* InStringBuilder)
+		: StringBuilder(InStringBuilder)
+		, IndentLevel(0)
+	{}
+
+	void WriteIndent();
+
+	FStringBuilderBase* StringBuilder;
+	int32 IndentLevel;
+};
+
+/** Tracks shared state while emitting HLSL code */
+class FEmitContext
+{
+public:
+	/** Returns a snippet of HLSL source code that references the given expression */
+	const TCHAR* AcquireHLSLReference(FExpression* Expression);
+
+	/** Returns a snippet of HLSL source code that references the given local variablew */
+	const TCHAR* AcquireHLSLReference(FLocalDeclaration* Declaration);
+
+	struct FScopeEntry
+	{
+		const FScope* Scope;
+		FCodeWriter* ExpressionCodeWriter;
+	};
+
+	struct FDeclarationEntry
+	{
+		const TCHAR* Definition;
+	};
+
+	FScopeEntry* FindScope(FScope* Scope);
+
+	TArray<FScopeEntry> ScopeStack;
+	TMap<FNode*, FDeclarationEntry> DeclarationMap;
+	FMemStackBase* Allocator = nullptr;
+	FMaterial* Material = nullptr; // TODO - remove preshader material dependency
+	FMaterialCompilationOutput* MaterialCompilationOutput = nullptr;
+	int32 NumExpressionLocals = 0;
+	int32 NumTexCoords = 0;
+};
+
+enum class ENodeVisitResult
+{
+	VisitDependentNodes,
+	SkipDependentNodes,
+};
+
+/** Should be overriden to inspect the nodes of an HLSLTree */
+class FNodeVisitor
+{
+public:
+	virtual ~FNodeVisitor() {}
+
+	void VisitNode(FNode* Node);
+
+	virtual ENodeVisitResult OnScope(FScope& Scope) { return ENodeVisitResult::VisitDependentNodes; }
+	virtual ENodeVisitResult OnStatement(FStatement& Statement) { return ENodeVisitResult::VisitDependentNodes; }
+	virtual ENodeVisitResult OnExpression(FExpression& Expression) { return ENodeVisitResult::VisitDependentNodes; }
+	virtual ENodeVisitResult OnLocalDeclaration(FLocalDeclaration& Declaration) { return ENodeVisitResult::VisitDependentNodes; }
+	virtual ENodeVisitResult OnParameterDeclaration(FParameterDeclaration& Declaration) { return ENodeVisitResult::VisitDependentNodes; }
+	virtual ENodeVisitResult OnTextureParameterDeclaration(FTextureParameterDeclaration& Declaration) { return ENodeVisitResult::VisitDependentNodes; }
+};
+
+/** Root class of the HLSL AST */
+class FNode
+{
+public:
+	virtual ~FNode() {}
+
+	static bool ShouldVisitDependentNodes(ENodeVisitResult Result) { return Result == ENodeVisitResult::VisitDependentNodes; }
+
+	virtual ENodeVisitResult Visit(FNodeVisitor& Visitor) = 0;
+
+	FScope* ParentScope = nullptr;
+};
+
+/**
+ * Represents an HLSL statement.  This is a piece of code that doesn't evaluate to any value, but instead should be executed sequentially, and likely has side-effects.
+ * Examples include assigning a value, or various control flow structures (if, for, while, etc)
+ * This is an abstract base class, with derived classes representing various types of statements
+ */
+class FStatement : public FNode
+{
+public:
+	/** Emits HLSL code for the statement. The generated code should include any required semi-colons and newlines */
+	virtual void EmitHLSL(FEmitContext& Context, FCodeWriter& Writer) const = 0;
+
+	virtual ENodeVisitResult Visit(FNodeVisitor& Visitor) override;
+
+	FStatement* NextStatement = nullptr;
+};
+
+/**
+ * Represents an HLSL expression.  This is a piece of code that evaluates to a value, but has no side effects.
+ * Unlike statements, expressions are not expected to execute in any particular order.  They may be cached (or not) in generated code, without the underlying implementation needing to care.
+ * Expressions track the outer-most scope in which they're accessed. The generated HLSL code will ensure they are defined in that scope.
+ * Examples include constant literals, variable accessors, and various types of math operations
+ * This is an abstract base class, with derived classes representing various types of expression
+ * Derived classes are expected to implement at least one of EmitHLSL or EmitPreshader
+ */
+class FExpression : public FNode
+{
+public:
+	inline bool IsInline() const { return bInline; }
+
+	virtual ENodeVisitResult Visit(FNodeVisitor& Visitor) override;
+
+	/* Emits HLSL code for the expression. The code should NOT include any newlines or semi-colons. The returned string may be assigned to a temporary variable, or embedded in another HLSL string */
+	virtual void EmitHLSL(FEmitContext& Context, FCodeWriter& Writer) const;
+
+	/* Emits bytecode for the preshader virtual machine */
+	virtual void EmitPreshader(FEmitContext& Context, FMaterialPreshaderData& OutPreshader) const;
+
+	/** The HLSL type of the expression (float, int, etc) */
+	EExpressionType Type;
+
+	EExpressionEvaluationType EvaluationType;
+	bool bInline;
+
+protected:
+	explicit FExpression(EExpressionType InType,
+		EExpressionEvaluationType InEvaluationType = EExpressionEvaluationType::Shader)
+		: Type(InType)
+		, EvaluationType(InEvaluationType)
+		, bInline(false)
+	{}
+};
+
+/**
+ * Represents an HLSL local variable.  This is used by statements/expressions to refer to a particular variable
+ */
+class FLocalDeclaration final : public FNode
+{
+public:
+	FLocalDeclaration(const FName& InName, EExpressionType InType) : Name(InName), Type(InType) {}
+
+	virtual ENodeVisitResult Visit(FNodeVisitor& Visitor) override;
+
+	FName Name;
+	EExpressionType Type;
+};
+
+/**
+ * Represents an HLSL parameter.  This is a uniform constant passed to the generated shader
+ */
+class FParameterDeclaration final : public FNode
+{
+public:
+	FParameterDeclaration(const FName& InName, const FConstant& InDefaultValue) : Name(InName), DefaultValue(InDefaultValue) {}
+
+	virtual ENodeVisitResult Visit(FNodeVisitor& Visitor) override;
+
+	FName Name;
+	FConstant DefaultValue;
+};
+
+/**
+ * Represents an HLSL texture parameter.
+ */
+class FTextureParameterDeclaration : public FNode
+{
+public:
+	FTextureParameterDeclaration(const FName& InName, const FTextureDescription& InDescription) : Name(InName), Description(InDescription) {}
+
+	virtual ENodeVisitResult Visit(FNodeVisitor& Visitor) override;
+
+	FName Name;
+	FTextureDescription Description;
+};
+
+/**
+ * Represents an HLSL scope.  This is an ordered list of statements (enclosed by an {} pair in HLSL)
+ * All HLSL nodes track which scopes they are accessed from
+ */
+class FScope final : public FNode
+{
+public:
+	virtual ENodeVisitResult Visit(FNodeVisitor& Visitor) override;
+	void EmitHLSL(FEmitContext& Context, FCodeWriter& Writer) const;
+
+	void AddDeclaration(FLocalDeclaration* Declaration);
+	void AddExpression(FExpression* Expression);
+	void AddStatement(FStatement* Statement);
+
+	void UseDeclaration(FLocalDeclaration* Declaration);
+	void UseExpression(FExpression* Expression);
+	bool TryMoveStatement(FStatement* Statement);
+
+private:
+	friend class FTree;
+	friend class FNodeVisitor_MoveToScope;
+
+	void UseNode(FNode* Node);
+
+	FScope* LinkedScope = nullptr;
+	FStatement* FirstStatement = nullptr;
+	FStatement* LastStatement = nullptr;
+	int32 NestedLevel = 0;
+};
+
+/**
+ * The HLSL AST.  Basically a wrapper around the root scope, with some helper methods
+ */
+class FTree
+{
+public:
+	static FTree* Create(FMemStackBase& Allocator);
+
+	void EmitHLSL(FEmitContext& Context, FCodeWriter& Writer) const;
+
+	FScope& GetRootScope() const { return *RootScope; }
+
+	template<typename T, typename... ArgTypes>
+	inline T* NewExpression(FScope& Scope, ArgTypes&&... Args)
+	{
+		T* Expression = NewNode<T>(Forward<ArgTypes>(Args)...);
+		Scope.AddExpression(Expression);
+		return Expression;
+	}
+
+	template<typename T, typename... ArgTypes>
+	inline T* NewStatement(FScope& Scope, ArgTypes&&... Args)
+	{
+		T* Statement = NewNode<T>(Forward<ArgTypes>(Args)...);
+		Scope.AddStatement(Statement);
+		return Statement;
+	}
+
+	FScope* NewScope(FScope& Scope);
+
+	/**
+	 * 2 scopes may be linked if they are both logically part of the same control structure (an 'if' scope and an 'else' scope are linked for example)
+	 * Linked scopes always share the same parent scope
+	 * Attempting to move a statement from a scope to its linked scope will instead move that statement into the parent scope
+	 * When translating node graphs, a chain of statements may be generated when translating an 'if' scope for example
+	 * Then when translating the 'else' scope, the flow may eventually hit a statement that was previously translated from the 'if' scope
+	 * At this point, that statement connected to both 'if' and 'else' scope should become the next statement after the if/else block,
+	 * which means it should be moved to the parent scope (the scope that contains the linked if/else scopes)
+	 * It's not clear if this pattern will be useful outside if/else block
+	 */
+	FScope* NewLinkedScope(FScope& Scope);
+	
+	FLocalDeclaration* NewLocalDeclaration(FScope& Scope, EExpressionType Type, const FName& Name);
+	FParameterDeclaration* NewParameterDeclaration(FScope& Scope, const FName& Name, const FConstant& DefaultValue);
+	FTextureParameterDeclaration* NewTextureParameterDeclaration(FScope& Scope, const FName& Name, const FTextureDescription& DefaultValue);
+private:
+	template<typename T, typename... ArgTypes>
+	inline T* NewNode(ArgTypes&&... Args)
+	{
+		T* Node = new(*Allocator) T(Forward<ArgTypes>(Args)...);
+		return Node;
+	}
+
+	FMemStackBase* Allocator;
+	FScope* RootScope;
+};
+
+} // namespace HLSLTree
+} // namespace UE
diff --git a/Engine/Source/Runtime/Engine/Public/HLSLTree/HLSLTreeCommon.h b/Engine/Source/Runtime/Engine/Public/HLSLTree/HLSLTreeCommon.h
new file mode 100644
--- /dev/null
+++ b/Engine/Source/Runtime/Engine/Public/HLSLTree/HLSLTreeCommon.h
@@ -0,0 +1,341 @@
+// Copyright Epic Games, Inc. All Rights Reserved.
+#pragma once
+
+#include "HLSLTree/HLSLTree.h"
+
+namespace UE
+{
+namespace HLSLTree
+{
+
+enum class EBinaryOp
+{
+	None,
+	Add,
+	Sub,
+	Mul,
+	Div,
+};
+
+class FExpressionConstant : public FExpression
+{
+public:
+	explicit FExpressionConstant(const FConstant& InValue)
+		: FExpression(InValue.Type, EExpressionEvaluationType::Constant), Value(InValue)
+	{}
+
+	FConstant Value;
+
+	virtual void EmitPreshader(FEmitContext& Context, FMaterialPreshaderData& OutPreshader) const override;
+	virtual void EmitHLSL(FEmitContext& Context, FCodeWriter& Writer) const override;
+};
+
+class FExpressionLocalVariable : public FExpression
+{
+public:
+	explicit FExpressionLocalVariable(FLocalDeclaration* InDeclaration) : FExpression(InDeclaration->Type), Declaration(InDeclaration) {}
+
+	FLocalDeclaration* Declaration;
+
+	virtual ENodeVisitResult Visit(FNodeVisitor& Visitor) override
+	{
+		const ENodeVisitResult Result = FExpression::Visit(Visitor);
+		if (ShouldVisitDependentNodes(Result))
+		{
+			Visitor.VisitNode(Declaration);
+		}
+		return Result;
+	}
+
+	virtual void EmitHLSL(FEmitContext& Context, FCodeWriter& Writer) const override;
+};
+
+class FExpressionParameter : public FExpression
+{
+public:
+	FExpressionParameter(FParameterDeclaration* InDeclaration)
+		: FExpression(InDeclaration->DefaultValue.Type, EExpressionEvaluationType::Preshader)
+		, Declaration(InDeclaration)
+	{}
+
+	FParameterDeclaration* Declaration;
+
+	virtual ENodeVisitResult Visit(FNodeVisitor& Visitor) override
+	{
+		const ENodeVisitResult Result = FExpression::Visit(Visitor);
+		if (ShouldVisitDependentNodes(Result))
+		{
+			Visitor.VisitNode(Declaration);
+		}
+		return Result;
+	}
+
+	virtual void EmitPreshader(FEmitContext& Context, FMaterialPreshaderData& OutPreshader) const override;
+};
+
+enum class EExternalInputType
+{
+	TexCoord0,
+	TexCoord1,
+	TexCoord2,
+	TexCoord3,
+	TexCoord4,
+	TexCoord5,
+	TexCoord6,
+	TexCoord7,
+};
+inline EExpressionType GetInputExpressionType(EExternalInputType Type)
+{
+	return EExpressionType::Float2;
+}
+inline EExternalInputType MakeInputTexCoord(int32 Index)
+{
+	check(Index >= 0 && Index < 8);
+	return (EExternalInputType)((int32)EExternalInputType::TexCoord0 + Index);
+}
+
+class FExpressionExternalInput : public FExpression
+{
+public:
+	FExpressionExternalInput(EExternalInputType InInputType) : FExpression(GetInputExpressionType(InInputType)), InputType(InInputType) {}
+
+	EExternalInputType InputType;
+
+	virtual void EmitHLSL(FEmitContext& Context, FCodeWriter& Writer) const override;
+};
+
+class FExpressionTextureSample : public FExpression
+{
+public:
+	FExpressionTextureSample(FTextureParameterDeclaration* InDeclaration, FExpression* InTexCoordExpression)
+		: FExpression(EExpressionType::Float4)
+		, Declaration(InDeclaration)
+		, TexCoordExpression(InTexCoordExpression)
+		, SamplerSource(SSM_FromTextureAsset)
+		, MipValueMode(TMVM_None)
+	{}
+
+	FTextureParameterDeclaration* Declaration;
+	FExpression* TexCoordExpression;
+	ESamplerSourceMode SamplerSource;
+	ETextureMipValueMode MipValueMode;
+
+	virtual ENodeVisitResult Visit(FNodeVisitor& Visitor) override
+	{
+		const ENodeVisitResult Result = FExpression::Visit(Visitor);
+		if (ShouldVisitDependentNodes(Result))
+		{
+			Visitor.VisitNode(Declaration);
+			Visitor.VisitNode(TexCoordExpression);
+		}
+		return Result;
+	}
+
+	virtual void EmitHLSL(FEmitContext& Context, FCodeWriter& Writer) const override;
+};
+
+class FExpressionDefaultMaterialAttributes : public FExpression
+{
+public:
+	FExpressionDefaultMaterialAttributes() : FExpression(EExpressionType::MaterialAttributes) {}
+
+	virtual void EmitHLSL(FEmitContext& Context, FCodeWriter& Writer) const override;
+};
+
+class FExpressionSetMaterialAttribute : public FExpression
+{
+public:
+	FExpressionSetMaterialAttribute() : FExpression(EExpressionType::MaterialAttributes) {}
+
+	FGuid AttributeID;
+	FExpression* AttributesExpression;
+	FExpression* ValueExpression;
+
+	virtual ENodeVisitResult Visit(FNodeVisitor& Visitor) override
+	{
+		const ENodeVisitResult Result = FExpression::Visit(Visitor);
+		if (ShouldVisitDependentNodes(Result))
+		{
+			Visitor.VisitNode(AttributesExpression);
+			Visitor.VisitNode(ValueExpression);
+		}
+		return Result;
+	}
+
+	virtual void EmitHLSL(FEmitContext& Context, FCodeWriter& Writer) const override;
+};
+
+class FExpressionBinaryOp : public FExpression
+{
+public:
+	FExpressionBinaryOp(EExpressionType ResultType, EBinaryOp InOp, FExpression* InLhs, FExpression* InRhs)
+		: FExpression(ResultType, CombineEvaluationTypes(InLhs->EvaluationType, InRhs->EvaluationType))
+		, Op(InOp)
+		, Lhs(InLhs)
+		, Rhs(InRhs)
+	{}
+
+	EBinaryOp Op;
+	FExpression* Lhs;
+	FExpression* Rhs;
+
+	virtual ENodeVisitResult Visit(FNodeVisitor& Visitor) override
+	{
+		const ENodeVisitResult Result = FExpression::Visit(Visitor);
+		if (ShouldVisitDependentNodes(Result))
+		{
+			Visitor.VisitNode(Lhs);
+			Visitor.VisitNode(Rhs);
+		}
+		return Result;
+	}
+
+	virtual void EmitPreshader(FEmitContext& Context, FMaterialPreshaderData& OutPreshader) const override;
+	virtual void EmitHLSL(FEmitContext& Context, FCodeWriter& Writer) const override;
+};
+
+struct FSwizzleParameters
+{
+	FSwizzleParameters() : NumComponents(0) { ComponentIndex[0] = ComponentIndex[1] = ComponentIndex[2] = ComponentIndex[3] = INDEX_NONE; }
+	FSwizzleParameters(int8 IndexR, int8 IndexG, int8 IndexB, int8 IndexA);
+
+	int8 ComponentIndex[4];
+	int32 NumComponents;
+};
+FSwizzleParameters MakeSwizzleMask(bool bInR, bool bInG, bool bInB, bool bInA);
+
+class FExpressionSwizzle : public FExpression
+{
+public:
+	FExpressionSwizzle(const FSwizzleParameters& InParams, FExpression* InInput)
+		: FExpression(MakeExpressionType(InInput->Type, InParams.NumComponents), InInput->EvaluationType)
+		, Parameters(InParams)
+		, Input(InInput)
+	{}
+
+	FSwizzleParameters Parameters;
+	FExpression* Input;
+
+	virtual ENodeVisitResult Visit(FNodeVisitor& Visitor) override
+	{
+		const ENodeVisitResult Result = FExpression::Visit(Visitor);
+		if (ShouldVisitDependentNodes(Result))
+		{
+			Visitor.VisitNode(Input);
+		}
+		return Result;
+	}
+
+	virtual void EmitPreshader(FEmitContext& Context, FMaterialPreshaderData& OutPreshader) const override;
+	virtual void EmitHLSL(FEmitContext& Context, FCodeWriter& Writer) const override;
+};
+
+class FExpressionCast : public FExpression
+{
+public:
+	FExpressionCast(EExpressionType InType, FExpression* InInput, ECastFlags InFlags = ECastFlags::None)
+		: FExpression(InType) // TODO preshader
+		, Input(InInput)
+		, Flags(InFlags)
+	{}
+
+	FExpression* Input;
+	ECastFlags Flags;
+
+	virtual ENodeVisitResult Visit(FNodeVisitor& Visitor) override
+	{
+		const ENodeVisitResult Result = FExpression::Visit(Visitor);
+		if (ShouldVisitDependentNodes(Result))
+		{
+			Visitor.VisitNode(Input);
+		}
+		return Result;
+	}
+
+	virtual void EmitHLSL(FEmitContext& Context, FCodeWriter& Writer) const override;
+};
+
+class FStatementReturn : public FStatement
+{
+public:
+	FExpression* Expression;
+
+	virtual ENodeVisitResult Visit(FNodeVisitor& Visitor) override
+	{
+		const ENodeVisitResult Result = FStatement::Visit(Visitor);
+		if (ShouldVisitDependentNodes(Result))
+		{
+			Visitor.VisitNode(Expression);
+		}
+		return Result;
+	}
+
+	virtual void EmitHLSL(FEmitContext& Context, FCodeWriter& Writer) const override;
+};
+
+class FStatementSetLocalVariable : public FStatement
+{
+public:
+	FLocalDeclaration* Declaration;
+	FExpression* Expression;
+
+	virtual ENodeVisitResult Visit(FNodeVisitor& Visitor) override
+	{
+		const ENodeVisitResult Result = FStatement::Visit(Visitor);
+		if (ShouldVisitDependentNodes(Result))
+		{
+			Visitor.VisitNode(Declaration);
+			Visitor.VisitNode(Expression);
+		}
+		return Result;
+	}
+
+	virtual void EmitHLSL(FEmitContext& Context, FCodeWriter& Writer) const override;
+};
+
+class FStatementIf : public FStatement
+{
+public:
+	FExpression* ConditionExpression;
+	FScope* ThenScope;
+	FScope* ElseScope;
+
+	virtual ENodeVisitResult Visit(FNodeVisitor& Visitor) override
+	{
+		const ENodeVisitResult Result = FStatement::Visit(Visitor);
+		if (ShouldVisitDependentNodes(Result))
+		{
+			Visitor.VisitNode(ConditionExpression);
+			Visitor.VisitNode(ThenScope);
+			Visitor.VisitNode(ElseScope);
+		}
+		return Result;
+	}
+
+	virtual void EmitHLSL(FEmitContext& Context, FCodeWriter& Writer) const override;
+};
+
+class FStatementFor : public FStatement
+{
+public:
+	FExpression* StartExpression;
+	FExpression* EndExpression;
+	FScope* LoopScope;
+
+	virtual ENodeVisitResult Visit(FNodeVisitor& Visitor) override
+	{
+		const ENodeVisitResult Result = FStatement::Visit(Visitor);
+		if (ShouldVisitDependentNodes(Result))
+		{
+			Visitor.VisitNode(StartExpression);
+			Visitor.VisitNode(EndExpression);
+			Visitor.VisitNode(LoopScope);
+		}
+		return Result;
+	}
+
+	virtual void EmitHLSL(FEmitContext& Context, FCodeWriter& Writer) const override;
+};
+
+}
+}
\ No newline at end of file
diff --git a/Engine/Source/Runtime/Engine/Public/HLSLTree/HLSLTreeTypes.h b/Engine/Source/Runtime/Engine/Public/HLSLTree/HLSLTreeTypes.h
new file mode 100644
--- /dev/null
+++ b/Engine/Source/Runtime/Engine/Public/HLSLTree/HLSLTreeTypes.h
@@ -0,0 +1,128 @@
+// Copyright Epic Games, Inc. All Rights Reserved.
+#pragma once
+
+#include "Misc/EnumClassFlags.h"
+
+class UTexture;
+
+namespace UE
+{
+namespace HLSLTree
+{
+
+class FNode;
+class FExpression;
+class FLocalDeclaration;
+class FParameterDeclaration;
+class FTextureParameterDeclaration;
+class FStatement;
+class FScope;
+class FTree;
+class FCodeWriter;
+
+enum class EExpressionComponentType
+{
+	Void,
+	Float,
+	MaterialAttributes,
+};
+
+struct FExpressionTypeDescription
+{
+	FExpressionTypeDescription() : Name(nullptr), ComponentType(EExpressionComponentType::Void), NumComponents(0) {}
+	FExpressionTypeDescription(const TCHAR* InName, EExpressionComponentType InComponentType, int32 InNumComponents) : Name(InName), ComponentType(InComponentType), NumComponents(InNumComponents) {}
+
+	const TCHAR* Name;
+	EExpressionComponentType ComponentType;
+	int32 NumComponents;
+};
+
+enum class EExpressionType : uint8
+{
+	Void,
+
+	Float1,
+	Float2,
+	Float3,
+	Float4,
+
+	MaterialAttributes,
+};
+
+FExpressionTypeDescription GetExpressionTypeDescription(EExpressionType Type);
+EExpressionType MakeExpressionType(EExpressionComponentType ComponentType, int32 NumComponents);
+EExpressionType MakeExpressionType(EExpressionType BaseType, int32 NumComponents);
+EExpressionType MakeArithmeticResultType(EExpressionType Lhs, EExpressionType Rhs, FString& OutErrorMessage);
+
+/**
+ * Describes how a given expression needs to be evaluated */
+enum class EExpressionEvaluationType
+{
+	/** The expression needs to generate HLSL code (EmitHLSL will be called) */
+	Shader,
+
+	/** The expression can generate a preshader (EmitPreshader will be called) */
+	Preshader,
+
+	/** The expression is constant (EmitPreshader will be called, generated preshader will be evaluated and cached at compile-time) */
+	Constant,
+};
+
+EExpressionEvaluationType CombineEvaluationTypes(EExpressionEvaluationType Lhs, EExpressionEvaluationType Rhs);
+
+struct FConstant
+{
+	inline FConstant() : Type(EExpressionType::Void) { FMemory::Memzero(Raw); }
+	inline FConstant(float v) : Type(EExpressionType::Float1) { FMemory::Memzero(Raw); Float[0] = v; }
+	inline FConstant(float X, float Y) : Type(EExpressionType::Float2) { FMemory::Memzero(Raw); Float[0] = X; Float[1] = Y; }
+	inline FConstant(float X, float Y, float Z) : Type(EExpressionType::Float3) { FMemory::Memzero(Raw); Float[0] = X; Float[1] = Y; Float[2] = Z; }
+	inline FConstant(float X, float Y, float Z, float W) : Type(EExpressionType::Float4) { FMemory::Memzero(Raw); Float[0] = X; Float[1] = Y; Float[2] = Z; Float[3] = W; }
+	inline FConstant(EExpressionType InType, const FLinearColor& Value) : Type(InType) { FMemory::Memzero(Raw); Float[0] = Value.R; Float[1] = Value.G; Float[2] = Value.B; Float[3] = Value.A; }
+
+	FLinearColor ToLinearColor() const;
+
+	void EmitHLSL(FCodeWriter& Writer) const;
+
+	EExpressionType Type;
+	union
+	{
+		uint32 Raw[4];
+		float Float[4];
+	};
+};
+
+struct FTextureDescription
+{
+	FTextureDescription()
+		: Texture(nullptr), SamplerType(SAMPLERTYPE_Color)
+	{}
+
+	FTextureDescription(UTexture* InTexture, EMaterialSamplerType InSamplerType)
+		: Texture(InTexture), SamplerType(InSamplerType)
+	{}
+
+	UTexture* Texture;
+	EMaterialSamplerType SamplerType;
+};
+inline bool operator==(const FTextureDescription& Lhs, const FTextureDescription& Rhs)
+{
+	return Lhs.Texture == Rhs.Texture && Lhs.SamplerType == Rhs.SamplerType;
+}
+inline bool operator!=(const FTextureDescription& Lhs, const FTextureDescription& Rhs)
+{
+	return !operator==(Lhs, Rhs);
+}
+inline uint32 GetTypeHash(const FTextureDescription& Ref)
+{
+	return HashCombine(GetTypeHash(Ref.Texture), GetTypeHash(Ref.SamplerType));
+}
+
+enum class ECastFlags : uint32
+{
+	None = 0u,
+	ReplicateScalar = (1u << 0),
+};
+ENUM_CLASS_FLAGS(ECastFlags);
+
+} // namespace HLSLTree
+} // namespace UE
diff --git a/Engine/Source/Runtime/Engine/Public/MaterialExpressionIO.h b/Engine/Source/Runtime/Engine/Public/MaterialExpressionIO.h
--- a/Engine/Source/Runtime/Engine/Public/MaterialExpressionIO.h
+++ b/Engine/Source/Runtime/Engine/Public/MaterialExpressionIO.h
@@ -6,6 +6,19 @@
 #include "UObject/Class.h"
 #include "SceneTypes.h"
 
+class FMaterialHLSLGenerator;
+
+namespace UE
+{
+namespace HLSLTree
+{
+enum class EExpressionType : uint8;
+class FScope;
+class FExpression;
+class FTextureParameterDeclaration;
+}
+}
+
 //
 //	FExpressionInput
 //
@@ -61,6 +74,9 @@ struct FExpressionInput
 
 #if WITH_EDITOR
 	ENGINE_API int32 Compile(class FMaterialCompiler* Compiler);
+	ENGINE_API UE::HLSLTree::FExpression* AcquireHLSLExpression(FMaterialHLSLGenerator& Generator, UE::HLSLTree::FScope& Scope) const;
+	ENGINE_API UE::HLSLTree::FExpression* AcquireHLSLExpressionWithCast(FMaterialHLSLGenerator& Generator, UE::HLSLTree::FScope& Scope, UE::HLSLTree::EExpressionType Type) const;
+	ENGINE_API UE::HLSLTree::FTextureParameterDeclaration* AcquireHLSLTexture(FMaterialHLSLGenerator& Generator, UE::HLSLTree::FScope& Scope) const;
 #endif // WITH_EDITOR
 
 	/**
diff --git a/Engine/Source/Runtime/Engine/Public/MaterialHLSLGenerator.h b/Engine/Source/Runtime/Engine/Public/MaterialHLSLGenerator.h
new file mode 100644
--- /dev/null
+++ b/Engine/Source/Runtime/Engine/Public/MaterialHLSLGenerator.h
@@ -0,0 +1,115 @@
+// Copyright Epic Games, Inc. All Rights Reserved.
+#pragma once
+
+#include "CoreMinimal.h"
+
+#if WITH_EDITOR
+
+#include "Containers/Map.h"
+#include "Misc/MemStack.h"
+#include "Templates/RefCounting.h"
+#include "RHIDefinitions.h"
+#include "HLSLTree/HLSLTree.h"
+
+class FMaterial;
+class FMaterialCompilationOutput;
+struct FSharedShaderCompilerEnvironment;
+class UMaterialExpression;
+class ITargetPlatform;
+enum class EMaterialGenerateHLSLStatus : uint8;
+
+namespace UE
+{
+namespace HLSLTree
+{
+struct FSwizzleParameters;
+class FExpressionConstant;
+class FExpressionExternalInput;
+class FExpressionSwizzle;
+class FExpressionCast;
+}
+}
+
+/**
+ * MaterialHLSLGenerator is a bridge between a material, and HLSLTree.  It facilitates generating HLSL source code for a given material, using HLSLTree
+ */
+class FMaterialHLSLGenerator
+{
+public:
+	FMaterialHLSLGenerator(EShaderPlatform InShaderPlatform,
+		ERHIFeatureLevel::Type InFeatureLevel,
+		const ITargetPlatform* InTargetPlatform);
+
+	bool GenerateHLSL(FMaterial& InMaterial,
+		FMaterialCompilationOutput& OutCompilationOutput,
+		TRefCountPtr<FSharedShaderCompilerEnvironment>& OutMaterialEnvironment);
+
+	EMaterialGenerateHLSLStatus Error(const FString& Message);
+
+	template <typename FormatType, typename... ArgTypes>
+	inline EMaterialGenerateHLSLStatus Errorf(const FormatType& Format, ArgTypes... Args)
+	{
+		return Error(FString::Printf(Format, Args...));
+	}
+
+	UE::HLSLTree::FTree& GetTree() const { return *HLSLTree; }
+
+	UE::HLSLTree::FExpressionConstant* NewConstant(UE::HLSLTree::FScope& Scope, const UE::HLSLTree::FConstant& Value);
+	UE::HLSLTree::FExpressionExternalInput* NewTexCoord(UE::HLSLTree::FScope& Scope, int32 Index);
+	UE::HLSLTree::FExpressionSwizzle* NewSwizzle(UE::HLSLTree::FScope& Scope, const UE::HLSLTree::FSwizzleParameters& Params, UE::HLSLTree::FExpression* Input);
+	UE::HLSLTree::FExpression* NewCast(UE::HLSLTree::FScope& Scope, UE::HLSLTree::EExpressionType Type, UE::HLSLTree::FExpression* Input, UE::HLSLTree::ECastFlags Flags = UE::HLSLTree::ECastFlags::None);
+
+	UE::HLSLTree::FLocalDeclaration* AcquireLocalDeclaration(UE::HLSLTree::FScope& Scope, UE::HLSLTree::EExpressionType Type, const FName& Name);
+	UE::HLSLTree::FParameterDeclaration* AcquireParameterDeclaration(UE::HLSLTree::FScope& Scope, const FName& Name, const UE::HLSLTree::FConstant& DefaultValue);
+
+	/** Returns a declaration to access the given texture, with no parameter */
+	UE::HLSLTree::FTextureParameterDeclaration* AcquireTextureDeclaration(UE::HLSLTree::FScope& Scope, const UE::HLSLTree::FTextureDescription& Value);
+
+	/** Returns a declaration to access the given texture parameter */
+	UE::HLSLTree::FTextureParameterDeclaration* AcquireTextureParameterDeclaration(UE::HLSLTree::FScope& Scope, const FName& Name, const UE::HLSLTree::FTextureDescription& DefaultValue);
+
+	/**
+	 * Returns the appropriate HLSLNode representing the given UMaterialExpression.
+	 * The node will be created if it doesn't exist. Otherwise, the tree will be updated to ensure the node is visible in the given scope
+	 * Note that a given UMaterialExpression may only support 1 of these node types, attempting to access an invalid node type will generate an error
+	 */
+	UE::HLSLTree::FExpression* AcquireExpression(UE::HLSLTree::FScope& Scope, UMaterialExpression* MaterialExpression, int32 OutputIndex);
+	UE::HLSLTree::FTextureParameterDeclaration* AcquireTextureDeclaration(UE::HLSLTree::FScope& Scope, UMaterialExpression* MaterialExpression, int32 OutputIndex);
+	UE::HLSLTree::FStatement* AcquireStatement(UE::HLSLTree::FScope& Scope, UMaterialExpression* MaterialExpression);
+
+private:
+	struct FExpressionKey
+	{
+		explicit FExpressionKey(UMaterialExpression* InExpression, int32 InOutputIndex = INDEX_NONE) : Expression(InExpression), OutputIndex(InOutputIndex) {}
+
+		UMaterialExpression* Expression;
+		int32 OutputIndex;
+
+		friend inline uint32 GetTypeHash(const FExpressionKey& Value)
+		{
+			return HashCombine(GetTypeHash(Value.Expression), GetTypeHash(Value.OutputIndex));
+		}
+
+		friend inline bool operator==(const FExpressionKey& Lhs, const FExpressionKey& Rhs)
+		{
+			return Lhs.Expression == Rhs.Expression && Lhs.OutputIndex == Rhs.OutputIndex;
+		}
+	};
+
+	FMemStackBase Allocator;
+	EShaderPlatform ShaderPlatform;
+	ERHIFeatureLevel::Type FeatureLevel;
+	const ITargetPlatform* TargetPlatform;
+	UE::HLSLTree::FTree* HLSLTree;
+	TArray<FExpressionKey> ExpressionStack;
+	TArray<FString> CompileErrors;
+	TArray<UMaterialExpression*> ErrorExpressions;
+	TMap<FName, UE::HLSLTree::FLocalDeclaration*> LocalDeclarationMap;
+	TMap<FName, UE::HLSLTree::FParameterDeclaration*> ParameterDeclarationMap;
+	TMap<UE::HLSLTree::FTextureDescription, UE::HLSLTree::FTextureParameterDeclaration*> TextureDeclarationMap;
+	TMap<FName, UE::HLSLTree::FTextureParameterDeclaration*> TextureParameterDeclarationMap;
+	TMap<FExpressionKey, UE::HLSLTree::FExpression*> ExpressionMap;
+	TMap<UMaterialExpression*, UE::HLSLTree::FStatement*> StatementMap;
+};
+
+#endif // WITH_EDITOR
diff --git a/Engine/Source/Runtime/Engine/Public/MaterialShared.h b/Engine/Source/Runtime/Engine/Public/MaterialShared.h
--- a/Engine/Source/Runtime/Engine/Public/MaterialShared.h
+++ b/Engine/Source/Runtime/Engine/Public/MaterialShared.h
@@ -61,6 +61,14 @@ class FMeshMaterialShaderMapLayout;
 struct FMaterialShaderTypes;
 struct FMaterialShaders;
 
+namespace UE
+{
+namespace HLSLTree
+{
+class FEmitContext;
+}
+}
+
 template <class ElementType> class TLinkedList;
 
 #define ME_CAPTION_HEIGHT		18
@@ -283,6 +291,8 @@ public:
 		return !operator==(Lhs, Rhs);
 	}
 
+	void Evaluate(FUniformExpressionSet* UniformExpressionSet, const struct FMaterialRenderContext& Context, FLinearColor& OutValue);
+
 	const int32 Num() const { return Data.Num(); }
 
 	void WriteData(const void* Value, uint32 Size);
@@ -550,6 +560,11 @@ public:
 	ENGINE_API void GetTextureValue(EMaterialTextureParameterType Type, int32 Index, const FMaterialRenderContext& Context, const FMaterial& Material, const UTexture*& OutValue) const;
 	ENGINE_API void GetTextureValue(int32 Index, const FMaterialRenderContext& Context, const FMaterial& Material, const URuntimeVirtualTexture*& OutValue) const;
 
+	int32 FindOrAddScalarParameter(const FHashedMaterialParameterInfo& ParameterInfo, float DefaultValue);
+	int32 FindOrAddVectorParameter(const FHashedMaterialParameterInfo& ParameterInfo, const FLinearColor& DefaultValue);
+
+	int32 FindOrAddTextureParameter(EMaterialTextureParameterType Type, const FMaterialTextureParameterInfo& Info);
+
 protected:
 	union FVTPackedStackAndLayerIndex
 	{
@@ -572,6 +587,7 @@ protected:
 	friend class FMaterialRenderProxy;
 	friend class FMaterialVirtualTextureStack;
 	friend class FDebugUniformExpressionSet;
+	friend class UE::HLSLTree::FEmitContext;
 
 	LAYOUT_FIELD(TMemoryImageArray<FMaterialUniformPreshaderHeader>, UniformVectorPreshaders);
 	LAYOUT_FIELD(TMemoryImageArray<FMaterialUniformPreshaderHeader>, UniformScalarPreshaders);
@@ -811,6 +827,9 @@ public:
 	
 	/** A hash of the base property overrides for this material instance. */
 	FSHAHash BasePropertyOverridesHash;
+
+	/** Is the material using the new HLSL generator? */
+	bool bUsingNewHLSLGenerator;
 #endif // WITH_EDITOR
 
 	/*
@@ -824,6 +843,7 @@ public:
 #if WITH_EDITOR
 		, Usage(EMaterialShaderMapUsage::Default)
 		, bIsCookedId(false)
+		, bUsingNewHLSLGenerator(false)
 #endif
 	{ }
 
@@ -1811,6 +1831,9 @@ public:
 
 	/** Should the material be compiled using exec pin? */
 	virtual bool IsCompiledWithExecutionFlow() const { return false; }
+
+	/** Is the material using the new (WIP) HLSL generator? */
+	virtual bool IsUsingNewHLSLGenerator() const { return false; }
 #endif // WITH_EDITOR
 
 	/**
@@ -2176,6 +2199,20 @@ private:
 		EMaterialShaderPrecompileMode PrecompileMode,
 		const ITargetPlatform* TargetPlatform = nullptr);
 
+	bool Translate_Legacy(const FMaterialShaderMapId& InShaderMapId,
+		const FStaticParameterSet& InStaticParameters,
+		EShaderPlatform InPlatform,
+		const ITargetPlatform* InTargetPlatform,
+		FMaterialCompilationOutput& OutCompilationOutput,
+		TRefCountPtr<FSharedShaderCompilerEnvironment>& OutMaterialEnvironment);
+
+	bool Translate_New(const FMaterialShaderMapId& InShaderMapId,
+		const FStaticParameterSet& InStaticParameters,
+		EShaderPlatform InPlatform,
+		const ITargetPlatform* InTargetPlatform,
+		FMaterialCompilationOutput& OutCompilationOutput,
+		TRefCountPtr<FSharedShaderCompilerEnvironment>& OutMaterialEnvironment);
+
 	/** Populates OutEnvironment with defines needed to compile shaders for this material. */
 	void SetupMaterialEnvironment(
 		EShaderPlatform Platform,
@@ -2194,6 +2231,7 @@ private:
 	friend class FMaterialShaderMap;
 	friend class FShaderCompilingManager;
 	friend class FHLSLMaterialTranslator;
+	friend class FMaterialHLSLGenerator;
 };
 
 /**
@@ -2648,6 +2686,7 @@ public:
 	ENGINE_API virtual FName GetAssetPath() const override;
 	ENGINE_API virtual bool ShouldInlineShaderCode() const override;
 	ENGINE_API virtual bool IsCompiledWithExecutionFlow() const override;
+	ENGINE_API virtual bool IsUsingNewHLSLGenerator() const override;
 #endif // WITH_EDITOR
 
 	ENGINE_API virtual FString GetFullPath() const override;
