Follow-up refactor to address fallout from FSoundWaveProxy

diff --git a/Engine/Plugins/Runtime/Metasound/Source/MetasoundEngine/Private/MetasoundWave.cpp b/Engine/Plugins/Runtime/Metasound/Source/MetasoundEngine/Private/MetasoundWave.cpp
--- a/Engine/Plugins/Runtime/Metasound/Source/MetasoundEngine/Private/MetasoundWave.cpp
+++ b/Engine/Plugins/Runtime/Metasound/Source/MetasoundEngine/Private/MetasoundWave.cpp
@@ -11,12 +11,12 @@
 
 namespace Audio
 {
-	bool FSimpleDecoderWrapper::Initialize(const InitParams& InInitParams, const FSoundWaveProxy& InWave)
+	bool FSimpleDecoderWrapper::Initialize(const InitParams& InInitParams, const FSoundWaveProxyPtr& InWave)
 	{
 		// validate data
-		if (false == InWave.IsStreaming())
+		if (!ensure(InWave.IsValid()) || !InWave->IsStreaming())
 		{
-			ensureAlwaysMsgf(InWave.IsStreaming(), TEXT("Metasounds does not support Force Inline (sound must be streaming)"));
+			ensureAlwaysMsgf(InWave->IsStreaming(), TEXT("Metasounds does not support Force Inline (sound must be streaming)"));
 
 			Input.Reset();
 			Output.Reset();
@@ -27,13 +27,13 @@ namespace Audio
 		}
 
 		// initialize input/output data
-		InputSampleRate = InWave.GetSampleRate();
+		InputSampleRate = InWave->GetSampleRate();
 		OutputSampleRate = InInitParams.OutputSampleRate;
 		FsInToFsOutRatio = (InputSampleRate / OutputSampleRate);
 		MaxPitchShiftRatio = FMath::Pow(2.f, InInitParams.MaxPitchShiftMagnitudeAllowedInOctaves);
 		MaxPitchShiftCents = InInitParams.MaxPitchShiftMagnitudeAllowedInOctaves * 1200.f;
 
-		NumChannels = InWave.GetNumChannels();
+		NumChannels = InWave->GetNumChannels();
 		DecodeBlockSizeInFrames = InInitParams.OutputBlockSizeInFrames;
 		DecodeBlockSizeInSamples = DecodeBlockSizeInFrames * NumChannels;
 
@@ -128,10 +128,15 @@ namespace Audio
 		return NumFramesConverted;
 	}
 
-	bool FSimpleDecoderWrapper::InitializeDecodersInternal(const FSoundWaveProxy& Wave)
+	bool FSimpleDecoderWrapper::InitializeDecodersInternal(const FSoundWaveProxyPtr& Wave)
 	{
+		if (!ensure(Wave.IsValid()))
+		{
+			return false;
+		}
+
 		// Input:
-		FName OldFormat = Wave.GetRuntimeFormat();
+		FName OldFormat = Wave->GetRuntimeFormat();
 		Input = MakeShareable(Audio::CreateBackCompatDecoderInput(OldFormat, Wave).Release());
 
 		if (!Input)
diff --git a/Engine/Plugins/Runtime/Metasound/Source/MetasoundEngine/Private/MetasoundWavePlayerNode.cpp b/Engine/Plugins/Runtime/Metasound/Source/MetasoundEngine/Private/MetasoundWavePlayerNode.cpp
--- a/Engine/Plugins/Runtime/Metasound/Source/MetasoundEngine/Private/MetasoundWavePlayerNode.cpp
+++ b/Engine/Plugins/Runtime/Metasound/Source/MetasoundEngine/Private/MetasoundWavePlayerNode.cpp
@@ -151,7 +151,7 @@ namespace Metasound
 				return false;
 			}
 
-			return Decoder.Initialize(Params, *Wave->SoundWaveProxy);
+			return Decoder.Initialize(Params, Wave->SoundWaveProxy);
 		}
 
 
diff --git a/Engine/Plugins/Runtime/Metasound/Source/MetasoundEngine/Public/MetasoundWave.h b/Engine/Plugins/Runtime/Metasound/Source/MetasoundEngine/Public/MetasoundWave.h
--- a/Engine/Plugins/Runtime/Metasound/Source/MetasoundEngine/Public/MetasoundWave.h
+++ b/Engine/Plugins/Runtime/Metasound/Source/MetasoundEngine/Public/MetasoundWave.h
@@ -23,7 +23,7 @@ namespace Metasound
 	public:
 		// TODO: Implement using TSharedPtr<> to enable cheaper assignment when 
 		// choosing between multiple instances of FWaveAsset.
-		TUniquePtr<FSoundWaveProxy> SoundWaveProxy;
+		FSoundWaveProxyPtr SoundWaveProxy;
 
 		FWaveAsset() = default;
 		FWaveAsset(const FWaveAsset&) = delete;
@@ -35,7 +35,7 @@ namespace Metasound
 
 		FWaveAsset& operator=(const FWaveAsset& Other)
 		{
-			SoundWaveProxy = MakeUnique<FSoundWaveProxy>(*Other.SoundWaveProxy);
+			SoundWaveProxy = Other.SoundWaveProxy;
 			return *this;
 		}
 
@@ -45,7 +45,8 @@ namespace Metasound
 
 			if (InInitData.IsValid())
 			{
-				SoundWaveProxy = MakeUnique<FSoundWaveProxy>(InInitData->GetAs<FSoundWaveProxy>());
+				// should we be getting handed a SharedPtr here?
+				SoundWaveProxy = MakeShared<FSoundWaveProxy, ESPMode::ThreadSafe>(InInitData->GetAs<FSoundWaveProxy>());
 			}
 		}
 
@@ -82,7 +83,7 @@ namespace Audio
 			return !bDecoderIsDone && Input.IsValid() && Output.IsValid() && Decoder.IsValid();
 		}
 
-		bool Initialize(const InitParams& InInitParams, const FSoundWaveProxy& InWave);
+		bool Initialize(const InitParams& InInitParams, const FSoundWaveProxyPtr& InWave);
 
 		// returns number of samples written.   
 		uint32 GenerateAudio(float* OutputDest, int32 NumOutputFrames, float PitchShiftInCents = 0.f);
@@ -94,7 +95,7 @@ namespace Audio
 		TSharedPtr<Audio::IDecoderInput, ESPMode::ThreadSafe> Input;
 
 		// init helper for decoders
-		bool InitializeDecodersInternal(const FSoundWaveProxy& Wave);
+		bool InitializeDecodersInternal(const FSoundWaveProxyPtr& Wave);
 
 		// SRC object
 		Audio::FResampler Resampler;
diff --git a/Engine/Source/Runtime/AudioCodecEngine/Private/DecoderInputBackCompat.cpp b/Engine/Source/Runtime/AudioCodecEngine/Private/DecoderInputBackCompat.cpp
--- a/Engine/Source/Runtime/AudioCodecEngine/Private/DecoderInputBackCompat.cpp
+++ b/Engine/Source/Runtime/AudioCodecEngine/Private/DecoderInputBackCompat.cpp
@@ -17,13 +17,13 @@ namespace Audio
 
 		if (!OldInfoObject.IsValid())
 		{
-			OldInfoObject.Reset(Handle->CreateCompressedAudioInfo(*Wave));
+			OldInfoObject.Reset(Handle->CreateCompressedAudioInfo(Wave));
 			audio_ensure(OldInfoObject.IsValid());
 
 			FSoundQualityInfo Info;
 			if (Wave->IsStreaming())
 			{
-				if (!OldInfoObject->StreamCompressedInfo(*Wave, &Info))
+				if (!OldInfoObject->StreamCompressedInfo(Wave, &Info))
 				{
 					return nullptr;
 				}
diff --git a/Engine/Source/Runtime/AudioCodecEngine/Private/DecoderInputBackCompat.h b/Engine/Source/Runtime/AudioCodecEngine/Private/DecoderInputBackCompat.h
--- a/Engine/Source/Runtime/AudioCodecEngine/Private/DecoderInputBackCompat.h
+++ b/Engine/Source/Runtime/AudioCodecEngine/Private/DecoderInputBackCompat.h
@@ -13,15 +13,15 @@ namespace Audio
 	struct FBackCompatInput : public IDecoderInput
 	{
 		FName OldFormatName;
-		TUniquePtr<FSoundWaveProxy> Wave;
+		FSoundWaveProxyPtr Wave;
 		mutable FFormatDescriptorSection Desc;
 		mutable TUniquePtr<ICompressedAudioInfo> OldInfoObject;
 
 		FBackCompatInput(
 			FName InOldFormatName,
-			const FSoundWaveProxy& InWave)
+			const FSoundWaveProxyPtr& InWave)
 			: OldFormatName(InOldFormatName)
-			, Wave(MakeUnique<FSoundWaveProxy>(InWave))
+			, Wave(InWave)
 		{
 		}
 
diff --git a/Engine/Source/Runtime/AudioCodecEngine/Private/DecoderInputFactory.cpp b/Engine/Source/Runtime/AudioCodecEngine/Private/DecoderInputFactory.cpp
--- a/Engine/Source/Runtime/AudioCodecEngine/Private/DecoderInputFactory.cpp
+++ b/Engine/Source/Runtime/AudioCodecEngine/Private/DecoderInputFactory.cpp
@@ -9,8 +9,13 @@ namespace Audio
 {		
 	TUniquePtr<Audio::IDecoderInput> CreateBackCompatDecoderInput(
 		FName InOldFormatName,
-		const FSoundWaveProxy& InSoundWave)
+		const FSoundWaveProxyPtr& InSoundWave)
 	{
-		return MakeUnique<FBackCompatInput>(InOldFormatName, InSoundWave);
+		if (ensure(InSoundWave.IsValid()))
+		{
+			return MakeUnique<FBackCompatInput>(InOldFormatName, InSoundWave);
+		}
+
+		return nullptr;
 	}
 }
diff --git a/Engine/Source/Runtime/AudioCodecEngine/Public/DecoderInputFactory.h b/Engine/Source/Runtime/AudioCodecEngine/Public/DecoderInputFactory.h
--- a/Engine/Source/Runtime/AudioCodecEngine/Public/DecoderInputFactory.h
+++ b/Engine/Source/Runtime/AudioCodecEngine/Public/DecoderInputFactory.h
@@ -16,5 +16,5 @@ namespace Audio
 	// Just loose for now.
 	AUDIOCODECENGINE_API TUniquePtr<IDecoderInput> CreateBackCompatDecoderInput(
 		FName InOldFormatName,
-		const FSoundWaveProxy& InSoundWave );
+		const FSoundWaveProxyPtr& InSoundWave );
 }
diff --git a/Engine/Source/Runtime/AudioMixer/Private/AudioMixerBlueprintLibrary.cpp b/Engine/Source/Runtime/AudioMixer/Private/AudioMixerBlueprintLibrary.cpp
--- a/Engine/Source/Runtime/AudioMixer/Private/AudioMixerBlueprintLibrary.cpp
+++ b/Engine/Source/Runtime/AudioMixer/Private/AudioMixerBlueprintLibrary.cpp
@@ -532,7 +532,7 @@ void UAudioMixerBlueprintLibrary::PrimeSoundForPlayback(USoundWave* SoundWave, c
 	}
 	else
 	{
-		IStreamingManager::Get().GetAudioStreamingManager().RequestChunk(SoundWave->GetThisAsProxy(), 1, [OnLoadCompletion, SoundWave](EAudioChunkLoadResult InResult)
+		IStreamingManager::Get().GetAudioStreamingManager().RequestChunk(SoundWave->CreateSoundWaveProxy(), 1, [OnLoadCompletion, SoundWave](EAudioChunkLoadResult InResult)
 		{
 			AsyncTask(ENamedThreads::GameThread, [OnLoadCompletion, SoundWave, InResult]() {
 				if (InResult == EAudioChunkLoadResult::Completed || InResult == EAudioChunkLoadResult::AlreadyLoaded)
diff --git a/Engine/Source/Runtime/AudioMixer/Private/AudioMixerDevice.cpp b/Engine/Source/Runtime/AudioMixer/Private/AudioMixerDevice.cpp
--- a/Engine/Source/Runtime/AudioMixer/Private/AudioMixerDevice.cpp
+++ b/Engine/Source/Runtime/AudioMixer/Private/AudioMixerDevice.cpp
@@ -586,7 +586,7 @@ namespace Audio
 		return AudioMixerPlatform->CreateCompressedAudioInfo(InSoundWave);
 	}
 
-	class ICompressedAudioInfo* FMixerDevice::CreateCompressedAudioInfo(const FSoundWaveProxy& InSoundWave)
+	class ICompressedAudioInfo* FMixerDevice::CreateCompressedAudioInfo(const FSoundWaveProxyPtr& InSoundWave)
 	{
 		check(AudioMixerPlatform);
 		return AudioMixerPlatform->CreateCompressedAudioInfo(InSoundWave);
diff --git a/Engine/Source/Runtime/AudioMixer/Private/AudioMixerSourceBuffer.cpp b/Engine/Source/Runtime/AudioMixer/Private/AudioMixerSourceBuffer.cpp
--- a/Engine/Source/Runtime/AudioMixer/Private/AudioMixerSourceBuffer.cpp
+++ b/Engine/Source/Runtime/AudioMixer/Private/AudioMixerSourceBuffer.cpp
@@ -368,7 +368,7 @@ namespace Audio
 		if (InDecoder && InDecoder->HasError())
 		{			
 			UE_LOG(LogAudioMixer, Warning, TEXT("Decoder Error, stopping source [%s]"), 
-				*InDecoder->GetStreamingSoundWave().GetFName().ToString());
+				*InDecoder->GetStreamingSoundWave()->GetFName().ToString());
 			bHasError = true;
 			bBufferFinished = true;
 			return false;	
diff --git a/Engine/Source/Runtime/AudioMixer/Public/AudioMixerDevice.h b/Engine/Source/Runtime/AudioMixer/Public/AudioMixerDevice.h
--- a/Engine/Source/Runtime/AudioMixer/Public/AudioMixerDevice.h
+++ b/Engine/Source/Runtime/AudioMixer/Public/AudioMixerDevice.h
@@ -87,7 +87,7 @@ namespace Audio
 		virtual bool SupportsRealtimeDecompression() const override;
 		virtual bool DisablePCMAudioCaching() const override;
 		virtual class ICompressedAudioInfo* CreateCompressedAudioInfo(USoundWave* SoundWave) override;
-		virtual class ICompressedAudioInfo* CreateCompressedAudioInfo(const FSoundWaveProxy& SoundWave) override;
+		virtual class ICompressedAudioInfo* CreateCompressedAudioInfo(const FSoundWaveProxyPtr& SoundWave) override;
 		virtual bool ValidateAPICall(const TCHAR* Function, uint32 ErrorCode) override;
 		virtual bool Exec(UWorld* InWorld, const TCHAR* Cmd, FOutputDevice& Ar) override;
 		virtual void CountBytes(class FArchive& Ar) override;
diff --git a/Engine/Source/Runtime/AudioMixerCore/Public/AudioMixer.h b/Engine/Source/Runtime/AudioMixerCore/Public/AudioMixer.h
--- a/Engine/Source/Runtime/AudioMixerCore/Public/AudioMixer.h
+++ b/Engine/Source/Runtime/AudioMixerCore/Public/AudioMixer.h
@@ -114,6 +114,8 @@ namespace EAudioMixerChannel
 class USoundWave;
 class FSoundWaveProxy;
 class ICompressedAudioInfo;
+using FSoundWaveProxyPtr = TSharedPtr<FSoundWaveProxy, ESPMode::ThreadSafe>;
+
 
 namespace Audio
 {
@@ -454,7 +456,7 @@ namespace Audio
 
 		/** Creates a Compressed audio info class suitable for decompressing this SoundWave. */
 		virtual ICompressedAudioInfo* CreateCompressedAudioInfo(USoundWave* SoundWave) = 0;
-		virtual ICompressedAudioInfo* CreateCompressedAudioInfo(const FSoundWaveProxy& SoundWave) { return nullptr; }
+		virtual ICompressedAudioInfo* CreateCompressedAudioInfo(const FSoundWaveProxyPtr& SoundWave) { return nullptr; }
 
 		/** Return any optional device name defined in platform configuratio. */
 		virtual FString GetDefaultDeviceName() = 0;
diff --git a/Engine/Source/Runtime/Engine/Classes/Sound/SoundWave.h b/Engine/Source/Runtime/Engine/Classes/Sound/SoundWave.h
--- a/Engine/Source/Runtime/Engine/Classes/Sound/SoundWave.h
+++ b/Engine/Source/Runtime/Engine/Classes/Sound/SoundWave.h
@@ -58,6 +58,11 @@ constexpr uint64 InvalidAudioStreamCacheLookupID = TNumericLimits<uint64>::Max()
  */
 struct FStreamedAudioChunk
 {
+	/** Serialization. */
+	void Serialize(FArchive& Ar, UObject* Owner, int32 ChunkIndex);
+
+	void GetCopy(void** OutChunkData);
+
 	/** Size of the chunk of data in bytes including zero padding */
 	int32 DataSize = 0;
 
@@ -67,9 +72,10 @@ struct FStreamedAudioChunk
 	/** Bulk data if stored in the package. */
 	FByteBulkData BulkData;
 
+private:
+	uint8* CachedDataPtr{ nullptr };
 
-	/** Serialization. */
-	void Serialize(FArchive& Ar, UObject* Owner, int32 ChunkIndex);
+public:
 
 #if WITH_EDITORONLY_DATA
 	/** Key if stored in the derived data cache. */
@@ -282,14 +288,14 @@ struct ISoundWaveClient
 	virtual void OnFinishDestroy(class USoundWave* Wave) = 0;
 };
 
-
+using FSoundWaveProxyPtr = TSharedPtr<FSoundWaveProxy, ESPMode::ThreadSafe>;
 class ENGINE_API FSoundWaveProxy : public Audio::TProxyData<FSoundWaveProxy>
 {
 public:
 	IMPL_AUDIOPROXY_CLASS(FSoundWaveProxy);
 
 	explicit FSoundWaveProxy(USoundWave* InWave);
-	FSoundWaveProxy(const FSoundWaveProxy& Other);
+	FSoundWaveProxy(const FSoundWaveProxy& Other) = default;
 
 	Audio::IProxyDataPtr Clone() const override
 	{
@@ -316,7 +322,7 @@ public:
 	bool WasLoadingBehaviorOverridden() const { return bLoadingBehaviorOverridden; }
 	ESoundWaveLoadingBehavior GetLoadingBehavior() const { return LoadingBehavior; }
 
-	TArrayView<const uint8> GetZerothChunk(bool bForImmediatePlayback = false) const;
+	static TArrayView<const uint8> GetZerothChunk(const FSoundWaveProxyPtr& SoundWaveProxy, bool bForImmediatePlayback = false);
 	bool GetChunkData(int32 ChunkIndex, uint8** OutChunkData, bool bMakeSureChunkIsLoaded = false);
 
 	bool IsZerothChunkDataLoaded() const;
@@ -332,10 +338,10 @@ public:
 	FString GetDerivedDataKey() const;
 
 	/** Zeroth Chunk of audio for sources that use Load On Demand. */
-	FBulkDataBuffer<uint8> ZerothChunkData;
+	TSharedPtr<FBulkDataBuffer<uint8>> ZerothChunkData;
 
 	/** The streaming derived data for this sound on this platform. */
-	TUniquePtr<FStreamedAudioPlatformData> RunningPlatformData{ nullptr };
+	TSharedPtr<FStreamedAudioPlatformData> RunningPlatformData{ nullptr };
 
 private:
 	FName NameCached;
@@ -469,6 +475,9 @@ private:
 	uint8 bLoadedFromCookedData : 1;
 #endif // !WITH_EDITOR
 
+	// cached proxy
+	FSoundWaveProxyPtr InternalProxy{ nullptr };
+
 	enum class ESoundWaveResourceState : uint8
 	{
 		NeedsFree,
@@ -725,10 +734,10 @@ public:
 	FFormatContainer CompressedFormatData;
 
 	/** Zeroth Chunk of audio for sources that use Load On Demand. */
-	FBulkDataBuffer<uint8> ZerothChunkData;
+	TSharedPtr<FBulkDataBuffer<uint8>> ZerothChunkData{ MakeShared<FBulkDataBuffer<uint8>>() };
 
 	/** The streaming derived data for this sound on this platform. */
-	FStreamedAudioPlatformData* RunningPlatformData{ nullptr };
+	TSharedPtr<FStreamedAudioPlatformData> RunningPlatformData{ nullptr };
 
 #if WITH_EDITORONLY_DATA
 	TMap<FName, uint32> AsyncLoadingDataFormats;
@@ -777,7 +786,8 @@ public:
 	virtual bool HasCookedAmplitudeEnvelopeData() const override;
 	//~ End USoundBase Interface.
 
-	TUniquePtr<FSoundWaveProxy> CreateSoundWaveProxy();
+	FSoundWaveProxyPtr CreateSoundWaveProxy();
+
 	//~Begin IAudioProxyDataFactory Interface.
 	virtual TUniquePtr<Audio::IProxyData> CreateNewProxyData(const Audio::FProxyDataInitParams& InitParams) override;
 	//~ End IAudioProxyDataFactory Interface.
@@ -1080,10 +1090,7 @@ public:
 		return (ESoundWavePrecacheState)PrecacheState.GetValue();
 	}
 
-	const FSoundWaveProxy& GetThisAsProxy();
-
 private:
-	TUniquePtr<FSoundWaveProxy> InternalProxy{ nullptr };
 	friend class FSoundWaveProxy;
 };
 
diff --git a/Engine/Source/Runtime/Engine/Private/ADPCMAudioInfo.cpp b/Engine/Source/Runtime/Engine/Private/ADPCMAudioInfo.cpp
--- a/Engine/Source/Runtime/Engine/Private/ADPCMAudioInfo.cpp
+++ b/Engine/Source/Runtime/Engine/Private/ADPCMAudioInfo.cpp
@@ -472,13 +472,16 @@ void FADPCMAudioInfo::ProcessSeekRequest()
 	}
 }
 
-bool FADPCMAudioInfo::StreamCompressedInfoInternal(const FSoundWaveProxy& InWaveProxy, struct FSoundQualityInfo* QualityInfo)
+bool FADPCMAudioInfo::StreamCompressedInfoInternal(const FSoundWaveProxyPtr& InWaveProxy, struct FSoundQualityInfo* QualityInfo)
 {
 	FScopeLock ScopeLock(&CurCompressedChunkHandleCriticalSection);
 
 	check(QualityInfo);
-
-	check(StreamingSoundWave.Get() == &InWaveProxy);
+	check(StreamingSoundWave == InWaveProxy);
+	if (!ensure(InWaveProxy.IsValid()))
+	{
+		return false;
+	}
 
 	CurrentChunkIndex = 0;
 
@@ -679,7 +682,7 @@ bool FADPCMAudioInfo::StreamCompressedData(uint8* Destination, bool bLooping, ui
 					}
 
 					// Request the next chunk of data from the streaming engine
-					CurCompressedChunkData = GetLoadedChunk(*StreamingSoundWave, CurrentChunkIndex, CurrentChunkDataSize);
+					CurCompressedChunkData = GetLoadedChunk(StreamingSoundWave, CurrentChunkIndex, CurrentChunkDataSize);
 
 					if(CurCompressedChunkData == nullptr)
 					{
@@ -874,7 +877,7 @@ bool FADPCMAudioInfo::StreamCompressedData(uint8* Destination, bool bLooping, ui
 				}
 
 				// Request the next chunk of data from the streaming engine
-				CurCompressedChunkData = GetLoadedChunk(*StreamingSoundWave, CurrentChunkIndex, CurrentChunkDataSize);
+				CurCompressedChunkData = GetLoadedChunk(StreamingSoundWave, CurrentChunkIndex, CurrentChunkDataSize);
 
 				if(CurCompressedChunkData == nullptr)
 				{
@@ -985,7 +988,7 @@ bool FADPCMAudioInfo::ReleaseStreamChunk(bool bBlockUntilReleased)
 	}
 }
 
-const uint8* FADPCMAudioInfo::GetLoadedChunk(const FSoundWaveProxy& InSoundWave, uint32 ChunkIndex, uint32& OutChunkSize)
+const uint8* FADPCMAudioInfo::GetLoadedChunk(const FSoundWaveProxyPtr& InSoundWave, uint32 ChunkIndex, uint32& OutChunkSize)
 {
 	if (ChunkIndex != 0 && FMath::RandRange(0.0f, 1.0f) < ChanceForIntentionalChunkMissCVar)
 	{
@@ -994,14 +997,14 @@ const uint8* FADPCMAudioInfo::GetLoadedChunk(const FSoundWaveProxy& InSoundWave,
 		return nullptr;
 	}
 
-	if (ChunkIndex >= InSoundWave.GetNumChunks())
+	if (!ensure(InSoundWave.IsValid()) || ChunkIndex >= InSoundWave->GetNumChunks())
 	{
 		OutChunkSize = 0;
 		return nullptr;
 	}
 	else if (ChunkIndex == 0)
 	{
-		TArrayView<const uint8> ZerothChunk = InSoundWave.GetZerothChunk(true);
+		TArrayView<const uint8> ZerothChunk = FSoundWaveProxy::GetZerothChunk(InSoundWave, true);
 		OutChunkSize = ZerothChunk.Num();
 		PreviouslyRequestedChunkIndex = ChunkIndex;
 		return ZerothChunk.GetData();
diff --git a/Engine/Source/Runtime/Engine/Private/AudioDecompress.cpp b/Engine/Source/Runtime/Engine/Private/AudioDecompress.cpp
--- a/Engine/Source/Runtime/Engine/Private/AudioDecompress.cpp
+++ b/Engine/Source/Runtime/Engine/Private/AudioDecompress.cpp
@@ -142,14 +142,18 @@ void IStreamedCompressedInfo::ExpandFile(uint8* DstBuffer, struct FSoundQualityI
 	}
 }
 
-bool IStreamedCompressedInfo::StreamCompressedInfoInternal(const FSoundWaveProxy& InWaveProxy, struct FSoundQualityInfo* QualityInfo)
+bool IStreamedCompressedInfo::StreamCompressedInfoInternal(const FSoundWaveProxyPtr& InWaveProxy, struct FSoundQualityInfo* QualityInfo)
 {
-	check(StreamingSoundWave.Get() == &InWaveProxy);
+	// we have already cached the wave, confirm it is valid and the same
+	if (!(ensureAlways(StreamingSoundWave.IsValid() && (StreamingSoundWave == InWaveProxy))))
+	{
+		return false;
+	}
 
 	// Get the first chunk of audio data (should always be loaded)
 	CurrentChunkIndex = 0;
 	uint32 ChunkSize = 0;
-	const uint8* FirstChunk = GetLoadedChunk(*StreamingSoundWave, CurrentChunkIndex, ChunkSize);
+	const uint8* FirstChunk = GetLoadedChunk(StreamingSoundWave, CurrentChunkIndex, ChunkSize);
 
 	bIsStreaming = true;
 	if (FirstChunk)
@@ -189,7 +193,7 @@ bool IStreamedCompressedInfo::StreamCompressedData(uint8* Destination, bool bLoo
 	if (SrcBufferData == NULL)
 	{
 		uint32 ChunkSize = 0;
-		SrcBufferData = GetLoadedChunk(*StreamingSoundWave, CurrentChunkIndex, ChunkSize);
+		SrcBufferData = GetLoadedChunk(StreamingSoundWave, CurrentChunkIndex, ChunkSize);
 		if (SrcBufferData)
 		{
 			bPrintChunkFailMessage = true;
@@ -274,7 +278,7 @@ bool IStreamedCompressedInfo::StreamCompressedData(uint8* Destination, bool bLoo
 					SrcBufferOffset = 0;
 				}
 
-				SrcBufferData = GetLoadedChunk(*StreamingSoundWave, CurrentChunkIndex, SrcBufferDataSize);
+				SrcBufferData = GetLoadedChunk(StreamingSoundWave, CurrentChunkIndex, SrcBufferDataSize);
 				if (SrcBufferData)
 				{
 					UE_LOG(LogAudio, Log, TEXT("Incremented current chunk from SoundWave'%s' - Chunk %d, Offset %d"), *StreamingSoundWave->GetFName().ToString(), CurrentChunkIndex, SrcBufferOffset);
@@ -362,16 +366,21 @@ uint32 IStreamedCompressedInfo::ZeroBuffer(uint8* Destination, uint32 BufferSize
 }
 
 
-const uint8* IStreamedCompressedInfo::GetLoadedChunk(const FSoundWaveProxy& InSoundWave, uint32 ChunkIndex, uint32& OutChunkSize)
+const uint8* IStreamedCompressedInfo::GetLoadedChunk(const FSoundWaveProxyPtr& InSoundWave, uint32 ChunkIndex, uint32& OutChunkSize)
 {
-	if (ChunkIndex >= InSoundWave.GetNumChunks())
+	if (!ensure(InSoundWave.IsValid()))
+	{
+		OutChunkSize = 0;
+		return nullptr;
+	}
+	else if (ChunkIndex >= InSoundWave->GetNumChunks())
 	{
 		OutChunkSize = 0;
 		return nullptr;
 	}
 	else if (ChunkIndex == 0)
 	{
-		TArrayView<const uint8> ZerothChunk = InSoundWave.GetZerothChunk(true);
+		TArrayView<const uint8> ZerothChunk = FSoundWaveProxy::GetZerothChunk(InSoundWave, true);
 		OutChunkSize = ZerothChunk.Num();
 		return ZerothChunk.GetData();
 	}
@@ -651,7 +660,8 @@ FAsyncAudioDecompressWorker::FAsyncAudioDecompressWorker(USoundWave* InWave, int
 	}
 	else if (GEngine && GEngine->GetMainAudioDevice())
 	{
-		AudioInfo = GEngine->GetMainAudioDevice()->CreateCompressedAudioInfo(Wave->GetThisAsProxy());
+		// alternatively, we could expose InWave::InternalProxy via public getter to avoid the proxy creation
+		AudioInfo = GEngine->GetMainAudioDevice()->CreateCompressedAudioInfo(InWave->CreateSoundWaveProxy());
 	}
 }
 
@@ -770,12 +780,12 @@ bool ICompressedAudioInfo::StreamCompressedInfo(USoundWave* Wave, FSoundQualityI
 		return false;
 	}
 
-	return StreamCompressedInfoInternal(*StreamingSoundWave, QualityInfo);
+	return StreamCompressedInfoInternal(StreamingSoundWave, QualityInfo);
 }
 
-bool ICompressedAudioInfo::StreamCompressedInfo(const FSoundWaveProxy& Wave, FSoundQualityInfo* QualityInfo)
+bool ICompressedAudioInfo::StreamCompressedInfo(const FSoundWaveProxyPtr& Wave, FSoundQualityInfo* QualityInfo)
 {
 	// Create our own copy of the proxy object
-	StreamingSoundWave = MakeUnique<FSoundWaveProxy>(Wave);
-	return StreamCompressedInfoInternal(*StreamingSoundWave, QualityInfo);
+	StreamingSoundWave = Wave;
+	return StreamCompressedInfoInternal(StreamingSoundWave, QualityInfo);
 }
diff --git a/Engine/Source/Runtime/Engine/Private/AudioDerivedData.cpp b/Engine/Source/Runtime/Engine/Private/AudioDerivedData.cpp
--- a/Engine/Source/Runtime/Engine/Private/AudioDerivedData.cpp
+++ b/Engine/Source/Runtime/Engine/Private/AudioDerivedData.cpp
@@ -845,7 +845,7 @@ int32 FStreamedAudioPlatformData::GetChunkFromDDC(int32 ChunkIndex, uint8** OutC
 			if (OutChunkData)
 			{
 				ChunkDataSize = Chunk.BulkData.GetBulkDataSize();
-				Chunk.BulkData.GetCopy((void**)OutChunkData, true);
+				Chunk.GetCopy((void**)OutChunkData);
 			}
 		}
 	}
@@ -868,7 +868,7 @@ int32 FStreamedAudioPlatformData::GetChunkFromDDC(int32 ChunkIndex, uint8** OutC
 		if (OutChunkData)
 		{
 			ChunkDataSize = Chunk.BulkData.GetBulkDataSize();
-			Chunk.BulkData.GetCopy((void**)OutChunkData, true);
+			Chunk.GetCopy((void**)OutChunkData);
 		}
 	}
 #endif // #if WITH_EDITORONLY_DATA
@@ -1451,11 +1451,7 @@ bool FDerivedAudioDataCompressor::Build(TArray<uint8>& OutData)
 
 void USoundWave::CleanupCachedRunningPlatformData()
 {
-	if (RunningPlatformData != NULL)
-	{
-		delete RunningPlatformData;
-		RunningPlatformData = NULL;
-	}
+	RunningPlatformData.Reset();
 }
 
 
@@ -1501,7 +1497,7 @@ void USoundWave::SerializeCookedPlatformData(FArchive& Ar)
 		check(RunningPlatformData == NULL);
 
 		// Don't serialize streaming data on servers, even if this platform supports streaming in theory
-		RunningPlatformData = new FStreamedAudioPlatformData();
+		RunningPlatformData = MakeShared<FStreamedAudioPlatformData>();
 		RunningPlatformData->Serialize(Ar, this);
 	}
 }
@@ -1518,7 +1514,7 @@ void USoundWave::CachePlatformData(bool bAsyncCache)
 	{
 		if (RunningPlatformData == NULL)
 		{
-			RunningPlatformData = new FStreamedAudioPlatformData();
+			RunningPlatformData = MakeShared<FStreamedAudioPlatformData>();
 		}
 
 		RunningPlatformData->Cache(*this, CompressionOverrides, AudioFormat, bAsyncCache ? EStreamedAudioCacheFlags::Async : EStreamedAudioCacheFlags::None);
diff --git a/Engine/Source/Runtime/Engine/Private/AudioStreaming.cpp b/Engine/Source/Runtime/Engine/Private/AudioStreaming.cpp
--- a/Engine/Source/Runtime/Engine/Private/AudioStreaming.cpp
+++ b/Engine/Source/Runtime/Engine/Private/AudioStreaming.cpp
@@ -127,9 +127,9 @@ void FStreamingWaveData::FreeResources()
 	}
 }
 
-bool FStreamingWaveData::Initialize(const FSoundWaveProxy& InSoundWave, FLegacyAudioStreamingManager* InAudioStreamingManager)
+bool FStreamingWaveData::Initialize(const FSoundWaveProxyPtr& InSoundWave, FLegacyAudioStreamingManager* InAudioStreamingManager)
 {
-	if (!InSoundWave.GetNumChunks())
+	if (!ensure(InSoundWave.IsValid()) && !InSoundWave->GetNumChunks())
 	{
 #if WITH_EDITOR
 		UE_LOG(LogAudio, Display, TEXT("Failed to initialize streaming wave data due to lack of serialized stream chunks. Error during stream cooking."));
@@ -139,7 +139,7 @@ bool FStreamingWaveData::Initialize(const FSoundWaveProxy& InSoundWave, FLegacyA
 		return false;
 	}
 
-	SoundWave = MakeUnique<FSoundWaveProxy>(InSoundWave);
+	SoundWave = InSoundWave;
 	AudioStreamingManager = InAudioStreamingManager;
 
 	// Always get the first chunk of data so we can play immediately
@@ -598,23 +598,23 @@ void FLegacyAudioStreamingManager::UpdateResourceStreaming(float DeltaTime, bool
 
 	for (ICompressedAudioInfo* Decoder : CompressedAudioInfos)
 	{
-		const FSoundWaveProxy& SoundWave = Decoder->GetStreamingSoundWave();
-		FStreamingWaveData** WaveDataPtr = StreamingSoundWaves.Find(SoundWave.GetFObjectKey());
+		const FSoundWaveProxyPtr& SoundWave = Decoder->GetStreamingSoundWave();
+		FStreamingWaveData** WaveDataPtr = StreamingSoundWaves.Find(SoundWave->GetFObjectKey());
 		if (WaveDataPtr && (*WaveDataPtr)->PendingChunkChangeRequestStatus.GetValue() == AudioState_ReadyFor_Requests)
 		{
 			FStreamingWaveData* WaveData = *WaveDataPtr;
 			// Request the chunk the source is using and the one after that
-			FWaveRequest& WaveRequest = GetWaveRequest(SoundWave.GetFObjectKey());
+			FWaveRequest& WaveRequest = GetWaveRequest(SoundWave->GetFObjectKey());
 			int32 SourceChunk = Decoder->GetCurrentChunkIndex();
-			if (SourceChunk >= 0 && SourceChunk < (int32)SoundWave.GetNumChunks())
+			if (SourceChunk >= 0 && SourceChunk < (int32)SoundWave->GetNumChunks())
 			{
 				WaveRequest.RequiredIndices.AddUnique(SourceChunk);
-				WaveRequest.RequiredIndices.AddUnique((SourceChunk + 1) % SoundWave.GetNumChunks());
+				WaveRequest.RequiredIndices.AddUnique((SourceChunk + 1) % SoundWave->GetNumChunks());
 				WaveRequest.bPrioritiseRequest = true;
 			}
 			else
 			{
-				UE_LOG(LogAudio, Log, TEXT("Invalid chunk request curIndex=%d numChunks=%d\n"), SourceChunk, SoundWave.GetNumChunks());
+				UE_LOG(LogAudio, Log, TEXT("Invalid chunk request curIndex=%d numChunks=%d\n"), SourceChunk, SoundWave->GetNumChunks());
 			}
 		}
 	}
@@ -699,18 +699,18 @@ void FLegacyAudioStreamingManager::NotifyLevelOffset(class ULevel* Level, const
 {
 }
 
-void FLegacyAudioStreamingManager::AddStreamingSoundWave(const FSoundWaveProxy& SoundWave)
+void FLegacyAudioStreamingManager::AddStreamingSoundWave(const FSoundWaveProxyPtr& SoundWave)
 {
-	if (FPlatformProperties::SupportsAudioStreaming() && SoundWave.IsStreaming())
+	if (ensure(SoundWave.IsValid()) && FPlatformProperties::SupportsAudioStreaming() && SoundWave->IsStreaming())
 	{
 		FScopeLock Lock(&CriticalSection);
 
-		if (StreamingSoundWaves.FindRef(SoundWave.GetFObjectKey()) == nullptr)
+		if (StreamingSoundWaves.FindRef(SoundWave->GetFObjectKey()) == nullptr)
 		{
 			FStreamingWaveData* NewStreamingWaveData = new FStreamingWaveData;
 			if (NewStreamingWaveData->Initialize(SoundWave, this))
 			{
-				StreamingSoundWaves.Add(SoundWave.GetFObjectKey(), NewStreamingWaveData);
+				StreamingSoundWaves.Add(SoundWave->GetFObjectKey(), NewStreamingWaveData);
 			}
 			else
 			{
@@ -721,13 +721,18 @@ void FLegacyAudioStreamingManager::AddStreamingSoundWave(const FSoundWaveProxy&
 	}
 }
 
-void FLegacyAudioStreamingManager::RemoveStreamingSoundWave(const FSoundWaveProxy& SoundWave)
+void FLegacyAudioStreamingManager::RemoveStreamingSoundWave(const FSoundWaveProxyPtr& SoundWave)
 {
+	if (!ensure(SoundWave.IsValid()))
+	{
+		return;
+	}
+
 	FScopeLock Lock(&CriticalSection);
-	FStreamingWaveData* WaveData = StreamingSoundWaves.FindRef(SoundWave.GetFObjectKey());
+	FStreamingWaveData* WaveData = StreamingSoundWaves.FindRef(SoundWave->GetFObjectKey());
 	if (WaveData)
 	{
-		StreamingSoundWaves.Remove(SoundWave.GetFObjectKey());
+		StreamingSoundWaves.Remove(SoundWave->GetFObjectKey());
 
 		// Free the resources of the streaming wave data. This blocks pending IO requests
 		WaveData->FreeResources();
@@ -749,7 +754,7 @@ void FLegacyAudioStreamingManager::RemoveStreamingSoundWave(const FSoundWaveProx
 		}
 		delete WaveData;
 	}
-	WaveRequests.Remove(SoundWave.GetFObjectKey());
+	WaveRequests.Remove(SoundWave->GetFObjectKey());
 }
 
 void FLegacyAudioStreamingManager::AddDecoder(ICompressedAudioInfo* InCompressedAudioInfo)
@@ -764,20 +769,24 @@ void FLegacyAudioStreamingManager::RemoveDecoder(ICompressedAudioInfo* InCompres
 	CompressedAudioInfos.Remove(InCompressedAudioInfo);
 }
 
-bool FLegacyAudioStreamingManager::IsManagedStreamingSoundWave(const FSoundWaveProxy&  SoundWave) const
+bool FLegacyAudioStreamingManager::IsManagedStreamingSoundWave(const FSoundWaveProxyPtr& SoundWave) const
 {
 	FScopeLock Lock(&CriticalSection);
-	return StreamingSoundWaves.FindRef(SoundWave.GetFObjectKey()) != NULL;
+	return ensure(SoundWave.IsValid()) && (StreamingSoundWaves.FindRef(SoundWave->GetFObjectKey()) != NULL);
 }
 
-bool FLegacyAudioStreamingManager::IsStreamingInProgress(const FSoundWaveProxy&  SoundWave)
+bool FLegacyAudioStreamingManager::IsStreamingInProgress(const FSoundWaveProxyPtr&  SoundWave)
 {
 	FScopeLock Lock(&CriticalSection);
-	FStreamingWaveData* WaveData = StreamingSoundWaves.FindRef(SoundWave.GetFObjectKey());
-	if (WaveData)
+	if (ensure(SoundWave.IsValid()))
 	{
-		return WaveData->UpdateStreamingStatus();
+		FStreamingWaveData* WaveData = StreamingSoundWaves.FindRef(SoundWave->GetFObjectKey());
+		if (WaveData)
+		{
+			return WaveData->UpdateStreamingStatus();
+		}
 	}
+
 	return false;
 }
 
@@ -882,16 +891,16 @@ bool FLegacyAudioStreamingManager::IsManagedStreamingSoundSource(const FSoundSou
 	return StreamingSoundSources.FindByKey(SoundSource) != NULL;
 }
 
-bool FLegacyAudioStreamingManager::RequestChunk(const FSoundWaveProxy&, uint32 ChunkIndex, TFunction<void(EAudioChunkLoadResult)> OnLoadCompleted, ENamedThreads::Type ThreadToCallOnLoadCompletedOn, bool bForImmediatePlayback)
+bool FLegacyAudioStreamingManager::RequestChunk(const FSoundWaveProxyPtr&, uint32 ChunkIndex, TFunction<void(EAudioChunkLoadResult)> OnLoadCompleted, ENamedThreads::Type ThreadToCallOnLoadCompletedOn, bool bForImmediatePlayback)
 {
 	UE_LOG(LogAudio, Warning, TEXT("RequestChunk is only supported in Stream Caching."));
 	return false;
 }
 
-FAudioChunkHandle FLegacyAudioStreamingManager::GetLoadedChunk(const FSoundWaveProxy&  SoundWave, uint32 ChunkIndex, bool bBlockForLoad, bool bForImmediatePlayback) const
+FAudioChunkHandle FLegacyAudioStreamingManager::GetLoadedChunk(const FSoundWaveProxyPtr&  SoundWave, uint32 ChunkIndex, bool bBlockForLoad, bool bForImmediatePlayback) const
 {
 	// Check for the spoof of failing to load a stream chunk
-	if (SpoofFailedStreamChunkLoad > 0)
+	if (ensure(SoundWave.IsValid()) && SpoofFailedStreamChunkLoad > 0)
 	{
 		return FAudioChunkHandle();
 	}
@@ -902,7 +911,7 @@ FAudioChunkHandle FLegacyAudioStreamingManager::GetLoadedChunk(const FSoundWaveP
 		return FAudioChunkHandle();
 	}
 
-	const FStreamingWaveData* WaveData = StreamingSoundWaves.FindRef(SoundWave.GetFObjectKey());
+	const FStreamingWaveData* WaveData = StreamingSoundWaves.FindRef(SoundWave->GetFObjectKey());
 	if (WaveData)
 	{
 		if (WaveData->LoadedChunkIndices.Contains(ChunkIndex))
@@ -912,7 +921,7 @@ FAudioChunkHandle FLegacyAudioStreamingManager::GetLoadedChunk(const FSoundWaveP
 				if (WaveData->LoadedChunks[Index].Index == ChunkIndex)
 				{
 					CriticalSection.Unlock();
-					return BuildChunkHandle(WaveData->LoadedChunks[Index].Data, WaveData->LoadedChunks[Index].AudioDataSize, SoundWave, SoundWave.GetFName(), ChunkIndex, InvalidAudioStreamCacheLookupID);
+					return BuildChunkHandle(WaveData->LoadedChunks[Index].Data, WaveData->LoadedChunks[Index].AudioDataSize, SoundWave, SoundWave->GetFName(), ChunkIndex, InvalidAudioStreamCacheLookupID);
 				}
 			}
 		}
diff --git a/Engine/Source/Runtime/Engine/Private/AudioStreamingCache.cpp b/Engine/Source/Runtime/Engine/Private/AudioStreamingCache.cpp
--- a/Engine/Source/Runtime/Engine/Private/AudioStreamingCache.cpp
+++ b/Engine/Source/Runtime/Engine/Private/AudioStreamingCache.cpp
@@ -235,27 +235,24 @@ static FAutoConsoleCommand GDisableProfilingAudioCacheCommand(
 
 
 FAudioChunkCache::FChunkKey::FChunkKey(const FChunkKey& Other)
-	: SoundWaveProxyPtr(Other.SoundWaveProxyPtr.IsValid() ? MakeUnique<FSoundWaveProxy>(*Other.SoundWaveProxyPtr) : nullptr)
-	, SoundWaveName(SoundWaveProxyPtr->GetFName())
+	: SoundWaveProxyPtr(Other.SoundWaveProxyPtr)
 	, ChunkIndex(Other.ChunkIndex)
 	, ObjectKey(Other.ObjectKey)
 #if WITH_EDITOR
 	, ChunkRevision(Other.ChunkRevision)
 #endif // #if WITH_EDITOR
 {
-	SoundWaveProxyPtr->ReleaseCompressedAudio();
+	if (SoundWaveProxyPtr.IsValid())
+	{
+		SoundWaveName = SoundWaveProxyPtr->GetFName();
+		SoundWaveProxyPtr->ReleaseCompressedAudio();
+	}
 }
 
 
 FAudioChunkCache::FChunkKey& FAudioChunkCache::FChunkKey::operator=(const FChunkKey& Other)
 {
-	SoundWaveProxyPtr.Reset();
-
-	if (Other.SoundWaveProxyPtr.IsValid())
-	{
-		SoundWaveProxyPtr = MakeUnique<FSoundWaveProxy>(*Other.SoundWaveProxyPtr);
-	}
-
+	SoundWaveProxyPtr = Other.SoundWaveProxyPtr;
 	SoundWaveName = Other.SoundWaveName;
 	ChunkIndex = Other.ChunkIndex;
 	ObjectKey = Other.ObjectKey;
@@ -268,21 +265,23 @@ FAudioChunkCache::FChunkKey& FAudioChunkCache::FChunkKey::operator=(const FChunk
 	return *this;
 }
 
-FAudioChunkCache::FChunkKey::FChunkKey(const FSoundWaveProxy& InSoundWave, uint32 InChunkIndex
+FAudioChunkCache::FChunkKey::FChunkKey(const FSoundWaveProxyPtr& InSoundWave, uint32 InChunkIndex
 #if WITH_EDITOR
 	, uint32 InChunkRevision
 #endif // #if WITH_EDITOR
 )
 
-	: SoundWaveProxyPtr(MakeUnique<FSoundWaveProxy>(InSoundWave))
-	, SoundWaveName(InSoundWave.GetFName())
+	: SoundWaveProxyPtr(InSoundWave)
 	, ChunkIndex(InChunkIndex)
-	, ObjectKey(InSoundWave.GetFObjectKey())
 #if WITH_EDITOR
 	, ChunkRevision(InChunkRevision)
 #endif // #if WITH_EDITOR
 {
-	SoundWaveProxyPtr->ReleaseCompressedAudio();
+	if (FSoundWaveProxyPtr ProxyPtr = SoundWaveProxyPtr)
+	{
+		SoundWaveName = ProxyPtr->GetFName();
+		ObjectKey = ProxyPtr->GetFObjectKey();
+	}
 }
 
 FAudioChunkCache::FChunkKey::FChunkKey(FName InSoundWaveName, uint32 InChunkIndex)
@@ -314,55 +313,68 @@ bool FAudioChunkCache::FChunkKey::operator==(const FChunkKey& Other) const
 
 bool FAudioChunkCache::FChunkKey::IsChunkStale()
 {
-	return !SoundWaveProxyPtr.IsValid()
+	FSoundWaveProxyPtr ProxyPtr = SoundWaveProxyPtr;
+
+	return !ProxyPtr.IsValid()
 #if WITH_EDITOR
-		|| (SoundWaveProxyPtr->GetCurrentChunkRevision() != ChunkRevision)
+		|| (ProxyPtr->GetCurrentChunkRevision() != ChunkRevision)
 #endif // #if WITH_EDITOR
 		;
 }
 
 ESoundWaveLoadingBehavior FAudioChunkCache::FChunkKey::GetLoadingBehavior() const
 {
-	check(SoundWaveProxyPtr.IsValid());
-	return SoundWaveProxyPtr->GetLoadingBehavior();
+	FSoundWaveProxyPtr ProxyPtr = SoundWaveProxyPtr;
+	check(ProxyPtr.IsValid());
+	return ProxyPtr->GetLoadingBehavior();
 }
 
 FStreamedAudioChunk& FAudioChunkCache::FChunkKey::GetChunk(uint32 InChunkIndex) const
 {
-	check(SoundWaveProxyPtr.IsValid());
-	check((ChunkIndex < (uint32)SoundWaveProxyPtr->GetNumChunks()));
+	FSoundWaveProxyPtr ProxyPtr = SoundWaveProxyPtr;
 
-	check(SoundWaveProxyPtr->RunningPlatformData);
+	check(ProxyPtr.IsValid());
+	check((ChunkIndex < (uint32)ProxyPtr->GetNumChunks()));
+
+	check(ProxyPtr->RunningPlatformData);
 
 	// This function shouldn't be called on audio marked "ForceInline."
-	ensureMsgf(SoundWaveProxyPtr->RunningPlatformData, TEXT("Calling GetNumChunks on a FSoundWaveProxy without RunnigPlatformData is not allowed! SoundWave: %s - %s")
-		, *SoundWaveProxyPtr->GetFName().ToString(), EnumToString(SoundWaveProxyPtr->GetLoadingBehavior()));
+	ensureMsgf(ProxyPtr->RunningPlatformData, TEXT("Calling GetNumChunks on a FSoundWaveProxy without RunnigPlatformData is not allowed! SoundWave: %s - %s")
+		, *ProxyPtr->GetFName().ToString(), EnumToString(ProxyPtr->GetLoadingBehavior()));
 
-	return SoundWaveProxyPtr->RunningPlatformData->Chunks[ChunkIndex];
+	return ProxyPtr->RunningPlatformData->Chunks[ChunkIndex];
 }
 
 uint32 FAudioChunkCache::FChunkKey::GetNumChunks() const
 {
-	check(SoundWaveProxyPtr.IsValid());
-	return SoundWaveProxyPtr->GetNumChunks();
+	FSoundWaveProxyPtr ProxyPtr = SoundWaveProxyPtr;
+
+	check(ProxyPtr.IsValid());
+	return ProxyPtr->GetNumChunks();
 }
 
 bool FAudioChunkCache::FChunkKey::IsRetainingAudio() const
 {
-	check(SoundWaveProxyPtr.IsValid());
-	return SoundWaveProxyPtr->IsRetainingAudio();
+	FSoundWaveProxyPtr ProxyPtr = SoundWaveProxyPtr;
+
+	check(ProxyPtr.IsValid());
+	return ProxyPtr->IsRetainingAudio();
 }
 
 void FAudioChunkCache::FChunkKey::ReleaseCompressedAudio()
 {
-	check(SoundWaveProxyPtr.IsValid());
-	return SoundWaveProxyPtr->ReleaseCompressedAudio();
+	FSoundWaveProxyPtr ProxyPtr = SoundWaveProxyPtr;
+
+	check(ProxyPtr.IsValid());
+	ProxyPtr->ReleaseCompressedAudio();
 }
 
 bool FAudioChunkCache::FChunkKey::WasLoadingBehaviorOverridden() const
 {
-	check(SoundWaveProxyPtr.IsValid());
-	return SoundWaveProxyPtr->WasLoadingBehaviorOverridden();
+	FSoundWaveProxyPtr ProxyPtr = SoundWaveProxyPtr;
+
+	check(ProxyPtr.IsValid());
+	return ProxyPtr->WasLoadingBehaviorOverridden();
 }
 
 FCachedAudioStreamingManager::FCachedAudioStreamingManager(const FCachedAudioStreamingManagerParams& InitParams)
@@ -433,12 +445,12 @@ void FCachedAudioStreamingManager::NotifyLevelOffset(class ULevel* Level, const
 	// Unused.
 }
 
-void FCachedAudioStreamingManager::AddStreamingSoundWave(const FSoundWaveProxy& SoundWave)
+void FCachedAudioStreamingManager::AddStreamingSoundWave(const FSoundWaveProxyPtr& SoundWave)
 {
 	// Unused.
 }
 
-void FCachedAudioStreamingManager::RemoveStreamingSoundWave(const FSoundWaveProxy& SoundWave)
+void FCachedAudioStreamingManager::RemoveStreamingSoundWave(const FSoundWaveProxyPtr& SoundWave)
 {
 	// Unused.
 }
@@ -453,14 +465,14 @@ void FCachedAudioStreamingManager::RemoveDecoder(ICompressedAudioInfo* InCompres
 	//Unused.
 }
 
-bool FCachedAudioStreamingManager::IsManagedStreamingSoundWave(const FSoundWaveProxy&  SoundWave) const
+bool FCachedAudioStreamingManager::IsManagedStreamingSoundWave(const FSoundWaveProxyPtr&  SoundWave) const
 {
 	// Unused. The concept of a sound wave being "managed" doesn't apply here.
 	checkf(false, TEXT("Not Implemented!"));
 	return true;
 }
 
-bool FCachedAudioStreamingManager::IsStreamingInProgress(const FSoundWaveProxy&  SoundWave)
+bool FCachedAudioStreamingManager::IsStreamingInProgress(const FSoundWaveProxyPtr&  SoundWave)
 {
 	// This function is used in USoundWave cleanup.
 	// Since this manager owns the binary data we are streaming off of,
@@ -468,7 +480,7 @@ bool FCachedAudioStreamingManager::IsStreamingInProgress(const FSoundWaveProxy&
 	// There are NO sound sources playing with this Sound Wave.
 	//
 	// This is because a playing sound source might kick off a load for a new chunk,
-	// which dereferences the corresponding USoundWave.
+	// which dereferences the corresponding USoundWave
 	//
 	// As of right now, this is handled by USoundWave::FreeResources(), called
 	// by USoundWave::IsReadyForFinishDestroy.
@@ -497,11 +509,16 @@ bool FCachedAudioStreamingManager::IsManagedStreamingSoundSource(const FSoundSou
 	return true;
 }
 
-FAudioChunkHandle FCachedAudioStreamingManager::GetLoadedChunk(const FSoundWaveProxy& SoundWave, uint32 ChunkIndex, bool bBlockForLoad, bool bForImmediatePlayback) const
+FAudioChunkHandle FCachedAudioStreamingManager::GetLoadedChunk(const FSoundWaveProxyPtr& SoundWave, uint32 ChunkIndex, bool bBlockForLoad, bool bForImmediatePlayback) const
 {
 	LLM_SCOPE(ELLMTag::AudioStreamCache);
 	bBlockForLoad |= (ForceBlockForLoadCVar != 0);
 
+	if (!ensure(SoundWave.IsValid()))
+	{
+		return {};
+	}
+
 	// If this sound wave is managed by a cache, use that to get the chunk:
 	FAudioChunkCache* Cache = GetCacheForWave(SoundWave);
 	if (Cache)
@@ -513,13 +530,13 @@ FAudioChunkHandle FCachedAudioStreamingManager::GetLoadedChunk(const FSoundWaveP
 			  SoundWave
 			, ChunkIndex
 #if WITH_EDITOR
-			, (uint32)SoundWave.GetCurrentChunkRevision()
+			, (uint32)SoundWave->GetCurrentChunkRevision()
 #endif
 		);
 
 		if (!FAudioChunkCache::IsKeyValid(ChunkKey))
 		{
-			UE_LOG(LogAudioStreamCaching, Warning, TEXT("Invalid Chunk Index %d Requested for Wave %s!"), ChunkIndex, *SoundWave.GetFName().ToString());
+			UE_LOG(LogAudioStreamCaching, Warning, TEXT("Invalid Chunk Index %d Requested for Wave %s!"), ChunkIndex, *SoundWave->GetFName().ToString());
 			return FAudioChunkHandle();
 		}
 
@@ -528,12 +545,12 @@ FAudioChunkHandle FCachedAudioStreamingManager::GetLoadedChunk(const FSoundWaveP
 		TArrayView<uint8> LoadedChunk = Cache->GetChunk(ChunkKey, bBlockForLoad, (bForImmediatePlayback || bBlockForLoad), LookupIDForChunk);
 
 		// Ensure that, if we requested a synchronous load of this chunk, we didn't fail to load said chunk.
-		UE_CLOG(bBlockForLoad && !LoadedChunk.GetData(), LogAudioStreamCaching, Display, TEXT("Synchronous load of chunk index %d for SoundWave %s failed to return any data. Likely because the cache was blown."), ChunkIndex, *SoundWave.GetFName().ToString());
+		UE_CLOG(bBlockForLoad && !LoadedChunk.GetData(), LogAudioStreamCaching, Display, TEXT("Synchronous load of chunk index %d for SoundWave %s failed to return any data. Likely because the cache was blown."), ChunkIndex, *SoundWave->GetFName().ToString());
 
 		// Set the updated cache offset for this chunk index.
 		Cache->SetCacheLookupIDForChunk(ChunkKey, LookupIDForChunk);
 
-		UE_CLOG(!bBlockForLoad && !LoadedChunk.GetData(), LogAudioStreamCaching, Verbose, TEXT("GetLoadedChunk called for chunk index %d of SoundWave %s when audio was not loaded yet. This will result in latency."), ChunkIndex, *SoundWave.GetFName().ToString());
+		UE_CLOG(!bBlockForLoad && !LoadedChunk.GetData(), LogAudioStreamCaching, Verbose, TEXT("GetLoadedChunk called for chunk index %d of SoundWave %s when audio was not loaded yet. This will result in latency."), ChunkIndex, *SoundWave->GetFName().ToString());
 
 		// Finally, if there's a chunk after this in the sound, request that it is in the cache.
 		const int32 NextChunk = GetNextChunkIndex(SoundWave, ChunkIndex);
@@ -544,14 +561,14 @@ FAudioChunkHandle FCachedAudioStreamingManager::GetLoadedChunk(const FSoundWaveP
 				SoundWave
 				, ((uint32)NextChunk)
 #if WITH_EDITOR
-				, (uint32)SoundWave.GetCurrentChunkRevision()
+				, (uint32)SoundWave->GetCurrentChunkRevision()
 #endif
 			);
 
 			uint64 LookupIDForNextChunk = Cache->AddOrTouchChunk(NextChunkKey, [](EAudioChunkLoadResult) {}, ENamedThreads::AnyThread, false);
 			if (LookupIDForNextChunk == InvalidAudioStreamCacheLookupID)
 			{
-				UE_LOG(LogAudioStreamCaching, Warning, TEXT("Failed to add follow-up chunk for %s. This should not happen with our current TrimMemory() logic"), *SoundWave.GetFName().ToString());
+				UE_LOG(LogAudioStreamCaching, Warning, TEXT("Failed to add follow-up chunk for %s. This should not happen with our current TrimMemory() logic"), *SoundWave->GetFName().ToString());
 			}
 			else
 			{
@@ -559,23 +576,23 @@ FAudioChunkHandle FCachedAudioStreamingManager::GetLoadedChunk(const FSoundWaveP
 			}
 		}
 
-		return BuildChunkHandle(LoadedChunk.GetData(), LoadedChunk.Num(), SoundWave, SoundWave.GetFName(), ChunkIndex, LookupIDForChunk);
+		return BuildChunkHandle(LoadedChunk.GetData(), LoadedChunk.Num(), SoundWave, SoundWave->GetFName(), ChunkIndex, LookupIDForChunk);
 	}
 	else
 	{
-		ensureMsgf(false, TEXT("Failed to find cache for wave %s. Are you sure this is a streaming wave?"), *SoundWave.GetFName().ToString());
+		ensureMsgf(false, TEXT("Failed to find cache for wave %s. Are you sure this is a streaming wave?"), *SoundWave->GetFName().ToString());
 		return FAudioChunkHandle();
 	}
 }
 
-FAudioChunkCache* FCachedAudioStreamingManager::GetCacheForWave(const FSoundWaveProxy&  InSoundWave) const
+FAudioChunkCache* FCachedAudioStreamingManager::GetCacheForWave(const FSoundWaveProxyPtr&  InSoundWave) const
 {
 	LLM_SCOPE(ELLMTag::AudioStreamCache);
 
 	// We only cache chunks beyond the zeroth chunk of audio (which is inlined directly on the asset)
-	if (InSoundWave.GetNumChunks() > 1)
+	if (ensure(InSoundWave.IsValid()) && InSoundWave->GetNumChunks() > 1)
 	{
-		const int32 SoundWaveChunkSize = InSoundWave.GetChunk(1).AudioDataSize;
+		const int32 SoundWaveChunkSize = InSoundWave->GetChunk(1).AudioDataSize;
 		return GetCacheForChunkSize(SoundWaveChunkSize);
 	}
 	else
@@ -604,15 +621,19 @@ FAudioChunkCache* FCachedAudioStreamingManager::GetCacheForChunkSize(uint32 InCh
 	return nullptr;
 }
 
-int32 FCachedAudioStreamingManager::GetNextChunkIndex(const FSoundWaveProxy&  InSoundWave, uint32 CurrentChunkIndex) const
+int32 FCachedAudioStreamingManager::GetNextChunkIndex(const FSoundWaveProxyPtr&  InSoundWave, uint32 CurrentChunkIndex) const
 {
 	LLM_SCOPE(ELLMTag::AudioStreamCache);
 
 	// TODO: Figure out a way to tell whether this wave is looping or not. For now we always prime the first chunk
 	// during the playback of the last chunk.
 	// if(bNotLooping) return ((int32) CurrentChunkIndex) < (InSoundWave->GetNumChunks() - 1);
+	if (!ensure(InSoundWave.IsValid()))
+	{
+		return INDEX_NONE;
+	}
 
-	const int32 NumChunksTotal = InSoundWave.GetNumChunks();
+	const int32 NumChunksTotal = InSoundWave->GetNumChunks();
 	if (NumChunksTotal <= 2)
 	{
 		// If there's only one chunk to cache (besides the zeroth chunk, which is inlined),
@@ -638,7 +659,7 @@ void FCachedAudioStreamingManager::AddReferenceToChunk(const FAudioChunkHandle&
 	check(Cache);
 
 	const FAudioChunkCache::FChunkKey ChunkKey(
-		  *InHandle.CorrespondingWave
+		  InHandle.CorrespondingWave.Pin()
 		, ((uint32)InHandle.ChunkIndex)
 #if WITH_EDITOR
 		, InHandle.ChunkGeneration
@@ -655,7 +676,7 @@ void FCachedAudioStreamingManager::RemoveReferenceToChunk(const FAudioChunkHandl
 	check(Cache);
 
 	const FAudioChunkCache::FChunkKey ChunkKey(
-		  *InHandle.CorrespondingWave
+		  InHandle.CorrespondingWave.Pin()
 		, ((uint32)InHandle.ChunkIndex)
 #if WITH_EDITOR
 		, InHandle.ChunkGeneration
@@ -665,17 +686,18 @@ void FCachedAudioStreamingManager::RemoveReferenceToChunk(const FAudioChunkHandl
 	Cache->RemoveReferenceToChunk(ChunkKey);
 }
 
-bool FCachedAudioStreamingManager::RequestChunk(const FSoundWaveProxy& SoundWave, uint32 ChunkIndex, TFunction<void(EAudioChunkLoadResult)> OnLoadCompleted, ENamedThreads::Type ThreadToCallOnLoadCompletedOn, bool bForImmediatePlayback)
+bool FCachedAudioStreamingManager::RequestChunk(const FSoundWaveProxyPtr& SoundWave, uint32 ChunkIndex, TFunction<void(EAudioChunkLoadResult)> OnLoadCompleted, ENamedThreads::Type ThreadToCallOnLoadCompletedOn, bool bForImmediatePlayback)
 {
 	LLM_SCOPE(ELLMTag::AudioStreamCache);
 	FAudioChunkCache* Cache = GetCacheForWave(SoundWave);
-	if (Cache)
+
+	if (Cache && ensure(SoundWave.IsValid()))
 	{
 		FAudioChunkCache::FChunkKey ChunkKey(
 			SoundWave
 			, ChunkIndex
 #if WITH_EDITOR
-			, (uint32)SoundWave.GetCurrentChunkRevision()
+			, (uint32)SoundWave->GetCurrentChunkRevision()
 #endif
 		);
 
@@ -688,7 +710,7 @@ bool FCachedAudioStreamingManager::RequestChunk(const FSoundWaveProxy& SoundWave
 	else
 	{
 		// This can hit if an out of bounds chunk was requested, or the zeroth chunk was requested from the streaming manager.
-		ensureMsgf(false, TEXT("GetCacheForWave failed for SoundWave %s!"), *SoundWave.GetFName().ToString());
+		ensureMsgf(false, TEXT("GetCacheForWave failed for SoundWave %s!"), *SoundWave->GetFName().ToString());
 
 		return false;
 	}
@@ -861,22 +883,8 @@ TArrayView<uint8> FAudioChunkCache::GetChunk(const FChunkKey& InKey, bool bBlock
 			// Reallocate our chunk data This allows us to shrink if possible.
 			FoundElement->ChunkData = (uint8*)FMemory::Realloc(FoundElement->ChunkData, ChunkAudioDataSize);
 
-			if (Chunk.DataSize != ChunkAudioDataSize)
-			{
-				// Unfortunately, GetCopy will write out the full zero-padded length of the bulk data,
-				// rather than just the audio data. So we set the array to Chunk.DataSize, then shrink to Chunk.AudioDataSize.
-				TArray<uint8> TempChunkBuffer;
-				TempChunkBuffer.AddUninitialized(Chunk.DataSize);
-				void* DataDestPtr = TempChunkBuffer.GetData();
-				Chunk.BulkData.GetCopy(&DataDestPtr, true);
-
-				FMemory::Memcpy(FoundElement->ChunkData, TempChunkBuffer.GetData(), ChunkAudioDataSize);
-			}
-			else
-			{
-				void* DataDestPtr = FoundElement->ChunkData;
-				Chunk.BulkData.GetCopy(&DataDestPtr, true);
-			}
+			void* DataDestPtr = FoundElement->ChunkData;
+			Chunk.GetCopy(&DataDestPtr);
 
 			MemoryCounterBytes += ChunkAudioDataSize;
 
diff --git a/Engine/Source/Runtime/Engine/Private/ContentStreaming.cpp b/Engine/Source/Runtime/Engine/Private/ContentStreaming.cpp
--- a/Engine/Source/Runtime/Engine/Private/ContentStreaming.cpp
+++ b/Engine/Source/Runtime/Engine/Private/ContentStreaming.cpp
@@ -1405,14 +1405,14 @@ FAudioChunkHandle::FAudioChunkHandle()
 {
 }
 
-FAudioChunkHandle::FAudioChunkHandle(const uint8* InData, uint32 NumBytes, const FSoundWaveProxy&  InSoundWave, const FName& SoundWaveName, uint32 InChunkIndex, uint64 InCacheLookupID)
+FAudioChunkHandle::FAudioChunkHandle(const uint8* InData, uint32 NumBytes, const FSoundWaveProxyPtr&  InSoundWave, const FName& SoundWaveName, uint32 InChunkIndex, uint64 InCacheLookupID)
 	: CachedData(InData)
 	, CachedDataNumBytes(NumBytes)
-	, CorrespondingWave(MakeUnique<FSoundWaveProxy>(InSoundWave))
+	, CorrespondingWave(InSoundWave)
 	, CorrespondingWaveName(SoundWaveName)
 	, ChunkIndex(InChunkIndex)
 #if WITH_EDITOR
-	, ChunkGeneration(InSoundWave.GetCurrentChunkRevision())
+	, ChunkGeneration(InSoundWave.IsValid()? InSoundWave->GetCurrentChunkRevision() : 0)
 #endif
 {
 }
@@ -1474,7 +1474,7 @@ FAudioChunkHandle& FAudioChunkHandle::operator=(const FAudioChunkHandle& Other)
 
 	if (Other.CorrespondingWave.IsValid())
 	{
-		CorrespondingWave = MakeUnique<FSoundWaveProxy>(*Other.CorrespondingWave);
+		CorrespondingWave = Other.CorrespondingWave;
 	}
 
 	CorrespondingWaveName = Other.CorrespondingWaveName;
@@ -1518,10 +1518,12 @@ bool FAudioChunkHandle::IsValid() const
 #if WITH_EDITOR
 bool FAudioChunkHandle::IsStale() const
 {
-	if (CorrespondingWave.IsValid())
+	FSoundWaveProxyPtr ProxyPtr = CorrespondingWave.Pin();
+
+	if (ProxyPtr.IsValid())
 	{
 		// NOTE: While this is currently safe in editor, there's no guarantee the USoundWave will be kept alive during the lifecycle of this chunk handle.
-		return ChunkGeneration != CorrespondingWave->GetCurrentChunkRevision();
+		return ChunkGeneration != ProxyPtr->GetCurrentChunkRevision();
 	}
 	else
 	{
@@ -1530,7 +1532,12 @@ bool FAudioChunkHandle::IsStale() const
 }
 #endif
 
-FAudioChunkHandle IAudioStreamingManager::BuildChunkHandle(const uint8* InData, uint32 NumBytes, const FSoundWaveProxy&  InSoundWave, const FName& SoundWaveName, uint32 InChunkIndex, uint64 InCacheLookupID)
+FAudioChunkHandle IAudioStreamingManager::BuildChunkHandle(const uint8* InData, uint32 NumBytes, const FSoundWaveProxyPtr&  InSoundWave, const FName& SoundWaveName, uint32 InChunkIndex, uint64 InCacheLookupID)
 {
-	return FAudioChunkHandle(InData, NumBytes, InSoundWave, SoundWaveName, InChunkIndex, InCacheLookupID);
+	if (ensure(InSoundWave.IsValid()))
+	{
+		return FAudioChunkHandle(InData, NumBytes, InSoundWave, SoundWaveName, InChunkIndex, InCacheLookupID);
+	}
+
+	return {};
 }
diff --git a/Engine/Source/Runtime/Engine/Private/GameplayStatics.cpp b/Engine/Source/Runtime/Engine/Private/GameplayStatics.cpp
--- a/Engine/Source/Runtime/Engine/Private/GameplayStatics.cpp
+++ b/Engine/Source/Runtime/Engine/Private/GameplayStatics.cpp
@@ -1726,7 +1726,7 @@ void UGameplayStatics::PrimeSound(USoundBase* InSound)
 
 		if (InSoundWave->GetNumChunks() > 1)
 		{
-			IStreamingManager::Get().GetAudioStreamingManager().RequestChunk(InSoundWave->GetThisAsProxy(), 1, [](EAudioChunkLoadResult) {});
+			IStreamingManager::Get().GetAudioStreamingManager().RequestChunk(InSoundWave->CreateSoundWaveProxy(), 1, [](EAudioChunkLoadResult) {});
 		}
 	}
 }
diff --git a/Engine/Source/Runtime/Engine/Private/SoundWave.cpp b/Engine/Source/Runtime/Engine/Private/SoundWave.cpp
--- a/Engine/Source/Runtime/Engine/Private/SoundWave.cpp
+++ b/Engine/Source/Runtime/Engine/Private/SoundWave.cpp
@@ -186,6 +186,29 @@ void FStreamedAudioChunk::Serialize(FArchive& Ar, UObject* Owner, int32 ChunkInd
 #endif // #if WITH_EDITORONLY_DATA
 }
 
+void FStreamedAudioChunk::GetCopy(void** OutChunkData)
+{
+	if (!CachedDataPtr)
+	{
+		if (AudioDataSize != DataSize)
+		{
+			TArray<uint8> TempChunkBuffer;
+			TempChunkBuffer.AddUninitialized(DataSize);
+			void* DataDestPtr = TempChunkBuffer.GetData();
+			BulkData.GetCopy(&DataDestPtr, true);
+
+			FMemory::Memcpy(*OutChunkData, TempChunkBuffer.GetData(), AudioDataSize);
+		}
+		else
+		{
+			BulkData.GetCopy((void**)&CachedDataPtr, true);
+		}
+	}
+
+	// todo: expose as shared ptr(?)
+	FMemory::Memcmp(*OutChunkData, CachedDataPtr, DataSize);
+}
+
 #if WITH_EDITORONLY_DATA
 uint32 FStreamedAudioChunk::StoreInDerivedDataCache(const FString& InDerivedDataKey, const FStringView& SoundWaveName)
 {
@@ -248,7 +271,7 @@ USoundWave::USoundWave(const FObjectInitializer& ObjectInitializer)
 #if WITH_EDITOR
 	bWasStreamCachingEnabledOnLastCook = FPlatformCompressionUtilities::IsCurrentPlatformUsingStreamCaching();
 	bLoadedFromCookedData = false;
-	RunningPlatformData = nullptr;
+	RunningPlatformData.Reset();
 
 	OwnedBulkDataPtr = nullptr;
 	ResourceData = nullptr;
@@ -279,7 +302,7 @@ void USoundWave::GetResourceSizeEx(FResourceSizeEx& CumulativeResourceSize)
 	CumulativeResourceSize.AddDedicatedSystemMemoryBytes(CookedEnvelopeTimeData.Num() * sizeof(FSoundWaveEnvelopeTimeData));
 
 	// Add zeroth chunk data, if it's used (if this USoundWave isn't streaming, this won't report).
-	CumulativeResourceSize.AddDedicatedSystemMemoryBytes(ZerothChunkData.GetView().Num());
+	CumulativeResourceSize.AddDedicatedSystemMemoryBytes(ZerothChunkData->GetView().Num());
 
 	// Finally, report the actual audio memory being used, if this asset isn't using the stream cache.
 	if (FAudioDevice* LocalAudioDevice = GEngine->GetMainAudioDeviceRaw())
@@ -536,7 +559,12 @@ void USoundWave::Serialize( FArchive& Ar )
 				// Prime first chunk of audio.
 				CachedSoundWaveLoadingBehavior = ESoundWaveLoadingBehavior::PrimeOnLoad;
 				bLoadingBehaviorOverridden = true;
-				IStreamingManager::Get().GetAudioStreamingManager().RequestChunk(GetThisAsProxy(), 1, [](EAudioChunkLoadResult) {});
+
+				if (!InternalProxy.IsValid())
+				{
+					InternalProxy = CreateSoundWaveProxy();
+				}
+				IStreamingManager::Get().GetAudioStreamingManager().RequestChunk(InternalProxy, 1, [](EAudioChunkLoadResult) {});
 			}
 		}
 	}
@@ -970,7 +998,7 @@ FByteBulkData* USoundWave::GetCompressedData(FName Format, const FPlatformAudioC
 void USoundWave::InvalidateCompressedData(bool bFreeResources, bool bRebuildStreamingChunks)
 {
 	CompressedDataGuid = FGuid::NewGuid();
-	ZerothChunkData.Empty();
+	ZerothChunkData->Empty();
 	CompressedFormatData.FlushData();
 
 	if (bFreeResources)
@@ -997,7 +1025,7 @@ void USoundWave::InvalidateCompressedData(bool bFreeResources, bool bRebuildStre
 
 bool USoundWave::HasStreamingChunks()
 {
-	return RunningPlatformData != NULL && RunningPlatformData->Chunks.Num() > 0;
+	return RunningPlatformData.IsValid() && RunningPlatformData->Chunks.Num() > 0;
 }
 
 void USoundWave::PostLoad()
@@ -1030,6 +1058,12 @@ void USoundWave::PostLoad()
 	// If our loading behavior is defined by a sound class, we need to update whether this sound wave actually needs to retain its audio data or not.
 	ESoundWaveLoadingBehavior ActualLoadingBehavior = GetLoadingBehavior();
 
+	if (!InternalProxy.IsValid())
+	{
+		InternalProxy = CreateSoundWaveProxy();
+	}
+	
+
 	if (ShouldUseStreamCaching() && ActualLoadingBehavior != GetLoadingBehavior(false))
 	{
 		if (!GIsEditor && !DisableRetainingCVar && ActualLoadingBehavior == ESoundWaveLoadingBehavior::RetainOnLoad)
@@ -1048,7 +1082,7 @@ void USoundWave::PostLoad()
 			
 			if (bShouldPrime && bHasMultipleChunks)
 			{
-				IStreamingManager::Get().GetAudioStreamingManager().RequestChunk(GetThisAsProxy(), 1, [](EAudioChunkLoadResult) {});
+				IStreamingManager::Get().GetAudioStreamingManager().RequestChunk(InternalProxy, 1, [](EAudioChunkLoadResult) {});
 			}
 		}
 
@@ -1111,7 +1145,7 @@ void USoundWave::PostLoad()
 #endif // #if WITH_EDITORONLY_DATA
 		if (!ShouldUseStreamCaching())
 		{
-			IStreamingManager::Get().GetAudioStreamingManager().AddStreamingSoundWave(GetThisAsProxy());
+			IStreamingManager::Get().GetAudioStreamingManager().AddStreamingSoundWave(InternalProxy);
 		}
 	}
 
@@ -1136,22 +1170,25 @@ void USoundWave::PostLoad()
 	INC_FLOAT_STAT_BY(STAT_AudioBufferTime, Duration);
 	INC_FLOAT_STAT_BY(STAT_AudioBufferTimeChannels, NumChannels * Duration);
 
-
-	// overwrite cached proxy now that post-load is complete.
+	// cache current state as a proxy
 	InternalProxy = CreateSoundWaveProxy();
-	InternalProxy->ReleaseCompressedAudio(); // release handles if there are any
+	if (InternalProxy.IsValid())
+	{
+		// release dupe handle already held by 'this'
+		InternalProxy->ReleaseCompressedAudio();
+	}
 }
 
 void USoundWave::EnsureZerothChunkIsLoaded()
 {
 	// If the zeroth chunk is already loaded, early exit.
-	if (ZerothChunkData.GetView().Num() > 0 || !ShouldUseStreamCaching())
+	if (ZerothChunkData->GetView().Num() > 0 || !ShouldUseStreamCaching())
 	{
 		return;
 	}
 
 #if WITH_EDITOR 
-	if (RunningPlatformData == nullptr)
+	if (RunningPlatformData.IsValid())
 	{
 		CachePlatformData(false);
 	}
@@ -1165,22 +1202,22 @@ void USoundWave::EnsureZerothChunkIsLoaded()
 		return;
 	}
 
-	ZerothChunkData.Reset(TempChunkBuffer, ChunkSizeInBytes);
+	ZerothChunkData->Reset(TempChunkBuffer, ChunkSizeInBytes);
 
 #else // WITH_EDITOR
 	// Otherwise, the zeroth chunk is cooked out to RunningPlatformData, and we just need to retrieve it.
-	check(RunningPlatformData && RunningPlatformData->Chunks.Num() > 0);
+	check(RunningPlatformData.IsValid() && RunningPlatformData->Chunks.Num() > 0);
 	FStreamedAudioChunk& ZerothChunk = RunningPlatformData->Chunks[0];
 	// Some sanity checks to ensure that the bulk size set up
 	UE_CLOG(ZerothChunk.BulkData.GetBulkDataSize() != ZerothChunk.DataSize, LogAudio, Warning, TEXT("Bulk data serialized out had a mismatched size with the DataSize field. Soundwave: %s Bulk Data Reported Size: %d Bulk Data Actual Size: %ld"), *GetFullName(), ZerothChunk.DataSize, ZerothChunk.BulkData.GetBulkDataSize());
 
-	ZerothChunkData = ZerothChunk.BulkData.GetCopyAsBuffer(ZerothChunk.AudioDataSize, true);
+	*ZerothChunkData = ZerothChunk.BulkData.GetCopyAsBuffer(ZerothChunk.AudioDataSize, true);
 #endif // WITH_EDITOR
 }
 
 uint32 USoundWave::GetNumChunks() const
 {
-	if (RunningPlatformData)
+	if (RunningPlatformData.IsValid())
 	{
 		return RunningPlatformData->Chunks.Num();
 	}
@@ -1204,7 +1241,7 @@ uint32 USoundWave::GetSizeOfChunk(uint32 ChunkIndex)
 {
 	check(ChunkIndex < GetNumChunks());
 
-	if (RunningPlatformData)
+	if (RunningPlatformData.IsValid())
 	{
 		return RunningPlatformData->Chunks[ChunkIndex].AudioDataSize;
 	}
@@ -1915,7 +1952,11 @@ bool USoundWave::IsReadyForFinishDestroy()
 		}
 	}
 
-	const bool bIsStreamingInProgress = IStreamingManager::Get().GetAudioStreamingManager().IsStreamingInProgress(GetThisAsProxy());
+	if (!InternalProxy.IsValid())
+	{
+		InternalProxy = CreateSoundWaveProxy();
+	}
+	const bool bIsStreamingInProgress = IStreamingManager::Get().GetAudioStreamingManager().IsStreamingInProgress(InternalProxy);
 
 	check(GetPrecacheState() != ESoundWavePrecacheState::InProgress);
 
@@ -1967,7 +2008,11 @@ void USoundWave::FinishDestroy()
 
 	if (FApp::CanEverRenderAudio())
 	{
-		IStreamingManager::Get().GetAudioStreamingManager().RemoveStreamingSoundWave(GetThisAsProxy());
+		if (!InternalProxy.IsValid())
+		{
+			InternalProxy = CreateSoundWaveProxy();
+		}
+		IStreamingManager::Get().GetAudioStreamingManager().RemoveStreamingSoundWave(InternalProxy);
 	}
 }
 
@@ -2347,19 +2392,27 @@ TArrayView<const uint8> USoundWave::GetZerothChunk(bool bForImmediatePlayback)
 			EnsureZerothChunkIsLoaded();
 		}
 
-		check(ZerothChunkData.GetView().Num() > 0);
+		check(ZerothChunkData->GetView().Num() > 0);
 
 		if (GetNumChunks() > 1)
 		{
 			// Prime first chunk for playback.
-			IStreamingManager::Get().GetAudioStreamingManager().RequestChunk(GetThisAsProxy(), 1, [](EAudioChunkLoadResult InResult) {}, ENamedThreads::AnyThread, bForImmediatePlayback);
+			if (!InternalProxy.IsValid())
+			{
+				InternalProxy = CreateSoundWaveProxy();
+			}
+			IStreamingManager::Get().GetAudioStreamingManager().RequestChunk(InternalProxy, 1, [](EAudioChunkLoadResult InResult) {}, ENamedThreads::AnyThread, bForImmediatePlayback);
 		}
 
-		return ZerothChunkData.GetView();
+		return ZerothChunkData->GetView();
 	}
 	else
 	{
-		FAudioChunkHandle ChunkHandle = IStreamingManager::Get().GetAudioStreamingManager().GetLoadedChunk(GetThisAsProxy(), 0);
+		if (!InternalProxy.IsValid())
+		{
+			InternalProxy = CreateSoundWaveProxy();
+		}
+		FAudioChunkHandle ChunkHandle = IStreamingManager::Get().GetAudioStreamingManager().GetLoadedChunk(InternalProxy, 0);
 		return TArrayView<const uint8>(ChunkHandle.GetData(), ChunkHandle.Num());
 	}
 }
@@ -2389,14 +2442,14 @@ bool USoundWave::HasCookedAmplitudeEnvelopeData() const
 	return CookedEnvelopeTimeData.Num() > 0;
 }
 
-TUniquePtr<FSoundWaveProxy> USoundWave::CreateSoundWaveProxy()
+FSoundWaveProxyPtr USoundWave::CreateSoundWaveProxy()
 {
-	return MakeUnique<FSoundWaveProxy>(this);
+	return MakeShared<FSoundWaveProxy, ESPMode::ThreadSafe>(this);
 }
 
 TUniquePtr<Audio::IProxyData> USoundWave::CreateNewProxyData(const Audio::FProxyDataInitParams& InitParams)
 {
-	return CreateSoundWaveProxy();
+	return MakeUnique<FSoundWaveProxy>(this);
 }
 
 void USoundWave::AddPlayingSource(const FSoundWaveClientPtr& Source)
@@ -2423,10 +2476,15 @@ void USoundWave::RemovePlayingSource(const FSoundWaveClientPtr& Source)
 
 void USoundWave::UpdatePlatformData()
 {
+	if (!InternalProxy.IsValid())
+	{
+		InternalProxy = CreateSoundWaveProxy();
+	}
+
 	if (IsStreaming(nullptr))
 	{
 		// Make sure there are no pending requests in flight.
-		while (IStreamingManager::Get().GetAudioStreamingManager().IsStreamingInProgress(GetThisAsProxy()))
+		while (IStreamingManager::Get().GetAudioStreamingManager().IsStreamingInProgress(InternalProxy))
 		{
 			// Give up timeslice.
 			FPlatformProcess::Sleep(0);
@@ -2434,18 +2492,18 @@ void USoundWave::UpdatePlatformData()
 
 #if WITH_EDITORONLY_DATA
 		// Temporarily remove from streaming manager to release currently used data chunks
-		IStreamingManager::Get().GetAudioStreamingManager().RemoveStreamingSoundWave(GetThisAsProxy());
+		IStreamingManager::Get().GetAudioStreamingManager().RemoveStreamingSoundWave(InternalProxy);
 
 		// Recache platform data if the source has changed.
 		CachePlatformData();
 
 		// Add back to the streaming manager to reload first chunk
-		IStreamingManager::Get().GetAudioStreamingManager().AddStreamingSoundWave(GetThisAsProxy());
+		IStreamingManager::Get().GetAudioStreamingManager().AddStreamingSoundWave(InternalProxy);
 #endif
 	}
 	else
 	{
-		IStreamingManager::Get().GetAudioStreamingManager().RemoveStreamingSoundWave(GetThisAsProxy());
+		IStreamingManager::Get().GetAudioStreamingManager().RemoveStreamingSoundWave(InternalProxy);
 	}
 }
 
@@ -2528,17 +2586,6 @@ bool USoundWave::GetChunkData(int32 ChunkIndex, uint8** OutChunkData, bool bMake
 	return true;
 }
 
-const FSoundWaveProxy& USoundWave::GetThisAsProxy()
-{
-	if (!InternalProxy.IsValid())
-	{
-		// PostLoad hasn't been called. The InternalProxy created will be overritten by PostLoad().
-		InternalProxy = CreateSoundWaveProxy();
-	}
-
-	check(InternalProxy.IsValid());
-	return *InternalProxy;
-}
 
 uint32 USoundWave::GetInterpolatedCookedFFTDataForTimeInternal(float InTime, uint32 StartingIndex, TArray<FSoundWaveSpectralData>& OutData, bool bLoop)
 {
@@ -2753,6 +2800,10 @@ void USoundWave::GetHandleForChunkOfAudio(TFunction<void(FAudioChunkHandle&&)> O
 {
 	ENamedThreads::Type ThreadToDispatchCallbackOn = (DispatchToGameThreadOnChunkRequestCVar != 0) ? ENamedThreads::GameThread : ENamedThreads::AnyThread;
 
+	if (!InternalProxy.IsValid())
+	{
+		InternalProxy = CreateSoundWaveProxy();
+	}
 
 	// if we are requesting a chunk that is out of bounds,
 	// early exit.
@@ -2764,7 +2815,7 @@ void USoundWave::GetHandleForChunkOfAudio(TFunction<void(FAudioChunkHandle&&)> O
 	else if (bForceSync)
 	{
 		// For sync cases, we call GetLoadedChunk with bBlockForLoad = true, then execute the callback immediately.
-		FAudioChunkHandle ChunkHandle = IStreamingManager::Get().GetAudioStreamingManager().GetLoadedChunk(GetThisAsProxy(), ChunkIndex, true);
+		FAudioChunkHandle ChunkHandle = IStreamingManager::Get().GetAudioStreamingManager().GetLoadedChunk(InternalProxy, ChunkIndex, true);
 		OnLoadCompleted(MoveTemp(ChunkHandle));
 	}
 	else
@@ -2772,7 +2823,7 @@ void USoundWave::GetHandleForChunkOfAudio(TFunction<void(FAudioChunkHandle&&)> O
 		TWeakObjectPtr<USoundWave> WeakThis = MakeWeakObjectPtr(this);
 
 		// For async cases, we call RequestChunk and request the loaded chunk in the completion callback.
-		IStreamingManager::Get().GetAudioStreamingManager().RequestChunk(GetThisAsProxy(), ChunkIndex, [ThreadToDispatchCallbackOn, WeakThis, OnLoadCompleted, ChunkIndex, CallbackThread](EAudioChunkLoadResult LoadResult)
+		IStreamingManager::Get().GetAudioStreamingManager().RequestChunk(InternalProxy, ChunkIndex, [ThreadToDispatchCallbackOn, WeakThis, OnLoadCompleted, ChunkIndex, CallbackThread](EAudioChunkLoadResult LoadResult)
 		{
 			auto DispatchOnLoadCompletedCallback = [ThreadToDispatchCallbackOn, OnLoadCompleted, CallbackThread](FAudioChunkHandle&& InHandle)
 			{
@@ -2794,7 +2845,7 @@ void USoundWave::GetHandleForChunkOfAudio(TFunction<void(FAudioChunkHandle&&)> O
 			if (WeakThis.IsValid() && (LoadResult == EAudioChunkLoadResult::Completed || LoadResult == EAudioChunkLoadResult::AlreadyLoaded))
 			{
 				USoundWave* ThisSoundWave = WeakThis.Get();
-				FAudioChunkHandle ChunkHandle = IStreamingManager::Get().GetAudioStreamingManager().GetLoadedChunk(ThisSoundWave->GetThisAsProxy(), ChunkIndex, (BlockOnChunkLoadCompletionCVar != 0));
+				FAudioChunkHandle ChunkHandle = IStreamingManager::Get().GetAudioStreamingManager().GetLoadedChunk(ThisSoundWave->InternalProxy, ChunkIndex, (BlockOnChunkLoadCompletionCVar != 0));
 
 				// If we hit this, something went wrong in GetLoadedChunk.
 				if (!ChunkHandle.IsValid())
@@ -2830,7 +2881,12 @@ void USoundWave::RetainCompressedAudio(bool bForceSync /*= false*/)
 	}
 	else if (bForceSync)
 	{
-		FirstChunk = IStreamingManager::Get().GetAudioStreamingManager().GetLoadedChunk(GetThisAsProxy(), 1, true);
+		if (!InternalProxy.IsValid())
+		{
+			InternalProxy = CreateSoundWaveProxy();
+		}
+
+		FirstChunk = IStreamingManager::Get().GetAudioStreamingManager().GetLoadedChunk(InternalProxy, 1, true);
 		UE_CLOG(!FirstChunk.IsValid(), LogAudio, Display, TEXT("First chunk was invalid after synchronous load in RetainCompressedAudio(). This was likely because the cache was blown. Sound: %s"), *GetFullName());
 	}
 	else
@@ -2903,7 +2959,12 @@ void USoundWave::OverrideLoadingBehavior(ESoundWaveLoadingBehavior InLoadingBeha
 		}
 		else if (InLoadingBehavior == ESoundWaveLoadingBehavior::PrimeOnLoad && GetNumChunks() > 1)
 		{
-			IStreamingManager::Get().GetAudioStreamingManager().RequestChunk(GetThisAsProxy(), 1, [](EAudioChunkLoadResult) {});
+			if (!InternalProxy.IsValid())
+			{
+				InternalProxy = CreateSoundWaveProxy();
+			}
+
+			IStreamingManager::Get().GetAudioStreamingManager().RequestChunk(InternalProxy, 1, [](EAudioChunkLoadResult) {});
 		}
 	}
 }
@@ -3018,38 +3079,12 @@ FSoundWaveProxy::FSoundWaveProxy(USoundWave* InWave)
 		}
 	}
 
-	if (InWave->RunningPlatformData)
+	if (InWave->RunningPlatformData.IsValid())
 	{
 		check(LoadingBehavior != ESoundWaveLoadingBehavior::ForceInline);
-		RunningPlatformData = MakeUnique<FStreamedAudioPlatformData>(*InWave->RunningPlatformData);
+		RunningPlatformData = InWave->RunningPlatformData;
 		NumChunks = InWave->GetNumChunks();
-		EnsureZerothChunkIsLoaded();
 	}
-
-//	UE_LOG(LogAudio, Display, TEXT("MAXDEBUG: FSoundWaveProxy(USoundWave* InWave)"));
-}
-
-FSoundWaveProxy::FSoundWaveProxy(const FSoundWaveProxy& Other)
-	: ZerothChunkData(Other.ZerothChunkData)
-	, RunningPlatformData(Other.RunningPlatformData? MakeUnique<FStreamedAudioPlatformData>(*Other.RunningPlatformData) : nullptr)
-	, NameCached(Other.NameCached)
-	, FullNameCached(Other.FullNameCached)
-	, PathNameCached(Other.PathNameCached)
-	, SoundWaveKeyCached(Other.SoundWaveKeyCached)
-	, RuntimeFormat(Other.RuntimeFormat)
-	, SampleRate(Other.SampleRate)
-	, NumChannels(Other.NumChannels)
-	, NumChunks(Other.NumChunks)
-	, FirstChunk(Other.FirstChunk)
-	, LoadingBehavior(Other.LoadingBehavior)
-	, bLoadingBehaviorOverridden(Other.bLoadingBehaviorOverridden)
-	, bIsStreaming(Other.bIsStreaming)
-	, bSeekableStreaming(Other.bSeekableStreaming)
-	, bShouldUseStreamCaching(Other.bShouldUseStreamCaching)
-#if WITH_EDITOR
-	, CurrentChunkRevision(Other.CurrentChunkRevision)
-#endif // #if WITH_EDITOR
-{
 }
 
 void FSoundWaveProxy::ReleaseCompressedAudio()
@@ -3063,25 +3098,30 @@ uint32 FSoundWaveProxy::GetSizeOfChunk(uint32 ChunkIndex) const
 	return GetChunk(ChunkIndex).AudioDataSize;
 }
 
-TArrayView<const uint8> FSoundWaveProxy::GetZerothChunk(bool bForImmediatePlayback) const
+TArrayView<const uint8> FSoundWaveProxy::GetZerothChunk(const FSoundWaveProxyPtr& SoundWaveProxy, bool bForImmediatePlayback)
 {
-	if (bShouldUseStreamCaching)
+	if (ensure(SoundWaveProxy.IsValid()))
 	{
-		check(IsZerothChunkDataLoaded());
+		if (SoundWaveProxy->bShouldUseStreamCaching)
+		{
+			check(SoundWaveProxy->IsZerothChunkDataLoaded());
 
-		if (GetNumChunks() > 1)
+			if (SoundWaveProxy->GetNumChunks() > 1)
+			{
+				// Prime first chunk for playback
+				IStreamingManager::Get().GetAudioStreamingManager().RequestChunk(SoundWaveProxy, 1, [](EAudioChunkLoadResult InResult) {}, ENamedThreads::AnyThread, bForImmediatePlayback);
+			}
+
+			return SoundWaveProxy->GetZerothChunkDataView();
+		}
+		else
 		{
-			// Prime first chunk for playback
-			IStreamingManager::Get().GetAudioStreamingManager().RequestChunk(*this, 1, [](EAudioChunkLoadResult InResult) {}, ENamedThreads::AnyThread, bForImmediatePlayback);
+			FAudioChunkHandle ChunkHandle = IStreamingManager::Get().GetAudioStreamingManager().GetLoadedChunk(SoundWaveProxy, 0);
+			return TArrayView<const uint8>(ChunkHandle.GetData(), ChunkHandle.Num());
 		}
-
-		return GetZerothChunkDataView();
-	}
-	else
-	{
-		FAudioChunkHandle ChunkHandle = IStreamingManager::Get().GetAudioStreamingManager().GetLoadedChunk(*this, 0);
-		return TArrayView<const uint8>(ChunkHandle.GetData(), ChunkHandle.Num());
 	}
+
+	return {};
 }
 
 bool FSoundWaveProxy::GetChunkData(int32 ChunkIndex, uint8** OutChunkData, bool bMakeSureChunkIsLoaded)
@@ -3114,16 +3154,21 @@ bool FSoundWaveProxy::GetChunkData(int32 ChunkIndex, uint8** OutChunkData, bool
 
 bool FSoundWaveProxy::IsZerothChunkDataLoaded() const
 {
-	return ZerothChunkData.GetView().Num() > 0;
+	return ZerothChunkData->GetView().Num() > 0;
 }
 
 const TArrayView<uint8> FSoundWaveProxy::GetZerothChunkDataView() const
 {
-	return ZerothChunkData.GetView();
+	return ZerothChunkData->GetView();
 }
 
 void FSoundWaveProxy::EnsureZerothChunkIsLoaded()
 {
+	// If the zeroth chunk is already loaded, early exit.
+	if (ZerothChunkData->GetView().Num() > 0 || !bShouldUseStreamCaching)
+	{
+		return;
+	}
 
 #ifdef WITH_EDITOR
 	// If we're running the editor, we'll need to retrieve the chunked audio from the DDC:
@@ -3135,7 +3180,7 @@ void FSoundWaveProxy::EnsureZerothChunkIsLoaded()
 		return;
 	}
 
-	ZerothChunkData.Reset(TempChunkBuffer, ChunkSizeInBytes);
+	ZerothChunkData->Reset(TempChunkBuffer, ChunkSizeInBytes);
 
 #else // WITH_EDITOR
 	// Otherwise, the zeroth chunk is cooked out to SharedRuntimeData::RunningPlatformData, and we just need to retrieve it.
@@ -3151,10 +3196,10 @@ void FSoundWaveProxy::EnsureZerothChunkIsLoaded()
 uint32 FSoundWaveProxy::GetNumChunks() const
 {
 	// This function shouldn't be called on audio marked "ForceInline."
-	ensureMsgf(RunningPlatformData, TEXT("Calling GetNumChunks on a FSoundWaveProxy without RunnigPlatformData is not allowed! SoundWave: %s - %s")
+	ensureMsgf(RunningPlatformData.IsValid(), TEXT("Calling GetNumChunks on a FSoundWaveProxy without RunnigPlatformData is not allowed! SoundWave: %s - %s")
 		, *GetFName().ToString(), EnumToString(LoadingBehavior));
 
-	if (RunningPlatformData)
+	if (RunningPlatformData.IsValid())
 	{
 		return RunningPlatformData->Chunks.Num();
 	}
@@ -3175,31 +3220,31 @@ int32 FSoundWaveProxy::GetCurrentChunkRevision() const
 const FStreamedAudioChunk& FSoundWaveProxy::GetChunk(uint32 ChunkIndex) const
 {
 	// This function shouldn't be called on audio marked "ForceInline."
-	ensureMsgf(RunningPlatformData, TEXT("Calling GetNumChunks on a FSoundWaveProxy without RunnigPlatformData is not allowed! SoundWave: %s - %s")
+	ensureMsgf(RunningPlatformData.IsValid(), TEXT("Calling GetNumChunks on a FSoundWaveProxy without RunnigPlatformData is not allowed! SoundWave: %s - %s")
 		, *GetFName().ToString(), EnumToString(LoadingBehavior));
 
 	check((ChunkIndex < (uint32)GetNumChunks()));
-	check(RunningPlatformData);
+	check(RunningPlatformData.IsValid());
 	return RunningPlatformData->Chunks[ChunkIndex];
 }
 
 int32 FSoundWaveProxy::GetChunkFromDDC(int32 ChunkIndex, uint8** OutChunkData, bool bMakeSureChunkIsLoaded)
 {
 	// This function shouldn't be called on audio marked "ForceInline."
-	ensureMsgf(RunningPlatformData, TEXT("Calling GetNumChunks on a FSoundWaveProxy without RunnigPlatformData is not allowed! SoundWave: %s - %s")
+	ensureMsgf(RunningPlatformData.IsValid(), TEXT("Calling GetNumChunks on a FSoundWaveProxy without RunnigPlatformData is not allowed! SoundWave: %s - %s")
 		, *GetFName().ToString(), EnumToString(LoadingBehavior));
 
-	check(RunningPlatformData);
+	check(RunningPlatformData.IsValid());
 	return RunningPlatformData->GetChunkFromDDC(ChunkIndex, OutChunkData, bMakeSureChunkIsLoaded);
 }
 
 #if WITH_EDITORONLY_DATA
 FString FSoundWaveProxy::GetDerivedDataKey() const
 {
-	check(RunningPlatformData);
+	check(RunningPlatformData.IsValid());
 
 	// This function shouldn't be called on audio marked "ForceInline."
-	ensureMsgf(RunningPlatformData, TEXT("Calling GetNumChunks on a FSoundWaveProxy without RunnigPlatformData is not allowed! SoundWave: %s - %s")
+	ensureMsgf(RunningPlatformData.IsValid(), TEXT("Calling GetNumChunks on a FSoundWaveProxy without RunnigPlatformData is not allowed! SoundWave: %s - %s")
 		, *GetFName().ToString(), EnumToString(LoadingBehavior));
 
 	return RunningPlatformData->DerivedDataKey;
diff --git a/Engine/Source/Runtime/Engine/Private/VorbisAudioInfo.cpp b/Engine/Source/Runtime/Engine/Private/VorbisAudioInfo.cpp
--- a/Engine/Source/Runtime/Engine/Private/VorbisAudioInfo.cpp
+++ b/Engine/Source/Runtime/Engine/Private/VorbisAudioInfo.cpp
@@ -206,7 +206,7 @@ size_t FVorbisAudioInfo::ReadStreaming(void *Ptr, uint32 Size )
 			}
 			else
 			{
-				CurrentStreamingChunkData = GetLoadedChunk(*StreamingSoundWave, CurrentStreamingChunkIndex, CurrentStreamingChunksSize);
+				CurrentStreamingChunkData = GetLoadedChunk(StreamingSoundWave, CurrentStreamingChunkIndex, CurrentStreamingChunksSize);
 			}
 
 			if (CurrentStreamingChunkData)
@@ -541,65 +541,75 @@ void FVorbisAudioInfo::EnableHalfRate( bool HalfRate )
 	ov_halfrate(&VFWrapper->vf, int32(HalfRate));
 }
 
-bool FVorbisAudioInfo::StreamCompressedInfoInternal(const FSoundWaveProxy& InWaveProxy, struct FSoundQualityInfo* QualityInfo)
+bool FVorbisAudioInfo::StreamCompressedInfoInternal(const FSoundWaveProxyPtr& InWaveProxy, struct FSoundQualityInfo* QualityInfo)
 {
-	if (!bDllLoaded)
+	if (ensure(InWaveProxy.IsValid()))
 	{
-		UE_LOG(LogAudio, Error, TEXT("FVorbisAudioInfo::StreamCompressedInfoInternal failed to parse header due to vorbis DLL not being loaded for sound '%s'."), *InWaveProxy.GetFName().ToString());
-		return false;
-	}
+		if (!bDllLoaded)
+		{
+			UE_LOG(LogAudio, Error, TEXT("FVorbisAudioInfo::StreamCompressedInfoInternal failed to parse header due to vorbis DLL not being loaded for sound '%s'."), *InWaveProxy->GetFName().ToString());
+			return false;
+		}
 
-	SCOPE_CYCLE_COUNTER( STAT_VorbisPrepareDecompressionTime );
+		SCOPE_CYCLE_COUNTER( STAT_VorbisPrepareDecompressionTime );
 
-	FScopeLock ScopeLock(&VorbisCriticalSection);
+		FScopeLock ScopeLock(&VorbisCriticalSection);
 
-	if (!VFWrapper)
-	{
-		UE_LOG(LogAudio, Error, TEXT("FVorbisAudioInfo::StreamCompressedInfoInternal failed due to no vorbis wrapper for sound '%s'."), *InWaveProxy.GetFName().ToString());
-		return false;
-	}
+		if (!VFWrapper)
+		{
+			UE_LOG(LogAudio, Error, TEXT("FVorbisAudioInfo::StreamCompressedInfoInternal failed due to no vorbis wrapper for sound '%s'."), *InWaveProxy->GetFName().ToString());
+			return false;
+		}
 
-	ov_callbacks Callbacks;
+		ov_callbacks Callbacks;
 
-	SrcBufferData = NULL;
-	SrcBufferDataSize = 0;
-	BufferOffset = 0;
+		SrcBufferData = NULL;
+		SrcBufferDataSize = 0;
+		BufferOffset = 0;
 
-	Callbacks.read_func = OggReadStreaming;
-	Callbacks.close_func = OggCloseStreaming;
-	Callbacks.seek_func = NULL;	// Force streaming
-	Callbacks.tell_func = NULL;	// Force streaming
+		Callbacks.read_func = OggReadStreaming;
+		Callbacks.close_func = OggCloseStreaming;
+		Callbacks.seek_func = NULL;	// Force streaming
+		Callbacks.tell_func = NULL;	// Force streaming
 
-	bHeaderParsed = GetCompressedInfoCommon(&Callbacks, QualityInfo);
-	if (!bHeaderParsed)
-	{
-		UE_LOG(LogAudio, Error, TEXT("FVorbisAudioInfo::StreamCompressedInfoInternal failed to parse header for '%s'."), *InWaveProxy.GetFName().ToString());
-	}
+		bHeaderParsed = GetCompressedInfoCommon(&Callbacks, QualityInfo);
+		if (!bHeaderParsed)
+		{
+			UE_LOG(LogAudio, Error, TEXT("FVorbisAudioInfo::StreamCompressedInfoInternal failed to parse header for '%s'."), *InWaveProxy->GetFName().ToString());
+		}
 	
 
-	return bHeaderParsed;
+		return bHeaderParsed;
+	}
+
+	return false;
 }
 
-// TODO: Make InSoundWave a const ref after fixing up FSoundWaveProxy::GetZerothChunk() to be const
-const uint8* FVorbisAudioInfo::GetLoadedChunk(FSoundWaveProxy& InSoundWave, uint32 ChunkIndex, uint32& OutChunkSize)
+
+const uint8* FVorbisAudioInfo::GetLoadedChunk(FSoundWaveProxyPtr InSoundWave, uint32 ChunkIndex, uint32& OutChunkSize)
 {
-	if (ChunkIndex >= InSoundWave.GetNumChunks())
-	{
-		OutChunkSize = 0;
-		return nullptr;
-	}
-	else if (ChunkIndex == 0)
+	if (ensure(InSoundWave.IsValid()))
 	{
-		TArrayView<const uint8> ZerothChunk = InSoundWave.GetZerothChunk(true);
-		OutChunkSize = ZerothChunk.Num();
-		return ZerothChunk.GetData();
-	}
-	else
-	{
-		CurCompressedChunkHandle = IStreamingManager::Get().GetAudioStreamingManager().GetLoadedChunk(InSoundWave, ChunkIndex, false, true);
-		OutChunkSize = CurCompressedChunkHandle.Num();
-		return CurCompressedChunkHandle.GetData();
+		if (ChunkIndex >= InSoundWave->GetNumChunks())
+		{
+			OutChunkSize = 0;
+			return nullptr;
+		}
+		else if (ChunkIndex == 0)
+		{
+			TArrayView<const uint8> ZerothChunk = FSoundWaveProxy::GetZerothChunk(InSoundWave, true);
+			OutChunkSize = ZerothChunk.Num();
+			return ZerothChunk.GetData();
+		}
+		else
+		{
+			CurCompressedChunkHandle = IStreamingManager::Get().GetAudioStreamingManager().GetLoadedChunk(InSoundWave, ChunkIndex, false, true);
+			OutChunkSize = CurCompressedChunkHandle.Num();
+			return CurCompressedChunkHandle.GetData();
+		}
 	}
+
+	return nullptr;
 }
 
 bool FVorbisAudioInfo::StreamCompressedData(uint8* InDestination, bool bLooping, uint32 BufferSize)
diff --git a/Engine/Source/Runtime/Engine/Public/ADPCMAudioInfo.h b/Engine/Source/Runtime/Engine/Public/ADPCMAudioInfo.h
--- a/Engine/Source/Runtime/Engine/Public/ADPCMAudioInfo.h
+++ b/Engine/Source/Runtime/Engine/Public/ADPCMAudioInfo.h
@@ -106,7 +106,7 @@ public:
 
 	// Additional overrides for streaming
 	virtual bool SupportsStreaming() const override {return true;}
-	virtual bool StreamCompressedInfoInternal(const FSoundWaveProxy& InWaveProxy, struct FSoundQualityInfo* QualityInfo) override;
+	virtual bool StreamCompressedInfoInternal(const FSoundWaveProxyPtr& InWaveProxy, struct FSoundQualityInfo* QualityInfo) override;
 	virtual bool StreamCompressedData(uint8* Destination, bool bLooping, uint32 BufferSize) override;
 	virtual int32 GetCurrentChunkIndex() const override
 	{
@@ -138,7 +138,7 @@ private:
 	// Wrapper function that returns a pointer to the currently used compressed data.
 	// If a non-zero chunk is requested, this function also aquires a reference to that chunk
 	// until we move on to a different chunk.
-	const uint8* GetLoadedChunk(const FSoundWaveProxy& InSoundWave, uint32 ChunkIndex, uint32& OutChunkSize);
+	const uint8* GetLoadedChunk(const FSoundWaveProxyPtr& InSoundWave, uint32 ChunkIndex, uint32& OutChunkSize);
 
 	int32 NumConsecutiveReadFailiures;
 
diff --git a/Engine/Source/Runtime/Engine/Public/AudioDecompress.h b/Engine/Source/Runtime/Engine/Public/AudioDecompress.h
--- a/Engine/Source/Runtime/Engine/Public/AudioDecompress.h
+++ b/Engine/Source/Runtime/Engine/Public/AudioDecompress.h
@@ -118,7 +118,7 @@ public:
 	* @param	QualityInfo		Quality Info (to be filled out)
 	*/
 	ENGINE_API bool StreamCompressedInfo(USoundWave* Wave, struct FSoundQualityInfo* QualityInfo);
-	ENGINE_API bool StreamCompressedInfo(const FSoundWaveProxy& Wave, struct FSoundQualityInfo* QualityInfo);
+	ENGINE_API bool StreamCompressedInfo(const FSoundWaveProxyPtr& Wave, struct FSoundQualityInfo* QualityInfo);
 
 	/**
 	*  Returns true if a non-recoverable error has occurred.
@@ -127,7 +127,7 @@ public:
 
 protected:
 	/** Internal override implemented by subclasses. */
-	virtual bool StreamCompressedInfoInternal(const FSoundWaveProxy& InWaveProxy, struct FSoundQualityInfo* QualityInfo) = 0;
+	virtual bool StreamCompressedInfoInternal(const FSoundWaveProxyPtr& InWaveProxy, struct FSoundQualityInfo* QualityInfo) = 0;
 
 public:
 
@@ -153,10 +153,10 @@ public:
 	virtual int32 GetCurrentChunkOffset() const {return -1;}
 
 	/** Return the streaming sound wave used by this decoder. Returns nullptr if there is not a streaming sound wave. */
-	virtual const FSoundWaveProxy& GetStreamingSoundWave() { return *StreamingSoundWave; }
+	virtual const FSoundWaveProxyPtr& GetStreamingSoundWave() { return StreamingSoundWave; }
 
 protected:
-	TUniquePtr<FSoundWaveProxy> StreamingSoundWave;
+	FSoundWaveProxyPtr StreamingSoundWave;
 };
 
 /** Struct used to store the results of a decode operation. **/
@@ -197,7 +197,7 @@ public:
 	virtual bool UsesVorbisChannelOrdering() const override { return false; }
 	virtual int GetStreamBufferSize() const override { return  MONO_PCM_BUFFER_SIZE; }
 	virtual bool SupportsStreaming() const override { return true; }
-	virtual bool StreamCompressedInfoInternal(const FSoundWaveProxy& InWaveProxy, FSoundQualityInfo* QualityInfo) override;
+	virtual bool StreamCompressedInfoInternal(const FSoundWaveProxyPtr& InWaveProxy, FSoundQualityInfo* QualityInfo) override;
 	virtual bool StreamCompressedData(uint8* Destination, bool bLooping, uint32 BufferSize) override;
 	virtual int32 GetCurrentChunkIndex() const override { return CurrentChunkIndex; }
 	virtual int32 GetCurrentChunkOffset() const override { return SrcBufferOffset; }
@@ -267,7 +267,7 @@ protected:
 	 * @param[out] OutChunkSize the size of the chunk.
 	 * @return a pointer to the chunk if it's loaded, nullptr otherwise.
 	 */
-	const uint8* GetLoadedChunk(const FSoundWaveProxy& InSoundWave, uint32 ChunkIndex, uint32& OutChunkSize);
+	const uint8* GetLoadedChunk(const FSoundWaveProxyPtr& InSoundWave, uint32 ChunkIndex, uint32& OutChunkSize);
 
 	/** bool set before ParseHeader. Whether we are streaming a file or not. */
 	bool bIsStreaming;
diff --git a/Engine/Source/Runtime/Engine/Public/AudioDevice.h b/Engine/Source/Runtime/Engine/Public/AudioDevice.h
--- a/Engine/Source/Runtime/Engine/Public/AudioDevice.h
+++ b/Engine/Source/Runtime/Engine/Public/AudioDevice.h
@@ -1042,7 +1042,7 @@ public:
 
 	/** Creates a Compressed audio info class suitable for decompressing this SoundWave */
 	virtual ICompressedAudioInfo* CreateCompressedAudioInfo(USoundWave* SoundWave) { return nullptr; }
-	virtual ICompressedAudioInfo* CreateCompressedAudioInfo(const FSoundWaveProxy& SoundWave) { return nullptr; }
+	virtual ICompressedAudioInfo* CreateCompressedAudioInfo(const FSoundWaveProxyPtr& SoundWave) { return nullptr; }
 
 	/**
 	 * Check for errors and output a human readable string
diff --git a/Engine/Source/Runtime/Engine/Public/AudioStreaming.h b/Engine/Source/Runtime/Engine/Public/AudioStreaming.h
--- a/Engine/Source/Runtime/Engine/Public/AudioStreaming.h
+++ b/Engine/Source/Runtime/Engine/Public/AudioStreaming.h
@@ -23,6 +23,8 @@ class ICompressedAudioInfo;
 struct FLegacyAudioStreamingManager;
 struct FWaveInstance;
 
+using FSoundWaveProxyPtr = TSharedPtr<FSoundWaveProxy, ESPMode::ThreadSafe>;
+
 /** Lists possible states used by Thread-safe counter. */
 enum EAudioStreamingState
 {
@@ -140,7 +142,7 @@ struct FStreamingWaveData final
 	 *
 	 * @param SoundWave	The SoundWave we are managing
 	 */
-	bool Initialize(const FSoundWaveProxy& SoundWave, FLegacyAudioStreamingManager* InStreamingManager);
+	bool Initialize(const FSoundWaveProxyPtr& SoundWave, FLegacyAudioStreamingManager* InStreamingManager);
 
 	/**
 	 * Updates the streaming status of the sound wave and performs finalization when appropriate. The function returns
@@ -200,7 +202,7 @@ private:
 
 public:
 	/** SoundWave this streaming data is for */
-	TUniquePtr<FSoundWaveProxy> SoundWave;
+	FSoundWaveProxyPtr SoundWave;
 
 	/** Thread-safe counter indicating the audio streaming state. */
 	mutable FThreadSafeCounter	PendingChunkChangeRequestStatus;
@@ -264,18 +266,18 @@ struct FLegacyAudioStreamingManager : public IAudioStreamingManager
 	// End IStreamingManager interface
 
 	// IAudioStreamingManager interface
-	virtual void AddStreamingSoundWave(const FSoundWaveProxy& SoundWave) override;
-	virtual void RemoveStreamingSoundWave(const FSoundWaveProxy& SoundWave) override;
+	virtual void AddStreamingSoundWave(const FSoundWaveProxyPtr& SoundWave) override;
+	virtual void RemoveStreamingSoundWave(const FSoundWaveProxyPtr& SoundWave) override;
 	virtual void AddDecoder(ICompressedAudioInfo* CompressedAudioInfo) override;
 	virtual void RemoveDecoder(ICompressedAudioInfo* CompressedAudioInfo) override;
-	virtual bool IsManagedStreamingSoundWave(const FSoundWaveProxy&  SoundWave) const override;
-	virtual bool IsStreamingInProgress(const FSoundWaveProxy&  SoundWave) override;
+	virtual bool IsManagedStreamingSoundWave(const FSoundWaveProxyPtr&  SoundWave) const override;
+	virtual bool IsStreamingInProgress(const FSoundWaveProxyPtr&  SoundWave) override;
 	virtual bool CanCreateSoundSource(const FWaveInstance* WaveInstance) const override;
 	virtual void AddStreamingSoundSource(FSoundSource* SoundSource) override;
 	virtual void RemoveStreamingSoundSource(FSoundSource* SoundSource) override;
 	virtual bool IsManagedStreamingSoundSource(const FSoundSource* SoundSource) const override;
-	virtual bool RequestChunk(const FSoundWaveProxy& SoundWave, uint32 ChunkIndex, TFunction<void(EAudioChunkLoadResult)> OnLoadCompleted, ENamedThreads::Type ThreadToCallOnLoadCompletedOn, bool bForImmediatePlayback = false) override;
-	virtual FAudioChunkHandle GetLoadedChunk(const FSoundWaveProxy&  SoundWave, uint32 ChunkIndex, bool bBlockForLoad = false, bool bForImmediatePlayback = false) const override;
+	virtual bool RequestChunk(const FSoundWaveProxyPtr& SoundWave, uint32 ChunkIndex, TFunction<void(EAudioChunkLoadResult)> OnLoadCompleted, ENamedThreads::Type ThreadToCallOnLoadCompletedOn, bool bForImmediatePlayback = false) override;
+	virtual FAudioChunkHandle GetLoadedChunk(const FSoundWaveProxyPtr&  SoundWave, uint32 ChunkIndex, bool bBlockForLoad = false, bool bForImmediatePlayback = false) const override;
 	virtual uint64 TrimMemory(uint64 NumBytesToFree) override;
 	virtual int32 RenderStatAudioStreaming(UWorld* World, FViewport* Viewport, FCanvas* Canvas, int32 X, int32 Y, const FVector* ViewLocation, const FRotator* ViewRotation) override;
 	virtual FString GenerateMemoryReport() override;
diff --git a/Engine/Source/Runtime/Engine/Public/AudioStreamingCache.h b/Engine/Source/Runtime/Engine/Public/AudioStreamingCache.h
--- a/Engine/Source/Runtime/Engine/Public/AudioStreamingCache.h
+++ b/Engine/Source/Runtime/Engine/Public/AudioStreamingCache.h
@@ -40,7 +40,7 @@ public:
 		FChunkKey(FName InSoundWaveName, uint32 InChunkIndex);
 
 		FChunkKey(
-			  const FSoundWaveProxy& InSoundWave
+			  const FSoundWaveProxyPtr& InSoundWave
 			, uint32 InChunkIndex
 #if WITH_EDITOR
 			, uint32 InChunkRevision = 0
@@ -51,7 +51,7 @@ public:
 
 
 	private:
-		mutable TUniquePtr<FSoundWaveProxy> SoundWaveProxyPtr{ nullptr };
+		mutable FSoundWaveProxyPtr SoundWaveProxyPtr{ nullptr };
 
 	public:
 		mutable FName SoundWaveName = FName(); // mutable for GetTypeHash()
@@ -467,12 +467,12 @@ public:
 	// End IStreamingManager interface
 
 	// IAudioStreamingManager interface (unused functions)
-	virtual void AddStreamingSoundWave(const FSoundWaveProxy& SoundWave) override;
-	virtual void RemoveStreamingSoundWave(const FSoundWaveProxy& SoundWave) override;
+	virtual void AddStreamingSoundWave(const FSoundWaveProxyPtr& SoundWave) override;
+	virtual void RemoveStreamingSoundWave(const FSoundWaveProxyPtr& SoundWave) override;
 	virtual void AddDecoder(ICompressedAudioInfo* CompressedAudioInfo) override;
 	virtual void RemoveDecoder(ICompressedAudioInfo* CompressedAudioInfo) override;
-	virtual bool IsManagedStreamingSoundWave(const FSoundWaveProxy&  SoundWave) const override;
-	virtual bool IsStreamingInProgress(const FSoundWaveProxy&  SoundWave) override;
+	virtual bool IsManagedStreamingSoundWave(const FSoundWaveProxyPtr&  SoundWave) const override;
+	virtual bool IsStreamingInProgress(const FSoundWaveProxyPtr&  SoundWave) override;
 	virtual bool CanCreateSoundSource(const FWaveInstance* WaveInstance) const override;
 	virtual void AddStreamingSoundSource(FSoundSource* SoundSource) override;
 	virtual void RemoveStreamingSoundSource(FSoundSource* SoundSource) override;
@@ -480,8 +480,8 @@ public:
 	// End IAudioStreamingManager interface (unused)
 
 	// IAudioStreamingManager interface (used functions)
-	virtual bool RequestChunk(const FSoundWaveProxy& SoundWave, uint32 ChunkIndex, TFunction<void(EAudioChunkLoadResult)> OnLoadCompleted, ENamedThreads::Type ThreadToCallOnLoadCompletedOn, bool bForImmediatePlayback = false) override;
-	virtual FAudioChunkHandle GetLoadedChunk(const FSoundWaveProxy&  SoundWave, uint32 ChunkIndex, bool bBlockForLoad = false, bool bForImmediatePlayback = false) const override;
+	virtual bool RequestChunk(const FSoundWaveProxyPtr& SoundWave, uint32 ChunkIndex, TFunction<void(EAudioChunkLoadResult)> OnLoadCompleted, ENamedThreads::Type ThreadToCallOnLoadCompletedOn, bool bForImmediatePlayback = false) override;
+	virtual FAudioChunkHandle GetLoadedChunk(const FSoundWaveProxyPtr&  SoundWave, uint32 ChunkIndex, bool bBlockForLoad = false, bool bForImmediatePlayback = false) const override;
 	virtual uint64 TrimMemory(uint64 NumBytesToFree) override;
 	virtual int32 RenderStatAudioStreaming(UWorld* World, FViewport* Viewport, FCanvas* Canvas, int32 X, int32 Y, const FVector* ViewLocation, const FRotator* ViewRotation) override;
 	virtual FString GenerateMemoryReport() override;
@@ -500,13 +500,13 @@ protected:
 	 * based on the size of this sound wave's chunk,
 	 * or nullptr if MemoryLoadOnDemand is disabled.
 	 */
-	FAudioChunkCache* GetCacheForWave(const FSoundWaveProxy&  InSoundWave) const;
+	FAudioChunkCache* GetCacheForWave(const FSoundWaveProxyPtr&  InSoundWave) const;
 	FAudioChunkCache* GetCacheForChunkSize(uint32 InChunkSize) const;
 
 	/**
 	 * Returns the next chunk to kick off a load for, or INDEX_NONE if there is only one chunk to cache.
 	 */
-	int32 GetNextChunkIndex(const FSoundWaveProxy&  InSoundWave, uint32 CurrentChunkIndex) const;
+	int32 GetNextChunkIndex(const FSoundWaveProxyPtr&  InSoundWave, uint32 CurrentChunkIndex) const;
 
 	/** Audio chunk caches. These are set up on initialization. */
 	TArray<FAudioChunkCache> CacheArray;
diff --git a/Engine/Source/Runtime/Engine/Public/ContentStreaming.h b/Engine/Source/Runtime/Engine/Public/ContentStreaming.h
--- a/Engine/Source/Runtime/Engine/Public/ContentStreaming.h
+++ b/Engine/Source/Runtime/Engine/Public/ContentStreaming.h
@@ -22,6 +22,8 @@ struct FWaveInstance;
 class UAnimStreamable;
 struct FCompressedAnimSequence;
 
+using FSoundWaveProxyPtr = TSharedPtr<FSoundWaveProxy, ESPMode::ThreadSafe>;
+
 /*-----------------------------------------------------------------------------
 	Stats.
 -----------------------------------------------------------------------------*/
@@ -35,7 +37,6 @@ class USkeletalMesh;
 class ULandscapeLODStreamingProxy;
 class UStreamableRenderAsset;
 class FSoundSource;
-class FSoundWaveProxy;
 struct FWaveInstance;
 struct FRenderAssetStreamingManager;
 
@@ -117,12 +118,12 @@ public:
 
 private:
 	// This constructor should only be called by an implementation of IAudioStreamingManager.
-	FAudioChunkHandle(const uint8* InData, uint32 NumBytes, const FSoundWaveProxy&  InSoundWave, const FName& SoundWaveName, uint32 InChunkIndex, uint64 InCacheLookupID);
+	FAudioChunkHandle(const uint8* InData, uint32 NumBytes, const FSoundWaveProxyPtr&  InSoundWave, const FName& SoundWaveName, uint32 InChunkIndex, uint64 InCacheLookupID);
 
 	const uint8*  CachedData;
 	int32 CachedDataNumBytes;
 
-	TUniquePtr<FSoundWaveProxy> CorrespondingWave;
+	TWeakPtr<FSoundWaveProxy, ESPMode::ThreadSafe> CorrespondingWave;
 	FName CorrespondingWaveName;
 
 	// The index of this chunk in the sound wave's full set of chunks of compressed audio.
@@ -498,10 +499,10 @@ enum class EAudioChunkLoadResult : uint8
 struct IAudioStreamingManager : public IStreamingManager
 {
 	/** Adds a new Sound Wave to the streaming manager. */
-	virtual void AddStreamingSoundWave(const FSoundWaveProxy& SoundWave) = 0;
+	virtual void AddStreamingSoundWave(const FSoundWaveProxyPtr& SoundWave) = 0;
 
 	/** Removes a Sound Wave from the streaming manager. */
-	virtual void RemoveStreamingSoundWave(const FSoundWaveProxy& SoundWave) = 0;
+	virtual void RemoveStreamingSoundWave(const FSoundWaveProxyPtr& SoundWave) = 0;
 
 	/** Adds the decoder to the streaming manager to prevent stream chunks from getting reaped from underneath it */
 	virtual void AddDecoder(ICompressedAudioInfo* CompressedAudioInfo) = 0;
@@ -510,10 +511,10 @@ struct IAudioStreamingManager : public IStreamingManager
 	virtual void RemoveDecoder(ICompressedAudioInfo* CompressedAudioInfo) = 0;
 
 	/** Returns true if this is a Sound Wave that is managed by the streaming manager. */
-	virtual bool IsManagedStreamingSoundWave(const FSoundWaveProxy&  SoundWave) const = 0;
+	virtual bool IsManagedStreamingSoundWave(const FSoundWaveProxyPtr&  SoundWave) const = 0;
 
 	/** Returns true if this Sound Wave is currently streaming a chunk. */
-	virtual bool IsStreamingInProgress(const FSoundWaveProxy&  SoundWave) = 0;
+	virtual bool IsStreamingInProgress(const FSoundWaveProxyPtr&  SoundWave) = 0;
 
 	virtual bool CanCreateSoundSource(const FWaveInstance* WaveInstance) const = 0;
 
@@ -534,7 +535,7 @@ struct IAudioStreamingManager : public IStreamingManager
 	 * @param ThreadToCallOnLoadCompleteOn. Optional specifier for which thread OnLoadCompleted should be called on.
 	 * @param bForImmediatePlaybac if true, this will optionally reprioritize this chunk's load request.
 	 */
-	virtual bool RequestChunk(const FSoundWaveProxy& SoundWave, uint32 ChunkIndex, TFunction<void(EAudioChunkLoadResult)> OnLoadCompleted = [](EAudioChunkLoadResult) {}, ENamedThreads::Type ThreadToCallOnLoadCompletedOn = ENamedThreads::AnyThread, bool bForImmediatePlayback = false) = 0;
+	virtual bool RequestChunk(const FSoundWaveProxyPtr& SoundWave, uint32 ChunkIndex, TFunction<void(EAudioChunkLoadResult)> OnLoadCompleted = [](EAudioChunkLoadResult) {}, ENamedThreads::Type ThreadToCallOnLoadCompletedOn = ENamedThreads::AnyThread, bool bForImmediatePlayback = false) = 0;
 
 	/**
 	 * Gets a pointer to a chunk of audio data
@@ -545,7 +546,7 @@ struct IAudioStreamingManager : public IStreamingManager
 	 * @param bForImmediatePlayback if true, will optionally reprioritize this chunk's load request. See au.streamcaching.PlaybackRequestPriority.
 	 * @return a handle to the loaded chunk. Can return a default constructed FAudioChunkHandle if the chunk is not loaded yet.
 	 */
-	virtual FAudioChunkHandle GetLoadedChunk(const FSoundWaveProxy&  SoundWave, uint32 ChunkIndex,  bool bBlockForLoad = false, bool bForImmediatePlayback = false) const = 0;
+	virtual FAudioChunkHandle GetLoadedChunk(const FSoundWaveProxyPtr&  SoundWave, uint32 ChunkIndex,  bool bBlockForLoad = false, bool bForImmediatePlayback = false) const = 0;
 
 	/**
 	 * This will start evicting elements from the cache until either hit our target of bytes or run out of chunks we can free.
@@ -574,7 +575,7 @@ protected:
 	friend FAudioChunkHandle;
 
 	/** This can be called by implementers of IAudioStreamingManager to construct an FAudioChunkHandle using an otherwise inaccessible constructor. */
-	static FAudioChunkHandle BuildChunkHandle(const uint8* InData, uint32 NumBytes, const FSoundWaveProxy&  InSoundWave, const FName& SoundWaveName, uint32 InChunkIndex, uint64 CacheLookupID);
+	static FAudioChunkHandle BuildChunkHandle(const uint8* InData, uint32 NumBytes, const FSoundWaveProxyPtr& InSoundWave, const FName& SoundWaveName, uint32 InChunkIndex, uint64 CacheLookupID);
 
 	/**
 	 * This can be used to increment reference counted handles to audio chunks. Called by the copy constructor of FAudioChunkHandle.
diff --git a/Engine/Source/Runtime/Engine/Public/VorbisAudioInfo.h b/Engine/Source/Runtime/Engine/Public/VorbisAudioInfo.h
--- a/Engine/Source/Runtime/Engine/Public/VorbisAudioInfo.h
+++ b/Engine/Source/Runtime/Engine/Public/VorbisAudioInfo.h
@@ -105,10 +105,10 @@ public:
 	// End of ICompressedAudioInfo Interface
 
 protected:
-	virtual bool StreamCompressedInfoInternal(const FSoundWaveProxy& InWaveProxy, struct FSoundQualityInfo* QualityInfo) override;
+	virtual bool StreamCompressedInfoInternal(const FSoundWaveProxyPtr& InWaveProxy, struct FSoundQualityInfo* QualityInfo) override;
 
 private:
-	const uint8* GetLoadedChunk(FSoundWaveProxy& InSoundWave, uint32 ChunkIndex, uint32& OutChunkSize);
+	const uint8* GetLoadedChunk(FSoundWaveProxyPtr InSoundWave, uint32 ChunkIndex, uint32& OutChunkSize);
 
 	struct FVorbisFileWrapper* VFWrapper;
 	const uint8* SrcBufferData;
diff --git a/Engine/Source/Runtime/Windows/AudioMixerXAudio2/Private/AudioMixerPlatformXAudio2.cpp b/Engine/Source/Runtime/Windows/AudioMixerXAudio2/Private/AudioMixerPlatformXAudio2.cpp
--- a/Engine/Source/Runtime/Windows/AudioMixerXAudio2/Private/AudioMixerPlatformXAudio2.cpp
+++ b/Engine/Source/Runtime/Windows/AudioMixerXAudio2/Private/AudioMixerPlatformXAudio2.cpp
@@ -1446,22 +1446,27 @@ namespace Audio
 		return nullptr;
 	}
 
-	ICompressedAudioInfo* FMixerPlatformXAudio2::CreateCompressedAudioInfo(const FSoundWaveProxy& InSoundWave)
+	ICompressedAudioInfo* FMixerPlatformXAudio2::CreateCompressedAudioInfo(const FSoundWaveProxyPtr& InSoundWave)
 	{
+		if (!InSoundWave.IsValid())
+		{
+			return nullptr;
+		}
+
 #if WITH_ENGINE
 		// Decoding with FSoundWaveProxy is only supported ot streaming audio at the moment
-		if (false == InSoundWave.IsStreaming())
+		if (false == InSoundWave->IsStreaming())
 		{
 			return nullptr;
 		}
 
-		if (InSoundWave.IsSeekableStreaming())
+		if (InSoundWave->IsSeekableStreaming())
 		{
 			return new FADPCMAudioInfo();
 		}
 
 #if WITH_XMA2 && USE_XMA2_FOR_STREAMING
-		if (InSoundWave.GetNumChannels() <= 2)
+		if (InSoundWave->GetNumChannels() <= 2)
 		{
 			return XMA2_INFO_NEW();
 		}
diff --git a/Engine/Source/Runtime/Windows/AudioMixerXAudio2/Private/AudioMixerPlatformXAudio2.h b/Engine/Source/Runtime/Windows/AudioMixerXAudio2/Private/AudioMixerPlatformXAudio2.h
--- a/Engine/Source/Runtime/Windows/AudioMixerXAudio2/Private/AudioMixerPlatformXAudio2.h
+++ b/Engine/Source/Runtime/Windows/AudioMixerXAudio2/Private/AudioMixerPlatformXAudio2.h
@@ -82,7 +82,7 @@ namespace Audio
 		virtual bool SupportsRealtimeDecompression() const override { return true; }
 		virtual bool DisablePCMAudioCaching() const override;
 		virtual ICompressedAudioInfo* CreateCompressedAudioInfo(USoundWave* InSoundWave) override;
-		virtual ICompressedAudioInfo* CreateCompressedAudioInfo(const FSoundWaveProxy& InSoundWave) override;
+		virtual ICompressedAudioInfo* CreateCompressedAudioInfo(const FSoundWaveProxyPtr& InSoundWave) override;
 		virtual FString GetDefaultDeviceName() override;
 		virtual FAudioPlatformSettings GetPlatformSettings() const override;
 		virtual void OnHardwareUpdate() override;
