World Partition: register actor dependencies in the same cells as their referencers. This will make sure actors referenced by always loaded actors will get registered in the world, and also makes loading of cells dependencies load at the same time as their referencers.

diff --git a/Engine/Source/Runtime/Engine/Private/WorldPartition/WorldPartition.cpp b/Engine/Source/Runtime/Engine/Private/WorldPartition/WorldPartition.cpp
--- a/Engine/Source/Runtime/Engine/Private/WorldPartition/WorldPartition.cpp
+++ b/Engine/Source/Runtime/Engine/Private/WorldPartition/WorldPartition.cpp
@@ -12,6 +12,7 @@
 #include "WorldPartition/WorldPartitionSubsystem.h"
 #include "WorldPartition/WorldPartitionStreamingSource.h"
 #include "Algo/Accumulate.h"
+#include "Algo/Transform.h"
 #include "Engine/World.h"
 #include "Engine/LevelStreaming.h"
 #include "GameFramework/WorldSettings.h"
@@ -256,9 +257,12 @@ void UWorldPartition::Initialize(UWorld* InWorld, const FTransform& InTransform)
 		if (bEditorOnly)
 		{
 			// Make sure to preload any AInfo actors first (mainly for ShouldActorBeLoaded as it needs to have access to the WorldDataLayers actor)
-			TArray<FWorldPartitionHandle> AlwaysLoadedActors = EditorHash->GetAlwaysLoadedCell()->Actors.Array();			
-			AlwaysLoadedActors.Sort([](const FWorldPartitionHandle& A, const FWorldPartitionHandle& B) { return A->GetActorClass()->IsChildOf<AInfo>(); });
-			TArray<FWorldPartitionReference> LoadedActors(AlwaysLoadedActors);
+			TArray<UWorldPartitionEditorCell::FActorHandle> AlwaysLoadedActors = EditorHash->GetAlwaysLoadedCell()->Actors.Array();
+			AlwaysLoadedActors.Sort([](const UWorldPartitionEditorCell::FActorHandle& A, const UWorldPartitionEditorCell::FActorHandle& B) { return A->GetActorClass()->IsChildOf<AInfo>(); });
+			
+			TArray<FWorldPartitionReference> LoadedActors;
+			LoadedActors.Reserve(AlwaysLoadedActors.Num());
+			Algo::Transform(AlwaysLoadedActors, LoadedActors, [](const UWorldPartitionEditorCell::FActorHandle& ActorHandle) { return ActorHandle.Handle; });
 
 			// Load the always loaded cell, don't call LoadCells to avoid creating a transaction
 			UpdateLoadingEditorCell(EditorHash->GetAlwaysLoadedCell(), true);
@@ -559,7 +563,7 @@ bool UWorldPartition::UpdateEditorCells(TFunctionRef<bool(TArray<UWorldPartition
 
 	for (UWorldPartitionEditorCell* Cell : CellsToProcess)
 	{
-		for (const FWorldPartitionReference& ActorDesc : Cell->LoadedActors)
+		for (const UWorldPartitionEditorCell::FActorReference& ActorDesc : Cell->LoadedActors)
 		{
 			if (!bIsCellShouldBeLoaded || !ShouldActorBeLoaded(*ActorDesc))
 			{
@@ -681,7 +685,7 @@ void UWorldPartition::UpdateLoadingEditorCell(UWorldPartitionEditorCell* Cell, b
 	}
 	else
 	{
-		for (FWorldPartitionHandle& ActorHandle: Cell->Actors)
+		for (UWorldPartitionEditorCell::FActorHandle& ActorHandle: Cell->Actors)
 		{
 			AActor* Actor = ActorHandle->GetActor();
 
@@ -692,16 +696,16 @@ void UWorldPartition::UpdateLoadingEditorCell(UWorldPartitionEditorCell* Cell, b
 			// Filter actor against DataLayers
 			if (ShouldActorBeLoaded(*ActorHandle))
 			{
-				Cell->LoadedActors.Add(ActorHandle);
+				Cell->LoadedActors.Add(UWorldPartitionEditorCell::FActorReference(ActorHandle.Source, ActorHandle.Handle));
 			}
 			else
 			{
 				// Don't call LoadedActors.Remove(ActorHandle) right away, as this will create a temporary reference and might try to load.
-				for (const FWorldPartitionReference& ActorReference: Cell->LoadedActors)
+				for (const UWorldPartitionEditorCell::FActorReference& ActorReference: Cell->LoadedActors)
 				{
-					if (ActorReference == ActorHandle)
+					if ((ActorReference.Source == ActorHandle.Source) && (ActorReference.Handle == ActorHandle.Handle))
 					{
-						Cell->LoadedActors.Remove(ActorHandle);
+						Cell->LoadedActors.Remove(UWorldPartitionEditorCell::FActorReference(ActorHandle.Source, ActorHandle.Handle));
 						bPotentiallyUnloadedActors = true;
 						break;
 					}
diff --git a/Engine/Source/Runtime/Engine/Private/WorldPartition/WorldPartitionActorDesc.cpp b/Engine/Source/Runtime/Engine/Private/WorldPartition/WorldPartitionActorDesc.cpp
--- a/Engine/Source/Runtime/Engine/Private/WorldPartition/WorldPartitionActorDesc.cpp
+++ b/Engine/Source/Runtime/Engine/Private/WorldPartition/WorldPartitionActorDesc.cpp
@@ -8,7 +8,6 @@
 #include "UObject/UObjectHash.h"
 #include "UObject/MetaData.h"
 #include "Algo/Transform.h"
-#include "Algo/RemoveIf.h"
 #include "Engine/World.h"
 #include "Engine/Level.h"
 #include "Editor/GroupActor.h"
diff --git a/Engine/Source/Runtime/Engine/Private/WorldPartition/WorldPartitionEditorCell.cpp b/Engine/Source/Runtime/Engine/Private/WorldPartition/WorldPartitionEditorCell.cpp
--- a/Engine/Source/Runtime/Engine/Private/WorldPartition/WorldPartitionEditorCell.cpp
+++ b/Engine/Source/Runtime/Engine/Private/WorldPartition/WorldPartitionEditorCell.cpp
@@ -20,7 +20,6 @@ void UWorldPartitionEditorCell::Serialize(FArchive& Ar)
 }
 
 #if WITH_EDITOR
-
 void UWorldPartitionEditorCell::BeginDestroy()
 {
 	// Release WorldPartition Actor Handles/References
@@ -30,34 +29,61 @@ void UWorldPartitionEditorCell::BeginDestroy()
 }
 
 void UWorldPartitionEditorCell::AddActor(const FWorldPartitionHandle& ActorHandle)
+{
+	AddActor(ActorHandle->GetGuid(), ActorHandle);
+}
+
+void UWorldPartitionEditorCell::AddActor(const FGuid& Source, const FWorldPartitionHandle& ActorHandle)
 {
 	check(ActorHandle.IsValid());
 	
 	bool bIsAlreadyInSet = false;
-	Actors.Add(ActorHandle, &bIsAlreadyInSet);
-	check(!bIsAlreadyInSet);
+	Actors.Add(FActorHandle(Source, ActorHandle), &bIsAlreadyInSet);
 
-	if (ActorHandle.IsLoaded())
-	{
-		bIsAlreadyInSet = false;
-		LoadedActors.Add(ActorHandle, &bIsAlreadyInSet);
-		check(!bIsAlreadyInSet);
+	if (!bIsAlreadyInSet)
+	{	
+		if (ActorHandle.IsLoaded())
+		{
+			LoadedActors.Add(FActorReference(Source, ActorHandle), &bIsAlreadyInSet);
+			check(!bIsAlreadyInSet);
+		}
+
+		UWorldPartition* WorldPartition = GetTypedOuter<UWorldPartition>();
+		for (const FGuid& ReferenceGuid : ActorHandle->GetReferences())
+		{
+			FWorldPartitionHandle ReferenceHandle(WorldPartition, ReferenceGuid);
+			AddActor(ActorHandle->GetGuid(), ReferenceHandle);
+		}
 	}
 }
 
 void UWorldPartitionEditorCell::RemoveActor(const FWorldPartitionHandle& ActorHandle)
 {
-	check(ActorHandle.IsValid());
-	verify(Actors.Remove(ActorHandle));
+	RemoveActor(ActorHandle->GetGuid(), ActorHandle);
+}
 
-	// Don't call LoadedActors.Remove(ActorHandle) right away, as this will create a temporary reference and might try to load
-	// a deleted actor. This is a temporary workaround.
-	for (const FWorldPartitionReference& ActorReference: LoadedActors)
+void UWorldPartitionEditorCell::RemoveActor(const FGuid& Source, const FWorldPartitionHandle& ActorHandle)
+{
+	check(ActorHandle.IsValid());
+	
+	if (Actors.Remove(FActorHandle(Source, ActorHandle)))
 	{
-		if (ActorReference == ActorHandle)
+		// Don't call LoadedActors.Remove(ActorHandle) right away, as this will create a temporary reference and might try to load
+		// a deleted actor. This is a temporary workaround.
+		for (const FActorReference& ActorReference: LoadedActors)
+		{
+			if ((ActorReference.Source == Source) && (ActorReference.Handle == ActorHandle))
+			{
+				LoadedActors.Remove(FActorReference(Source, ActorHandle));
+				break;
+			}
+		}
+
+		UWorldPartition* WorldPartition = GetTypedOuter<UWorldPartition>();
+		for (const FGuid& ReferenceGuid : ActorHandle->GetReferences())
 		{
-			LoadedActors.Remove(ActorHandle);
-			break;
+			FWorldPartitionHandle ReferenceHandle(WorldPartition, ReferenceGuid);
+			RemoveActor(ActorHandle->GetGuid(), ReferenceHandle);
 		}
 	}
 }
diff --git a/Engine/Source/Runtime/Engine/Private/WorldPartition/WorldPartitionEditorSpatialHash.cpp b/Engine/Source/Runtime/Engine/Private/WorldPartition/WorldPartitionEditorSpatialHash.cpp
--- a/Engine/Source/Runtime/Engine/Private/WorldPartition/WorldPartitionEditorSpatialHash.cpp
+++ b/Engine/Source/Runtime/Engine/Private/WorldPartition/WorldPartitionEditorSpatialHash.cpp
@@ -392,7 +392,7 @@ int32 UWorldPartitionEditorSpatialHash::ForEachIntersectingActor(const FBox& Box
 
 	ForEachIntersectingCell(Box, [&](UWorldPartitionEditorCell* EditorCell)
 	{
-		for(FWorldPartitionHandle& ActorDesc: EditorCell->Actors)
+		for(UWorldPartitionEditorCell::FActorHandle& ActorDesc: EditorCell->Actors)
 		{
 			if (ActorDesc->Tag != FWorldPartitionActorDesc::GlobalTag)
 			{
@@ -407,17 +407,18 @@ int32 UWorldPartitionEditorSpatialHash::ForEachIntersectingActor(const FBox& Box
 		}
 	});
 
-	for(FWorldPartitionHandle& ActorDesc: AlwaysLoadedCell->Actors)
+	for(UWorldPartitionEditorCell::FActorHandle& ActorDesc: AlwaysLoadedCell->Actors)
 	{
-		check(ActorDesc->Tag != FWorldPartitionActorDesc::GlobalTag);
-
-		if (Box.Intersect(ActorDesc->GetBounds()))
+		if (ActorDesc->Tag != FWorldPartitionActorDesc::GlobalTag)
 		{
-			InOperation(*ActorDesc);
-			NumIntersecting++;
-		}
+			if (Box.Intersect(ActorDesc->GetBounds()))
+			{
+				InOperation(*ActorDesc);
+				NumIntersecting++;
+			}
 
-		ActorDesc->Tag = FWorldPartitionActorDesc::GlobalTag;		
+			ActorDesc->Tag = FWorldPartitionActorDesc::GlobalTag;
+		}
 	}
 
 	return NumIntersecting;
diff --git a/Engine/Source/Runtime/Engine/Public/WorldPartition/WorldPartitionEditorCell.h b/Engine/Source/Runtime/Engine/Public/WorldPartition/WorldPartitionEditorCell.h
--- a/Engine/Source/Runtime/Engine/Public/WorldPartition/WorldPartitionEditorCell.h
+++ b/Engine/Source/Runtime/Engine/Public/WorldPartition/WorldPartitionEditorCell.h
@@ -2,6 +2,7 @@
 #pragma once
 
 #include "CoreMinimal.h"
+#include "Misc/HashBuilder.h"
 #include "WorldPartition/WorldPartitionHandle.h"
 #include "WorldPartitionEditorCell.generated.h"
 
@@ -20,16 +21,61 @@ public:
 	virtual void Serialize(FArchive& Ar) override;
 
 #if WITH_EDITOR
+	template<typename Type>
+	struct TActorHandle
+	{
+		inline TActorHandle(const Type& InHandle)
+			: Source(InHandle->GetGuid())
+			, Handle(InHandle)
+		{}
+
+		inline TActorHandle(const FGuid& InSource, const Type& InHandle)
+			: Source(InSource)
+			, Handle(InHandle)
+		{}
+
+		inline FWorldPartitionActorDesc* operator->() const
+		{
+			return *Handle;
+		}
+
+		inline FWorldPartitionActorDesc* operator*() const
+		{
+			return *Handle;
+		}
+
+		inline bool operator==(const TActorHandle& Other) const
+		{
+			return (Source == Other.Source) && (Handle == Other.Handle);
+		}
+
+		friend ENGINE_API inline uint32 GetTypeHash(const TActorHandle& ActorHandle)
+		{
+			FHashBuilder HashBuilder;
+			HashBuilder << ActorHandle.Source << ActorHandle.Handle;
+			return HashBuilder.GetHash();
+		}
+
+		FGuid Source;
+		Type Handle;
+	};
+
+	typedef TActorHandle<FWorldPartitionHandle> FActorHandle;
+	typedef TActorHandle<FWorldPartitionReference> FActorReference;
+
 	virtual void BeginDestroy() override;
+
 	void AddActor(const FWorldPartitionHandle& ActorHandle);
+	void AddActor(const FGuid& Source, const FWorldPartitionHandle& ActorHandle);
 	void RemoveActor(const FWorldPartitionHandle& ActorHandle);
+	void RemoveActor(const FGuid& Source, const FWorldPartitionHandle& ActorHandle);
 
 	FBox						Bounds;
 
 	/** Tells if the cell was manually loaded in the editor */
 	bool						bLoaded : 1;
 
-	TSet<FWorldPartitionHandle> Actors;
-	TSet<FWorldPartitionReference> LoadedActors;
+	TSet<FActorHandle> Actors;
+	TSet<FActorReference> LoadedActors;
 #endif
 };
\ No newline at end of file
