World Partition - Fix crash when running with unbuilt HLODs

diff --git a/Engine/Source/Runtime/Engine/Private/WorldPartition/HLOD/HLODActor.cpp b/Engine/Source/Runtime/Engine/Private/WorldPartition/HLOD/HLODActor.cpp
--- a/Engine/Source/Runtime/Engine/Private/WorldPartition/HLOD/HLODActor.cpp
+++ b/Engine/Source/Runtime/Engine/Private/WorldPartition/HLOD/HLODActor.cpp
@@ -31,12 +31,18 @@ UPrimitiveComponent* AWorldPartitionHLOD::GetHLODComponent()
 
 void AWorldPartitionHLOD::OnCellShown(FName InCellName)
 {
-	GetRootComponent()->SetVisibility(false, true);
+	if (GetRootComponent())
+	{
+		GetRootComponent()->SetVisibility(false, true);
+	}
 }
 
 void AWorldPartitionHLOD::OnCellHidden(FName InCellName)
 {
-	GetRootComponent()->SetVisibility(true, true);
+	if (GetRootComponent())
+	{
+		GetRootComponent()->SetVisibility(true, true);
+	}
 }
 
 void AWorldPartitionHLOD::BeginPlay()
