Ensure HLOD generation can finish the first HLOD level properly even in case of WP actor registration errors
To be investigated

diff --git a/Engine/Source/Runtime/Engine/Private/WorldPartition/RuntimeSpatialHash/RuntimeSpatialHashHLOD.cpp b/Engine/Source/Runtime/Engine/Private/WorldPartition/RuntimeSpatialHash/RuntimeSpatialHashHLOD.cpp
--- a/Engine/Source/Runtime/Engine/Private/WorldPartition/RuntimeSpatialHash/RuntimeSpatialHashHLOD.cpp
+++ b/Engine/Source/Runtime/Engine/Private/WorldPartition/RuntimeSpatialHash/RuntimeSpatialHashHLOD.cpp
@@ -522,9 +522,10 @@ bool UWorldPartitionRuntimeSpatialHash::GenerateHLOD(ISourceControlHelper* Sourc
 		for (const FGuid& HLODActorGuid : HLODActors)
 		{
 			FWorldPartitionActorDesc* HLODActorDesc = WorldPartition->GetActorDesc(HLODActorGuid);
-			check(HLODActorDesc);
-
-			GridsHLODActors.FindOrAdd(HLODActorDesc->GetRuntimeGrid()).Add(HLODActorGuid);
+			if(ensureAlways(HLODActorDesc))
+			{
+				GridsHLODActors.FindOrAdd(HLODActorDesc->GetRuntimeGrid()).Add(HLODActorGuid);
+			}
 		}
 	};
 
