WP HLOD:
ensureAlways() -> ensure()

diff --git a/Engine/Source/Runtime/Engine/Private/WorldPartition/RuntimeSpatialHash/RuntimeSpatialHashHLOD.cpp b/Engine/Source/Runtime/Engine/Private/WorldPartition/RuntimeSpatialHash/RuntimeSpatialHashHLOD.cpp
--- a/Engine/Source/Runtime/Engine/Private/WorldPartition/RuntimeSpatialHash/RuntimeSpatialHashHLOD.cpp
+++ b/Engine/Source/Runtime/Engine/Private/WorldPartition/RuntimeSpatialHash/RuntimeSpatialHashHLOD.cpp
@@ -522,7 +522,7 @@ bool UWorldPartitionRuntimeSpatialHash::GenerateHLOD(ISourceControlHelper* Sourc
 		for (const FGuid& HLODActorGuid : HLODActors)
 		{
 			FWorldPartitionActorDesc* HLODActorDesc = WorldPartition->GetActorDesc(HLODActorGuid);
-			if(ensureAlways(HLODActorDesc))
+			if(ensure(HLODActorDesc))
 			{
 				GridsHLODActors.FindOrAdd(HLODActorDesc->GetRuntimeGrid()).Add(HLODActorGuid);
 			}
