Mesh merging - Added option to generate Nanite enabled meshes

diff --git a/Engine/Source/Developer/MeshMergeUtilities/Private/MeshMergeUtilities.cpp b/Engine/Source/Developer/MeshMergeUtilities/Private/MeshMergeUtilities.cpp
--- a/Engine/Source/Developer/MeshMergeUtilities/Private/MeshMergeUtilities.cpp
+++ b/Engine/Source/Developer/MeshMergeUtilities/Private/MeshMergeUtilities.cpp
@@ -2847,6 +2847,10 @@ void FMeshMergeUtilities::MergeComponentsToStaticMesh(const TArray<UPrimitiveCom
 		StaticMesh->ImportVersion = EImportStaticMeshVersion::LastVersion;
 		StaticMesh->SetLightMapResolution(InSettings.bComputedLightMapResolution ? DataTracker.GetLightMapDimension() : InSettings.TargetLightMapResolution);
 
+		// Nanite settings
+		StaticMesh->NaniteSettings.bEnabled = InSettings.bGenerateNaniteEnabledMesh;
+		StaticMesh->NaniteSettings.PercentTriangles = 1.0f;	// For now, assume meshes generated here are low poly enough to be used directly as the nanite proxy.
+
 #if WITH_EDITOR
 		//If we are running the automation test
 		if (GIsAutomationTesting)
diff --git a/Engine/Source/Developer/MeshMergeUtilities/Private/ProxyGenerationProcessor.cpp b/Engine/Source/Developer/MeshMergeUtilities/Private/ProxyGenerationProcessor.cpp
--- a/Engine/Source/Developer/MeshMergeUtilities/Private/ProxyGenerationProcessor.cpp
+++ b/Engine/Source/Developer/MeshMergeUtilities/Private/ProxyGenerationProcessor.cpp
@@ -286,6 +286,10 @@ void FProxyGenerationProcessor::ProcessJob(const FGuid& JobGuid, FProxyGeneratio
 		RemoveVertexColorAndCommitMeshDescription();
 	}
 
+	// Nanite settings
+	StaticMesh->NaniteSettings.bEnabled = Data->MergeData->InProxySettings.bGenerateNaniteEnabledMesh;
+	StaticMesh->NaniteSettings.PercentTriangles = 1.0f;	// For now, assume meshes generated here are low poly enough to be used directly as the nanite proxy.
+
 	//Set the Imported version before calling the build
 	StaticMesh->ImportVersion = EImportStaticMeshVersion::LastVersion;
 
diff --git a/Engine/Source/Runtime/Engine/Classes/Engine/MeshMerging.h b/Engine/Source/Runtime/Engine/Classes/Engine/MeshMerging.h
--- a/Engine/Source/Runtime/Engine/Classes/Engine/MeshMerging.h
+++ b/Engine/Source/Runtime/Engine/Classes/Engine/MeshMerging.h
@@ -319,6 +319,10 @@ struct FMeshProxySettings
 	UPROPERTY(BlueprintReadWrite, EditAnywhere, Category = ProxySettings)
 	uint8 bGenerateLightmapUVs:1;
 
+	/** Whether to generate a nanite-enabled mesh */
+	UPROPERTY(BlueprintReadWrite, EditAnywhere, Category = ProxySettings)
+	uint8 bGenerateNaniteEnabledMesh : 1;
+
 	/** Default settings. */
 	FMeshProxySettings()
 		: ScreenSize(300)
@@ -351,6 +355,7 @@ struct FMeshProxySettings
 		, bCreateCollision(true)
 		, bAllowVertexColors(false)
 		, bGenerateLightmapUVs(false)
+		, bGenerateNaniteEnabledMesh(false)
 	{
 		MaterialSettings.MaterialMergeType = EMaterialMergeType::MaterialMergeType_Simplygon;
 	}
@@ -359,17 +364,28 @@ struct FMeshProxySettings
 	bool operator==(const FMeshProxySettings& Other) const
 	{
 		return ScreenSize == Other.ScreenSize
+			&& VoxelSize == Other.VoxelSize
 			&& MaterialSettings == Other.MaterialSettings
-			&& bRecalculateNormals == Other.bRecalculateNormals
-			&& bOverrideTransferDistance == Other.bOverrideTransferDistance
+			&& MergeDistance == Other.MergeDistance
+			&& UnresolvedGeometryColor == Other.UnresolvedGeometryColor
 			&& MaxRayCastDist == Other.MaxRayCastDist
-			&& bUseHardAngleThreshold == Other.bUseHardAngleThreshold
 			&& HardAngleThreshold == Other.HardAngleThreshold
+			&& LightMapResolution == Other.LightMapResolution
 			&& NormalCalculationMethod == Other.NormalCalculationMethod
-			&& MergeDistance == Other.MergeDistance
-			&& UnresolvedGeometryColor == Other.UnresolvedGeometryColor
+			&& LandscapeCullingPrecision == Other.LandscapeCullingPrecision
+			&& bCalculateCorrectLODModel == Other.bCalculateCorrectLODModel
 			&& bOverrideVoxelSize == Other.bOverrideVoxelSize
-			&& VoxelSize == Other.VoxelSize;
+			&& bOverrideTransferDistance == Other.bOverrideTransferDistance
+			&& bUseHardAngleThreshold == Other.bUseHardAngleThreshold
+			&& bComputeLightMapResolution == Other.bComputeLightMapResolution
+			&& bRecalculateNormals == Other.bRecalculateNormals
+			&& bUseLandscapeCulling == Other.bUseLandscapeCulling
+			&& bAllowDistanceField == Other.bAllowDistanceField
+			&& bReuseMeshLightmapUVs == Other.bReuseMeshLightmapUVs
+			&& bCreateCollision == Other.bCreateCollision
+			&& bAllowVertexColors == Other.bAllowVertexColors
+			&& bGenerateLightmapUVs == Other.bGenerateLightmapUVs
+			&& bGenerateNaniteEnabledMesh == Other.bGenerateNaniteEnabledMesh;
 	}
 
 	/** Inequality. */
@@ -501,6 +517,10 @@ struct FMeshMergingSettings
 	UPROPERTY(EditAnywhere, Category = MeshSettings)
 	uint8 bAllowDistanceField:1;
 
+	/** Whether to generate a nanite-enabled mesh */
+	UPROPERTY(EditAnywhere, BlueprintReadWrite, Category = MeshSettings)
+	uint8 bGenerateNaniteEnabledMesh : 1;
+
 #if WITH_EDITORONLY_DATA
 	/** Whether we should import vertex colors into merged mesh */
 	UPROPERTY()
@@ -548,6 +568,7 @@ struct FMeshMergingSettings
 		, bUseLandscapeCulling(false)
 		, bIncludeImposters(true)
 		, bAllowDistanceField(false)
+		, bGenerateNaniteEnabledMesh(false)
 #if WITH_EDITORONLY_DATA
 		, bImportVertexColors_DEPRECATED(false)
 		, bCalculateCorrectLODModel_DEPRECATED(false)
