Fixed logic for building runtime hash grid level count.

diff --git a/Engine/Source/Runtime/Engine/Private/WorldPartition/RuntimeSpatialHash/RuntimeSpatialHashGridHelper.cpp b/Engine/Source/Runtime/Engine/Private/WorldPartition/RuntimeSpatialHash/RuntimeSpatialHashGridHelper.cpp
--- a/Engine/Source/Runtime/Engine/Private/WorldPartition/RuntimeSpatialHash/RuntimeSpatialHashGridHelper.cpp
+++ b/Engine/Source/Runtime/Engine/Private/WorldPartition/RuntimeSpatialHash/RuntimeSpatialHashGridHelper.cpp
@@ -21,12 +21,12 @@ FSquare2DGridHelper::FSquare2DGridHelper(const FBox& InWorldBounds, const FVecto
 	int32 GridLevelCount = 1;
 	if (WorldBoundsMaxExtent > 0.f)
 	{
-		GridSize = 2.f * FMath::CeilToFloat(WorldBoundsMaxExtent / CellSize);
+		GridSize = 2.f * FMath::CeilToFloat(WorldBoundsMaxExtent / CellSize); 
 		if (!FMath::IsPowerOfTwo(GridSize))
 		{
 			GridSize = FMath::Pow(2.f, FMath::CeilToFloat(FMath::Log2(static_cast<float>(GridSize))));
 		}
-		GridLevelCount = FMath::Log2(static_cast<float>(GridSize)) + 1;
+		GridLevelCount = FMath::FloorLog2(GridSize) + 1;
 	}
 	else
 	{
