SpatialHashStreamingGrid is now using sparse arrays to avoid crashing out of memory when duplicating very large worlds for PIE.

diff --git a/Engine/Source/Runtime/Engine/Private/WorldPartition/WorldPartitionRuntimeSpatialHash.cpp b/Engine/Source/Runtime/Engine/Private/WorldPartition/WorldPartitionRuntimeSpatialHash.cpp
--- a/Engine/Source/Runtime/Engine/Private/WorldPartition/WorldPartitionRuntimeSpatialHash.cpp
+++ b/Engine/Source/Runtime/Engine/Private/WorldPartition/WorldPartitionRuntimeSpatialHash.cpp
@@ -111,12 +111,15 @@ void FSpatialHashStreamingGrid::GetCells(const TArray<FWorldPartitionStreamingSo
 		const FSphere GridSphere(Source.Location, GetLoadingRange());
 		Helper.ForEachIntersectingCells(GridSphere, [&](const FIntVector& Coords)
 		{
-			const FSpatialHashStreamingGridLayerCell& LayerCell = GridLevels[Coords.Z].LayerCells[Coords.Y * Helper.Levels[Coords.Z].GridSize + Coords.X];
-			for (const UWorldPartitionRuntimeCell* Cell : LayerCell.GridCells)
+			if (const int32* LayerCellIndexPtr = GridLevels[Coords.Z].LayerCellsMapping.Find(Coords.Y * Helper.Levels[Coords.Z].GridSize + Coords.X))
 			{
-				if (!Cell->HasDataLayers() || (DataLayerSubsystem && DataLayerSubsystem->IsAnyDataLayerActive(Cell->GetDataLayers())))
+				const FSpatialHashStreamingGridLayerCell& LayerCell = GridLevels[Coords.Z].LayerCells[*LayerCellIndexPtr];
+				for (const UWorldPartitionRuntimeCell* Cell : LayerCell.GridCells)
 				{
-					Cells.Add(Cell);
+					if (!Cell->HasDataLayers() || (DataLayerSubsystem && DataLayerSubsystem->IsAnyDataLayerActive(Cell->GetDataLayers())))
+					{
+						Cells.Add(Cell);
+					}
 				}
 			}
 		});
@@ -175,16 +178,22 @@ void FSpatialHashStreamingGrid::Draw3D(UWorld* World, const TArray<FWorldPartiti
 				Helper.Levels[GridLevel].GetCellBounds(FIntVector2(Coords.X, Coords.Y), CellWorldBounds);
 				
 				FVector BoundsExtent(CellWorldBounds.GetExtent(), 100.f);
-				const FSpatialHashStreamingGridLayerCell& LayerCell = GridLevels[GridLevel].LayerCells[Coords.Y * Helper.Levels[GridLevel].GridSize + Coords.X];
+
 				const UWorldPartitionRuntimeSpatialHashCell* Cell = nullptr;
-				for (const UWorldPartitionRuntimeSpatialHashCell* GridCell : LayerCell.GridCells)
+				if (const int32* LayerCellIndexPtr = GridLevels[GridLevel].LayerCellsMapping.Find(Coords.Y * Helper.Levels[GridLevel].GridSize + Coords.X))
 				{
-					Cell = GridCell;
-					if (!GridCell->HasDataLayers())
+					const FSpatialHashStreamingGridLayerCell& LayerCell = GridLevels[GridLevel].LayerCells[*LayerCellIndexPtr];
+					
+					for (const UWorldPartitionRuntimeSpatialHashCell* GridCell : LayerCell.GridCells)
 					{
-						break;
+						Cell = GridCell;
+						if (!GridCell->HasDataLayers())
+						{
+							break;
+						}
 					}
 				}
+
 				FColor CellColor = Cell ? Cell->GetDebugColor().ToFColor(false).WithAlpha(16) : FColor(0, 0, 0, 16);
 				FVector BoundsOrigin(CellWorldBounds.GetCenter(), Z);
 				FBox Box = FBox::BuildAABB(BoundsOrigin, BoundsExtent);
@@ -264,16 +273,20 @@ void FSpatialHashStreamingGrid::Draw2D(UCanvas* Canvas, const TArray<FWorldParti
 				}
 			}
 
-			const FSpatialHashStreamingGridLayerCell& LayerCell = GridLevels[GridLevel].LayerCells[Coords.Y * Helper.Levels[GridLevel].GridSize + Coords.X];
 			const UWorldPartitionRuntimeSpatialHashCell* Cell = nullptr;
-			for (const UWorldPartitionRuntimeSpatialHashCell* GridCell : LayerCell.GridCells)
+			if (const int32* LayerCellIndexPtr = GridLevels[GridLevel].LayerCellsMapping.Find(Coords.Y * Helper.Levels[GridLevel].GridSize + Coords.X))
 			{
-				Cell = GridCell;
-				if (!GridCell->HasDataLayers())
+				const FSpatialHashStreamingGridLayerCell& LayerCell = GridLevels[GridLevel].LayerCells[*LayerCellIndexPtr];
+				for (const UWorldPartitionRuntimeSpatialHashCell* GridCell : LayerCell.GridCells)
 				{
-					break;
+					Cell = GridCell;
+					if (!GridCell->HasDataLayers())
+					{
+						break;
+					}
 				}
 			}
+
 			FLinearColor CellColor = Cell ? Cell->GetDebugColor() : FLinearColor(0.f, 0.f, 0.f, 0.25f);
 			FCanvasTileItem Item(CellScreenBounds.Min, GWhiteTexture, CellScreenBounds.GetSize(), CellColor);
 			Item.BlendMode = SE_BLEND_Translucent;
@@ -599,15 +612,12 @@ bool UWorldPartitionRuntimeSpatialHash::CreateStreamingGrid(const FSpatialHashRu
 
 		FSpatialHashStreamingGridLevel& GridLevel = CurrentStreamingGrid.GridLevels.AddDefaulted_GetRef();
 
-		GridLevel.LayerCells.SetNum(TempLevel.Cells.Num());
-
 		int32 CellIndex = INDEX_NONE;
 		for (const FSquare2DGridHelper::FGridLevel::FGridCell& TempCell : TempLevel.Cells)
 		{
 			CellIndex++;
 			const int32 CellCoordX = CellIndex % TempLevel.GridSize;
 			const int32 CellCoordY = CellIndex / TempLevel.GridSize;
-			TArray<UWorldPartitionRuntimeSpatialHashCell*>& GridCells = GridLevel.LayerCells[CellIndex].GridCells;
 
 			for (const FSquare2DGridHelper::FGridLevel::FGridCellDataChunk& GridCellDataChunk : TempCell.GetDataChunks())
 			{
@@ -654,7 +664,20 @@ bool UWorldPartitionRuntimeSpatialHash::CreateStreamingGrid(const FSpatialHashRu
 				FName CellName = GetCellName(CurrentStreamingGrid.GridName, Level, CellCoordX, CellCoordY, GridCellDataChunk.GetDataLayersID());
 
 				UWorldPartitionRuntimeSpatialHashCell* StreamingCell = NewObject<UWorldPartitionRuntimeSpatialHashCell>(WorldPartition, StreamingPolicy->GetRuntimeCellClass(), CellName);
-				GridCells.Add(StreamingCell);
+
+				int32 LayerCellIndex;
+				int32* LayerCellIndexPtr = GridLevel.LayerCellsMapping.Find(CellIndex);
+				if (LayerCellIndexPtr)
+				{
+					LayerCellIndex = *LayerCellIndexPtr;
+				}
+				else
+				{
+					LayerCellIndex = GridLevel.LayerCells.AddDefaulted();
+					GridLevel.LayerCellsMapping.Add(CellIndex, LayerCellIndex);
+				}
+
+				GridLevel.LayerCells[LayerCellIndex].GridCells.Add(StreamingCell);
 				StreamingCell->SetIsAlwaysLoaded(bIsCellAlwaysLoaded);
 				StreamingCell->SetDataLayers(GridCellDataChunk.GetDataLayers());
 				StreamingCell->Level = Level;
diff --git a/Engine/Source/Runtime/Engine/Public/WorldPartition/WorldPartitionRuntimeSpatialHash.h b/Engine/Source/Runtime/Engine/Public/WorldPartition/WorldPartitionRuntimeSpatialHash.h
--- a/Engine/Source/Runtime/Engine/Public/WorldPartition/WorldPartitionRuntimeSpatialHash.h
+++ b/Engine/Source/Runtime/Engine/Public/WorldPartition/WorldPartitionRuntimeSpatialHash.h
@@ -61,6 +61,12 @@ struct FSpatialHashStreamingGridLevel
 	 */
 	UPROPERTY()
 	TArray<FSpatialHashStreamingGridLayerCell> LayerCells;
+
+	/** 
+	 * Represents the grid cells index mapping
+	 */
+	UPROPERTY()
+	TMap<int32, int32> LayerCellsMapping;
 };
 
 // Disable copy of this USTRUCT by the autogenerated code.
