Derivatives refactor
Changed enums to enum classes
Moved FMaterialDerivativeAutogen to a separate file

diff --git a/Engine/Source/Runtime/Engine/Private/Materials/HLSLMaterialDerivativeAutogen.cpp b/Engine/Source/Runtime/Engine/Private/Materials/HLSLMaterialDerivativeAutogen.cpp
new file mode 100644
--- /dev/null
+++ b/Engine/Source/Runtime/Engine/Private/Materials/HLSLMaterialDerivativeAutogen.cpp
@@ -0,0 +1,1770 @@
+// Copyright Epic Games, Inc. All Rights Reserved.
+#include "HLSLMaterialDerivativeAutogen.h"
+#include "HLSLMaterialTranslator.h"
+
+#if WITH_EDITORONLY_DATA
+
+static int32 GDebugGenerateAllFunctionsEnabled = 0;
+static FAutoConsoleVariableRef CVarAnalyticDerivDebugGenerateAllFunctions(
+	TEXT("r.MaterialEditor.AnalyticDeriv.DebugGenerateAllFunctions"),
+	GDebugGenerateAllFunctionsEnabled,
+	TEXT("Debug: Generate all derivative functions.")
+);
+
+static inline bool IsDebugGenerateAllFunctionsEnabled()
+{
+	return GDebugGenerateAllFunctionsEnabled != 0;
+}
+
+static inline const TCHAR * GetBoolVectorName(int32 TypeIndex)
+{
+	switch(TypeIndex)
+	{
+	case 0:
+		return TEXT("bool");
+	case 1:
+		return TEXT("bool2");
+	case 2:
+		return TEXT("bool3");
+	case 3:
+		return TEXT("bool4");
+	default:
+		check(0);
+		return TEXT("");
+	}
+}
+
+static inline const TCHAR * GetFloatVectorName(int32 TypeIndex)
+{
+	switch(TypeIndex)
+	{
+	case 0:
+		return TEXT("float");
+	case 1:
+		return TEXT("float2");
+	case 2:
+		return TEXT("float3");
+	case 3:
+		return TEXT("float4");
+	default:
+		check(0);
+		return TEXT("");
+	}
+}
+
+static inline const TCHAR * GetDerivVectorName(int32 TypeIndex)
+{
+	switch(TypeIndex)
+	{
+	case 0:
+		return TEXT("FloatDeriv");
+	case 1:
+		return TEXT("FloatDeriv2");
+	case 2:
+		return TEXT("FloatDeriv3");
+	case 3:
+		return TEXT("FloatDeriv4");
+	default:
+		check(0);
+		return TEXT("");
+	}
+}
+
+
+// 0: Float
+// 1: Float2
+// 2: Float3
+// 3: Float4
+// -1: Everything else
+int32 GetDerivTypeIndex(EMaterialValueType ValueType)
+{
+	int32 Ret = INDEX_NONE;
+	if (ValueType == MCT_Float1 || ValueType == MCT_Float)
+	{
+		Ret = 0;
+	}
+	else if (ValueType == MCT_Float2)
+	{
+		Ret = 1;
+	}
+	else if (ValueType == MCT_Float3)
+	{
+		Ret = 2;
+	}
+	else if (ValueType == MCT_Float4)
+	{
+		Ret = 3;
+	}
+	else
+	{
+		Ret = -1; // redundant, but leaving it here so we can set breakpoints
+		check(false);
+	}
+	return Ret;
+}
+
+static EMaterialValueType GetMaterialTypeFromDerivTypeIndex(int32 Index)
+{
+	switch (Index)
+	{
+	case 0:
+		return MCT_Float;
+	case 1:
+		return MCT_Float2;
+	case 2:
+		return MCT_Float3;
+	case 3:
+		return MCT_Float4;
+	default:
+		check(0);
+		break;
+	}
+
+	return (EMaterialValueType)0; // invalid, should be a Float 1/2/3/4, break at the check(0);
+}
+
+
+static FString CoerceFloat(const TCHAR* Value, int32 DstType, int32 SrcType)
+{
+	if (DstType == SrcType)
+	{
+		return Value;
+	}
+
+	if (SrcType == 0)
+	{
+		const TCHAR* Mask[4] = { TEXT("x"), TEXT("xx"), TEXT("xxx"), TEXT("xxxx") };
+		return FString::Printf(TEXT("%s.%s"), Value, Mask[DstType]);
+	}
+	
+	if (DstType < SrcType)
+	{
+		const TCHAR* Mask[4] = { TEXT("x"), TEXT("xy"), TEXT("xyz"), TEXT("xyzw") };
+		return FString::Printf(TEXT("%s.%s"), Value, Mask[DstType]);
+	}
+	else
+	{
+		check(DstType > SrcType);
+		const TCHAR* Zeros[4] = { TEXT("0.0f"), TEXT("0.0f, 0.0f"), TEXT("0.0f, 0.0f, 0.0f"), TEXT("0.0f, 0.0f, 0.0f, 0.0f") };
+		return FString::Printf(TEXT("%s(%s, %s)"), GetFloatVectorName(DstType), Value, Zeros[DstType - SrcType - 1]);
+	}
+}
+
+void FMaterialDerivativeAutogen::EnableGeneratedDepencencies()
+{
+
+	for (int32 Index = 0; Index < 4; Index++)
+	{
+		// PowPositiveClamped requires Pow
+		if (bFunc2OpIsEnabled[(int32)EFunc2::PowPositiveClamped][Index])
+		{
+			bFunc2OpIsEnabled[(int32)EFunc2::Pow][Index] = true;
+		}
+	}
+
+	for (int32 Index = 0; Index < 4; Index++)
+	{
+		// normalize requires rsqrt1, dot, expand, and mul
+		if (bFunc1OpIsEnabled[(int32)EFunc1::Normalize][Index])
+		{
+			bConvertDerivEnabled[Index][0] = true;
+			bFunc2OpIsEnabled[(int32)EFunc2::Dot][Index] = true;
+			bFunc1OpIsEnabled[(int32)EFunc1::Rsqrt][0] = true;
+			bFunc2OpIsEnabled[(int32)EFunc2::Mul][Index] = true;
+		}
+
+		// length requires sqrt1 and dot, dot requires a few other things, but those are handled below
+		if (bFunc1OpIsEnabled[(int32)EFunc1::Length][Index])
+		{
+			bFunc2OpIsEnabled[(int32)EFunc2::Dot][Index] = true;
+			bFunc1OpIsEnabled[(int32)EFunc1::Sqrt][0] = true;
+		}
+
+		// inv length requires rsqrt1 (instead of sqrt1) and dot
+		if (bFunc1OpIsEnabled[(int32)EFunc1::InvLength][Index])
+		{
+			bFunc2OpIsEnabled[(int32)EFunc2::Dot][Index] = true;
+			bFunc1OpIsEnabled[(int32)EFunc1::Rsqrt][0] = true;
+		}
+	}
+
+	// min, max, saturate, and PosClampedPow require the select helper
+	for (int32 Index = 0; Index < 4; Index++)
+	{
+		if (bFunc2OpIsEnabled[(int32)EFunc2::Max][Index] ||
+			bFunc2OpIsEnabled[(int32)EFunc2::Min][Index] ||
+			bFunc1OpIsEnabled[(int32)EFunc1::Saturate][Index] ||
+			bFunc2OpIsEnabled[(int32)EFunc2::PowPositiveClamped][Index])
+		{
+			bSelectElemHelperEnabled[Index] = true;
+		}
+	}
+
+	// Dot requires extract, mul1, add1 and FloatDeriv constructor
+	for (int32 Index = 0; Index < 4; Index++)
+	{
+		if (bFunc2OpIsEnabled[(int32)EFunc2::Dot][Index])
+		{
+			bExtractIndexEnabled[Index] = true;
+			bConstructConstantDerivEnabled[0] = true;
+			bFunc2OpIsEnabled[(int32)EFunc2::Add][0] = true;
+			bFunc2OpIsEnabled[(int32)EFunc2::Mul][0] = true;
+		}
+	}
+
+	if (bRotateScaleOffsetTexCoords)
+	{
+		bFunc2OpIsEnabled[(int32)EFunc2::Add][1] = true;
+		bFunc2OpIsEnabled[(int32)EFunc2::Mul][1] = true;
+		bConstructDerivEnabled[1] = true;
+	}
+}
+
+// given a string, convert it from type to type
+FString FMaterialDerivativeAutogen::CoerceValueRaw(const FString& Token, int32 SrcType, EDerivativeStatus SrcStatus, int32 DstType)
+{
+	FString Ret = Token;
+
+	// If the original value is a derivative, grab the raw value
+	if (SrcStatus == EDerivativeStatus::Valid)
+	{
+		Ret = Ret + TEXT(".Value");
+	}
+
+	if (SrcType == DstType)
+	{
+		// no op
+	}
+	else
+	{
+		check(SrcType == 0); // can only coerce a float1
+		if (DstType == 0)
+		{
+			Ret = FString::Printf( TEXT("MaterialFloat(%s)"), *Ret);
+		}
+		else if (DstType == 1)
+		{
+			Ret = FString::Printf( TEXT("MaterialFloat2(%s,%s)"), *Ret, *Ret);
+		}
+		else if (DstType == 2)
+		{
+			Ret = FString::Printf( TEXT("MaterialFloat3(%s,%s,%s)"), *Ret, *Ret, *Ret);
+		}
+		else if (DstType == 3)
+		{
+			Ret = FString::Printf( TEXT("MaterialFloat4(%s,%s,%s,%s)"), *Ret, *Ret, *Ret, *Ret);
+		}
+		else
+		{
+			check(0);
+		}
+	}
+	return Ret;
+}
+
+// given a string, convert it from type to type
+FString FMaterialDerivativeAutogen::CoerceValueDeriv(const FString& Token, int32 SrcType, EDerivativeStatus SrcStatus, int32 DstType)
+{
+	check(SrcType < 4);
+	check(DstType < 4);
+
+	FString Ret = Token;
+
+	check(IsDerivativeValid(SrcStatus));
+
+	// If it's valid, then it's already a type. Otherwise, we need to convert it from raw to deriv
+	if (SrcStatus == EDerivativeStatus::Zero)
+	{
+		FString SrcDerivType = GetDerivVectorName(SrcType);
+		bConstructConstantDerivEnabled[SrcType] = true;
+		Ret = TEXT("ConstructConstant") + SrcDerivType + TEXT("(") + Ret + TEXT(")");
+	}
+
+	return ConvertDeriv(Ret, DstType, SrcType);
+}
+
+FString FMaterialDerivativeAutogen::ConstructDeriv(const FString& Value, const FString& Ddx, const FString& Ddy, int32 DstType)
+{
+	check(DstType >= 0);
+	check(DstType < 4);
+
+	bConstructDerivEnabled[DstType] = true;
+	FString TypeName = GetDerivVectorName(DstType);
+	FString Ret = TEXT("Construct") + TypeName + TEXT("(") + Value + TEXT(",") + Ddx + TEXT(",") + Ddy + TEXT(")");
+	return Ret;
+}
+
+FString FMaterialDerivativeAutogen::ConstructDerivFinite(const FString& Value, int32 DstType)
+{
+	check(DstType >= 0);
+	check(DstType < 4);
+
+	bConstructFiniteDerivEnabled[DstType] = true;
+
+	FString TypeName = GetDerivVectorName(DstType);
+	FString Ret = TEXT("ConstructFinite") + TypeName + TEXT("(") + Value + TEXT(")");
+	return Ret;
+}
+
+FString FMaterialDerivativeAutogen::ExtractElement(const FString& Value, int32 SrcType, int32 ElementIndex)
+{
+	check(SrcType < 4);
+	check(ElementIndex <= SrcType);
+	bExtractIndexEnabled[SrcType] = true;
+
+	FString TypeName = GetDerivVectorName(SrcType);
+
+	FString Ret = FString::Printf(TEXT("Extract%s_%d(%s %s)"), *TypeName, ElementIndex + 1, *TypeName, *Value);
+
+	return Ret;
+}
+
+FString FMaterialDerivativeAutogen::ConvertDeriv(const FString& Value, int32 DstType, int32 SrcType)
+{
+	check(DstType < 4);
+	check(SrcType < 4);
+	if (DstType == SrcType)
+		return Value;
+	
+	bConvertDerivEnabled[DstType][SrcType] = true;
+
+	FString DstTypeName = GetDerivVectorName(DstType);
+	return FString::Printf(TEXT("Convert%s(%s)"), *DstTypeName, *Value);
+}
+
+int32 FMaterialDerivativeAutogen::GetFunc1ReturnNumComponents(int32 SrcTypeIndex, EFunc1 Op)
+{
+	int32 DstTypeIndex = INDEX_NONE;
+
+	switch(Op)
+	{
+	case EFunc1::Abs:
+	case EFunc1::Log2:
+	case EFunc1::Log10:
+	case EFunc1::Exp:
+	case EFunc1::Sin:
+	case EFunc1::Cos:
+	case EFunc1::Tan:
+	case EFunc1::Asin:
+	case EFunc1::AsinFast:
+	case EFunc1::Acos:
+	case EFunc1::AcosFast:
+	case EFunc1::Atan:
+	case EFunc1::AtanFast:
+	case EFunc1::Sqrt:
+	case EFunc1::Rcp:
+	case EFunc1::Rsqrt:
+	case EFunc1::Saturate:
+	case EFunc1::Frac:
+	case EFunc1::Normalize:
+		DstTypeIndex = SrcTypeIndex;
+		break;
+
+	case EFunc1::Length:
+	case EFunc1::InvLength:
+		DstTypeIndex = 0;
+		break;
+	default:
+		break;
+	}
+
+	return DstTypeIndex;
+}
+
+int32 FMaterialDerivativeAutogen::GetFunc2ReturnNumComponents(int32 LhsTypeIndex, int32 RhsTypeIndex, EFunc2 Op)
+{
+	int32 DstTypeIndex = INDEX_NONE;
+	
+	switch(Op)
+	{
+	case EFunc2::Add:
+	case EFunc2::Sub:
+	case EFunc2::Mul:
+	case EFunc2::Div:
+	case EFunc2::Fmod:
+	case EFunc2::Max:
+	case EFunc2::Min:
+	case EFunc2::Pow:
+	case EFunc2::PowPositiveClamped:
+	case EFunc2::Atan2:
+	case EFunc2::Atan2Fast:
+		// if the initial type is different from the output type, then it's only valid if type is 0 (float). We can convert
+		// a float to a type with more components, but for example, we can't implicitly convert a float2 to a float3/float4.
+		if (LhsTypeIndex == RhsTypeIndex || RhsTypeIndex == 0 || LhsTypeIndex == 0)
+		{
+			DstTypeIndex = FMath::Max<int32>(LhsTypeIndex, RhsTypeIndex);
+		}
+		break;
+	case EFunc2::Dot:
+		DstTypeIndex = 0;
+		break;
+	case EFunc2::Cross:
+		check(LhsTypeIndex == 2);
+		check(RhsTypeIndex == 2);
+		DstTypeIndex = 2;
+		break;
+	default:
+		break;
+	}
+
+	return DstTypeIndex;
+}
+
+
+int32 FMaterialDerivativeAutogen::GenerateExpressionFunc1(FHLSLMaterialTranslator& Translator, EFunc1 Op, int32 SrcCode)
+{
+	if (SrcCode == INDEX_NONE)
+	{
+		return INDEX_NONE;
+	}
+
+	const FDerivInfo SrcDerivInfo = Translator.GetDerivInfo(SrcCode);
+	int32 OutputTypeIndex	= GetFunc1ReturnNumComponents(SrcDerivInfo.TypeIndex, Op);
+
+	if (OutputTypeIndex < 0)
+	{
+		return INDEX_NONE;
+	}
+
+	EDerivativeStatus DstStatus = IsDerivativeValid(SrcDerivInfo.DerivativeStatus) ? EDerivativeStatus::Valid : EDerivativeStatus::NotValid;
+	bool bUseScalarVersion = (DstStatus != EDerivativeStatus::Valid);
+
+	// make initial tokens
+	FString DstTokens[CompiledPDV_MAX];
+
+	for (int32 Index = 0; Index < CompiledPDV_MAX; Index++)
+	{
+		ECompiledPartialDerivativeVariation Variation = (ECompiledPartialDerivativeVariation)Index;
+
+		FString SrcToken = Translator.GetParameterCodeDeriv(SrcCode,Variation);
+
+		// The token is the symbol name. If we are in finite mode, that's all we have to do. But if
+		// we are in analytic mode, we may need to get the value.
+		if (Index == CompiledPDV_Analytic)
+		{
+			SrcToken = CoerceValueRaw(SrcToken, SrcDerivInfo.TypeIndex, SrcDerivInfo.DerivativeStatus, SrcDerivInfo.TypeIndex);
+		}
+
+		FString DstToken;
+		// just generate a type
+		switch(Op)
+		{
+		case EFunc1::Abs:
+			DstToken = TEXT("abs(") + SrcToken + TEXT(")");
+			break;
+		case EFunc1::Log2:
+			DstToken = TEXT("log2(") + SrcToken + TEXT(")");
+			break;
+		case EFunc1::Log10:
+			DstToken = TEXT("log10(") + SrcToken + TEXT(")");
+			break;
+		case EFunc1::Exp:
+			DstToken = TEXT("exp(") + SrcToken + TEXT(")");
+			break;
+		case EFunc1::Sin:
+			DstToken = TEXT("sin(") + SrcToken + TEXT(")");
+			break;
+		case EFunc1::Cos:
+			DstToken = TEXT("cos(") + SrcToken + TEXT(")");
+			break;
+		case EFunc1::Tan:
+			DstToken = TEXT("tan(") + SrcToken + TEXT(")");
+			break;
+		case EFunc1::Asin:
+			DstToken = TEXT("asin(") + SrcToken + TEXT(")");
+			break;
+		case EFunc1::AsinFast:
+			DstToken = TEXT("asinFast(") + SrcToken + TEXT(")");
+			break;
+		case EFunc1::Acos:
+			DstToken = TEXT("acos(") + SrcToken + TEXT(")");
+			break;
+		case EFunc1::AcosFast:
+			DstToken = TEXT("acosFast(") + SrcToken + TEXT(")");
+			break;
+		case EFunc1::Atan:
+			DstToken = TEXT("atan(") + SrcToken + TEXT(")");
+			break;
+		case EFunc1::AtanFast:
+			DstToken = TEXT("atanFast(") + SrcToken + TEXT(")");
+			break;
+		case EFunc1::Sqrt:
+			DstToken = TEXT("sqrt(") + SrcToken + TEXT(")");
+			break;
+		case EFunc1::Rcp:
+			DstToken = TEXT("rcp(") + SrcToken + TEXT(")");
+			break;
+		case EFunc1::Rsqrt:
+			DstToken = TEXT("rsqrt(") + SrcToken + TEXT(")");
+			break;
+		case EFunc1::Saturate:
+			DstToken = TEXT("saturate(") + SrcToken + TEXT(")");
+			break;
+		case EFunc1::Frac:
+			DstToken = TEXT("frac(") + SrcToken + TEXT(")");
+			break;
+		case EFunc1::Length:
+			DstToken = TEXT("length(") + SrcToken + TEXT(")");
+			break;
+		case EFunc1::InvLength:
+			DstToken = TEXT("rcp(length(") + SrcToken + TEXT("))");
+			break;
+		case EFunc1::Normalize:
+			DstToken = TEXT("normalize(") + SrcToken + TEXT(")");
+			break;
+		default:
+			check(0);
+			break;
+		}
+
+		DstTokens[Index] = DstToken;
+	}
+
+	if (!bUseScalarVersion)
+	{
+		FString SrcToken = Translator.GetParameterCodeDeriv(SrcCode, CompiledPDV_Analytic);
+
+		SrcToken = CoerceValueDeriv(SrcToken, SrcDerivInfo.TypeIndex, SrcDerivInfo.DerivativeStatus, SrcDerivInfo.TypeIndex);
+
+		check(Op < EFunc1::Num);
+		bFunc1OpIsEnabled[(int32)Op][SrcDerivInfo.TypeIndex] = true;
+		
+		FString DstToken;
+		switch(Op)
+		{
+		case EFunc1::Abs:
+			DstToken = TEXT("AbsDeriv(") + SrcToken + TEXT(")");
+			break;
+		case EFunc1::Log2:
+			DstToken = TEXT("Log2Deriv(") + SrcToken + TEXT(")");
+			break;
+		case EFunc1::Log10:
+			DstToken = TEXT("Log10Deriv(") + SrcToken + TEXT(")");
+			break;
+		case EFunc1::Exp:
+			DstToken = TEXT("ExpDeriv(") + SrcToken + TEXT(")");
+			break;
+		case EFunc1::Sin:
+			DstToken = TEXT("SinDeriv(") + SrcToken + TEXT(")");
+			break;
+		case EFunc1::Cos:
+			DstToken = TEXT("CosDeriv(") + SrcToken + TEXT(")");
+			break;
+		case EFunc1::Tan:
+			DstToken = TEXT("TanDeriv(") + SrcToken + TEXT(")");
+			break;
+		case EFunc1::Asin:
+			DstToken = TEXT("ASinDeriv(") + SrcToken + TEXT(")");
+			break;
+		case EFunc1::AsinFast:
+			DstToken = TEXT("ASinFastDeriv(") + SrcToken + TEXT(")");
+			break;
+		case EFunc1::Acos:
+			DstToken = TEXT("ACosDeriv(") + SrcToken + TEXT(")");
+			break;
+		case EFunc1::AcosFast:
+			DstToken = TEXT("ACosFastDeriv(") + SrcToken + TEXT(")");
+			break;
+		case EFunc1::Atan:
+			DstToken = TEXT("ATanDeriv(") + SrcToken + TEXT(")");
+			break;
+		case EFunc1::AtanFast:
+			DstToken = TEXT("ATanFastDeriv(") + SrcToken + TEXT(")");
+			break;
+		case EFunc1::Sqrt:
+			DstToken = TEXT("SqrtDeriv(") + SrcToken + TEXT(")");
+			break;
+		case EFunc1::Rcp:
+			DstToken = TEXT("RcpDeriv(") + SrcToken + TEXT(")");
+			break;
+		case EFunc1::Rsqrt:
+			DstToken = TEXT("RsqrtDeriv(") + SrcToken + TEXT(")");
+			break;
+		case EFunc1::Saturate:
+			DstToken = TEXT("SaturateDeriv(") + SrcToken + TEXT(")");
+			break;
+		case EFunc1::Frac:
+			DstToken = TEXT("FracDeriv(") + SrcToken + TEXT(")");
+			break;
+		case EFunc1::Length:
+			DstToken = TEXT("LengthDeriv(") + SrcToken + TEXT(")");
+			break;
+		case EFunc1::InvLength:
+			DstToken = TEXT("InvLengthDeriv(") + SrcToken + TEXT(")");
+			break;
+		case EFunc1::Normalize:
+			DstToken = TEXT("NormalizeDeriv(") + SrcToken + TEXT(")");
+			break;
+		default:
+			check(0);
+			break;
+		}
+
+		DstTokens[CompiledPDV_Analytic] = DstToken;
+	}
+
+	EMaterialValueType DstMatType = GetMaterialTypeFromDerivTypeIndex(OutputTypeIndex);
+	int32 Ret = Translator.AddCodeChunkInnerDeriv(*DstTokens[CompiledPDV_FiniteDifferences],*DstTokens[CompiledPDV_Analytic],DstMatType,false /*?*/, DstStatus);
+	return Ret;
+}
+
+int32 FMaterialDerivativeAutogen::GenerateExpressionFunc2(FHLSLMaterialTranslator& Translator, EFunc2 Op, int32 LhsCode, int32 RhsCode)
+{
+	if (LhsCode == INDEX_NONE || RhsCode == INDEX_NONE)
+	{
+		return INDEX_NONE;
+	}
+
+
+	const FDerivInfo LhsDerivInfo = Translator.GetDerivInfo(LhsCode);
+	const FDerivInfo RhsDerivInfo = Translator.GetDerivInfo(RhsCode);
+
+	int32 IntermediaryTypeIndex = FMath::Max<int32>(LhsDerivInfo.TypeIndex, RhsDerivInfo.TypeIndex);
+
+	if (Op == EFunc2::Cross)
+	{
+		IntermediaryTypeIndex = 2;
+	}
+
+	if (Op == EFunc2::Fmod)
+	{
+		check(RhsDerivInfo.DerivativeStatus == EDerivativeStatus::Zero);
+	}
+
+	int32 OutputTypeIndex = GetFunc2ReturnNumComponents(LhsDerivInfo.TypeIndex, RhsDerivInfo.TypeIndex, Op);
+
+	if (OutputTypeIndex < 0)
+	{
+		return INDEX_NONE;
+	}
+
+	EDerivativeStatus DstStatus = EDerivativeStatus::NotValid;
+
+	// Rules for derivatives:
+	// 1. If either the LHS or RHS is Not Valid or Not Aware, then the derivative is not valid. Run scalar route.
+	// 2. If both LHS and RHS are known to be Zero, then run raw code, and specify a known zero status.
+	// 3. If both LHS and RHS are Valid derivatives, then run deriv path.
+	// 4. If one is Valid and the other is known Zero, then promote the Zero to Valid, and run deriv path.
+
+	bool bUseScalarVersion = false;
+	bool bIsDerivValidZero = true;
+	bool bMakeDerivLhs = false;
+	bool bMakeDerivRhs = false;
+	if (!IsDerivativeValid(LhsDerivInfo.DerivativeStatus) || !IsDerivativeValid(RhsDerivInfo.DerivativeStatus))
+	{
+		// use scalar version as a fallback
+		bUseScalarVersion = true;
+		// derivative is not valid
+		bIsDerivValidZero = false;
+
+		// We output status as invalid, since one of the parameters is either not aware or not valid
+		DstStatus = EDerivativeStatus::NotValid;
+
+	}
+	else if (LhsDerivInfo.DerivativeStatus == EDerivativeStatus::Zero && RhsDerivInfo.DerivativeStatus == EDerivativeStatus::Zero)
+	{
+		// use scalar version
+		bUseScalarVersion = true;
+		// since we know both incoming values have derivatives of zero, we know the output is zero
+		bIsDerivValidZero = true;
+		// we know that the value is zero
+		DstStatus = EDerivativeStatus::Zero;
+	}
+	else
+	{
+		check(IsDerivativeValid(LhsDerivInfo.DerivativeStatus));
+		check(IsDerivativeValid(RhsDerivInfo.DerivativeStatus));
+
+		// use deriv version
+		bUseScalarVersion = false;
+		// we will be calculating a valid derivative, so this value doesn't matter, but make it false for consistency
+		bIsDerivValidZero = false;
+
+		// if the lhs has a derivitive of zero, and rhs is non-zero, convert lhs from a scalar type to deriv type
+		if (LhsDerivInfo.DerivativeStatus == EDerivativeStatus::Zero)
+		{
+			bMakeDerivLhs = true;
+		}
+
+		// if the rhs has a derivative of zero, and lhs is non-zero, convert rhs from a scalar type to deriv type
+		if (RhsDerivInfo.DerivativeStatus == EDerivativeStatus::Zero)
+		{
+			bMakeDerivRhs = true;
+		}
+
+		// derivative results will be valid
+		DstStatus = EDerivativeStatus::Valid;
+	}
+
+	FString DstTokens[CompiledPDV_MAX];
+
+	for (int32 Index = 0; Index < CompiledPDV_MAX; Index++)
+	{
+		ECompiledPartialDerivativeVariation Variation = (ECompiledPartialDerivativeVariation)Index;
+
+		FString LhsToken = Translator.GetParameterCodeDeriv(LhsCode,Variation);
+		FString RhsToken = Translator.GetParameterCodeDeriv(RhsCode,Variation);
+
+		// The token is the symbol name. If we are in finite mode, that's all we have to do. But if
+		// we are in analytic mode, we may need to get the value.
+		if (Index == CompiledPDV_Analytic)
+		{
+			LhsToken = CoerceValueRaw(LhsToken, LhsDerivInfo.TypeIndex, LhsDerivInfo.DerivativeStatus, IntermediaryTypeIndex);
+			RhsToken = CoerceValueRaw(RhsToken, RhsDerivInfo.TypeIndex, RhsDerivInfo.DerivativeStatus, IntermediaryTypeIndex);
+		}
+
+		FString DstToken;
+		// just generate a type
+		switch(Op)
+		{
+		case EFunc2::Add:
+			DstToken = TEXT("(") + LhsToken + TEXT(" + ") + RhsToken + TEXT(")");
+			break;
+		case EFunc2::Sub:
+			DstToken = TEXT("(") + LhsToken + TEXT(" - ") + RhsToken + TEXT(")");
+			break;
+		case EFunc2::Mul:
+			DstToken = TEXT("(") + LhsToken + TEXT(" * ") + RhsToken + TEXT(")");
+			break;
+		case EFunc2::Div:
+			DstToken = TEXT("(") + LhsToken + TEXT(" / ") + RhsToken + TEXT(")");
+			break;
+		case EFunc2::Fmod:
+			DstToken = TEXT("fmod(") + LhsToken + TEXT(",") + RhsToken + TEXT(")");
+			break;
+		case EFunc2::Min:
+			DstToken = TEXT("min(") + LhsToken + TEXT(",") + RhsToken + TEXT(")");
+			break;
+		case EFunc2::Max:
+			DstToken = TEXT("max(") + LhsToken + TEXT(",") + RhsToken + TEXT(")");
+			break;
+		case EFunc2::Dot:
+			DstToken = TEXT("dot(") + LhsToken + TEXT(",") + RhsToken + TEXT(")");
+			break;
+		case EFunc2::Pow:
+			DstToken = TEXT("pow(") + LhsToken + TEXT(",") + RhsToken + TEXT(")");
+			break;
+		case EFunc2::PowPositiveClamped:
+			DstToken = TEXT("PositiveClampedPow(") + LhsToken + TEXT(",") + RhsToken + TEXT(")");
+			break;
+		case EFunc2::Atan2:
+			DstToken = TEXT("atan2(") + LhsToken + TEXT(",") + RhsToken + TEXT(")");
+			break;
+		case EFunc2::Atan2Fast:
+			DstToken = TEXT("atan2Fast(") + LhsToken + TEXT(",") + RhsToken + TEXT(")");
+			break;
+		case EFunc2::Cross:
+			DstToken = TEXT("cross(") + LhsToken + TEXT(",") + RhsToken + TEXT(")");
+			break;
+		default:
+			check(0);
+			break;
+		}
+
+		DstTokens[Index] = DstToken;
+	}
+
+	if (!bUseScalarVersion)
+	{
+		FString LhsToken = Translator.GetParameterCodeDeriv(LhsCode,CompiledPDV_Analytic);
+		FString RhsToken = Translator.GetParameterCodeDeriv(RhsCode,CompiledPDV_Analytic);
+
+		LhsToken = CoerceValueDeriv(LhsToken, LhsDerivInfo.TypeIndex, LhsDerivInfo.DerivativeStatus, IntermediaryTypeIndex);
+		RhsToken = CoerceValueDeriv(RhsToken, RhsDerivInfo.TypeIndex, RhsDerivInfo.DerivativeStatus, IntermediaryTypeIndex);
+
+		check(Op < EFunc2::Num);
+		bFunc2OpIsEnabled[(int32)Op][IntermediaryTypeIndex] = true;
+
+		FString DstToken;
+		switch(Op)
+		{
+		case EFunc2::Add:
+			DstToken = TEXT("AddDeriv(") + LhsToken + TEXT(",") + RhsToken + TEXT(")");
+			break;
+		case EFunc2::Sub:
+			DstToken = TEXT("SubDeriv(") + LhsToken + TEXT(",") + RhsToken + TEXT(")");
+			break;
+		case EFunc2::Mul:
+			DstToken = TEXT("MulDeriv(") + LhsToken + TEXT(",") + RhsToken + TEXT(")");
+			break;
+		case EFunc2::Div:
+			DstToken = TEXT("DivDeriv(") + LhsToken + TEXT(",") + RhsToken + TEXT(")");
+			break;
+		case EFunc2::Fmod:
+			DstToken = TEXT("FmodDeriv(") + LhsToken + TEXT(",") + RhsToken + TEXT(")");
+			break;
+		case EFunc2::Dot:
+			DstToken = TEXT("DotDeriv(") + LhsToken + TEXT(",") + RhsToken + TEXT(")");
+			break;
+		case EFunc2::Min:
+			DstToken = TEXT("MinDeriv(") + LhsToken + TEXT(",") + RhsToken + TEXT(")");
+			break;
+		case EFunc2::Max:
+			DstToken = TEXT("MaxDeriv(") + LhsToken + TEXT(",") + RhsToken + TEXT(")");
+			break;
+		case EFunc2::Pow:
+			DstToken = TEXT("PowDeriv(") + LhsToken + TEXT(",") + RhsToken + TEXT(")");
+			break;
+		case EFunc2::PowPositiveClamped:
+			DstToken = TEXT("PowPositiveClampedDeriv(") + LhsToken + TEXT(",") + RhsToken + TEXT(")");
+			break;
+		case EFunc2::Cross:
+			DstToken = TEXT("CrossDeriv(") + LhsToken + TEXT(",") + RhsToken + TEXT(")");
+			break;
+		case EFunc2::Atan2:
+			DstToken = TEXT("Atan2Deriv(") + LhsToken + TEXT(",") + RhsToken + TEXT(")");
+			break;
+		case EFunc2::Atan2Fast:
+			DstToken = TEXT("Atan2FastDeriv(") + LhsToken + TEXT(",") + RhsToken + TEXT(")");
+			break;
+		default:
+			check(0);
+			break;
+		}
+
+		DstTokens[CompiledPDV_Analytic] = DstToken;
+		DstStatus = EDerivativeStatus::Valid;
+	}
+
+	EMaterialValueType DstMatType = GetMaterialTypeFromDerivTypeIndex(OutputTypeIndex);
+
+	int32 Ret = Translator.AddCodeChunkInnerDeriv(*DstTokens[CompiledPDV_FiniteDifferences],*DstTokens[CompiledPDV_Analytic],DstMatType,false /*?*/, DstStatus);
+	return Ret;
+}
+
+int32 FMaterialDerivativeAutogen::GenerateLerpFunc(FHLSLMaterialTranslator& Translator, int32 A, int32 B, int32 S)
+{
+	// TODO: generalize to Func3?
+	if (A == INDEX_NONE || B == INDEX_NONE || S == INDEX_NONE)
+	{
+		return INDEX_NONE;
+	}
+
+	const FDerivInfo ADerivInfo = Translator.GetDerivInfo(A);
+	const FDerivInfo BDerivInfo = Translator.GetDerivInfo(B);
+	const FDerivInfo SDerivInfo = Translator.GetDerivInfo(S);
+	
+	const EMaterialValueType ResultType = Translator.GetArithmeticResultType(A, B);
+	const EMaterialValueType AlphaType = (ResultType == SDerivInfo.Type) ? ResultType : MCT_Float1;
+	const uint32 NumResultComponents = GetNumComponents(ResultType);
+
+	const bool bAllZeroDeriv = (ADerivInfo.DerivativeStatus == EDerivativeStatus::Zero && BDerivInfo.DerivativeStatus == EDerivativeStatus::Zero && SDerivInfo.DerivativeStatus == EDerivativeStatus::Zero);
+	FString FiniteString = FString::Printf(TEXT("lerp(%s,%s,%s)"), *Translator.CoerceParameter(A, ResultType), *Translator.CoerceParameter(B, ResultType), *Translator.CoerceParameter(S, AlphaType));
+	
+	if (!bAllZeroDeriv && IsDerivativeValid(ADerivInfo.DerivativeStatus) && IsDerivativeValid(BDerivInfo.DerivativeStatus) && IsDerivativeValid(SDerivInfo.DerivativeStatus))
+	{
+		const uint32 ResultTypeIndex = GetDerivTypeIndex(ResultType);
+		//const uint32 AlphaTypeIndex = GetDerivTypeIndex(AlphaType);
+		FString ADeriv = Translator.GetParameterCodeDeriv(A, CompiledPDV_Analytic);
+		FString BDeriv = Translator.GetParameterCodeDeriv(B, CompiledPDV_Analytic);
+		FString SDeriv = Translator.GetParameterCodeDeriv(S, CompiledPDV_Analytic);
+
+		ADeriv = CoerceValueDeriv(ADeriv, ADerivInfo.TypeIndex, ADerivInfo.DerivativeStatus, ResultTypeIndex);
+		BDeriv = CoerceValueDeriv(BDeriv, BDerivInfo.TypeIndex, BDerivInfo.DerivativeStatus, ResultTypeIndex);
+		SDeriv = CoerceValueDeriv(SDeriv, SDerivInfo.TypeIndex, SDerivInfo.DerivativeStatus, ResultTypeIndex);
+
+		FString AnalyticString = FString::Printf(TEXT("LerpDeriv(%s, %s, %s)"), *ADeriv, *BDeriv, *SDeriv);
+
+		check(NumResultComponents <= 4);
+		bLerpEnabled[ResultTypeIndex] = true;
+
+		return Translator.AddCodeChunkInnerDeriv(*FiniteString, *AnalyticString, ResultType, false, EDerivativeStatus::Valid);
+	}
+	else
+	{
+		return Translator.AddCodeChunkInnerDeriv(*FiniteString, *FiniteString, ResultType, false, bAllZeroDeriv ? EDerivativeStatus::Zero : EDerivativeStatus::NotValid);
+	}
+}
+
+int32 FMaterialDerivativeAutogen::GenerateRotateScaleOffsetTexCoordsFunc(FHLSLMaterialTranslator& Translator, int32 TexCoord, int32 RotationScale, int32 Offset)
+{
+	// TODO: generalize to Func3?
+	if (TexCoord == INDEX_NONE || RotationScale == INDEX_NONE || Offset == INDEX_NONE)
+	{
+		return INDEX_NONE;
+	}
+	const FDerivInfo TexCoordDerivInfo		= Translator.GetDerivInfo(TexCoord);
+	const FDerivInfo RotationScaleDerivInfo	= Translator.GetDerivInfo(RotationScale);
+	const FDerivInfo OffsetDerivInfo		= Translator.GetDerivInfo(Offset);
+
+	const EMaterialValueType ResultType = MCT_Float2;
+	const uint32 NumResultComponents = 2;
+
+	const bool bAllZeroDeriv = (TexCoordDerivInfo.DerivativeStatus == EDerivativeStatus::Zero && RotationScaleDerivInfo.DerivativeStatus == EDerivativeStatus::Zero && OffsetDerivInfo.DerivativeStatus == EDerivativeStatus::Zero);
+	FString FiniteString = FString::Printf(TEXT("RotateScaleOffsetTexCoords(%s, %s, %s.xy)"),	*Translator.CoerceParameter(TexCoord,		ResultType), 
+																								*Translator.CoerceParameter(RotationScale,	ResultType),
+																								*Translator.CoerceParameter(Offset,			ResultType));
+
+	if (!bAllZeroDeriv && IsDerivativeValid(TexCoordDerivInfo.DerivativeStatus) && IsDerivativeValid(RotationScaleDerivInfo.DerivativeStatus) && IsDerivativeValid(OffsetDerivInfo.DerivativeStatus))
+	{
+		const uint32 ResultTypeIndex = GetDerivTypeIndex(ResultType);
+		FString TexCoordDeriv		= Translator.GetParameterCodeDeriv(TexCoord, CompiledPDV_Analytic);
+		FString RotationScaleDeriv	= Translator.GetParameterCodeDeriv(RotationScale, CompiledPDV_Analytic);
+		FString OffsetDeriv			= Translator.GetParameterCodeDeriv(Offset, CompiledPDV_Analytic);
+
+		TexCoordDeriv		= CoerceValueDeriv(TexCoordDeriv,		TexCoordDerivInfo.TypeIndex,		TexCoordDerivInfo.DerivativeStatus,			ResultTypeIndex);
+		RotationScaleDeriv	= CoerceValueDeriv(RotationScaleDeriv,	RotationScaleDerivInfo.TypeIndex,	RotationScaleDerivInfo.DerivativeStatus,	ResultTypeIndex);
+		OffsetDeriv			= CoerceValueDeriv(OffsetDeriv,			OffsetDerivInfo.TypeIndex,			OffsetDerivInfo.DerivativeStatus,			ResultTypeIndex);
+
+		FString AnalyticString = FString::Printf(TEXT("RotateScaleOffsetTexCoordsDeriv(%s, %s, %s)"), *TexCoordDeriv, *RotationScaleDeriv, *OffsetDeriv);
+
+		check(NumResultComponents <= 4);
+		bRotateScaleOffsetTexCoords = true;
+
+		return Translator.AddCodeChunkInnerDeriv(*FiniteString, *AnalyticString, ResultType, false, EDerivativeStatus::Valid);
+	}
+	else
+	{
+		return Translator.AddCodeChunkInnerDeriv(*FiniteString, *FiniteString, ResultType, false, bAllZeroDeriv ? EDerivativeStatus::Zero : EDerivativeStatus::NotValid);
+	}
+}
+
+int32 FMaterialDerivativeAutogen::GenerateIfFunc(FHLSLMaterialTranslator& Translator, int32 A, int32 B, int32 Greater, int32 Equal, int32 Less, int32 Threshold)
+{
+	if (A == INDEX_NONE || B == INDEX_NONE || Greater == INDEX_NONE || Less == INDEX_NONE)
+	{
+		return INDEX_NONE;
+	}
+
+	const FString AFinite = Translator.GetParameterCode(A);
+	const FString BFinite = Translator.GetParameterCode(B);
+
+	const bool bEqual = (Equal != INDEX_NONE);
+
+	if (bEqual && Threshold == INDEX_NONE)
+		return INDEX_NONE;
+
+	EMaterialValueType ResultType = Translator.GetArithmeticResultType(Greater, Less);
+
+	if (bEqual)
+		ResultType = Translator.GetArithmeticResultType(ResultType, Translator.GetParameterType(Greater));
+	
+	Greater = Translator.ForceCast(Greater,		ResultType);
+	Less	= Translator.ForceCast(Less,		ResultType);
+	if (Greater == INDEX_NONE || Less == INDEX_NONE)
+		return INDEX_NONE;
+
+	const FString GreaterFinite = Translator.GetParameterCode(Greater);
+	const FString LessFinite	= Translator.GetParameterCode(Less);
+
+	if (bEqual)
+	{
+		Equal = Translator.ForceCast(Equal, ResultType);
+		if (Equal == INDEX_NONE)
+			return INDEX_NONE;
+	}
+	
+	const FDerivInfo GreaterDerivInfo	= Translator.GetDerivInfo(Greater);
+	const FDerivInfo LessDerivInfo		= Translator.GetDerivInfo(Less);
+
+	FString CodeFinite;
+	FString ThresholdFinite;
+	if (bEqual)
+	{
+		ThresholdFinite = *Translator.GetParameterCode(Threshold);
+
+		CodeFinite = FString::Printf(
+			TEXT("((abs(%s - %s) > %s) ? (%s >= %s ? %s : %s) : %s)"),
+			*AFinite, *BFinite, *ThresholdFinite,
+			*AFinite, *BFinite,
+			*GreaterFinite, *LessFinite, *Translator.GetParameterCode(Equal));
+	}
+	else
+	{
+		CodeFinite = FString::Printf(
+			TEXT("((%s >= %s) ? %s : %s)"),
+			*AFinite, *BFinite,
+			*GreaterFinite, *LessFinite
+		);
+	}
+
+	const bool bAllDerivValid = IsDerivativeValid(GreaterDerivInfo.DerivativeStatus) && IsDerivativeValid(LessDerivInfo.DerivativeStatus) && (!bEqual || IsDerivativeValid(Translator.GetDerivativeStatus(Equal)));
+	const bool bAllDerivZero = GreaterDerivInfo.DerivativeStatus == EDerivativeStatus::Zero && LessDerivInfo.DerivativeStatus == EDerivativeStatus::Zero && (!bEqual || Translator.GetDerivativeStatus(Equal) == EDerivativeStatus::Zero);
+	if (bAllDerivValid && !bAllDerivZero)
+	{
+		const uint32 ResultTypeIndex = GetDerivTypeIndex(ResultType);
+
+		FString GreaterDeriv	= Translator.GetParameterCodeDeriv(Greater, CompiledPDV_Analytic);
+		FString LessDeriv		= Translator.GetParameterCodeDeriv(Less,	CompiledPDV_Analytic);
+
+		GreaterDeriv	= CoerceValueDeriv(GreaterDeriv,	ResultTypeIndex, GreaterDerivInfo.DerivativeStatus, ResultTypeIndex);
+		LessDeriv		= CoerceValueDeriv(LessDeriv,		ResultTypeIndex, LessDerivInfo.DerivativeStatus,	ResultTypeIndex);
+
+		FString CodeAnalytic;
+		if (bEqual)
+		{
+			const FDerivInfo EqualDerivInfo = Translator.GetDerivInfo(Equal);
+			FString EqualDeriv = Translator.GetParameterCodeDeriv(Equal, CompiledPDV_Analytic);
+			EqualDeriv = CoerceValueDeriv(EqualDeriv, ResultTypeIndex, EqualDerivInfo.DerivativeStatus, ResultTypeIndex);
+
+			CodeAnalytic = FString::Printf(TEXT("IfDeriv(%s, %s, %s, %s, %s, %s)"), *AFinite, *BFinite, *GreaterDeriv, *LessDeriv, *EqualDeriv, *ThresholdFinite);
+			bIf2Enabled[ResultTypeIndex] = true;
+		}
+		else
+		{
+			CodeAnalytic = FString::Printf(TEXT("IfDeriv(%s, %s, %s, %s)"), *AFinite, *BFinite, *GreaterDeriv, *LessDeriv);
+			bIfEnabled[ResultTypeIndex] = true;
+		}
+
+		return Translator.AddCodeChunkInnerDeriv(*CodeFinite, *CodeAnalytic, ResultType, false, EDerivativeStatus::Valid);
+
+	}
+	else
+	{
+		return Translator.AddCodeChunkInnerDeriv(*CodeFinite, *CodeFinite, ResultType, false, bAllDerivZero ? EDerivativeStatus::Zero : EDerivativeStatus::NotValid);
+	}
+}
+
+FString FMaterialDerivativeAutogen::GenerateUsedFunctions(FHLSLMaterialTranslator& Translator)
+{
+	// Certain derivative functions rely on other derivative functions. For example, Dot() requires Mul() and Add(). So if (for example) dot is enabled, then enable mul1 and add1.
+	EnableGeneratedDepencencies();
+
+	FString Ret;
+
+	// The basic structs (FloatDeriv, FloatDeriv2, FloatDeriv3, FloatDeriv4)
+	// It's not worth keeping track of all the times these are used, just make them.
+	for (int32 Index = 0; Index < 4; Index++)
+	{
+		FString BaseName = GetDerivVectorName(Index);
+		FString FieldName = GetFloatVectorName(Index);
+
+		Ret += TEXT("struct ") + BaseName + LINE_TERMINATOR;
+		Ret += TEXT("{") LINE_TERMINATOR;
+		Ret += TEXT("\t") + FieldName + TEXT(" Value;") LINE_TERMINATOR;
+		Ret += TEXT("\t") + FieldName + TEXT(" Ddx;") LINE_TERMINATOR;
+		Ret += TEXT("\t") + FieldName + TEXT(" Ddy;") LINE_TERMINATOR;
+		Ret += TEXT("};") LINE_TERMINATOR;
+		Ret += TEXT("") LINE_TERMINATOR;
+	}
+
+	// Full FloatDerivX constructors with explicit derivatives.
+	for (int32 Index = 0; Index < 4; Index++)
+	{
+		if (bConstructDerivEnabled[Index] || IsDebugGenerateAllFunctionsEnabled())
+		{
+			FString BaseName = GetDerivVectorName(Index);
+			FString FieldName = GetFloatVectorName(Index);
+
+			Ret += BaseName + TEXT(" Construct") + BaseName + TEXT("(") + FieldName + TEXT(" InValue,") + FieldName + TEXT(" InDdx,") + FieldName + TEXT(" InDdy)") LINE_TERMINATOR;
+			Ret += TEXT("{") LINE_TERMINATOR;
+			Ret += TEXT("\t") + BaseName + TEXT(" Ret;") LINE_TERMINATOR;
+			Ret += TEXT("\tRet.Value = InValue;") LINE_TERMINATOR;
+			Ret += TEXT("\tRet.Ddx = InDdx;") LINE_TERMINATOR;
+			Ret += TEXT("\tRet.Ddy = InDdy;") LINE_TERMINATOR;
+			Ret += TEXT("\treturn Ret;") LINE_TERMINATOR;
+			Ret += TEXT("}") LINE_TERMINATOR;
+			Ret += TEXT("") LINE_TERMINATOR;
+		}
+	}
+
+	// FloatDerivX constructors from constant floatX.
+	for (int32 Index = 0; Index < 4; Index++)
+	{
+		if (bConstructConstantDerivEnabled[Index] || IsDebugGenerateAllFunctionsEnabled())
+		{
+			FString BaseName = GetDerivVectorName(Index);
+			FString FieldName = GetFloatVectorName(Index);
+
+			Ret += BaseName + TEXT(" ConstructConstant") + BaseName + TEXT("(") + FieldName + TEXT(" Value)") LINE_TERMINATOR;
+			Ret += TEXT("{") LINE_TERMINATOR;
+			Ret += TEXT("\t") + BaseName + TEXT(" Ret;") LINE_TERMINATOR;
+			Ret += TEXT("\tRet.Value = Value;") LINE_TERMINATOR;
+			Ret += TEXT("\tRet.Ddx = 0;") LINE_TERMINATOR;
+			Ret += TEXT("\tRet.Ddy = 0;") LINE_TERMINATOR;
+			Ret += TEXT("\treturn Ret;") LINE_TERMINATOR;
+			Ret += TEXT("}") LINE_TERMINATOR;
+			Ret += TEXT("") LINE_TERMINATOR;
+		}
+	}
+
+	// FloatDerivX constructor from floatX with implicit derivatives.
+	for (int32 Index = 0; Index < 4; Index++)
+	{
+		if (bConstructFiniteDerivEnabled[Index] || IsDebugGenerateAllFunctionsEnabled())
+		{
+			FString BaseName = GetDerivVectorName(Index);
+			FString FieldName = GetFloatVectorName(Index);
+
+			Ret += BaseName + TEXT(" ConstructFinite") + BaseName + TEXT("(") + FieldName + TEXT(" InValue)") LINE_TERMINATOR;
+			Ret += TEXT("{") LINE_TERMINATOR;
+			Ret += TEXT("\t") + BaseName + TEXT(" Ret;") LINE_TERMINATOR;
+			Ret += TEXT("\tRet.Value = InValue;") LINE_TERMINATOR;
+			Ret += TEXT("\tRet.Ddx = ddx(InValue);") LINE_TERMINATOR;
+			Ret += TEXT("\tRet.Ddy = ddy(InValue);") LINE_TERMINATOR;
+			Ret += TEXT("\treturn Ret;") LINE_TERMINATOR;
+			Ret += TEXT("}") LINE_TERMINATOR;
+			Ret += TEXT("") LINE_TERMINATOR;
+		}
+	}
+
+	// Convert between FloatDeriv types
+	for (int32 DstIndex = 0; DstIndex < 4; DstIndex++)
+	{
+		for (int32 SrcIndex = 0; SrcIndex < 4; SrcIndex++)
+		{
+			if (SrcIndex == DstIndex)
+				continue;
+
+			if (bConvertDerivEnabled[DstIndex][SrcIndex] || IsDebugGenerateAllFunctionsEnabled())
+			{
+				FString DstBaseName = GetDerivVectorName(DstIndex);
+				FString SrcBaseName = GetDerivVectorName(SrcIndex);
+				FString ScalarName = GetDerivVectorName(0);
+
+				Ret += DstBaseName + TEXT(" Convert") + DstBaseName + TEXT("(") + SrcBaseName + TEXT(" Src)") LINE_TERMINATOR;
+				Ret += TEXT("{") LINE_TERMINATOR;
+				Ret += TEXT("\t") + DstBaseName + TEXT(" Ret;") LINE_TERMINATOR;
+				Ret += TEXT("\tRet.Value = ") + CoerceFloat(TEXT("Src.Value"), DstIndex, SrcIndex) + TEXT(";") LINE_TERMINATOR;
+				Ret += TEXT("\tRet.Ddx = ") + CoerceFloat(TEXT("Src.Ddx"), DstIndex, SrcIndex) + TEXT(";") LINE_TERMINATOR;
+				Ret += TEXT("\tRet.Ddy = ") + CoerceFloat(TEXT("Src.Ddy"), DstIndex, SrcIndex) + TEXT(";") LINE_TERMINATOR;
+				Ret += TEXT("\treturn Ret;") LINE_TERMINATOR;
+				Ret += TEXT("}") LINE_TERMINATOR;
+				Ret += TEXT("") LINE_TERMINATOR;
+			}
+		}
+	}
+
+	const TCHAR* SwizzleList[4] =
+	{
+		TEXT("x"),
+		TEXT("y"),
+		TEXT("z"),
+		TEXT("w")
+	};
+
+	// Extract single FloatDeriv element from FloatDerivX
+	for (int32 StructIndex = 0; StructIndex < 4; StructIndex++)
+	{
+		for (int32 ElemIndex = 0; ElemIndex <= StructIndex; ElemIndex++)
+		{
+			if (bExtractIndexEnabled[StructIndex] || IsDebugGenerateAllFunctionsEnabled())
+			{
+				// i.e. can't grab the 4th component of a float3
+				check(ElemIndex <= StructIndex);
+
+				FString BaseName = GetDerivVectorName(StructIndex);
+
+				FString ElemStr = FString::Printf(TEXT("%d"), ElemIndex + 1);
+
+				Ret += TEXT("FloatDeriv Extract") + BaseName + TEXT("_") + ElemStr + TEXT("(") + BaseName + TEXT(" InValue)") LINE_TERMINATOR;
+				Ret += TEXT("{") LINE_TERMINATOR;
+				Ret += TEXT("\tFloatDeriv Ret;") LINE_TERMINATOR;
+				Ret += TEXT("\tRet.Value = InValue.Value.") + FString(SwizzleList[ElemIndex]) + TEXT(";") LINE_TERMINATOR;
+				Ret += TEXT("\tRet.Ddx = InValue.Ddx.") + FString(SwizzleList[ElemIndex]) + TEXT(";") LINE_TERMINATOR;
+				Ret += TEXT("\tRet.Ddy = InValue.Ddy.") + FString(SwizzleList[ElemIndex]) + TEXT(";") LINE_TERMINATOR;
+				Ret += TEXT("\treturn Ret;") LINE_TERMINATOR;
+				Ret += TEXT("}") LINE_TERMINATOR;
+				Ret += TEXT("") LINE_TERMINATOR;
+			}
+		}
+	}
+
+
+	// SelectElemHelper
+	for (int32 Index = 0; Index < 4; Index++)
+	{
+		if (bSelectElemHelperEnabled[Index] || IsDebugGenerateAllFunctionsEnabled())
+		{
+			FString FieldName = GetFloatVectorName(Index);
+			FString BoolName = GetBoolVectorName(Index);
+
+			Ret += FieldName + TEXT(" SelectElemHelper(") + BoolName + TEXT(" Cmp, ") + FieldName + TEXT(" Lhs, ") + FieldName + TEXT(" Rhs)") LINE_TERMINATOR;
+			Ret += TEXT("{") LINE_TERMINATOR;
+			Ret += TEXT("\t") + FieldName + TEXT(" Ret = 0.0f;") LINE_TERMINATOR;
+			for (int32 Component = 0; Component <= Index; Component++)
+			{
+				const TCHAR* ChanName = SwizzleList[Component];
+				Ret += FString::Printf(TEXT("\tRet.%s = Cmp.%s ? Lhs.%s : Rhs.%s;"),ChanName,ChanName,ChanName,ChanName) + LINE_TERMINATOR;
+			}
+
+			Ret += TEXT("\treturn Ret;") LINE_TERMINATOR;
+			Ret += TEXT("}") LINE_TERMINATOR;
+			Ret += TEXT("") LINE_TERMINATOR;
+		}
+	}
+
+	// Func2s
+	for (int32 Op = 0; Op < (int32)EFunc2::Num; Op++)
+	{
+		for (int32 Index = 0; Index < 4; Index++)
+		{
+			if (bFunc2OpIsEnabled[Op][Index] || IsDebugGenerateAllFunctionsEnabled())
+			{
+				FString BaseName = GetDerivVectorName(Index);
+				FString FieldName = GetFloatVectorName(Index);
+				FString BoolName = GetBoolVectorName(Index);
+
+				switch((EFunc2)Op)
+				{
+				case EFunc2::Add:
+					Ret += BaseName + TEXT(" AddDeriv(") + BaseName + TEXT(" A, ") + BaseName + TEXT(" B)") LINE_TERMINATOR;
+					Ret += TEXT("{") LINE_TERMINATOR;
+					Ret += TEXT("\t") + BaseName + TEXT(" Ret;") LINE_TERMINATOR;
+					Ret += TEXT("\tRet.Value = A.Value + B.Value;") LINE_TERMINATOR;
+					Ret += TEXT("\tRet.Ddx = A.Ddx + B.Ddx;") LINE_TERMINATOR;
+					Ret += TEXT("\tRet.Ddy = A.Ddy + B.Ddy;") LINE_TERMINATOR;
+					Ret += TEXT("\treturn Ret;") LINE_TERMINATOR;
+					Ret += TEXT("}") LINE_TERMINATOR;
+					Ret += TEXT("") LINE_TERMINATOR;
+					break;
+				case EFunc2::Sub:
+					Ret += BaseName + TEXT(" SubDeriv(") + BaseName + TEXT(" A, ") + BaseName + TEXT(" B)") LINE_TERMINATOR;
+					Ret += TEXT("{") LINE_TERMINATOR;
+					Ret += TEXT("\t") + BaseName + TEXT(" Ret;") LINE_TERMINATOR;
+					Ret += TEXT("\tRet.Value = A.Value - B.Value;") LINE_TERMINATOR;
+					Ret += TEXT("\tRet.Ddx = A.Ddx - B.Ddx;") LINE_TERMINATOR;
+					Ret += TEXT("\tRet.Ddy = A.Ddy - B.Ddy;") LINE_TERMINATOR;
+					Ret += TEXT("\treturn Ret;") LINE_TERMINATOR;
+					Ret += TEXT("}") LINE_TERMINATOR;
+					Ret += TEXT("") LINE_TERMINATOR;
+					break;
+				case EFunc2::Mul:
+					Ret += BaseName + TEXT(" MulDeriv(") + BaseName + TEXT(" A, ") + BaseName + TEXT(" B)") LINE_TERMINATOR;
+					Ret += TEXT("{") LINE_TERMINATOR;
+					Ret += TEXT("\t") + BaseName + TEXT(" Ret;") LINE_TERMINATOR;
+					Ret += TEXT("\tRet.Value = A.Value * B.Value;") LINE_TERMINATOR;
+					Ret += TEXT("\tRet.Ddx = A.Ddx * B.Value + A.Value * B.Ddx;") LINE_TERMINATOR;
+					Ret += TEXT("\tRet.Ddy = A.Ddy * B.Value + A.Value * B.Ddy;") LINE_TERMINATOR;
+					Ret += TEXT("\treturn Ret;") LINE_TERMINATOR;
+					Ret += TEXT("}") LINE_TERMINATOR;
+					Ret += TEXT("") LINE_TERMINATOR;
+					break;
+				case EFunc2::Div:
+					Ret += BaseName + TEXT(" DivDeriv(") + BaseName + TEXT(" A, ") + BaseName + TEXT(" B)") LINE_TERMINATOR;
+					Ret += TEXT("{") LINE_TERMINATOR;
+					Ret += TEXT("\t") + BaseName + TEXT(" Ret;") LINE_TERMINATOR;
+					Ret += TEXT("\tRet.Value = A.Value / B.Value;") LINE_TERMINATOR;
+					Ret += TEXT("\t") + FieldName + TEXT(" Denom = rcp(B.Value * B.Value);") LINE_TERMINATOR;
+					Ret += TEXT("\t") + FieldName + TEXT(" dFdA =  B.Value * Denom;") LINE_TERMINATOR;
+					Ret += TEXT("\t") + FieldName + TEXT(" dFdB = -A.Value * Denom;") LINE_TERMINATOR;
+					Ret += TEXT("\tRet.Ddx = dFdA * A.Ddx + dFdB * B.Ddx;") LINE_TERMINATOR;
+					Ret += TEXT("\tRet.Ddy = dFdA * A.Ddy + dFdB * B.Ddy;") LINE_TERMINATOR;
+					Ret += TEXT("\treturn Ret;") LINE_TERMINATOR;
+					Ret += TEXT("}") LINE_TERMINATOR;
+					Ret += TEXT("") LINE_TERMINATOR;
+					break;
+				case EFunc2::Fmod:
+					// Only valid when B derivatives are zero.
+					// We can't really do anything meaningful in the non-zero case.
+					Ret += BaseName + TEXT(" FmodDeriv(") + BaseName + TEXT(" A, ") + BaseName + TEXT(" B)") LINE_TERMINATOR;
+					Ret += TEXT("{") LINE_TERMINATOR;
+					Ret += TEXT("\t") + BaseName + TEXT(" Ret;") LINE_TERMINATOR;
+					Ret += TEXT("\tRet.Value = fmod(A.Value, B.Value);") LINE_TERMINATOR;
+					Ret += TEXT("\tRet.Ddx = A.Ddx;") LINE_TERMINATOR;
+					Ret += TEXT("\tRet.Ddy = A.Ddy;") LINE_TERMINATOR;
+					Ret += TEXT("\treturn Ret;") LINE_TERMINATOR;
+					Ret += TEXT("}") LINE_TERMINATOR;
+					Ret += TEXT("") LINE_TERMINATOR;
+					break;
+				case EFunc2::Min:
+					Ret += BaseName + TEXT(" MinDeriv(") + BaseName + TEXT(" A, ") + BaseName + TEXT(" B)") LINE_TERMINATOR;
+					Ret += TEXT("{") LINE_TERMINATOR;
+					Ret += TEXT("\t") + BaseName + TEXT(" Ret;") LINE_TERMINATOR;
+					Ret += TEXT("\t") + BoolName + TEXT(" Cmp = A.Value < B.Value;") LINE_TERMINATOR;
+					Ret += TEXT("\tRet.Value = SelectElemHelper(Cmp, A.Value, B.Value);") LINE_TERMINATOR;
+					Ret += TEXT("\tRet.Ddx = SelectElemHelper(Cmp, A.Ddx, B.Ddx);") LINE_TERMINATOR;
+					Ret += TEXT("\tRet.Ddy = SelectElemHelper(Cmp, A.Ddy, B.Ddy);") LINE_TERMINATOR;
+					Ret += TEXT("\treturn Ret;") LINE_TERMINATOR;
+					Ret += TEXT("}") LINE_TERMINATOR;
+					Ret += TEXT("") LINE_TERMINATOR;
+					break;
+				case EFunc2::Max:
+					Ret += BaseName + TEXT(" MaxDeriv(") + BaseName + TEXT(" A, ") + BaseName + TEXT(" B)") LINE_TERMINATOR;
+					Ret += TEXT("{") LINE_TERMINATOR;
+					Ret += TEXT("\t") + BaseName + TEXT(" Ret;") LINE_TERMINATOR;
+					Ret += TEXT("\t") + BoolName + TEXT(" Cmp = A.Value > B.Value;") LINE_TERMINATOR;
+					Ret += TEXT("\tRet.Value = SelectElemHelper(Cmp, A.Value, B.Value);") LINE_TERMINATOR;
+					Ret += TEXT("\tRet.Ddx = SelectElemHelper(Cmp, A.Ddx, B.Ddx);") LINE_TERMINATOR;
+					Ret += TEXT("\tRet.Ddy = SelectElemHelper(Cmp, A.Ddy, B.Ddy);") LINE_TERMINATOR;
+					Ret += TEXT("\treturn Ret;") LINE_TERMINATOR;
+					Ret += TEXT("}") LINE_TERMINATOR;
+					Ret += TEXT("") LINE_TERMINATOR;
+					break;
+				case EFunc2::Dot:
+					Ret += TEXT("FloatDeriv DotDeriv(") + BaseName + TEXT(" A, ") + BaseName + TEXT(" B)") LINE_TERMINATOR;
+					Ret += TEXT("{") LINE_TERMINATOR;
+					Ret += TEXT("\tFloatDeriv Ret = ConstructConstantFloatDeriv(0);") LINE_TERMINATOR;
+					for (int32 Component = 0; Component <= Index; Component++)
+					{
+						Ret += FString::Printf(TEXT("\tRet = AddDeriv(Ret,MulDeriv(Extract%s_%d(A),Extract%s_%d(B)));"), *BaseName, Component + 1, *BaseName, Component + 1) + LINE_TERMINATOR;
+					}
+					Ret += TEXT("\treturn Ret;") LINE_TERMINATOR;
+					Ret += TEXT("}") LINE_TERMINATOR;
+					Ret += TEXT("") LINE_TERMINATOR;
+					break;
+				case EFunc2::Pow:
+					// pow(A,B) = exp(B*log(A))
+					//     pow'(A,B) = exp(B*log(A)) * (B'*log(A) + (B/A)*A')
+					//     pow'(A,B) = pow(A,B) * (B'*log(A) + (B/A)*A')
+					// sanity check when B is constant and A is a linear function (B'=0,A'=1)
+					//     pow'(A,B) = pow(A,B) * (0*log(A) + (B/A)*1)
+					//     pow'(A,B) = B * pow(A,B-1)
+					Ret += BaseName + TEXT(" PowDeriv(") + BaseName + TEXT(" A, ") + BaseName + TEXT(" B)") LINE_TERMINATOR;
+					Ret += TEXT("{") LINE_TERMINATOR;
+					Ret += TEXT("\t") + BaseName + TEXT(" Ret;") LINE_TERMINATOR;
+					Ret += TEXT("\tRet.Value = pow(A.Value, B.Value);") LINE_TERMINATOR;
+					Ret += TEXT("\tRet.Ddx = Ret.Value * (B.Ddx * log(A.Value) + (B.Value/A.Value)*A.Ddx);") LINE_TERMINATOR;
+					Ret += TEXT("\tRet.Ddy = Ret.Value * (B.Ddy * log(A.Value) + (B.Value/A.Value)*A.Ddy);") LINE_TERMINATOR;
+					Ret += TEXT("\treturn Ret;") LINE_TERMINATOR;
+					Ret += TEXT("}") LINE_TERMINATOR;
+					Ret += TEXT("") LINE_TERMINATOR;
+					break;
+				case EFunc2::PowPositiveClamped:
+					Ret += BaseName + TEXT(" PowPositiveClampedDeriv(") + BaseName + TEXT(" A, ") + BaseName + TEXT(" B)") LINE_TERMINATOR;
+					Ret += TEXT("{") LINE_TERMINATOR;
+					Ret += TEXT("\t") + BaseName + TEXT(" Ret;") LINE_TERMINATOR;
+					Ret += TEXT("\t") + BoolName + TEXT(" InRange = (0.0 < B.Value);") LINE_TERMINATOR; // should we check for A as well?
+					Ret += TEXT("\t") + FieldName + TEXT(" Zero = 0.0;") LINE_TERMINATOR;
+					Ret += TEXT("\tRet.Value = PositiveClampedPow(A.Value, B.Value);") LINE_TERMINATOR;
+					Ret += TEXT("\tRet.Ddx = Ret.Value * (B.Ddx * log(A.Value) + (B.Value/A.Value)*A.Ddx);") LINE_TERMINATOR;
+					Ret += TEXT("\tRet.Ddy = Ret.Value * (B.Ddy * log(A.Value) + (B.Value/A.Value)*A.Ddy);") LINE_TERMINATOR;
+					Ret += TEXT("\tRet.Ddx = SelectElemHelper(InRange,Ret.Ddx,Zero);") LINE_TERMINATOR;
+					Ret += TEXT("\tRet.Ddy = SelectElemHelper(InRange,Ret.Ddy,Zero);") LINE_TERMINATOR;
+					Ret += TEXT("\treturn Ret;") LINE_TERMINATOR;
+					Ret += TEXT("}") LINE_TERMINATOR;
+					Ret += TEXT("") LINE_TERMINATOR;
+					break;
+				case EFunc2::Atan2:
+					Ret += BaseName + TEXT(" Atan2Deriv(") + BaseName + TEXT(" A, ") + BaseName + TEXT(" B)") LINE_TERMINATOR;
+					Ret += TEXT("{") LINE_TERMINATOR;
+					Ret += TEXT("\t") + BaseName + TEXT(" Ret;") LINE_TERMINATOR;
+					Ret += TEXT("\tRet.Value = atan2(A.Value, B.Value);") LINE_TERMINATOR;
+					Ret += TEXT("\t") + FieldName + TEXT(" Denom = rcp(A.Value * A.Value + B.Value * B.Value);") LINE_TERMINATOR;
+					Ret += TEXT("\t") + FieldName + TEXT(" dFdA =  B.Value * Denom;") LINE_TERMINATOR;
+					Ret += TEXT("\t") + FieldName + TEXT(" dFdB = -A.Value * Denom;") LINE_TERMINATOR;
+					Ret += TEXT("\tRet.Ddx = dFdA * A.Ddx + dFdB * B.Ddx;") LINE_TERMINATOR;
+					Ret += TEXT("\tRet.Ddy = dFdA * A.Ddy + dFdB * B.Ddy;") LINE_TERMINATOR;
+					Ret += TEXT("\treturn Ret;") LINE_TERMINATOR;
+					Ret += TEXT("}") LINE_TERMINATOR;
+					Ret += TEXT("") LINE_TERMINATOR;
+					break;
+				case EFunc2::Atan2Fast:
+					Ret += BaseName + TEXT(" Atan2FastDeriv(") + BaseName + TEXT(" A, ") + BaseName + TEXT(" B)") LINE_TERMINATOR;
+					Ret += TEXT("{") LINE_TERMINATOR;
+					Ret += TEXT("\t") + BaseName + TEXT(" Ret;") LINE_TERMINATOR;
+					Ret += TEXT("\tRet.Value = atan2Fast(A.Value, B.Value);") LINE_TERMINATOR;
+					Ret += TEXT("\t") + FieldName + TEXT(" Denom = rcp(A.Value * A.Value + B.Value * B.Value);") LINE_TERMINATOR;
+					Ret += TEXT("\t") + FieldName + TEXT(" dFdA =  B.Value * Denom;") LINE_TERMINATOR;
+					Ret += TEXT("\t") + FieldName + TEXT(" dFdB = -A.Value * Denom;") LINE_TERMINATOR;
+					Ret += TEXT("\tRet.Ddx = dFdA * A.Ddx + dFdB * B.Ddx;") LINE_TERMINATOR;
+					Ret += TEXT("\tRet.Ddy = dFdA * A.Ddy + dFdB * B.Ddy;") LINE_TERMINATOR;
+					Ret += TEXT("\treturn Ret;") LINE_TERMINATOR;
+					Ret += TEXT("}") LINE_TERMINATOR;
+					Ret += TEXT("") LINE_TERMINATOR;
+					break;
+				case EFunc2::Cross:
+					if (Index == 2)
+					{
+						// (A*B)' = A' * B + A * B'
+						// Cross(A, B) = A.yzx * B.zxy - A.zxy * B.yzx;
+						// Cross(A, B)' = A.yzx' * B.zxy + A.yzx * B.zxy' - A.zxy' * B.yzx - A.zxy * B.yzx';
+						Ret += BaseName + TEXT(" CrossDeriv(") + BaseName + TEXT(" A, ") + BaseName + TEXT(" B)") LINE_TERMINATOR;
+						Ret += TEXT("{") LINE_TERMINATOR;
+						Ret += TEXT("\t") + BaseName + TEXT(" Ret;") LINE_TERMINATOR;
+						Ret += TEXT("\tRet.Value = cross(A.Value, B.Value);") LINE_TERMINATOR;
+						Ret += TEXT("\tRet.Ddx = A.Ddx.yzx * B.Value.zxy + A.Value.yzx * B.Ddx.zxy - A.Ddx.zxy * B.Value.yzx - A.Value.zxy * B.Ddx.yzx;") LINE_TERMINATOR;
+						Ret += TEXT("\tRet.Ddy = A.Ddy.yzx * B.Value.zxy + A.Value.yzx * B.Ddy.zxy - A.Ddy.zxy * B.Value.yzx - A.Value.zxy * B.Ddy.yzx;") LINE_TERMINATOR;
+						Ret += TEXT("\treturn Ret;") LINE_TERMINATOR;
+						Ret += TEXT("}") LINE_TERMINATOR;
+						Ret += TEXT("") LINE_TERMINATOR;
+					}
+					break;
+				default:
+					check(0);
+					break;
+				}
+			}
+		}
+	}
+
+	for (int32 Op = 0; Op < (int32)EFunc1::Num; Op++)
+	{
+		for (int32 Index = 0; Index < 4; Index++)
+		{
+			if (bFunc1OpIsEnabled[Op][Index] || IsDebugGenerateAllFunctionsEnabled())
+			{
+				FString BaseName	= GetDerivVectorName(Index);
+				FString FieldName	= GetFloatVectorName(Index);
+				FString BoolName	= GetBoolVectorName(Index);
+
+				switch((EFunc1)Op)
+				{
+				case EFunc1::Abs:
+					Ret += BaseName + TEXT(" AbsDeriv(") + BaseName + TEXT(" A)") LINE_TERMINATOR;
+					Ret += TEXT("{") LINE_TERMINATOR;
+					Ret += TEXT("\t") + BaseName + TEXT(" Ret;") LINE_TERMINATOR;
+					Ret += TEXT("\tRet.Value = abs(A.Value);") LINE_TERMINATOR;
+					Ret += TEXT("\t") + FieldName + TEXT(" dFdA = (A.Value >= 0.0f ? 1.0f : -1.0f);") LINE_TERMINATOR;
+					Ret += TEXT("\tRet.Ddx = dFdA * A.Ddx;") LINE_TERMINATOR;
+					Ret += TEXT("\tRet.Ddy = dFdA * A.Ddy;") LINE_TERMINATOR;
+					Ret += TEXT("\treturn Ret;") LINE_TERMINATOR;
+					Ret += TEXT("}") LINE_TERMINATOR;
+					Ret += TEXT("") LINE_TERMINATOR;
+					break;
+				case EFunc1::Sin:
+					Ret += BaseName + TEXT(" SinDeriv(") + BaseName + TEXT(" A)") LINE_TERMINATOR;
+					Ret += TEXT("{") LINE_TERMINATOR;
+					Ret += TEXT("\t") + BaseName + TEXT(" Ret;") LINE_TERMINATOR;
+					Ret += TEXT("\tRet.Value = sin(A.Value);") LINE_TERMINATOR;
+					Ret += TEXT("\t") + FieldName + TEXT(" dFdA = cos(A.Value);");
+					Ret += TEXT("\tRet.Ddx = dFdA * A.Ddx;") LINE_TERMINATOR;
+					Ret += TEXT("\tRet.Ddy = dFdA * A.Ddy;") LINE_TERMINATOR;
+					Ret += TEXT("\treturn Ret;") LINE_TERMINATOR;
+					Ret += TEXT("}") LINE_TERMINATOR;
+					Ret += TEXT("") LINE_TERMINATOR;
+					break;
+				case EFunc1::Cos:
+					Ret += BaseName + TEXT(" CosDeriv(") + BaseName + TEXT(" A)") LINE_TERMINATOR;
+					Ret += TEXT("{") LINE_TERMINATOR;
+					Ret += TEXT("\t") + BaseName + TEXT(" Ret;") LINE_TERMINATOR;
+					Ret += TEXT("\tRet.Value = cos(A.Value);") LINE_TERMINATOR;
+					Ret += TEXT("\t") + FieldName + TEXT(" dFdA = -sin(A.Value);");
+					Ret += TEXT("\tRet.Ddx = dFdA * A.Ddx;") LINE_TERMINATOR;
+					Ret += TEXT("\tRet.Ddy = dFdA * A.Ddy;") LINE_TERMINATOR;
+					Ret += TEXT("\treturn Ret;") LINE_TERMINATOR;
+					Ret += TEXT("}") LINE_TERMINATOR;
+					Ret += TEXT("") LINE_TERMINATOR;
+					break;
+				case EFunc1::Tan:
+					Ret += BaseName + TEXT(" TanDeriv(") + BaseName + TEXT(" A)") LINE_TERMINATOR;
+					Ret += TEXT("{") LINE_TERMINATOR;
+					Ret += TEXT("\t") + BaseName + TEXT(" Ret;") LINE_TERMINATOR;
+					Ret += TEXT("\tRet.Value = tan(A.Value);") LINE_TERMINATOR;
+					Ret += TEXT("\t") + FieldName + TEXT(" dFdA = rcp(cos(A.Value) * cos(A.Value));");
+					Ret += TEXT("\tRet.Ddx = dFdA * A.Ddx;") LINE_TERMINATOR;
+					Ret += TEXT("\tRet.Ddy = dFdA * A.Ddy;") LINE_TERMINATOR;
+					Ret += TEXT("\treturn Ret;") LINE_TERMINATOR;
+					Ret += TEXT("}") LINE_TERMINATOR;
+					Ret += TEXT("") LINE_TERMINATOR;
+					break;
+				case EFunc1::Asin:
+					Ret += BaseName + TEXT(" AsinDeriv(") + BaseName + TEXT(" A)") LINE_TERMINATOR;
+					Ret += TEXT("{") LINE_TERMINATOR;
+					Ret += TEXT("\t") + BaseName + TEXT(" Ret;") LINE_TERMINATOR;
+					Ret += TEXT("\tRet.Value = asin(A.Value);") LINE_TERMINATOR;
+					Ret += TEXT("\t") + FieldName + TEXT(" dFdA = rsqrt(max(1.0f - A.Value * A.Value, 0.00001f));");
+					Ret += TEXT("\tRet.Ddx = dFdA * A.Ddx;") LINE_TERMINATOR;
+					Ret += TEXT("\tRet.Ddy = dFdA * A.Ddy;") LINE_TERMINATOR;
+					Ret += TEXT("\treturn Ret;") LINE_TERMINATOR;
+					Ret += TEXT("}") LINE_TERMINATOR;
+					Ret += TEXT("") LINE_TERMINATOR;
+					break;
+				case EFunc1::AsinFast:
+					Ret += BaseName + TEXT(" AsinFastDeriv(") + BaseName + TEXT(" A)") LINE_TERMINATOR;
+					Ret += TEXT("{") LINE_TERMINATOR;
+					Ret += TEXT("\t") + BaseName + TEXT(" Ret;") LINE_TERMINATOR;
+					Ret += TEXT("\tRet.Value = asinFast(A.Value);") LINE_TERMINATOR;
+					Ret += TEXT("\t") + FieldName + TEXT(" dFdA = rsqrt(max(1.0f - A.Value * A.Value, 0.00001f));");
+					Ret += TEXT("\tRet.Ddx = dFdA * A.Ddx;") LINE_TERMINATOR;
+					Ret += TEXT("\tRet.Ddy = dFdA * A.Ddy;") LINE_TERMINATOR;
+					Ret += TEXT("\treturn Ret;") LINE_TERMINATOR;
+					Ret += TEXT("}") LINE_TERMINATOR;
+					Ret += TEXT("") LINE_TERMINATOR;
+					break;
+				case EFunc1::Acos:
+					Ret += BaseName + TEXT(" AcosDeriv(") + BaseName + TEXT(" A)") LINE_TERMINATOR;
+					Ret += TEXT("{") LINE_TERMINATOR;
+					Ret += TEXT("\t") + BaseName + TEXT(" Ret;") LINE_TERMINATOR;
+					Ret += TEXT("\tRet.Value = acos(A.Value);") LINE_TERMINATOR;
+					Ret += TEXT("\t") + FieldName + TEXT(" dFdA = -rsqrt(max(1.0f - A.Value * A.Value, 0.00001f));");
+					Ret += TEXT("\tRet.Ddx = dFdA * A.Ddx;") LINE_TERMINATOR;
+					Ret += TEXT("\tRet.Ddy = dFdA * A.Ddy;") LINE_TERMINATOR;
+					Ret += TEXT("\treturn Ret;") LINE_TERMINATOR;
+					Ret += TEXT("}") LINE_TERMINATOR;
+					Ret += TEXT("") LINE_TERMINATOR;
+					break;
+				case EFunc1::AcosFast:
+					Ret += BaseName + TEXT(" AcosFastDeriv(") + BaseName + TEXT(" A)") LINE_TERMINATOR;
+					Ret += TEXT("{") LINE_TERMINATOR;
+					Ret += TEXT("\t") + BaseName + TEXT(" Ret;") LINE_TERMINATOR;
+					Ret += TEXT("\tRet.Value = acosFast(A.Value);") LINE_TERMINATOR;
+					Ret += TEXT("\t") + FieldName + TEXT(" dFdA = -rsqrt(max(1.0f - A.Value * A.Value, 0.00001f));");
+					Ret += TEXT("\tRet.Ddx = dFdA * A.Ddx;") LINE_TERMINATOR;
+					Ret += TEXT("\tRet.Ddy = dFdA * A.Ddy;") LINE_TERMINATOR;
+					Ret += TEXT("\treturn Ret;") LINE_TERMINATOR;
+					Ret += TEXT("}") LINE_TERMINATOR;
+					Ret += TEXT("") LINE_TERMINATOR;
+					break;
+				case EFunc1::Atan:
+					Ret += BaseName + TEXT(" AtanDeriv(") + BaseName + TEXT(" A)") LINE_TERMINATOR;
+					Ret += TEXT("{") LINE_TERMINATOR;
+					Ret += TEXT("\t") + BaseName + TEXT(" Ret;") LINE_TERMINATOR;
+					Ret += TEXT("\tRet.Value = atan(A.Value);") LINE_TERMINATOR;
+					Ret += TEXT("\t") + FieldName + TEXT(" dFdA = rcp(A.Value * A.Value + 1.0f);");
+					Ret += TEXT("\tRet.Ddx = dFdA * A.Ddx;") LINE_TERMINATOR;
+					Ret += TEXT("\tRet.Ddy = dFdA * A.Ddy;") LINE_TERMINATOR;
+					Ret += TEXT("\treturn Ret;") LINE_TERMINATOR;
+					Ret += TEXT("}") LINE_TERMINATOR;
+					Ret += TEXT("") LINE_TERMINATOR;
+					break;
+				case EFunc1::AtanFast:
+					Ret += BaseName + TEXT(" AtanFastDeriv(") + BaseName + TEXT(" A)") LINE_TERMINATOR;
+					Ret += TEXT("{") LINE_TERMINATOR;
+					Ret += TEXT("\t") + BaseName + TEXT(" Ret;") LINE_TERMINATOR;
+					Ret += TEXT("\tRet.Value = atanFast(A.Value);") LINE_TERMINATOR;
+					Ret += TEXT("\t") + FieldName + TEXT(" dFdA = rcp(A.Value * A.Value + 1.0f);");
+					Ret += TEXT("\tRet.Ddx = dFdA * A.Ddx;") LINE_TERMINATOR;
+					Ret += TEXT("\tRet.Ddy = dFdA * A.Ddy;") LINE_TERMINATOR;
+					Ret += TEXT("\treturn Ret;") LINE_TERMINATOR;
+					Ret += TEXT("}") LINE_TERMINATOR;
+					Ret += TEXT("") LINE_TERMINATOR;
+					break;
+				case EFunc1::Sqrt:
+					Ret += BaseName + TEXT(" SqrtDeriv(") + BaseName + TEXT(" A)") LINE_TERMINATOR;
+					Ret += TEXT("{") LINE_TERMINATOR;
+					Ret += TEXT("\t") + BaseName + TEXT(" Ret;") LINE_TERMINATOR;
+					Ret += TEXT("\tRet.Value = sqrt(A.Value);") LINE_TERMINATOR;
+					Ret += TEXT("\t") + FieldName + TEXT(" dFdA = 0.5f * rsqrt(max(A.Value, 0.00001f));");
+					Ret += TEXT("\tRet.Ddx = dFdA * A.Ddx;") LINE_TERMINATOR;
+					Ret += TEXT("\tRet.Ddy = dFdA * A.Ddy;") LINE_TERMINATOR;
+					Ret += TEXT("\treturn Ret;") LINE_TERMINATOR;
+					Ret += TEXT("}") LINE_TERMINATOR;
+					Ret += TEXT("") LINE_TERMINATOR;
+					break;
+				case EFunc1::Rcp:
+					Ret += BaseName + TEXT(" RcpDeriv(") + BaseName + TEXT(" A)") LINE_TERMINATOR;
+					Ret += TEXT("{") LINE_TERMINATOR;
+					Ret += TEXT("\t") + BaseName + TEXT(" Ret;") LINE_TERMINATOR;
+					Ret += TEXT("\tRet.Value = rcp(A.Value);") LINE_TERMINATOR;
+					Ret += TEXT("\t") + FieldName + TEXT(" dFdA = -Ret.Value * Ret.Value;");
+					Ret += TEXT("\tRet.Ddx = dFdA * A.Ddx;") LINE_TERMINATOR;
+					Ret += TEXT("\tRet.Ddy = dFdA * A.Ddy;") LINE_TERMINATOR;
+					Ret += TEXT("\treturn Ret;") LINE_TERMINATOR;
+					Ret += TEXT("}") LINE_TERMINATOR;
+					Ret += TEXT("") LINE_TERMINATOR;
+					break;
+				case EFunc1::Rsqrt:
+					Ret += BaseName + TEXT(" RsqrtDeriv(") + BaseName + TEXT(" A)") LINE_TERMINATOR;
+					Ret += TEXT("{") LINE_TERMINATOR;
+					Ret += TEXT("\t") + BaseName + TEXT(" Ret;") LINE_TERMINATOR;
+					Ret += TEXT("\tRet.Value = rsqrt(A.Value);") LINE_TERMINATOR;
+					Ret += TEXT("\t") + FieldName + TEXT(" dFdA = -0.5f * rsqrt(A.Value) * rcp(A.Value);");
+					Ret += TEXT("\tRet.Ddx = dFdA * A.Ddx;") LINE_TERMINATOR;
+					Ret += TEXT("\tRet.Ddy = dFdA * A.Ddy;") LINE_TERMINATOR;
+					Ret += TEXT("\treturn Ret;") LINE_TERMINATOR;
+					Ret += TEXT("}") LINE_TERMINATOR;
+					Ret += TEXT("") LINE_TERMINATOR;
+					break;
+				case EFunc1::Saturate:
+					Ret += BaseName + TEXT(" SaturateDeriv(") + BaseName + TEXT(" A)") LINE_TERMINATOR;
+					Ret += TEXT("{") LINE_TERMINATOR;
+					Ret += TEXT("\t") + BoolName + TEXT(" InRange = (0.0 < A.Value && A.Value < 1.0);") LINE_TERMINATOR;
+					Ret += TEXT("\t") + FieldName + TEXT(" Zero = 0.0f;") LINE_TERMINATOR;
+					Ret += TEXT("\t") + BaseName + TEXT(" Ret;") LINE_TERMINATOR;
+					Ret += TEXT("\tRet.Value = saturate(A.Value);") LINE_TERMINATOR;
+					Ret += TEXT("\tRet.Ddx = SelectElemHelper(InRange, A.Ddx, Zero);") LINE_TERMINATOR;
+					Ret += TEXT("\tRet.Ddy = SelectElemHelper(InRange, A.Ddy, Zero);") LINE_TERMINATOR;
+					Ret += TEXT("\treturn Ret;") LINE_TERMINATOR;
+					Ret += TEXT("}") LINE_TERMINATOR;
+					Ret += TEXT("") LINE_TERMINATOR;
+					break;
+				case EFunc1::Frac:
+					Ret += BaseName + TEXT(" FracDeriv(") + BaseName + TEXT(" A)") LINE_TERMINATOR;
+					Ret += TEXT("{") LINE_TERMINATOR;
+					Ret += TEXT("\t") + BaseName + TEXT(" Ret;") LINE_TERMINATOR;
+					Ret += TEXT("\tRet.Value = frac(A.Value);") LINE_TERMINATOR;
+					Ret += TEXT("\tRet.Ddx = A.Ddx;") LINE_TERMINATOR;
+					Ret += TEXT("\tRet.Ddy = A.Ddy;") LINE_TERMINATOR;
+					Ret += TEXT("\treturn Ret;") LINE_TERMINATOR;
+					Ret += TEXT("}") LINE_TERMINATOR;
+					Ret += TEXT("") LINE_TERMINATOR;
+					break;
+				case EFunc1::Log2:
+					Ret += BaseName + TEXT(" Log2Deriv(") + BaseName + TEXT(" A)") LINE_TERMINATOR;
+					Ret += TEXT("{") LINE_TERMINATOR;
+					Ret += TEXT("\t") + BaseName + TEXT(" Ret;") LINE_TERMINATOR;
+					Ret += TEXT("\tRet.Value = log2(A.Value);") LINE_TERMINATOR;
+					Ret += TEXT("\t") + FieldName + TEXT(" dFdA = rcp(A.Value) * 1.442695f;");
+					Ret += TEXT("\tRet.Ddx = dFdA * A.Ddx ;") LINE_TERMINATOR;
+					Ret += TEXT("\tRet.Ddy = dFdA * A.Ddy;") LINE_TERMINATOR;
+					Ret += TEXT("\treturn Ret;") LINE_TERMINATOR;
+					Ret += TEXT("}") LINE_TERMINATOR;
+					Ret += TEXT("") LINE_TERMINATOR;
+					break;
+				case EFunc1::Log10:
+					Ret += BaseName + TEXT(" Log10Deriv(") + BaseName + TEXT(" A)") LINE_TERMINATOR;
+					Ret += TEXT("{") LINE_TERMINATOR;
+					Ret += TEXT("\t") + BaseName + TEXT(" Ret;") LINE_TERMINATOR;
+					Ret += TEXT("\tRet.Value = log10(A.Value);") LINE_TERMINATOR;
+					Ret += TEXT("\t") + FieldName + TEXT(" dFdA = rcp(A.Value) * 0.4342945f;");
+					Ret += TEXT("\tRet.Ddx = dFdA * A.Ddx;") LINE_TERMINATOR;
+					Ret += TEXT("\tRet.Ddy = dFdA * A.Ddy;") LINE_TERMINATOR;
+					Ret += TEXT("\treturn Ret;") LINE_TERMINATOR;
+					Ret += TEXT("}") LINE_TERMINATOR;
+					Ret += TEXT("") LINE_TERMINATOR;
+					break;
+				case EFunc1::Exp:
+					Ret += BaseName + TEXT(" ExpDeriv(") + BaseName + TEXT(" A)") LINE_TERMINATOR;
+					Ret += TEXT("{") LINE_TERMINATOR;
+					Ret += TEXT("\t") + BaseName + TEXT(" Ret;") LINE_TERMINATOR;
+					Ret += TEXT("\tRet.Value = exp(A.Value);") LINE_TERMINATOR;
+					Ret += TEXT("\tRet.Ddx = exp(A.Value) * A.Ddx;") LINE_TERMINATOR;
+					Ret += TEXT("\tRet.Ddy = exp(A.Value) * A.Ddy;") LINE_TERMINATOR;
+					Ret += TEXT("\treturn Ret;") LINE_TERMINATOR;
+					Ret += TEXT("}") LINE_TERMINATOR;
+					Ret += TEXT("") LINE_TERMINATOR;
+					break;
+				case EFunc1::Length:
+					Ret += TEXT("FloatDeriv LengthDeriv(") + BaseName + TEXT(" A)") LINE_TERMINATOR;
+					Ret += TEXT("{") LINE_TERMINATOR;
+					Ret += TEXT("\tFloatDeriv Ret = SqrtDeriv(DotDeriv(A,A));") LINE_TERMINATOR;
+					Ret += TEXT("\treturn Ret;") LINE_TERMINATOR;
+					Ret += TEXT("}") LINE_TERMINATOR;
+					Ret += TEXT("") LINE_TERMINATOR;
+					break;
+				case EFunc1::InvLength:
+					Ret += TEXT("FloatDeriv InvLengthDeriv(") + BaseName + TEXT(" A)") LINE_TERMINATOR;
+					Ret += TEXT("{") LINE_TERMINATOR;
+					Ret += TEXT("\tFloatDeriv Ret = RsqrtDeriv(DotDeriv(A,A));") LINE_TERMINATOR;
+					Ret += TEXT("\treturn Ret;") LINE_TERMINATOR;
+					Ret += TEXT("}") LINE_TERMINATOR;
+					Ret += TEXT("") LINE_TERMINATOR;
+					break;
+				case EFunc1::Normalize:
+					Ret += BaseName + TEXT(" NormalizeDeriv(") + BaseName + TEXT(" A)") LINE_TERMINATOR;
+					Ret += TEXT("{") LINE_TERMINATOR;
+					Ret += TEXT("\tFloatDeriv InvLen = RsqrtDeriv(DotDeriv(A,A));") LINE_TERMINATOR;
+					Ret += TEXT("\t") + BaseName + TEXT(" Ret = MulDeriv(") + ConvertDeriv(TEXT("InvLen"), Index, 0) + TEXT(", A);") LINE_TERMINATOR;
+					Ret += TEXT("\treturn Ret;") LINE_TERMINATOR;
+					Ret += TEXT("}") LINE_TERMINATOR;
+					Ret += TEXT("") LINE_TERMINATOR;
+					break;
+				default:
+					check(0);
+					break;
+				}
+			}
+		}
+	}
+
+	for (int32 Index = 0; Index < 4; Index++)
+	{
+		const FString BaseName = GetDerivVectorName(Index);
+		const FString FieldName = GetFloatVectorName(Index);
+		const FString BoolName = GetBoolVectorName(Index);
+
+		if (bLerpEnabled[Index] || IsDebugGenerateAllFunctionsEnabled())
+		{
+			
+			// lerp(a,b,s) = a*(1-s) + b*s
+			// lerp(a,b,s)' = a' * (1 - s') + b' * s + s' * (b - a)
+			Ret += FString::Printf(TEXT("%s LerpDeriv(%s A, %s B, %s S)"), *BaseName, *BaseName, *BaseName, *BaseName) + LINE_TERMINATOR;
+			Ret += TEXT("{") LINE_TERMINATOR;
+			Ret += TEXT("\t") + BaseName + TEXT(" Ret;") LINE_TERMINATOR;
+			Ret += TEXT("\tRet.Value = lerp(A.Value, B.Value, S.Value);") LINE_TERMINATOR;
+			Ret += TEXT("\tRet.Ddx = lerp(A.Ddx, B.Ddx, S.Value) + S.Ddx * (B.Value - A.Value);") LINE_TERMINATOR;
+			Ret += TEXT("\tRet.Ddy = lerp(A.Ddy, B.Ddy, S.Value) + S.Ddy * (B.Value - A.Value);") LINE_TERMINATOR;
+			Ret += TEXT("\treturn Ret;") LINE_TERMINATOR;
+			Ret += TEXT("}") LINE_TERMINATOR;
+			Ret += TEXT("") LINE_TERMINATOR;
+		}
+
+		if (bIfEnabled[Index] || IsDebugGenerateAllFunctionsEnabled())
+		{
+			Ret += FString::Printf(TEXT("%s IfDeriv(float A, float B, %s Greater, %s Less)"), *BaseName, *BaseName, *BaseName) + LINE_TERMINATOR;
+			Ret += TEXT("{") LINE_TERMINATOR;
+			Ret += TEXT("\tif(A >= B)") LINE_TERMINATOR;
+			Ret += TEXT("\t\treturn Greater;") LINE_TERMINATOR;
+			Ret += TEXT("\telse") LINE_TERMINATOR;
+			Ret += TEXT("\t\treturn Less;") LINE_TERMINATOR;
+			Ret += TEXT("}") LINE_TERMINATOR;
+			Ret += TEXT("") LINE_TERMINATOR;
+		}
+
+		if (bIf2Enabled[Index] || IsDebugGenerateAllFunctionsEnabled())
+		{
+			Ret += FString::Printf(TEXT("%s IfDeriv(float A, float B, %s Greater, %s Less, %s Equal, float Threshold)"), *BaseName, *BaseName, *BaseName, *BaseName) + LINE_TERMINATOR;
+			Ret += TEXT("{") LINE_TERMINATOR;
+			Ret += TEXT("\tif(!(abs(A - B) > Threshold))") LINE_TERMINATOR;	// Written like this to preserve NaN behavior of original code.
+			Ret += TEXT("\t\treturn Equal;") LINE_TERMINATOR;
+			Ret += TEXT("\tif(A >= B)") LINE_TERMINATOR;
+			Ret += TEXT("\t\treturn Greater;") LINE_TERMINATOR;
+			Ret += TEXT("\telse") LINE_TERMINATOR;
+			Ret += TEXT("\t\treturn Less;") LINE_TERMINATOR;
+			Ret += TEXT("}") LINE_TERMINATOR;
+			Ret += TEXT("") LINE_TERMINATOR;
+		}
+	}
+
+	if (bRotateScaleOffsetTexCoords || IsDebugGenerateAllFunctionsEnabled())
+	{
+		// float2(dot(InTexCoords, InRotationScale.xy), dot(InTexCoords, InRotationScale.zw)) + InOffset;
+		// InTexCoords.xy * InRotationScale.xw + InTexCoords.yx * InRotationScale.yz + InOffset;
+		Ret += TEXT("FloatDeriv2 RotateScaleOffsetTexCoordsDeriv(FloatDeriv2 TexCoord, FloatDeriv2 RotationScale, FloatDeriv2 Offset)") LINE_TERMINATOR;
+		Ret += TEXT("{") LINE_TERMINATOR;
+		Ret += TEXT("\tFloatDeriv2 Ret = Offset;") LINE_TERMINATOR;
+		Ret += TEXT("\tRet = AddDeriv(Ret, MulDeriv(TexCoord, SwizzleDeriv2(RotationScale, xw));") LINE_TERMINATOR;
+		Ret += TEXT("\tRet = AddDeriv(Ret, MulDeriv(SwizzleDeriv2(TexCoord, yx), SwizzleDeriv2(RotationScale, yz));") LINE_TERMINATOR;
+		Ret += TEXT("\treturn Ret;") LINE_TERMINATOR;
+		Ret += TEXT("}") LINE_TERMINATOR;
+		Ret += TEXT("") LINE_TERMINATOR;
+	}
+	
+	if (bUnMirrorEnabled[1][1] || IsDebugGenerateAllFunctionsEnabled())
+	{
+		// UnMirrorUV
+		Ret += TEXT("FloatDeriv2 UnMirrorU(FloatDeriv2 UV, FMaterialPixelParameters Parameters)") LINE_TERMINATOR;
+		Ret += TEXT("{") LINE_TERMINATOR;
+		Ret += TEXT("\tconst MaterialFloat Scale = (Parameters.UnMirrored * 0.5f);") LINE_TERMINATOR;
+		Ret += TEXT("\tUV.Value = UV.Value * Scale + 0.5f;") LINE_TERMINATOR;
+		Ret += TEXT("\tUV.Ddx *= Scale;") LINE_TERMINATOR;
+		Ret += TEXT("\tUV.Ddy *= Scale;") LINE_TERMINATOR;
+		Ret += TEXT("\treturn UV;") LINE_TERMINATOR;
+		Ret += TEXT("}") LINE_TERMINATOR;
+		Ret += TEXT("") LINE_TERMINATOR;
+	}
+	
+	if(bUnMirrorEnabled[1][0] || IsDebugGenerateAllFunctionsEnabled())
+	{
+		// UnMirrorU
+		Ret += TEXT("FloatDeriv2 UnMirrorU(FloatDeriv2 UV, FMaterialPixelParameters Parameters)") LINE_TERMINATOR;
+		Ret += TEXT("{") LINE_TERMINATOR;
+		Ret += TEXT("\tconst MaterialFloat Scale = (Parameters.UnMirrored * 0.5f);") LINE_TERMINATOR;
+		Ret += TEXT("\tUV.Value.x = UV.Value.x * Scale + 0.5f;") LINE_TERMINATOR;
+		Ret += TEXT("\tUV.Ddx.x *= Scale;") LINE_TERMINATOR;
+		Ret += TEXT("\tUV.Ddy.x *= Scale;") LINE_TERMINATOR;
+		Ret += TEXT("\treturn UV;") LINE_TERMINATOR;
+		Ret += TEXT("}") LINE_TERMINATOR;
+		Ret += TEXT("") LINE_TERMINATOR;
+	}
+	
+	if (bUnMirrorEnabled[0][1] || IsDebugGenerateAllFunctionsEnabled())
+	{
+		// UnMirrorV
+		Ret += TEXT("FloatDeriv2 UnMirrorV(FloatDeriv2 UV, FMaterialPixelParameters Parameters)") LINE_TERMINATOR;
+		Ret += TEXT("{") LINE_TERMINATOR;
+		Ret += TEXT("\tconst MaterialFloat Scale = (Parameters.UnMirrored * 0.5f);") LINE_TERMINATOR;
+		Ret += TEXT("\tUV.Value.y = UV.Value.y * Scale + 0.5f;") LINE_TERMINATOR;
+		Ret += TEXT("\tUV.Ddx.y *= Scale;") LINE_TERMINATOR;
+		Ret += TEXT("\tUV.Ddy.y *= Scale;") LINE_TERMINATOR;
+		Ret += TEXT("\treturn UV;") LINE_TERMINATOR;
+		Ret += TEXT("}") LINE_TERMINATOR;
+		Ret += TEXT("") LINE_TERMINATOR;
+	}
+
+	return Ret;
+}
+
+
+FString FMaterialDerivativeAutogen::ApplyUnMirror(FString Value, bool bUnMirrorU, bool bUnMirrorV)
+{
+	if (bUnMirrorU && bUnMirrorV)
+	{
+		Value = FString::Printf(TEXT("UnMirrorUV(%s, Parameters)"), *Value);
+	}
+	else if (bUnMirrorU)
+	{
+		Value = FString::Printf(TEXT("UnMirrorU(%s, Parameters)"), *Value);
+	}
+	else if (bUnMirrorV)
+	{
+		Value = FString::Printf(TEXT("UnMirrorV(%s, Parameters)"), *Value);
+	}
+
+	bUnMirrorEnabled[bUnMirrorU][bUnMirrorV] = true;
+	
+	return Value;	
+}
+
+#endif // WITH_EDITORONLY_DATA
\ No newline at end of file
diff --git a/Engine/Source/Runtime/Engine/Private/Materials/HLSLMaterialDerivativeAutogen.h b/Engine/Source/Runtime/Engine/Private/Materials/HLSLMaterialDerivativeAutogen.h
new file mode 100644
--- /dev/null
+++ b/Engine/Source/Runtime/Engine/Private/Materials/HLSLMaterialDerivativeAutogen.h
@@ -0,0 +1,139 @@
+// Copyright Epic Games, Inc. All Rights Reserved.
+
+#pragma once
+
+#include "CoreMinimal.h"
+#include "MaterialShared.h"
+
+#if WITH_EDITORONLY_DATA
+
+class FHLSLMaterialTranslator;
+
+/**
+* For a node, the known information of the partial derivatives.
+* NotAware	- This node is made by a function that has no knowledge of analytic partial derivatives.
+* NotValid	- This node is aware of partial derivatives, and knows that one of its source inputs is not partial derivative aware, and therefore its value is not to be used.
+* Zero		- This node is aware of partial derivatives, and knows that it's value is zero, as is the case for uniform parameters.
+* Valid		- This node is aware of partial derivatives, and knows that it has a calculated value.
+**/
+enum class EDerivativeStatus
+{
+	NotAware,
+	NotValid,
+	Zero,
+	Valid,
+};
+
+inline bool IsDerivativeValid(EDerivativeStatus Status)
+{
+	return Status == EDerivativeStatus::Valid || Status == EDerivativeStatus::Zero;
+}
+
+int32 GetDerivTypeIndex(EMaterialValueType ValueType);
+
+struct FDerivInfo
+{
+	EMaterialValueType	Type;
+	int32				TypeIndex;
+	EDerivativeStatus	DerivativeStatus;
+};
+
+class FMaterialDerivativeAutogen
+{
+public:
+	// Unary functions
+	enum class EFunc1
+	{
+		Abs,
+		Log2,
+		Log10,
+		Exp,
+		Sin,
+		Cos,
+		Tan,
+		Asin,
+		AsinFast,
+		Acos,
+		AcosFast,
+		Atan,
+		AtanFast,
+		Sqrt,
+		Rcp,
+		Rsqrt,
+		Saturate,
+		Frac,
+		Length,
+		InvLength,
+		Normalize,
+		Num
+	};
+
+	// Binary functions
+	enum class EFunc2
+	{
+		Add,
+		Sub,
+		Mul,
+		Div,
+		Fmod,
+		Max,
+		Min,
+		Dot,	// Depends on Add/Mul, so it must come after them
+		Pow,
+		PowPositiveClamped,
+		Cross,
+		Atan2,
+		Atan2Fast,
+		Num
+	};
+
+	// Note that the type index is from [0,3] for float1 to float4. I.e. A float3 would be index 2.
+	static int32 GetFunc1ReturnNumComponents(int32 SrcTypeIndex, EFunc1 Op);
+	static int32 GetFunc2ReturnNumComponents(int32 LhsTypeIndex, int32 RhsTypeIndex, EFunc2 Op);
+
+	void EnableGeneratedDepencencies();
+
+	int32 GenerateExpressionFunc1(FHLSLMaterialTranslator& Translator, EFunc1 Op, int32 SrcCode);
+	int32 GenerateExpressionFunc2(FHLSLMaterialTranslator& Translator, EFunc2 Op, int32 LhsCode, int32 RhsCode);
+
+	int32 GenerateLerpFunc(FHLSLMaterialTranslator& Translator, int32 A, int32 B, int32 S);
+	int32 GenerateRotateScaleOffsetTexCoordsFunc(FHLSLMaterialTranslator& Translator, int32 TexCoord, int32 RotationScale, int32 Offset);
+	int32 GenerateIfFunc(FHLSLMaterialTranslator& Translator, int32 A, int32 B, int32 Greater, int32 Equal, int32 Less, int32 Threshold);
+	
+	FString GenerateUsedFunctions(FHLSLMaterialTranslator& Translator);
+
+	FString ApplyUnMirror(FString Value, bool bUnMirrorU, bool bUnMirrorV);
+
+	FString CoerceValueRaw(const FString& Token, int32 SrcType, EDerivativeStatus SrcStatus, int32 DstType);
+	FString CoerceValueDeriv(const FString& Token, int32 SrcType, EDerivativeStatus SrcStatus, int32 DstType);
+
+	FString ConstructDeriv(const FString& Value, const FString& Ddx, const FString& Ddy, int32 DstType);
+	FString ConstructDerivFinite(const FString& Value, int32 DstType);
+
+	FString ConvertDeriv(const FString& Value, int32 DstType, int32 SrcType);
+
+	FString ExtractElement(const FString& Value, int32 SrcType, int32 ElementIndex);
+
+private:
+	// State to keep track of which derivative functions have been used and need to be generated.
+	bool bFunc1OpIsEnabled[(int32)EFunc1::Num][4] = {};
+	bool bFunc2OpIsEnabled[(int32)EFunc2::Num][4] = {};
+	bool bLerpEnabled[4] = {};
+	bool bIfEnabled[4] = {};
+	bool bIf2Enabled[4] = {};
+	bool bUnMirrorEnabled[2][2] = {};
+
+	bool bRotateScaleOffsetTexCoords = false;
+
+	bool bConstructDerivEnabled[4] = {};
+	bool bConstructConstantDerivEnabled[4] = {};
+	bool bConstructFiniteDerivEnabled[4] = {};
+
+	bool bConvertDerivEnabled[4][4] = {};
+
+	bool bExtractIndexEnabled[4] = {};
+
+	bool bSelectElemHelperEnabled[4] = {};
+};
+
+#endif // WITH_EDITORONLY_DATA
diff --git a/Engine/Source/Runtime/Engine/Private/Materials/HLSLMaterialTranslator.cpp b/Engine/Source/Runtime/Engine/Private/Materials/HLSLMaterialTranslator.cpp
--- a/Engine/Source/Runtime/Engine/Private/Materials/HLSLMaterialTranslator.cpp
+++ b/Engine/Source/Runtime/Engine/Private/Materials/HLSLMaterialTranslator.cpp
@@ -28,13 +28,6 @@ static FAutoConsoleVariableRef CVarAnalyticDerivDebugEmitInvalidDerivTokens(
 	TEXT("Debug: Emit '$' tokens to mark expressions with invalid derivatives.\n")
 );
 
-static int32 GDebugGenerateAllFunctionsEnabled = 0;
-static FAutoConsoleVariableRef CVarAnalyticDerivDebugGenerateAllFunctions(
-	TEXT("r.MaterialEditor.AnalyticDeriv.DebugGenerateAllFunctions"),
-	GDebugGenerateAllFunctionsEnabled,
-	TEXT("Debug: Generate all derivative functions.")
-);
-
 static inline bool IsAnalyticDerivEnabled()
 {
 	return GAnalyticDerivEnabled != 0;
@@ -45,13 +38,8 @@ static inline bool IsDebugTextureSampleEnabled()
 	return IsAnalyticDerivEnabled() && (GDebugTextureSampleEnabled != 0);
 }
 
-static inline bool IsDebugGenerateAllFunctionsEnabled()
-{
-	return IsAnalyticDerivEnabled() && (GDebugGenerateAllFunctionsEnabled != 0);
-}
-
 /** @return the number of components in a vector type. */
-static inline uint32 GetNumComponents(EMaterialValueType Type)
+uint32 GetNumComponents(EMaterialValueType Type)
 {
 	switch(Type)
 	{
@@ -90,61 +78,6 @@ static inline int32 SwizzleComponentToIndex(TCHAR Component)
 	}
 }
 
-
-static inline const TCHAR * GetBoolVectorName(int32 TypeIndex)
-{
-	switch(TypeIndex)
-	{
-	case 0:
-		return TEXT("bool");
-	case 1:
-		return TEXT("bool2");
-	case 2:
-		return TEXT("bool3");
-	case 3:
-		return TEXT("bool4");
-	default:
-		check(0);
-		return TEXT("");
-	}
-}
-
-static inline const TCHAR * GetFloatVectorName(int32 TypeIndex)
-{
-	switch(TypeIndex)
-	{
-	case 0:
-		return TEXT("float");
-	case 1:
-		return TEXT("float2");
-	case 2:
-		return TEXT("float3");
-	case 3:
-		return TEXT("float4");
-	default:
-		check(0);
-		return TEXT("");
-	}
-}
-
-static inline const TCHAR * GetDerivVectorName(int32 TypeIndex)
-{
-	switch(TypeIndex)
-	{
-	case 0:
-		return TEXT("FloatDeriv");
-	case 1:
-		return TEXT("FloatDeriv2");
-	case 2:
-		return TEXT("FloatDeriv3");
-	case 3:
-		return TEXT("FloatDeriv4");
-	default:
-		check(0);
-		return TEXT("");
-	}
-}
-
 static inline const TCHAR * GetFloatZeroVector(uint32 NumComponents)
 {
 	switch (NumComponents)
@@ -163,1709 +96,6 @@ static inline const TCHAR * GetFloatZeroVector(uint32 NumComponents)
 	}
 }
 
-static inline bool IsDerivativeValid(EDerivativeStatus Status)
-{
-	return Status == DS_Valid || Status == DS_Zero;
-}
-
-void FMaterialDerivativeAutogen::EnableGeneratedDepencencies()
-{
-
-	for (int32 Index = 0; Index < 4; Index++)
-	{
-		// PowPositiveClamped requires Pow
-		if (bFunc2OpIsEnabled[EF2_PowPositiveClamped][Index])
-		{
-			bFunc2OpIsEnabled[EF2_Pow][Index] = true;
-		}
-	}
-
-	for (int32 Index = 0; Index < 4; Index++)
-	{
-		// normalize requires rsqrt1, dot, expand, and mul
-		if (bFunc1OpIsEnabled[EF1_Normalize][Index])
-		{
-			bConvertDerivEnabled[Index][0] = true;
-			bFunc2OpIsEnabled[EF2_Dot][Index] = true;
-			bFunc1OpIsEnabled[EF1_Rsqrt][0] = true;
-			bFunc2OpIsEnabled[EF2_Mul][Index] = true;
-		}
-
-		// length requires sqrt1 and dot, dot requires a few other things, but those are handled below
-		if (bFunc1OpIsEnabled[EF1_Length][Index])
-		{
-			bFunc2OpIsEnabled[EF2_Dot][Index] = true;
-			bFunc1OpIsEnabled[EF1_Sqrt][0] = true;
-		}
-
-		// inv length requires rsqrt1 (instead of sqrt1) and dot
-		if (bFunc1OpIsEnabled[EF1_InvLength][Index])
-		{
-			bFunc2OpIsEnabled[EF2_Dot][Index] = true;
-			bFunc1OpIsEnabled[EF1_Rsqrt][0] = true;
-		}
-	}
-
-	// min, max, saturate, and PosClampedPow require the select helper
-	for (int32 Index = 0; Index < 4; Index++)
-	{
-		if (bFunc2OpIsEnabled[EF2_Max][Index] ||
-			bFunc2OpIsEnabled[EF2_Min][Index] ||
-			bFunc1OpIsEnabled[EF1_Saturate][Index] ||
-			bFunc2OpIsEnabled[EF2_PowPositiveClamped][Index])
-		{
-			bSelectElemHelperEnabled[Index] = true;
-		}
-	}
-
-	// Dot requires extract, mul1, add1 and FloatDeriv constructor
-	for (int32 Index = 0; Index < 4; Index++)
-	{
-		if (bFunc2OpIsEnabled[EF2_Dot][Index])
-		{
-			bExtractIndexEnabled[Index] = true;
-			bConstructConstantDerivEnabled[0] = true;
-			bFunc2OpIsEnabled[EF2_Add][0] = true;
-			bFunc2OpIsEnabled[EF2_Mul][0] = true;
-		}
-	}
-
-	if (bRotateScaleOffsetTexCoords)
-	{
-		bFunc2OpIsEnabled[EF2_Add][1] = true;
-		bFunc2OpIsEnabled[EF2_Mul][1] = true;
-		bConstructDerivEnabled[1] = true;
-	}
-}
-
-// 0: Float
-// 1: Float2
-// 2: Float3
-// 3: Float4
-// -1: Everything else
-static int32 GetDerivTypeIndex(EMaterialValueType ValueType)
-{
-	int32 Ret = INDEX_NONE;
-	if (ValueType == MCT_Float1 || ValueType == MCT_Float)
-	{
-		Ret = 0;
-	}
-	else if (ValueType == MCT_Float2)
-	{
-		Ret = 1;
-	}
-	else if (ValueType == MCT_Float3)
-	{
-		Ret = 2;
-	}
-	else if (ValueType == MCT_Float4)
-	{
-		Ret = 3;
-	}
-	else
-	{
-		Ret = -1; // redundant, but leaving it here so we can set breakpoints
-		check(false);
-	}
-	return Ret;
-}
-
-
-static EMaterialValueType GetMaterialTypFromDerivTypeIndex(int32 Index)
-{
-	switch(Index)
-	{
-	case 0:
-		return MCT_Float;
-	case 1:
-		return MCT_Float2;
-	case 2:
-		return MCT_Float3;
-	case 3:
-		return MCT_Float4;
-	default:
-		check(0);
-		break;
-	}
-
-	return (EMaterialValueType)0; // invalid, should be a Float 1/2/3/4, break at the check(0);
-}
-
-// given a string, convert it from type to type
-FString FMaterialDerivativeAutogen::CoerceValueRaw(const FString& Token, int32 SrcType, EDerivativeStatus SrcStatus, int32 DstType)
-{
-	FString Ret = Token;
-
-	// If the original value is a derivative, grab the raw value
-	if (SrcStatus == DS_Valid)
-	{
-		Ret = Ret + TEXT(".Value");
-	}
-
-	if (SrcType == DstType)
-	{
-		// no op
-	}
-	else
-	{
-		check(SrcType == 0); // can only coerce a float1
-		if (DstType == 0)
-		{
-			Ret = FString::Printf( TEXT("MaterialFloat(%s)"), *Ret);
-		}
-		else if (DstType == 1)
-		{
-			Ret = FString::Printf( TEXT("MaterialFloat2(%s,%s)"), *Ret, *Ret);
-		}
-		else if (DstType == 2)
-		{
-			Ret = FString::Printf( TEXT("MaterialFloat3(%s,%s,%s)"), *Ret, *Ret, *Ret);
-		}
-		else if (DstType == 3)
-		{
-			Ret = FString::Printf( TEXT("MaterialFloat4(%s,%s,%s,%s)"), *Ret, *Ret, *Ret, *Ret);
-		}
-		else
-		{
-			check(0);
-		}
-	}
-	return Ret;
-}
-
-// given a string, convert it from type to type
-FString FMaterialDerivativeAutogen::CoerceValueDeriv(const FString& Token, int32 SrcType, EDerivativeStatus SrcStatus, int32 DstType)
-{
-	check(SrcType < 4);
-	check(DstType < 4);
-
-	FString Ret = Token;
-
-	check(IsDerivativeValid(SrcStatus));
-
-	// If it's valid, then it's already a type. Otherwise, we need to convert it from raw to deriv
-	if (SrcStatus == DS_Zero)
-	{
-		FString SrcDerivType = GetDerivVectorName(SrcType);
-		bConstructConstantDerivEnabled[SrcType] = true;
-		Ret = TEXT("ConstructConstant") + SrcDerivType + TEXT("(") + Ret + TEXT(")");
-	}
-
-	return ConvertDeriv(Ret, DstType, SrcType);
-}
-
-FString FMaterialDerivativeAutogen::ConstructDeriv(const FString& Value, const FString& Ddx, const FString& Ddy, int32 DstType)
-{
-	check(DstType >= 0);
-	check(DstType < 4);
-
-	bConstructDerivEnabled[DstType] = true;
-	FString TypeName = GetDerivVectorName(DstType);
-	FString Ret = TEXT("Construct") + TypeName + TEXT("(") + Value + TEXT(",") + Ddx + TEXT(",") + Ddy + TEXT(")");
-	return Ret;
-}
-
-FString FMaterialDerivativeAutogen::ConstructDerivFinite(const FString& Value, int32 DstType)
-{
-	check(DstType >= 0);
-	check(DstType < 4);
-
-	bConstructFiniteDerivEnabled[DstType] = true;
-
-	FString TypeName = GetDerivVectorName(DstType);
-	FString Ret = TEXT("ConstructFinite") + TypeName + TEXT("(") + Value + TEXT(")");
-	return Ret;
-}
-
-FString FMaterialDerivativeAutogen::ExtractElement(const FString& Value, int32 SrcType, int32 ElementIndex)
-{
-	check(SrcType < 4);
-	check(ElementIndex <= SrcType);
-	bExtractIndexEnabled[SrcType] = true;
-
-	FString TypeName = GetDerivVectorName(SrcType);
-
-	FString Ret = FString::Printf(TEXT("Extract%s_%d(%s %s)"), *TypeName, ElementIndex + 1, *TypeName, *Value);
-
-	return Ret;
-}
-
-FString FMaterialDerivativeAutogen::ConvertDeriv(const FString& Value, int32 DstType, int32 SrcType)
-{
-	check(DstType < 4);
-	check(SrcType < 4);
-	if (DstType == SrcType)
-		return Value;
-	
-	bConvertDerivEnabled[DstType][SrcType] = true;
-
-	FString DstTypeName = GetDerivVectorName(DstType);
-	return FString::Printf(TEXT("Convert%s(%s)"), *DstTypeName, *Value);
-}
-
-int32 FMaterialDerivativeAutogen::GetFunc1ReturnNumComponents(int32 SrcTypeIndex, EFunc1 Op)
-{
-	int32 DstTypeIndex = INDEX_NONE;
-
-	switch(Op)
-	{
-	case EF1_Abs:
-	case EF1_Log2:
-	case EF1_Log10:
-	case EF1_Exp:
-	case EF1_Sin:
-	case EF1_Cos:
-	case EF1_Tan:
-	case EF1_Asin:
-	case EF1_AsinFast:
-	case EF1_Acos:
-	case EF1_AcosFast:
-	case EF1_Atan:
-	case EF1_AtanFast:
-	case EF1_Sqrt:
-	case EF1_Rcp:
-	case EF1_Rsqrt:
-	case EF1_Saturate:
-	case EF1_Frac:
-	case EF1_Normalize:
-		DstTypeIndex = SrcTypeIndex;
-		break;
-
-	case EF1_Length:
-	case EF1_InvLength:
-		DstTypeIndex = 0;
-		break;
-	default:
-		break;
-	}
-
-	return DstTypeIndex;
-}
-
-int32 FMaterialDerivativeAutogen::GetFunc2ReturnNumComponents(int32 LhsTypeIndex, int32 RhsTypeIndex, EFunc2 Op)
-{
-	int32 DstTypeIndex = INDEX_NONE;
-	
-	switch(Op)
-	{
-	case EF2_Add:
-	case EF2_Sub:
-	case EF2_Mul:
-	case EF2_Div:
-	case EF2_Fmod:
-	case EF2_Max:
-	case EF2_Min:
-	case EF2_Pow:
-	case EF2_PowPositiveClamped:
-	case EF2_Atan2:
-	case EF2_Atan2Fast:
-		// if the initial type is different from the output type, then it's only valid if type is 0 (float). We can convert
-		// a float to a type with more components, but for example, we can't implicitly convert a float2 to a float3/float4.
-		if (LhsTypeIndex == RhsTypeIndex || RhsTypeIndex == 0 || LhsTypeIndex == 0)
-		{
-			DstTypeIndex = FMath::Max<int32>(LhsTypeIndex, RhsTypeIndex);
-		}
-		break;
-	case EF2_Dot:
-		DstTypeIndex = 0;
-		break;
-	case EF2_Cross:
-		check(LhsTypeIndex == 2);
-		check(RhsTypeIndex == 2);
-		DstTypeIndex = 2;
-		break;
-	default:
-		break;
-	}
-
-	return DstTypeIndex;
-}
-
-
-int32 FMaterialDerivativeAutogen::GenerateExpressionFunc1(FHLSLMaterialTranslator& Translator, EFunc1 Op, int32 SrcCode)
-{
-	if (SrcCode == INDEX_NONE)
-	{
-		return INDEX_NONE;
-	}
-
-	const FDerivInfo SrcDerivInfo = Translator.GetDerivInfo(SrcCode);
-	int32 OutputTypeIndex	= GetFunc1ReturnNumComponents(SrcDerivInfo.TypeIndex, Op);
-
-	if (OutputTypeIndex < 0)
-	{
-		return INDEX_NONE;
-	}
-
-	EDerivativeStatus DstStatus = IsDerivativeValid(SrcDerivInfo.DerivativeStatus) ? DS_Valid : DS_NotValid;
-	bool bUseScalarVersion = (DstStatus != DS_Valid);
-
-	// make initial tokens
-	FString DstTokens[CompiledPDV_MAX];
-
-	for (int32 Index = 0; Index < CompiledPDV_MAX; Index++)
-	{
-		ECompiledPartialDerivativeVariation Variation = (ECompiledPartialDerivativeVariation)Index;
-
-		FString SrcToken = Translator.GetParameterCodeDeriv(SrcCode,Variation);
-
-		// The token is the symbol name. If we are in finite mode, that's all we have to do. But if
-		// we are in analytic mode, we may need to get the value.
-		if (Index == CompiledPDV_Analytic)
-		{
-			SrcToken = CoerceValueRaw(SrcToken, SrcDerivInfo.TypeIndex, SrcDerivInfo.DerivativeStatus, SrcDerivInfo.TypeIndex);
-		}
-
-		FString DstToken;
-		// just generate a type
-		switch(Op)
-		{
-		case EF1_Abs:
-			DstToken = TEXT("abs(") + SrcToken + TEXT(")");
-			break;
-		case EF1_Log2:
-			DstToken = TEXT("log2(") + SrcToken + TEXT(")");
-			break;
-		case EF1_Log10:
-			DstToken = TEXT("log10(") + SrcToken + TEXT(")");
-			break;
-		case EF1_Exp:
-			DstToken = TEXT("exp(") + SrcToken + TEXT(")");
-			break;
-		case EF1_Sin:
-			DstToken = TEXT("sin(") + SrcToken + TEXT(")");
-			break;
-		case EF1_Cos:
-			DstToken = TEXT("cos(") + SrcToken + TEXT(")");
-			break;
-		case EF1_Tan:
-			DstToken = TEXT("tan(") + SrcToken + TEXT(")");
-			break;
-		case EF1_Asin:
-			DstToken = TEXT("asin(") + SrcToken + TEXT(")");
-			break;
-		case EF1_AsinFast:
-			DstToken = TEXT("asinFast(") + SrcToken + TEXT(")");
-			break;
-		case EF1_Acos:
-			DstToken = TEXT("acos(") + SrcToken + TEXT(")");
-			break;
-		case EF1_AcosFast:
-			DstToken = TEXT("acosFast(") + SrcToken + TEXT(")");
-			break;
-		case EF1_Atan:
-			DstToken = TEXT("atan(") + SrcToken + TEXT(")");
-			break;
-		case EF1_AtanFast:
-			DstToken = TEXT("atanFast(") + SrcToken + TEXT(")");
-			break;
-		case EF1_Sqrt:
-			DstToken = TEXT("sqrt(") + SrcToken + TEXT(")");
-			break;
-		case EF1_Rcp:
-			DstToken = TEXT("rcp(") + SrcToken + TEXT(")");
-			break;
-		case EF1_Rsqrt:
-			DstToken = TEXT("rsqrt(") + SrcToken + TEXT(")");
-			break;
-		case EF1_Saturate:
-			DstToken = TEXT("saturate(") + SrcToken + TEXT(")");
-			break;
-		case EF1_Frac:
-			DstToken = TEXT("frac(") + SrcToken + TEXT(")");
-			break;
-		case EF1_Length:
-			DstToken = TEXT("length(") + SrcToken + TEXT(")");
-			break;
-		case EF1_InvLength:
-			DstToken = TEXT("rcp(length(") + SrcToken + TEXT("))");
-			break;
-		case EF1_Normalize:
-			DstToken = TEXT("normalize(") + SrcToken + TEXT(")");
-			break;
-		default:
-			check(0);
-			break;
-		}
-
-		DstTokens[Index] = DstToken;
-	}
-
-	if (!bUseScalarVersion)
-	{
-		FString SrcToken = Translator.GetParameterCodeDeriv(SrcCode, CompiledPDV_Analytic);
-
-		SrcToken = CoerceValueDeriv(SrcToken, SrcDerivInfo.TypeIndex, SrcDerivInfo.DerivativeStatus, SrcDerivInfo.TypeIndex);
-
-		check(Op < EF1_Num);
-		bFunc1OpIsEnabled[Op][SrcDerivInfo.TypeIndex] = true;
-		
-		FString DstToken;
-		switch(Op)
-		{
-		case EF1_Abs:
-			DstToken = TEXT("AbsDeriv(") + SrcToken + TEXT(")");
-			break;
-		case EF1_Log2:
-			DstToken = TEXT("Log2Deriv(") + SrcToken + TEXT(")");
-			break;
-		case EF1_Log10:
-			DstToken = TEXT("Log10Deriv(") + SrcToken + TEXT(")");
-			break;
-		case EF1_Exp:
-			DstToken = TEXT("ExpDeriv(") + SrcToken + TEXT(")");
-			break;
-		case EF1_Sin:
-			DstToken = TEXT("SinDeriv(") + SrcToken + TEXT(")");
-			break;
-		case EF1_Cos:
-			DstToken = TEXT("CosDeriv(") + SrcToken + TEXT(")");
-			break;
-		case EF1_Tan:
-			DstToken = TEXT("TanDeriv(") + SrcToken + TEXT(")");
-			break;
-		case EF1_Asin:
-			DstToken = TEXT("ASinDeriv(") + SrcToken + TEXT(")");
-			break;
-		case EF1_AsinFast:
-			DstToken = TEXT("ASinFastDeriv(") + SrcToken + TEXT(")");
-			break;
-		case EF1_Acos:
-			DstToken = TEXT("ACosDeriv(") + SrcToken + TEXT(")");
-			break;
-		case EF1_AcosFast:
-			DstToken = TEXT("ACosFastDeriv(") + SrcToken + TEXT(")");
-			break;
-		case EF1_Atan:
-			DstToken = TEXT("ATanDeriv(") + SrcToken + TEXT(")");
-			break;
-		case EF1_AtanFast:
-			DstToken = TEXT("ATanFastDeriv(") + SrcToken + TEXT(")");
-			break;
-		case EF1_Sqrt:
-			DstToken = TEXT("SqrtDeriv(") + SrcToken + TEXT(")");
-			break;
-		case EF1_Rcp:
-			DstToken = TEXT("RcpDeriv(") + SrcToken + TEXT(")");
-			break;
-		case EF1_Rsqrt:
-			DstToken = TEXT("RsqrtDeriv(") + SrcToken + TEXT(")");
-			break;
-		case EF1_Saturate:
-			DstToken = TEXT("SaturateDeriv(") + SrcToken + TEXT(")");
-			break;
-		case EF1_Frac:
-			DstToken = TEXT("FracDeriv(") + SrcToken + TEXT(")");
-			break;
-		case EF1_Length:
-			DstToken = TEXT("LengthDeriv(") + SrcToken + TEXT(")");
-			break;
-		case EF1_InvLength:
-			DstToken = TEXT("InvLengthDeriv(") + SrcToken + TEXT(")");
-			break;
-		case EF1_Normalize:
-			DstToken = TEXT("NormalizeDeriv(") + SrcToken + TEXT(")");
-			break;
-		default:
-			check(0);
-			break;
-		}
-
-		DstTokens[CompiledPDV_Analytic] = DstToken;
-	}
-
-	EMaterialValueType DstMatType = GetMaterialTypFromDerivTypeIndex(OutputTypeIndex);
-	int32 Ret = Translator.AddCodeChunkInnerDeriv(*DstTokens[CompiledPDV_FiniteDifferences],*DstTokens[CompiledPDV_Analytic],DstMatType,false /*?*/, DstStatus);
-	return Ret;
-}
-
-int32 FMaterialDerivativeAutogen::GenerateExpressionFunc2(FHLSLMaterialTranslator& Translator, EFunc2 Op, int32 LhsCode, int32 RhsCode)
-{
-	if (LhsCode == INDEX_NONE || RhsCode == INDEX_NONE)
-	{
-		return INDEX_NONE;
-	}
-
-
-	const FDerivInfo LhsDerivInfo = Translator.GetDerivInfo(LhsCode);
-	const FDerivInfo RhsDerivInfo = Translator.GetDerivInfo(RhsCode);
-
-	int32 IntermediaryTypeIndex = FMath::Max<int32>(LhsDerivInfo.TypeIndex, RhsDerivInfo.TypeIndex);
-
-	if (Op == EF2_Cross)
-	{
-		IntermediaryTypeIndex = 2;
-	}
-
-	if (Op == EF2_Fmod)
-	{
-		check(RhsDerivInfo.DerivativeStatus == DS_Zero);
-	}
-
-	int32 OutputTypeIndex = GetFunc2ReturnNumComponents(LhsDerivInfo.TypeIndex, RhsDerivInfo.TypeIndex, Op);
-
-	if (OutputTypeIndex < 0)
-	{
-		return INDEX_NONE;
-	}
-
-	EDerivativeStatus DstStatus = DS_NotValid;
-
-	// Rules for derivatives:
-	// 1. If either the LHS or RHS is Not Valid or Not Aware, then the derivative is not valid. Run scalar route.
-	// 2. If both LHS and RHS are known to be Zero, then run raw code, and specify a known zero status.
-	// 3. If both LHS and RHS are Valid derivatives, then run deriv path.
-	// 4. If one is Valid and the other is known Zero, then promote the Zero to Valid, and run deriv path.
-
-	bool bUseScalarVersion = false;
-	bool bIsDerivValidZero = true;
-	bool bMakeDerivLhs = false;
-	bool bMakeDerivRhs = false;
-	if (!IsDerivativeValid(LhsDerivInfo.DerivativeStatus) || !IsDerivativeValid(RhsDerivInfo.DerivativeStatus))
-	{
-		// use scalar version as a fallback
-		bUseScalarVersion = true;
-		// derivative is not valid
-		bIsDerivValidZero = false;
-
-		// We output status as invalid, since one of the parameters is either not aware or not valid
-		DstStatus = DS_NotValid;
-
-	}
-	else if (LhsDerivInfo.DerivativeStatus == DS_Zero && RhsDerivInfo.DerivativeStatus == DS_Zero)
-	{
-		// use scalar version
-		bUseScalarVersion = true;
-		// since we know both incoming values have derivatives of zero, we know the output is zero
-		bIsDerivValidZero = true;
-		// we know that the value is zero
-		DstStatus = DS_Zero;
-	}
-	else
-	{
-		check(IsDerivativeValid(LhsDerivInfo.DerivativeStatus));
-		check(IsDerivativeValid(RhsDerivInfo.DerivativeStatus));
-
-		// use deriv version
-		bUseScalarVersion = false;
-		// we will be calculating a valid derivative, so this value doesn't matter, but make it false for consistency
-		bIsDerivValidZero = false;
-
-		// if the lhs has a derivitive of zero, and rhs is non-zero, convert lhs from a scalar type to deriv type
-		if (LhsDerivInfo.DerivativeStatus == DS_Zero)
-		{
-			bMakeDerivLhs = true;
-		}
-
-		// if the rhs has a derivative of zero, and lhs is non-zero, convert rhs from a scalar type to deriv type
-		if (RhsDerivInfo.DerivativeStatus == DS_Zero)
-		{
-			bMakeDerivRhs = true;
-		}
-
-		// derivative results will be valid
-		DstStatus = DS_Valid;
-	}
-
-	FString DstTokens[CompiledPDV_MAX];
-
-	for (int32 Index = 0; Index < CompiledPDV_MAX; Index++)
-	{
-		ECompiledPartialDerivativeVariation Variation = (ECompiledPartialDerivativeVariation)Index;
-
-		FString LhsToken = Translator.GetParameterCodeDeriv(LhsCode,Variation);
-		FString RhsToken = Translator.GetParameterCodeDeriv(RhsCode,Variation);
-
-		// The token is the symbol name. If we are in finite mode, that's all we have to do. But if
-		// we are in analytic mode, we may need to get the value.
-		if (Index == CompiledPDV_Analytic)
-		{
-			LhsToken = CoerceValueRaw(LhsToken, LhsDerivInfo.TypeIndex, LhsDerivInfo.DerivativeStatus, IntermediaryTypeIndex);
-			RhsToken = CoerceValueRaw(RhsToken, RhsDerivInfo.TypeIndex, RhsDerivInfo.DerivativeStatus, IntermediaryTypeIndex);
-		}
-
-		FString DstToken;
-		// just generate a type
-		switch(Op)
-		{
-		case EF2_Add:
-			DstToken = TEXT("(") + LhsToken + TEXT(" + ") + RhsToken + TEXT(")");
-			break;
-		case EF2_Sub:
-			DstToken = TEXT("(") + LhsToken + TEXT(" - ") + RhsToken + TEXT(")");
-			break;
-		case EF2_Mul:
-			DstToken = TEXT("(") + LhsToken + TEXT(" * ") + RhsToken + TEXT(")");
-			break;
-		case EF2_Div:
-			DstToken = TEXT("(") + LhsToken + TEXT(" / ") + RhsToken + TEXT(")");
-			break;
-		case EF2_Fmod:
-			DstToken = TEXT("fmod(") + LhsToken + TEXT(",") + RhsToken + TEXT(")");
-			break;
-		case EF2_Min:
-			DstToken = TEXT("min(") + LhsToken + TEXT(",") + RhsToken + TEXT(")");
-			break;
-		case EF2_Max:
-			DstToken = TEXT("max(") + LhsToken + TEXT(",") + RhsToken + TEXT(")");
-			break;
-		case EF2_Dot:
-			DstToken = TEXT("dot(") + LhsToken + TEXT(",") + RhsToken + TEXT(")");
-			break;
-		case EF2_Pow:
-			DstToken = TEXT("pow(") + LhsToken + TEXT(",") + RhsToken + TEXT(")");
-			break;
-		case EF2_PowPositiveClamped:
-			DstToken = TEXT("PositiveClampedPow(") + LhsToken + TEXT(",") + RhsToken + TEXT(")");
-			break;
-		case EF2_Atan2:
-			DstToken = TEXT("atan2(") + LhsToken + TEXT(",") + RhsToken + TEXT(")");
-			break;
-		case EF2_Atan2Fast:
-			DstToken = TEXT("atan2Fast(") + LhsToken + TEXT(",") + RhsToken + TEXT(")");
-			break;
-		case EF2_Cross:
-			DstToken = TEXT("cross(") + LhsToken + TEXT(",") + RhsToken + TEXT(")");
-			break;
-		default:
-			check(0);
-			break;
-		}
-
-		DstTokens[Index] = DstToken;
-	}
-
-	if (!bUseScalarVersion)
-	{
-		FString LhsToken = Translator.GetParameterCodeDeriv(LhsCode,CompiledPDV_Analytic);
-		FString RhsToken = Translator.GetParameterCodeDeriv(RhsCode,CompiledPDV_Analytic);
-
-		LhsToken = CoerceValueDeriv(LhsToken, LhsDerivInfo.TypeIndex, LhsDerivInfo.DerivativeStatus, IntermediaryTypeIndex);
-		RhsToken = CoerceValueDeriv(RhsToken, RhsDerivInfo.TypeIndex, RhsDerivInfo.DerivativeStatus, IntermediaryTypeIndex);
-
-		check(Op < EF2_Num);
-		bFunc2OpIsEnabled[Op][IntermediaryTypeIndex] = true;
-
-		FString DstToken;
-		switch(Op)
-		{
-		case EF2_Add:
-			DstToken = TEXT("AddDeriv(") + LhsToken + TEXT(",") + RhsToken + TEXT(")");
-			break;
-		case EF2_Sub:
-			DstToken = TEXT("SubDeriv(") + LhsToken + TEXT(",") + RhsToken + TEXT(")");
-			break;
-		case EF2_Mul:
-			DstToken = TEXT("MulDeriv(") + LhsToken + TEXT(",") + RhsToken + TEXT(")");
-			break;
-		case EF2_Div:
-			DstToken = TEXT("DivDeriv(") + LhsToken + TEXT(",") + RhsToken + TEXT(")");
-			break;
-		case EF2_Fmod:
-			DstToken = TEXT("FmodDeriv(") + LhsToken + TEXT(",") + RhsToken + TEXT(")");
-			break;
-		case EF2_Dot:
-			DstToken = TEXT("DotDeriv(") + LhsToken + TEXT(",") + RhsToken + TEXT(")");
-			break;
-		case EF2_Min:
-			DstToken = TEXT("MinDeriv(") + LhsToken + TEXT(",") + RhsToken + TEXT(")");
-			break;
-		case EF2_Max:
-			DstToken = TEXT("MaxDeriv(") + LhsToken + TEXT(",") + RhsToken + TEXT(")");
-			break;
-		case EF2_Pow:
-			DstToken = TEXT("PowDeriv(") + LhsToken + TEXT(",") + RhsToken + TEXT(")");
-			break;
-		case EF2_PowPositiveClamped:
-			DstToken = TEXT("PowPositiveClampedDeriv(") + LhsToken + TEXT(",") + RhsToken + TEXT(")");
-			break;
-		case EF2_Cross:
-			DstToken = TEXT("CrossDeriv(") + LhsToken + TEXT(",") + RhsToken + TEXT(")");
-			break;
-		case EF2_Atan2:
-			DstToken = TEXT("Atan2Deriv(") + LhsToken + TEXT(",") + RhsToken + TEXT(")");
-			break;
-		case EF2_Atan2Fast:
-			DstToken = TEXT("Atan2FastDeriv(") + LhsToken + TEXT(",") + RhsToken + TEXT(")");
-			break;
-		default:
-			check(0);
-			break;
-		}
-
-		DstTokens[CompiledPDV_Analytic] = DstToken;
-		DstStatus = DS_Valid;
-	}
-
-	EMaterialValueType DstMatType = GetMaterialTypFromDerivTypeIndex(OutputTypeIndex);
-
-	int32 Ret = Translator.AddCodeChunkInnerDeriv(*DstTokens[CompiledPDV_FiniteDifferences],*DstTokens[CompiledPDV_Analytic],DstMatType,false /*?*/, DstStatus);
-	return Ret;
-}
-
-int32 FMaterialDerivativeAutogen::GenerateLerpFunc(FHLSLMaterialTranslator& Translator, int32 A, int32 B, int32 S)
-{
-	// TODO: generalize to Func3?
-	if (A == INDEX_NONE || B == INDEX_NONE || S == INDEX_NONE)
-	{
-		return INDEX_NONE;
-	}
-
-	const FDerivInfo ADerivInfo = Translator.GetDerivInfo(A);
-	const FDerivInfo BDerivInfo = Translator.GetDerivInfo(B);
-	const FDerivInfo SDerivInfo = Translator.GetDerivInfo(S);
-	
-	const EMaterialValueType ResultType = Translator.GetArithmeticResultType(A, B);
-	const EMaterialValueType AlphaType = (ResultType == SDerivInfo.Type) ? ResultType : MCT_Float1;
-	const uint32 NumResultComponents = GetNumComponents(ResultType);
-
-	const bool bAllZeroDeriv = (ADerivInfo.DerivativeStatus == DS_Zero && BDerivInfo.DerivativeStatus == DS_Zero && SDerivInfo.DerivativeStatus == DS_Zero);
-	FString FiniteString = FString::Printf(TEXT("lerp(%s,%s,%s)"), *Translator.CoerceParameter(A, ResultType), *Translator.CoerceParameter(B, ResultType), *Translator.CoerceParameter(S, AlphaType));
-	
-	if (!bAllZeroDeriv && IsDerivativeValid(ADerivInfo.DerivativeStatus) && IsDerivativeValid(BDerivInfo.DerivativeStatus) && IsDerivativeValid(SDerivInfo.DerivativeStatus))
-	{
-		const uint32 ResultTypeIndex = GetDerivTypeIndex(ResultType);
-		//const uint32 AlphaTypeIndex = GetDerivTypeIndex(AlphaType);
-		FString ADeriv = Translator.GetParameterCodeDeriv(A, CompiledPDV_Analytic);
-		FString BDeriv = Translator.GetParameterCodeDeriv(B, CompiledPDV_Analytic);
-		FString SDeriv = Translator.GetParameterCodeDeriv(S, CompiledPDV_Analytic);
-
-		ADeriv = CoerceValueDeriv(ADeriv, ADerivInfo.TypeIndex, ADerivInfo.DerivativeStatus, ResultTypeIndex);
-		BDeriv = CoerceValueDeriv(BDeriv, BDerivInfo.TypeIndex, BDerivInfo.DerivativeStatus, ResultTypeIndex);
-		SDeriv = CoerceValueDeriv(SDeriv, SDerivInfo.TypeIndex, SDerivInfo.DerivativeStatus, ResultTypeIndex);
-
-		FString AnalyticString = FString::Printf(TEXT("LerpDeriv(%s, %s, %s)"), *ADeriv, *BDeriv, *SDeriv);
-
-		check(NumResultComponents <= 4);
-		bLerpEnabled[ResultTypeIndex] = true;
-
-		return Translator.AddCodeChunkInnerDeriv(*FiniteString, *AnalyticString, ResultType, false, DS_Valid);
-	}
-	else
-	{
-		return Translator.AddCodeChunkInnerDeriv(*FiniteString, *FiniteString, ResultType, false, bAllZeroDeriv ? DS_Zero : DS_NotValid);
-	}
-}
-
-int32 FMaterialDerivativeAutogen::GenerateRotateScaleOffsetTexCoordsFunc(FHLSLMaterialTranslator& Translator, int32 TexCoord, int32 RotationScale, int32 Offset)
-{
-	// TODO: generalize to Func3?
-	if (TexCoord == INDEX_NONE || RotationScale == INDEX_NONE || Offset == INDEX_NONE)
-	{
-		return INDEX_NONE;
-	}
-	const FDerivInfo TexCoordDerivInfo		= Translator.GetDerivInfo(TexCoord);
-	const FDerivInfo RotationScaleDerivInfo	= Translator.GetDerivInfo(RotationScale);
-	const FDerivInfo OffsetDerivInfo		= Translator.GetDerivInfo(Offset);
-
-	const EMaterialValueType ResultType = MCT_Float2;
-	const uint32 NumResultComponents = 2;
-
-	const bool bAllZeroDeriv = (TexCoordDerivInfo.DerivativeStatus == DS_Zero && RotationScaleDerivInfo.DerivativeStatus == DS_Zero && OffsetDerivInfo.DerivativeStatus == DS_Zero);
-	FString FiniteString = FString::Printf(TEXT("RotateScaleOffsetTexCoords(%s, %s, %s.xy)"),	*Translator.CoerceParameter(TexCoord,		ResultType), 
-																								*Translator.CoerceParameter(RotationScale,	ResultType),
-																								*Translator.CoerceParameter(Offset,			ResultType));
-
-	if (!bAllZeroDeriv && IsDerivativeValid(TexCoordDerivInfo.DerivativeStatus) && IsDerivativeValid(RotationScaleDerivInfo.DerivativeStatus) && IsDerivativeValid(OffsetDerivInfo.DerivativeStatus))
-	{
-		const uint32 ResultTypeIndex = GetDerivTypeIndex(ResultType);
-		FString TexCoordDeriv		= Translator.GetParameterCodeDeriv(TexCoord, CompiledPDV_Analytic);
-		FString RotationScaleDeriv	= Translator.GetParameterCodeDeriv(RotationScale, CompiledPDV_Analytic);
-		FString OffsetDeriv			= Translator.GetParameterCodeDeriv(Offset, CompiledPDV_Analytic);
-
-		TexCoordDeriv		= CoerceValueDeriv(TexCoordDeriv,		TexCoordDerivInfo.TypeIndex,		TexCoordDerivInfo.DerivativeStatus,			ResultTypeIndex);
-		RotationScaleDeriv	= CoerceValueDeriv(RotationScaleDeriv,	RotationScaleDerivInfo.TypeIndex,	RotationScaleDerivInfo.DerivativeStatus,	ResultTypeIndex);
-		OffsetDeriv			= CoerceValueDeriv(OffsetDeriv,			OffsetDerivInfo.TypeIndex,			OffsetDerivInfo.DerivativeStatus,			ResultTypeIndex);
-
-		FString AnalyticString = FString::Printf(TEXT("RotateScaleOffsetTexCoordsDeriv(%s, %s, %s)"), *TexCoordDeriv, *RotationScaleDeriv, *OffsetDeriv);
-
-		check(NumResultComponents <= 4);
-		bRotateScaleOffsetTexCoords = true;
-
-		return Translator.AddCodeChunkInnerDeriv(*FiniteString, *AnalyticString, ResultType, false, DS_Valid);
-	}
-	else
-	{
-		return Translator.AddCodeChunkInnerDeriv(*FiniteString, *FiniteString, ResultType, false, bAllZeroDeriv ? DS_Zero : DS_NotValid);
-	}
-}
-
-int32 FMaterialDerivativeAutogen::GenerateIfFunc(FHLSLMaterialTranslator& Translator, int32 A, int32 B, int32 Greater, int32 Equal, int32 Less, int32 Threshold)
-{
-	if (A == INDEX_NONE || B == INDEX_NONE || Greater == INDEX_NONE || Less == INDEX_NONE)
-	{
-		return INDEX_NONE;
-	}
-
-	const FString AFinite = Translator.GetParameterCode(A);
-	const FString BFinite = Translator.GetParameterCode(B);
-
-	const bool bEqual = (Equal != INDEX_NONE);
-
-	if (bEqual && Threshold == INDEX_NONE)
-		return INDEX_NONE;
-
-	EMaterialValueType ResultType = Translator.GetArithmeticResultType(Greater, Less);
-
-	if (bEqual)
-		ResultType = Translator.GetArithmeticResultType(ResultType, Translator.GetParameterType(Greater));
-	
-	Greater = Translator.ForceCast(Greater,		ResultType);
-	Less	= Translator.ForceCast(Less,		ResultType);
-	if (Greater == INDEX_NONE || Less == INDEX_NONE)
-		return INDEX_NONE;
-
-	const FString GreaterFinite = Translator.GetParameterCode(Greater);
-	const FString LessFinite	= Translator.GetParameterCode(Less);
-
-	if (bEqual)
-	{
-		Equal = Translator.ForceCast(Equal, ResultType);
-		if (Equal == INDEX_NONE)
-			return INDEX_NONE;
-	}
-	
-	const FDerivInfo GreaterDerivInfo	= Translator.GetDerivInfo(Greater);
-	const FDerivInfo LessDerivInfo		= Translator.GetDerivInfo(Less);
-
-	FString CodeFinite;
-	FString ThresholdFinite;
-	if (bEqual)
-	{
-		ThresholdFinite = *Translator.GetParameterCode(Threshold);
-
-		CodeFinite = FString::Printf(
-			TEXT("((abs(%s - %s) > %s) ? (%s >= %s ? %s : %s) : %s)"),
-			*AFinite, *BFinite, *ThresholdFinite,
-			*AFinite, *BFinite,
-			*GreaterFinite, *LessFinite, *Translator.GetParameterCode(Equal));
-	}
-	else
-	{
-		CodeFinite = FString::Printf(
-			TEXT("((%s >= %s) ? %s : %s)"),
-			*AFinite, *BFinite,
-			*GreaterFinite, *LessFinite
-		);
-	}
-
-	const bool bAllDerivValid = IsDerivativeValid(GreaterDerivInfo.DerivativeStatus) && IsDerivativeValid(LessDerivInfo.DerivativeStatus) && (!bEqual || IsDerivativeValid(Translator.GetDerivativeStatus(Equal)));
-	const bool bAllDerivZero = GreaterDerivInfo.DerivativeStatus == DS_Zero && LessDerivInfo.DerivativeStatus == DS_Zero && (!bEqual || Translator.GetDerivativeStatus(Equal) == DS_Zero);
-	if (bAllDerivValid && !bAllDerivZero)
-	{
-		const uint32 ResultTypeIndex = GetDerivTypeIndex(ResultType);
-
-		FString GreaterDeriv	= Translator.GetParameterCodeDeriv(Greater, CompiledPDV_Analytic);
-		FString LessDeriv		= Translator.GetParameterCodeDeriv(Less,	CompiledPDV_Analytic);
-
-		GreaterDeriv	= CoerceValueDeriv(GreaterDeriv,	ResultTypeIndex, GreaterDerivInfo.DerivativeStatus, ResultTypeIndex);
-		LessDeriv		= CoerceValueDeriv(LessDeriv,		ResultTypeIndex, LessDerivInfo.DerivativeStatus,	ResultTypeIndex);
-
-		FString CodeAnalytic;
-		if (bEqual)
-		{
-			const FDerivInfo EqualDerivInfo = Translator.GetDerivInfo(Equal);
-			FString EqualDeriv = Translator.GetParameterCodeDeriv(Equal, CompiledPDV_Analytic);
-			EqualDeriv = CoerceValueDeriv(EqualDeriv, ResultTypeIndex, EqualDerivInfo.DerivativeStatus, ResultTypeIndex);
-
-			CodeAnalytic = FString::Printf(TEXT("IfDeriv(%s, %s, %s, %s, %s, %s)"), *AFinite, *BFinite, *GreaterDeriv, *LessDeriv, *EqualDeriv, *ThresholdFinite);
-			bIf2Enabled[ResultTypeIndex] = true;
-		}
-		else
-		{
-			CodeAnalytic = FString::Printf(TEXT("IfDeriv(%s, %s, %s, %s)"), *AFinite, *BFinite, *GreaterDeriv, *LessDeriv);
-			bIfEnabled[ResultTypeIndex] = true;
-		}
-
-		return Translator.AddCodeChunkInnerDeriv(*CodeFinite, *CodeAnalytic, ResultType, false, DS_Valid);
-
-	}
-	else
-	{
-		return Translator.AddCodeChunkInnerDeriv(*CodeFinite, *CodeFinite, ResultType, false, bAllDerivZero ? DS_Zero : DS_NotValid);
-	}
-}
-
-static FString CoerceFloat(const TCHAR* Value, int32 DstType, int32 SrcType)
-{
-	if (DstType == SrcType)
-	{
-		return Value;
-	}
-
-	if (SrcType == 0)
-	{
-		const TCHAR* Mask[4] = { TEXT("x"), TEXT("xx"), TEXT("xxx"), TEXT("xxxx") };
-		return FString::Printf(TEXT("%s.%s"), Value, Mask[DstType]);
-	}
-
-	
-	if (DstType < SrcType)
-	{
-		const TCHAR* Mask[4] = { TEXT("x"), TEXT("xy"), TEXT("xyz"), TEXT("xyzw") };
-		return FString::Printf(TEXT("%s.%s"), Value, Mask[DstType]);
-	}
-	else
-	{
-		check(DstType > SrcType);
-		const TCHAR* Zeros[4] = { TEXT("0.0f"), TEXT("0.0f, 0.0f"), TEXT("0.0f, 0.0f, 0.0f"), TEXT("0.0f, 0.0f, 0.0f, 0.0f") };
-		return FString::Printf(TEXT("%s(%s, %s)"), GetFloatVectorName(DstType), Value, Zeros[DstType - SrcType - 1]);
-	}
-}
-
-FString FMaterialDerivativeAutogen::GenerateUsedFunctions(FHLSLMaterialTranslator& Translator)
-{
-	// Certain derivative functions rely on other derivative functions. For example, Dot() requires Mul() and Add(). So if (for example) dot is enabled, then enable mul1 and add1.
-	EnableGeneratedDepencencies();
-
-	FString Ret;
-
-	// The basic structs (FloatDeriv, FloatDeriv2, FloatDeriv3, FloatDeriv4)
-	// It's not worth keeping track of all the times these are used, just make them.
-	for (int32 Index = 0; Index < 4; Index++)
-	{
-		FString BaseName = GetDerivVectorName(Index);
-		FString FieldName = GetFloatVectorName(Index);
-
-		Ret += TEXT("struct ") + BaseName + LINE_TERMINATOR;
-		Ret += TEXT("{") LINE_TERMINATOR;
-		Ret += TEXT("\t") + FieldName + TEXT(" Value;") LINE_TERMINATOR;
-		Ret += TEXT("\t") + FieldName + TEXT(" Ddx;") LINE_TERMINATOR;
-		Ret += TEXT("\t") + FieldName + TEXT(" Ddy;") LINE_TERMINATOR;
-		Ret += TEXT("};") LINE_TERMINATOR;
-		Ret += TEXT("") LINE_TERMINATOR;
-	}
-
-	// Full FloatDerivX constructors with explicit derivatives.
-	for (int32 Index = 0; Index < 4; Index++)
-	{
-		if (bConstructDerivEnabled[Index] || IsDebugGenerateAllFunctionsEnabled())
-		{
-			FString BaseName = GetDerivVectorName(Index);
-			FString FieldName = GetFloatVectorName(Index);
-
-			Ret += BaseName + TEXT(" Construct") + BaseName + TEXT("(") + FieldName + TEXT(" InValue,") + FieldName + TEXT(" InDdx,") + FieldName + TEXT(" InDdy)") LINE_TERMINATOR;
-			Ret += TEXT("{") LINE_TERMINATOR;
-			Ret += TEXT("\t") + BaseName + TEXT(" Ret;") LINE_TERMINATOR;
-			Ret += TEXT("\tRet.Value = InValue;") LINE_TERMINATOR;
-			Ret += TEXT("\tRet.Ddx = InDdx;") LINE_TERMINATOR;
-			Ret += TEXT("\tRet.Ddy = InDdy;") LINE_TERMINATOR;
-			Ret += TEXT("\treturn Ret;") LINE_TERMINATOR;
-			Ret += TEXT("}") LINE_TERMINATOR;
-			Ret += TEXT("") LINE_TERMINATOR;
-		}
-	}
-
-	// FloatDerivX constructors from constant floatX.
-	for (int32 Index = 0; Index < 4; Index++)
-	{
-		if (bConstructConstantDerivEnabled[Index] || IsDebugGenerateAllFunctionsEnabled())
-		{
-			FString BaseName = GetDerivVectorName(Index);
-			FString FieldName = GetFloatVectorName(Index);
-
-			Ret += BaseName + TEXT(" ConstructConstant") + BaseName + TEXT("(") + FieldName + TEXT(" Value)") LINE_TERMINATOR;
-			Ret += TEXT("{") LINE_TERMINATOR;
-			Ret += TEXT("\t") + BaseName + TEXT(" Ret;") LINE_TERMINATOR;
-			Ret += TEXT("\tRet.Value = Value;") LINE_TERMINATOR;
-			Ret += TEXT("\tRet.Ddx = 0;") LINE_TERMINATOR;
-			Ret += TEXT("\tRet.Ddy = 0;") LINE_TERMINATOR;
-			Ret += TEXT("\treturn Ret;") LINE_TERMINATOR;
-			Ret += TEXT("}") LINE_TERMINATOR;
-			Ret += TEXT("") LINE_TERMINATOR;
-		}
-	}
-
-	// FloatDerivX constructor from floatX with implicit derivatives.
-	for (int32 Index = 0; Index < 4; Index++)
-	{
-		if (bConstructFiniteDerivEnabled[Index] || IsDebugGenerateAllFunctionsEnabled())
-		{
-			FString BaseName = GetDerivVectorName(Index);
-			FString FieldName = GetFloatVectorName(Index);
-
-			Ret += BaseName + TEXT(" ConstructFinite") + BaseName + TEXT("(") + FieldName + TEXT(" InValue)") LINE_TERMINATOR;
-			Ret += TEXT("{") LINE_TERMINATOR;
-			Ret += TEXT("\t") + BaseName + TEXT(" Ret;") LINE_TERMINATOR;
-			Ret += TEXT("\tRet.Value = InValue;") LINE_TERMINATOR;
-			Ret += TEXT("\tRet.Ddx = ddx(InValue);") LINE_TERMINATOR;
-			Ret += TEXT("\tRet.Ddy = ddy(InValue);") LINE_TERMINATOR;
-			Ret += TEXT("\treturn Ret;") LINE_TERMINATOR;
-			Ret += TEXT("}") LINE_TERMINATOR;
-			Ret += TEXT("") LINE_TERMINATOR;
-		}
-	}
-
-	// Convert between FloatDeriv types
-	for (int32 DstIndex = 0; DstIndex < 4; DstIndex++)
-	{
-		for (int32 SrcIndex = 0; SrcIndex < 4; SrcIndex++)
-		{
-			if (SrcIndex == DstIndex)
-				continue;
-
-			if (bConvertDerivEnabled[DstIndex][SrcIndex] || IsDebugGenerateAllFunctionsEnabled())
-			{
-				FString DstBaseName = GetDerivVectorName(DstIndex);
-				FString SrcBaseName = GetDerivVectorName(SrcIndex);
-				FString ScalarName = GetDerivVectorName(0);
-
-				Ret += DstBaseName + TEXT(" Convert") + DstBaseName + TEXT("(") + SrcBaseName + TEXT(" Src)") LINE_TERMINATOR;
-				Ret += TEXT("{") LINE_TERMINATOR;
-				Ret += TEXT("\t") + DstBaseName + TEXT(" Ret;") LINE_TERMINATOR;
-				Ret += TEXT("\tRet.Value = ") + CoerceFloat(TEXT("Src.Value"), DstIndex, SrcIndex) + TEXT(";") LINE_TERMINATOR;
-				Ret += TEXT("\tRet.Ddx = ") + CoerceFloat(TEXT("Src.Ddx"), DstIndex, SrcIndex) + TEXT(";") LINE_TERMINATOR;
-				Ret += TEXT("\tRet.Ddy = ") + CoerceFloat(TEXT("Src.Ddy"), DstIndex, SrcIndex) + TEXT(";") LINE_TERMINATOR;
-				Ret += TEXT("\treturn Ret;") LINE_TERMINATOR;
-				Ret += TEXT("}") LINE_TERMINATOR;
-				Ret += TEXT("") LINE_TERMINATOR;
-			}
-		}
-	}
-
-	const TCHAR* SwizzleList[4] =
-	{
-		TEXT("x"),
-		TEXT("y"),
-		TEXT("z"),
-		TEXT("w")
-	};
-
-	// Extract single FloatDeriv element from FloatDerivX
-	for (int32 StructIndex = 0; StructIndex < 4; StructIndex++)
-	{
-		for (int32 ElemIndex = 0; ElemIndex <= StructIndex; ElemIndex++)
-		{
-			if (bExtractIndexEnabled[StructIndex] || IsDebugGenerateAllFunctionsEnabled())
-			{
-				// i.e. can't grab the 4th component of a float3
-				check(ElemIndex <= StructIndex);
-
-				FString BaseName = GetDerivVectorName(StructIndex);
-
-				FString ElemStr = FString::Printf(TEXT("%d"), ElemIndex + 1);
-
-				Ret += TEXT("FloatDeriv Extract") + BaseName + TEXT("_") + ElemStr + TEXT("(") + BaseName + TEXT(" InValue)") LINE_TERMINATOR;
-				Ret += TEXT("{") LINE_TERMINATOR;
-				Ret += TEXT("\tFloatDeriv Ret;") LINE_TERMINATOR;
-				Ret += TEXT("\tRet.Value = InValue.Value.") + FString(SwizzleList[ElemIndex]) + TEXT(";") LINE_TERMINATOR;
-				Ret += TEXT("\tRet.Ddx = InValue.Ddx.") + FString(SwizzleList[ElemIndex]) + TEXT(";") LINE_TERMINATOR;
-				Ret += TEXT("\tRet.Ddy = InValue.Ddy.") + FString(SwizzleList[ElemIndex]) + TEXT(";") LINE_TERMINATOR;
-				Ret += TEXT("\treturn Ret;") LINE_TERMINATOR;
-				Ret += TEXT("}") LINE_TERMINATOR;
-				Ret += TEXT("") LINE_TERMINATOR;
-			}
-		}
-	}
-
-
-	// SelectElemHelper
-	for (int32 Index = 0; Index < 4; Index++)
-	{
-		if (bSelectElemHelperEnabled[Index] || IsDebugGenerateAllFunctionsEnabled())
-		{
-			FString FieldName = GetFloatVectorName(Index);
-			FString BoolName = GetBoolVectorName(Index);
-
-			Ret += FieldName + TEXT(" SelectElemHelper(") + BoolName + TEXT(" Cmp, ") + FieldName + TEXT(" Lhs, ") + FieldName + TEXT(" Rhs)") LINE_TERMINATOR;
-			Ret += TEXT("{") LINE_TERMINATOR;
-			Ret += TEXT("\t") + FieldName + TEXT(" Ret = 0.0f;") LINE_TERMINATOR;
-			for (int32 Component = 0; Component <= Index; Component++)
-			{
-				const TCHAR* ChanName = SwizzleList[Component];
-				Ret += FString::Printf(TEXT("\tRet.%s = Cmp.%s ? Lhs.%s : Rhs.%s;"),ChanName,ChanName,ChanName,ChanName) + LINE_TERMINATOR;
-			}
-
-			Ret += TEXT("\treturn Ret;") LINE_TERMINATOR;
-			Ret += TEXT("}") LINE_TERMINATOR;
-			Ret += TEXT("") LINE_TERMINATOR;
-		}
-	}
-
-	// Func2s
-	for (int32 Op = 0; Op < EF2_Num; Op++)
-	{
-		for (int32 Index = 0; Index < 4; Index++)
-		{
-			if (bFunc2OpIsEnabled[Op][Index] || IsDebugGenerateAllFunctionsEnabled())
-			{
-				FString BaseName = GetDerivVectorName(Index);
-				FString FieldName = GetFloatVectorName(Index);
-				FString BoolName = GetBoolVectorName(Index);
-
-				switch(Op)
-				{
-				case EF2_Add:
-					Ret += BaseName + TEXT(" AddDeriv(") + BaseName + TEXT(" A, ") + BaseName + TEXT(" B)") LINE_TERMINATOR;
-					Ret += TEXT("{") LINE_TERMINATOR;
-					Ret += TEXT("\t") + BaseName + TEXT(" Ret;") LINE_TERMINATOR;
-					Ret += TEXT("\tRet.Value = A.Value + B.Value;") LINE_TERMINATOR;
-					Ret += TEXT("\tRet.Ddx = A.Ddx + B.Ddx;") LINE_TERMINATOR;
-					Ret += TEXT("\tRet.Ddy = A.Ddy + B.Ddy;") LINE_TERMINATOR;
-					Ret += TEXT("\treturn Ret;") LINE_TERMINATOR;
-					Ret += TEXT("}") LINE_TERMINATOR;
-					Ret += TEXT("") LINE_TERMINATOR;
-					break;
-				case EF2_Sub:
-					Ret += BaseName + TEXT(" SubDeriv(") + BaseName + TEXT(" A, ") + BaseName + TEXT(" B)") LINE_TERMINATOR;
-					Ret += TEXT("{") LINE_TERMINATOR;
-					Ret += TEXT("\t") + BaseName + TEXT(" Ret;") LINE_TERMINATOR;
-					Ret += TEXT("\tRet.Value = A.Value - B.Value;") LINE_TERMINATOR;
-					Ret += TEXT("\tRet.Ddx = A.Ddx - B.Ddx;") LINE_TERMINATOR;
-					Ret += TEXT("\tRet.Ddy = A.Ddy - B.Ddy;") LINE_TERMINATOR;
-					Ret += TEXT("\treturn Ret;") LINE_TERMINATOR;
-					Ret += TEXT("}") LINE_TERMINATOR;
-					Ret += TEXT("") LINE_TERMINATOR;
-					break;
-				case EF2_Mul:
-					Ret += BaseName + TEXT(" MulDeriv(") + BaseName + TEXT(" A, ") + BaseName + TEXT(" B)") LINE_TERMINATOR;
-					Ret += TEXT("{") LINE_TERMINATOR;
-					Ret += TEXT("\t") + BaseName + TEXT(" Ret;") LINE_TERMINATOR;
-					Ret += TEXT("\tRet.Value = A.Value * B.Value;") LINE_TERMINATOR;
-					Ret += TEXT("\tRet.Ddx = A.Ddx * B.Value + A.Value * B.Ddx;") LINE_TERMINATOR;
-					Ret += TEXT("\tRet.Ddy = A.Ddy * B.Value + A.Value * B.Ddy;") LINE_TERMINATOR;
-					Ret += TEXT("\treturn Ret;") LINE_TERMINATOR;
-					Ret += TEXT("}") LINE_TERMINATOR;
-					Ret += TEXT("") LINE_TERMINATOR;
-					break;
-				case EF2_Div:
-					Ret += BaseName + TEXT(" DivDeriv(") + BaseName + TEXT(" A, ") + BaseName + TEXT(" B)") LINE_TERMINATOR;
-					Ret += TEXT("{") LINE_TERMINATOR;
-					Ret += TEXT("\t") + BaseName + TEXT(" Ret;") LINE_TERMINATOR;
-					Ret += TEXT("\tRet.Value = A.Value / B.Value;") LINE_TERMINATOR;
-					Ret += TEXT("\t") + FieldName + TEXT(" Denom = rcp(B.Value * B.Value);") LINE_TERMINATOR;
-					Ret += TEXT("\t") + FieldName + TEXT(" dFdA =  B.Value * Denom;") LINE_TERMINATOR;
-					Ret += TEXT("\t") + FieldName + TEXT(" dFdB = -A.Value * Denom;") LINE_TERMINATOR;
-					Ret += TEXT("\tRet.Ddx = dFdA * A.Ddx + dFdB * B.Ddx;") LINE_TERMINATOR;
-					Ret += TEXT("\tRet.Ddy = dFdA * A.Ddy + dFdB * B.Ddy;") LINE_TERMINATOR;
-					Ret += TEXT("\treturn Ret;") LINE_TERMINATOR;
-					Ret += TEXT("}") LINE_TERMINATOR;
-					Ret += TEXT("") LINE_TERMINATOR;
-					break;
-				case EF2_Fmod:
-					// Only valid when B derivatives are zero.
-					// We can't really do anything meaningful in the non-zero case.
-					Ret += BaseName + TEXT(" FmodDeriv(") + BaseName + TEXT(" A, ") + BaseName + TEXT(" B)") LINE_TERMINATOR;
-					Ret += TEXT("{") LINE_TERMINATOR;
-					Ret += TEXT("\t") + BaseName + TEXT(" Ret;") LINE_TERMINATOR;
-					Ret += TEXT("\tRet.Value = fmod(A.Value, B.Value);") LINE_TERMINATOR;
-					Ret += TEXT("\tRet.Ddx = A.Ddx;") LINE_TERMINATOR;
-					Ret += TEXT("\tRet.Ddy = A.Ddy;") LINE_TERMINATOR;
-					Ret += TEXT("\treturn Ret;") LINE_TERMINATOR;
-					Ret += TEXT("}") LINE_TERMINATOR;
-					Ret += TEXT("") LINE_TERMINATOR;
-					break;
-				case EF2_Min:
-					Ret += BaseName + TEXT(" MinDeriv(") + BaseName + TEXT(" A, ") + BaseName + TEXT(" B)") LINE_TERMINATOR;
-					Ret += TEXT("{") LINE_TERMINATOR;
-					Ret += TEXT("\t") + BaseName + TEXT(" Ret;") LINE_TERMINATOR;
-					Ret += TEXT("\t") + BoolName + TEXT(" Cmp = A.Value < B.Value;") LINE_TERMINATOR;
-					Ret += TEXT("\tRet.Value = SelectElemHelper(Cmp, A.Value, B.Value);") LINE_TERMINATOR;
-					Ret += TEXT("\tRet.Ddx = SelectElemHelper(Cmp, A.Ddx, B.Ddx);") LINE_TERMINATOR;
-					Ret += TEXT("\tRet.Ddy = SelectElemHelper(Cmp, A.Ddy, B.Ddy);") LINE_TERMINATOR;
-					Ret += TEXT("\treturn Ret;") LINE_TERMINATOR;
-					Ret += TEXT("}") LINE_TERMINATOR;
-					Ret += TEXT("") LINE_TERMINATOR;
-					break;
-				case EF2_Max:
-					Ret += BaseName + TEXT(" MaxDeriv(") + BaseName + TEXT(" A, ") + BaseName + TEXT(" B)") LINE_TERMINATOR;
-					Ret += TEXT("{") LINE_TERMINATOR;
-					Ret += TEXT("\t") + BaseName + TEXT(" Ret;") LINE_TERMINATOR;
-					Ret += TEXT("\t") + BoolName + TEXT(" Cmp = A.Value > B.Value;") LINE_TERMINATOR;
-					Ret += TEXT("\tRet.Value = SelectElemHelper(Cmp, A.Value, B.Value);") LINE_TERMINATOR;
-					Ret += TEXT("\tRet.Ddx = SelectElemHelper(Cmp, A.Ddx, B.Ddx);") LINE_TERMINATOR;
-					Ret += TEXT("\tRet.Ddy = SelectElemHelper(Cmp, A.Ddy, B.Ddy);") LINE_TERMINATOR;
-					Ret += TEXT("\treturn Ret;") LINE_TERMINATOR;
-					Ret += TEXT("}") LINE_TERMINATOR;
-					Ret += TEXT("") LINE_TERMINATOR;
-					break;
-				case EF2_Dot:
-					Ret += TEXT("FloatDeriv DotDeriv(") + BaseName + TEXT(" A, ") + BaseName + TEXT(" B)") LINE_TERMINATOR;
-					Ret += TEXT("{") LINE_TERMINATOR;
-					Ret += TEXT("\tFloatDeriv Ret = ConstructConstantFloatDeriv(0);") LINE_TERMINATOR;
-					for (int32 Component = 0; Component <= Index; Component++)
-					{
-						Ret += FString::Printf(TEXT("\tRet = AddDeriv(Ret,MulDeriv(Extract%s_%d(A),Extract%s_%d(B)));"), *BaseName, Component + 1, *BaseName, Component + 1) + LINE_TERMINATOR;
-					}
-					Ret += TEXT("\treturn Ret;") LINE_TERMINATOR;
-					Ret += TEXT("}") LINE_TERMINATOR;
-					Ret += TEXT("") LINE_TERMINATOR;
-					break;
-				case EF2_Pow:
-					// pow(A,B) = exp(B*log(A))
-					//     pow'(A,B) = exp(B*log(A)) * (B'*log(A) + (B/A)*A')
-					//     pow'(A,B) = pow(A,B) * (B'*log(A) + (B/A)*A')
-					// sanity check when B is constant and A is a linear function (B'=0,A'=1)
-					//     pow'(A,B) = pow(A,B) * (0*log(A) + (B/A)*1)
-					//     pow'(A,B) = B * pow(A,B-1)
-					Ret += BaseName + TEXT(" PowDeriv(") + BaseName + TEXT(" A, ") + BaseName + TEXT(" B)") LINE_TERMINATOR;
-					Ret += TEXT("{") LINE_TERMINATOR;
-					Ret += TEXT("\t") + BaseName + TEXT(" Ret;") LINE_TERMINATOR;
-					Ret += TEXT("\tRet.Value = pow(A.Value, B.Value);") LINE_TERMINATOR;
-					Ret += TEXT("\tRet.Ddx = Ret.Value * (B.Ddx * log(A.Value) + (B.Value/A.Value)*A.Ddx);") LINE_TERMINATOR;
-					Ret += TEXT("\tRet.Ddy = Ret.Value * (B.Ddy * log(A.Value) + (B.Value/A.Value)*A.Ddy);") LINE_TERMINATOR;
-					Ret += TEXT("\treturn Ret;") LINE_TERMINATOR;
-					Ret += TEXT("}") LINE_TERMINATOR;
-					Ret += TEXT("") LINE_TERMINATOR;
-					break;
-				case EF2_PowPositiveClamped:
-					Ret += BaseName + TEXT(" PowPositiveClampedDeriv(") + BaseName + TEXT(" A, ") + BaseName + TEXT(" B)") LINE_TERMINATOR;
-					Ret += TEXT("{") LINE_TERMINATOR;
-					Ret += TEXT("\t") + BaseName + TEXT(" Ret;") LINE_TERMINATOR;
-					Ret += TEXT("\t") + BoolName + TEXT(" InRange = (0.0 < B.Value);") LINE_TERMINATOR; // should we check for A as well?
-					Ret += TEXT("\t") + FieldName + TEXT(" Zero = 0.0;") LINE_TERMINATOR;
-					Ret += TEXT("\tRet.Value = PositiveClampedPow(A.Value, B.Value);") LINE_TERMINATOR;
-					Ret += TEXT("\tRet.Ddx = Ret.Value * (B.Ddx * log(A.Value) + (B.Value/A.Value)*A.Ddx);") LINE_TERMINATOR;
-					Ret += TEXT("\tRet.Ddy = Ret.Value * (B.Ddy * log(A.Value) + (B.Value/A.Value)*A.Ddy);") LINE_TERMINATOR;
-					Ret += TEXT("\tRet.Ddx = SelectElemHelper(InRange,Ret.Ddx,Zero);") LINE_TERMINATOR;
-					Ret += TEXT("\tRet.Ddy = SelectElemHelper(InRange,Ret.Ddy,Zero);") LINE_TERMINATOR;
-					Ret += TEXT("\treturn Ret;") LINE_TERMINATOR;
-					Ret += TEXT("}") LINE_TERMINATOR;
-					Ret += TEXT("") LINE_TERMINATOR;
-					break;
-				case EF2_Atan2:
-					Ret += BaseName + TEXT(" Atan2Deriv(") + BaseName + TEXT(" A, ") + BaseName + TEXT(" B)") LINE_TERMINATOR;
-					Ret += TEXT("{") LINE_TERMINATOR;
-					Ret += TEXT("\t") + BaseName + TEXT(" Ret;") LINE_TERMINATOR;
-					Ret += TEXT("\tRet.Value = atan2(A.Value, B.Value);") LINE_TERMINATOR;
-					Ret += TEXT("\t") + FieldName + TEXT(" Denom = rcp(A.Value * A.Value + B.Value * B.Value);") LINE_TERMINATOR;
-					Ret += TEXT("\t") + FieldName + TEXT(" dFdA =  B.Value * Denom;") LINE_TERMINATOR;
-					Ret += TEXT("\t") + FieldName + TEXT(" dFdB = -A.Value * Denom;") LINE_TERMINATOR;
-					Ret += TEXT("\tRet.Ddx = dFdA * A.Ddx + dFdB * B.Ddx;") LINE_TERMINATOR;
-					Ret += TEXT("\tRet.Ddy = dFdA * A.Ddy + dFdB * B.Ddy;") LINE_TERMINATOR;
-					Ret += TEXT("\treturn Ret;") LINE_TERMINATOR;
-					Ret += TEXT("}") LINE_TERMINATOR;
-					Ret += TEXT("") LINE_TERMINATOR;
-					break;
-				case EF2_Atan2Fast:
-					Ret += BaseName + TEXT(" Atan2FastDeriv(") + BaseName + TEXT(" A, ") + BaseName + TEXT(" B)") LINE_TERMINATOR;
-					Ret += TEXT("{") LINE_TERMINATOR;
-					Ret += TEXT("\t") + BaseName + TEXT(" Ret;") LINE_TERMINATOR;
-					Ret += TEXT("\tRet.Value = atan2Fast(A.Value, B.Value);") LINE_TERMINATOR;
-					Ret += TEXT("\t") + FieldName + TEXT(" Denom = rcp(A.Value * A.Value + B.Value * B.Value);") LINE_TERMINATOR;
-					Ret += TEXT("\t") + FieldName + TEXT(" dFdA =  B.Value * Denom;") LINE_TERMINATOR;
-					Ret += TEXT("\t") + FieldName + TEXT(" dFdB = -A.Value * Denom;") LINE_TERMINATOR;
-					Ret += TEXT("\tRet.Ddx = dFdA * A.Ddx + dFdB * B.Ddx;") LINE_TERMINATOR;
-					Ret += TEXT("\tRet.Ddy = dFdA * A.Ddy + dFdB * B.Ddy;") LINE_TERMINATOR;
-					Ret += TEXT("\treturn Ret;") LINE_TERMINATOR;
-					Ret += TEXT("}") LINE_TERMINATOR;
-					Ret += TEXT("") LINE_TERMINATOR;
-					break;
-				case EF2_Cross:
-					if (Index == 2)
-					{
-						// (A*B)' = A' * B + A * B'
-						// Cross(A, B) = A.yzx * B.zxy - A.zxy * B.yzx;
-						// Cross(A, B)' = A.yzx' * B.zxy + A.yzx * B.zxy' - A.zxy' * B.yzx - A.zxy * B.yzx';
-						Ret += BaseName + TEXT(" CrossDeriv(") + BaseName + TEXT(" A, ") + BaseName + TEXT(" B)") LINE_TERMINATOR;
-						Ret += TEXT("{") LINE_TERMINATOR;
-						Ret += TEXT("\t") + BaseName + TEXT(" Ret;") LINE_TERMINATOR;
-						Ret += TEXT("\tRet.Value = cross(A.Value, B.Value);") LINE_TERMINATOR;
-						Ret += TEXT("\tRet.Ddx = A.Ddx.yzx * B.Value.zxy + A.Value.yzx * B.Ddx.zxy - A.Ddx.zxy * B.Value.yzx - A.Value.zxy * B.Ddx.yzx;") LINE_TERMINATOR;
-						Ret += TEXT("\tRet.Ddy = A.Ddy.yzx * B.Value.zxy + A.Value.yzx * B.Ddy.zxy - A.Ddy.zxy * B.Value.yzx - A.Value.zxy * B.Ddy.yzx;") LINE_TERMINATOR;
-						Ret += TEXT("\treturn Ret;") LINE_TERMINATOR;
-						Ret += TEXT("}") LINE_TERMINATOR;
-						Ret += TEXT("") LINE_TERMINATOR;
-					}
-					break;
-				default:
-					check(0);
-					break;
-				}
-			}
-		}
-	}
-
-	for (int32 Op = 0; Op < EF1_Num; Op++)
-	{
-		for (int32 Index = 0; Index < 4; Index++)
-		{
-			if (bFunc1OpIsEnabled[Op][Index] || IsDebugGenerateAllFunctionsEnabled())
-			{
-				FString BaseName	= GetDerivVectorName(Index);
-				FString FieldName	= GetFloatVectorName(Index);
-				FString BoolName	= GetBoolVectorName(Index);
-
-				switch(Op)
-				{
-				case EF1_Abs:
-					Ret += BaseName + TEXT(" AbsDeriv(") + BaseName + TEXT(" A)") LINE_TERMINATOR;
-					Ret += TEXT("{") LINE_TERMINATOR;
-					Ret += TEXT("\t") + BaseName + TEXT(" Ret;") LINE_TERMINATOR;
-					Ret += TEXT("\tRet.Value = abs(A.Value);") LINE_TERMINATOR;
-					Ret += TEXT("\t") + FieldName + TEXT(" dFdA = (A.Value >= 0.0f ? 1.0f : -1.0f);") LINE_TERMINATOR;
-					Ret += TEXT("\tRet.Ddx = dFdA * A.Ddx;") LINE_TERMINATOR;
-					Ret += TEXT("\tRet.Ddy = dFdA * A.Ddy;") LINE_TERMINATOR;
-					Ret += TEXT("\treturn Ret;") LINE_TERMINATOR;
-					Ret += TEXT("}") LINE_TERMINATOR;
-					Ret += TEXT("") LINE_TERMINATOR;
-					break;
-				case EF1_Sin:
-					Ret += BaseName + TEXT(" SinDeriv(") + BaseName + TEXT(" A)") LINE_TERMINATOR;
-					Ret += TEXT("{") LINE_TERMINATOR;
-					Ret += TEXT("\t") + BaseName + TEXT(" Ret;") LINE_TERMINATOR;
-					Ret += TEXT("\tRet.Value = sin(A.Value);") LINE_TERMINATOR;
-					Ret += TEXT("\t") + FieldName + TEXT(" dFdA = cos(A.Value);");
-					Ret += TEXT("\tRet.Ddx = dFdA * A.Ddx;") LINE_TERMINATOR;
-					Ret += TEXT("\tRet.Ddy = dFdA * A.Ddy;") LINE_TERMINATOR;
-					Ret += TEXT("\treturn Ret;") LINE_TERMINATOR;
-					Ret += TEXT("}") LINE_TERMINATOR;
-					Ret += TEXT("") LINE_TERMINATOR;
-					break;
-				case EF1_Cos:
-					Ret += BaseName + TEXT(" CosDeriv(") + BaseName + TEXT(" A)") LINE_TERMINATOR;
-					Ret += TEXT("{") LINE_TERMINATOR;
-					Ret += TEXT("\t") + BaseName + TEXT(" Ret;") LINE_TERMINATOR;
-					Ret += TEXT("\tRet.Value = cos(A.Value);") LINE_TERMINATOR;
-					Ret += TEXT("\t") + FieldName + TEXT(" dFdA = -sin(A.Value);");
-					Ret += TEXT("\tRet.Ddx = dFdA * A.Ddx;") LINE_TERMINATOR;
-					Ret += TEXT("\tRet.Ddy = dFdA * A.Ddy;") LINE_TERMINATOR;
-					Ret += TEXT("\treturn Ret;") LINE_TERMINATOR;
-					Ret += TEXT("}") LINE_TERMINATOR;
-					Ret += TEXT("") LINE_TERMINATOR;
-					break;
-				case EF1_Tan:
-					Ret += BaseName + TEXT(" TanDeriv(") + BaseName + TEXT(" A)") LINE_TERMINATOR;
-					Ret += TEXT("{") LINE_TERMINATOR;
-					Ret += TEXT("\t") + BaseName + TEXT(" Ret;") LINE_TERMINATOR;
-					Ret += TEXT("\tRet.Value = tan(A.Value);") LINE_TERMINATOR;
-					Ret += TEXT("\t") + FieldName + TEXT(" dFdA = rcp(cos(A.Value) * cos(A.Value));");
-					Ret += TEXT("\tRet.Ddx = dFdA * A.Ddx;") LINE_TERMINATOR;
-					Ret += TEXT("\tRet.Ddy = dFdA * A.Ddy;") LINE_TERMINATOR;
-					Ret += TEXT("\treturn Ret;") LINE_TERMINATOR;
-					Ret += TEXT("}") LINE_TERMINATOR;
-					Ret += TEXT("") LINE_TERMINATOR;
-					break;
-				case EF1_Asin:
-					Ret += BaseName + TEXT(" AsinDeriv(") + BaseName + TEXT(" A)") LINE_TERMINATOR;
-					Ret += TEXT("{") LINE_TERMINATOR;
-					Ret += TEXT("\t") + BaseName + TEXT(" Ret;") LINE_TERMINATOR;
-					Ret += TEXT("\tRet.Value = asin(A.Value);") LINE_TERMINATOR;
-					Ret += TEXT("\t") + FieldName + TEXT(" dFdA = rsqrt(max(1.0f - A.Value * A.Value, 0.00001f));");
-					Ret += TEXT("\tRet.Ddx = dFdA * A.Ddx;") LINE_TERMINATOR;
-					Ret += TEXT("\tRet.Ddy = dFdA * A.Ddy;") LINE_TERMINATOR;
-					Ret += TEXT("\treturn Ret;") LINE_TERMINATOR;
-					Ret += TEXT("}") LINE_TERMINATOR;
-					Ret += TEXT("") LINE_TERMINATOR;
-					break;
-				case EF1_AsinFast:
-					Ret += BaseName + TEXT(" AsinFastDeriv(") + BaseName + TEXT(" A)") LINE_TERMINATOR;
-					Ret += TEXT("{") LINE_TERMINATOR;
-					Ret += TEXT("\t") + BaseName + TEXT(" Ret;") LINE_TERMINATOR;
-					Ret += TEXT("\tRet.Value = asinFast(A.Value);") LINE_TERMINATOR;
-					Ret += TEXT("\t") + FieldName + TEXT(" dFdA = rsqrt(max(1.0f - A.Value * A.Value, 0.00001f));");
-					Ret += TEXT("\tRet.Ddx = dFdA * A.Ddx;") LINE_TERMINATOR;
-					Ret += TEXT("\tRet.Ddy = dFdA * A.Ddy;") LINE_TERMINATOR;
-					Ret += TEXT("\treturn Ret;") LINE_TERMINATOR;
-					Ret += TEXT("}") LINE_TERMINATOR;
-					Ret += TEXT("") LINE_TERMINATOR;
-					break;
-				case EF1_Acos:
-					Ret += BaseName + TEXT(" AcosDeriv(") + BaseName + TEXT(" A)") LINE_TERMINATOR;
-					Ret += TEXT("{") LINE_TERMINATOR;
-					Ret += TEXT("\t") + BaseName + TEXT(" Ret;") LINE_TERMINATOR;
-					Ret += TEXT("\tRet.Value = acos(A.Value);") LINE_TERMINATOR;
-					Ret += TEXT("\t") + FieldName + TEXT(" dFdA = -rsqrt(max(1.0f - A.Value * A.Value, 0.00001f));");
-					Ret += TEXT("\tRet.Ddx = dFdA * A.Ddx;") LINE_TERMINATOR;
-					Ret += TEXT("\tRet.Ddy = dFdA * A.Ddy;") LINE_TERMINATOR;
-					Ret += TEXT("\treturn Ret;") LINE_TERMINATOR;
-					Ret += TEXT("}") LINE_TERMINATOR;
-					Ret += TEXT("") LINE_TERMINATOR;
-					break;
-				case EF1_AcosFast:
-					Ret += BaseName + TEXT(" AcosFastDeriv(") + BaseName + TEXT(" A)") LINE_TERMINATOR;
-					Ret += TEXT("{") LINE_TERMINATOR;
-					Ret += TEXT("\t") + BaseName + TEXT(" Ret;") LINE_TERMINATOR;
-					Ret += TEXT("\tRet.Value = acosFast(A.Value);") LINE_TERMINATOR;
-					Ret += TEXT("\t") + FieldName + TEXT(" dFdA = -rsqrt(max(1.0f - A.Value * A.Value, 0.00001f));");
-					Ret += TEXT("\tRet.Ddx = dFdA * A.Ddx;") LINE_TERMINATOR;
-					Ret += TEXT("\tRet.Ddy = dFdA * A.Ddy;") LINE_TERMINATOR;
-					Ret += TEXT("\treturn Ret;") LINE_TERMINATOR;
-					Ret += TEXT("}") LINE_TERMINATOR;
-					Ret += TEXT("") LINE_TERMINATOR;
-					break;
-				case EF1_Atan:
-					Ret += BaseName + TEXT(" AtanDeriv(") + BaseName + TEXT(" A)") LINE_TERMINATOR;
-					Ret += TEXT("{") LINE_TERMINATOR;
-					Ret += TEXT("\t") + BaseName + TEXT(" Ret;") LINE_TERMINATOR;
-					Ret += TEXT("\tRet.Value = atan(A.Value);") LINE_TERMINATOR;
-					Ret += TEXT("\t") + FieldName + TEXT(" dFdA = rcp(A.Value * A.Value + 1.0f);");
-					Ret += TEXT("\tRet.Ddx = dFdA * A.Ddx;") LINE_TERMINATOR;
-					Ret += TEXT("\tRet.Ddy = dFdA * A.Ddy;") LINE_TERMINATOR;
-					Ret += TEXT("\treturn Ret;") LINE_TERMINATOR;
-					Ret += TEXT("}") LINE_TERMINATOR;
-					Ret += TEXT("") LINE_TERMINATOR;
-					break;
-				case EF1_AtanFast:
-					Ret += BaseName + TEXT(" AtanFastDeriv(") + BaseName + TEXT(" A)") LINE_TERMINATOR;
-					Ret += TEXT("{") LINE_TERMINATOR;
-					Ret += TEXT("\t") + BaseName + TEXT(" Ret;") LINE_TERMINATOR;
-					Ret += TEXT("\tRet.Value = atanFast(A.Value);") LINE_TERMINATOR;
-					Ret += TEXT("\t") + FieldName + TEXT(" dFdA = rcp(A.Value * A.Value + 1.0f);");
-					Ret += TEXT("\tRet.Ddx = dFdA * A.Ddx;") LINE_TERMINATOR;
-					Ret += TEXT("\tRet.Ddy = dFdA * A.Ddy;") LINE_TERMINATOR;
-					Ret += TEXT("\treturn Ret;") LINE_TERMINATOR;
-					Ret += TEXT("}") LINE_TERMINATOR;
-					Ret += TEXT("") LINE_TERMINATOR;
-					break;
-				case EF1_Sqrt:
-					Ret += BaseName + TEXT(" SqrtDeriv(") + BaseName + TEXT(" A)") LINE_TERMINATOR;
-					Ret += TEXT("{") LINE_TERMINATOR;
-					Ret += TEXT("\t") + BaseName + TEXT(" Ret;") LINE_TERMINATOR;
-					Ret += TEXT("\tRet.Value = sqrt(A.Value);") LINE_TERMINATOR;
-					Ret += TEXT("\t") + FieldName + TEXT(" dFdA = 0.5f * rsqrt(max(A.Value, 0.00001f));");
-					Ret += TEXT("\tRet.Ddx = dFdA * A.Ddx;") LINE_TERMINATOR;
-					Ret += TEXT("\tRet.Ddy = dFdA * A.Ddy;") LINE_TERMINATOR;
-					Ret += TEXT("\treturn Ret;") LINE_TERMINATOR;
-					Ret += TEXT("}") LINE_TERMINATOR;
-					Ret += TEXT("") LINE_TERMINATOR;
-					break;
-				case EF1_Rcp:
-					Ret += BaseName + TEXT(" RcpDeriv(") + BaseName + TEXT(" A)") LINE_TERMINATOR;
-					Ret += TEXT("{") LINE_TERMINATOR;
-					Ret += TEXT("\t") + BaseName + TEXT(" Ret;") LINE_TERMINATOR;
-					Ret += TEXT("\tRet.Value = rcp(A.Value);") LINE_TERMINATOR;
-					Ret += TEXT("\t") + FieldName + TEXT(" dFdA = -Ret.Value * Ret.Value;");
-					Ret += TEXT("\tRet.Ddx = dFdA * A.Ddx;") LINE_TERMINATOR;
-					Ret += TEXT("\tRet.Ddy = dFdA * A.Ddy;") LINE_TERMINATOR;
-					Ret += TEXT("\treturn Ret;") LINE_TERMINATOR;
-					Ret += TEXT("}") LINE_TERMINATOR;
-					Ret += TEXT("") LINE_TERMINATOR;
-					break;
-				case EF1_Rsqrt:
-					Ret += BaseName + TEXT(" RsqrtDeriv(") + BaseName + TEXT(" A)") LINE_TERMINATOR;
-					Ret += TEXT("{") LINE_TERMINATOR;
-					Ret += TEXT("\t") + BaseName + TEXT(" Ret;") LINE_TERMINATOR;
-					Ret += TEXT("\tRet.Value = rsqrt(A.Value);") LINE_TERMINATOR;
-					Ret += TEXT("\t") + FieldName + TEXT(" dFdA = -0.5f * rsqrt(A.Value) * rcp(A.Value);");
-					Ret += TEXT("\tRet.Ddx = dFdA * A.Ddx;") LINE_TERMINATOR;
-					Ret += TEXT("\tRet.Ddy = dFdA * A.Ddy;") LINE_TERMINATOR;
-					Ret += TEXT("\treturn Ret;") LINE_TERMINATOR;
-					Ret += TEXT("}") LINE_TERMINATOR;
-					Ret += TEXT("") LINE_TERMINATOR;
-					break;
-				case EF1_Saturate:
-					Ret += BaseName + TEXT(" SaturateDeriv(") + BaseName + TEXT(" A)") LINE_TERMINATOR;
-					Ret += TEXT("{") LINE_TERMINATOR;
-					Ret += TEXT("\t") + BoolName + TEXT(" InRange = (0.0 < A.Value && A.Value < 1.0);") LINE_TERMINATOR;
-					Ret += TEXT("\t") + FieldName + TEXT(" Zero = 0.0f;") LINE_TERMINATOR;
-					Ret += TEXT("\t") + BaseName + TEXT(" Ret;") LINE_TERMINATOR;
-					Ret += TEXT("\tRet.Value = saturate(A.Value);") LINE_TERMINATOR;
-					Ret += TEXT("\tRet.Ddx = SelectElemHelper(InRange, A.Ddx, Zero);") LINE_TERMINATOR;
-					Ret += TEXT("\tRet.Ddy = SelectElemHelper(InRange, A.Ddy, Zero);") LINE_TERMINATOR;
-					Ret += TEXT("\treturn Ret;") LINE_TERMINATOR;
-					Ret += TEXT("}") LINE_TERMINATOR;
-					Ret += TEXT("") LINE_TERMINATOR;
-					break;
-				case EF1_Frac:
-					Ret += BaseName + TEXT(" FracDeriv(") + BaseName + TEXT(" A)") LINE_TERMINATOR;
-					Ret += TEXT("{") LINE_TERMINATOR;
-					Ret += TEXT("\t") + BaseName + TEXT(" Ret;") LINE_TERMINATOR;
-					Ret += TEXT("\tRet.Value = frac(A.Value);") LINE_TERMINATOR;
-					Ret += TEXT("\tRet.Ddx = A.Ddx;") LINE_TERMINATOR;
-					Ret += TEXT("\tRet.Ddy = A.Ddy;") LINE_TERMINATOR;
-					Ret += TEXT("\treturn Ret;") LINE_TERMINATOR;
-					Ret += TEXT("}") LINE_TERMINATOR;
-					Ret += TEXT("") LINE_TERMINATOR;
-					break;
-				case EF1_Log2:
-					Ret += BaseName + TEXT(" Log2Deriv(") + BaseName + TEXT(" A)") LINE_TERMINATOR;
-					Ret += TEXT("{") LINE_TERMINATOR;
-					Ret += TEXT("\t") + BaseName + TEXT(" Ret;") LINE_TERMINATOR;
-					Ret += TEXT("\tRet.Value = log2(A.Value);") LINE_TERMINATOR;
-					Ret += TEXT("\t") + FieldName + TEXT(" dFdA = rcp(A.Value) * 1.442695f;");
-					Ret += TEXT("\tRet.Ddx = dFdA * A.Ddx ;") LINE_TERMINATOR;
-					Ret += TEXT("\tRet.Ddy = dFdA * A.Ddy;") LINE_TERMINATOR;
-					Ret += TEXT("\treturn Ret;") LINE_TERMINATOR;
-					Ret += TEXT("}") LINE_TERMINATOR;
-					Ret += TEXT("") LINE_TERMINATOR;
-					break;
-				case EF1_Log10:
-					Ret += BaseName + TEXT(" Log10Deriv(") + BaseName + TEXT(" A)") LINE_TERMINATOR;
-					Ret += TEXT("{") LINE_TERMINATOR;
-					Ret += TEXT("\t") + BaseName + TEXT(" Ret;") LINE_TERMINATOR;
-					Ret += TEXT("\tRet.Value = log10(A.Value);") LINE_TERMINATOR;
-					Ret += TEXT("\t") + FieldName + TEXT(" dFdA = rcp(A.Value) * 0.4342945f;");
-					Ret += TEXT("\tRet.Ddx = dFdA * A.Ddx;") LINE_TERMINATOR;
-					Ret += TEXT("\tRet.Ddy = dFdA * A.Ddy;") LINE_TERMINATOR;
-					Ret += TEXT("\treturn Ret;") LINE_TERMINATOR;
-					Ret += TEXT("}") LINE_TERMINATOR;
-					Ret += TEXT("") LINE_TERMINATOR;
-					break;
-				case EF1_Exp:
-					Ret += BaseName + TEXT(" ExpDeriv(") + BaseName + TEXT(" A)") LINE_TERMINATOR;
-					Ret += TEXT("{") LINE_TERMINATOR;
-					Ret += TEXT("\t") + BaseName + TEXT(" Ret;") LINE_TERMINATOR;
-					Ret += TEXT("\tRet.Value = exp(A.Value);") LINE_TERMINATOR;
-					Ret += TEXT("\tRet.Ddx = exp(A.Value) * A.Ddx;") LINE_TERMINATOR;
-					Ret += TEXT("\tRet.Ddy = exp(A.Value) * A.Ddy;") LINE_TERMINATOR;
-					Ret += TEXT("\treturn Ret;") LINE_TERMINATOR;
-					Ret += TEXT("}") LINE_TERMINATOR;
-					Ret += TEXT("") LINE_TERMINATOR;
-					break;
-				case EF1_Length:
-					Ret += TEXT("FloatDeriv LengthDeriv(") + BaseName + TEXT(" A)") LINE_TERMINATOR;
-					Ret += TEXT("{") LINE_TERMINATOR;
-					Ret += TEXT("\tFloatDeriv Ret = SqrtDeriv(DotDeriv(A,A));") LINE_TERMINATOR;
-					Ret += TEXT("\treturn Ret;") LINE_TERMINATOR;
-					Ret += TEXT("}") LINE_TERMINATOR;
-					Ret += TEXT("") LINE_TERMINATOR;
-					break;
-				case EF1_InvLength:
-					Ret += TEXT("FloatDeriv InvLengthDeriv(") + BaseName + TEXT(" A)") LINE_TERMINATOR;
-					Ret += TEXT("{") LINE_TERMINATOR;
-					Ret += TEXT("\tFloatDeriv Ret = RsqrtDeriv(DotDeriv(A,A));") LINE_TERMINATOR;
-					Ret += TEXT("\treturn Ret;") LINE_TERMINATOR;
-					Ret += TEXT("}") LINE_TERMINATOR;
-					Ret += TEXT("") LINE_TERMINATOR;
-					break;
-				case EF1_Normalize:
-					Ret += BaseName + TEXT(" NormalizeDeriv(") + BaseName + TEXT(" A)") LINE_TERMINATOR;
-					Ret += TEXT("{") LINE_TERMINATOR;
-					Ret += TEXT("\tFloatDeriv InvLen = RsqrtDeriv(DotDeriv(A,A));") LINE_TERMINATOR;
-					Ret += TEXT("\t") + BaseName + TEXT(" Ret = MulDeriv(") + ConvertDeriv(TEXT("InvLen"), Index, 0) + TEXT(", A);") LINE_TERMINATOR;
-					Ret += TEXT("\treturn Ret;") LINE_TERMINATOR;
-					Ret += TEXT("}") LINE_TERMINATOR;
-					Ret += TEXT("") LINE_TERMINATOR;
-					break;
-				default:
-					check(0);
-					break;
-				}
-			}
-		}
-	}
-
-	for (int32 Index = 0; Index < 4; Index++)
-	{
-		const FString BaseName = GetDerivVectorName(Index);
-		const FString FieldName = GetFloatVectorName(Index);
-		const FString BoolName = GetBoolVectorName(Index);
-
-		if (bLerpEnabled[Index] || IsDebugGenerateAllFunctionsEnabled())
-		{
-			
-			// lerp(a,b,s) = a*(1-s) + b*s
-			// lerp(a,b,s)' = a' * (1 - s') + b' * s + s' * (b - a)
-			Ret += FString::Printf(TEXT("%s LerpDeriv(%s A, %s B, %s S)"), *BaseName, *BaseName, *BaseName, *BaseName) + LINE_TERMINATOR;
-			Ret += TEXT("{") LINE_TERMINATOR;
-			Ret += TEXT("\t") + BaseName + TEXT(" Ret;") LINE_TERMINATOR;
-			Ret += TEXT("\tRet.Value = lerp(A.Value, B.Value, S.Value);") LINE_TERMINATOR;
-			Ret += TEXT("\tRet.Ddx = lerp(A.Ddx, B.Ddx, S.Value) + S.Ddx * (B.Value - A.Value);") LINE_TERMINATOR;
-			Ret += TEXT("\tRet.Ddy = lerp(A.Ddy, B.Ddy, S.Value) + S.Ddy * (B.Value - A.Value);") LINE_TERMINATOR;
-			Ret += TEXT("\treturn Ret;") LINE_TERMINATOR;
-			Ret += TEXT("}") LINE_TERMINATOR;
-			Ret += TEXT("") LINE_TERMINATOR;
-		}
-
-		if (bIfEnabled[Index] || IsDebugGenerateAllFunctionsEnabled())
-		{
-			Ret += FString::Printf(TEXT("%s IfDeriv(float A, float B, %s Greater, %s Less)"), *BaseName, *BaseName, *BaseName) + LINE_TERMINATOR;
-			Ret += TEXT("{") LINE_TERMINATOR;
-			Ret += TEXT("\tif(A >= B)") LINE_TERMINATOR;
-			Ret += TEXT("\t\treturn Greater;") LINE_TERMINATOR;
-			Ret += TEXT("\telse") LINE_TERMINATOR;
-			Ret += TEXT("\t\treturn Less;") LINE_TERMINATOR;
-			Ret += TEXT("}") LINE_TERMINATOR;
-			Ret += TEXT("") LINE_TERMINATOR;
-		}
-
-		if (bIf2Enabled[Index] || IsDebugGenerateAllFunctionsEnabled())
-		{
-			Ret += FString::Printf(TEXT("%s IfDeriv(float A, float B, %s Greater, %s Less, %s Equal, float Threshold)"), *BaseName, *BaseName, *BaseName, *BaseName) + LINE_TERMINATOR;
-			Ret += TEXT("{") LINE_TERMINATOR;
-			Ret += TEXT("\tif(!(abs(A - B) > Threshold))") LINE_TERMINATOR;	// Written like this to preserve NaN behavior of original code.
-			Ret += TEXT("\t\treturn Equal;") LINE_TERMINATOR;
-			Ret += TEXT("\tif(A >= B)") LINE_TERMINATOR;
-			Ret += TEXT("\t\treturn Greater;") LINE_TERMINATOR;
-			Ret += TEXT("\telse") LINE_TERMINATOR;
-			Ret += TEXT("\t\treturn Less;") LINE_TERMINATOR;
-			Ret += TEXT("}") LINE_TERMINATOR;
-			Ret += TEXT("") LINE_TERMINATOR;
-		}
-	}
-
-	if (bRotateScaleOffsetTexCoords || IsDebugGenerateAllFunctionsEnabled())
-	{
-		// float2(dot(InTexCoords, InRotationScale.xy), dot(InTexCoords, InRotationScale.zw)) + InOffset;
-		// InTexCoords.xy * InRotationScale.xw + InTexCoords.yx * InRotationScale.yz + InOffset;
-		Ret += TEXT("FloatDeriv2 RotateScaleOffsetTexCoordsDeriv(FloatDeriv2 TexCoord, FloatDeriv2 RotationScale, FloatDeriv2 Offset)") LINE_TERMINATOR;
-		Ret += TEXT("{") LINE_TERMINATOR;
-		Ret += TEXT("\tFloatDeriv2 Ret = Offset;") LINE_TERMINATOR;
-		Ret += TEXT("\tRet = AddDeriv(Ret, MulDeriv(TexCoord, SwizzleDeriv2(RotationScale, xw));") LINE_TERMINATOR;
-		Ret += TEXT("\tRet = AddDeriv(Ret, MulDeriv(SwizzleDeriv2(TexCoord, yx), SwizzleDeriv2(RotationScale, yz));") LINE_TERMINATOR;
-		Ret += TEXT("\treturn Ret;") LINE_TERMINATOR;
-		Ret += TEXT("}") LINE_TERMINATOR;
-		Ret += TEXT("") LINE_TERMINATOR;
-	}
-	
-	if (bUnMirrorEnabled[1][1] || IsDebugGenerateAllFunctionsEnabled())
-	{
-		// UnMirrorUV
-		Ret += TEXT("FloatDeriv2 UnMirrorU(FloatDeriv2 UV, FMaterialPixelParameters Parameters)") LINE_TERMINATOR;
-		Ret += TEXT("{") LINE_TERMINATOR;
-		Ret += TEXT("\tconst MaterialFloat Scale = (Parameters.UnMirrored * 0.5f);") LINE_TERMINATOR;
-		Ret += TEXT("\tUV.Value = UV.Value * Scale + 0.5f;") LINE_TERMINATOR;
-		Ret += TEXT("\tUV.Ddx *= Scale;") LINE_TERMINATOR;
-		Ret += TEXT("\tUV.Ddy *= Scale;") LINE_TERMINATOR;
-		Ret += TEXT("\treturn UV;") LINE_TERMINATOR;
-		Ret += TEXT("}") LINE_TERMINATOR;
-		Ret += TEXT("") LINE_TERMINATOR;
-	}
-	
-	if(bUnMirrorEnabled[1][0] || IsDebugGenerateAllFunctionsEnabled())
-	{
-		// UnMirrorU
-		Ret += TEXT("FloatDeriv2 UnMirrorU(FloatDeriv2 UV, FMaterialPixelParameters Parameters)") LINE_TERMINATOR;
-		Ret += TEXT("{") LINE_TERMINATOR;
-		Ret += TEXT("\tconst MaterialFloat Scale = (Parameters.UnMirrored * 0.5f);") LINE_TERMINATOR;
-		Ret += TEXT("\tUV.Value.x = UV.Value.x * Scale + 0.5f;") LINE_TERMINATOR;
-		Ret += TEXT("\tUV.Ddx.x *= Scale;") LINE_TERMINATOR;
-		Ret += TEXT("\tUV.Ddy.x *= Scale;") LINE_TERMINATOR;
-		Ret += TEXT("\treturn UV;") LINE_TERMINATOR;
-		Ret += TEXT("}") LINE_TERMINATOR;
-		Ret += TEXT("") LINE_TERMINATOR;
-	}
-	
-	if (bUnMirrorEnabled[0][1] || IsDebugGenerateAllFunctionsEnabled())
-	{
-		// UnMirrorV
-		Ret += TEXT("FloatDeriv2 UnMirrorV(FloatDeriv2 UV, FMaterialPixelParameters Parameters)") LINE_TERMINATOR;
-		Ret += TEXT("{") LINE_TERMINATOR;
-		Ret += TEXT("\tconst MaterialFloat Scale = (Parameters.UnMirrored * 0.5f);") LINE_TERMINATOR;
-		Ret += TEXT("\tUV.Value.y = UV.Value.y * Scale + 0.5f;") LINE_TERMINATOR;
-		Ret += TEXT("\tUV.Ddx.y *= Scale;") LINE_TERMINATOR;
-		Ret += TEXT("\tUV.Ddy.y *= Scale;") LINE_TERMINATOR;
-		Ret += TEXT("\treturn UV;") LINE_TERMINATOR;
-		Ret += TEXT("}") LINE_TERMINATOR;
-		Ret += TEXT("") LINE_TERMINATOR;
-	}
-
-	return Ret;
-}
-
-
-FString FMaterialDerivativeAutogen::ApplyUnMirror(FString Value, bool bUnMirrorU, bool bUnMirrorV)
-{
-	if (bUnMirrorU && bUnMirrorV)
-	{
-		Value = FString::Printf(TEXT("UnMirrorUV(%s, Parameters)"), *Value);
-	}
-	else if (bUnMirrorU)
-	{
-		Value = FString::Printf(TEXT("UnMirrorU(%s, Parameters)"), *Value);
-	}
-	else if (bUnMirrorV)
-	{
-		Value = FString::Printf(TEXT("UnMirrorV(%s, Parameters)"), *Value);
-	}
-
-	bUnMirrorEnabled[bUnMirrorU][bUnMirrorV] = true;
-	
-	return Value;	
-}
-
-
 FHLSLMaterialTranslator::FHLSLMaterialTranslator(FMaterial* InMaterial,
 	FMaterialCompilationOutput& InMaterialCompilationOutput,
 	const FStaticParameterSet& InStaticParameters,
@@ -4011,7 +2241,7 @@ FString FHLSLMaterialTranslator::GetParameterCode(int32 Index, const TCHAR* Defa
 {
 	FString Ret = GetParameterCodeRaw(Index,Default);
 
-	if (GetDerivativeStatus(Index) == DS_Valid)
+	if (GetDerivativeStatus(Index) == EDerivativeStatus::Valid)
 	{
 		Ret = TEXT("DERIV_BASE_VALUE(") + Ret + TEXT(")");
 	}
@@ -4145,7 +2375,7 @@ void FHLSLMaterialTranslator::GetFixedParameterCode(int32 StartChunk, int32 EndC
 			check(ResultChunk.bInline || ResultChunk.SymbolName.Len() > 0);
 			OutDefinitions = GetDefinitions(CodeChunks, StartChunk, EndChunk, Variation);
 			OutValue = ResultChunk.bInline ? ResultChunk.AtDefinition(Variation) : ResultChunk.SymbolName;
-			if (Variation == CompiledPDV_Analytic && ResultChunk.DerivativeStatus == DS_Valid)
+			if (Variation == CompiledPDV_Analytic && ResultChunk.DerivativeStatus == EDerivativeStatus::Valid)
 			{
 				OutValue = FString(TEXT("(")) + OutValue + TEXT(").Value");
 			}
@@ -4289,11 +2519,11 @@ const TCHAR* FHLSLMaterialTranslator::HLSLTypeStringDeriv(EMaterialValueType Typ
 {
 	switch(Type)
 	{
-	case MCT_Float1:				return (DerivativeStatus == DS_Valid) ? TEXT("FloatDeriv") : TEXT("MaterialFloat");
-	case MCT_Float2:				return (DerivativeStatus == DS_Valid) ? TEXT("FloatDeriv2") : TEXT("MaterialFloat2");
-	case MCT_Float3:				return (DerivativeStatus == DS_Valid) ? TEXT("FloatDeriv3") : TEXT("MaterialFloat3");
-	case MCT_Float4:				return (DerivativeStatus == DS_Valid) ? TEXT("FloatDeriv4") : TEXT("MaterialFloat4");
-	case MCT_Float:					return (DerivativeStatus == DS_Valid) ? TEXT("FloatDeriv") : TEXT("MaterialFloat");
+	case MCT_Float1:				return (DerivativeStatus == EDerivativeStatus::Valid) ? TEXT("FloatDeriv") : TEXT("MaterialFloat");
+	case MCT_Float2:				return (DerivativeStatus == EDerivativeStatus::Valid) ? TEXT("FloatDeriv2") : TEXT("MaterialFloat2");
+	case MCT_Float3:				return (DerivativeStatus == EDerivativeStatus::Valid) ? TEXT("FloatDeriv3") : TEXT("MaterialFloat3");
+	case MCT_Float4:				return (DerivativeStatus == EDerivativeStatus::Valid) ? TEXT("FloatDeriv4") : TEXT("MaterialFloat4");
+	case MCT_Float:					return (DerivativeStatus == EDerivativeStatus::Valid) ? TEXT("FloatDeriv") : TEXT("MaterialFloat");
 	case MCT_Texture2D:				return TEXT("texture2D");
 	case MCT_TextureCube:			return TEXT("textureCube");
 	case MCT_Texture2DArray:		return TEXT("texture2DArray");
@@ -4554,7 +2784,7 @@ int32 FHLSLMaterialTranslator::AddCodeChunkWithHash(uint64 BaseHash, EMaterialVa
 
 	uint64 Hash = CityHash64((char*)FormattedCode, Result * sizeof(TCHAR));
 	Hash = CityHash128to64({ BaseHash, Hash });
-	const int32 CodeIndex = AddCodeChunkInner(Hash, FormattedCode, Type, DS_NotAware, false);
+	const int32 CodeIndex = AddCodeChunkInner(Hash, FormattedCode, Type, EDerivativeStatus::NotAware, false);
 	FMemory::Free(FormattedCode);
 
 	return CodeIndex;
@@ -4662,12 +2892,12 @@ int32 FHLSLMaterialTranslator::AddUniformExpressionInner(uint64 Hash, FMaterialU
 	const int32 ReturnIndex = CurrentScopeChunks->Num();
 	// Create a new code chunk for the uniform expression
 	// Note that uniforms have a known-zero derivative
-	new(*CurrentScopeChunks) FShaderCodeChunk(Hash, UniformExpression, FormattedCode, FormattedCode, Type, DS_Zero);
+	new(*CurrentScopeChunks) FShaderCodeChunk(Hash, UniformExpression, FormattedCode, FormattedCode, Type, EDerivativeStatus::Zero);
 
 	if (!bFoundExistingExpression)
 	{
 		// Add an entry to the material-wide list of uniform expressions
-		new(UniformExpressions) FShaderCodeChunk(Hash, UniformExpression, FormattedCode, FormattedCode, Type, DS_Zero);
+		new(UniformExpressions) FShaderCodeChunk(Hash, UniformExpression, FormattedCode, FormattedCode, Type, EDerivativeStatus::Zero);
 	}
 
 	AddCodeChunkToCurrentScope(ReturnIndex);
@@ -4723,7 +2953,7 @@ int32 FHLSLMaterialTranslator::AccessUniformExpression(int32 Index)
 	}
 	else if(CodeChunk.Type & MCT_Float)
 	{
-		check(CodeChunk.DerivativeStatus != DS_Valid);
+		check(CodeChunk.DerivativeStatus != EDerivativeStatus::Valid);
 
 		const TCHAR* Mask;
 		switch(CodeChunk.Type)
@@ -5338,20 +3568,20 @@ int32 FHLSLMaterialTranslator::ValidCast(int32 Code, EMaterialValueType DestType
 			FString FiniteCode = FString::Printf(TEXT("%s%s"), *GetParameterCode(Code), Mask);
 			if (IsAnalyticDerivEnabled() && IsDerivativeValid(CodeDerivInfo.DerivativeStatus))
 			{
-				if (CodeDerivInfo.DerivativeStatus == DS_Valid)
+				if (CodeDerivInfo.DerivativeStatus == EDerivativeStatus::Valid)
 				{
 					FString DerivString = *GetParameterCodeDeriv(Code, CompiledPDV_Analytic);
 
 					FString DDXCode = TEXT("0.0f");
 					FString DDYCode = TEXT("0.0f");
-					if (CodeDerivInfo.DerivativeStatus == DS_Valid)
+					if (CodeDerivInfo.DerivativeStatus == EDerivativeStatus::Valid)
 					{
 						DDXCode = DerivString + TEXT(".Ddx") + Mask;
 						DDYCode = DerivString + TEXT(".Ddy") + Mask;
 					}
 
 					FString AnalyticCode = DerivativeAutogen.ConstructDeriv(FiniteCode, *DDXCode, *DDYCode, GetDerivTypeIndex(DestType));
-					return AddCodeChunkInnerDeriv(*FiniteCode, *AnalyticCode, DestType, false, DS_Valid);
+					return AddCodeChunkInnerDeriv(*FiniteCode, *AnalyticCode, DestType, false, EDerivativeStatus::Valid);
 				}
 				else
 				{
@@ -5381,7 +3611,7 @@ int32 FHLSLMaterialTranslator::ValidCast(int32 Code, EMaterialValueType DestType
 
 			if (IsAnalyticDerivEnabled() && IsDerivativeValid(CodeDerivInfo.DerivativeStatus))
 			{
-				if (CodeDerivInfo.DerivativeStatus == DS_Valid)
+				if (CodeDerivInfo.DerivativeStatus == EDerivativeStatus::Valid)
 				{
 					FString DerivString = *GetParameterCodeDeriv(Code, CompiledPDV_Analytic);
 					FString DerivStringDDX = DerivString + TEXT(".Ddx");
@@ -5398,7 +3628,7 @@ int32 FHLSLMaterialTranslator::ValidCast(int32 Code, EMaterialValueType DestType
 					DDYCode += TEXT(")");
 
 					FString AnalyticCode = DerivativeAutogen.ConstructDeriv(FiniteCode, *DDXCode, *DDYCode, GetDerivTypeIndex(DestType));
-					return AddCodeChunkInnerDeriv(*FiniteCode, *AnalyticCode, DestType, false, DS_Valid);
+					return AddCodeChunkInnerDeriv(*FiniteCode, *AnalyticCode, DestType, false, EDerivativeStatus::Valid);
 				}
 				else
 				{
@@ -5473,20 +3703,20 @@ int32 FHLSLMaterialTranslator::ForceCast(int32 Code, EMaterialValueType DestType
 			FString FiniteCode = FString::Printf(TEXT("%s%s"), *GetParameterCode(Code), Mask);
 			if (IsAnalyticDerivEnabled() && IsDerivativeValid(CodeDerivInfo.DerivativeStatus))
 			{
-				if (CodeDerivInfo.DerivativeStatus == DS_Valid)
+				if (CodeDerivInfo.DerivativeStatus == EDerivativeStatus::Valid)
 				{
 					FString DerivString = *GetParameterCodeDeriv(Code, CompiledPDV_Analytic);
 
 					FString DDXCode = TEXT("0.0f");
 					FString DDYCode = TEXT("0.0f");
-					if (CodeDerivInfo.DerivativeStatus == DS_Valid)
+					if (CodeDerivInfo.DerivativeStatus == EDerivativeStatus::Valid)
 					{
 						DDXCode = DerivString + TEXT(".Ddx") + Mask;
 						DDYCode = DerivString + TEXT(".Ddy") + Mask;
 					}
 
 					FString AnalyticCode = DerivativeAutogen.ConstructDeriv(FiniteCode, *DDXCode, *DDYCode, GetDerivTypeIndex(DestType));
-					return AddCodeChunkInnerDeriv(*FiniteCode, *AnalyticCode, DestType, false, DS_Valid);
+					return AddCodeChunkInnerDeriv(*FiniteCode, *AnalyticCode, DestType, false, EDerivativeStatus::Valid);
 				}
 				else
 				{
@@ -5518,7 +3748,7 @@ int32 FHLSLMaterialTranslator::ForceCast(int32 Code, EMaterialValueType DestType
 
 			if (IsAnalyticDerivEnabled() && IsDerivativeValid(CodeDerivInfo.DerivativeStatus))
 			{
-				if (CodeDerivInfo.DerivativeStatus == DS_Valid)
+				if (CodeDerivInfo.DerivativeStatus == EDerivativeStatus::Valid)
 				{
 					FString DerivString = *GetParameterCodeDeriv(Code, CompiledPDV_Analytic);
 					FString DerivStringDDX = DerivString + TEXT(".Ddx");
@@ -5534,7 +3764,7 @@ int32 FHLSLMaterialTranslator::ForceCast(int32 Code, EMaterialValueType DestType
 					DDYCode += TEXT(")");
 
 					FString AnalyticCode = DerivativeAutogen.ConstructDeriv(FiniteCode, *DDXCode, *DDYCode, GetDerivTypeIndex(DestType));
-					return AddCodeChunkInnerDeriv(*FiniteCode, *AnalyticCode, DestType, false, DS_Valid);
+					return AddCodeChunkInnerDeriv(*FiniteCode, *AnalyticCode, DestType, false, EDerivativeStatus::Valid);
 				}
 				else
 				{
@@ -5826,7 +4056,7 @@ int32 FHLSLMaterialTranslator::Sine(int32 X)
 	{
 		if (IsAnalyticDerivEnabled())
 		{
-			return DerivativeAutogen.GenerateExpressionFunc1(*this,FMaterialDerivativeAutogen::EF1_Sin,X);
+			return DerivativeAutogen.GenerateExpressionFunc1(*this,FMaterialDerivativeAutogen::EFunc1::Sin,X);
 		}
 		else
 		{
@@ -5850,7 +4080,7 @@ int32 FHLSLMaterialTranslator::Cosine(int32 X)
 	{
 		if (IsAnalyticDerivEnabled())
 		{
-			return DerivativeAutogen.GenerateExpressionFunc1(*this,FMaterialDerivativeAutogen::EF1_Cos,X);
+			return DerivativeAutogen.GenerateExpressionFunc1(*this,FMaterialDerivativeAutogen::EFunc1::Cos,X);
 		}
 		else
 		{
@@ -5874,7 +4104,7 @@ int32 FHLSLMaterialTranslator::Tangent(int32 X)
 	{
 		if (IsAnalyticDerivEnabled())
 		{
-			return DerivativeAutogen.GenerateExpressionFunc1(*this, FMaterialDerivativeAutogen::EF1_Tan, X);
+			return DerivativeAutogen.GenerateExpressionFunc1(*this, FMaterialDerivativeAutogen::EFunc1::Tan, X);
 		}
 		else
 		{
@@ -5898,7 +4128,7 @@ int32 FHLSLMaterialTranslator::Arcsine(int32 X)
 	{
 		if (IsAnalyticDerivEnabled())
 		{
-			return DerivativeAutogen.GenerateExpressionFunc1(*this, FMaterialDerivativeAutogen::EF1_Asin, X);
+			return DerivativeAutogen.GenerateExpressionFunc1(*this, FMaterialDerivativeAutogen::EFunc1::Asin, X);
 		}
 		else
 		{
@@ -5922,7 +4152,7 @@ int32 FHLSLMaterialTranslator::ArcsineFast(int32 X)
 	{
 		if (IsAnalyticDerivEnabled())
 		{
-			return DerivativeAutogen.GenerateExpressionFunc1(*this, FMaterialDerivativeAutogen::EF1_AsinFast, X);
+			return DerivativeAutogen.GenerateExpressionFunc1(*this, FMaterialDerivativeAutogen::EFunc1::AsinFast, X);
 		}
 		else
 		{
@@ -5946,7 +4176,7 @@ int32 FHLSLMaterialTranslator::Arccosine(int32 X)
 	{
 		if (IsAnalyticDerivEnabled())
 		{
-			return DerivativeAutogen.GenerateExpressionFunc1(*this, FMaterialDerivativeAutogen::EF1_Acos, X);
+			return DerivativeAutogen.GenerateExpressionFunc1(*this, FMaterialDerivativeAutogen::EFunc1::Acos, X);
 		}
 		else
 		{
@@ -5970,7 +4200,7 @@ int32 FHLSLMaterialTranslator::ArccosineFast(int32 X)
 	{
 		if (IsAnalyticDerivEnabled())
 		{
-			return DerivativeAutogen.GenerateExpressionFunc1(*this, FMaterialDerivativeAutogen::EF1_AcosFast, X);
+			return DerivativeAutogen.GenerateExpressionFunc1(*this, FMaterialDerivativeAutogen::EFunc1::AcosFast, X);
 		}
 		else
 		{
@@ -5994,7 +4224,7 @@ int32 FHLSLMaterialTranslator::Arctangent(int32 X)
 	{
 		if (IsAnalyticDerivEnabled())
 		{
-			return DerivativeAutogen.GenerateExpressionFunc1(*this, FMaterialDerivativeAutogen::EF1_Atan, X);
+			return DerivativeAutogen.GenerateExpressionFunc1(*this, FMaterialDerivativeAutogen::EFunc1::Atan, X);
 		}
 		else
 		{
@@ -6018,7 +4248,7 @@ int32 FHLSLMaterialTranslator::ArctangentFast(int32 X)
 	{
 		if (IsAnalyticDerivEnabled())
 		{
-			return DerivativeAutogen.GenerateExpressionFunc1(*this, FMaterialDerivativeAutogen::EF1_AtanFast, X);
+			return DerivativeAutogen.GenerateExpressionFunc1(*this, FMaterialDerivativeAutogen::EFunc1::AtanFast, X);
 		}
 		else
 		{
@@ -6042,7 +4272,7 @@ int32 FHLSLMaterialTranslator::Arctangent2(int32 Y, int32 X)
 	{
 		if (IsAnalyticDerivEnabled())
 		{
-			return DerivativeAutogen.GenerateExpressionFunc2(*this, FMaterialDerivativeAutogen::EF2_Atan2, Y, X);
+			return DerivativeAutogen.GenerateExpressionFunc2(*this, FMaterialDerivativeAutogen::EFunc2::Atan2, Y, X);
 		}
 		else
 		{
@@ -6066,7 +4296,7 @@ int32 FHLSLMaterialTranslator::Arctangent2Fast(int32 Y, int32 X)
 	{
 		if (IsAnalyticDerivEnabled())
 		{
-			return DerivativeAutogen.GenerateExpressionFunc2(*this, FMaterialDerivativeAutogen::EF2_Atan2Fast, Y, X);
+			return DerivativeAutogen.GenerateExpressionFunc2(*this, FMaterialDerivativeAutogen::EFunc2::Atan2Fast, Y, X);
 		}
 		else
 		{
@@ -6175,7 +4405,7 @@ int32 FHLSLMaterialTranslator::Frac(int32 X)
 	{
 		if (IsAnalyticDerivEnabled())
 		{
-			return DerivativeAutogen.GenerateExpressionFunc1(*this, FMaterialDerivativeAutogen::EF1_Frac, X);
+			return DerivativeAutogen.GenerateExpressionFunc1(*this, FMaterialDerivativeAutogen::EFunc1::Frac, X);
 		}
 		else
 		{
@@ -6199,10 +4429,10 @@ int32 FHLSLMaterialTranslator::Fmod(int32 A, int32 B)
 	else
 	{
 		const FDerivInfo& BDerivInfo = GetDerivInfo(B);
-		if (IsAnalyticDerivEnabled() && BDerivInfo.DerivativeStatus == DS_Zero)
+		if (IsAnalyticDerivEnabled() && BDerivInfo.DerivativeStatus == EDerivativeStatus::Zero)
 		{
 			// Analytic derivatives only make sense when RHS derivatives are zero.
-			return DerivativeAutogen.GenerateExpressionFunc2(*this, FMaterialDerivativeAutogen::EF2_Fmod, A, B);
+			return DerivativeAutogen.GenerateExpressionFunc2(*this, FMaterialDerivativeAutogen::EFunc2::Fmod, A, B);
 		}
 		else
 		{
@@ -6233,7 +4463,7 @@ int32 FHLSLMaterialTranslator::Abs(int32 X)
 	{
 		if (IsAnalyticDerivEnabled())
 		{
-			return DerivativeAutogen.GenerateExpressionFunc1(*this, FMaterialDerivativeAutogen::EF1_Abs, X);
+			return DerivativeAutogen.GenerateExpressionFunc1(*this, FMaterialDerivativeAutogen::EFunc1::Abs, X);
 		}
 		else
 		{
@@ -6316,7 +4546,7 @@ int32 FHLSLMaterialTranslator::GetViewportUV()
 	if (IsAnalyticDerivEnabled())
 	{
 		FString AnalyticCode = DerivativeAutogen.ConstructDeriv(FiniteCode, TEXT("float2(View.ViewSizeAndInvSize.z, 0.0f)"), TEXT("float2(0.0f, View.ViewSizeAndInvSize.w)"), 1);
-		return AddCodeChunkInnerDeriv(*FiniteCode, *AnalyticCode, MCT_Float2, false, DS_Valid);
+		return AddCodeChunkInnerDeriv(*FiniteCode, *AnalyticCode, MCT_Float2, false, EDerivativeStatus::Valid);
 	}
 	else
 	{
@@ -6578,7 +4808,7 @@ int32 FHLSLMaterialTranslator::WorldPosition(EWorldPositionIncludedOffsets World
 	if (IsAnalyticDerivEnabled())
 	{
 		FString AnalyticCode = DerivativeAutogen.ConstructDeriv(FiniteCode, TEXT("Parameters.WorldPosition_DDX"), TEXT("Parameters.WorldPosition_DDY"), 2);
-		return AddCodeChunkInnerDeriv(*FiniteCode, *AnalyticCode, MCT_Float3, false, DS_Valid);
+		return AddCodeChunkInnerDeriv(*FiniteCode, *AnalyticCode, MCT_Float3, false, EDerivativeStatus::Valid);
 	}
 	else
 	{
@@ -6779,7 +5009,7 @@ int32 FHLSLMaterialTranslator::TextureCoordinate(uint32 CoordinateIndex, bool Un
 		FString TexCoordCodeAnalytic = DerivativeAutogen.ConstructDeriv(TexCoordCode, TexCoordCodeDDX, TexCoordCodeDDY, 1);	// 1 = float2
 		FString SampleCodeAnalytic = DerivativeAutogen.ApplyUnMirror(TexCoordCodeAnalytic, UnMirrorU, UnMirrorV);
 
-		return AddCodeChunkInnerDeriv(*SampleCodeFinite, *SampleCodeAnalytic, MCT_Float2, false, DS_Valid);
+		return AddCodeChunkInnerDeriv(*SampleCodeFinite, *SampleCodeAnalytic, MCT_Float2, false, EDerivativeStatus::Valid);
 	}
 	else
 	{
@@ -6915,7 +5145,7 @@ uint32 FHLSLMaterialTranslator::AcquireVTStackIndex(
 				*UV_Value, GetVTAddressMode(AddressU), GetVTAddressMode(AddressV), 
 				*UV_Ddx, *UV_Ddy,
 				StackIndex);
-			Entry.CodeIndex = AddCodeChunkInnerDeriv(*SampleCodeFinite, *SampleCodeAnalytic, MCT_VTPageTableResult, false, DS_NotValid);
+			Entry.CodeIndex = AddCodeChunkInnerDeriv(*SampleCodeFinite, *SampleCodeAnalytic, MCT_VTPageTableResult, false, EDerivativeStatus::NotValid);
 		}
 		else
 		{
@@ -6952,7 +5182,7 @@ uint32 FHLSLMaterialTranslator::AcquireVTStackIndex(
 				*UV_Value, GetVTAddressMode(AddressU), GetVTAddressMode(AddressV),
 				*UV_Ddx, *UV_Ddy,
 				StackIndex);
-			Entry.CodeIndex = AddCodeChunkInnerDeriv(*SampleCodeFinite, *SampleCodeAnalytic, MCT_VTPageTableResult, false, DS_NotValid);
+			Entry.CodeIndex = AddCodeChunkInnerDeriv(*SampleCodeFinite, *SampleCodeAnalytic, MCT_VTPageTableResult, false, EDerivativeStatus::NotValid);
 		}
 		else
 		{
@@ -7332,7 +5562,7 @@ int32 FHLSLMaterialTranslator::TextureSample(
 	FString UV_Scale = TEXT("1.0f");
 	if (IsAnalyticDerivEnabled() && IsDerivativeValid(UvDerivativeStatus))
 	{
-		if (UvDerivativeStatus == DS_Valid)
+		if (UvDerivativeStatus == EDerivativeStatus::Valid)
 		{
 			const FString UVAnalytic = GetParameterCodeDeriv(CoordinateIndex, CompiledPDV_Analytic);
 			UV_Value = UVAnalytic + TEXT(".Value");
@@ -7456,7 +5686,7 @@ int32 FHLSLMaterialTranslator::TextureSample(
 				SampleCodeAnalytic = FString::Printf(TEXT("%s(%s, %s, %d, VTUniform_Unpack(Material.VTPackedUniform[%d]))"), *TextureTypeName, *TextureName, *VTPageTableResult_Analytic, VTPageTableIndex, VirtualTextureIndex);
 
 			SampleCodeAnalytic = ApplySamplerType(SampleCodeAnalytic, SamplerType);
-			SamplingCodeIndex = AddCodeChunkInnerDeriv(*SampleCodeFinite, *SampleCodeAnalytic, MCT_Float4, false, DS_NotValid);
+			SamplingCodeIndex = AddCodeChunkInnerDeriv(*SampleCodeFinite, *SampleCodeAnalytic, MCT_Float4, false, EDerivativeStatus::NotValid);
 		}
 		else
 		{
@@ -7526,7 +5756,7 @@ int32 FHLSLMaterialTranslator::TextureSample(
 				SampleCodeAnalytic = ApplySamplerType(SampleCodeAnalytic, SamplerType);
 			}
 
-			SamplingCodeIndex = AddCodeChunkInnerDeriv(*SampleCodeFinite, *SampleCodeAnalytic, MCT_Float4, false /*?*/, DS_NotValid);
+			SamplingCodeIndex = AddCodeChunkInnerDeriv(*SampleCodeFinite, *SampleCodeAnalytic, MCT_Float4, false /*?*/, EDerivativeStatus::NotValid);
 		}
 		else
 		{
@@ -7624,7 +5854,7 @@ int32 FHLSLMaterialTranslator::PixelDepth()
 	if (IsAnalyticDerivEnabled())
 	{
 		FString AnalyticCode = DerivativeAutogen.ConstructDeriv(TEXT("Parameters.ScreenPosition.w"), TEXT("Parameters.ScreenPosition_DDX.w"), TEXT("Parameters.ScreenPosition_DDY.w"), 0);
-		return AddCodeChunkInnerDeriv(*FiniteCode, *AnalyticCode, MCT_Float, false, DS_Valid);
+		return AddCodeChunkInnerDeriv(*FiniteCode, *AnalyticCode, MCT_Float, false, EDerivativeStatus::Valid);
 	}
 	else
 	{
@@ -8282,7 +6512,7 @@ int32 FHLSLMaterialTranslator::VertexColor()
 	if (IsAnalyticDerivEnabled())
 	{
 		FString AnalyticCode = DerivativeAutogen.ConstructDeriv(FiniteCode, TEXT("Parameters.VertexColor_DDX"), TEXT("Parameters.VertexColor_DDY"), 3);
-		return AddCodeChunkInnerDeriv(*FiniteCode, *AnalyticCode, MCT_Float4, false, DS_Valid);
+		return AddCodeChunkInnerDeriv(*FiniteCode, *AnalyticCode, MCT_Float4, false, EDerivativeStatus::Valid);
 	}
 	else
 	{
@@ -8402,7 +6632,7 @@ int32 FHLSLMaterialTranslator::VertexInterpolator(uint32 InterpolatorIndex)
 	if (IsAnalyticDerivEnabled())
 	{
 		FString AnalyticCode = DerivativeAutogen.ConstructDeriv(*ValueCode, *DDXCode, *DDYCode, GetDerivTypeIndex(Interpolator->InterpolatedType));
-		return AddCodeChunkInnerDeriv(*ValueCode, *AnalyticCode, Interpolator->InterpolatedType, false, DS_Valid);
+		return AddCodeChunkInnerDeriv(*ValueCode, *AnalyticCode, Interpolator->InterpolatedType, false, EDerivativeStatus::Valid);
 	}
 	else
 	{
@@ -8426,7 +6656,7 @@ int32 FHLSLMaterialTranslator::Add(int32 A,int32 B)
 	{
 		if (IsAnalyticDerivEnabled())
 		{
-			return DerivativeAutogen.GenerateExpressionFunc2(*this,FMaterialDerivativeAutogen::EF2_Add,A,B);
+			return DerivativeAutogen.GenerateExpressionFunc2(*this,FMaterialDerivativeAutogen::EFunc2::Add,A,B);
 		}
 		else
 		{
@@ -8450,7 +6680,7 @@ int32 FHLSLMaterialTranslator::Sub(int32 A,int32 B)
 	{
 		if (IsAnalyticDerivEnabled())
 		{
-			return DerivativeAutogen.GenerateExpressionFunc2(*this,FMaterialDerivativeAutogen::EF2_Sub,A,B);
+			return DerivativeAutogen.GenerateExpressionFunc2(*this,FMaterialDerivativeAutogen::EFunc2::Sub,A,B);
 		}
 		else
 		{
@@ -8474,7 +6704,7 @@ int32 FHLSLMaterialTranslator::Mul(int32 A,int32 B)
 	{
 		if (IsAnalyticDerivEnabled())
 		{
-			return DerivativeAutogen.GenerateExpressionFunc2(*this,FMaterialDerivativeAutogen::EF2_Mul,A,B);
+			return DerivativeAutogen.GenerateExpressionFunc2(*this,FMaterialDerivativeAutogen::EFunc2::Mul,A,B);
 		}
 		else
 		{
@@ -8498,7 +6728,7 @@ int32 FHLSLMaterialTranslator::Div(int32 A,int32 B)
 	{
 		if (IsAnalyticDerivEnabled())
 		{
-			return DerivativeAutogen.GenerateExpressionFunc2(*this,FMaterialDerivativeAutogen::EF2_Div,A,B);
+			return DerivativeAutogen.GenerateExpressionFunc2(*this,FMaterialDerivativeAutogen::EFunc2::Div,A,B);
 		}
 		else
 		{
@@ -8550,7 +6780,7 @@ int32 FHLSLMaterialTranslator::Dot(int32 A,int32 B)
 		// Promote scalar (or truncate the bigger type)
 		if (IsAnalyticDerivEnabled())
 		{
-			return DerivativeAutogen.GenerateExpressionFunc2(*this, FMaterialDerivativeAutogen::EF2_Dot, A, B);
+			return DerivativeAutogen.GenerateExpressionFunc2(*this, FMaterialDerivativeAutogen::EFunc2::Dot, A, B);
 		}
 		else
 		{
@@ -8585,7 +6815,7 @@ int32 FHLSLMaterialTranslator::Cross(int32 A,int32 B)
 	else
 	{
 		if(IsAnalyticDerivEnabled())
-			return DerivativeAutogen.GenerateExpressionFunc2(*this, FMaterialDerivativeAutogen::EF2_Cross, A, B);
+			return DerivativeAutogen.GenerateExpressionFunc2(*this, FMaterialDerivativeAutogen::EFunc2::Cross, A, B);
 		else
 			return AddCodeChunk(MCT_Float3,TEXT("cross(%s,%s)"),*CoerceParameter(A,MCT_Float3),*CoerceParameter(B,MCT_Float3));
 	}
@@ -8600,7 +6830,7 @@ int32 FHLSLMaterialTranslator::Power(int32 Base,int32 Exponent)
 
 	if (IsAnalyticDerivEnabled())
 	{
-		return DerivativeAutogen.GenerateExpressionFunc2(*this,FMaterialDerivativeAutogen::EF2_PowPositiveClamped,Base,Exponent);
+		return DerivativeAutogen.GenerateExpressionFunc2(*this,FMaterialDerivativeAutogen::EFunc2::PowPositiveClamped,Base,Exponent);
 	}
 	else
 	{
@@ -8622,7 +6852,7 @@ int32 FHLSLMaterialTranslator::Logarithm2(int32 X)
 	}
 	else if (IsAnalyticDerivEnabled())
 	{
-		return DerivativeAutogen.GenerateExpressionFunc1(*this, FMaterialDerivativeAutogen::EF1_Log2, X);
+		return DerivativeAutogen.GenerateExpressionFunc1(*this, FMaterialDerivativeAutogen::EFunc1::Log2, X);
 	}
 	else
 	{
@@ -8643,7 +6873,7 @@ int32 FHLSLMaterialTranslator::Logarithm10(int32 X)
 	}
 	else if (IsAnalyticDerivEnabled())
 	{
-		return DerivativeAutogen.GenerateExpressionFunc1(*this, FMaterialDerivativeAutogen::EF1_Log10, X);
+		return DerivativeAutogen.GenerateExpressionFunc1(*this, FMaterialDerivativeAutogen::EFunc1::Log10, X);
 	}
 	else
 	{
@@ -8666,7 +6896,7 @@ int32 FHLSLMaterialTranslator::SquareRoot(int32 X)
 	{
 		if (IsAnalyticDerivEnabled())
 		{
-			return DerivativeAutogen.GenerateExpressionFunc1(*this,FMaterialDerivativeAutogen::EF1_Sqrt,X);
+			return DerivativeAutogen.GenerateExpressionFunc1(*this,FMaterialDerivativeAutogen::EFunc1::Sqrt,X);
 		}
 		else
 		{
@@ -8690,7 +6920,7 @@ int32 FHLSLMaterialTranslator::Length(int32 X)
 	{
 		if (IsAnalyticDerivEnabled())
 		{
-			return DerivativeAutogen.GenerateExpressionFunc1(*this,FMaterialDerivativeAutogen::EF1_Length,X);
+			return DerivativeAutogen.GenerateExpressionFunc1(*this,FMaterialDerivativeAutogen::EFunc1::Length,X);
 		}
 		else
 		{
@@ -9101,7 +7331,7 @@ int32 FHLSLMaterialTranslator::Min(int32 A,int32 B)
 	{
 		if (IsAnalyticDerivEnabled())
 		{
-			return DerivativeAutogen.GenerateExpressionFunc2(*this,FMaterialDerivativeAutogen::EF2_Min,A,B);
+			return DerivativeAutogen.GenerateExpressionFunc2(*this,FMaterialDerivativeAutogen::EFunc2::Min,A,B);
 		}
 		else
 		{
@@ -9125,7 +7355,7 @@ int32 FHLSLMaterialTranslator::Max(int32 A,int32 B)
 	{
 		if (IsAnalyticDerivEnabled())
 		{
-			return DerivativeAutogen.GenerateExpressionFunc2(*this,FMaterialDerivativeAutogen::EF2_Max,A,B);
+			return DerivativeAutogen.GenerateExpressionFunc2(*this,FMaterialDerivativeAutogen::EFunc2::Max,A,B);
 		}
 		else
 		{
@@ -9149,8 +7379,8 @@ int32 FHLSLMaterialTranslator::Clamp(int32 X,int32 A,int32 B)
 	{
 		if (IsAnalyticDerivEnabled())
 		{
-			int32 MaxAX = DerivativeAutogen.GenerateExpressionFunc2(*this,FMaterialDerivativeAutogen::EF2_Max,A,X);
-			int32 Result = DerivativeAutogen.GenerateExpressionFunc2(*this,FMaterialDerivativeAutogen::EF2_Min,MaxAX,B);
+			int32 MaxAX = DerivativeAutogen.GenerateExpressionFunc2(*this,FMaterialDerivativeAutogen::EFunc2::Max,A,X);
+			int32 Result = DerivativeAutogen.GenerateExpressionFunc2(*this,FMaterialDerivativeAutogen::EFunc2::Min,MaxAX,B);
 			return Result;
 		}
 		else
@@ -9175,7 +7405,7 @@ int32 FHLSLMaterialTranslator::Saturate(int32 X)
 	{
 		if (IsAnalyticDerivEnabled())
 		{
-			return DerivativeAutogen.GenerateExpressionFunc1(*this,FMaterialDerivativeAutogen::EF1_Saturate,X);
+			return DerivativeAutogen.GenerateExpressionFunc1(*this,FMaterialDerivativeAutogen::EFunc1::Saturate,X);
 		}
 		else
 		{
@@ -9245,7 +7475,7 @@ int32 FHLSLMaterialTranslator::ComponentMask(int32 Vector,bool R,bool G,bool B,b
 		const FString VectorDeriv = GetParameterCodeDeriv(Vector, CompiledPDV_Analytic);
 
 		FString CodeAnalytic;
-		if (VectorDerivStatus == DS_Valid)
+		if (VectorDerivStatus == EDerivativeStatus::Valid)
 		{
 			CodeAnalytic = DerivativeAutogen.ConstructDeriv(
 				FString::Printf(TEXT("%s.Value.%s"), *VectorDeriv, *MaskString),
@@ -9253,7 +7483,7 @@ int32 FHLSLMaterialTranslator::ComponentMask(int32 Vector,bool R,bool G,bool B,b
 				FString::Printf(TEXT("%s.Ddy.%s"), *VectorDeriv, *MaskString),
 				GetDerivTypeIndex(ResultType)
 			);
-			return AddCodeChunkInnerDeriv(*CodeFinite, *CodeAnalytic, ResultType, false, DS_Valid);
+			return AddCodeChunkInnerDeriv(*CodeFinite, *CodeAnalytic, ResultType, false, EDerivativeStatus::Valid);
 		}
 		else
 		{
@@ -9290,7 +7520,7 @@ int32 FHLSLMaterialTranslator::AppendVector(int32 A,int32 B)
 		const EDerivativeStatus BDerivativeStatus = GetDerivativeStatus(B);
 		if (IsAnalyticDerivEnabled() && IsDerivativeValid(ADerivativeStatus) && IsDerivativeValid(BDerivativeStatus))
 		{
-			if (ADerivativeStatus == DS_Zero && BDerivativeStatus == DS_Zero)
+			if (ADerivativeStatus == EDerivativeStatus::Zero && BDerivativeStatus == EDerivativeStatus::Zero)
 			{
 				return AddInlinedCodeChunkZeroDeriv(ResultType, TEXT("MaterialFloat%u(%s,%s)"), NumResultComponents, *GetParameterCode(A), *GetParameterCode(B));
 			}
@@ -9301,14 +7531,14 @@ int32 FHLSLMaterialTranslator::AppendVector(int32 A,int32 B)
 				FString B_DDX = GetFloatZeroVector(NumComponentsB);
 				FString B_DDY = GetFloatZeroVector(NumComponentsB);
 
-				if (ADerivativeStatus == DS_Valid)
+				if (ADerivativeStatus == EDerivativeStatus::Valid)
 				{
 					FString Deriv = *GetParameterCodeDeriv(A, CompiledPDV_Analytic);
 					A_DDX = Deriv + TEXT(".Ddx");
 					A_DDY = Deriv + TEXT(".Ddy");
 				}
 
-				if (BDerivativeStatus == DS_Valid)
+				if (BDerivativeStatus == EDerivativeStatus::Valid)
 				{
 					FString Deriv = *GetParameterCodeDeriv(B, CompiledPDV_Analytic);
 					B_DDX = Deriv + TEXT(".Ddx");
@@ -9318,7 +7548,7 @@ int32 FHLSLMaterialTranslator::AppendVector(int32 A,int32 B)
 				FString DDXCode = FString::Printf(TEXT("MaterialFloat%u(%s, %s)"), NumResultComponents, *A_DDX, *B_DDX);
 				FString DDYCode = FString::Printf(TEXT("MaterialFloat%u(%s, %s)"), NumResultComponents, *A_DDY, *B_DDY);
 				FString AnalyticCode = DerivativeAutogen.ConstructDeriv(FiniteCode, *DDXCode, *DDYCode, GetDerivTypeIndex(ResultType));
-				return AddCodeChunkInnerDeriv(*FiniteCode, *AnalyticCode, ResultType, false, DS_Valid);
+				return AddCodeChunkInnerDeriv(*FiniteCode, *AnalyticCode, ResultType, false, EDerivativeStatus::Valid);
 			}
 		}
 		else
@@ -9617,7 +7847,7 @@ int32 FHLSLMaterialTranslator::LightmapUVs()
 			FiniteCode,
 			TEXT("GetLightmapUVs_DDX(Parameters)"),
 			TEXT("GetLightmapUVs_DDY(Parameters)"), 1);
-		return AddCodeChunkInnerDeriv(FiniteCode, *AnalyticCode, MCT_Float2, false, DS_Valid);
+		return AddCodeChunkInnerDeriv(FiniteCode, *AnalyticCode, MCT_Float2, false, EDerivativeStatus::Valid);
 	}
 	else
 	{
@@ -9815,7 +8045,7 @@ int32 FHLSLMaterialTranslator::DDX( int32 A )
 	const FDerivInfo ADerivInfo = GetDerivInfo(A);
 	const FString FiniteCode = FString::Printf(TEXT("ddx(%s)"), *GetParameterCode(A));
 
-	if (IsAnalyticDerivEnabled() && ADerivInfo.DerivativeStatus == DS_Valid)
+	if (IsAnalyticDerivEnabled() && ADerivInfo.DerivativeStatus == EDerivativeStatus::Valid)
 	{
 		FString ADeriv = GetParameterCodeDeriv(A, CompiledPDV_Analytic);
 
@@ -9826,11 +8056,11 @@ int32 FHLSLMaterialTranslator::DDX( int32 A )
 			ADerivInfo.TypeIndex
 		);
 
-		return AddCodeChunkInnerDeriv(*FiniteCode, *AnalyticCode, ADerivInfo.Type, false, DS_Valid);
+		return AddCodeChunkInnerDeriv(*FiniteCode, *AnalyticCode, ADerivInfo.Type, false, EDerivativeStatus::Valid);
 	}
 	else
 	{
-		return AddCodeChunkInnerDeriv(*FiniteCode, ADerivInfo.Type, false, (ADerivInfo.DerivativeStatus == DS_Zero) ? DS_Zero : DS_NotValid);
+		return AddCodeChunkInnerDeriv(*FiniteCode, ADerivInfo.Type, false, (ADerivInfo.DerivativeStatus == EDerivativeStatus::Zero) ? EDerivativeStatus::Zero : EDerivativeStatus::NotValid);
 	}
 }
 
@@ -9854,7 +8084,7 @@ int32 FHLSLMaterialTranslator::DDY( int32 A )
 	const FDerivInfo ADerivInfo = GetDerivInfo(A);
 	const FString FiniteCode = FString::Printf(TEXT("ddy(%s)"), *GetParameterCode(A));
 
-	if (IsAnalyticDerivEnabled() && ADerivInfo.DerivativeStatus == DS_Valid)
+	if (IsAnalyticDerivEnabled() && ADerivInfo.DerivativeStatus == EDerivativeStatus::Valid)
 	{
 		FString ADeriv = GetParameterCodeDeriv(A, CompiledPDV_Analytic);
 
@@ -9865,11 +8095,11 @@ int32 FHLSLMaterialTranslator::DDY( int32 A )
 			GetDerivTypeIndex(ADerivInfo.Type)
 			);
 
-		return AddCodeChunkInnerDeriv(*FiniteCode, *AnalyticCode, ADerivInfo.Type, false, DS_Valid);
+		return AddCodeChunkInnerDeriv(*FiniteCode, *AnalyticCode, ADerivInfo.Type, false, EDerivativeStatus::Valid);
 	}
 	else
 	{
-		return AddCodeChunkInnerDeriv(*FiniteCode, ADerivInfo.Type, false, (ADerivInfo.DerivativeStatus == DS_Zero) ? DS_Zero : DS_NotValid);
+		return AddCodeChunkInnerDeriv(*FiniteCode, ADerivInfo.Type, false, (ADerivInfo.DerivativeStatus == EDerivativeStatus::Zero) ? EDerivativeStatus::Zero : EDerivativeStatus::NotValid);
 	}
 }
 
diff --git a/Engine/Source/Runtime/Engine/Private/Materials/HLSLMaterialTranslator.h b/Engine/Source/Runtime/Engine/Private/Materials/HLSLMaterialTranslator.h
--- a/Engine/Source/Runtime/Engine/Private/Materials/HLSLMaterialTranslator.h
+++ b/Engine/Source/Runtime/Engine/Private/Materials/HLSLMaterialTranslator.h
@@ -47,14 +47,13 @@
 #include "Containers/HashTable.h"
 #include "Engine/Texture2D.h"
 #include "StrataMaterial.h"
+#include "HLSLMaterialDerivativeAutogen.h"
 #endif
 
 class Error;
 
 #if WITH_EDITORONLY_DATA
 
-class FHLSLMaterialTranslator;
-
 enum EMaterialExpressionVisitResult
 {
 	MVR_CONTINUE,
@@ -68,28 +67,7 @@ public:
 	virtual EMaterialExpressionVisitResult Visit(UMaterialExpression* InExpression) = 0;
 };
 
-/**
-* For a node, the known information of the partial derivatives.
-* DS_NotAware	- This node is made by a function that has no knowledge of analytic partial derivatives.
-* DS_NotValid	- This node is aware of partial derivatives, and knows that one of its source inputs is not partial derivative aware, and therefore its value is not to be used.
-* DS_Zero		- This node is aware of partial derivatives, and knows that it's value is zero, as is the case for uniform parameters.
-* DS_Valid		- This node is aware of partial derivatives, and knows that it has a calculated value.
-**/
-enum EDerivativeStatus
-{
-	DS_NotAware,
-	DS_NotValid,
-	DS_Zero,
-	DS_Valid,
-	DS_MAX,
-};
-
-struct FDerivInfo
-{
-	EMaterialValueType	Type;
-	int32				TypeIndex;
-	EDerivativeStatus	DerivativeStatus;
-};
+uint32 GetNumComponents(EMaterialValueType Type);
 
 struct FShaderCodeChunk
 {
@@ -210,103 +188,6 @@ struct FMaterialCustomExpressionEntry
 	TArray<int32> OutputCodeIndex;
 };
 
-class FMaterialDerivativeAutogen
-{
-public:
-	// Unary functions
-	enum EFunc1
-	{
-		EF1_Abs,
-		EF1_Log2,
-		EF1_Log10,
-		EF1_Exp,
-		EF1_Sin,
-		EF1_Cos,
-		EF1_Tan,
-		EF1_Asin,
-		EF1_AsinFast,
-		EF1_Acos,
-		EF1_AcosFast,
-		EF1_Atan,
-		EF1_AtanFast,
-		EF1_Sqrt,
-		EF1_Rcp,
-		EF1_Rsqrt,
-		EF1_Saturate,
-		EF1_Frac,
-		EF1_Length,
-		EF1_InvLength,
-		EF1_Normalize,
-		EF1_Num
-	};
-
-	// Binary functions
-	enum EFunc2
-	{
-		EF2_Add,
-		EF2_Sub,
-		EF2_Mul,
-		EF2_Div,
-		EF2_Fmod,
-		EF2_Max,
-		EF2_Min,
-		EF2_Dot,	// Depends on Add/Mul, so it must come after them
-		EF2_Pow,
-		EF2_PowPositiveClamped,
-		EF2_Cross,
-		EF2_Atan2,
-		EF2_Atan2Fast,
-		EF2_Num
-	};
-
-	// Note that the type index is from [0,3] for float1 to float4. I.e. A float3 would be index 2.
-	static int32 GetFunc1ReturnNumComponents(int32 SrcTypeIndex, EFunc1 Op);
-	static int32 GetFunc2ReturnNumComponents(int32 LhsTypeIndex, int32 RhsTypeIndex, EFunc2 Op);
-
-	void EnableGeneratedDepencencies();
-
-	int32 GenerateExpressionFunc1(FHLSLMaterialTranslator& Translator, EFunc1 Op, int32 SrcCode);
-	int32 GenerateExpressionFunc2(FHLSLMaterialTranslator& Translator, EFunc2 Op, int32 LhsCode, int32 RhsCode);
-
-	int32 GenerateLerpFunc(FHLSLMaterialTranslator& Translator, int32 A, int32 B, int32 S);
-	int32 GenerateRotateScaleOffsetTexCoordsFunc(FHLSLMaterialTranslator& Translator, int32 TexCoord, int32 RotationScale, int32 Offset);
-	int32 GenerateIfFunc(FHLSLMaterialTranslator& Translator, int32 A, int32 B, int32 Greater, int32 Equal, int32 Less, int32 Threshold);
-	
-	FString GenerateUsedFunctions(FHLSLMaterialTranslator& Translator);
-
-	FString ApplyUnMirror(FString Value, bool bUnMirrorU, bool bUnMirrorV);
-
-	FString CoerceValueRaw(const FString& Token, int32 SrcType, EDerivativeStatus SrcStatus, int32 DstType);
-	FString CoerceValueDeriv(const FString& Token, int32 SrcType, EDerivativeStatus SrcStatus, int32 DstType);
-
-	FString ConstructDeriv(const FString& Value, const FString& Ddx, const FString& Ddy, int32 DstType);
-	FString ConstructDerivFinite(const FString& Value, int32 DstType);
-
-	FString ConvertDeriv(const FString& Value, int32 DstType, int32 SrcType);
-
-	FString ExtractElement(const FString& Value, int32 SrcType, int32 ElementIndex);
-
-private:
-	// State to keep track of which derivative functions have been used and need to be generated.
-	bool bFunc1OpIsEnabled[EF1_Num][4] = {};
-	bool bFunc2OpIsEnabled[EF2_Num][4] = {};
-	bool bLerpEnabled[4] = {};
-	bool bIfEnabled[4] = {};
-	bool bIf2Enabled[4] = {};
-	bool bUnMirrorEnabled[2][2] = {};
-
-	bool bRotateScaleOffsetTexCoords = false;
-
-	bool bConstructDerivEnabled[4] = {};
-	bool bConstructConstantDerivEnabled[4] = {};
-	bool bConstructFiniteDerivEnabled[4] = {};
-
-	bool bConvertDerivEnabled[4][4] = {};
-
-	bool bExtractIndexEnabled[4] = {};
-
-	bool bSelectElemHelperEnabled[4] = {};
-};
 
 struct FMaterialDerivativeVariation
 {
@@ -663,21 +544,21 @@ protected:
 	int32 AddCodeChunk(EMaterialValueType Type, const TCHAR* Format, Types... Args)
 	{
 		static_assert(TAnd<TIsValidVariadicFunctionArg<Types>...>::Value, "Invalid argument(s) passed to AddCodeChunk");
-		return AddCodeChunkInner(Type, DS_NotAware, false, Format, Args...);
+		return AddCodeChunkInner(Type, EDerivativeStatus::NotAware, false, Format, Args...);
 	}
 
 	template <typename... Types>
 	int32 AddCodeChunkZeroDeriv(EMaterialValueType Type, const TCHAR* Format, Types... Args)
 	{
 		static_assert(TAnd<TIsValidVariadicFunctionArg<Types>...>::Value, "Invalid argument(s) passed to AddCodeChunkZeroDeriv");
-		return AddCodeChunkInner(Type, DS_Zero, false, Format, Args...);
+		return AddCodeChunkInner(Type, EDerivativeStatus::Zero, false, Format, Args...);
 	}
 
 	template <typename... Types>
 	int32 AddCodeChunkFiniteDeriv(EMaterialValueType Type, const TCHAR* Format, Types... Args)
 	{
 		static_assert(TAnd<TIsValidVariadicFunctionArg<Types>...>::Value, "Invalid argument(s) passed to AddCodeChunkFiniteDeriv");
-		return AddCodeChunkInner(Type, DS_NotValid, false, Format, Args...);
+		return AddCodeChunkInner(Type, EDerivativeStatus::NotValid, false, Format, Args...);
 	}
 	
 
@@ -689,21 +570,21 @@ protected:
 	int32 AddInlinedCodeChunk(EMaterialValueType Type, const TCHAR* Format, Types... Args)
 	{
 		static_assert(TAnd<TIsValidVariadicFunctionArg<Types>...>::Value, "Invalid argument(s) passed to AddInlinedCodeChunk");
-		return AddCodeChunkInner(Type, DS_NotAware, true, Format, Args...);
+		return AddCodeChunkInner(Type, EDerivativeStatus::NotAware, true, Format, Args...);
 	}
 
 	template <typename... Types>
 	int32 AddInlinedCodeChunkZeroDeriv(EMaterialValueType Type, const TCHAR* Format, Types... Args)
 	{
 		static_assert(TAnd<TIsValidVariadicFunctionArg<Types>...>::Value, "Invalid argument(s) passed to AddInlinedCodeChunkZeroDeriv");
-		return AddCodeChunkInner(Type, DS_Zero, true, Format, Args...);
+		return AddCodeChunkInner(Type, EDerivativeStatus::Zero, true, Format, Args...);
 	}
 
 	template <typename... Types>
 	int32 AddInlinedCodeChunkFiniteDeriv(EMaterialValueType Type, const TCHAR* Format, Types... Args)
 	{
 		static_assert(TAnd<TIsValidVariadicFunctionArg<Types>...>::Value, "Invalid argument(s) passed to AddInlinedCodeChunkFiniteDeriv");
-		return AddCodeChunkInner(Type, DS_NotValid, true, Format, Args...);
+		return AddCodeChunkInner(Type, EDerivativeStatus::NotValid, true, Format, Args...);
 	}
 
 	int32 AddUniformExpressionInner(uint64 Hash, FMaterialUniformExpression* UniformExpression, EMaterialValueType Type, const TCHAR* FormattedCode);
