Don't register/unregister references to deleted actors.

diff --git a/Engine/Source/Runtime/Engine/Private/WorldPartition/WorldPartitionEditorCell.cpp b/Engine/Source/Runtime/Engine/Private/WorldPartition/WorldPartitionEditorCell.cpp
--- a/Engine/Source/Runtime/Engine/Private/WorldPartition/WorldPartitionEditorCell.cpp
+++ b/Engine/Source/Runtime/Engine/Private/WorldPartition/WorldPartitionEditorCell.cpp
@@ -52,7 +52,11 @@ void UWorldPartitionEditorCell::AddActor(const FGuid& Source, const FWorldPartit
 		for (const FGuid& ReferenceGuid : ActorHandle->GetReferences())
 		{
 			FWorldPartitionHandle ReferenceHandle(WorldPartition, ReferenceGuid);
-			AddActor(ActorHandle->GetGuid(), ReferenceHandle);
+
+			if (ReferenceHandle.IsValid())
+			{
+				AddActor(ActorHandle->GetGuid(), ReferenceHandle);
+			}
 		}
 	}
 }
@@ -83,7 +87,11 @@ void UWorldPartitionEditorCell::RemoveActor(const FGuid& Source, const FWorldPar
 		for (const FGuid& ReferenceGuid : ActorHandle->GetReferences())
 		{
 			FWorldPartitionHandle ReferenceHandle(WorldPartition, ReferenceGuid);
-			RemoveActor(ActorHandle->GetGuid(), ReferenceHandle);
+
+			if (ReferenceHandle.IsValid())
+			{
+				RemoveActor(ActorHandle->GetGuid(), ReferenceHandle);
+			}
 		}
 	}
 }
