Use FReadOnlyCVARCache for strata runtime checks to work on console with correct cvar value.
Fixes strata translucency.

diff --git a/Engine/Source/Runtime/Engine/Private/Materials/Material.cpp b/Engine/Source/Runtime/Engine/Private/Materials/Material.cpp
--- a/Engine/Source/Runtime/Engine/Private/Materials/Material.cpp
+++ b/Engine/Source/Runtime/Engine/Private/Materials/Material.cpp
@@ -3633,10 +3633,7 @@ void UMaterial::ConvertMaterialToStrataMaterial()
 		return;
 	}
 
-	const URendererSettings* RendererSettings = GetDefault<URendererSettings>();
-	const bool bStrataEnabled = RendererSettings && RendererSettings->bEnableStrata;
-
-	if (bStrataEnabled && !FrontMaterial.IsConnected())
+	if (FReadOnlyCVARCache::Get().bEnableStrata && !FrontMaterial.IsConnected())
 	{
 		auto MoveConnectionTo = [](auto& OldNodeInput, UMaterialExpression* NewNode, uint32 NewInputIndex)
 		{
@@ -4922,9 +4919,7 @@ void UMaterial::RebuildShadingModelField()
 {
 	ShadingModels.ClearShadingModels();
 
-	const URendererSettings* RendererSettings = GetDefault<URendererSettings>();
-	const bool bStrataEnabled = RendererSettings && RendererSettings->bEnableStrata;
-	if (bStrataEnabled && FrontMaterial.IsConnected())
+	if (FReadOnlyCVARCache::Get().bEnableStrata && FrontMaterial.IsConnected())
 	{
 		FStrataMaterialInfo StrataMaterialInfo;
 		check(this->FrontMaterial.Expression);
@@ -6719,8 +6714,7 @@ static bool IsPropertyActive_Internal(EMaterialProperty InProperty,
 		break;
 	case MP_FrontMaterial:
 		{
-			const URendererSettings* RendererSettings = GetDefault<URendererSettings>();
-			Active = RendererSettings && RendererSettings->bEnableStrata;
+			Active = FReadOnlyCVARCache::Get().bEnableStrata;
 			break;
 		}
 	case MP_MaterialAttributes:
diff --git a/Engine/Source/Runtime/Engine/Private/Materials/MaterialShared.cpp b/Engine/Source/Runtime/Engine/Private/Materials/MaterialShared.cpp
--- a/Engine/Source/Runtime/Engine/Private/Materials/MaterialShared.cpp
+++ b/Engine/Source/Runtime/Engine/Private/Materials/MaterialShared.cpp
@@ -1611,8 +1611,7 @@ bool FMaterialResource::HasAmbientOcclusionConnected() const
 bool FMaterialResource::IsStrataMaterial() const
 {
 	// STRATA_TODO IsStrataMaterial should go away once Strata implementation is finished
-	const URendererSettings* RendererSettings = GetDefault<URendererSettings>();
-	if (RendererSettings && RendererSettings->bEnableStrata)
+	if (FReadOnlyCVARCache::Get().bEnableStrata)
 	{
 		return Material->bUseMaterialAttributes ?
 			Material->MaterialAttributes.IsConnected(MP_FrontMaterial) :
diff --git a/Engine/Source/Runtime/Engine/Private/SceneManagement.cpp b/Engine/Source/Runtime/Engine/Private/SceneManagement.cpp
--- a/Engine/Source/Runtime/Engine/Private/SceneManagement.cpp
+++ b/Engine/Source/Runtime/Engine/Private/SceneManagement.cpp
@@ -1122,6 +1122,7 @@ void FReadOnlyCVARCache::Init()
 	static const auto CVarVertexFoggingForOpaque = IConsoleManager::Get().FindTConsoleVariableDataInt(TEXT("r.VertexFoggingForOpaque"));	
 	static const auto CVarAllowStaticLighting = IConsoleManager::Get().FindTConsoleVariableDataInt(TEXT("r.AllowStaticLighting"));
 	static const auto CVarSupportSkyAtmosphere = IConsoleManager::Get().FindTConsoleVariableDataInt(TEXT("r.SupportSkyAtmosphere"));
+	static const auto CVarStrata = IConsoleManager::Get().FindTConsoleVariableDataInt(TEXT("r.Strata"));
 
 	static const auto CVarMobileAllowMovableDirectionalLights = IConsoleManager::Get().FindTConsoleVariableDataInt(TEXT("r.Mobile.AllowMovableDirectionalLights"));
 	static const auto CVarMobileEnableStaticAndCSMShadowReceivers = IConsoleManager::Get().FindTConsoleVariableDataInt(TEXT("r.Mobile.EnableStaticAndCSMShadowReceivers"));
@@ -1140,6 +1141,7 @@ void FReadOnlyCVARCache::Init()
 	bEnableLowQualityLightmaps = !CVarSupportLowQualityLightmaps || CVarSupportLowQualityLightmaps->GetValueOnAnyThread() != 0 || bForceAllPermutations;
 	bAllowStaticLighting = CVarAllowStaticLighting->GetValueOnAnyThread() != 0;
 	bSupportSkyAtmosphere = !CVarSupportSkyAtmosphere || CVarSupportSkyAtmosphere->GetValueOnAnyThread() != 0;
+	bEnableStrata = CVarStrata ? CVarStrata->GetValueOnAnyThread() != 0 : false;
 
 	// mobile
 	bMobileAllowMovableDirectionalLights = CVarMobileAllowMovableDirectionalLights->GetValueOnAnyThread() != 0;
diff --git a/Engine/Source/Runtime/Engine/Public/SceneManagement.h b/Engine/Source/Runtime/Engine/Public/SceneManagement.h
--- a/Engine/Source/Runtime/Engine/Public/SceneManagement.h
+++ b/Engine/Source/Runtime/Engine/Public/SceneManagement.h
@@ -3189,6 +3189,7 @@ struct FReadOnlyCVARCache
 	bool bEnableLowQualityLightmaps;
 	bool bAllowStaticLighting;
 	bool bSupportSkyAtmosphere;
+	bool bEnableStrata;
 
 	// Mobile specific
 	bool bMobileAllowMovableDirectionalLights;
