FSquare2DGridHelper is now using sparse arrays to avoid very long creation times when going in PIE for very large worlds.
Along with 15541845, this changes make it so going into PIE for very large worlds is instant and uses far less memory.

diff --git a/Engine/Source/Runtime/Engine/Private/WorldPartition/RuntimeSpatialHash/RuntimeSpatialHashGridHelper.cpp b/Engine/Source/Runtime/Engine/Private/WorldPartition/RuntimeSpatialHash/RuntimeSpatialHashGridHelper.cpp
--- a/Engine/Source/Runtime/Engine/Private/WorldPartition/RuntimeSpatialHash/RuntimeSpatialHashGridHelper.cpp
+++ b/Engine/Source/Runtime/Engine/Private/WorldPartition/RuntimeSpatialHash/RuntimeSpatialHashGridHelper.cpp
@@ -51,6 +51,9 @@ FSquare2DGridHelper::FSquare2DGridHelper(const FBox& InWorldBounds, const FVecto
 		CurrentCellSize <<= 1;
 		CurrentGridSize >>= 1;
 	}
+
+	// Make sure the always loaded cell exists
+	GetAlwaysLoadedCell();
 }
 
 void FSquare2DGridHelper::ForEachCells(TFunctionRef<void(const FIntVector&)> InOperation) const
@@ -93,20 +96,15 @@ void FSquare2DGridHelper::ValidateSingleActorReferer()
 	TSet<FActorInstance> ActorUsage;
 	for (int32 Level = 0; Level < Levels.Num() - 1; Level++)
 	{
-		const int32 CurrentGridSize = Levels[Level].GridSize;
-		for (int32 y = 0; y < CurrentGridSize; y++)
+		for (const FSquare2DGridHelper::FGridLevel::FGridCell& ThisCell : Levels[Level].Cells)
 		{
-			for (int32 x = 0; x < CurrentGridSize; x++)
+			for (const FGridLevel::FGridCellDataChunk& DataChunk : ThisCell.GetDataChunks())
 			{
-				FGridLevel::FGridCell& ThisCell = Levels[Level].Cells[y * CurrentGridSize + x];
-				for (const FGridLevel::FGridCellDataChunk& DataChunk : ThisCell.GetDataChunks())
+				for (const FActorInstance& ActorInstance : DataChunk.GetActors())
 				{
-					for (const FActorInstance& ActorInstance : DataChunk.GetActors())
-					{
-						bool bWasAlreadyInSet;
-						ActorUsage.Add(ActorInstance, &bWasAlreadyInSet);
-						check(!bWasAlreadyInSet);
-					}
+					bool bWasAlreadyInSet;
+					ActorUsage.Add(ActorInstance, &bWasAlreadyInSet);
+					check(!bWasAlreadyInSet);
 				}
 			}
 		}
diff --git a/Engine/Source/Runtime/Engine/Private/WorldPartition/RuntimeSpatialHash/RuntimeSpatialHashGridHelper.h b/Engine/Source/Runtime/Engine/Private/WorldPartition/RuntimeSpatialHash/RuntimeSpatialHashGridHelper.h
--- a/Engine/Source/Runtime/Engine/Private/WorldPartition/RuntimeSpatialHash/RuntimeSpatialHashGridHelper.h
+++ b/Engine/Source/Runtime/Engine/Private/WorldPartition/RuntimeSpatialHash/RuntimeSpatialHashGridHelper.h
@@ -307,11 +307,7 @@ struct FSquare2DGridHelper
 
 		inline FGridLevel(const FVector2D& InOrigin, int32 InCellSize, int32 InGridSize)
 			: FGrid2D(InOrigin, InCellSize, InGridSize)
-		{
-#if WITH_EDITOR
-			Cells.InsertDefaulted(0, GridSize * GridSize);
-#endif
-		}
+		{}
 
 #if WITH_EDITOR
 		/**
@@ -322,9 +318,23 @@ struct FSquare2DGridHelper
 		inline FGridCell& GetCell(const FIntVector2& InCoords)
 		{
 			check(IsValidCoords(InCoords));
+
 			uint32 CellIndex;
 			GetCellIndex(InCoords, CellIndex);
-			return Cells[CellIndex];
+
+			int32 CellIndexMapping;
+			int32* CellIndexMappingPtr = CellsMapping.Find(CellIndex);
+			if (CellIndexMappingPtr)
+			{
+				CellIndexMapping = *CellIndexMappingPtr;
+			}
+			else
+			{
+				CellIndexMapping = Cells.AddDefaulted();
+				CellsMapping.Add(CellIndex, CellIndexMapping);
+			}
+
+			return Cells[CellIndexMapping];
 		}
 
 		/**
@@ -335,12 +345,16 @@ struct FSquare2DGridHelper
 		inline const FGridCell& GetCell(const FIntVector2& InCoords) const
 		{
 			check(IsValidCoords(InCoords));
+
 			uint32 CellIndex;
 			GetCellIndex(InCoords, CellIndex);
-			return Cells[CellIndex];
+
+			int32 CellIndexMapping = CellsMapping.FindChecked(CellIndex);
+			return Cells[CellIndexMapping];
 		}
 
 		TArray<FGridCell> Cells;
+		TMap<int32, int32> CellsMapping;
 #endif
 	};
 
@@ -351,10 +365,10 @@ struct FSquare2DGridHelper
 	inline FGridLevel& GetLowestLevel() { return Levels[0]; }
 
 	// Returns the always loaded (top level) cell
-	inline FGridLevel::FGridCell& GetAlwaysLoadedCell() { return Levels.Last().Cells[0]; }
+	inline FGridLevel::FGridCell& GetAlwaysLoadedCell() { return Levels.Last().GetCell(FIntVector2(0,0)); }
 
 	// Returns the always loaded (top level) cell
-	inline const FGridLevel::FGridCell& GetAlwaysLoadedCell() const { return Levels.Last().Cells[0]; }
+	inline const FGridLevel::FGridCell& GetAlwaysLoadedCell() const { return Levels.Last().GetCell(FIntVector2(0,0)); }
 
 	// Returns the cell at the given coord
 	inline const FGridLevel::FGridCell& GetCell(const FIntVector& InCoords) const { return Levels[InCoords.Z].GetCell(FIntVector2(InCoords.X, InCoords.Y)); }
@@ -391,8 +405,8 @@ struct FSquare2DGridHelper
 	 */
 	int32 ForEachIntersectingCells(const FSphere& InSphere, TFunctionRef<void(const FIntVector&)> InOperation) const;
 
-	// Validates that actor is not referenced by multiple cells
 #if WITH_EDITOR
+	// Validates that actor is not referenced by multiple cells
 	void ValidateSingleActorReferer();
 #endif
 
diff --git a/Engine/Source/Runtime/Engine/Private/WorldPartition/WorldPartitionRuntimeSpatialHash.cpp b/Engine/Source/Runtime/Engine/Private/WorldPartition/WorldPartitionRuntimeSpatialHash.cpp
--- a/Engine/Source/Runtime/Engine/Private/WorldPartition/WorldPartitionRuntimeSpatialHash.cpp
+++ b/Engine/Source/Runtime/Engine/Private/WorldPartition/WorldPartitionRuntimeSpatialHash.cpp
@@ -612,13 +612,14 @@ bool UWorldPartitionRuntimeSpatialHash::CreateStreamingGrid(const FSpatialHashRu
 
 		FSpatialHashStreamingGridLevel& GridLevel = CurrentStreamingGrid.GridLevels.AddDefaulted_GetRef();
 
-		int32 CellIndex = INDEX_NONE;
-		for (const FSquare2DGridHelper::FGridLevel::FGridCell& TempCell : TempLevel.Cells)
+		for (const auto& TempCellMapping : TempLevel.CellsMapping)
 		{
-			CellIndex++;
+			const uint32 CellIndex = TempCellMapping.Key;
 			const int32 CellCoordX = CellIndex % TempLevel.GridSize;
 			const int32 CellCoordY = CellIndex / TempLevel.GridSize;
 
+			const FSquare2DGridHelper::FGridLevel::FGridCell& TempCell = TempLevel.Cells[TempCellMapping.Value];
+
 			for (const FSquare2DGridHelper::FGridLevel::FGridCellDataChunk& GridCellDataChunk : TempCell.GetDataChunks())
 			{
 				// Cell cannot be treated as always loaded if it has data layers
