Fix non-editor builds.

diff --git a/Engine/Source/Runtime/Engine/Private/WorldPartition/RuntimeSpatialHash/RuntimeSpatialHashGridHelper.cpp b/Engine/Source/Runtime/Engine/Private/WorldPartition/RuntimeSpatialHash/RuntimeSpatialHashGridHelper.cpp
--- a/Engine/Source/Runtime/Engine/Private/WorldPartition/RuntimeSpatialHash/RuntimeSpatialHashGridHelper.cpp
+++ b/Engine/Source/Runtime/Engine/Private/WorldPartition/RuntimeSpatialHash/RuntimeSpatialHashGridHelper.cpp
@@ -52,8 +52,10 @@ FSquare2DGridHelper::FSquare2DGridHelper(const FBox& InWorldBounds, const FVecto
 		CurrentGridSize >>= 1;
 	}
 
+#if WITH_EDITOR
 	// Make sure the always loaded cell exists
 	GetAlwaysLoadedCell();
+#endif
 }
 
 void FSquare2DGridHelper::ForEachCells(TFunctionRef<void(const FIntVector&)> InOperation) const
