Undo 15555840 + convert-to-edit. this because FReadOnlyCVARCache is not created soon enough (in init) while some mateiral compilation can happen in preinit due to ProcessNewlyLoadedUObjects.

diff --git a/Engine/Source/Runtime/Engine/Private/Materials/Material.cpp b/Engine/Source/Runtime/Engine/Private/Materials/Material.cpp
--- a/Engine/Source/Runtime/Engine/Private/Materials/Material.cpp
+++ b/Engine/Source/Runtime/Engine/Private/Materials/Material.cpp
@@ -3633,7 +3633,10 @@ void UMaterial::ConvertMaterialToStrataMaterial()
 		return;
 	}
 
-	if (FReadOnlyCVARCache::Get().bEnableStrata && !FrontMaterial.IsConnected())
+	const URendererSettings* RendererSettings = GetDefault<URendererSettings>();
+	const bool bStrataEnabled = RendererSettings && RendererSettings->bEnableStrata;
+
+	if (bStrataEnabled && !FrontMaterial.IsConnected())
 	{
 		auto MoveConnectionTo = [](auto& OldNodeInput, UMaterialExpression* NewNode, uint32 NewInputIndex)
 		{
@@ -4919,7 +4922,9 @@ void UMaterial::RebuildShadingModelField()
 {
 	ShadingModels.ClearShadingModels();
 
-	if (FReadOnlyCVARCache::Get().bEnableStrata && FrontMaterial.IsConnected())
+	const URendererSettings* RendererSettings = GetDefault<URendererSettings>();
+	const bool bStrataEnabled = RendererSettings && RendererSettings->bEnableStrata;
+	if (bStrataEnabled && FrontMaterial.IsConnected())
 	{
 		FStrataMaterialInfo StrataMaterialInfo;
 		check(this->FrontMaterial.Expression);
@@ -6714,7 +6719,8 @@ static bool IsPropertyActive_Internal(EMaterialProperty InProperty,
 		break;
 	case MP_FrontMaterial:
 		{
-			Active = FReadOnlyCVARCache::Get().bEnableStrata;
+			const URendererSettings* RendererSettings = GetDefault<URendererSettings>();
+			Active = RendererSettings && RendererSettings->bEnableStrata;
 			break;
 		}
 	case MP_MaterialAttributes:
diff --git a/Engine/Source/Runtime/Engine/Private/Materials/MaterialShared.cpp b/Engine/Source/Runtime/Engine/Private/Materials/MaterialShared.cpp
--- a/Engine/Source/Runtime/Engine/Private/Materials/MaterialShared.cpp
+++ b/Engine/Source/Runtime/Engine/Private/Materials/MaterialShared.cpp
@@ -1611,7 +1611,8 @@ bool FMaterialResource::HasAmbientOcclusionConnected() const
 bool FMaterialResource::IsStrataMaterial() const
 {
 	// STRATA_TODO IsStrataMaterial should go away once Strata implementation is finished
-	if (FReadOnlyCVARCache::Get().bEnableStrata)
+	const URendererSettings* RendererSettings = GetDefault<URendererSettings>();
+	if (RendererSettings && RendererSettings->bEnableStrata)
 	{
 		return Material->bUseMaterialAttributes ?
 			Material->MaterialAttributes.IsConnected(MP_FrontMaterial) :
diff --git a/Engine/Source/Runtime/Engine/Private/SceneManagement.cpp b/Engine/Source/Runtime/Engine/Private/SceneManagement.cpp
--- a/Engine/Source/Runtime/Engine/Private/SceneManagement.cpp
+++ b/Engine/Source/Runtime/Engine/Private/SceneManagement.cpp
@@ -1122,7 +1122,6 @@ void FReadOnlyCVARCache::Init()
 	static const auto CVarVertexFoggingForOpaque = IConsoleManager::Get().FindTConsoleVariableDataInt(TEXT("r.VertexFoggingForOpaque"));	
 	static const auto CVarAllowStaticLighting = IConsoleManager::Get().FindTConsoleVariableDataInt(TEXT("r.AllowStaticLighting"));
 	static const auto CVarSupportSkyAtmosphere = IConsoleManager::Get().FindTConsoleVariableDataInt(TEXT("r.SupportSkyAtmosphere"));
-	static const auto CVarStrata = IConsoleManager::Get().FindTConsoleVariableDataInt(TEXT("r.Strata"));
 
 	static const auto CVarMobileAllowMovableDirectionalLights = IConsoleManager::Get().FindTConsoleVariableDataInt(TEXT("r.Mobile.AllowMovableDirectionalLights"));
 	static const auto CVarMobileEnableStaticAndCSMShadowReceivers = IConsoleManager::Get().FindTConsoleVariableDataInt(TEXT("r.Mobile.EnableStaticAndCSMShadowReceivers"));
@@ -1141,7 +1140,6 @@ void FReadOnlyCVARCache::Init()
 	bEnableLowQualityLightmaps = !CVarSupportLowQualityLightmaps || CVarSupportLowQualityLightmaps->GetValueOnAnyThread() != 0 || bForceAllPermutations;
 	bAllowStaticLighting = CVarAllowStaticLighting->GetValueOnAnyThread() != 0;
 	bSupportSkyAtmosphere = !CVarSupportSkyAtmosphere || CVarSupportSkyAtmosphere->GetValueOnAnyThread() != 0;
-	bEnableStrata = CVarStrata ? CVarStrata->GetValueOnAnyThread() != 0 : false;
 
 	// mobile
 	bMobileAllowMovableDirectionalLights = CVarMobileAllowMovableDirectionalLights->GetValueOnAnyThread() != 0;
diff --git a/Engine/Source/Runtime/Engine/Public/SceneManagement.h b/Engine/Source/Runtime/Engine/Public/SceneManagement.h
--- a/Engine/Source/Runtime/Engine/Public/SceneManagement.h
+++ b/Engine/Source/Runtime/Engine/Public/SceneManagement.h
@@ -3189,7 +3189,6 @@ struct FReadOnlyCVARCache
 	bool bEnableLowQualityLightmaps;
 	bool bAllowStaticLighting;
 	bool bSupportSkyAtmosphere;
-	bool bEnableStrata;
 
 	// Mobile specific
 	bool bMobileAllowMovableDirectionalLights;
