Fixed assert adding a new batch element

#rnx

diff --git a/Engine/Source/Runtime/Engine/Private/HierarchicalInstancedStaticMesh.cpp b/Engine/Source/Runtime/Engine/Private/HierarchicalInstancedStaticMesh.cpp
--- a/Engine/Source/Runtime/Engine/Private/HierarchicalInstancedStaticMesh.cpp
+++ b/Engine/Source/Runtime/Engine/Private/HierarchicalInstancedStaticMesh.cpp
@@ -1546,7 +1546,8 @@ void FHierarchicalStaticMeshSceneProxy::FillDynamicMeshElements(FMeshElementColl
 								}
 								else
 								{
-									NewBatchElement = &MeshElement.Elements.Add_GetRef(MeshElement.Elements[0]);
+									NewBatchElement = &MeshElement.Elements.AddDefaulted_GetRef();
+									*NewBatchElement = MeshElement.Elements[0];
 								}
 
 								const int32 InstanceOffset = RunArray[CurrentRun];
