non-unity compile fixes

diff --git a/Engine/Source/Runtime/Engine/Private/WorldPartition/WorldPartitionRuntimeHash.cpp b/Engine/Source/Runtime/Engine/Private/WorldPartition/WorldPartitionRuntimeHash.cpp
--- a/Engine/Source/Runtime/Engine/Private/WorldPartition/WorldPartitionRuntimeHash.cpp
+++ b/Engine/Source/Runtime/Engine/Private/WorldPartition/WorldPartitionRuntimeHash.cpp
@@ -3,6 +3,7 @@
 #include "WorldPartition/WorldPartitionRuntimeHash.h"
 #include "WorldPartition/HLOD/HLODActor.h"
 #include "WorldPartition/HLOD/HLODActorDesc.h"
+#include "Engine/World.h"
 #include "Engine/LevelScriptBlueprint.h"
 #include "ActorReferencesUtils.h"
 
diff --git a/Engine/Source/Runtime/Engine/Public/WorldPartition/WorldPartitionActorDescView.h b/Engine/Source/Runtime/Engine/Public/WorldPartition/WorldPartitionActorDescView.h
--- a/Engine/Source/Runtime/Engine/Public/WorldPartition/WorldPartitionActorDescView.h
+++ b/Engine/Source/Runtime/Engine/Public/WorldPartition/WorldPartitionActorDescView.h
@@ -6,7 +6,9 @@
 #if WITH_EDITOR
 class UHLODLayer;
 class FWorldPartitionActorDesc;
+class UActorDescContainer;
 enum class EContainerClusterMode : uint8;
+enum class EActorGridPlacement : uint8;
 
 /**
  * A view on top of an actor desc, used to cache information that can be (potentially) different than the actor desc
