The bottom of the cloud layer can now go below the virtual planet ground.

diff --git a/Engine/Source/Runtime/Engine/Classes/Components/VolumetricCloudComponent.h b/Engine/Source/Runtime/Engine/Classes/Components/VolumetricCloudComponent.h
--- a/Engine/Source/Runtime/Engine/Classes/Components/VolumetricCloudComponent.h
+++ b/Engine/Source/Runtime/Engine/Classes/Components/VolumetricCloudComponent.h
@@ -26,7 +26,7 @@ class UVolumetricCloudComponent : public USceneComponent
 	~UVolumetricCloudComponent();
 
 	/** The altitude at which the cloud layer starts. (kilometers above the ground) */
-	UPROPERTY(EditAnywhere, BlueprintReadOnly, interp, Category = "Layer", meta = (UIMin = 0.0f, UIMax = 20.0f, ClampMin = 0.0f, SliderExponent = 2.0))
+	UPROPERTY(EditAnywhere, BlueprintReadOnly, interp, Category = "Layer", meta = (UIMin = 0.0f, UIMax = 20.0f, SliderExponent = 2.0))
 	float LayerBottomAltitude;
 
 	/** The altitude at which the cloud layer ends. (kilometers above the ground) */
