Fix Save2 import assert because skeletal mesh cooked platform data isn't adequately cached before the import harvest phase resulting in missing information when it happens lazily during serialization

diff --git a/Engine/Source/Runtime/Engine/Private/SkeletalMesh.cpp b/Engine/Source/Runtime/Engine/Private/SkeletalMesh.cpp
--- a/Engine/Source/Runtime/Engine/Private/SkeletalMesh.cpp
+++ b/Engine/Source/Runtime/Engine/Private/SkeletalMesh.cpp
@@ -4234,6 +4234,8 @@ void USkeletalMesh::ValidateBoneWeights(const ITargetPlatform* TargetPlatform)
 
 void USkeletalMesh::BeginCacheForCookedPlatformData(const ITargetPlatform* TargetPlatform)
 {
+	// Make sure to cache platform data so it doesn't happen lazily during serialization of the skeletal mesh
+	GetPlatformSkeletalMeshRenderData(this, TargetPlatform);
 	ValidateBoneWeights(TargetPlatform);
 }
 
