CIS fix. changed mecmp to memcpy (typo)

diff --git a/Engine/Source/Runtime/Engine/Private/SoundWave.cpp b/Engine/Source/Runtime/Engine/Private/SoundWave.cpp
--- a/Engine/Source/Runtime/Engine/Private/SoundWave.cpp
+++ b/Engine/Source/Runtime/Engine/Private/SoundWave.cpp
@@ -206,7 +206,7 @@ void FStreamedAudioChunk::GetCopy(void** OutChunkData)
 	}
 
 	// todo: expose as shared ptr(?)
-	FMemory::Memcmp(*OutChunkData, CachedDataPtr, DataSize);
+	FMemory::Memcpy(*OutChunkData, CachedDataPtr, DataSize);
 }
 
 #if WITH_EDITORONLY_DATA
