Physics - Fix improper use of JointName as a way to find the child bone of the constraint
This was broken when using user created constraint ( in Physics Asset Editor ) that result in constraints named with "UserConstraint" prefix
This is mostly affecting SkeletalMeshComponent limit and motor related functions ( and constraint breaking functions ) as well as  potentially some physics assets editor selection behaviour
- Added GetChildBoneName() and GetParentBoneName() to FConstraintInstance for API clarity
- Fixed all UE code improperly using JointName and used GetChildBoneName instead

diff --git a/Engine/Source/Editor/PhysicsAssetEditor/Private/PhysicsAssetEditorSharedData.cpp b/Engine/Source/Editor/PhysicsAssetEditor/Private/PhysicsAssetEditorSharedData.cpp
--- a/Engine/Source/Editor/PhysicsAssetEditor/Private/PhysicsAssetEditorSharedData.cpp
+++ b/Engine/Source/Editor/PhysicsAssetEditor/Private/PhysicsAssetEditorSharedData.cpp
@@ -542,7 +542,7 @@ void FPhysicsAssetEditorSharedData::ToggleSelectionType()
 		for (int32 BodyIdx = 0; BodyIdx < PhysicsAsset->SkeletalBodySetups.Num(); ++BodyIdx)
 		{
 			UBodySetup* BodySetup = PhysicsAsset->SkeletalBodySetups[BodyIdx];
-			if (DefaultInstance.JointName == BodySetup->BoneName && BodySetup->AggGeom.GetElementCount() > 0)
+			if (DefaultInstance.GetChildBoneName() == BodySetup->BoneName && BodySetup->AggGeom.GetElementCount() > 0)
 			{
 				SetSelectedBodyAnyPrim(BodyIdx, true);
 			}
@@ -556,7 +556,7 @@ void FPhysicsAssetEditorSharedData::ToggleSelectionType()
 		for(int32 ConstraintIdx = 0; ConstraintIdx < PhysicsAsset->ConstraintSetup.Num(); ++ConstraintIdx)
 		{
 			const UPhysicsConstraintTemplate* ConstraintTemplate = PhysicsAsset->ConstraintSetup[ConstraintIdx]; 
-			if(ConstraintTemplate->DefaultInstance.JointName == BodySetup->BoneName && !TmpSelectedConstraints.Contains(ConstraintIdx))
+			if(ConstraintTemplate->DefaultInstance.GetChildBoneName() == BodySetup->BoneName && !TmpSelectedConstraints.Contains(ConstraintIdx))
 			{
 				TmpSelectedConstraints.Add(ConstraintIdx);
 				SetSelectedConstraint(ConstraintIdx, true);
diff --git a/Engine/Source/Runtime/Engine/Classes/PhysicsEngine/ConstraintInstance.h b/Engine/Source/Runtime/Engine/Classes/PhysicsEngine/ConstraintInstance.h
--- a/Engine/Source/Runtime/Engine/Classes/PhysicsEngine/ConstraintInstance.h
+++ b/Engine/Source/Runtime/Engine/Classes/PhysicsEngine/ConstraintInstance.h
@@ -250,6 +250,12 @@ public:
 	/** Constructor **/
 	FConstraintInstance();
 
+	/** Get the child bone name */
+	const FName& GetChildBoneName() const { return ConstraintBone1; }
+
+	/** Get the parent bone name */
+	const FName& GetParentBoneName() const { return ConstraintBone2; }
+
 	/** Gets the linear limit size */
 	float GetLinearLimit() const
 	{
diff --git a/Engine/Source/Runtime/Engine/Private/PhysicsEngine/PhysAnim.cpp b/Engine/Source/Runtime/Engine/Private/PhysicsEngine/PhysAnim.cpp
--- a/Engine/Source/Runtime/Engine/Private/PhysicsEngine/PhysAnim.cpp
+++ b/Engine/Source/Runtime/Engine/Private/PhysicsEngine/PhysAnim.cpp
@@ -742,8 +742,8 @@ void USkeletalMeshComponent::UpdateRBJointMotors()
 			UPhysicsConstraintTemplate* CS = PhysicsAsset->ConstraintSetup[i];
 			FConstraintInstance* CI = Constraints[i];
 
-			FName JointName = CS->DefaultInstance.JointName;
-			int32 BoneIndex = SkeletalMesh->GetRefSkeleton().FindBoneIndex(JointName);
+			FName JointChildBoneName = CS->DefaultInstance.GetChildBoneName();
+			int32 BoneIndex = SkeletalMesh->GetRefSkeleton().FindBoneIndex(JointChildBoneName);
 
 			// If we found this bone, and a visible bone that is not the root, and its joint is motorised in some way..
 			if( (BoneIndex != INDEX_NONE) && (BoneIndex != 0) &&
diff --git a/Engine/Source/Runtime/Engine/Private/PhysicsEngine/PhysicsAsset.cpp b/Engine/Source/Runtime/Engine/Private/PhysicsEngine/PhysicsAsset.cpp
--- a/Engine/Source/Runtime/Engine/Private/PhysicsEngine/PhysicsAsset.cpp
+++ b/Engine/Source/Runtime/Engine/Private/PhysicsEngine/PhysicsAsset.cpp
@@ -499,7 +499,7 @@ FName UPhysicsAsset::FindConstraintBoneName(int32 ConstraintIndex)
 		return NAME_None;
 	}
 
-	return ConstraintSetup[ConstraintIndex]->DefaultInstance.JointName;
+	return ConstraintSetup[ConstraintIndex]->DefaultInstance.GetChildBoneName();
 }
 
 int32 UPhysicsAsset::FindMirroredBone(class USkeletalMesh* skelMesh, int32 BoneIndex)
diff --git a/Engine/Source/Runtime/Engine/Private/SkeletalMeshComponentPhysics.cpp b/Engine/Source/Runtime/Engine/Private/SkeletalMeshComponentPhysics.cpp
--- a/Engine/Source/Runtime/Engine/Private/SkeletalMeshComponentPhysics.cpp
+++ b/Engine/Source/Runtime/Engine/Private/SkeletalMeshComponentPhysics.cpp
@@ -1014,11 +1014,11 @@ void USkeletalMeshComponent::TermBodiesBelow(FName ParentBoneName)
 		for(int32 i=0; i<Constraints.Num(); i++)
 		{
 			// Get bone index of constraint
-			FName JointName = Constraints[i]->JointName;
-			int32 JointBoneIndex = GetBoneIndex(JointName);
+			FName JointChildBoneName = Constraints[i]->GetChildBoneName();
+			int32 JointBoneIndex = GetBoneIndex(JointChildBoneName);
 
 			// If constraint has bone in mesh, and is either the parent or child of it, term it
-			if(	JointBoneIndex != INDEX_NONE && (JointName == ParentBoneName ||	SkeletalMesh->GetRefSkeleton().BoneIsChildOf(JointBoneIndex, ParentBoneIndex)) )
+			if(	JointBoneIndex != INDEX_NONE && (JointChildBoneName == ParentBoneName ||	SkeletalMesh->GetRefSkeleton().BoneIsChildOf(JointBoneIndex, ParentBoneIndex)) )
 			{
 				Constraints[i]->TermConstraint();
 			}
@@ -1133,7 +1133,7 @@ void USkeletalMeshComponent::SetAllMotorsAngularPositionDrive(bool bEnableSwingD
 	{
 		if( bSkipCustomPhysicsType )
 		{
-			int32 BodyIndex = PhysicsAsset->FindBodyIndex(Constraints[i]->JointName);
+			int32 BodyIndex = PhysicsAsset->FindBodyIndex(Constraints[i]->GetChildBoneName());
 			if( BodyIndex != INDEX_NONE && PhysicsAsset->SkeletalBodySetups[BodyIndex]->PhysicsType != PhysType_Default)
 			{
 				continue;
@@ -1155,7 +1155,7 @@ void USkeletalMeshComponent::SetNamedMotorsAngularPositionDrive(bool bEnableSwin
 	for(int32 i=0; i<Constraints.Num(); i++)
 	{
 		FConstraintInstance* Instance = Constraints[i];
-		if( BoneNames.Contains(Instance->JointName) )
+		if( BoneNames.Contains(Instance->GetChildBoneName()) )
 		{
 			Constraints[i]->SetOrientationDriveTwistAndSwing(bEnableTwistDrive, bEnableSwingDrive);
 		}
@@ -1177,7 +1177,7 @@ void USkeletalMeshComponent::SetNamedMotorsAngularVelocityDrive(bool bEnableSwin
 	for(int32 i=0; i<Constraints.Num(); i++)
 	{
 		FConstraintInstance* Instance = Constraints[i];
-		if( BoneNames.Contains(Instance->JointName) )
+		if( BoneNames.Contains(Instance->GetChildBoneName()) )
 		{
 			Constraints[i]->SetAngularVelocityDriveTwistAndSwing(bEnableTwistDrive, bEnableSwingDrive);
 		}
@@ -1200,7 +1200,7 @@ void USkeletalMeshComponent::SetAllMotorsAngularVelocityDrive(bool bEnableSwingD
 	{
 		if( bSkipCustomPhysicsType )
 		{
-			int32 BodyIndex = PhysicsAsset->FindBodyIndex(Constraints[i]->JointName);
+			int32 BodyIndex = PhysicsAsset->FindBodyIndex(Constraints[i]->GetChildBoneName());
 			if( BodyIndex != INDEX_NONE && PhysicsAsset->SkeletalBodySetups[BodyIndex]->PhysicsType != PhysType_Default )
 			{
 				continue;
@@ -1255,7 +1255,7 @@ void USkeletalMeshComponent::SetAllMotorsAngularDriveParams(float InSpring, floa
 	{
 		if( bSkipCustomPhysicsType )
 		{
-			int32 BodyIndex = PhysicsAsset->FindBodyIndex(Constraints[i]->JointName);
+			int32 BodyIndex = PhysicsAsset->FindBodyIndex(Constraints[i]->GetChildBoneName());
 			if( BodyIndex != INDEX_NONE && PhysicsAsset->SkeletalBodySetups[BodyIndex]->PhysicsType != PhysType_Default )
 			{
 				continue;
@@ -1609,7 +1609,7 @@ void USkeletalMeshComponent::UpdateMeshForBrokenConstraints()
 		if( ConstraintInst && ConstraintInst->IsTerminated() )
 		{
 			// Get the associated joint bone index.
-			int32 JointBoneIndex = GetBoneIndex(ConstraintInst->JointName);
+			int32 JointBoneIndex = GetBoneIndex(ConstraintInst->GetChildBoneName());
 			if( JointBoneIndex == INDEX_NONE )
 			{
 				continue;
@@ -1895,7 +1895,7 @@ void USkeletalMeshComponent::BreakConstraint(FVector Impulse, FVector HitLocatio
 	UPhysicsAsset * const PhysicsAsset = GetPhysicsAsset();
 
 	// Figure out if Body is fixed or not
-	FBodyInstance* Body = GetBodyInstance(Constraint->JointName);
+	FBodyInstance* Body = GetBodyInstance(Constraint->GetChildBoneName());
 
 	if( Body != NULL && !Body->IsInstanceSimulatingPhysics() )
 	{
@@ -1930,7 +1930,7 @@ void USkeletalMeshComponent::SetAngularLimits(FName InBoneName, float Swing1Limi
 	UPhysicsAsset * const PhysicsAsset = GetPhysicsAsset();
 
 	// Figure out if Body is fixed or not
-	FBodyInstance* Body = GetBodyInstance(Constraint->JointName);
+	FBodyInstance* Body = GetBodyInstance(Constraint->GetChildBoneName());
 
 	if (Body != NULL && Body->IsInstanceSimulatingPhysics())
 	{
