Prevent creating static mesh instance elements for foliage

#rnx

diff --git a/Engine/Source/Runtime/Engine/Private/Elements/Framework/EngineElementsLibrary.cpp b/Engine/Source/Runtime/Engine/Private/Elements/Framework/EngineElementsLibrary.cpp
--- a/Engine/Source/Runtime/Engine/Private/Elements/Framework/EngineElementsLibrary.cpp
+++ b/Engine/Source/Runtime/Engine/Private/Elements/Framework/EngineElementsLibrary.cpp
@@ -386,7 +386,7 @@ TTypedElementOwner<FSMInstanceElementData> UEngineElementsLibrary::CreateSMInsta
 	UTypedElementRegistry* Registry = UTypedElementRegistry::GetInstance();
 
 	TTypedElementOwner<FSMInstanceElementData> SMInstanceElement;
-	if (ensureAlways(Registry) && SMInstanceElementDataUtil::SMInstanceElementsEnabled())
+	if (ensureAlways(Registry) && SMInstanceElementDataUtil::SMInstanceElementsEnabled() && SMInstanceElementDataUtil::IsValidComponentForSMInstanceElements(InSMInstanceElementId.ISMComponent))
 	{
 		SMInstanceElement = Registry->CreateElement<FSMInstanceElementData>(NAME_SMInstance);
 		if (SMInstanceElement)
diff --git a/Engine/Source/Runtime/Engine/Private/Elements/SMInstance/SMInstanceElementData.cpp b/Engine/Source/Runtime/Engine/Private/Elements/SMInstance/SMInstanceElementData.cpp
--- a/Engine/Source/Runtime/Engine/Private/Elements/SMInstance/SMInstanceElementData.cpp
+++ b/Engine/Source/Runtime/Engine/Private/Elements/SMInstance/SMInstanceElementData.cpp
@@ -2,6 +2,7 @@
 
 #include "Elements/SMInstance/SMInstanceElementData.h"
 #include "Elements/Framework/TypedElementHandle.h"
+#include "Components/InstancedStaticMeshComponent.h"
 #include "HAL/IConsoleManager.h"
 
 UE_DEFINE_TYPED_ELEMENT_DATA_RTTI(FSMInstanceElementData);
@@ -21,6 +22,31 @@ bool SMInstanceElementsEnabled()
 	return GEnableSMInstanceElements != 0;
 }
 
+bool IsValidComponentForSMInstanceElements(const UInstancedStaticMeshComponent* InComponent)
+{
+	if (!InComponent)
+	{
+		return false;
+	}
+
+	if (const AActor* OwnerActor = InComponent->GetOwner())
+	{
+		// Foliage actors have extra bookkeeping data which isn't correctly updated by static mesh instance elements
+		// Disable being able to create static mesh instance elements for foliage actors until this is resolved...
+		// Note: This test is by name as we cannot link directly to AInstancedFoliageActor
+		static const FName NAME_InstancedFoliageActor = "InstancedFoliageActor";
+		for (const UClass* ActorClass = OwnerActor->GetClass(); ActorClass; ActorClass = ActorClass->GetSuperClass())
+		{
+			if (ActorClass->GetFName() == NAME_InstancedFoliageActor)
+			{
+				return false;
+			}
+		}
+	}
+
+	return true;
+}
+
 FSMInstanceId GetSMInstanceFromHandle(const FTypedElementHandle& InHandle, const bool bSilent)
 {
 	const FSMInstanceElementData* SMInstanceElement = InHandle.GetData<FSMInstanceElementData>(bSilent);
diff --git a/Engine/Source/Runtime/Engine/Public/Elements/SMInstance/SMInstanceElementData.h b/Engine/Source/Runtime/Engine/Public/Elements/SMInstance/SMInstanceElementData.h
--- a/Engine/Source/Runtime/Engine/Public/Elements/SMInstance/SMInstanceElementData.h
+++ b/Engine/Source/Runtime/Engine/Public/Elements/SMInstance/SMInstanceElementData.h
@@ -26,6 +26,11 @@ namespace SMInstanceElementDataUtil
  */
 ENGINE_API bool SMInstanceElementsEnabled();
 
+/**
+ * Test whether the given ISM component is valid to be used with static mesh instance elements.
+ */
+ENGINE_API bool IsValidComponentForSMInstanceElements(const UInstancedStaticMeshComponent* InComponent);
+
 /**
  * Attempt to get the static mesh instance ID from the given element handle.
  * @note This is not typically something you'd want to use outside of data access within an interface implementation.
