Use Strata cvar instead of default settings for strata runtime and pipeline checks to work on console.
Fixes strata translucency at runtime when !WITH_EDITOR

diff --git a/Engine/Source/Runtime/Engine/Private/Materials/Material.cpp b/Engine/Source/Runtime/Engine/Private/Materials/Material.cpp
--- a/Engine/Source/Runtime/Engine/Private/Materials/Material.cpp
+++ b/Engine/Source/Runtime/Engine/Private/Materials/Material.cpp
@@ -3633,9 +3633,8 @@ void UMaterial::ConvertMaterialToStrataMaterial()
 		return;
 	}
 
-	const URendererSettings* RendererSettings = GetDefault<URendererSettings>();
-	const bool bStrataEnabled = RendererSettings && RendererSettings->bEnableStrata;
-
+	static const auto CVar = IConsoleManager::Get().FindTConsoleVariableDataInt(TEXT("r.Strata"));
+	const bool bStrataEnabled = CVar && CVar->GetValueOnAnyThread() > 0;
 	if (bStrataEnabled && !FrontMaterial.IsConnected())
 	{
 		auto MoveConnectionTo = [](auto& OldNodeInput, UMaterialExpression* NewNode, uint32 NewInputIndex)
@@ -4922,8 +4921,8 @@ void UMaterial::RebuildShadingModelField()
 {
 	ShadingModels.ClearShadingModels();
 
-	const URendererSettings* RendererSettings = GetDefault<URendererSettings>();
-	const bool bStrataEnabled = RendererSettings && RendererSettings->bEnableStrata;
+	static const auto CVar = IConsoleManager::Get().FindTConsoleVariableDataInt(TEXT("r.Strata"));
+	const bool bStrataEnabled = CVar && CVar->GetValueOnAnyThread() > 0;
 	if (bStrataEnabled && FrontMaterial.IsConnected())
 	{
 		FStrataMaterialInfo StrataMaterialInfo;
@@ -6719,8 +6718,9 @@ static bool IsPropertyActive_Internal(EMaterialProperty InProperty,
 		break;
 	case MP_FrontMaterial:
 		{
-			const URendererSettings* RendererSettings = GetDefault<URendererSettings>();
-			Active = RendererSettings && RendererSettings->bEnableStrata;
+			static const auto CVar = IConsoleManager::Get().FindTConsoleVariableDataInt(TEXT("r.Strata"));
+			const bool bStrataEnabled = CVar && CVar->GetValueOnAnyThread() > 0;
+			Active = bStrataEnabled;
 			break;
 		}
 	case MP_MaterialAttributes:
diff --git a/Engine/Source/Runtime/Engine/Private/Materials/MaterialShared.cpp b/Engine/Source/Runtime/Engine/Private/Materials/MaterialShared.cpp
--- a/Engine/Source/Runtime/Engine/Private/Materials/MaterialShared.cpp
+++ b/Engine/Source/Runtime/Engine/Private/Materials/MaterialShared.cpp
@@ -1611,14 +1611,11 @@ bool FMaterialResource::HasAmbientOcclusionConnected() const
 bool FMaterialResource::IsStrataMaterial() const
 {
 	// STRATA_TODO IsStrataMaterial should go away once Strata implementation is finished
-	const URendererSettings* RendererSettings = GetDefault<URendererSettings>();
-	if (RendererSettings && RendererSettings->bEnableStrata)
-	{
-		return Material->bUseMaterialAttributes ?
-			Material->MaterialAttributes.IsConnected(MP_FrontMaterial) :
-			Material->HasStrataFrontMaterialConnected();
-	}
-	return false;
+	static const auto CVar = IConsoleManager::Get().FindTConsoleVariableDataInt(TEXT("r.Strata"));
+	const bool bStrataEnabled = CVar && CVar->GetValueOnAnyThread() > 0;
+	// We no longer support both types of material (Strata and non strata) so no need to check if FrontMaterial is plugged in.
+	// We simply consider all material as Strata when Strata is enabled.
+	return bStrataEnabled;
 }
 
 bool FMaterialResource::RequiresSynchronousCompilation() const
