Render the local bounds of the proxy in addition to global and bounding sphere.

diff --git a/Engine/Source/Runtime/Engine/Private/PrimitiveSceneProxy.cpp b/Engine/Source/Runtime/Engine/Private/PrimitiveSceneProxy.cpp
--- a/Engine/Source/Runtime/Engine/Private/PrimitiveSceneProxy.cpp
+++ b/Engine/Source/Runtime/Engine/Private/PrimitiveSceneProxy.cpp
@@ -914,6 +914,8 @@ void FPrimitiveSceneProxy::RenderBounds(
 		DrawCircle(PDI, InBounds.Origin, FVector(1, 0, 0), FVector(0, 1, 0), FColor::Yellow, InBounds.SphereRadius, 32, DrawBoundsDPG);
 		DrawCircle(PDI, InBounds.Origin, FVector(1, 0, 0), FVector(0, 0, 1), FColor::Yellow, InBounds.SphereRadius, 32, DrawBoundsDPG);
 		DrawCircle(PDI, InBounds.Origin, FVector(0, 1, 0), FVector(0, 0, 1), FColor::Yellow, InBounds.SphereRadius, 32, DrawBoundsDPG);
+		// render the Local bounds in a red color
+		DrawWireBox(PDI, GetLocalToWorld(), GetLocalBounds().GetBox(), FColor(255, 72, 72), DrawBoundsDPG);
 	}
 }
 
