RigidBody hit reporting: Hit reporting will now function if enabled after initialization
Contribution by [REDACTED_PERSON]

diff --git a/Engine/Source/Runtime/Engine/Private/PhysicsEngine/BodyInstance.cpp b/Engine/Source/Runtime/Engine/Private/PhysicsEngine/BodyInstance.cpp
--- a/Engine/Source/Runtime/Engine/Private/PhysicsEngine/BodyInstance.cpp
+++ b/Engine/Source/Runtime/Engine/Private/PhysicsEngine/BodyInstance.cpp
@@ -3514,6 +3514,36 @@ void FBodyInstance::AddImpulseAtPosition(const FVector& Impulse, const FVector&
 
 void FBodyInstance::SetInstanceNotifyRBCollision(bool bNewNotifyCollision)
 {
+
+	if (bNewNotifyCollision == bNotifyRigidBodyCollision)
+	{
+		// don't update if we've already set it
+		return;
+	}
+	
+#if WITH_CHAOS
+	// make sure to register the component for collision events 
+	if (UPrimitiveComponent* PrimComp = OwnerComponent.Get())
+	{
+		if (UWorld* World = PrimComp->GetWorld())
+		{
+			if (FPhysScene_Chaos* PhysScene = World->GetPhysicsScene())
+			{
+				if (bNewNotifyCollision)
+				{
+					// add to the list
+					PhysScene->RegisterForCollisionEvents(PrimComp);
+				}
+				else
+				{
+					PhysScene->UnRegisterForCollisionEvents(PrimComp);
+				}
+			}
+		}
+	}
+#endif // WITH_CHAOS	
+
+
 	bNotifyRigidBodyCollision = bNewNotifyCollision;
 	UpdatePhysicsFilterData();
 }
