Remove superflous %s to fix crash introduced in 15536071

diff --git a/Engine/Source/Runtime/Engine/Private/Materials/MaterialExpressions.cpp b/Engine/Source/Runtime/Engine/Private/Materials/MaterialExpressions.cpp
--- a/Engine/Source/Runtime/Engine/Private/Materials/MaterialExpressions.cpp
+++ b/Engine/Source/Runtime/Engine/Private/Materials/MaterialExpressions.cpp
@@ -1824,7 +1824,7 @@ bool UMaterialExpressionTextureBase::VerifySamplerType(
 
 			FString SamplerTypeDisplayName = SamplerTypeEnum->GetDisplayNameTextByValue(SamplerType).ToString();
 
-			OutErrorMessage = FString::Printf(TEXT("%s> To use '%s' as sampler type, SRGB must be disabled for %s"),
+			OutErrorMessage = FString::Printf(TEXT("To use '%s' as sampler type, SRGB must be disabled for %s"),
 				*SamplerTypeDisplayName,
 				*Texture->GetPathName() );
 			return false;
