Autoderiv: Cleanup. Moved some public members to private.

diff --git a/Engine/Source/Runtime/Engine/Private/Materials/HLSLMaterialDerivativeAutogen.h b/Engine/Source/Runtime/Engine/Private/Materials/HLSLMaterialDerivativeAutogen.h
--- a/Engine/Source/Runtime/Engine/Private/Materials/HLSLMaterialDerivativeAutogen.h
+++ b/Engine/Source/Runtime/Engine/Private/Materials/HLSLMaterialDerivativeAutogen.h
@@ -87,12 +87,6 @@ public:
 		Num
 	};
 
-	// Note that the type index is from [0,3] for float1 to float4. I.e. A float3 would be index 2.
-	static int32 GetFunc1ReturnNumComponents(int32 SrcTypeIndex, EFunc1 Op);
-	static int32 GetFunc2ReturnNumComponents(int32 LhsTypeIndex, int32 RhsTypeIndex, EFunc2 Op);
-
-	void EnableGeneratedDepencencies();
-
 	int32 GenerateExpressionFunc1(FHLSLMaterialTranslator& Translator, EFunc1 Op, int32 SrcCode);
 	int32 GenerateExpressionFunc2(FHLSLMaterialTranslator& Translator, EFunc2 Op, int32 LhsCode, int32 RhsCode);
 
@@ -104,9 +98,6 @@ public:
 
 	FString ApplyUnMirror(FString Value, bool bUnMirrorU, bool bUnMirrorV);
 
-	FString CoerceValueRaw(const FString& Token, int32 SrcType, EDerivativeStatus SrcStatus, int32 DstType);
-	FString CoerceValueDeriv(const FString& Token, int32 SrcType, EDerivativeStatus SrcStatus, int32 DstType);
-
 	FString ConstructDeriv(const FString& Value, const FString& Ddx, const FString& Ddy, int32 DstType);
 	FString ConstructDerivFinite(const FString& Value, int32 DstType);
 
@@ -115,15 +106,18 @@ public:
 	FString ExtractElement(const FString& Value, int32 SrcType, int32 ElementIndex);
 
 private:
+	// Note that the type index is from [0,3] for float1 to float4. I.e. A float3 would be index 2.
+	static int32 GetFunc1ReturnNumComponents(int32 SrcTypeIndex, EFunc1 Op);
+	static int32 GetFunc2ReturnNumComponents(int32 LhsTypeIndex, int32 RhsTypeIndex, EFunc2 Op);
+
+	FString CoerceValueRaw(const FString& Token, int32 SrcType, EDerivativeStatus SrcStatus, int32 DstType);
+	FString CoerceValueDeriv(const FString& Token, int32 SrcType, EDerivativeStatus SrcStatus, int32 DstType);
+
+	void EnableGeneratedDepencencies();
+
 	// State to keep track of which derivative functions have been used and need to be generated.
 	bool bFunc1OpIsEnabled[(int32)EFunc1::Num][4] = {};
 	bool bFunc2OpIsEnabled[(int32)EFunc2::Num][4] = {};
-	bool bLerpEnabled[4] = {};
-	bool bIfEnabled[4] = {};
-	bool bIf2Enabled[4] = {};
-	bool bUnMirrorEnabled[2][2] = {};
-
-	bool bRotateScaleOffsetTexCoords = false;
 
 	bool bConstructDerivEnabled[4] = {};
 	bool bConstructConstantDerivEnabled[4] = {};
@@ -133,7 +127,13 @@ private:
 
 	bool bExtractIndexEnabled[4] = {};
 
+	bool bIfEnabled[4] = {};
+	bool bIf2Enabled[4] = {};
+	bool bLerpEnabled[4] = {};
+	bool bRotateScaleOffsetTexCoords = false;
 	bool bSelectElemHelperEnabled[4] = {};
+	bool bUnMirrorEnabled[2][2] = {};
+	
 };
 
 #endif // WITH_EDITORONLY_DATA
