Autoderiv: Removed unused function ExtractElement

diff --git a/Engine/Source/Runtime/Engine/Private/Materials/HLSLMaterialDerivativeAutogen.cpp b/Engine/Source/Runtime/Engine/Private/Materials/HLSLMaterialDerivativeAutogen.cpp
--- a/Engine/Source/Runtime/Engine/Private/Materials/HLSLMaterialDerivativeAutogen.cpp
+++ b/Engine/Source/Runtime/Engine/Private/Materials/HLSLMaterialDerivativeAutogen.cpp
@@ -306,19 +306,6 @@ FString FMaterialDerivativeAutogen::ConstructDerivFinite(const FString& Value, i
 	return Ret;
 }
 
-FString FMaterialDerivativeAutogen::ExtractElement(const FString& Value, int32 SrcType, int32 ElementIndex)
-{
-	check(SrcType < 4);
-	check(ElementIndex <= SrcType);
-	bExtractIndexEnabled[SrcType] = true;
-
-	FString TypeName = GetDerivVectorName(SrcType);
-
-	FString Ret = FString::Printf(TEXT("Extract%s_%d(%s %s)"), *TypeName, ElementIndex + 1, *TypeName, *Value);
-
-	return Ret;
-}
-
 FString FMaterialDerivativeAutogen::ConvertDeriv(const FString& Value, int32 DstType, int32 SrcType)
 {
 	check(DstType < 4);
diff --git a/Engine/Source/Runtime/Engine/Private/Materials/HLSLMaterialDerivativeAutogen.h b/Engine/Source/Runtime/Engine/Private/Materials/HLSLMaterialDerivativeAutogen.h
--- a/Engine/Source/Runtime/Engine/Private/Materials/HLSLMaterialDerivativeAutogen.h
+++ b/Engine/Source/Runtime/Engine/Private/Materials/HLSLMaterialDerivativeAutogen.h
@@ -103,8 +103,6 @@ public:
 
 	FString ConvertDeriv(const FString& Value, int32 DstType, int32 SrcType);
 
-	FString ExtractElement(const FString& Value, int32 SrcType, int32 ElementIndex);
-
 private:
 	// Note that the type index is from [0,3] for float1 to float4. I.e. A float3 would be index 2.
 	static int32 GetFunc1ReturnNumComponents(int32 SrcTypeIndex, EFunc1 Op);
