New Animation Attributes system, replacing Custom Attributes:

+ Attribute structures to UAnimDataModel
    * These are sampled/copied into AnimSequence whenever they change
+ Attribute related Notifies and Payloads
+ Controller API and Actions for Attribute related behaviour
+ Type traits (TAttributeTypeTraitsBase) to determine support functionality for user-defined attribute types
+ TAttributeContainer equivalent to TCustomAttributes, used for keeping track of attributes at runtime in a TMap similar fashion
    * Has two exported specializations FStack/Heap-AttributeContainer
+ IAttributeBlendOperator interface used for Attribute related operations in Anim graph
    * Allows for user-defined blending behaviour for their associated types
+ TAttributeBlendOperator providing out-of-the-box blending behaviour for user-defined types
+ FAttributeBlendData helper structure, this encapsulates and abstracts the blend / attribute operations
    * Exposes two iterators, allowing BlendOperator to loop through (type) overlapping Attributes and unique attributes
+ Float/Integer/String Animation Attribute structures used to support legacy TVariant CustomAttribute data types
+ Transform animation attribute structure to add support for single-FTransform based attributes
+ FAnimationAttributeIdentifier identifier used to reference an attribute in a script-friendly manor
+ AttributeTypes static API for registering Attribute types
+ FAttributeCurve providing a curve-type with an Attribute type as its underlying key-value
+ TWrappedAttribute helper structure to wrap end template operate on raw memory (TArray buffer)
+ Added tests for
    * Attribute related controller functionality and actions
    * Attribute curve key reduction
    * Evaluating attributes from AnimSequence
    * Attribute operations (blend, accumulate etc)
    * Functional testing for blendspace attribute evaluation and blending

* Changed default attribute blend type to Blend vs Override
* Updated FBX import/export paths to handle and use new Attribute data structures
* Attribute data is now incorporated into animation source data DDC key
* Deprecated Custom Attributes stored on AnimSequence get converted into their equivalent Attribute structures
* Deprecated all previous CustomAttribute structures, APIs and files
* Corrected some comments in UAnimDataController.h
* Updated existing custom attribute tests to adhere to new blend expectations/behaviour
* Updated AnimSequence resize tests to also incorporate an attribute curve
* Changed layered bone blend to use .5 blend weight vs 1.0 to cover more behaviour
* Added transform attribute used to compare against bone transform during pre-existing functional testing (blended only)

- Deleted CustomAttributes details customization

diff --git a/Engine/Plugins/Experimental/Animation/ContextualAnimation/Source/ContextualAnimation/Private/ContextualAnimUtilities.cpp b/Engine/Plugins/Experimental/Animation/ContextualAnimation/Source/ContextualAnimation/Private/ContextualAnimUtilities.cpp
--- a/Engine/Plugins/Experimental/Animation/ContextualAnimation/Source/ContextualAnimation/Private/ContextualAnimUtilities.cpp
+++ b/Engine/Plugins/Experimental/Animation/ContextualAnimation/Source/ContextualAnimation/Private/ContextualAnimUtilities.cpp
@@ -5,7 +5,7 @@
 #include "AnimationRuntime.h"
 #include "Animation/AnimMontage.h"
 #include "Animation/AnimationPoseData.h"
-#include "Animation/CustomAttributesRuntime.h"
+#include "Animation/AttributesRuntime.h"
 #include "Animation/AnimTypes.h"
 #include "DrawDebugHelpers.h"
 #include "Engine/World.h"
@@ -21,7 +21,7 @@ void UContextualAnimUtilities::ExtractLocalSpacePose(const UAnimSequenceBase* An
 
 	FAnimExtractContext Context(Time, bExtractRootMotion);
 
-	FStackCustomAttributes Attributes;
+	UE::Anim::FStackAttributeContainer Attributes;
 	FAnimationPoseData AnimationPoseData(OutPose, Curve, Attributes);
 	if (const UAnimSequence* AnimSequence = Cast<UAnimSequence>(Animation))
 	{
diff --git a/Engine/Plugins/Experimental/Animation/MotionWarping/Source/MotionWarping/Private/MotionWarpingComponent.cpp b/Engine/Plugins/Experimental/Animation/MotionWarping/Source/MotionWarping/Private/MotionWarpingComponent.cpp
--- a/Engine/Plugins/Experimental/Animation/MotionWarping/Source/MotionWarping/Private/MotionWarpingComponent.cpp
+++ b/Engine/Plugins/Experimental/Animation/MotionWarping/Source/MotionWarping/Private/MotionWarpingComponent.cpp
@@ -43,7 +43,7 @@ void UMotionWarpingUtilities::ExtractLocalSpacePose(const UAnimSequenceBase* Ani
 
 	FAnimExtractContext Context(Time, bExtractRootMotion);
 
-	FStackCustomAttributes Attributes;
+	UE::Anim::FStackAttributeContainer Attributes;
 	FAnimationPoseData AnimationPoseData(OutPose, Curve, Attributes);
 	if (const UAnimSequence* AnimSequence = Cast<UAnimSequence>(Animation))
 	{
diff --git a/Engine/Plugins/Experimental/Animation/PoseSearch/Source/Runtime/Private/PoseSearch.cpp b/Engine/Plugins/Experimental/Animation/PoseSearch/Source/Runtime/Private/PoseSearch.cpp
--- a/Engine/Plugins/Experimental/Animation/PoseSearch/Source/Runtime/Private/PoseSearch.cpp
+++ b/Engine/Plugins/Experimental/Animation/PoseSearch/Source/Runtime/Private/PoseSearch.cpp
@@ -1206,7 +1206,7 @@ void FSequenceSampler::ExtractPoses()
 	FCSPose<FCompactPose> ComponentSpacePose;
 
 	FBlendedCurve UnusedCurve;
-	FStackCustomAttributes UnusedAttributes;
+	UE::Anim::FStackAttributeContainer UnusedAttributes;
 
 	FAnimExtractContext ExtractionCtx;
 	// ExtractionCtx.PoseCurves is intentionally left empty
diff --git a/Engine/Source/Editor/DetailCustomizations/Private/DetailCustomizations.cpp b/Engine/Source/Editor/DetailCustomizations/Private/DetailCustomizations.cpp
--- a/Engine/Source/Editor/DetailCustomizations/Private/DetailCustomizations.cpp
+++ b/Engine/Source/Editor/DetailCustomizations/Private/DetailCustomizations.cpp
@@ -142,7 +142,6 @@
 #include "DebugCameraControllerSettingsCustomization.h"
 #include "BoundsCopyComponentDetails.h"
 #include "SupportedRangeTypes.h"	// StructsSupportingRangeVisibility
-#include "CustomAttributeCustomization.h"
 
 IMPLEMENT_MODULE( FDetailCustomizationsModule, DetailCustomizations );
 
@@ -283,8 +282,6 @@ void FDetailCustomizationsModule::RegisterPropertyTypeCustomizations()
 	RegisterCustomPropertyTypeLayout("ComponentReference", FOnGetPropertyTypeCustomizationInstance::CreateStatic(&FComponentReferenceCustomization::MakeInstance));
 	RegisterCustomPropertyTypeLayout("EMaterialShadingModel", FOnGetPropertyTypeCustomizationInstance::CreateStatic(&FMaterialShadingModelCustomization::MakeInstance));
 	RegisterCustomPropertyTypeLayout("DebugCameraControllerSettingsViewModeIndex", FOnGetPropertyTypeCustomizationInstance::CreateStatic(&FDebugCameraControllerSettingsViewModeIndexCustomization::MakeInstance));
-	RegisterCustomPropertyTypeLayout("CustomAttributePerBoneData", FOnGetPropertyTypeCustomizationInstance::CreateStatic(&FCustomAttributePerBoneDataCustomization::MakeInstance));
-	RegisterCustomPropertyTypeLayout("CustomAttribute", FOnGetPropertyTypeCustomizationInstance::CreateStatic(&FCustomAttributeCustomization::MakeInstance));
 }
 
 #undef REGISTER_UIMINMAX_CUSTOMIZATION
diff --git a/Engine/Source/Editor/MovieSceneTools/Private/EditModes/SkeletalAnimationTrackEditMode.cpp b/Engine/Source/Editor/MovieSceneTools/Private/EditModes/SkeletalAnimationTrackEditMode.cpp
--- a/Engine/Source/Editor/MovieSceneTools/Private/EditModes/SkeletalAnimationTrackEditMode.cpp
+++ b/Engine/Source/Editor/MovieSceneTools/Private/EditModes/SkeletalAnimationTrackEditMode.cpp
@@ -20,7 +20,7 @@
 #include "Rendering/SkeletalMeshRenderData.h"
 #include "Animation/AnimSequence.h"
 #include "Animation/AnimationPoseData.h"
-#include "Animation/CustomAttributesRuntime.h"
+#include "Animation/AttributesRuntime.h"
 
 #include "Tracks/MovieScene3DTransformTrack.h"
 #include "Sections/MovieScene3DTransformSection.h"
@@ -382,7 +382,7 @@ void FSkeletalAnimationTrackEditMode::Render(const FSceneView* View, FViewport*
 
 											FBlendedCurve OutCurve;
 											OutCurve.InitFrom(SkelMeshComp->GetAnimInstance()->GetRequiredBones());
-											FStackCustomAttributes TempAttributes;
+											UE::Anim::FStackAttributeContainer TempAttributes;
 											FAnimationPoseData OutAnimationPoseData(OutPose, OutCurve, TempAttributes);
 
 											const float Seconds = AnimSection->MapTimeToAnimation(CurrentTime.Time, CurrentTime.Rate);
diff --git a/Engine/Source/Editor/SequenceRecorder/Private/AnimationRecorder.cpp b/Engine/Source/Editor/SequenceRecorder/Private/AnimationRecorder.cpp
--- a/Engine/Source/Editor/SequenceRecorder/Private/AnimationRecorder.cpp
+++ b/Engine/Source/Editor/SequenceRecorder/Private/AnimationRecorder.cpp
@@ -20,6 +20,7 @@
 #include "Animation/AnimNotifies/AnimNotifyState.h"
 #include "Subsystems/AssetEditorSubsystem.h"
 #include "Animation/AnimData/AnimDataController.h"
+#include "Animation/BuiltInAttributeTypes.h"
 
 #define LOCTEXT_NAMESPACE "FAnimationRecorder"
 
@@ -528,6 +529,9 @@ void FAnimationRecorder::ProcessRecordedTimes(UAnimSequence* AnimSequence, USkel
 	TArray<FString> Slates(&Slate, 1);
 	TArray<float> SlateTimes(&Times[0], 1);
 
+	UAnimDataController* Controller = AnimSequence->GetController();
+	UAnimDataController::FScopedBracket(Controller, LOCTEXT("AddTimeCodeAttributesBracket", "Adding Time Code attributes"));
+
 	for (int32 BoneIndex = 0; BoneIndex < SpaceBases->Num(); ++BoneIndex)
 	{
 		// verify if this bone exists in skeleton
@@ -537,13 +541,14 @@ void FAnimationRecorder::ProcessRecordedTimes(UAnimSequence* AnimSequence, USkel
 			// add tracks for the bone existing
 			FName BoneTreeName = AnimSkeleton->GetReferenceSkeleton().GetBoneName(BoneTreeIndex);
 			
-			AnimSequence->AddBoneIntegerCustomAttribute(BoneTreeName, FName(*HoursName), Times, Hours);
-			AnimSequence->AddBoneIntegerCustomAttribute(BoneTreeName, FName(*MinutesName), Times, Minutes);
-			AnimSequence->AddBoneIntegerCustomAttribute(BoneTreeName, FName(*SecondsName), Times, Seconds);
-			AnimSequence->AddBoneIntegerCustomAttribute(BoneTreeName, FName(*FramesName), Times, Frames);
-			AnimSequence->AddBoneFloatCustomAttribute(BoneTreeName, FName(*SubFramesName), Times, SubFrames);
+			UE::Anim::AddTypedCustomAttribute<FIntegerAnimationAttribute, int32>(FName(*HoursName), BoneTreeName, AnimSequence, MakeArrayView(Times), MakeArrayView(Hours));
+			UE::Anim::AddTypedCustomAttribute<FIntegerAnimationAttribute, int32>(FName(*MinutesName), BoneTreeName, AnimSequence, MakeArrayView(Times), MakeArrayView(Minutes));
+			UE::Anim::AddTypedCustomAttribute<FIntegerAnimationAttribute, int32>(FName(*SecondsName), BoneTreeName, AnimSequence, MakeArrayView(Times), MakeArrayView(Seconds));
+			UE::Anim::AddTypedCustomAttribute<FIntegerAnimationAttribute, int32>(FName(*FramesName), BoneTreeName, AnimSequence, MakeArrayView(Times), MakeArrayView(Frames));
+
+			UE::Anim::AddTypedCustomAttribute<FFloatAnimationAttribute, float>(FName(*SubFramesName), BoneTreeName, AnimSequence, MakeArrayView(Times), MakeArrayView(SubFrames));
 
-			AnimSequence->AddBoneStringCustomAttribute(BoneTreeName, FName(*SlateName), SlateTimes, Slates);
+			UE::Anim::AddTypedCustomAttribute<FStringAnimationAttribute, FString>(FName(*SlateName), BoneTreeName, AnimSequence, MakeArrayView(SlateTimes), MakeArrayView(Slates));
 		}
 	}
 }
diff --git a/Engine/Source/Editor/UnrealEd/Private/Animation/DebugSkelMeshComponent.cpp b/Engine/Source/Editor/UnrealEd/Private/Animation/DebugSkelMeshComponent.cpp
--- a/Engine/Source/Editor/UnrealEd/Private/Animation/DebugSkelMeshComponent.cpp
+++ b/Engine/Source/Editor/UnrealEd/Private/Animation/DebugSkelMeshComponent.cpp
@@ -484,7 +484,7 @@ void UDebugSkelMeshComponent::GenSpaceBases(TArray<FTransform>& OutSpaceBases)
 	TempBoneSpaceTransforms.AddUninitialized(OutSpaceBases.Num());
 	FVector TempRootBoneTranslation;
 	FBlendedHeapCurve TempCurve;
-	FHeapCustomAttributes TempAtttributes;
+	UE::Anim::FHeapAttributeContainer TempAtttributes;
 	DoInstancePreEvaluation();
 	PerformAnimationEvaluation(SkeletalMesh, AnimScriptInstance, OutSpaceBases, TempBoneSpaceTransforms, TempRootBoneTranslation, TempCurve, TempAtttributes);
 	DoInstancePostEvaluation();
@@ -574,7 +574,7 @@ void UDebugSkelMeshComponent::RefreshBoneTransforms(FActorComponentTickFunction*
 					{
 						FCompactPose AdditiveBasePose;
 						FBlendedCurve AdditiveCurve;
-						FStackCustomAttributes AdditiveAttributes;
+						UE::Anim::FStackAttributeContainer AdditiveAttributes;
 						AdditiveCurve.InitFrom(BoneContainer);
 						AdditiveBasePose.SetBoneContainer(&BoneContainer);
 						
diff --git a/Engine/Source/Editor/UnrealEd/Private/Fbx/FbxAnimationExport.cpp b/Engine/Source/Editor/UnrealEd/Private/Fbx/FbxAnimationExport.cpp
--- a/Engine/Source/Editor/UnrealEd/Private/Fbx/FbxAnimationExport.cpp
+++ b/Engine/Source/Editor/UnrealEd/Private/Fbx/FbxAnimationExport.cpp
@@ -17,7 +17,8 @@
 #include "Animation/SkeletalMeshActor.h"
 #include "FbxExporter.h"
 #include "Exporters/FbxExportOption.h"
-#include "Animation/CustomAttributesRuntime.h"
+#include "Animation/AttributesRuntime.h"
+#include "Animation/BuiltInAttributeTypes.h"
 
 DEFINE_LOG_CATEGORY_STATIC(LogFbxAnimationExport, Log, All);
 
@@ -115,7 +116,7 @@ void FFbxExporter::ExportAnimSequenceToFbx(const UAnimSequence* AnimSeq,
 
 	ExportCustomAnimCurvesToFbx(CustomCurveMap, AnimSeq, AnimStartOffset, AnimEndOffset, AnimPlayRate, StartTime);
 
-	TArray<FCustomAttribute> CustomAttributes;
+	TArray<const FAnimatedBoneAttribute*> CustomAttributes;
 
 	// Add the animation data to the bone nodes
 	for(int32 BoneIndex = 0; BoneIndex < BoneNodes.Num(); ++BoneIndex)
@@ -126,7 +127,7 @@ void FFbxExporter::ExportAnimSequenceToFbx(const UAnimSequence* AnimSeq,
 		FName BoneName = Skeleton->GetReferenceSkeleton().GetBoneName(BoneTreeIndex);
 		
 		CustomAttributes.Reset();
-		AnimSeq->GetCustomAttributesForBone(BoneName, CustomAttributes);
+		AnimSeq->GetDataModel()->GetAttributesForBone(BoneName, CustomAttributes);
 
 		TArray<TPair<int32, FbxAnimCurve*>> FloatCustomAttributeIndices;
 		TArray<TPair<int32, FbxAnimCurve*>> IntCustomAttributeIndices;
@@ -134,45 +135,49 @@ void FFbxExporter::ExportAnimSequenceToFbx(const UAnimSequence* AnimSeq,
 		// Setup custom attribute properties and curves
 		for (int32 AttributeIndex = 0; AttributeIndex < CustomAttributes.Num(); ++AttributeIndex)
 		{
-			const FCustomAttribute& Attribute = CustomAttributes[AttributeIndex];
-			const FName& AttributeName = Attribute.Name;
+			const FAnimatedBoneAttribute* AttributePtr = CustomAttributes[AttributeIndex];
+			if (AttributePtr)
+			{
+				const FAnimatedBoneAttribute& Attribute = *AttributePtr;
 
-			const EVariantTypes VariantType = static_cast<EVariantTypes>(Attribute.VariantType);
+				const FName& AttributeName = Attribute.Identifier.GetName();
+				const UScriptStruct* AttributeType = Attribute.Identifier.GetType();
 
-			if (VariantType == EVariantTypes::Int32)
-			{
-				FbxProperty AnimCurveFbxProp = FbxProperty::Create(CurrentBoneNode, FbxIntDT, TCHAR_TO_UTF8(*AttributeName.ToString()));
-				AnimCurveFbxProp.ModifyFlag(FbxPropertyFlags::eAnimatable, true);
-				AnimCurveFbxProp.ModifyFlag(FbxPropertyFlags::eUserDefined, true);
+				if (AttributeType == FIntegerAnimationAttribute::StaticStruct())
+				{
+					FbxProperty AnimCurveFbxProp = FbxProperty::Create(CurrentBoneNode, FbxIntDT, TCHAR_TO_UTF8(*AttributeName.ToString()));
+					AnimCurveFbxProp.ModifyFlag(FbxPropertyFlags::eAnimatable, true);
+					AnimCurveFbxProp.ModifyFlag(FbxPropertyFlags::eUserDefined, true);
 
-				FbxAnimCurve* AnimFbxCurve = AnimCurveFbxProp.GetCurve(InAnimLayer, true);
-				AnimFbxCurve->KeyModifyBegin();
-				IntCustomAttributeIndices.Emplace(AttributeIndex, AnimFbxCurve);
-			}
-			else if (VariantType == EVariantTypes::Float)
-			{
-				FbxProperty AnimCurveFbxProp = FbxProperty::Create(CurrentBoneNode, FbxFloatDT, TCHAR_TO_UTF8(*AttributeName.ToString()));
-				AnimCurveFbxProp.ModifyFlag(FbxPropertyFlags::eAnimatable, true);
-				AnimCurveFbxProp.ModifyFlag(FbxPropertyFlags::eUserDefined, true);
+					FbxAnimCurve* AnimFbxCurve = AnimCurveFbxProp.GetCurve(InAnimLayer, true);
+					AnimFbxCurve->KeyModifyBegin();
+					IntCustomAttributeIndices.Emplace(AttributeIndex, AnimFbxCurve);
+				}
+				else if (AttributeType == FFloatAnimationAttribute::StaticStruct())
+				{
+					FbxProperty AnimCurveFbxProp = FbxProperty::Create(CurrentBoneNode, FbxFloatDT, TCHAR_TO_UTF8(*AttributeName.ToString()));
+					AnimCurveFbxProp.ModifyFlag(FbxPropertyFlags::eAnimatable, true);
+					AnimCurveFbxProp.ModifyFlag(FbxPropertyFlags::eUserDefined, true);
 
-				FbxAnimCurve* AnimFbxCurve = AnimCurveFbxProp.GetCurve(InAnimLayer, true);
-				AnimFbxCurve->KeyModifyBegin();
-				FloatCustomAttributeIndices.Emplace(AttributeIndex, AnimFbxCurve);
-			}
-			else if (VariantType == EVariantTypes::String)
-			{
-				FbxProperty AnimCurveFbxProp = FbxProperty::Create(CurrentBoneNode, FbxStringDT, TCHAR_TO_UTF8(*AttributeName.ToString()));
-				AnimCurveFbxProp.ModifyFlag(FbxPropertyFlags::eUserDefined, true);
-
-				// String attributes can't be keyed, simply set a normal value.
-				FString AttributeValue;
-				FCustomAttributesRuntime::GetAttributeValue(Attribute, 0.f, AttributeValue);
-				FbxString FbxValueString(TCHAR_TO_UTF8(*AttributeValue));
-				AnimCurveFbxProp.Set(FbxValueString);
-			}
-			else
-			{
-				ensureMsgf(false, TEXT("Trying to export unsupported custom attribte (float, int32 and FString are currently supported)"));
+					FbxAnimCurve* AnimFbxCurve = AnimCurveFbxProp.GetCurve(InAnimLayer, true);
+					AnimFbxCurve->KeyModifyBegin();
+					FloatCustomAttributeIndices.Emplace(AttributeIndex, AnimFbxCurve);
+				}
+				else if (AttributeType == FStringAnimationAttribute::StaticStruct())
+				{
+					FbxProperty AnimCurveFbxProp = FbxProperty::Create(CurrentBoneNode, FbxStringDT, TCHAR_TO_UTF8(*AttributeName.ToString()));
+					AnimCurveFbxProp.ModifyFlag(FbxPropertyFlags::eUserDefined, true);
+
+					// String attributes can't be keyed, simply set a normal value.
+					FStringAnimationAttribute EvaluatedAttribute = Attribute.Curve.Evaluate<FStringAnimationAttribute>(0.f);
+
+					FbxString FbxValueString(TCHAR_TO_UTF8(*EvaluatedAttribute.Value));
+					AnimCurveFbxProp.Set(FbxValueString);
+				}
+				else
+				{
+					ensureMsgf(false, TEXT("Trying to export unsupported custom attribte (float, int32 and FString are currently supported)"));
+				}
 			}
 		}
 
@@ -234,18 +239,31 @@ void FFbxExporter::ExportAnimSequenceToFbx(const UAnimSequence* AnimSeq,
 
 			for (TPair<int32, FbxAnimCurve*>& CurrentAttributeCurve : FloatCustomAttributeIndices)
 			{
-				float AttributeValue = 0.f;
-				FCustomAttributesRuntime::GetAttributeValue(CustomAttributes[CurrentAttributeCurve.Key], AnimTime, AttributeValue);
-				int32 KeyIndex = CurrentAttributeCurve.Value->KeyAdd(ExportTime);
-				CurrentAttributeCurve.Value->KeySetValue(KeyIndex, AttributeValue);
+				const FAnimatedBoneAttribute* AttributePtr = CustomAttributes[CurrentAttributeCurve.Key];
+
+				if (AttributePtr)
+				{
+					ensure(AttributePtr->Identifier.GetType() == FFloatAnimationAttribute::StaticStruct());
+
+					FFloatAnimationAttribute EvaluatedAttribute = AttributePtr->Curve.Evaluate<FFloatAnimationAttribute>(AnimTime);
+					int32 KeyIndex = CurrentAttributeCurve.Value->KeyAdd(ExportTime);
+					CurrentAttributeCurve.Value->KeySetValue(KeyIndex, EvaluatedAttribute.Value);
+
+				}
 			}
 
 			for (TPair<int32, FbxAnimCurve*>& CurrentAttributeCurve : IntCustomAttributeIndices)
 			{
-				int32 AttributeValue = 0;
-				FCustomAttributesRuntime::GetAttributeValue(CustomAttributes[CurrentAttributeCurve.Key], AnimTime, AttributeValue);
-				int32 KeyIndex = CurrentAttributeCurve.Value->KeyAdd(ExportTime);
-				CurrentAttributeCurve.Value->KeySetValue(KeyIndex, static_cast<float>(AttributeValue));
+				const FAnimatedBoneAttribute* AttributePtr = CustomAttributes[CurrentAttributeCurve.Key];
+
+				if (AttributePtr)
+				{
+					ensure(AttributePtr->Identifier.GetType() == FIntegerAnimationAttribute::StaticStruct());
+
+					FIntegerAnimationAttribute EvaluatedAttribute = AttributePtr->Curve.Evaluate<FIntegerAnimationAttribute>(AnimTime);
+					int32 KeyIndex = CurrentAttributeCurve.Value->KeyAdd(ExportTime);
+					CurrentAttributeCurve.Value->KeySetValue(KeyIndex, static_cast<float>(EvaluatedAttribute.Value));
+				}
 			}
 		};
 
@@ -648,7 +666,7 @@ void FFbxExporter::ExportAnimTrack(IAnimTrackAdapter& AnimTrackAdapter, AActor*
 		}
 	}
 	
-	TArray<FCustomAttribute> CustomAttributes;
+	TArray<const FAnimatedBoneAttribute*> CustomAttributes;
 	
 	FTransform InitialInvParentTransform;
 
@@ -759,7 +777,7 @@ void FFbxExporter::ExportAnimTrack(IAnimTrackAdapter& AnimTrackAdapter, AActor*
 			}
 
 			CustomAttributes.Reset();
-			AnimSeq->GetCustomAttributesForBone(BoneName, CustomAttributes);
+			AnimSeq->GetDataModel()->GetAttributesForBone(BoneName, CustomAttributes);
 
 			TArray<TPair<int32, FbxAnimCurve*>> FloatCustomAttributeIndices;
 			TArray<TPair<int32, FbxAnimCurve*>> IntCustomAttributeIndices;
@@ -767,61 +785,71 @@ void FFbxExporter::ExportAnimTrack(IAnimTrackAdapter& AnimTrackAdapter, AActor*
 			// Setup custom attribute properties and curves
 			for (int32 AttributeIndex = 0; AttributeIndex < CustomAttributes.Num(); ++AttributeIndex)
 			{
-				const FCustomAttribute& Attribute = CustomAttributes[AttributeIndex];
-				const FName& AttributeName = Attribute.Name;
+				const FAnimatedBoneAttribute* AttributePtr = CustomAttributes[AttributeIndex];
+				if (AttributePtr)
+				{
+					const FAnimatedBoneAttribute& Attribute = *AttributePtr;
 
-				const EVariantTypes VariantType = static_cast<EVariantTypes>(Attribute.VariantType);
+					const FName& AttributeName = Attribute.Identifier.GetName();
+					const UScriptStruct* AttributeType = Attribute.Identifier.GetType();
 
-				if (VariantType == EVariantTypes::Int32)
-				{
-					FbxProperty AnimCurveFbxProp = FbxProperty::Create(CurrentBoneNode, FbxIntDT, TCHAR_TO_UTF8(*AttributeName.ToString()));
-					AnimCurveFbxProp.ModifyFlag(FbxPropertyFlags::eAnimatable, true);
-					AnimCurveFbxProp.ModifyFlag(FbxPropertyFlags::eUserDefined, true);
+					if (AttributeType == FIntegerAnimationAttribute::StaticStruct())
+					{
+						FbxProperty AnimCurveFbxProp = FbxProperty::Create(CurrentBoneNode, FbxIntDT, TCHAR_TO_UTF8(*AttributeName.ToString()));
+						AnimCurveFbxProp.ModifyFlag(FbxPropertyFlags::eAnimatable, true);
+						AnimCurveFbxProp.ModifyFlag(FbxPropertyFlags::eUserDefined, true);
 
-					FbxAnimCurve* AnimFbxCurve = AnimCurveFbxProp.GetCurve(AnimLayer, true);
-					AnimFbxCurve->KeyModifyBegin();
-					IntCustomAttributeIndices.Emplace(AttributeIndex, AnimFbxCurve);
-				}
-				else if (VariantType == EVariantTypes::Float)
-				{
-					FbxProperty AnimCurveFbxProp = FbxProperty::Create(CurrentBoneNode, FbxFloatDT, TCHAR_TO_UTF8(*AttributeName.ToString()));
-					AnimCurveFbxProp.ModifyFlag(FbxPropertyFlags::eAnimatable, true);
-					AnimCurveFbxProp.ModifyFlag(FbxPropertyFlags::eUserDefined, true);
+						FbxAnimCurve* AnimFbxCurve = AnimCurveFbxProp.GetCurve(AnimLayer, true);
+						AnimFbxCurve->KeyModifyBegin();
+						IntCustomAttributeIndices.Emplace(AttributeIndex, AnimFbxCurve);
+					}
+					else if (AttributeType == FFloatAnimationAttribute::StaticStruct())
+					{
+						FbxProperty AnimCurveFbxProp = FbxProperty::Create(CurrentBoneNode, FbxFloatDT, TCHAR_TO_UTF8(*AttributeName.ToString()));
+						AnimCurveFbxProp.ModifyFlag(FbxPropertyFlags::eAnimatable, true);
+						AnimCurveFbxProp.ModifyFlag(FbxPropertyFlags::eUserDefined, true);
 
-					FbxAnimCurve* AnimFbxCurve = AnimCurveFbxProp.GetCurve(AnimLayer, true);
-					AnimFbxCurve->KeyModifyBegin();
-					FloatCustomAttributeIndices.Emplace(AttributeIndex, AnimFbxCurve);
-				}
-				else if (VariantType == EVariantTypes::String)
-				{
-					FbxProperty AnimCurveFbxProp = FbxProperty::Create(CurrentBoneNode, FbxStringDT, TCHAR_TO_UTF8(*AttributeName.ToString()));
-					AnimCurveFbxProp.ModifyFlag(FbxPropertyFlags::eUserDefined, true);
+						FbxAnimCurve* AnimFbxCurve = AnimCurveFbxProp.GetCurve(AnimLayer, true);
+						AnimFbxCurve->KeyModifyBegin();
+						FloatCustomAttributeIndices.Emplace(AttributeIndex, AnimFbxCurve);
+					}
+					else if (AttributeType == FStringAnimationAttribute::StaticStruct())
+					{
+						FbxProperty AnimCurveFbxProp = FbxProperty::Create(CurrentBoneNode, FbxStringDT, TCHAR_TO_UTF8(*AttributeName.ToString()));
+						AnimCurveFbxProp.ModifyFlag(FbxPropertyFlags::eUserDefined, true);
 
-					// String attributes can't be keyed, simply set a normal value.
-					FString AttributeValue;
-					FCustomAttributesRuntime::GetAttributeValue(Attribute, 0.f, AttributeValue);
-					FbxString FbxValueString(TCHAR_TO_UTF8(*AttributeValue));
-					AnimCurveFbxProp.Set(FbxValueString);
-				}
-				else
-				{
-					ensureMsgf(false, TEXT("Trying to export unsupported custom attribte (float, int32 and FString are currently supported)"));
-				}
+						// String attributes can't be keyed, simply set a normal value.
+						FStringAnimationAttribute EvaluatedAttribute = Attribute.Curve.Evaluate<FStringAnimationAttribute>(0.f);
 
-				for (TPair<int32, FbxAnimCurve*>& CurrentAttributeCurve : FloatCustomAttributeIndices)
-				{
-					float AttributeValue = 0.f;
-					FCustomAttributesRuntime::GetAttributeValue(CustomAttributes[CurrentAttributeCurve.Key], AnimTime, AttributeValue);
-					int32 KeyIndex = CurrentAttributeCurve.Value->KeyAdd(ExportTime);
-					CurrentAttributeCurve.Value->KeySetValue(KeyIndex, AttributeValue);
-				}
+						FbxString FbxValueString(TCHAR_TO_UTF8(*EvaluatedAttribute.Value));
+						AnimCurveFbxProp.Set(FbxValueString);
+					}
+					else
+					{
+						ensureMsgf(false, TEXT("Trying to export unsupported custom attribte (float, int32 and FString are currently supported)"));
+					}
 
-				for (TPair<int32, FbxAnimCurve*>& CurrentAttributeCurve : IntCustomAttributeIndices)
-				{
-					int32 AttributeValue = 0;
-					FCustomAttributesRuntime::GetAttributeValue(CustomAttributes[CurrentAttributeCurve.Key], AnimTime, AttributeValue);
-					int32 KeyIndex = CurrentAttributeCurve.Value->KeyAdd(ExportTime);
-					CurrentAttributeCurve.Value->KeySetValue(KeyIndex, static_cast<float>(AttributeValue));
+					for (TPair<int32, FbxAnimCurve*>& CurrentAttributeCurve : FloatCustomAttributeIndices)
+					{
+						const FAnimatedBoneAttribute* FloatAttributePtr = CustomAttributes[CurrentAttributeCurve.Key];
+						ensure(FloatAttributePtr->Identifier.GetType() == FFloatAnimationAttribute::StaticStruct());
+
+						FFloatAnimationAttribute EvaluatedAttribute = FloatAttributePtr->Curve.Evaluate<FFloatAnimationAttribute>(AnimTime);
+
+						int32 KeyIndex = CurrentAttributeCurve.Value->KeyAdd(ExportTime);
+						CurrentAttributeCurve.Value->KeySetValue(KeyIndex, EvaluatedAttribute.Value);
+					}
+
+					for (TPair<int32, FbxAnimCurve*>& CurrentAttributeCurve : IntCustomAttributeIndices)
+					{
+						const FAnimatedBoneAttribute* IntAttributePtr = CustomAttributes[CurrentAttributeCurve.Key];
+						ensure(IntAttributePtr->Identifier.GetType() == FIntegerAnimationAttribute::StaticStruct());
+
+						FIntegerAnimationAttribute EvaluatedAttribute = IntAttributePtr->Curve.Evaluate< FIntegerAnimationAttribute>(AnimTime);
+
+						int32 KeyIndex = CurrentAttributeCurve.Value->KeyAdd(ExportTime);
+						CurrentAttributeCurve.Value->KeySetValue(KeyIndex, static_cast<float>(EvaluatedAttribute.Value));
+					}
 				}
 			}
 
diff --git a/Engine/Source/Editor/UnrealEd/Private/SkeletalMeshEdit.cpp b/Engine/Source/Editor/UnrealEd/Private/SkeletalMeshEdit.cpp
--- a/Engine/Source/Editor/UnrealEd/Private/SkeletalMeshEdit.cpp
+++ b/Engine/Source/Editor/UnrealEd/Private/SkeletalMeshEdit.cpp
@@ -38,6 +38,7 @@
 #include "FbxAnimUtils.h"
 #include "Animation/AnimData/AnimDataController.h"
 #include "Animation/AnimSequenceHelpers.h"
+#include "Animation/BuiltInAttributeTypes.h"
 
 #define LOCTEXT_NAMESPACE "SkeletalMeshEdit"
 
@@ -1390,7 +1391,8 @@ bool UnFbx::FFbxImporter::ImportCustomAttributeToBone(UAnimSequence* TargetSeque
 						return InProperty.Get<float>() * ValueScale;
 					}
 				});
-				TargetSequence->AddBoneFloatCustomAttribute(BoneName, FName(CurveName), TimeArray, FloatValues);
+
+				UE::Anim::AddTypedCustomAttribute<FFloatAnimationAttribute, float>(FName(CurveName), BoneName, TargetSequence, MakeArrayView(TimeArray), MakeArrayView(FloatValues));
 				break;
 			}
 			case EFbxType::eFbxBool:
@@ -1414,7 +1416,8 @@ bool UnFbx::FFbxImporter::ImportCustomAttributeToBone(UAnimSequence* TargetSeque
 						return static_cast<int32>(InProperty.Get<int32>() * ValueScale);
 					}
 				});
-				TargetSequence->AddBoneIntegerCustomAttribute(BoneName, FName(CurveName), TimeArray, IntValues);
+
+				UE::Anim::AddTypedCustomAttribute<FIntegerAnimationAttribute, int32>(FName(CurveName), BoneName, TargetSequence, MakeArrayView(TimeArray), MakeArrayView(IntValues));
 				break;
 			}
 			case EFbxType::eFbxString:
@@ -1434,7 +1437,8 @@ bool UnFbx::FFbxImporter::ImportCustomAttributeToBone(UAnimSequence* TargetSeque
 						return FString(UTF8_TO_TCHAR(InProperty.Get<FbxString>()));
 					}
 				});
-				TargetSequence->AddBoneStringCustomAttribute(BoneName, FName(CurveName), TimeArray, StringValues);
+
+				UE::Anim::AddTypedCustomAttribute<FStringAnimationAttribute, FString>(FName(CurveName), BoneName, TargetSequence, MakeArrayView(TimeArray), MakeArrayView(StringValues));
 				break;
 			}
 			default:
@@ -1530,7 +1534,7 @@ bool UnFbx::FFbxImporter::ImportAnimation(USkeleton* Skeleton, UAnimSequence * D
 
 	if (ImportOptions->bDeleteExistingNonCurveCustomAttributes)
 	{
-		DestSeq->RemoveAllCustomAttributes();
+		Controller->RemoveAllAttributes();
 	}
 	
 	const bool bReimportWarnings = GetDefault<UEditorPerProjectUserSettings>()->bAnimationReimportWarnings;
diff --git a/Engine/Source/Runtime/AnimGraphRuntime/Private/AnimNodes/AnimNode_BlendListBase.cpp b/Engine/Source/Runtime/AnimGraphRuntime/Private/AnimNodes/AnimNode_BlendListBase.cpp
--- a/Engine/Source/Runtime/AnimGraphRuntime/Private/AnimNodes/AnimNode_BlendListBase.cpp
+++ b/Engine/Source/Runtime/AnimGraphRuntime/Private/AnimNodes/AnimNode_BlendListBase.cpp
@@ -259,7 +259,7 @@ void FAnimNode_BlendListBase::Evaluate_AnyThread(FPoseContext& Output)
 		// Scratch arrays for evaluation, stack allocated
 		TArray<FCompactPose, TInlineAllocator<8>> FilteredPoses;
 		TArray<FBlendedCurve, TInlineAllocator<8>> FilteredCurve;
-		TArray<FStackCustomAttributes, TInlineAllocator<8>> FilteredAttributes;
+		TArray<UE::Anim::FStackAttributeContainer, TInlineAllocator<8>> FilteredAttributes;
 
 		FilteredPoses.SetNum(NumPoses, false);
 		FilteredCurve.SetNum(NumPoses, false);
diff --git a/Engine/Source/Runtime/AnimGraphRuntime/Private/AnimNodes/AnimNode_CopyPoseFromMesh.cpp b/Engine/Source/Runtime/AnimGraphRuntime/Private/AnimNodes/AnimNode_CopyPoseFromMesh.cpp
--- a/Engine/Source/Runtime/AnimGraphRuntime/Private/AnimNodes/AnimNode_CopyPoseFromMesh.cpp
+++ b/Engine/Source/Runtime/AnimGraphRuntime/Private/AnimNodes/AnimNode_CopyPoseFromMesh.cpp
@@ -236,7 +236,7 @@ void FAnimNode_CopyPoseFromMesh::Evaluate_AnyThread(FPoseContext& Output)
 	if (bCopyCustomAttributes)
 	{	
 		const FBoneContainer& RequiredBones = OutPose.GetBoneContainer();
-		FCustomAttributesRuntime::CopyAndRemapAttributes(SourceCustomAttributes, Output.CustomAttributes, BoneMapToSource, RequiredBones);		
+		UE::Anim::Attributes::CopyAndRemapAttributes(SourceCustomAttributes, Output.CustomAttributes, BoneMapToSource, RequiredBones);		
 	}
 }
 
diff --git a/Engine/Source/Runtime/AnimGraphRuntime/Private/AnimNodes/AnimNode_LayeredBoneBlend.cpp b/Engine/Source/Runtime/AnimGraphRuntime/Private/AnimNodes/AnimNode_LayeredBoneBlend.cpp
--- a/Engine/Source/Runtime/AnimGraphRuntime/Private/AnimNodes/AnimNode_LayeredBoneBlend.cpp
+++ b/Engine/Source/Runtime/AnimGraphRuntime/Private/AnimNodes/AnimNode_LayeredBoneBlend.cpp
@@ -242,7 +242,7 @@ void FAnimNode_LayeredBoneBlend::Evaluate_AnyThread(FPoseContext& Output)
 		TArray<FBlendedCurve> TargetBlendCurves;
 		TargetBlendCurves.SetNum(NumPoses);
 
-		TArray<FStackCustomAttributes> TargetBlendAttributes;
+		TArray<UE::Anim::FStackAttributeContainer> TargetBlendAttributes;
 		TargetBlendAttributes.SetNum(NumPoses);
 
 		for (int32 ChildIndex = 0; ChildIndex < NumPoses; ++ChildIndex)
diff --git a/Engine/Source/Runtime/AnimGraphRuntime/Private/AnimNodes/AnimNode_MakeDynamicAdditive.cpp b/Engine/Source/Runtime/AnimGraphRuntime/Private/AnimNodes/AnimNode_MakeDynamicAdditive.cpp
--- a/Engine/Source/Runtime/AnimGraphRuntime/Private/AnimNodes/AnimNode_MakeDynamicAdditive.cpp
+++ b/Engine/Source/Runtime/AnimGraphRuntime/Private/AnimNodes/AnimNode_MakeDynamicAdditive.cpp
@@ -54,7 +54,7 @@ void FAnimNode_MakeDynamicAdditive::Evaluate_AnyThread(FPoseContext& Output)
 	FAnimationRuntime::ConvertPoseToAdditive(Output.Pose, BaseEvalContext.Pose);
 	Output.Curve.ConvertToAdditive(BaseEvalContext.Curve);
 
-	FCustomAttributesRuntime::SubtractAttributes(BaseEvalContext.CustomAttributes, Output.CustomAttributes);
+	UE::Anim::Attributes::ConvertToAdditive(BaseEvalContext.CustomAttributes, Output.CustomAttributes);
 }
 
 void FAnimNode_MakeDynamicAdditive::GatherDebugData(FNodeDebugData& DebugData)
diff --git a/Engine/Source/Runtime/AnimGraphRuntime/Private/AnimNodes/AnimNode_MultiWayBlend.cpp b/Engine/Source/Runtime/AnimGraphRuntime/Private/AnimNodes/AnimNode_MultiWayBlend.cpp
--- a/Engine/Source/Runtime/AnimGraphRuntime/Private/AnimNodes/AnimNode_MultiWayBlend.cpp
+++ b/Engine/Source/Runtime/AnimGraphRuntime/Private/AnimNodes/AnimNode_MultiWayBlend.cpp
@@ -9,7 +9,7 @@ struct FMultiBlendData : public TThreadSingleton<FMultiBlendData>
 	TArray<FCompactPose, TInlineAllocator<8>> SourcePoses;
 	TArray<float, TInlineAllocator<8>> SourceWeights;
 	TArray<FBlendedCurve, TInlineAllocator<8>> SourceCurves;
-	TArray<FStackCustomAttributes, TInlineAllocator<8>> SourceAttributes;
+	TArray<UE::Anim::FStackAttributeContainer, TInlineAllocator<8>> SourceAttributes;
 };
 
 /////////////////////////////////////////////////////
@@ -115,7 +115,7 @@ void FAnimNode_MultiWayBlend::Evaluate_AnyThread(FPoseContext& Output)
 	FMultiBlendData& BlendData = FMultiBlendData::Get();
 	TArray<FCompactPose, TInlineAllocator<8>>& SourcePoses = BlendData.SourcePoses;
 	TArray<FBlendedCurve, TInlineAllocator<8>>& SourceCurves = BlendData.SourceCurves;
-	TArray<FStackCustomAttributes, TInlineAllocator<8>>& SourceAttributes = BlendData.SourceAttributes;
+	TArray<UE::Anim::FStackAttributeContainer, TInlineAllocator<8>>& SourceAttributes = BlendData.SourceAttributes;
 	TArray<float, TInlineAllocator<8>>& SourceWeights = BlendData.SourceWeights;
 
 	const int32 SourcePosesInitialNum = SourcePoses.Num();
@@ -138,7 +138,7 @@ void FAnimNode_MultiWayBlend::Evaluate_AnyThread(FPoseContext& Output)
 				FBlendedCurve& SourceCurve = SourceCurves.AddDefaulted_GetRef();
 				SourceCurve.MoveFrom(PoseContext.Curve);
 
-				FStackCustomAttributes& SourceAttribute = SourceAttributes.AddDefaulted_GetRef();
+				UE::Anim::FStackAttributeContainer& SourceAttribute = SourceAttributes.AddDefaulted_GetRef();
 				SourceAttribute.MoveFrom(PoseContext.CustomAttributes);
 
 				SourceWeights.Add(CurrentAlpha);
@@ -153,7 +153,7 @@ void FAnimNode_MultiWayBlend::Evaluate_AnyThread(FPoseContext& Output)
 		// obtain views onto the ends of our stacks
 		TArrayView<FCompactPose> SourcePosesView = MakeArrayView(&SourcePoses[SourcePosesInitialNum], SourcePosesAdded);
 		TArrayView<FBlendedCurve> SourceCurvesView = MakeArrayView(&SourceCurves[SourcePosesInitialNum], SourcePosesAdded);
-		TArrayView<FStackCustomAttributes> SourceAttributesView = MakeArrayView(&SourceAttributes[SourcePosesInitialNum], SourcePosesAdded);
+		TArrayView<UE::Anim::FStackAttributeContainer> SourceAttributesView = MakeArrayView(&SourceAttributes[SourcePosesInitialNum], SourcePosesAdded);
 		TArrayView<float> SourceWeightsView = MakeArrayView(&SourceWeights[SourcePosesInitialNum], SourcePosesAdded);
 
 		FAnimationPoseData AnimationPoseData(Output);
diff --git a/Engine/Source/Runtime/AnimGraphRuntime/Private/AnimNodes/AnimNode_RandomPlayer.cpp b/Engine/Source/Runtime/AnimGraphRuntime/Private/AnimNodes/AnimNode_RandomPlayer.cpp
--- a/Engine/Source/Runtime/AnimGraphRuntime/Private/AnimNodes/AnimNode_RandomPlayer.cpp
+++ b/Engine/Source/Runtime/AnimGraphRuntime/Private/AnimNodes/AnimNode_RandomPlayer.cpp
@@ -271,7 +271,7 @@ void FAnimNode_RandomPlayer::Evaluate_AnyThread(FPoseContext& Output)
 		// Start Blending
 		FCompactPose Poses[2];
 		FBlendedCurve Curves[2];
-		FStackCustomAttributes Attributes[2];
+		UE::Anim::FStackAttributeContainer Attributes[2];
 		float Weights[2];
 
 		const FBoneContainer& RequiredBone = AnimProxy->GetRequiredBones();
diff --git a/Engine/Source/Runtime/AnimGraphRuntime/Public/AnimNodes/AnimNode_CopyPoseFromMesh.h b/Engine/Source/Runtime/AnimGraphRuntime/Public/AnimNodes/AnimNode_CopyPoseFromMesh.h
--- a/Engine/Source/Runtime/AnimGraphRuntime/Public/AnimNodes/AnimNode_CopyPoseFromMesh.h
+++ b/Engine/Source/Runtime/AnimGraphRuntime/Public/AnimNodes/AnimNode_CopyPoseFromMesh.h
@@ -74,7 +74,7 @@ private:
 	TMap<FName, float> SourceCurveList;
 
 	// Cached attributes, copied on the game thread
-	FHeapCustomAttributes SourceCustomAttributes;
+	UE::Anim::FHeapAttributeContainer SourceCustomAttributes;
 
 	// reinitialize mesh component 
 	void ReinitializeMeshComponent(USkeletalMeshComponent* NewSkeletalMeshComponent, USkeletalMeshComponent* TargetMeshComponent);
diff --git a/Engine/Source/Runtime/Core/Public/UObject/UE5MainStreamObjectVersion.h b/Engine/Source/Runtime/Core/Public/UObject/UE5MainStreamObjectVersion.h
--- a/Engine/Source/Runtime/Core/Public/UObject/UE5MainStreamObjectVersion.h
+++ b/Engine/Source/Runtime/Core/Public/UObject/UE5MainStreamObjectVersion.h
@@ -97,6 +97,9 @@ struct CORE_API FUE5MainStreamObjectVersion
 		// Expose particle Disabled flag to the game thread
 		AddDisabledFlag,
 
+		// Moving animation custom attributes from AnimationSequence to UAnimDataModel
+		MoveCustomAttributesToDataModel,
+
 		// -----<new versions can be added above this line>-------------------------------------------------
 		VersionPlusOne,
 		LatestVersion = VersionPlusOne - 1
diff --git a/Engine/Source/Runtime/Engine/Classes/Animation/AnimData/AnimDataController.h b/Engine/Source/Runtime/Engine/Classes/Animation/AnimData/AnimDataController.h
--- a/Engine/Source/Runtime/Engine/Classes/Animation/AnimData/AnimDataController.h
+++ b/Engine/Source/Runtime/Engine/Classes/Animation/AnimData/AnimDataController.h
@@ -7,6 +7,7 @@
 #include "UObject/ObjectMacros.h"
 #include "Animation/AnimData/AnimDataModel.h"
 #include "Animation/AnimCurveTypes.h"
+#include "Algo/Transform.h"
 
 #if WITH_EDITOR
 #include "ChangeTransactor.h"
@@ -19,6 +20,7 @@
 #endif // WITH_EDITOR
 
 struct FAnimationCurveIdentifier;
+struct FAnimationAttributeIdentifier;
 static const int32 DefaultCurveFlags = EAnimAssetCurveFlags::AACF_Editable;
 
 namespace UE {
@@ -306,7 +308,7 @@ public:
 	/**
 	* Removes a single key for the transform curve with provided identifier. Broadcasts a EAnimDataModelNotifyType::CurveChanged notify if successful.
 	*
-	* @param	CurveId			    Identifier for the transform curve for which the key is to be set
+	* @param	CurveId			    Identifier for the transform curve for which the key is to be removed
 	* @param	Time				Time of the key to be removed
 	* @param	bShouldTransact		Whether or not any undo-redo changes should be generated
 	*
@@ -368,10 +370,9 @@ public:
 	
 	/**
 	* Remove a single key from the curve with provided identifier and name. Broadcasts a EAnimDataModelNotifyType::CurveChanged notify if successful.
-	* In case a key for the provided key time already exists the key is replaced.
 	*
-	* @param	CurveId			    Identifier for the curve for which the key is to be set
-	* @param	NewCurveId			Time of the key to be removed
+	* @param	CurveId			    Identifier for the curve for which the key is to be removed
+	* @param	Time				Time of the key to be removed
 	* @param	bShouldTransact		Whether or not any undo-redo changes should be generated
 	*
 	* @return	Whether or not the curve key was succesfully removed
@@ -382,11 +383,11 @@ public:
 	/**
 	* Replace the keys for the curve with provided identifier and name. Broadcasts a EAnimDataModelNotifyType::CurveChanged notify if successful.
 	*
-	* @param	CurveId			    Identifier for the curve for which the key is to be set
+	* @param	CurveId			    Identifier for the curve for which the keys are to be replaced
 	* @param	CurveKeys			Keys with which the existing keys are to be replaced
 	* @param	bShouldTransact		Whether or not any undo-redo changes should be generated
 	*
-	* @return	Whether or not the curve key was succesfully set
+	* @return	Whether or not replacing curve keys was succesful
 	*/
 	UFUNCTION(BlueprintCallable, Category = CurveData)
 	bool SetCurveKeys(const FAnimationCurveIdentifier& CurveId, const TArray<FRichCurveKey>& CurveKeys, bool bShouldTransact = true);
@@ -435,19 +436,149 @@ public:
 	* @param	bShouldTransact		Whether or not any undo-redo changes should be generated
 	*/
 	void ResetModel(bool bShouldTransact = true);
+
+	/**
+	* Adds a new attribute with the provided information. Broadcasts a EAnimDataModelNotifyType::AttributeAdded notify if successful.
+	*
+	* @param	AttributeIdentifier		Identifier for the to-be-added attribute
+	* @param	bShouldTransact			Whether or not any undo-redo changes should be generated
+	*
+	* @return	Whether or not the attribute was succesfully added
+	*/
+	UFUNCTION(BlueprintCallable, Category = CurveData)
+	bool AddAttribute(const FAnimationAttributeIdentifier& AttributeIdentifier, bool bShouldTransact = true);
+
+	/**
+	* Removes an attribute, if found, with the provided information. Broadcasts a EAnimDataModelNotifyType::AttributeRemoved notify if successful.
+	*
+	* @param	AttributeIdentifier		Identifier for the to-be-removed attribute
+	* @param	bShouldTransact			Whether or not any undo-redo changes should be generated
+	*
+	* @return	Whether or not the attribute was succesfully removed
+	*/
+	UFUNCTION(BlueprintCallable, Category = CurveData)
+	bool RemoveAttribute(const FAnimationAttributeIdentifier& AttributeIdentifier, bool bShouldTransact = true);
+
+	/**
+	* Removes all attributes for the specified bone name, if any. Broadcasts a EAnimDataModelNotifyType::AttributeRemoved notify for each removed attribute.
+	*
+	* @param	BoneName			Name of the bone to remove attributes for
+	* @param	bShouldTransact		Whether or not any undo-redo changes should be generated
+	*
+	* @return	Total number of removes attributes
+	*/
+	UFUNCTION(BlueprintCallable, Category = CurveData)
+	int32 RemoveAllAttributesForBone(const FName& BoneName, bool bShouldTransact = true);
+
+	/**
+	* Removes all stored attributes. Broadcasts a EAnimDataModelNotifyType::AttributeRemoved notify for each removed attribute.
+	*
+	* @param	bShouldTransact		Whether or not any undo-redo changes should be generated
+	*
+	* @return	Total number of removes attributes
+	*/
+	UFUNCTION(BlueprintCallable, Category = CurveData)
+	int32 RemoveAllAttributes(bool bShouldTransact = true);	
+
+	/**
+	* Sets a single key for the attribute with provided identifier. Broadcasts a EAnimDataModelNotifyType::AttributeChanged notify if successful.
+	* In case a key for the provided key time already exists the key is replaced.
+	*
+	* @param	AttributeIdentifier		Identifier for the attribute for which the key is to be set
+	* @param	Time					Time of the to-be-set key
+	* @param	KeyValue				Value (templated) of the to-be-set key
+	* @param	bShouldTransact			Whether or not any undo-redo changes should be generated
+	*
+	* @return	Whether or not the key was succesfully set
+	*/
+	template<typename AttributeType>
+	bool SetTypedAttributeKey(const FAnimationAttributeIdentifier& AttributeIdentifier, float Time, AttributeType& KeyValue, bool bShouldTransact = true)
+	{
+		return SetAttributeKey_Internal(AttributeIdentifier, Time, (const void*)&KeyValue, AttributeType::StaticStruct(), bShouldTransact);
+	}
+
+	/**
+	* Sets a single key for the attribute with provided identifier. Broadcasts a EAnimDataModelNotifyType::AttributeChanged notify if successful.
+	* In case a key for the provided key time already exists the key is replaced.
+	*
+	* @param	AttributeIdentifier		Identifier for the attribute for which the key is to be set
+	* @param	Time					Time of the to-be-set key
+	* @param	KeyValue				Value of the to-be-set key
+	* @param	bShouldTransact			Whether or not any undo-redo changes should be generated
+	*
+	* @return	Whether or not the key was succesfully set
+	*/
+	bool SetAttributeKey(const FAnimationAttributeIdentifier& AttributeIdentifier, float Time, const void* KeyValue, bool bShouldTransact = true)
+	{
+		return SetAttributeKey_Internal(AttributeIdentifier, Time, KeyValue, AttributeIdentifier.GetType(), bShouldTransact);
+	}
+	
+	/**
+	* Replace the keys for the attribute with provided identifier. Broadcasts a EAnimDataModelNotifyType::AttributeChanged notify if successful.
+	*
+	* @param	AttributeIdentifier		Identifier for the attribute for which the keys are to be replaced
+	* @param	Times					Times with which the existing key timings are to be replaced
+	* @param	KeyValues				Values with which the existing key values are to be replaced
+	* @param	bShouldTransact			Whether or not any undo-redo changes should be generated
+	*
+	* @return	Whether or not replacing the attribute keys was succesful
+	*/
+	bool SetAttributeKeys(const FAnimationAttributeIdentifier& AttributeIdentifier, TArrayView<const float> Times, TArrayView<const void*> KeyValues, bool bShouldTransact = true)
+	{
+		return SetAttributeKeys_Internal(AttributeIdentifier, Times, KeyValues, AttributeIdentifier.GetType(), bShouldTransact);
+	}
+
+	/**
+	* Replace the keys for the attribute with provided identifier. Broadcasts a EAnimDataModelNotifyType::AttributeChanged notify if successful.
+	*
+	* @param	AttributeIdentifier		Identifier for the attribute for which the keys are to be replaced
+	* @param	Times					Times with which the existing key timings are to be replaced
+	* @param	KeyValues				Values (templated) with which the existing key values are to be replaced
+	* @param	bShouldTransact			Whether or not any undo-redo changes should be generated
+	*
+	* @return	Whether or not replacing the attribute keys was succesful
+	*/
+	template<typename AttributeType>
+	bool SetTypedAttributeKeys(const FAnimationAttributeIdentifier& AttributeIdentifier, TArrayView<const float> Times, TArrayView<AttributeType> KeyValues, bool bShouldTransact = true)
+	{		
+		TArray<const void*> KeyValuePtrs;
+		Algo::Transform(KeyValues, KeyValuePtrs, [](const AttributeType& Value)
+		{
+			return (const void*)&Value;
+		});
+
+		return SetAttributeKeys_Internal(AttributeIdentifier, Times, MakeArrayView(KeyValuePtrs), AttributeType::StaticStruct(), bShouldTransact);
+	}
+
+	/**
+	* Remove a single key from the attribute with provided identifier. Broadcasts a EAnimDataModelNotifyType::AttributeChanged notify if successful.
+	*
+	* @param	AttributeIdentifier		Identifier for the attribute from which the key is to be removed
+	* @param	Time					Time of the key to be removed
+	* @param	bShouldTransact			Whether or not any undo-redo changes should be generated
+	*
+	* @return	Whether or not the attribute key was succesfully removed
+	*/
+	UFUNCTION(BlueprintCallable, Category = CurveData)
+	bool RemoveAttributeKey(const FAnimationAttributeIdentifier& AttributeIdentifier, float Time, bool bShouldTransact = true);
 protected:
 	/** Functionality used by FOpenBracketAction and FCloseBracketAction to broadcast their equivalent notifies without actually opening a bracket. */
 	void NotifyBracketOpen();
 	void NotifyBracketClosed();
 
 private:
+	/** Internal functionality for setting Attribute curve key(s) */
+	bool SetAttributeKey_Internal(const FAnimationAttributeIdentifier& AttributeIdentifier, float Time, const void* KeyValue, const UScriptStruct* TypeStruct, bool bShouldTransact = true);
+	bool SetAttributeKeys_Internal(const FAnimationAttributeIdentifier& AttributeIdentifier, TArrayView<const float> Times, TArrayView<const void*> KeyValues, const UScriptStruct* TypeStruct, bool bShouldTransact = true);
+
 	/** Returns whether or not the supplied curve type is supported by the controller functionality */
 	const bool IsSupportedCurveType(ERawCurveTrackTypes CurveType) const;
 	/** Returns the string representation of the provided curve enum type value */
 	FString GetCurveTypeValueName(ERawCurveTrackTypes InType) const;
 	
-	/** Resizes the curve data stored on the model according to the provided new length and time at which to insert or remove time */
+	/** Resizes the curve/attribute data stored on the model according to the provided new length and time at which to insert or remove time */
 	void ResizeCurves(float NewLength, bool bInserted, float T0, float T1, bool bShouldTransact = true);
+	void ResizeAttributes(float NewLength, bool bInserted, float T0, float T1, bool bShouldTransact = true);
 
 	/** Ensures that a valid model is currently targeted */
 	void ValidateModel() const;
diff --git a/Engine/Source/Runtime/Engine/Classes/Animation/AnimData/AnimDataModel.h b/Engine/Source/Runtime/Engine/Classes/Animation/AnimData/AnimDataModel.h
--- a/Engine/Source/Runtime/Engine/Classes/Animation/AnimData/AnimDataModel.h
+++ b/Engine/Source/Runtime/Engine/Classes/Animation/AnimData/AnimDataModel.h
@@ -5,6 +5,7 @@
 #include "CoreMinimal.h"
 #include "UObject/ObjectMacros.h"
 #include "AnimDataNotifications.h"
+#include "Animation/AttributeCurve.h"
 
 #include "AnimDataModel.generated.h"
 
@@ -46,6 +47,23 @@ struct ENGINE_API FAnimationCurveData
 	TArray<FTransformCurve>	TransformCurves;
 };
 
+/**
+ * Structure encapsulating animated (bone) attribute data.
+ */
+USTRUCT(BlueprintType)
+struct ENGINE_API FAnimatedBoneAttribute
+{
+	GENERATED_BODY()
+
+	/** Identifier to reference this attribute by */
+	UPROPERTY(VisibleAnywhere, BlueprintReadOnly, Category = "Model")
+	FAnimationAttributeIdentifier Identifier;	
+
+	/** Curve containing the (animated) attribute data */
+	UPROPERTY(VisibleAnywhere, BlueprintReadOnly, Category = "Model")
+	FAttributeCurve Curve;
+};
+
 struct FTrackToSkeletonMap;
 struct FAnimationCurveIdentifier;
 
@@ -213,6 +231,36 @@ public:
 	*/
 	const FRichCurve& GetRichCurve(const FAnimationCurveIdentifier& CurveIdentifier) const;
 		
+	/**
+	* @return	Animated (bone) attributes stored
+	*/
+	TArrayView<const FAnimatedBoneAttribute> GetAttributes() const;
+
+	/**
+	* @return	Number of animated (bone) attributes stored
+	*/
+	int32 GetNumberOfAttributes() const;
+
+	/**
+	* @return	Number of animated (bone) attributes stored for the specified bone index
+	*/
+	int32 GetNumberOfAttributesForBoneIndex(const int32 BoneIndex) const;
+
+	/**
+	* @return	All animated (bone) attributes stored for the specified bone name
+	*/
+	void GetAttributesForBone(const FName& BoneName, TArray<const FAnimatedBoneAttribute*>& OutBoneAttributes) const;
+
+	/**
+	* @return	Animated (bone) attribute object for the provided identifier if valid
+	*/
+	const FAnimatedBoneAttribute& GetAttribute(const FAnimationAttributeIdentifier& AttributeIdentifier) const;
+
+	/**
+	* @return	Animated (bone) attribute ptr for the provided identifier if valid, otherwise returns a nullptr 
+	*/
+	const FAnimatedBoneAttribute* FindAttribute(const FAnimationAttributeIdentifier& AttributeIdentifier) const;
+		
 	/**
 	* @return	The outer UAnimSequence object if found, otherwise returns a nullptr 
 	*/
@@ -276,7 +324,23 @@ private:
 			ModifiedEventDynamic.Broadcast(NotifyType, this, Payload);
 		}
 
-		GenerateTransientData();
+		// Only regenerate transient data when not in a bracket, or at the end of one
+		{
+			if (NotifyType == EAnimDataModelNotifyType::BracketOpened)
+			{
+				++BracketCounter;
+			}
+			if (NotifyType == EAnimDataModelNotifyType::BracketClosed)
+			{
+				--BracketCounter;
+			}
+
+			check(BracketCounter >= 0);
+			if (BracketCounter == 0)
+			{
+				GenerateTransientData();
+			}
+		}
 	}
 
 	/**
@@ -307,6 +371,7 @@ private:
 				--BracketCounter;
 			}
 
+			check(BracketCounter >= 0);
 			if (BracketCounter == 0)
 			{
 				GenerateTransientData();
@@ -351,6 +416,10 @@ private:
 	UPROPERTY(VisibleAnywhere, BlueprintReadOnly, Category = "Animation Data Model", meta = (AllowPrivateAccess = "true"))
 	FAnimationCurveData CurveData;
 	
+	/** Container with all animated (bone) attribute data */
+	UPROPERTY(VisibleAnywhere, BlueprintReadOnly, Category = "Animation Data Model", meta = (AllowPrivateAccess = "true"))
+	TArray<FAnimatedBoneAttribute> AnimatedBoneAttributes;
+	
 private:
 	/** Transient data, kept around for backward-compatibility */
 	UPROPERTY(Transient, VisibleAnywhere, Category = "Animation Data Model")
diff --git a/Engine/Source/Runtime/Engine/Classes/Animation/AnimData/AnimDataNotifications.h b/Engine/Source/Runtime/Engine/Classes/Animation/AnimData/AnimDataNotifications.h
--- a/Engine/Source/Runtime/Engine/Classes/Animation/AnimData/AnimDataNotifications.h
+++ b/Engine/Source/Runtime/Engine/Classes/Animation/AnimData/AnimDataNotifications.h
@@ -9,6 +9,7 @@
 #include "Animation/SmartName.h"
 #include "Animation/AnimCurveTypes.h"
 #include "Animation/AnimData/CurveIdentifier.h"
+#include "Animation/AnimData/AttributeIdentifier.h"
 
 #include "AnimDataNotifications.generated.h"
 
@@ -58,6 +59,15 @@ enum class EAnimDataModelNotifyType : uint8
 
 	/** Indicates a curve its color has changed. Type of payload: FCurveChangedPayload */
 	CurveColorChanged,
+
+	/** Indicates a new attribute has been added. Type of payload: FAttributeAddedPayload */
+	AttributeAdded,
+	
+	/** Indicates a new attribute has been removed. Type of payload: FAttributeRemovedPayload */
+	AttributeRemoved,
+
+	/** Indicates an attribute has been changed. Type of payload: FAttributeChangedPayload */
+	AttributeChanged,
 	
 	/** Indicates the data model has been populated from the source UAnimSequence. Type of payload: FEmptyPayload */
 	Populated,
@@ -183,6 +193,20 @@ struct FCurveFlagsChangedPayload : public FCurvePayload
 	int32 OldFlags;
 };
 
+USTRUCT(BlueprintType)
+struct FAttributePayload : public FEmptyPayload
+{
+	GENERATED_BODY()
+
+	/** Identifier of the attribute */
+	UPROPERTY(BlueprintReadOnly, Category = Payload)
+	FAnimationAttributeIdentifier Identifier;
+};
+
+typedef FAttributePayload FAttributeAddedPayload;
+typedef FAttributePayload FAttributeRemovedPayload;
+typedef FAttributePayload FAttributeChangedPayload;
+
 USTRUCT(BlueprintType)
 struct FAnimDataModelNotifPayload
 {
diff --git a/Engine/Source/Runtime/Engine/Classes/Animation/AnimInstance.h b/Engine/Source/Runtime/Engine/Classes/Animation/AnimInstance.h
--- a/Engine/Source/Runtime/Engine/Classes/Animation/AnimInstance.h
+++ b/Engine/Source/Runtime/Engine/Classes/Animation/AnimInstance.h
@@ -39,7 +39,7 @@ struct FParallelEvaluationData
 {
 	FBlendedHeapCurve& OutCurve;
 	FCompactPose& OutPose;
-	FHeapCustomAttributes& OutAttributes;
+	UE::Anim::FHeapAttributeContainer& OutAttributes;
 };
 
 UENUM()
@@ -180,7 +180,7 @@ struct FSlotEvaluationPose
 	/* These Pose/Curve is stack allocator. You should not use it outside of stack. */
 	FCompactPose Pose;
 	FBlendedCurve Curve;
-	FStackCustomAttributes Attributes;
+	UE::Anim::FStackAttributeContainer Attributes;
 
 	FSlotEvaluationPose()
 		: AdditiveType(AAT_None)
diff --git a/Engine/Source/Runtime/Engine/Classes/Animation/AnimNodeBase.h b/Engine/Source/Runtime/Engine/Classes/Animation/AnimNodeBase.h
--- a/Engine/Source/Runtime/Engine/Classes/Animation/AnimNodeBase.h
+++ b/Engine/Source/Runtime/Engine/Classes/Animation/AnimNodeBase.h
@@ -15,7 +15,7 @@
 #include "Animation/AnimTrace.h"
 #include "Animation/AnimationPoseData.h"
 #include "UObject/FieldPath.h"
-#include "CustomAttributesRuntime.h"
+#include "Animation/AttributesRuntime.h"
 #include "Animation/AnimNodeMessages.h"
 
 // WARNING: This should always be the last include in any file that needs it (except .generated.h)
@@ -427,7 +427,7 @@ public:
 	/* These Pose/Curve/Attributes are allocated using MemStack. You should not use it outside of stack. */
 	FCompactPose	Pose;
 	FBlendedCurve	Curve;
-	FStackCustomAttributes CustomAttributes;
+	UE::Anim::FStackAttributeContainer CustomAttributes;
 
 public:
 	// This constructor allocates a new uninitialized pose for the specified anim instance
@@ -511,7 +511,7 @@ struct FComponentSpacePoseContext : public FAnimationBaseContext
 public:
 	FCSPose<FCompactPose>	Pose;
 	FBlendedCurve			Curve;
-	FStackCustomAttributes CustomAttributes;
+	UE::Anim::FStackAttributeContainer CustomAttributes;
 
 public:
 	// This constructor allocates a new uninitialized pose for the specified anim instance
diff --git a/Engine/Source/Runtime/Engine/Classes/Animation/AnimNode_LinkedInputPose.h b/Engine/Source/Runtime/Engine/Classes/Animation/AnimNode_LinkedInputPose.h
--- a/Engine/Source/Runtime/Engine/Classes/Animation/AnimNode_LinkedInputPose.h
+++ b/Engine/Source/Runtime/Engine/Classes/Animation/AnimNode_LinkedInputPose.h
@@ -43,7 +43,7 @@ struct ENGINE_API FAnimNode_LinkedInputPose : public FAnimNode_Base
 	 */
 	FCompactHeapPose CachedInputPose;
 	FBlendedHeapCurve CachedInputCurve;
-	FHeapCustomAttributes CachedAttributes;
+	UE::Anim::FHeapAttributeContainer CachedAttributes;
 
 	// The node index of the currently-linked outer node
 	int32 OuterGraphNodeIndex;
diff --git a/Engine/Source/Runtime/Engine/Classes/Animation/AnimNode_SaveCachedPose.h b/Engine/Source/Runtime/Engine/Classes/Animation/AnimNode_SaveCachedPose.h
--- a/Engine/Source/Runtime/Engine/Classes/Animation/AnimNode_SaveCachedPose.h
+++ b/Engine/Source/Runtime/Engine/Classes/Animation/AnimNode_SaveCachedPose.h
@@ -52,7 +52,7 @@ struct ENGINE_API FAnimNode_SaveCachedPose : public FAnimNode_Base
 protected:
 	FCompactPose CachedPose;
 	FBlendedCurve CachedCurve;
-	FStackCustomAttributes CachedAttributes;
+	UE::Anim::FStackAttributeContainer CachedAttributes;
 
 	struct FCachedUpdateContext
 	{
diff --git a/Engine/Source/Runtime/Engine/Classes/Animation/AnimNode_TransitionPoseEvaluator.h b/Engine/Source/Runtime/Engine/Classes/Animation/AnimNode_TransitionPoseEvaluator.h
--- a/Engine/Source/Runtime/Engine/Classes/Animation/AnimNode_TransitionPoseEvaluator.h
+++ b/Engine/Source/Runtime/Engine/Classes/Animation/AnimNode_TransitionPoseEvaluator.h
@@ -48,7 +48,7 @@ struct ENGINE_API FAnimNode_TransitionPoseEvaluator : public FAnimNode_Base
 
 	FCompactHeapPose CachedPose;
 	FBlendedHeapCurve CachedCurve;
-	FStackCustomAttributes CachedAttributes;
+	UE::Anim::FStackAttributeContainer CachedAttributes;
 
 	UPROPERTY(EditAnywhere, BlueprintReadWrite, Category=Pose, meta=(NeverAsPin, ClampMin="1", UIMin="1"))
 	int32 FramesToCachePose;
diff --git a/Engine/Source/Runtime/Engine/Classes/Animation/AnimSequence.h b/Engine/Source/Runtime/Engine/Classes/Animation/AnimSequence.h
--- a/Engine/Source/Runtime/Engine/Classes/Animation/AnimSequence.h
+++ b/Engine/Source/Runtime/Engine/Classes/Animation/AnimSequence.h
@@ -20,6 +20,7 @@
 #include "Containers/ArrayView.h"
 #include "Animation/CustomAttributes.h"
 #include "Animation/AnimData/AnimDataNotifications.h"
+#include "Animation/AttributeCurve.h"
 
 #if WITH_EDITOR
 #include "AnimData/AnimDataModel.h"
@@ -810,7 +811,7 @@ public:
 	/**
 	* Return true if compressed data is out of date / missing and so animation needs to use raw data
 	*/
-	bool DoesNeedRecompress() const { return GetSkeleton() && (bUseRawDataOnly || (GetSkeletonVirtualBoneGuid() != GetSkeleton()->GetVirtualBoneGuid()) || !HasValidBakedCustomAttributes()); }
+	bool DoesNeedRecompress() const { return GetSkeleton() && (bUseRawDataOnly || (GetSkeletonVirtualBoneGuid() != GetSkeleton()->GetVirtualBoneGuid())); }
 
 	/**
 	 * Create Animation Sequence from Reference Pose of the Mesh
@@ -1002,79 +1003,46 @@ public:
 
 public:
 #if WITH_EDITOR
+	UE_DEPRECATED(5.0, "AddBoneCustomAttribute has been deprecated see UAnimDataController::AddAttribute")
 	UFUNCTION(BlueprintCallable, Category=CustomAttributes)
-	void AddBoneFloatCustomAttribute(const FName& BoneName, const FName& AttributeName, const TArray<float>& TimeKeys, const TArray<float>& ValueKeys)
-	{
-		AddBoneCustomAttribute<float>(BoneName, AttributeName, TimeKeys, ValueKeys);
-	}
-	
+	void AddBoneFloatCustomAttribute(const FName& BoneName, const FName& AttributeName, const TArray<float>& TimeKeys, const TArray<float>& ValueKeys);
+
+	UE_DEPRECATED(5.0, "AddBoneCustomAttribute has been deprecated see UAnimDataController::AddAttribute")
 	UFUNCTION(BlueprintCallable, Category = CustomAttributes)
-	void AddBoneIntegerCustomAttribute(const FName& BoneName, const FName& AttributeName, const TArray<float>& TimeKeys, const TArray<int32>& ValueKeys)
-	{
-		AddBoneCustomAttribute<int32>(BoneName, AttributeName, TimeKeys, ValueKeys);
-	}
+	void AddBoneIntegerCustomAttribute(const FName& BoneName, const FName& AttributeName, const TArray<float>& TimeKeys, const TArray<int32>& ValueKeys);
 
+	UE_DEPRECATED(5.0, "AddBoneStringCustomAttribute has been deprecated see UAnimDataController::AddAttribute")
 	UFUNCTION(BlueprintCallable, Category = CustomAttributes)
-	void AddBoneStringCustomAttribute(const FName& BoneName, const FName& AttributeName, const TArray<float>& TimeKeys, const TArray<FString>& ValueKeys)
-	{
-		AddBoneCustomAttribute<FString>(BoneName, AttributeName, TimeKeys, ValueKeys);
-	}
+	void AddBoneStringCustomAttribute(const FName& BoneName, const FName& AttributeName, const TArray<float>& TimeKeys, const TArray<FString>& ValueKeys);
 
+	UE_DEPRECATED(5.0, "RemoveCustomAttribute has been deprecated see UAnimDataController::RemoveAttribute")
 	UFUNCTION(BlueprintCallable, Category = CustomAttributes)
 	void RemoveCustomAttribute(const FName& BoneName, const FName& AttributeName);
 
+	UE_DEPRECATED(5.0, "RemoveAllCustomAttributesForBone has been deprecated see UAnimDataController::RemoveAllAttributesForBone")
 	UFUNCTION(BlueprintCallable, Category = CustomAttributes)
 	void RemoveAllCustomAttributesForBone(const FName& BoneName);
 
+	UE_DEPRECATED(5.0, "RemoveAllCustomAttributes has been deprecated see UAnimDataController::RemoveAllAttributes")
 	UFUNCTION(BlueprintCallable, Category = CustomAttributes)
 	void RemoveAllCustomAttributes();
-
-	void GetCustomAttributesForBone(const FName& BoneName, TArray<FCustomAttribute>& OutAttributes) const;
+	
+	PRAGMA_DISABLE_DEPRECATION_WARNINGS
+	UE_DEPRECATED(5.0, "GetCustomAttributesForBone has been deprecated see UAnimDataModel::GetAttributesForBone")
+	void GetCustomAttributesForBone(const FName& BoneName, TArray<FCustomAttribute>& OutAttributes) const {}
+	PRAGMA_ENABLE_DEPRECATION_WARNINGS
 #endif // WITH_EDITOR
+	UE_DEPRECATED(5.0, "GetCustomAttributesForBone has been deprecated use EvaluateAttribute instead")
+	void GetCustomAttributes(FAnimationPoseData& OutAnimationPoseData, const FAnimExtractContext& ExtractionContext, bool bUseRawData) const
+	{
+		EvaluateAttributes(OutAnimationPoseData, ExtractionContext, bUseRawData);
+	}
 
-	void GetCustomAttributes(FAnimationPoseData& OutAnimationPoseData, const FAnimExtractContext& ExtractionContext, bool bUseRawData) const;
+	void EvaluateAttributes(FAnimationPoseData& OutAnimationPoseData, const FAnimExtractContext& ExtractionContext, bool bUseRawData) const;	
 protected:
 #if WITH_EDITOR
-	template<typename DataType>
-	void AddBoneCustomAttribute(const FName& BoneName, const FName& AttributeName, const TArrayView<const float> TimeKeys, const TArrayView<const DataType> ValueKeys)
-	{
-		ensureMsgf(TimeKeys.Num() == ValueKeys.Num(), TEXT("Time keys do not match value keys"));
-
-		constexpr EVariantTypes VariantType = TVariantTraits<DataType>::GetType();
-		static_assert(VariantType == EVariantTypes::Int32 || VariantType == EVariantTypes::Float || VariantType == EVariantTypes::String, "Unsupported variant (data) type");
-
-		FCustomAttributePerBoneData& PerBoneData = FindOrAddCustomAttributeForBone(BoneName);
-		PerBoneData.BoneTreeIndex = GetSkeleton()->GetReferenceSkeleton().FindBoneIndex(BoneName);
-
-		const bool bAlreadyExists = PerBoneData.Attributes.ContainsByPredicate([AttributeName](FCustomAttribute& Attribute)
-		{
-			return Attribute.Name == AttributeName;
-		});
-
-		if (!bAlreadyExists)
-		{
-			FCustomAttribute& NewAttribute = PerBoneData.Attributes.AddDefaulted_GetRef();
-			NewAttribute.Name = AttributeName;
-			NewAttribute.VariantType = (int32)VariantType;
-
-			NewAttribute.Times = TimeKeys;
-
-			for (const DataType& Value : ValueKeys)
-			{
-				NewAttribute.Values.Add(FVariant(Value));
-			}
-			
-			// Update the Guid used to keep track of raw / baked versions
-			CustomAttributesGuid = FGuid::NewGuid();
-		}
-		else
-		{
-			UE_LOG(LogAnimation, Warning, TEXT("Unable to add Custom Attribute %s to bone %s as it already exist."), *AttributeName.ToString(), *BoneName.ToString());
-		}
-	}
-	
-	void SynchronousCustomAttributesCompression();
-	FCustomAttributePerBoneData& FindOrAddCustomAttributeForBone(const FName& BoneName);
+	void SynchronousAnimatedBoneAttributesCompression();
+	void MoveAttributesToModel();
 #endif // WITH_EDITOR
 
 protected:
@@ -1112,24 +1080,17 @@ protected:
 	bool bBlockCompressionRequests;
 
 private:
+	PRAGMA_DISABLE_DEPRECATION_WARNINGS
+	UE_DEPRECATED(5.0, "PerBoneCustomAttributeData has been deprecated see UAnimDataModel::AnimatedBoneAttributes")
 	UPROPERTY(VisibleAnywhere, EditFixedSize, Category=CustomAttributes)
 	TArray<FCustomAttributePerBoneData> PerBoneCustomAttributeData;
-	
-	UPROPERTY()
-	FGuid CustomAttributesGuid;
+	PRAGMA_ENABLE_DEPRECATION_WARNINGS
+#endif // WITH_EDITORONLY_DATA
 
+protected:
 	UPROPERTY()
-	FGuid BakedCustomAttributesGuid;
+	TMap<FAnimationAttributeIdentifier, FAttributeCurve> AttributeCurves;
 
-	bool HasValidBakedCustomAttributes() const
-	{
-		// Ensure the raw / baked versions match
-		return CustomAttributesGuid == BakedCustomAttributesGuid;
-	}
-#endif // WITH_EDITOR
-
-	UPROPERTY()
-	TArray<FBakedCustomAttributePerBoneData> BakedPerBoneCustomAttributeData;
 public:
 	friend class UAnimationAsset;
 	friend struct FScopedAnimSequenceRawDataCache;
diff --git a/Engine/Source/Runtime/Engine/Classes/Animation/AnimationPoseData.h b/Engine/Source/Runtime/Engine/Classes/Animation/AnimationPoseData.h
--- a/Engine/Source/Runtime/Engine/Classes/Animation/AnimationPoseData.h
+++ b/Engine/Source/Runtime/Engine/Classes/Animation/AnimationPoseData.h
@@ -6,16 +6,17 @@
 
 struct FCompactPose;
 struct FBlendedCurve;
-struct FStackCustomAttributes;
 struct FPoseContext;
 struct FSlotEvaluationPose;
 
+namespace UE { namespace Anim { struct FStackAttributeContainer; } }
+
 /** Structure used for passing around animation pose related data throughout the Animation Runtime */
 struct ENGINE_API FAnimationPoseData
 {
 	FAnimationPoseData(FPoseContext& InPoseContext);
 	FAnimationPoseData(FSlotEvaluationPose& InSlotPoseContext);
-	FAnimationPoseData(FCompactPose& InPose, FBlendedCurve& InCurve, FStackCustomAttributes& InAttributes);
+	FAnimationPoseData(FCompactPose& InPose, FBlendedCurve& InCurve, UE::Anim::FStackAttributeContainer& InAttributes);
 	
 	/** No default constructor, or assignment */
 	FAnimationPoseData() = delete;
@@ -26,11 +27,11 @@ struct ENGINE_API FAnimationPoseData
 	FCompactPose& GetPose();
 	const FBlendedCurve& GetCurve() const;
 	FBlendedCurve& GetCurve();
-	const FStackCustomAttributes& GetAttributes() const;
-	FStackCustomAttributes& GetAttributes();
+	const UE::Anim::FStackAttributeContainer& GetAttributes() const;
+	UE::Anim::FStackAttributeContainer& GetAttributes();
 
 protected:
 	FCompactPose& Pose;
 	FBlendedCurve& Curve;
-	FStackCustomAttributes& Attributes;
+	UE::Anim::FStackAttributeContainer& Attributes;
 };
\ No newline at end of file
diff --git a/Engine/Source/Runtime/Engine/Classes/Animation/BlendProfile.h b/Engine/Source/Runtime/Engine/Classes/Animation/BlendProfile.h
--- a/Engine/Source/Runtime/Engine/Classes/Animation/BlendProfile.h
+++ b/Engine/Source/Runtime/Engine/Classes/Animation/BlendProfile.h
@@ -12,8 +12,8 @@
 struct FAlphaBlend;
 struct FCompactPose;
 struct FBlendedCurve;
-struct FStackCustomAttributes;
 struct FSlotEvaluationPose;
+namespace UE { namespace Anim { struct FStackAttributeContainer; } }
 
 /** The mode in which the blend profile should be applied. */
 UENUM()
diff --git a/Engine/Source/Runtime/Engine/Classes/Animation/CustomAttributes.h b/Engine/Source/Runtime/Engine/Classes/Animation/CustomAttributes.h
--- a/Engine/Source/Runtime/Engine/Classes/Animation/CustomAttributes.h
+++ b/Engine/Source/Runtime/Engine/Classes/Animation/CustomAttributes.h
@@ -2,7 +2,7 @@
 
 #pragma once
 
-#include <UObject/ObjectMacros.h>
+#include "UObject/ObjectMacros.h"
 #include "UObject/NameTypes.h"
 #include "Misc/Variant.h"
 #include "Animation/AnimTypes.h"
@@ -37,11 +37,13 @@ struct ENGINE_API FCustomAttributeSetting
 	FString Meaning;
 };
 
+struct UE_DEPRECATED(5.0, "FCustomAttribute has been deprecated") FCustomAttribute;
 USTRUCT(Experimental)
 struct ENGINE_API FCustomAttribute
 {
 	GENERATED_BODY()
 
+#if WITH_EDITORONLY_DATA
 	/** Name of this attribute */
 	UPROPERTY(VisibleAnywhere, Category = CustomAttribute)
 	FName Name;
@@ -66,8 +68,11 @@ struct ENGINE_API FCustomAttribute
 
 		return true;
 	}
+#endif // WITH_EDITORONLY_DATA
 };
 
+#if WITH_EDITORONLY_DATA
+PRAGMA_DISABLE_DEPRECATION_WARNINGS
 // Custom serializer required for FVariant array
 template<>
 struct TStructOpsTypeTraits<FCustomAttribute> : public TStructOpsTypeTraitsBase2<FCustomAttribute>
@@ -77,68 +82,92 @@ struct TStructOpsTypeTraits<FCustomAttribute> : public TStructOpsTypeTraitsBase2
 		WithSerializer = true
 	};
 };
+PRAGMA_ENABLE_DEPRECATION_WARNINGS
+#endif // WITH_EDITORONLY_DATA
 
 /** Structure describing custom attributes for a single bone (index) */
+struct UE_DEPRECATED(5.0, "FCustomAttributePerBoneData has been deprecated") FCustomAttributePerBoneData;
+
 USTRUCT(Experimental)
 struct ENGINE_API FCustomAttributePerBoneData
 {
 	GENERATED_BODY()
 
+#if WITH_EDITORONLY_DATA
 	UPROPERTY(VisibleAnywhere, Category = CustomAttributeBoneData)
 	int32 BoneTreeIndex = 0;
 
+	PRAGMA_DISABLE_DEPRECATION_WARNINGS
 	UPROPERTY(VisibleAnywhere, EditFixedSize, Category = CustomAttributeBoneData)
 	TArray<FCustomAttribute> Attributes;
+	PRAGMA_ENABLE_DEPRECATION_WARNINGS
+#endif // WITH_EDITORONLY_DATA
 };
 
 /** (Baked) string custom attribute, uses FStringCurve for evaluation instead of FVariant array */
+struct UE_DEPRECATED(5.0, "FBakedStringCustomAttribute has been deprecated") FBakedStringCustomAttribute;
+
 USTRUCT(Experimental)
 struct FBakedStringCustomAttribute
 {
 	GENERATED_BODY()
 
+#if WITH_EDITORONLY_DATA
 	UPROPERTY(VisibleAnywhere, Category = CustomAttributeBoneData)
 	FName AttributeName;
 
 	UPROPERTY(VisibleAnywhere, Category = CustomAttributeBoneData)
 	FStringCurve StringCurve;
+#endif // WITH_EDITORONLY_DATA
 };
 
 /** (Baked) int32 custom attribute, uses FIntegralCurve for evaluation instead of FVariant array */
+struct UE_DEPRECATED(5.0, "FBakedIntegerCustomAttribute has been deprecated") FBakedIntegerCustomAttribute;
+
 USTRUCT(Experimental)
 struct FBakedIntegerCustomAttribute
 {
 	GENERATED_BODY()
 
+#if WITH_EDITORONLY_DATA
 	UPROPERTY(VisibleAnywhere, Category = CustomAttributeBoneData)
 	FName AttributeName;
 
 	UPROPERTY(VisibleAnywhere, Category = CustomAttributeBoneData)
 	FIntegralCurve IntCurve;
+#endif // WITH_EDITORONLY_DATA
 };
 
 /** (Baked) float custom attribute, uses FSimpleCurve for evaluation instead of FVariant array */
+struct UE_DEPRECATED(5.0, "FBakedFloatCustomAttribute has been deprecated") FBakedFloatCustomAttribute;
+
 USTRUCT(Experimental)
 struct FBakedFloatCustomAttribute
 {
 	GENERATED_BODY()
 
+#if WITH_EDITORONLY_DATA
 	UPROPERTY(VisibleAnywhere, Category = CustomAttributeBoneData)
 	FName AttributeName;
 
 	UPROPERTY(VisibleAnywhere, Category = CustomAttributeBoneData)
 	FSimpleCurve FloatCurve;
+#endif // WITH_EDITORONLY_DATA
 };
 
 /** Structure describing baked custom attributes for a single bone (index) */
+struct UE_DEPRECATED(5.0, "FBakedCustomAttributePerBoneData has been deprecated") FBakedCustomAttributePerBoneData;
+
 USTRUCT(Experimental)
 struct FBakedCustomAttributePerBoneData
 {
 	GENERATED_BODY()
 
+#if WITH_EDITORONLY_DATA
 	UPROPERTY()
 	int32 BoneTreeIndex = 0;
 
+	PRAGMA_DISABLE_DEPRECATION_WARNINGS
 	UPROPERTY(VisibleAnywhere, EditFixedSize, Category = CustomAttributeBoneData)
 	TArray<FBakedStringCustomAttribute> StringAttributes;
 
@@ -147,4 +176,6 @@ struct FBakedCustomAttributePerBoneData
 
 	UPROPERTY(VisibleAnywhere, EditFixedSize, Category = CustomAttributeBoneData)
 	TArray<FBakedFloatCustomAttribute> FloatAttributes;
-};
+	PRAGMA_ENABLE_DEPRECATION_WARNINGS
+#endif // WITH_EDITORONLY_DATA
+};
\ No newline at end of file
diff --git a/Engine/Source/Runtime/Engine/Classes/Animation/CustomAttributesRuntime.h b/Engine/Source/Runtime/Engine/Classes/Animation/CustomAttributesRuntime.h
--- a/Engine/Source/Runtime/Engine/Classes/Animation/CustomAttributesRuntime.h
+++ b/Engine/Source/Runtime/Engine/Classes/Animation/CustomAttributesRuntime.h
@@ -2,303 +2,72 @@
 
 #pragma once
 
-#include "Containers/ArrayView.h"
-#include "Misc/Variant.h"
-#include "Animation/AnimTypes.h"
-#include "UObject/NameTypes.h"
-#include "Containers/ContainersFwd.h"
+#include "CoreMinimal.h"
 
-struct FCompactPoseBoneIndex;
-struct FBlendSampleData;
-struct FBoneContainer;
+PRAGMA_DISABLE_DEPRECATION_WARNINGS
+#include "Animation/AttributesRuntime.h"
+PRAGMA_ENABLE_DEPRECATION_WARNINGS
 
-enum class ECustomAttributeBlendType : uint8;
+UE_STATIC_DEPRECATE(5.0, true, "CustomAttributesRuntime.h has been deprecated - please include and use AttributesRuntime.h instead.");
 
-/** "Header" structure describing a runtime custom attribute */
-struct FCustomAttributeInfo
-{
-	FCustomAttributeInfo(const FName& InName, const FCompactPoseBoneIndex& InCompactBoneIndex, const ECustomAttributeBlendType& InBlendType);
-
-	// Hash contains CombineHash(GetTypeHash(BoneIndex), GetTypeHash(AttributeName))
-	uint32 Hash;
-	// Index to the bone this attribute is stored for
-	int32 BoneIndex;
-	// Type of blending to use throughout the animation runtime
-	ECustomAttributeBlendType BlendType;
-};
-
-template<class BoneIndexType, typename InAllocator>
-struct TBaseCustomAttributes
-{
-	template <typename OtherBoneIndexType, typename OtherAllocator>
-	friend struct TBaseCustomAttributes;	
-
-	/** Retrieve the typed  array containing individual attribute values */
-	template<typename DataType> TArray<DataType, InAllocator>& GetValuesArray()
-	{
-		TArray<DataType, InAllocator>* ArrayPtr = nullptr;
-		GetValuesArray_Internal(ArrayPtr);
-		return *ArrayPtr;
-	}
-
-	template<typename DataType>
-	const TArray<DataType, InAllocator>& GetValuesArray() const
-	{
-		const TArray<DataType, InAllocator>* ArrayPtr = nullptr;
-		GetValuesArray_Internal(ArrayPtr);
-		return *ArrayPtr;
-	}
-
-	template<typename DataType>
-	int32 AddBoneAttribute(const BoneIndexType& BoneIndex, const FName& AttributeName, ECustomAttributeBlendType BlendType, const DataType& Value)
-	{
-		FCustomAttributeInfo AttributeInfo(AttributeName, BoneIndex, BlendType);
-		return AddBoneAttribute<DataType>(AttributeInfo, Value);
-	}
-	
-	template<typename DataType>
-	int32 AddBoneAttribute(const FCustomAttributeInfo& AttributeInfo, const DataType& Value)
-	{
-		constexpr int32 DataTypeIndex = GetIndexForType<DataType>();
-		const int32 ExistingAttributeIndex = AttributeInfos[DataTypeIndex].IndexOfByPredicate([BoneAttributeHash = AttributeInfo.Hash](const FCustomAttributeInfo& Attribute)
-		{
-			return Attribute.Hash == BoneAttributeHash;
-		});
-
-		// Should only add an attribute once
-		if (ensure(ExistingAttributeIndex == INDEX_NONE))
-		{
-			int32 NewIndex = INDEX_NONE;
-
-			AttributeInfos[DataTypeIndex].Add(AttributeInfo);
-			UniqueTypedBoneIndices[DataTypeIndex].AddUnique(AttributeInfo.BoneIndex);
-			TArray<DataType, InAllocator>& TypedArray = GetValuesArray<DataType>();
-			NewIndex = TypedArray.Add(Value);
-
-			// Ensure arrays match in size
-			ensure(AttributeInfos[DataTypeIndex].Num() == TypedArray.Num());
-
-			return NewIndex;
-		}
-
-		return ExistingAttributeIndex;
-	}
-
-	template<typename DataType>
-	bool GetBoneAttribute(const BoneIndexType& BoneIndex, const FName& AttributeName, DataType& OutValue) const
-	{
-		constexpr int32 DataTypeIndex = GetIndexForType<DataType>();
-		const int32 BoneIndexInt = BoneIndex.GetInt();
-
-		// Early out if for this bone index no attributes are currently contained
-		if (UniqueTypedBoneIndices[DataTypeIndex].Contains(BoneIndexInt))
-		{
-			const uint32 BoneAttributeHash = HashCombine(GetTypeHash(BoneIndexInt), GetTypeHash(AttributeName));
-			const int32 AttributeIndex = IndexOfBoneAttribute<DataType>(BoneAttributeHash, BoneIndexInt);
-			if (AttributeIndex != INDEX_NONE)
-			{
-				const TArray<DataType, InAllocator>& TypedArray = GetValuesArray<DataType>();
-				ensure(TypedArray.IsValidIndex(AttributeIndex));
-				OutValue = TypedArray[AttributeIndex];
-
-				return true;
-			}
-		}
-
-		return false;
-	}
-
-	template<typename DataType>
-	int32 IndexOfBoneAttribute(uint32 BoneAttributeHash, int32 BoneIndexInt) const
-	{
-		constexpr int32 DataTypeIndex = GetIndexForType<DataType>();
-
-		if (UniqueTypedBoneIndices[DataTypeIndex].Contains(BoneIndexInt))
-		{
-			return AttributeInfos[DataTypeIndex].IndexOfByPredicate([BoneAttributeHash](const FCustomAttributeInfo& Attribute)
-			{
-				return Attribute.Hash == BoneAttributeHash;
-			});
-		}
-
-		return INDEX_NONE;
-	}
-
-	template<typename DataType>
-	bool ContainsBoneAttribute(uint32 BoneAttributeHash, int32 BoneIndexInt) const
-	{
-		return IndexOfBoneAttribute<DataType>(BoneAttributeHash, BoneIndexInt) != INDEX_NONE;
-	}
-
-	template<typename DataType>
-	const TArray<FCustomAttributeInfo, InAllocator>& GetAttributeInfo() const
-	{
-		constexpr int32 DataTypeIndex = GetIndexForType<DataType>();
-		return AttributeInfos[DataTypeIndex];
-	}
-
-	template<typename DataType>
-	const TArray<int32, InAllocator>& GetUniqueBoneIndices() const
-	{
-		constexpr int32 DataTypeIndex = GetIndexForType<DataType>();
-		return UniqueTypedBoneIndices[DataTypeIndex];
-	}
-
-	template <typename OtherAllocator>
-	void CopyFrom(const TBaseCustomAttributes<BoneIndexType, OtherAllocator>& Other)
-	{
-		FloatValues = Other.FloatValues;
-		IntValues = Other.IntValues;
-		StringValues = Other.StringValues;
-
-		for (int32 Index = 0; Index < NumSupportedDataTypes; ++Index)
-		{
-			AttributeInfos[Index] = Other.AttributeInfos[Index];
-			UniqueTypedBoneIndices[Index] = Other.UniqueTypedBoneIndices[Index];
-		}
-	}
-
-	void CopyFrom(const TBaseCustomAttributes<BoneIndexType, InAllocator>& Other)
-	{
-		if (&Other != this)
-		{
-			FloatValues = Other.FloatValues;
-			IntValues = Other.IntValues;
-			StringValues = Other.StringValues;
-
-			for (int32 Index = 0; Index < NumSupportedDataTypes; ++Index)
-			{
-				AttributeInfos[Index] = Other.AttributeInfos[Index];
-				UniqueTypedBoneIndices[Index] = Other.UniqueTypedBoneIndices[Index];
-			}
-		}
-	}
-
-	/** Once moved, source is invalid */
-	void MoveFrom(TBaseCustomAttributes<BoneIndexType, InAllocator>& Other)
-	{
-		FloatValues = MoveTemp(Other.FloatValues);
-		IntValues = MoveTemp(Other.IntValues);
-		StringValues = MoveTemp(Other.StringValues);
-
-		for (int32 Index = 0; Index < NumSupportedDataTypes; ++Index)
-		{
-			AttributeInfos[Index] = MoveTemp(Other.AttributeInfos[Index]);
-			UniqueTypedBoneIndices[Index] = MoveTemp(Other.UniqueTypedBoneIndices[Index]);
-		}
-	}
-
-	bool ContainsData() const
-	{
-		return FloatValues.Num() || IntValues.Num() || StringValues.Num();
-	}
-
-	void Empty()
-	{
-		FloatValues.Empty();
-		IntValues.Empty();
-		StringValues.Empty();
-
-		for (int32 Index = 0; Index < NumSupportedDataTypes; ++Index)
-		{
-			AttributeInfos[Index].Empty();
-			UniqueTypedBoneIndices[Index].Empty();
-		}
-	}
-
-	bool operator!=(const TBaseCustomAttributes<BoneIndexType, InAllocator>& Other)
-	{
-		return FloatValues.Num() != Other.FloatValues.Num() 
-		|| IntValues.Num() != Other.IntValues.Num() 
-		|| StringValues.Num() != Other.StringValues.Num();
-	}
-
-protected:
-	// Number of, and explicit types which are currently supported
-	static CONSTEXPR int32 NumSupportedDataTypes = 3;
-	static CONSTEXPR EVariantTypes SupportedTypes[NumSupportedDataTypes] = { EVariantTypes::Float, EVariantTypes::Int32, EVariantTypes::String };	
-	
-	template<typename DataType>
-	static CONSTEXPR int32 GetIndexForType()
-	{
-		constexpr EVariantTypes VariantType = TVariantTraits<DataType>::GetType();
-		for (int32 TypeIndex = 0; TypeIndex < NumSupportedDataTypes; ++TypeIndex)
-		{
-			if (SupportedTypes[TypeIndex] == VariantType)
-			{
-				return TypeIndex;
-			}
-		}
-
-		return INDEX_NONE;
-	};
-
-	// Internal type array getters, used to retrieve by-reference array in the public API
-	void GetValuesArray_Internal(TArray<float, InAllocator>*& OutPtr) { OutPtr = &FloatValues; }
-	void GetValuesArray_Internal(TArray<int32, InAllocator>*& OutPtr) { OutPtr = &IntValues; }
-	void GetValuesArray_Internal(TArray<FString, InAllocator>*& OutPtr) { OutPtr = &StringValues; }
-
-	void GetValuesArray_Internal(const TArray<float, InAllocator>*& OutPtr) const { OutPtr = &FloatValues; }
-	void GetValuesArray_Internal(const TArray<int32, InAllocator>*& OutPtr) const { OutPtr = &IntValues; }
-	void GetValuesArray_Internal(const TArray<FString, InAllocator>*& OutPtr) const { OutPtr = &StringValues; }
-
-protected:
-	// Information for each stored custom attribute
-	TArray<FCustomAttributeInfo, InAllocator> AttributeInfos[NumSupportedDataTypes];
-	
-	/* Contains the uniquely added bone indices, on a per-type basis (could make this a non-per-type array and index into it from BoneIndices, would add some cost for runtime evaluation to retrieve unique bone indices per type) */
-	TArray<int32, InAllocator> UniqueTypedBoneIndices[NumSupportedDataTypes];
-
-	// Attribute typed value arrays
-	TArray<float, InAllocator> FloatValues;
-	TArray<int32, InAllocator> IntValues;
-	TArray<FString, InAllocator> StringValues;	
-};
-
-struct ENGINE_API FStackCustomAttributes : public TBaseCustomAttributes<FCompactPoseBoneIndex, FAnimStackAllocator> {};
-struct ENGINE_API FHeapCustomAttributes : public TBaseCustomAttributes<FCompactPoseBoneIndex, FDefaultAllocator> {};
-
-/** Helper functionality for custom attributes animation runtime */
 struct ENGINE_API FCustomAttributesRuntime
 {
 #if WITH_EDITOR
 	/** Editor functionality to retrieve custom attribute values from the raw data */
-	static void GetAttributeValue(struct FStackCustomAttributes& OutAttributes, const struct FCompactPoseBoneIndex& PoseBoneIndex, const struct FCustomAttribute& Attribute, const struct FAnimExtractContext& ExtractionContext);
-	static void GetAttributeValue(const struct FCustomAttribute& Attribute, float Time, float& OutValue);
-	static void GetAttributeValue(const struct FCustomAttribute& Attribute, float Time, int32& OutValue);
-	static void GetAttributeValue(const struct FCustomAttribute& Attribute, float Time, FString& OutValue);
+	UE_DEPRECATED(5.0, "Use UE::Animation::Attributes::GetAttributeValue with different signature")
+	static void GetAttributeValue(struct UE::Anim::FStackAttributeContainer& OutAttributes, const struct FCompactPoseBoneIndex& PoseBoneIndex, const struct FCustomAttribute& Attribute, const struct FAnimExtractContext& ExtractionContext) {}
+	UE_DEPRECATED(5.0, "Use UE::Animation::Attributes::GetAttributeValue with different signature")
+	static void GetAttributeValue(const struct FCustomAttribute& Attribute, float Time, float& OutValue) {}
+	UE_DEPRECATED(5.0, "Use UE::Animation::Attributes::GetAttributeValue with different signature")
+	static void GetAttributeValue(const struct FCustomAttribute& Attribute, float Time, int32& OutValue) {}
+	UE_DEPRECATED(5.0, "Use UE::Animation::Attributes::GetAttributeValue with different signature")
+	static void GetAttributeValue(const struct FCustomAttribute& Attribute, float Time, FString& OutValue) {}
 #endif // WITH_EDITOR
 
 	/** Blend custom attribute values from N set of inputs */
-	static void BlendAttributes(const TArrayView<const FStackCustomAttributes> SourceAttributes, const TArrayView<const float> SourceWeights, FStackCustomAttributes& OutAttributes);
+	UE_DEPRECATED(5.0, "Functionality moved, see UE::Animation::Attributes::BlendAttributes")
+	static void BlendAttributes(const TArrayView<const UE::Anim::FStackAttributeContainer> SourceAttributes, const TArrayView<const float> SourceWeights, UE::Anim::FStackAttributeContainer& OutAttributes) { UE::Anim::Attributes::BlendAttributes(SourceAttributes, SourceWeights, OutAttributes); }
 
 	/** Blend custom attribute values from N set of inputs (ptr-values) */
-	static void BlendAttributes(const TArrayView<const FStackCustomAttributes* const> SourceAttributes, const TArrayView<const float> SourceWeights, FStackCustomAttributes& OutAttributes);
+	UE_DEPRECATED(5.0, "Functionality moved, see UE::Animation::Attributes::BlendAttributes")
+	static void BlendAttributes(const TArrayView<const UE::Anim::FStackAttributeContainer* const> SourceAttributes, const TArrayView<const float> SourceWeights, UE::Anim::FStackAttributeContainer& OutAttributes) { UE::Anim::Attributes::BlendAttributes(SourceAttributes, SourceWeights, OutAttributes); }
 
 	/** Blend custom attribute values from N set of inputs, using input weight remapping */
-	static void BlendAttributes(const TArrayView<const FStackCustomAttributes> SourceAttributes, const TArrayView<const float> SourceWeights, const TArrayView<const int32> SourceWeightsIndices, FStackCustomAttributes& OutAttributes);
+	UE_DEPRECATED(5.0, "Functionality moved, see UE::Animation::Attributes::BlendAttributes")
+	static void BlendAttributes(const TArrayView<const UE::Anim::FStackAttributeContainer> SourceAttributes, const TArrayView<const float> SourceWeights, const TArrayView<const int32> SourceWeightsIndices, UE::Anim::FStackAttributeContainer& OutAttributes) { UE::Anim::Attributes::BlendAttributes(SourceAttributes, SourceWeights, SourceWeightsIndices, OutAttributes); }
 
 	/* Blend custom attribute values from 2 inputs, using per-bone weights */ 
-	static void BlendAttributesPerBone(const FStackCustomAttributes& SourceAttributes1, const FStackCustomAttributes& SourceAttributes2, const TArrayView<const float> WeightsOfSource2, FStackCustomAttributes& OutAttributes);
+	UE_DEPRECATED(5.0, "Functionality moved, see UE::Animation::Attributes::BlendAttributesPerBone")
+	static void BlendAttributesPerBone(const UE::Anim::FStackAttributeContainer& SourceAttributes1, const UE::Anim::FStackAttributeContainer& SourceAttributes2, const TArrayView<const float> WeightsOfSource2, UE::Anim::FStackAttributeContainer& OutAttributes) { UE::Anim::Attributes::BlendAttributesPerBone(SourceAttributes1, SourceAttributes2, WeightsOfSource2, OutAttributes); }
 		
 	/* Blend custom attribute values from N set of inputs, using bone filter pose weights */ 
-	static void BlendAttributesPerBoneFilter(const TArrayView<const FStackCustomAttributes> BlendAttributes, const TArray<FPerBoneBlendWeight>& BoneBlendWeights, FStackCustomAttributes& OutAttributes);
+	UE_DEPRECATED(5.0, "Functionality moved, see UE::Animation::Attributes::BlendAttributesPerBoneFilter")
+	static void BlendAttributesPerBoneFilter(const TArrayView<const UE::Anim::FStackAttributeContainer> BlendAttributes, const TArray<FPerBoneBlendWeight>& BoneBlendWeights, UE::Anim::FStackAttributeContainer& OutAttributes)
+	{
+		UE::Anim::Attributes::BlendAttributesPerBoneFilter(OutAttributes, BlendAttributes, BoneBlendWeights, OutAttributes);
+	}
 		
 	/** Add any new or override existing custom attributes */
-	static void OverrideAttributes(const FStackCustomAttributes& SourceAttributes, FStackCustomAttributes& OutAttributes, float Weight);
+	UE_DEPRECATED(5.0, "Functionality moved, see UE::Animation::Attributes::OverrideAttributes")
+	static void OverrideAttributes(const UE::Anim::FStackAttributeContainer& SourceAttributes, UE::Anim::FStackAttributeContainer& OutAttributes, float Weight) { UE::Anim::Attributes::OverrideAttributes(SourceAttributes, OutAttributes, Weight); }
 
 	/** Add any new or accumulate with existing custom attributes */
-	static void AccumulateAttributes(const FStackCustomAttributes& SourceAttributes, FStackCustomAttributes& OutAttributes, float Weight);
+	UE_DEPRECATED(5.0, "Functionality moved, see UE::Animation::Attributes::AccumulateAttributes")
+	static void AccumulateAttributes(const UE::Anim::FStackAttributeContainer& SourceAttributes, UE::Anim::FStackAttributeContainer& OutAttributes, float Weight) { UE::Anim::Attributes::AccumulateAttributes(SourceAttributes, OutAttributes, Weight, AAT_None); }
 
 	/** Add (negated) any new or subtract from existing custom attributes */
-	static void SubtractAttributes(const FStackCustomAttributes& SourceAttributes, FStackCustomAttributes& OutAttributes);
+	UE_DEPRECATED(5.0, "Functionality moved, see UE::Animation::Attributes::OverrideAttributes")
+	static void SubtractAttributes(const UE::Anim::FStackAttributeContainer& SourceAttributes, UE::Anim::FStackAttributeContainer& OutAttributes) { UE::Anim::Attributes::ConvertToAdditive(SourceAttributes, OutAttributes); }
 
 	/** Copy attributes from source, and remap the bone indices according to BoneMapToSource */
-	static void CopyAndRemapAttributes(const FHeapCustomAttributes& SourceAttributes, FStackCustomAttributes& OutAttributes, const TMap<int32, int32>& BoneMapToSource, const FBoneContainer& RequiredBones);
+	UE_DEPRECATED(5.0, "Functionality moved, see UE::Animation::Attributes::CopyAndRemapAttributes")
+	static void CopyAndRemapAttributes(const UE::Anim::FHeapAttributeContainer& SourceAttributes, UE::Anim::FStackAttributeContainer& OutAttributes, const TMap<int32, int32>& BoneMapToSource, const FBoneContainer& RequiredBones) { UE::Anim::Attributes::CopyAndRemapAttributes(SourceAttributes, OutAttributes, BoneMapToSource, RequiredBones); }
 
 	/** Interpolates between two sets of attributes */
-	static void InterpolateAttributes(const FHeapCustomAttributes& SourceAttributes, FHeapCustomAttributes& OutAttributes, float Alpha);
+	UE_DEPRECATED(5.0, "Functionality moved, see UE::Animation::Attributes::InterpolateAttributes")
+	static void InterpolateAttributes(const UE::Anim::FHeapAttributeContainer& SourceAttributes, UE::Anim::FHeapAttributeContainer& OutAttributes, float Alpha) { UE::Anim::Attributes::InterpolateAttributes(OutAttributes, SourceAttributes, Alpha); }
 
+	UE_DEPRECATED(5.0, "Functionality moved, see UE::Animation::Attributes::GetAttributeBlendType")
 	/** Helper functionality to retrieve the correct blend type (from UAnimationSettings) for the provided attribute name */
-	static ECustomAttributeBlendType GetAttributeBlendType(const FName& InName);
+	static ECustomAttributeBlendType GetAttributeBlendType(const FName& InName) { return UE::Anim::Attributes::GetAttributeBlendType(InName); }
 };
