New Animation Attributes system, replacing Custom Attributes:

diff --git a/Engine/Source/Runtime/Engine/Classes/Components/SkeletalMeshComponent.h b/Engine/Source/Runtime/Engine/Classes/Components/SkeletalMeshComponent.h
--- a/Engine/Source/Runtime/Engine/Classes/Components/SkeletalMeshComponent.h
+++ b/Engine/Source/Runtime/Engine/Classes/Components/SkeletalMeshComponent.h
@@ -25,7 +25,7 @@
 #include "ClothingSimulationFactory.h"
 #include "ClothCollisionPrim.h"
 #include "PhysicsEngine/PhysicsAsset.h"
-#include "Animation/CustomAttributesRuntime.h"
+#include "Animation/AttributesRuntime.h"
 
 #include "SkeletalMeshComponent.generated.h"
 
@@ -81,7 +81,7 @@ enum class EClothMassMode : uint8
 	MaxClothMassMode UMETA(Hidden)
 };
 
-/** Method used when retrieving a Custom Attribute value*/
+/** Method used when retrieving a attribute value*/
 UENUM()
 enum class ECustomBoneAttributeLookup : uint8
 {
@@ -130,9 +130,9 @@ struct FAnimationEvaluationContext
 	FBlendedHeapCurve	Curve;
 	FBlendedHeapCurve	CachedCurve;
 
-	// Custom attribute data, swapped in from the component when we are running parallel eval
-	FHeapCustomAttributes CustomAttributes;
-	FHeapCustomAttributes CachedCustomAttributes;
+	// attribute data, swapped in from the component when we are running parallel eval
+	UE::Anim::FHeapAttributeContainer CustomAttributes;
+	UE::Anim::FHeapAttributeContainer CachedCustomAttributes;
 
 	FAnimationEvaluationContext()
 	{
@@ -431,16 +431,15 @@ private:
 	/** Cached Curve result for Update Rate optimization */
 	FBlendedHeapCurve CachedCurve;
 
-	/** Current and cached custom attribute evaluation data, used for Update Rate optimization */
-	FHeapCustomAttributes CachedAttributes;
-	FHeapCustomAttributes CustomAttributes;
+	/** Current and cached attribute evaluation data, used for Update Rate optimization */
+	UE::Anim::FHeapAttributeContainer CachedAttributes;
+	UE::Anim::FHeapAttributeContainer CustomAttributes;
 public:
 	/** 
-	 * Get float type custom attribute value.
+	 * Get float type attribute value.
 
 	 * @param BoneName Name of the bone to retrieve try and retrieve the attribute from
 	 * @param AttributeName Name of the attribute to retrieve
-	 * @param DefaultValue In case the attribute could not be found
      * @param OutValue (reference) Retrieved attribute value if found, otherwise is set to DefaultValue
 	 * @param LookupType Determines how the attribute is retrieved from the specified BoneName (see ECustomBoneAttributeLookup)
 	 * @return Whether or not the atttribute was successfully retrieved
@@ -449,11 +448,22 @@ public:
 	bool GetFloatAttribute_Ref(const FName& BoneName, const FName& AttributeName, UPARAM(ref) float& OutValue, ECustomBoneAttributeLookup LookupType = ECustomBoneAttributeLookup::BoneOnly);
 
 	/** 
-	 * Get integer type custom attribute value.
+	 * Get FTransform type attribute value.
+
+	 * @param BoneName Name of the bone to retrieve try and retrieve the attribute from
+	 * @param AttributeName Name of the attribute to retrieve
+     * @param OutValue (reference) Retrieved attribute value if found, otherwise is set to DefaultValue
+	 * @param LookupType Determines how the attribute is retrieved from the specified BoneName (see ECustomBoneAttributeLookup)
+	 * @return Whether or not the atttribute was successfully retrieved
+	*/
+	UFUNCTION(BlueprintCallable, Category = CustomAttributes)
+	bool GetTransformAttribute_Ref(const FName& BoneName, const FName& AttributeName, UPARAM(ref) FTransform& OutValue, ECustomBoneAttributeLookup LookupType = ECustomBoneAttributeLookup::BoneOnly);
+
+	/** 
+	 * Get integer type attribute value.
 
 	 * @param BoneName Name of the bone to retrieve try and retrieve the attribute from
 	 * @param AttributeName Name of the attribute to retrieve
-	 * @param DefaultValue In case the attribute could not be found
      * @param OutValue (reference) Retrieved attribute value if found, otherwise is set to DefaultValue
 	 * @param LookupType Determines how the attribute is retrieved from the specified BoneName (see ECustomBoneAttributeLookup)
 	 * @return Whether or not the atttribute was successfully retrieved
@@ -462,11 +472,10 @@ public:
 	bool GetIntegerAttribute_Ref(const FName& BoneName, const FName& AttributeName, UPARAM(ref) int32& OutValue, ECustomBoneAttributeLookup LookupType = ECustomBoneAttributeLookup::BoneOnly);
 
 	/** 
-	 * Get string type custom attribute value.
+	 * Get string type attribute value.
 
 	 * @param BoneName Name of the bone to retrieve try and retrieve the attribute from
 	 * @param AttributeName Name of the attribute to retrieve
-	 * @param DefaultValue In case the attribute could not be found
      * @param OutValue (reference) Retrieved attribute value if found, otherwise is set to DefaultValue
 	 * @param LookupType Determines how the attribute is retrieved from the specified BoneName (see ECustomBoneAttributeLookup)
 	 * @return Whether or not the atttribute was successfully retrieved
@@ -474,9 +483,8 @@ public:
 	UFUNCTION(BlueprintCallable, Category = CustomAttributes)
 	bool GetStringAttribute_Ref(const FName& BoneName, const FName& AttributeName, UPARAM(ref) FString& OutValue, ECustomBoneAttributeLookup LookupType = ECustomBoneAttributeLookup::BoneOnly);
 
-
 	/** 
-	 * Get float type custom attribute value.
+	 * Get float type attribute value.
 
 	 * @param BoneName Name of the bone to retrieve try and retrieve the attribute from
 	 * @param AttributeName Name of the attribute to retrieve
@@ -487,9 +495,21 @@ public:
 	*/
 	UFUNCTION(BlueprintCallable, Category = CustomAttributes)
 	bool GetFloatAttribute(const FName& BoneName, const FName& AttributeName, float DefaultValue, float& OutValue, ECustomBoneAttributeLookup LookupType = ECustomBoneAttributeLookup::BoneOnly);
+	
+	/**
+	 * Get FTransform type attribute value.
+
+	 * @param BoneName Name of the bone to retrieve try and retrieve the attribute from
+	 * @param AttributeName Name of the attribute to retrieve
+	 * @param OutValue (reference) Retrieved attribute value if found, otherwise is set to DefaultValue
+	 * @param LookupType Determines how the attribute is retrieved from the specified BoneName (see ECustomBoneAttributeLookup)
+	 * @return Whether or not the atttribute was successfully retrieved
+	*/
+	UFUNCTION(BlueprintCallable, Category = CustomAttributes)
+	bool GetTransformAttribute(const FName& BoneName, const FName& AttributeName, FTransform DefaultValue, FTransform& OutValue, ECustomBoneAttributeLookup LookupType = ECustomBoneAttributeLookup::BoneOnly);
 
 	/** 
-	 * Get integer type custom attribute value.
+	 * Get integer type attribute value.
 
 	 * @param BoneName Name of the bone to retrieve try and retrieve the attribute from
 	 * @param AttributeName Name of the attribute to retrieve
@@ -502,7 +522,7 @@ public:
 	bool GetIntegerAttribute(const FName& BoneName, const FName& AttributeName, int32 DefaultValue, int32& OutValue, ECustomBoneAttributeLookup LookupType = ECustomBoneAttributeLookup::BoneOnly);
 
 	/** 
-	 * Get string type custom attribute value.
+	 * Get string type attribute value.
 
 	 * @param BoneName Name of the bone to retrieve try and retrieve the attribute from
 	 * @param AttributeName Name of the attribute to retrieve
@@ -516,8 +536,8 @@ public:
 
 protected:
 	/** Templated version to try and retrieve a typed bone attribute's value */
-	template<typename DataType>
-	bool GetBoneAttribute(const FName& BoneName, const FName& AttributeName, DataType DefaultValue, DataType& OutValue, ECustomBoneAttributeLookup LookupType);	
+	template<typename DataType, typename CustomAttributeType>
+	bool FindAttributeChecked(const FName& BoneName, const FName& AttributeName, DataType DefaultValue, DataType& OutValue, ECustomBoneAttributeLookup LookupType);	
 
 public:
 	/** Used to scale speed of all animations on this skeletal mesh. */
@@ -1896,12 +1916,12 @@ public:
 	* @param	OutCurves				Blended Curve
 	*/
 #if WITH_EDITOR
-	void PerformAnimationEvaluation(const USkeletalMesh* InSkeletalMesh, UAnimInstance* InAnimInstance, TArray<FTransform>& OutSpaceBases, TArray<FTransform>& OutBoneSpaceTransforms, FVector& OutRootBoneTranslation, FBlendedHeapCurve& OutCurve, FHeapCustomAttributes& OutAttributes);
+	void PerformAnimationEvaluation(const USkeletalMesh* InSkeletalMesh, UAnimInstance* InAnimInstance, TArray<FTransform>& OutSpaceBases, TArray<FTransform>& OutBoneSpaceTransforms, FVector& OutRootBoneTranslation, FBlendedHeapCurve& OutCurve, UE::Anim::FHeapAttributeContainer& OutAttributes);
 
 	UE_DEPRECATED(4.26, "Please use PerformAnimationEvaluation with different signature")
 	void PerformAnimationEvaluation(const USkeletalMesh* InSkeletalMesh, UAnimInstance* InAnimInstance, TArray<FTransform>& OutSpaceBases, TArray<FTransform>& OutBoneSpaceTransforms, FVector& OutRootBoneTranslation, FBlendedHeapCurve& OutCurve);
 #endif
-	void PerformAnimationProcessing(const USkeletalMesh* InSkeletalMesh, UAnimInstance* InAnimInstance, bool bInDoEvaluation, TArray<FTransform>& OutSpaceBases, TArray<FTransform>& OutBoneSpaceTransforms, FVector& OutRootBoneTranslation, FBlendedHeapCurve& OutCurve, FHeapCustomAttributes& OutAttributes);
+	void PerformAnimationProcessing(const USkeletalMesh* InSkeletalMesh, UAnimInstance* InAnimInstance, bool bInDoEvaluation, TArray<FTransform>& OutSpaceBases, TArray<FTransform>& OutBoneSpaceTransforms, FVector& OutRootBoneTranslation, FBlendedHeapCurve& OutCurve, UE::Anim::FHeapAttributeContainer& OutAttributes);
 
 	UE_DEPRECATED(4.26, "Please use PerformAnimationEvaluation with different signature")
 	void PerformAnimationProcessing(const USkeletalMesh* InSkeletalMesh, UAnimInstance* InAnimInstance, bool bInDoEvaluation, TArray<FTransform>& OutSpaceBases, TArray<FTransform>& OutBoneSpaceTransforms, FVector& OutRootBoneTranslation, FBlendedHeapCurve& OutCurve);
@@ -1909,7 +1929,7 @@ public:
 	/**
 	 * Evaluates the post process instance from the skeletal mesh this component is using.
 	 */
-	void EvaluatePostProcessMeshInstance(TArray<FTransform>& OutBoneSpaceTransforms, FCompactPose& InOutPose, FBlendedHeapCurve& OutCurve, const USkeletalMesh* InSkeletalMesh, FVector& OutRootBoneTranslation, FHeapCustomAttributes& OutAttributes) const;
+	void EvaluatePostProcessMeshInstance(TArray<FTransform>& OutBoneSpaceTransforms, FCompactPose& InOutPose, FBlendedHeapCurve& OutCurve, const USkeletalMesh* InSkeletalMesh, FVector& OutRootBoneTranslation, UE::Anim::FHeapAttributeContainer& OutAttributes) const;
 
 	UE_DEPRECATED(4.26, "Please use EvaluatePostProcessMeshInstance with different signature")
 	void EvaluatePostProcessMeshInstance(TArray<FTransform>& OutBoneSpaceTransforms, FCompactPose& InOutPose, FBlendedHeapCurve& OutCurve, const USkeletalMesh* InSkeletalMesh, FVector& OutRootBoneTranslation) const;
@@ -2261,7 +2281,7 @@ private:
 	void EndPhysicsTickComponent(FSkeletalMeshComponentEndPhysicsTickFunction& ThisTickFunction);
 
 	/** Evaluate Anim System **/
-	void EvaluateAnimation(const USkeletalMesh* InSkeletalMesh, UAnimInstance* InAnimInstance, FVector& OutRootBoneTranslation, FBlendedHeapCurve& OutCurve, FCompactPose& OutPose, FHeapCustomAttributes& OutAttributes) const;
+	void EvaluateAnimation(const USkeletalMesh* InSkeletalMesh, UAnimInstance* InAnimInstance, FVector& OutRootBoneTranslation, FBlendedHeapCurve& OutCurve, FCompactPose& OutPose, UE::Anim::FHeapAttributeContainer& OutAttributes) const;
 
 	/** Queues up tasks for parallel update/evaluation, as well as the chained game thread completion task */
 	void DispatchParallelEvaluationTasks(FActorComponentTickFunction* TickFunction);
@@ -2444,13 +2464,13 @@ public:
 
 
 private:
-	/** Temporary array of custom attributes that are active on this component - keeps same buffer index as SpaceBases - Please check SkinnedMeshComponent*/
-	FHeapCustomAttributes AttributesArray[2];
+	/** Temporary array of attributes that are active on this component - keeps same buffer index as SpaceBases - Please check SkinnedMeshComponent*/
+	UE::Anim::FHeapAttributeContainer AttributesArray[2];
 
-	FHeapCustomAttributes& GetEditableCustomAttributes() { return AttributesArray[CurrentEditableComponentTransforms]; }
+	UE::Anim::FHeapAttributeContainer& GetEditableCustomAttributes() { return AttributesArray[CurrentEditableComponentTransforms]; }
 
 public:
-	const FHeapCustomAttributes& GetCustomAttributes() const { return AttributesArray[CurrentReadComponentTransforms]; }
+	const UE::Anim::FHeapAttributeContainer& GetCustomAttributes() const { return AttributesArray[CurrentReadComponentTransforms]; }
 public:
 	/** Skeletal mesh component should not be able to have its mobility set to static */
 	virtual const bool CanHaveStaticMobility() const override { return false; }
diff --git a/Engine/Source/Runtime/Engine/Private/Animation/AnimCompositeBase.cpp b/Engine/Source/Runtime/Engine/Private/Animation/AnimCompositeBase.cpp
--- a/Engine/Source/Runtime/Engine/Private/Animation/AnimCompositeBase.cpp
+++ b/Engine/Source/Runtime/Engine/Private/Animation/AnimCompositeBase.cpp
@@ -10,7 +10,7 @@
 #include "BonePose.h"
 #include "AnimationRuntime.h"
 #include "Animation/AnimationPoseData.h"
-#include "Animation/CustomAttributesRuntime.h"
+#include "Animation/AttributesRuntime.h"
 
 ///////////////////////////////////////////////////////
 // FAnimSegment
@@ -545,7 +545,7 @@ void FAnimTrack::SortAnimSegments()
 
 void FAnimTrack::GetAnimationPose(/*out*/ FCompactPose& OutPose, /*out*/ FBlendedCurve& OutCurve, const FAnimExtractContext& ExtractionContext) const
 {
-	FStackCustomAttributes TempAttributes;
+	UE::Anim::FStackAttributeContainer TempAttributes;
 	FAnimationPoseData OutAnimationPoseData(OutPose, OutCurve, TempAttributes);
 	GetAnimationPose(OutAnimationPoseData, ExtractionContext);
 }
diff --git a/Engine/Source/Runtime/Engine/Private/Animation/AnimDataModel.cpp b/Engine/Source/Runtime/Engine/Private/Animation/AnimDataModel.cpp
--- a/Engine/Source/Runtime/Engine/Private/Animation/AnimDataModel.cpp
+++ b/Engine/Source/Runtime/Engine/Private/Animation/AnimDataModel.cpp
@@ -6,6 +6,7 @@
 #include "Animation/AnimSequence.h"
 
 #include "Algo/Transform.h"
+#include "Algo/Accumulate.h"
 #include "Animation/SmartName.h"
 
 
@@ -235,6 +236,56 @@ const FRichCurve& UAnimDataModel::GetRichCurve(const FAnimationCurveIdentifier&
 	return *CurvePtr;
 }
 
+TArrayView<const FAnimatedBoneAttribute> UAnimDataModel::GetAttributes() const
+{
+	return AnimatedBoneAttributes;
+}
+
+int32 UAnimDataModel::GetNumberOfAttributes() const
+{
+	return AnimatedBoneAttributes.Num();
+}
+
+int32 UAnimDataModel::GetNumberOfAttributesForBoneIndex(const int32 BoneIndex) const
+{
+	// Sum up total number of attributes with provided bone index
+	const int32 NumberOfBoneAttributes = Algo::Accumulate<int32>(AnimatedBoneAttributes, 0, [BoneIndex](int32 Sum, const FAnimatedBoneAttribute& Attribute) -> int32
+	{
+		Sum += Attribute.Identifier.GetBoneIndex() == BoneIndex ? 1 : 0;
+		return Sum;
+	});
+	return NumberOfBoneAttributes;
+}
+
+void UAnimDataModel::GetAttributesForBone(const FName& BoneName, TArray<const FAnimatedBoneAttribute*>& OutBoneAttributes) const
+{
+	// Sum up total number of attributes with provided bone name
+	Algo::TransformIf(AnimatedBoneAttributes, OutBoneAttributes, [BoneName](const FAnimatedBoneAttribute& Attribute) -> bool
+	{
+		return Attribute.Identifier.GetBoneName() == BoneName;
+	},
+	[](const FAnimatedBoneAttribute& Attribute) 
+	{
+		return &Attribute;
+	});
+}
+
+const FAnimatedBoneAttribute& UAnimDataModel::GetAttribute(const FAnimationAttributeIdentifier& AttributeIdentifier) const
+{
+	const FAnimatedBoneAttribute* AttributePtr = FindAttribute(AttributeIdentifier);
+	checkf(AttributePtr, TEXT("Unable to find attribute for provided identifier"));
+
+	return *AttributePtr;
+}
+
+const FAnimatedBoneAttribute* UAnimDataModel::FindAttribute(const FAnimationAttributeIdentifier& AttributeIdentifier) const
+{
+	return AnimatedBoneAttributes.FindByPredicate([AttributeIdentifier](const FAnimatedBoneAttribute& Attribute)
+	{
+		return Attribute.Identifier == AttributeIdentifier;
+	});
+}
+
 const FAnimCurveBase* UAnimDataModel::FindCurve(const FAnimationCurveIdentifier& CurveIdentifier) const
 {
 	switch (CurveIdentifier.CurveType)
@@ -277,6 +328,12 @@ FGuid UAnimDataModel::GenerateGuid() const
 		UpdateWithData(Curve.FloatCurve.PreInfinityExtrap);
 		UpdateWithData(Curve.FloatCurve.PostInfinityExtrap);
 	}
+	
+	for (const FAnimatedBoneAttribute& Attribute : AnimatedBoneAttributes)
+	{
+		UpdateWithData(Attribute.Identifier);
+		UpdateSHAWithArray(Attribute.Curve.GetConstRefOfKeys());
+	}
 
 	Sha.Final();
 
diff --git a/Engine/Source/Runtime/Engine/Private/Animation/AnimInstance.cpp b/Engine/Source/Runtime/Engine/Private/Animation/AnimInstance.cpp
--- a/Engine/Source/Runtime/Engine/Private/Animation/AnimInstance.cpp
+++ b/Engine/Source/Runtime/Engine/Private/Animation/AnimInstance.cpp
@@ -673,7 +673,7 @@ bool UAnimInstance::ParallelCanEvaluate(const USkeletalMesh* InSkeletalMesh) con
 
 void UAnimInstance::ParallelEvaluateAnimation(bool bForceRefPose, const USkeletalMesh* InSkeletalMesh, FBlendedHeapCurve& OutCurve, FCompactPose& OutPose)
 {
-	FHeapCustomAttributes Attributes;
+	UE::Anim::FHeapAttributeContainer Attributes;
 	FParallelEvaluationData EvalData = { OutCurve, OutPose, Attributes };
 	ParallelEvaluateAnimation(bForceRefPose, InSkeletalMesh, EvalData);
 }
diff --git a/Engine/Source/Runtime/Engine/Private/Animation/AnimInstanceProxy.cpp b/Engine/Source/Runtime/Engine/Private/Animation/AnimInstanceProxy.cpp
--- a/Engine/Source/Runtime/Engine/Private/Animation/AnimInstanceProxy.cpp
+++ b/Engine/Source/Runtime/Engine/Private/Animation/AnimInstanceProxy.cpp
@@ -1248,7 +1248,7 @@ void FAnimInstanceProxy::EvaluateAnimationNode_WithRoot(FPoseContext& Output, FA
 
 void FAnimInstanceProxy::SlotEvaluatePose(const FName& SlotNodeName, const FCompactPose& SourcePose, const FBlendedCurve& SourceCurve, float InSourceWeight, FCompactPose& BlendedPose, FBlendedCurve& BlendedCurve, float InBlendWeight, float InTotalNodeWeight)
 {
-	FStackCustomAttributes TempAttributes;
+	UE::Anim::FStackAttributeContainer TempAttributes;
 
 	const FAnimationPoseData SourceAnimationPoseData(*const_cast<FCompactPose*>(&SourcePose), *const_cast<FBlendedCurve*>(&SourceCurve), TempAttributes);
 	FAnimationPoseData BlendedAnimationPoseData(BlendedPose, BlendedCurve, TempAttributes);
@@ -1261,11 +1261,11 @@ void FAnimInstanceProxy::SlotEvaluatePoseWithBlendProfiles(const FName& SlotNode
 {
 	const FCompactPose& SourcePose = SourceAnimationPoseData.GetPose();
 	const FBlendedCurve& SourceCurve = SourceAnimationPoseData.GetCurve();
-	const FStackCustomAttributes& SourceAttributes = SourceAnimationPoseData.GetAttributes();
+	const UE::Anim::FStackAttributeContainer& SourceAttributes = SourceAnimationPoseData.GetAttributes();
 
 	FCompactPose& BlendedPose = OutBlendedAnimationPoseData.GetPose();
 	FBlendedCurve& BlendedCurve = OutBlendedAnimationPoseData.GetCurve();
-	FStackCustomAttributes& BlendedAttributes = OutBlendedAnimationPoseData.GetAttributes();
+	UE::Anim::FStackAttributeContainer& BlendedAttributes = OutBlendedAnimationPoseData.GetAttributes();
 
 	if (InBlendWeight <= ZERO_ANIMWEIGHT_THRESH)
 	{
@@ -1304,7 +1304,7 @@ void FAnimInstanceProxy::SlotEvaluatePoseWithBlendProfiles(const FName& SlotNode
 	TArray<float, TInlineAllocator<8>>& BlendingWeights = BlendProfileScratchData.BlendingWeights;
 	TArray<const FCompactPose*, TInlineAllocator<8>>& BlendingPoses = BlendProfileScratchData.BlendingPoses;
 	TArray<const FBlendedCurve*, TInlineAllocator<8>>& BlendingCurves = BlendProfileScratchData.BlendingCurves;
-	TArray<const FStackCustomAttributes*, TInlineAllocator<8>>& BlendingAttributes = BlendProfileScratchData.BlendingAttributes;
+	TArray<const UE::Anim::FStackAttributeContainer*, TInlineAllocator<8>>& BlendingAttributes = BlendProfileScratchData.BlendingAttributes;
 	Poses.Reset();
 	AdditivePoses.Reset();
 	PoseIndices.Reset();
@@ -1581,7 +1581,7 @@ void FAnimInstanceProxy::SlotEvaluatePoseWithBlendProfiles(const FName& SlotNode
 
 	if (BlendingAttributes.Num() > 0)
 	{
-		FCustomAttributesRuntime::BlendAttributes(BlendingAttributes, BlendingWeights, OutBlendedAnimationPoseData.GetAttributes());
+		UE::Anim::Attributes::BlendAttributes(BlendingAttributes, BlendingWeights, OutBlendedAnimationPoseData.GetAttributes());
 	}
 
 	// Additives.
@@ -1590,7 +1590,7 @@ void FAnimInstanceProxy::SlotEvaluatePoseWithBlendProfiles(const FName& SlotNode
 		FSlotEvaluationPose& AdditivePose = Poses[PoseIndex];
 		const FAnimationPoseData AdditiveAnimationPoseData(AdditivePose);
 		OutBlendedAnimationPoseData.GetCurve().Accumulate(AdditiveAnimationPoseData.GetCurve(), AdditivePose.Weight);
-		FCustomAttributesRuntime::AccumulateAttributes(AdditiveAnimationPoseData.GetAttributes(), OutBlendedAnimationPoseData.GetAttributes(), AdditivePose.Weight);
+		UE::Anim::Attributes::AccumulateAttributes(AdditiveAnimationPoseData.GetAttributes(), OutBlendedAnimationPoseData.GetAttributes(), AdditivePose.Weight, AdditivePoses[PoseIndex].AdditiveType);
 	}
 }
 
@@ -1598,11 +1598,11 @@ void FAnimInstanceProxy::SlotEvaluatePose(const FName& SlotNodeName, const FAnim
 {
 	const FCompactPose& SourcePose = SourceAnimationPoseData.GetPose();
 	const FBlendedCurve& SourceCurve = SourceAnimationPoseData.GetCurve();
-	const FStackCustomAttributes& SourceAttributes = SourceAnimationPoseData.GetAttributes();
+	const UE::Anim::FStackAttributeContainer& SourceAttributes = SourceAnimationPoseData.GetAttributes();
 	
 	FCompactPose& BlendedPose = OutBlendedAnimationPoseData.GetPose();
 	FBlendedCurve& BlendedCurve = OutBlendedAnimationPoseData.GetCurve();
-	FStackCustomAttributes& BlendedAttributes = OutBlendedAnimationPoseData.GetAttributes();
+	UE::Anim::FStackAttributeContainer& BlendedAttributes = OutBlendedAnimationPoseData.GetAttributes();
 
 	// Accessing MontageInstances from this function is not safe (as this can be called during Parallel Anim Evaluation!
 	// Any montage data you need to add should be part of MontageEvaluationData
@@ -1632,7 +1632,7 @@ void FAnimInstanceProxy::SlotEvaluatePose(const FName& SlotNodeName, const FAnim
 	TArray<float, TInlineAllocator<8>>& BlendingWeights = BlendProfileScratchData.BlendingWeights;
 	TArray<const FCompactPose*, TInlineAllocator<8>>& BlendingPoses = BlendProfileScratchData.BlendingPoses;
 	TArray<const FBlendedCurve*, TInlineAllocator<8>>& BlendingCurves = BlendProfileScratchData.BlendingCurves;
-	TArray<const FStackCustomAttributes*, TInlineAllocator<8>>& BlendingAttributes = BlendProfileScratchData.BlendingAttributes;
+	TArray<const UE::Anim::FStackAttributeContainer*, TInlineAllocator<8>>& BlendingAttributes = BlendProfileScratchData.BlendingAttributes;
 	NonAdditivePoses.Reset();
 	AdditivePoses.Reset();
 	BlendingWeights.Reset();
diff --git a/Engine/Source/Runtime/Engine/Private/Animation/AnimNode_StateMachine.cpp b/Engine/Source/Runtime/Engine/Private/Animation/AnimNode_StateMachine.cpp
--- a/Engine/Source/Runtime/Engine/Private/Animation/AnimNode_StateMachine.cpp
+++ b/Engine/Source/Runtime/Engine/Private/Animation/AnimNode_StateMachine.cpp
@@ -891,8 +891,8 @@ void FAnimNode_StateMachine::EvaluateTransitionStandardBlendInternal(FPoseContex
 	Output.Curve.Override(PreviousStateResult.Curve, 1.0 - Transition.Alpha);
 	Output.Curve.Accumulate(NextStateResult.Curve, Transition.Alpha);
 
-	FCustomAttributesRuntime::OverrideAttributes(PreviousStateResult.CustomAttributes, Output.CustomAttributes, 1.0 - Transition.Alpha);
-	FCustomAttributesRuntime::AccumulateAttributes(NextStateResult.CustomAttributes, Output.CustomAttributes, Transition.Alpha);
+	UE::Anim::Attributes::OverrideAttributes(PreviousStateResult.CustomAttributes, Output.CustomAttributes, 1.0 - Transition.Alpha);
+	UE::Anim::Attributes::AccumulateAttributes(NextStateResult.CustomAttributes, Output.CustomAttributes, Transition.Alpha, AAT_None);
 }
 
 void FAnimNode_StateMachine::EvaluateTransitionCustomBlend(FPoseContext& Output, FAnimationActiveTransitionEntry& Transition, bool bIntermediatePoseIsValid)
diff --git a/Engine/Source/Runtime/Engine/Private/Animation/AnimSequence.cpp b/Engine/Source/Runtime/Engine/Private/Animation/AnimSequence.cpp
--- a/Engine/Source/Runtime/Engine/Private/Animation/AnimSequence.cpp
+++ b/Engine/Source/Runtime/Engine/Private/Animation/AnimSequence.cpp
@@ -46,7 +46,7 @@
 #include "Animation/AnimStreamable.h"
 #include "Modules/ModuleManager.h"
 #include "ProfilingDebugging/CookStats.h"
-#include "Animation/CustomAttributesRuntime.h"
+#include "Animation/AttributesRuntime.h"
 #include "Stats/StatsHierarchical.h"
 #include "Animation/AnimationPoseData.h"
 #include "ITimeManagementModule.h"
@@ -57,6 +57,7 @@
 #if WITH_EDITOR
 #include "Animation/AnimData/AnimDataController.h"
 #include "Animation/AnimData/AnimDataModel.h"
+#include "Animation/BuiltInAttributeTypes.h"
 #endif // WITH_EDITOR
 
 #include "Animation/AnimSequenceHelpers.h"
@@ -434,9 +435,6 @@ UAnimSequence::UAnimSequence(const FObjectInitializer& ObjectInitializer)
 	bAllowFrameStripping = true;
 	CompressionErrorThresholdScale = 1.f;
 
-	CustomAttributesGuid.Invalidate();
-	BakedCustomAttributesGuid.Invalidate();
-
 	if (!HasAllFlags(EObjectFlags::RF_ClassDefaultObject))
 	{
 		check(DataModel);
@@ -960,6 +958,14 @@ void UAnimSequence::PostLoad()
 		}
 #endif
 	}
+
+#if WITH_EDITOR
+	if (GetLinkerCustomVersion(FUE5MainStreamObjectVersion::GUID) < FUE5MainStreamObjectVersion::MoveCustomAttributesToDataModel)
+	{
+		ValidateModel();
+		MoveAttributesToModel();
+	}
+#endif // WITH_EDITOR
 }
 
 #if WITH_EDITOR
@@ -1410,7 +1416,7 @@ void UAnimSequence::GetAnimationPose(FAnimationPoseData& OutAnimationPoseData, c
 
 void UAnimSequence::GetBonePose(FCompactPose& OutPose, FBlendedCurve& OutCurve, const FAnimExtractContext& ExtractionContext, bool bForceUseRawData) const
 {
-	FStackCustomAttributes TempAttributes;
+	UE::Anim::FStackAttributeContainer TempAttributes;
 	FAnimationPoseData OutAnimationPoseData(OutPose, OutCurve, TempAttributes);
 	GetBonePose(OutAnimationPoseData, ExtractionContext, bForceUseRawData);
 }
@@ -1520,7 +1526,7 @@ void UAnimSequence::GetBonePose(FAnimationPoseData& OutAnimationPoseData, const
 			RootMotionReset.ResetRootBoneForRootMotion(OutPose[FCompactPoseBoneIndex(0)], RequiredBones);
 		}
 
-		GetCustomAttributes(OutAnimationPoseData, ExtractionContext, true);
+		EvaluateAttributes(OutAnimationPoseData, ExtractionContext, true);
 
 		return;
 	}
@@ -1528,7 +1534,7 @@ void UAnimSequence::GetBonePose(FAnimationPoseData& OutAnimationPoseData, const
 
 	DecompressPose(OutPose, CompressedData, ExtractionContext, GetSkeleton(), GetPlayLength(), Interpolation, bIsBakedAdditive, GetRetargetTransforms(), GetRetargetTransformsSourceName(), RootMotionReset);
 
-	GetCustomAttributes(OutAnimationPoseData, ExtractionContext, false);
+	EvaluateAttributes(OutAnimationPoseData, ExtractionContext, false);
 }
 
 const TArray<FRawAnimSequenceTrack>& UAnimSequence::GetRawAnimationData() const
@@ -1612,7 +1618,7 @@ int32 UAnimSequence::AddNewRawTrack(FName TrackName, FRawAnimSequenceTrack* Trac
 
 void UAnimSequence::GetBonePose_Additive(FCompactPose& OutPose, FBlendedCurve& OutCurve, const FAnimExtractContext& ExtractionContext) const
 {
-	FStackCustomAttributes TempAttributes;
+	UE::Anim::FStackAttributeContainer TempAttributes;
 	FAnimationPoseData OutAnimationPoseData(OutPose, OutCurve, TempAttributes);
 
 	GetBonePose_Additive(OutAnimationPoseData, ExtractionContext);
@@ -1622,7 +1628,7 @@ void UAnimSequence::GetBonePose_Additive(FAnimationPoseData& OutAnimationPoseDat
 {
 	FCompactPose& OutPose = OutAnimationPoseData.GetPose();
 	FBlendedCurve& OutCurve = OutAnimationPoseData.GetCurve();
-	FStackCustomAttributes& OutAttributes = OutAnimationPoseData.GetAttributes();
+	UE::Anim::FStackAttributeContainer& OutAttributes = OutAnimationPoseData.GetAttributes();
 
 	if (!IsValidAdditive())
 	{
@@ -1636,7 +1642,7 @@ void UAnimSequence::GetBonePose_Additive(FAnimationPoseData& OutAnimationPoseDat
 	// Extract base pose
 	FCompactPose BasePose;
 	FBlendedCurve BaseCurve;
-	FStackCustomAttributes BaseAttributes;
+	UE::Anim::FStackAttributeContainer BaseAttributes;
 	
 	BasePose.SetBoneContainer(&OutPose.GetBoneContainer());
 	BaseCurve.InitFrom(OutCurve);	
@@ -1649,12 +1655,12 @@ void UAnimSequence::GetBonePose_Additive(FAnimationPoseData& OutAnimationPoseDat
 	FAnimationRuntime::ConvertPoseToAdditive(OutPose, BasePose);
 	OutCurve.ConvertToAdditive(BaseCurve);
 
-	FCustomAttributesRuntime::SubtractAttributes(BaseAttributes, OutAttributes);
+	UE::Anim::Attributes::ConvertToAdditive(BaseAttributes, OutAttributes);
 }
 
 void UAnimSequence::GetAdditiveBasePose(FCompactPose& OutPose, FBlendedCurve& OutCurve, const FAnimExtractContext& ExtractionContext) const
 {
-	FStackCustomAttributes TempAttributes;
+	UE::Anim::FStackAttributeContainer TempAttributes;
 	FAnimationPoseData OutAnimationPoseData(OutPose, OutCurve, TempAttributes);
 
 	GetAdditiveBasePose(OutAnimationPoseData, ExtractionContext);
@@ -1698,7 +1704,7 @@ void UAnimSequence::GetAdditiveBasePose(FAnimationPoseData& OutAnimationPoseData
 
 void UAnimSequence::GetBonePose_AdditiveMeshRotationOnly(FCompactPose& OutPose, FBlendedCurve& OutCurve, const FAnimExtractContext& ExtractionContext) const
 {
-	FStackCustomAttributes TempAttributes;
+	UE::Anim::FStackAttributeContainer TempAttributes;
 	FAnimationPoseData OutAnimationPoseData(OutPose, OutCurve, TempAttributes);
 
 	GetBonePose_AdditiveMeshRotationOnly(OutAnimationPoseData, ExtractionContext);
@@ -1708,7 +1714,7 @@ void UAnimSequence::GetBonePose_AdditiveMeshRotationOnly(FAnimationPoseData& Out
 {
 	FCompactPose& OutPose = OutAnimationPoseData.GetPose();
 	FBlendedCurve& OutCurve = OutAnimationPoseData.GetCurve();
-	FStackCustomAttributes& OutAttributes = OutAnimationPoseData.GetAttributes();
+	UE::Anim::FStackAttributeContainer& OutAttributes = OutAnimationPoseData.GetAttributes();
 
 	if (!IsValidAdditive())
 	{
@@ -1723,7 +1729,7 @@ void UAnimSequence::GetBonePose_AdditiveMeshRotationOnly(FAnimationPoseData& Out
 	// get base pose
 	FCompactPose BasePose;
 	FBlendedCurve BaseCurve;
-	FStackCustomAttributes BaseAttributes;
+	UE::Anim::FStackAttributeContainer BaseAttributes;
 
 	FAnimationPoseData BasePoseData(BasePose, BaseCurve, BaseAttributes);
 
@@ -1740,7 +1746,7 @@ void UAnimSequence::GetBonePose_AdditiveMeshRotationOnly(FAnimationPoseData& Out
 	FAnimationRuntime::ConvertPoseToAdditive(OutPose, BasePose);
 	OutCurve.ConvertToAdditive(BaseCurve);
 
-	FCustomAttributesRuntime::SubtractAttributes(BaseAttributes, OutAttributes);
+	UE::Anim::Attributes::ConvertToAdditive(BaseAttributes, OutAttributes);
 }
 
 #if WITH_EDITORONLY_DATA
@@ -2106,7 +2112,7 @@ void UAnimSequence::ApplyCompressedData(const TArray<uint8>& Data)
 #if WITH_EDITOR
 	bCompressionInProgress = false;
 	
-	SynchronousCustomAttributesCompression();
+	SynchronousAnimatedBoneAttributesCompression();
 #endif
 	if(Data.Num() > 0)
 	{
@@ -2346,7 +2352,7 @@ void UAnimSequence::BakeOutVirtualBoneTracks(TArray<FRawAnimSequenceTrack>& NewR
 		const float CurrentFrameTime = Frame * EvalContext.IntervalTime;
 		ExtractContext.CurrentTime = CurrentFrameTime;
 
-		FStackCustomAttributes TempAttributes;
+		UE::Anim::FStackAttributeContainer TempAttributes;
 		FAnimationPoseData AnimPoseData(Pose, Curve, TempAttributes);
 		GetAnimationPose(AnimPoseData, ExtractContext);
 
@@ -2405,7 +2411,7 @@ void UAnimSequence::TestEvalauteAnimation() const
 		const float CurrentFrameTime = Frame * EvalContext.IntervalTime;
 		ExtractContext.CurrentTime = CurrentFrameTime;
 
-		FStackCustomAttributes TempAttributes;
+		UE::Anim::FStackAttributeContainer TempAttributes;
 		FAnimationPoseData AnimPoseData(Pose, Curve, TempAttributes);
 		GetAnimationPose(AnimPoseData, ExtractContext);
 	}
@@ -2515,11 +2521,11 @@ void UAnimSequence::BakeOutAdditiveIntoRawData(TArray<FRawAnimSequenceTrack>& Ne
 		const float CurrentFrameTime = Frame * EvalContext.IntervalTime;
 		ExtractContext.CurrentTime = CurrentFrameTime;
 
-		FStackCustomAttributes BaseAttributes;
+		UE::Anim::FStackAttributeContainer BaseAttributes;
 		FAnimationPoseData AnimPoseData(Pose, Curve, BaseAttributes);
 		GetAnimationPose(AnimPoseData, ExtractContext);
 
-		FStackCustomAttributes AdditiveAttributes;
+		UE::Anim::FStackAttributeContainer AdditiveAttributes;
 		FAnimationPoseData AnimBasePoseData(BasePose, DummyBaseCurve, AdditiveAttributes);
 		GetAdditiveBasePose(AnimBasePoseData, ExtractContext);
 
@@ -4583,52 +4589,50 @@ bool UAnimSequence::UseRawDataForPoseExtraction(const FBoneContainer& RequiredBo
 	return bUseRawDataOnly || (GetSkeletonVirtualBoneGuid() != GetSkeleton()->GetVirtualBoneGuid()) || RequiredBones.GetDisableRetargeting() || RequiredBones.ShouldUseRawData() || RequiredBones.ShouldUseSourceData();
 }
 
-void UAnimSequence::GetCustomAttributes(FAnimationPoseData& OutAnimationPoseData, const FAnimExtractContext& ExtractionContext, bool bUseRawData) const
+#if WITH_EDITOR
+void UAnimSequence::AddBoneFloatCustomAttribute(const FName& BoneName, const FName& AttributeName, const TArray<float>& TimeKeys, const TArray<float>& ValueKeys)
+{
+	UE::Anim::AddTypedCustomAttribute<FFloatAnimationAttribute, float>(AttributeName, BoneName, this, MakeArrayView(TimeKeys), MakeArrayView(ValueKeys));
+}
+
+void UAnimSequence::AddBoneIntegerCustomAttribute(const FName& BoneName, const FName& AttributeName, const TArray<float>& TimeKeys, const TArray<int32>& ValueKeys)
 {
-	QUICK_SCOPE_CYCLE_COUNTER(STAT_GetCustomAttributes);
+	UE::Anim::AddTypedCustomAttribute<FIntegerAnimationAttribute, int32>(AttributeName, BoneName, this, MakeArrayView(TimeKeys), MakeArrayView(ValueKeys));
+}
+
+void UAnimSequence::AddBoneStringCustomAttribute(const FName& BoneName, const FName& AttributeName, const TArray<float>& TimeKeys, const TArray<FString>& ValueKeys)
+{
+	UE::Anim::AddTypedCustomAttribute<FStringAnimationAttribute, FString>(AttributeName, BoneName, this, MakeArrayView(TimeKeys), MakeArrayView(ValueKeys));
+}
+#endif // WITH_EDITOR
+
+void UAnimSequence::EvaluateAttributes(FAnimationPoseData& OutAnimationPoseData, const FAnimExtractContext& ExtractionContext, bool bUseRawData) const
+{
+	QUICK_SCOPE_CYCLE_COUNTER(STAT_EvaluateAttributes);
 
 	const FBoneContainer& RequiredBones = OutAnimationPoseData.GetPose().GetBoneContainer();
-	FStackCustomAttributes& OutAttributes = OutAnimationPoseData.GetAttributes();
+	UE::Anim::FStackAttributeContainer& OutAttributes = OutAnimationPoseData.GetAttributes();
 
 #if WITH_EDITOR
 	if (bUseRawData)
 	{
-		for (const FCustomAttributePerBoneData& BoneAttributes : PerBoneCustomAttributeData)
+		ValidateModel();
+		for (const FAnimatedBoneAttribute& Attribute : DataModel->GetAttributes())
 		{
-			const FCompactPoseBoneIndex PoseBoneIndex = RequiredBones.GetCompactPoseIndexFromSkeletonIndex(BoneAttributes.BoneTreeIndex);
-
-			for (const FCustomAttribute& Attribute : BoneAttributes.Attributes)
-			{
-				FCustomAttributesRuntime::GetAttributeValue(OutAttributes, PoseBoneIndex, Attribute, ExtractionContext);
-			}
+			const FCompactPoseBoneIndex PoseBoneIndex = RequiredBones.GetCompactPoseIndexFromSkeletonIndex(Attribute.Identifier.GetBoneIndex());
+			UE::Anim::Attributes::GetAttributeValue(OutAttributes, PoseBoneIndex, Attribute, ExtractionContext);
 		}
 	}
 	else
 #endif // WITH_EDITOR
 	{
-		for (const FBakedCustomAttributePerBoneData& BakedBoneAttributes : BakedPerBoneCustomAttributeData)
+		for (const TPair<FAnimationAttributeIdentifier, FAttributeCurve>& BakedAttribute : AttributeCurves)
 		{
-			const FCompactPoseBoneIndex PoseBoneIndex = RequiredBones.GetCompactPoseIndexFromSkeletonIndex(BakedBoneAttributes.BoneTreeIndex);
-			for (const FBakedFloatCustomAttribute& Attribute : BakedBoneAttributes.FloatAttributes)
-			{
-				const ECustomAttributeBlendType BlendType = FCustomAttributesRuntime::GetAttributeBlendType(Attribute.AttributeName);
-				const float Value = Attribute.FloatCurve.Eval(ExtractionContext.CurrentTime);
-				OutAttributes.AddBoneAttribute<float>(PoseBoneIndex, Attribute.AttributeName, BlendType, Value);
-			}
-
-			for (const FBakedIntegerCustomAttribute& Attribute : BakedBoneAttributes.IntAttributes)
-			{
-				const ECustomAttributeBlendType BlendType = FCustomAttributesRuntime::GetAttributeBlendType(Attribute.AttributeName);
-				const int32 Value = Attribute.IntCurve.Evaluate(ExtractionContext.CurrentTime);
-				OutAttributes.AddBoneAttribute<int32>(PoseBoneIndex, Attribute.AttributeName, BlendType, Value);
-			}
+			const FCompactPoseBoneIndex PoseBoneIndex = RequiredBones.GetCompactPoseIndexFromSkeletonIndex(BakedAttribute.Key.GetBoneIndex());
+			UE::Anim::FAttributeId Info(BakedAttribute.Key.GetName(), PoseBoneIndex);
+			uint8* AttributePtr = OutAttributes.FindOrAdd(BakedAttribute.Key.GetType(), Info);
 
-			for (const FBakedStringCustomAttribute& Attribute : BakedBoneAttributes.StringAttributes)
-			{
-				static const FString DefaultValue = TEXT("");
-				const FString Value = Attribute.StringCurve.Eval(ExtractionContext.CurrentTime, DefaultValue);
-				OutAttributes.AddBoneAttribute<FString>(PoseBoneIndex, Attribute.AttributeName, ECustomAttributeBlendType::Override, Value);
-			}
+			BakedAttribute.Value.EvaluateToPtr(BakedAttribute.Key.GetType(), ExtractionContext.CurrentTime, AttributePtr);
 		}
 	}
 }
@@ -4636,182 +4640,51 @@ void UAnimSequence::GetCustomAttributes(FAnimationPoseData& OutAnimationPoseData
 #if WITH_EDITOR
 void UAnimSequence::RemoveCustomAttribute(const FName& BoneName, const FName& AttributeName)
 {
-	FCustomAttributePerBoneData* DataPtr = PerBoneCustomAttributeData.FindByPredicate([BoneName, this](FCustomAttributePerBoneData& Attribute)
-	{
-		return Attribute.BoneTreeIndex == GetSkeleton()->GetReferenceSkeleton().FindBoneIndex(BoneName);
-	});
+	ValidateModel();
 
-	if (DataPtr)
-	{
-		const int32 NumRemoved = DataPtr->Attributes.RemoveAll([AttributeName](FCustomAttribute& Attribute)
-		{
-			return Attribute.Name == AttributeName;
-		});
+	TArray<FAnimationAttributeIdentifier> IdentifiersToRemove;
 
-		// In case there are no custom attributes left for this bone, remove the wrapping structure entry as well
-		if (DataPtr->Attributes.Num() == 0)
+	Algo::TransformIf(DataModel->GetAttributes(), IdentifiersToRemove, 
+		[BoneName, AttributeName](const FAnimatedBoneAttribute& Attribute)
 		{
-			ensure(PerBoneCustomAttributeData.RemoveAll([DataPtr](FCustomAttributePerBoneData& Attribute)
-			{
-				return Attribute.BoneTreeIndex == DataPtr->BoneTreeIndex;
-			}) == 1);
-		}
-
-		if (NumRemoved)
+			return Attribute.Identifier.GetBoneName() == BoneName && Attribute.Identifier.GetName() == AttributeName;
+		},
+		[] (const FAnimatedBoneAttribute& Attribute)
 		{
-			// Update the Guid used to keep track of raw / baked versions
-			CustomAttributesGuid = FGuid::NewGuid();
+			return Attribute.Identifier;
 		}
-	}
-}
-
-void UAnimSequence::RemoveAllCustomAttributesForBone(const FName& BoneName)
-{
-	const USkeleton* CurrentSkeleton = GetSkeleton();
+	);
 
-	if (CurrentSkeleton)
+	if (IdentifiersToRemove.Num())
 	{
-		const int32 BoneIndex = CurrentSkeleton->GetReferenceSkeleton().FindBoneIndex(BoneName);
-
-		if (BoneIndex != INDEX_NONE)
+		UAnimDataController::FScopedBracket Bracket(Controller, LOCTEXT("RemoveCustomAttributeDeprecated", "Removing Custom Attributes by Bone and Name"));
+		for (const FAnimationAttributeIdentifier& AttributeIdentifier : IdentifiersToRemove)
 		{
-			const int32 NumRemoved = PerBoneCustomAttributeData.RemoveAll([BoneIndex](const FCustomAttributePerBoneData PerBoneData)
-			{
-				return PerBoneData.BoneTreeIndex == BoneIndex;
-			});
-
-			if (NumRemoved)
-			{
-				// Update the Guid used to keep track of raw / baked versions
-				CustomAttributesGuid = FGuid::NewGuid();
-			}
+			Controller->RemoveAttribute(AttributeIdentifier);
 		}
 	}
 }
 
-void UAnimSequence::RemoveAllCustomAttributes()
+void UAnimSequence::RemoveAllCustomAttributesForBone(const FName& BoneName)
 {
-	if (PerBoneCustomAttributeData.Num())
-	{
-		// Update the Guid used to keep track of raw / baked versions
-		CustomAttributesGuid = FGuid::NewGuid();
-	}
-
-	PerBoneCustomAttributeData.Empty();
+	ValidateModel();
+	Controller->RemoveAllAttributesForBone(BoneName);
 }
 
-void UAnimSequence::GetCustomAttributesForBone(const FName& BoneName, TArray<FCustomAttribute>& OutAttributes) const
+void UAnimSequence::RemoveAllCustomAttributes()
 {
-	const USkeleton* CurrentSkeleton = GetSkeleton();
-
-	if (CurrentSkeleton)
-	{
-		const int32 BoneIndex = CurrentSkeleton->GetReferenceSkeleton().FindBoneIndex(BoneName);
-
-		if (BoneIndex != INDEX_NONE)
-		{
-			for (const FCustomAttributePerBoneData& PerBoneData : PerBoneCustomAttributeData)
-			{
-				if (PerBoneData.BoneTreeIndex == BoneIndex)
-				{
-					OutAttributes.Append(PerBoneData.Attributes);
-				}
-			}
-		}
-	}
+	ValidateModel();
+	Controller->RemoveAllAttributes();
 }
 
-// Helper functionality to populate a curve by sampling the custom attribute data
-template<typename DataType, typename CurveType>
-void ConvertAttributeToAdditive(const FCustomAttribute& AdditiveAttribute, const FCustomAttribute& RefAttribute, CurveType& InOutCurve, float SamplingTime, int32 NumberOfFrames, TFunctionRef<float(float Time)> GetReferenceTime)
+void UAnimSequence::SynchronousAnimatedBoneAttributesCompression()
 {
-	for (int32 Frame = 0; Frame < NumberOfFrames; ++Frame)
-	{
-		const float CurrentFrameTime = Frame * SamplingTime;
-
-		DataType AdditiveValue;
-		FCustomAttributesRuntime::GetAttributeValue(AdditiveAttribute, CurrentFrameTime, AdditiveValue);
-
-		DataType RefValue;
-		FCustomAttributesRuntime::GetAttributeValue(RefAttribute, GetReferenceTime(CurrentFrameTime), RefValue);
+	ValidateModel();
 
-		const DataType Value = RefValue - AdditiveValue;
-		InOutCurve.AddKey(CurrentFrameTime, Value);
-	}
-}
+	AttributeCurves.Empty();
 
-void UAnimSequence::SynchronousCustomAttributesCompression()
-{
 	// If we are additive, we'll need to sample the base pose (against we're additive) and subtract the attributes from the base ones
 	const bool bShouldSampleBasePose = IsValidAdditive() && RefPoseType != ABPT_RefPose;
-	
-	BakedPerBoneCustomAttributeData.Empty(PerBoneCustomAttributeData.Num());
-
-	auto ProcessCustomAttribute = [this](const FCustomAttribute& Attribute, FBakedCustomAttributePerBoneData& BakedBoneAttributes)
-	{		
-		switch (static_cast<EVariantTypes>(Attribute.VariantType))
-		{
-			case EVariantTypes::Float:
-			{
-				FBakedFloatCustomAttribute& BakedFloatAttribute = BakedBoneAttributes.FloatAttributes.AddDefaulted_GetRef();
-				BakedFloatAttribute.AttributeName = Attribute.Name;
-
-				FSimpleCurve& FloatCurve = BakedFloatAttribute.FloatCurve;
-
-				TArray<FSimpleCurveKey> Keys;
-				for (int32 KeyIndex = 0; KeyIndex < Attribute.Times.Num(); ++KeyIndex)
-				{
-					const FVariant& VariantValue = Attribute.Values[KeyIndex];
-					FloatCurve.AddKey(Attribute.Times[KeyIndex], VariantValue.GetValue<float>());
-				}
-
-				FloatCurve.SetDefaultValue(FloatCurve.GetFirstKey().Value);
-				FloatCurve.RemoveRedundantKeys(0.f);
-				break;
-			}
-
-			case EVariantTypes::Int32:
-			{
-				FBakedIntegerCustomAttribute& BakedIntAttribute = BakedBoneAttributes.IntAttributes.AddDefaulted_GetRef();
-				BakedIntAttribute.AttributeName = Attribute.Name;
-
-				FIntegralCurve& IntCurve = BakedIntAttribute.IntCurve;
-				for (int32 KeyIndex = 0; KeyIndex < Attribute.Times.Num(); ++KeyIndex)
-				{
-					const FVariant& VariantValue = Attribute.Values[KeyIndex];
-					IntCurve.AddKey(Attribute.Times[KeyIndex], VariantValue.GetValue<int32>());
-				}
-
-				IntCurve.SetDefaultValue(IntCurve.GetKey(IntCurve.GetFirstKeyHandle()).Value);
-				IntCurve.RemoveRedundantKeys();
-				break;
-			}
-
-			case EVariantTypes::String:
-			{
-				FBakedStringCustomAttribute& BakedStringAttribute = BakedBoneAttributes.StringAttributes.AddDefaulted_GetRef();
-				BakedStringAttribute.AttributeName = Attribute.Name;
-
-				FStringCurve& StringCurve = BakedStringAttribute.StringCurve;
-				for (int32 KeyIndex = 0; KeyIndex < Attribute.Times.Num(); ++KeyIndex)
-				{
-					const FVariant& VariantValue = Attribute.Values[KeyIndex];
-					StringCurve.AddKey(Attribute.Times[KeyIndex], VariantValue.GetValue<FString>());
-				}
-
-				StringCurve.SetDefaultValue(StringCurve.GetKey(StringCurve.GetFirstKeyHandle()).Value);
-				StringCurve.RemoveRedundantKeys();
-				break;
-			}
-
-			default:
-			{
-				ensureMsgf(false, TEXT("Invalid data variant type for custom attribute, only int32, float and FString are currently supported"));
-				break;
-			}
-		}
-	};
-
 	if (bShouldSampleBasePose)
 	{
 		// Behaviour for determining the time to sample the base pose attributes
@@ -4829,7 +4702,6 @@ void UAnimSequence::SynchronousCustomAttributesCompression()
 			{
 				const float Fraction = (RefPoseSeq->GetNumberOfSampledKeys() > 0) ? FMath::Clamp<float>((float)RefFrameIndex / (float)RefPoseSeq->GetNumberOfSampledKeys(), 0.f, 1.f) : 0.f;
 				BasePoseTime = RefPoseSeq->GetPlayLength() * Fraction;
-
 			}
 
 			return BasePoseTime;
@@ -4839,101 +4711,176 @@ void UAnimSequence::SynchronousCustomAttributesCompression()
 
 		// Helper struct to match sample timings with regular additive baking
 		FByFramePoseEvalContext EvalContext(this);
-		for (const FCustomAttributePerBoneData& BoneAttributes : PerBoneCustomAttributeData)
-		{
-			FBakedCustomAttributePerBoneData& BakedBoneAttributes = BakedPerBoneCustomAttributeData.AddDefaulted_GetRef();
-			BakedBoneAttributes.BoneTreeIndex = BoneAttributes.BoneTreeIndex;
 
-			TArray<FCustomAttribute> ReferenceSequenceAttributes;			
-			RefPoseSeq->GetCustomAttributesForBone(RefSkeleton.GetBoneName(BoneAttributes.BoneTreeIndex), ReferenceSequenceAttributes);
+		for (const FAnimatedBoneAttribute& AdditiveAttribute : DataModel->GetAttributes())
+		{
+			ensure(!AttributeCurves.Contains(AdditiveAttribute.Identifier));
+			FAttributeCurve& AttributeCurve = AttributeCurves.Add(AdditiveAttribute.Identifier);
 
-			// Check whether or not the base sequence has any attributes
-			if (!ReferenceSequenceAttributes.Num())
+			if (const FAnimatedBoneAttribute* RefPoseAttributePtr = RefPoseSeq->GetDataModel()->FindAttribute(AdditiveAttribute.Identifier))
 			{
-				for (const FCustomAttribute& Attribute : BoneAttributes.Attributes)
+				AttributeCurve.SetScriptStruct(AdditiveAttribute.Identifier.GetType());
+				const FAnimatedBoneAttribute& RefAttribute = *RefPoseAttributePtr;
+				for (int32 Frame = 0; Frame < NumberOfSampledKeys; ++Frame)
 				{
-					ProcessCustomAttribute(Attribute, BakedBoneAttributes);
+					const float CurrentFrameTime = Frame * EvalContext.IntervalTime;
+
+					UE::Anim::FStackAttributeContainer AdditiveAttributes;
+					uint8* AdditivePtr = AdditiveAttributes.FindOrAdd(AdditiveAttribute.Identifier.GetType(), UE::Anim::FAttributeId(NAME_None, 0));
+					AdditiveAttribute.Curve.EvaluateToPtr(AdditiveAttribute.Identifier.GetType(), CurrentFrameTime, AdditivePtr);
+
+					UE::Anim::FStackAttributeContainer RefAttributes;
+					uint8* RefPtr = RefAttributes.FindOrAdd(RefAttribute.Identifier.GetType(), UE::Anim::FAttributeId(NAME_None, 0));
+					RefAttribute.Curve.EvaluateToPtr(RefAttribute.Identifier.GetType(), GetBasePoseTimeToSample(CurrentFrameTime), RefPtr);
+
+					UE::Anim::Attributes::ConvertToAdditive(RefAttributes, AdditiveAttributes);
+
+					AttributeCurve.AddKey(CurrentFrameTime, (void*)AdditivePtr);
 				}
 			}
 			else
 			{
-				for (const FCustomAttribute& Attribute : BoneAttributes.Attributes)
-				{
-					// Try and find equivalent in reference sequence
-					const FCustomAttribute* RefAttribute = ReferenceSequenceAttributes.FindByPredicate([Attribute](const FCustomAttribute& Attr)
-					{	
-						return Attribute.Name == Attr.Name && Attribute.VariantType == Attr.VariantType;
-					});
+				AttributeCurve = AdditiveAttribute.Curve;
+			}
+		}
+	}
+	else
+	{
+		for (const FAnimatedBoneAttribute& Attribute : DataModel->GetAttributes())
+		{
+			// Do something with the attributes
+			ensure(!AttributeCurves.Contains(Attribute.Identifier));
 
-					if (RefAttribute)
+			FAttributeCurve& BakedCurve = AttributeCurves.Add(Attribute.Identifier);
+			BakedCurve = Attribute.Curve;
+		}
+	}
+
+	for (TPair<FAnimationAttributeIdentifier, FAttributeCurve>& BakedAttribute : AttributeCurves)
+	{
+		BakedAttribute.Value.RemoveRedundantKeys();
+	}
+}
+
+void UAnimSequence::MoveAttributesToModel()
+{
+	WaitOnExistingCompression(false);
+
+	USkeleton* TargetSkeleton = GetSkeleton();
+	PRAGMA_DISABLE_DEPRECATION_WARNINGS
+	if (TargetSkeleton && PerBoneCustomAttributeData.Num())
+	{
+		UAnimDataController::FScopedBracket Bracket(Controller, LOCTEXT("MoveAttributesToModel", "Moving legacy Custom Attributes to Model"));
+
+		for (const FCustomAttributePerBoneData& PerBoneData : PerBoneCustomAttributeData)
+		{
+			const FName BoneName = TargetSkeleton->GetReferenceSkeleton().GetBoneName(PerBoneData.BoneTreeIndex);
+			if (BoneName != NAME_None)
+			{
+				auto ProcessCustomAttribute = [BoneName, this](const FCustomAttribute& Attribute)
+				{
+					switch (static_cast<EVariantTypes>(Attribute.VariantType))
 					{
-						switch (static_cast<EVariantTypes>(Attribute.VariantType))
+					case EVariantTypes::Float:
+					{
+						FAnimationAttributeIdentifier Identifier = UAnimationAttributeIdentifierExtensions::CreateAttributeIdentifier(this, Attribute.Name, BoneName, FFloatAnimationAttribute::StaticStruct());
+
+						if (Controller->AddAttribute(Identifier))
 						{
-							case EVariantTypes::Float:
-							{
-								FBakedFloatCustomAttribute& BakedFloatAttribute = BakedBoneAttributes.FloatAttributes.AddDefaulted_GetRef();
-								BakedFloatAttribute.AttributeName = Attribute.Name;
+							TArray<int8> TempArray;
+							TempArray.SetNumZeroed(FFloatAnimationAttribute::StaticStruct()->GetStructureSize());
+							FFloatAnimationAttribute::StaticStruct()->InitializeStruct(TempArray.GetData());
 
-								FSimpleCurve& FloatCurve = BakedFloatAttribute.FloatCurve;
-								ConvertAttributeToAdditive<float, FSimpleCurve>(Attribute, *RefAttribute, FloatCurve, EvalContext.IntervalTime, NumberOfSampledKeys, GetBasePoseTimeToSample);
-								FloatCurve.RemoveRedundantKeys(0.f);
+							FFloatAnimationAttribute* TempAttribute = (FFloatAnimationAttribute*)TempArray.GetData();
 
-								break;
+							for (int32 KeyIndex = 0; KeyIndex < Attribute.Times.Num(); ++KeyIndex)
+							{
+								const FVariant& VariantValue = Attribute.Values[KeyIndex];
+								TempAttribute->Value = VariantValue.GetValue<float>();
+								Controller->SetTypedAttributeKey<FFloatAnimationAttribute>(Identifier, Attribute.Times[KeyIndex], *TempAttribute);
 							}
+						}
+						else
+						{
+							UE_LOG(LogAnimation, Warning, TEXT("Failed to upgrade float attribute %s for bone %s"), *Attribute.Name.ToString(), *BoneName.ToString());
+
+						}
 
-							case EVariantTypes::Int32:
+						break;
+					}
+
+					case EVariantTypes::Int32:
+					{
+						FAnimationAttributeIdentifier Identifier = UAnimationAttributeIdentifierExtensions::CreateAttributeIdentifier(this, Attribute.Name, BoneName, FIntegerAnimationAttribute::StaticStruct());
+						if (Controller->AddAttribute(Identifier))
+						{
+							TArray<int8> TempArray;
+							TempArray.SetNumZeroed(FIntegerAnimationAttribute::StaticStruct()->GetStructureSize());
+							FIntegerAnimationAttribute::StaticStruct()->InitializeStruct(TempArray.GetData());
+
+							FIntegerAnimationAttribute* TempAttribute = (FIntegerAnimationAttribute*)TempArray.GetData();
+
+							for (int32 KeyIndex = 0; KeyIndex < Attribute.Times.Num(); ++KeyIndex)
 							{
-								FBakedIntegerCustomAttribute& BakedIntAttribute = BakedBoneAttributes.IntAttributes.AddDefaulted_GetRef();
-								BakedIntAttribute.AttributeName = Attribute.Name;
-
-								FIntegralCurve& IntCurve = BakedIntAttribute.IntCurve;
-								ConvertAttributeToAdditive<int32, FIntegralCurve>(Attribute, *RefAttribute, IntCurve, EvalContext.IntervalTime, NumberOfSampledKeys, GetBasePoseTimeToSample);
-								IntCurve.RemoveRedundantKeys();
-							
-								break;
+								const FVariant& VariantValue = Attribute.Values[KeyIndex];
+								TempAttribute->Value = VariantValue.GetValue<int32>();
+								Controller->SetTypedAttributeKey<FIntegerAnimationAttribute>(Identifier, Attribute.Times[KeyIndex], *TempAttribute);
 							}
+						}
+						else
+						{
+							UE_LOG(LogAnimation, Warning, TEXT("Failed to upgrade integer attribute %s for bone %s"), *Attribute.Name.ToString(), *BoneName.ToString());
+						}
 
-							case EVariantTypes::String:
+						break;
+					}
+
+					case EVariantTypes::String:
+					{
+						FAnimationAttributeIdentifier Identifier = UAnimationAttributeIdentifierExtensions::CreateAttributeIdentifier(this, Attribute.Name, BoneName, FStringAnimationAttribute::StaticStruct());
+						if (Controller->AddAttribute(Identifier))
+						{
+							TArray<int8> TempArray;
+							TempArray.SetNumZeroed(FStringAnimationAttribute::StaticStruct()->GetStructureSize());
+							FStringAnimationAttribute::StaticStruct()->InitializeStruct(TempArray.GetData());
+
+							FStringAnimationAttribute*TempAttribute = (FStringAnimationAttribute*)TempArray.GetData();
+							for (int32 KeyIndex = 0; KeyIndex < Attribute.Times.Num(); ++KeyIndex)
 							{
-								ProcessCustomAttribute(Attribute, BakedBoneAttributes);
-								break;
+								const FVariant& VariantValue = Attribute.Values[KeyIndex];
+								TempAttribute->Value = VariantValue.GetValue<FString>();
+								Controller->SetTypedAttributeKey<FStringAnimationAttribute>(Identifier, Attribute.Times[KeyIndex], *TempAttribute);
 							}
 						}
+						else
+						{
+							UE_LOG(LogAnimation, Warning, TEXT("Failed to upgrade string attribute %s for bone %s"), *Attribute.Name.ToString(), *BoneName.ToString());
+						}
+						break;
 					}
-					else
+
+					default:
 					{
-						ProcessCustomAttribute(Attribute, BakedBoneAttributes);
-					}					
+						ensureMsgf(false, TEXT("Invalid data variant type for custom attribute, only int32, float and FString are currently supported"));
+						break;
+					}
+					}
+				};
+
+				for (const FCustomAttribute& Attribute : PerBoneData.Attributes)
+				{
+					ProcessCustomAttribute(Attribute);
 				}
 			}
-		}
-	}
-	else
-	{
-		for (const FCustomAttributePerBoneData& BoneAttributes : PerBoneCustomAttributeData)
-		{
-			FBakedCustomAttributePerBoneData& BakedBoneAttributes = BakedPerBoneCustomAttributeData.AddDefaulted_GetRef();
-			BakedBoneAttributes.BoneTreeIndex = BoneAttributes.BoneTreeIndex;
-
-			for (const FCustomAttribute& Attribute : BoneAttributes.Attributes)
+			else
 			{
-				ProcessCustomAttribute(Attribute, BakedBoneAttributes);
+				UE_LOG(LogAnimation, Warning, TEXT("Failed to upgrade custom attributes for bone index %i as no such bone index existing in the Skeleton %s"), PerBoneData.BoneTreeIndex, *TargetSkeleton->GetName());
 			}
 		}
-	}
-
-	// Match baked/raw attributes guid
-	BakedCustomAttributesGuid = CustomAttributesGuid;
-}
-
-FCustomAttributePerBoneData& UAnimSequence::FindOrAddCustomAttributeForBone(const FName& BoneName)
-{
-	FCustomAttributePerBoneData* DataPtr = PerBoneCustomAttributeData.FindByPredicate([BoneName, this](FCustomAttributePerBoneData& Attribute)
-	{
-		return Attribute.BoneTreeIndex == GetSkeleton()->GetReferenceSkeleton().FindBoneIndex(BoneName);
-	});
 
-	return DataPtr ? *DataPtr : PerBoneCustomAttributeData.AddDefaulted_GetRef();
+		PerBoneCustomAttributeData.Empty();
+	}
+	PRAGMA_ENABLE_DEPRECATION_WARNINGS
 }
 #endif // WITH_EDITOR
 
@@ -5670,7 +5617,8 @@ void UAnimSequence::OnModelModified(const EAnimDataModelNotifyType& NotifyType,
 			{
 				const auto LengthChangingNotifies = { EAnimDataModelNotifyType::SequenceLengthChanged, EAnimDataModelNotifyType::FrameRateChanged, EAnimDataModelNotifyType::Reset };
 				const auto ResamplingNotifies = { EAnimDataModelNotifyType::TrackAdded, EAnimDataModelNotifyType::TrackChanged, EAnimDataModelNotifyType::TrackRemoved,  EAnimDataModelNotifyType::Populated };
-				const auto RecompressNotifies = { EAnimDataModelNotifyType::CurveAdded, EAnimDataModelNotifyType::CurveChanged, EAnimDataModelNotifyType::CurveRemoved, EAnimDataModelNotifyType::CurveFlagsChanged, EAnimDataModelNotifyType::CurveScaled };
+				const auto RecompressNotifies = { EAnimDataModelNotifyType::CurveAdded, EAnimDataModelNotifyType::CurveChanged, EAnimDataModelNotifyType::CurveRemoved, EAnimDataModelNotifyType::CurveFlagsChanged, EAnimDataModelNotifyType::CurveScaled,
+				EAnimDataModelNotifyType::AttributeAdded, EAnimDataModelNotifyType::AttributeChanged, EAnimDataModelNotifyType::AttributeRemoved };
 				
 				if (NotifyCollector.Contains(LengthChangingNotifies) || NotifyCollector.Contains(ResamplingNotifies))
 				{
@@ -5730,6 +5678,18 @@ void UAnimSequence::OnModelModified(const EAnimDataModelNotifyType& NotifyType,
 			break;
 		}
 
+		case EAnimDataModelNotifyType::AttributeAdded:
+		case EAnimDataModelNotifyType::AttributeChanged:
+		case EAnimDataModelNotifyType::AttributeRemoved:
+		{
+			if (NotifyCollector.IsNotWithinBracket())
+			{
+				UpdateRawDataGuid(RegenerateGUID);
+				RecompressAnimationData();
+			}
+			break;
+		}
+
 		case EAnimDataModelNotifyType::CurveColorChanged:
 		{
 			break;
diff --git a/Engine/Source/Runtime/Engine/Private/Animation/AnimSequenceBase.cpp b/Engine/Source/Runtime/Engine/Private/Animation/AnimSequenceBase.cpp
--- a/Engine/Source/Runtime/Engine/Private/Animation/AnimSequenceBase.cpp
+++ b/Engine/Source/Runtime/Engine/Private/Animation/AnimSequenceBase.cpp
@@ -895,7 +895,7 @@ void UAnimSequenceBase::Serialize(FArchive& Ar)
 
 void UAnimSequenceBase::GetAnimationPose(struct FCompactPose& OutPose, FBlendedCurve & OutCurve, const FAnimExtractContext & ExtractionContext) const
 {
-	FStackCustomAttributes TempAttributes;
+	UE::Anim::FStackAttributeContainer TempAttributes;
 	FAnimationPoseData OutAnimationPoseData(OutPose, OutCurve, TempAttributes);
 	GetAnimationPose(OutAnimationPoseData, ExtractionContext);
 }
diff --git a/Engine/Source/Runtime/Engine/Private/Animation/AnimSingleNodeInstanceProxy.cpp b/Engine/Source/Runtime/Engine/Private/Animation/AnimSingleNodeInstanceProxy.cpp
--- a/Engine/Source/Runtime/Engine/Private/Animation/AnimSingleNodeInstanceProxy.cpp
+++ b/Engine/Source/Runtime/Engine/Private/Animation/AnimSingleNodeInstanceProxy.cpp
@@ -166,7 +166,7 @@ void FAnimSingleNodeInstanceProxy::InternalBlendSpaceEvaluatePose(class UBlendSp
 		FBlendedCurve& OutCurve = OutContext.Curve;
 		FCompactPose AdditivePose;
 		FBlendedCurve AdditiveCurve;
-		FStackCustomAttributes AdditiveAttributes;
+		UE::Anim::FStackAttributeContainer AdditiveAttributes;
 		AdditivePose.SetBoneContainer(&OutPose.GetBoneContainer());
 		AdditiveCurve.InitFrom(OutCurve);
 #if WITH_EDITORONLY_DATA
@@ -299,7 +299,7 @@ void FAnimNode_SingleNode::Evaluate_AnyThread(FPoseContext& Output)
 
 				FCompactPose AdditivePose;
 				FBlendedCurve AdditiveCurve;
-				FStackCustomAttributes AdditiveAttributes;
+				UE::Anim::FStackAttributeContainer AdditiveAttributes;
 				AdditivePose.SetBoneContainer(&Output.Pose.GetBoneContainer());
 				AdditiveCurve.InitFrom(Output.Curve);
 
@@ -370,7 +370,7 @@ void FAnimNode_SingleNode::Evaluate_AnyThread(FPoseContext& Output)
 
 				FCompactPose AdditivePose;
 				FBlendedCurve AdditiveCurve;
-				FStackCustomAttributes AdditiveAttributes;
+				UE::Anim::FStackAttributeContainer AdditiveAttributes;
 				AdditivePose.SetBoneContainer(&Output.Pose.GetBoneContainer());
 				AdditiveCurve.InitFrom(Output.Curve);
 
@@ -396,7 +396,7 @@ void FAnimNode_SingleNode::Evaluate_AnyThread(FPoseContext& Output)
 				LocalSourcePose.SetBoneContainer(&Output.Pose.GetBoneContainer());
 				LocalSourceCurve.InitFrom(Output.Curve);
 
-				FStackCustomAttributes LocalSourceAttributes;
+				UE::Anim::FStackAttributeContainer LocalSourceAttributes;
 			
 				FAnimTrack const* const AnimTrack = Montage->GetAnimationData(ActiveMontageSlot);
 				if (AnimTrack && AnimTrack->IsAdditive())
@@ -462,7 +462,7 @@ void FAnimNode_SingleNode::Evaluate_AnyThread(FPoseContext& Output)
 				{
 					FCompactPose AdditivePose;
 					FBlendedCurve AdditiveCurve;
-					FStackCustomAttributes AdditiveAttributes;
+					UE::Anim::FStackAttributeContainer AdditiveAttributes;
 					AdditivePose.SetBoneContainer(&Output.Pose.GetBoneContainer());
 					AdditiveCurve.InitFrom(Output.Curve);
 
diff --git a/Engine/Source/Runtime/Engine/Private/Animation/AnimStreamable.cpp b/Engine/Source/Runtime/Engine/Private/Animation/AnimStreamable.cpp
--- a/Engine/Source/Runtime/Engine/Private/Animation/AnimStreamable.cpp
+++ b/Engine/Source/Runtime/Engine/Private/Animation/AnimStreamable.cpp
@@ -20,7 +20,7 @@
 #include "ContentStreaming.h"
 #include "ProfilingDebugging/CookStats.h"
 #include "Animation/AnimationPoseData.h"
-#include "Animation/CustomAttributesRuntime.h"
+#include "Animation/AttributesRuntime.h"
 #include "UObject/UE5MainStreamObjectVersion.h"
 
 #include "Animation/AnimData/AnimDataModel.h"
diff --git a/Engine/Source/Runtime/Engine/Private/Animation/AnimationPoseData.cpp b/Engine/Source/Runtime/Engine/Private/Animation/AnimationPoseData.cpp
--- a/Engine/Source/Runtime/Engine/Private/Animation/AnimationPoseData.cpp
+++ b/Engine/Source/Runtime/Engine/Private/Animation/AnimationPoseData.cpp
@@ -4,7 +4,7 @@
 #include "Animation/AnimNodeBase.h"
 #include "Animation/AnimInstance.h"
 #include "Animation/AnimCurveTypes.h"
-#include "Animation/CustomAttributesRuntime.h"
+#include "Animation/AttributesRuntime.h"
 
 #include "BonePose.h"
 
@@ -18,7 +18,7 @@ FAnimationPoseData::FAnimationPoseData(FSlotEvaluationPose& InSlotPoseContext)
 {
 }
 
-FAnimationPoseData::FAnimationPoseData(FCompactPose& InPose, FBlendedCurve& InCurve, FStackCustomAttributes& InAttributes) : Pose(InPose), Curve(InCurve), Attributes(InAttributes)
+FAnimationPoseData::FAnimationPoseData(FCompactPose& InPose, FBlendedCurve& InCurve, UE::Anim::FStackAttributeContainer& InAttributes) : Pose(InPose), Curve(InCurve), Attributes(InAttributes)
 {
 }
 
@@ -42,12 +42,12 @@ const FBlendedCurve& FAnimationPoseData::GetCurve() const
 	return Curve;
 }
 
-FStackCustomAttributes& FAnimationPoseData::GetAttributes()
+UE::Anim::FStackAttributeContainer& FAnimationPoseData::GetAttributes()
 {
 	return Attributes;
 }
 
-const FStackCustomAttributes& FAnimationPoseData::GetAttributes() const
+const UE::Anim::FStackAttributeContainer& FAnimationPoseData::GetAttributes() const
 {
 	return Attributes;
 }
\ No newline at end of file
diff --git a/Engine/Source/Runtime/Engine/Private/Animation/AnimationRuntime.cpp b/Engine/Source/Runtime/Engine/Private/Animation/AnimationRuntime.cpp
--- a/Engine/Source/Runtime/Engine/Private/Animation/AnimationRuntime.cpp
+++ b/Engine/Source/Runtime/Engine/Private/Animation/AnimationRuntime.cpp
@@ -10,7 +10,7 @@
 #include "Animation/AnimInstance.h"
 #include "SkeletalRender.h"
 #include "Animation/CustomAttributes.h"
-#include "Animation/CustomAttributesRuntime.h"
+#include "Animation/AttributesRuntime.h"
 #include "GenericPlatform/GenericPlatformCompilerPreSetup.h"
 #include "Animation/AnimationPoseData.h"
 #include "Animation/BlendProfile.h"
@@ -266,19 +266,19 @@ void FAnimationRuntime::BlendPosesTogether(
 	/*out*/ FCompactPose& ResultPose, 
 	/*out*/ FBlendedCurve& ResultCurve)
 {
-	FStackCustomAttributes TempAttributes;
+	UE::Anim::FStackAttributeContainer TempAttributes;
 	FAnimationPoseData AnimationPoseData = { ResultPose, ResultCurve, TempAttributes };
 
 	BlendPosesTogether(SourcePoses, SourceCurves, {}, SourceWeights, AnimationPoseData);	
 }
 
-void FAnimationRuntime::BlendPosesTogether(TArrayView<const FCompactPose> SourcePoses, TArrayView<const FBlendedCurve> SourceCurves, TArrayView<const FStackCustomAttributes> SourceAttributes, TArrayView<const float> SourceWeights, FAnimationPoseData& OutAnimationPoseData)
+void FAnimationRuntime::BlendPosesTogether(TArrayView<const FCompactPose> SourcePoses, TArrayView<const FBlendedCurve> SourceCurves, TArrayView<const UE::Anim::FStackAttributeContainer> SourceAttributes, TArrayView<const float> SourceWeights, FAnimationPoseData& OutAnimationPoseData)
 {
 	check(SourcePoses.Num() > 0);
 
 	FCompactPose& OutPose = OutAnimationPoseData.GetPose();
 	FBlendedCurve& OutCurve = OutAnimationPoseData.GetCurve();
-	FStackCustomAttributes& OutAttributes = OutAnimationPoseData.GetAttributes();
+	UE::Anim::FStackAttributeContainer& OutAttributes = OutAnimationPoseData.GetAttributes();
 
 	BlendPose<ETransformBlendMode::Overwrite>(SourcePoses[0], OutPose, SourceWeights[0]);
 
@@ -301,7 +301,7 @@ void FAnimationRuntime::BlendPosesTogether(TArrayView<const FCompactPose> Source
 
 	if (SourceAttributes.Num() > 0)
 	{
-		FCustomAttributesRuntime::BlendAttributes(SourceAttributes, SourceWeights, OutAttributes);
+		UE::Anim::Attributes::BlendAttributes(SourceAttributes, SourceWeights, OutAttributes);
 	}
 }
 
@@ -313,20 +313,20 @@ void FAnimationRuntime::BlendPosesTogether(
 	/*out*/ FCompactPose& ResultPose,
 	/*out*/ FBlendedCurve& ResultCurve)
 {
-	FStackCustomAttributes TempAttributes;
+	UE::Anim::FStackAttributeContainer TempAttributes;
 	FAnimationPoseData AnimationPoseData = { ResultPose, ResultCurve, TempAttributes };
 
 	BlendPosesTogether(SourcePoses, SourceCurves, {}, SourceWeights, SourceWeightsIndices, AnimationPoseData);
 }
 
 
-void FAnimationRuntime::BlendPosesTogether(TArrayView<const FCompactPose> SourcePoses, TArrayView<const FBlendedCurve> SourceCurves, TArrayView<const FStackCustomAttributes> SourceAttributes, TArrayView<const float> SourceWeights, TArrayView<const int32> SourceWeightsIndices, /*out*/ FAnimationPoseData& OutAnimationPoseData)
+void FAnimationRuntime::BlendPosesTogether(TArrayView<const FCompactPose> SourcePoses, TArrayView<const FBlendedCurve> SourceCurves, TArrayView<const UE::Anim::FStackAttributeContainer> SourceAttributes, TArrayView<const float> SourceWeights, TArrayView<const int32> SourceWeightsIndices, /*out*/ FAnimationPoseData& OutAnimationPoseData)
 {
 	check(SourcePoses.Num() > 0);
 
 	FCompactPose& OutPose = OutAnimationPoseData.GetPose();
 	FBlendedCurve& OutCurve = OutAnimationPoseData.GetCurve();
-	FStackCustomAttributes& OutAttributes = OutAnimationPoseData.GetAttributes();
+	UE::Anim::FStackAttributeContainer& OutAttributes = OutAnimationPoseData.GetAttributes();
 
 	BlendPose<ETransformBlendMode::Overwrite>(SourcePoses[0], OutPose, SourceWeights[SourceWeightsIndices[0]]);
 
@@ -349,7 +349,7 @@ void FAnimationRuntime::BlendPosesTogether(TArrayView<const FCompactPose> Source
 
 	if (SourceAttributes.Num() > 0)
 	{
-		FCustomAttributesRuntime::BlendAttributes(SourceAttributes, SourceWeights, SourceWeightsIndices, OutAttributes);
+		UE::Anim::Attributes::BlendAttributes(SourceAttributes, SourceWeights, SourceWeightsIndices, OutAttributes);
 	}
 }
 
@@ -360,19 +360,19 @@ void FAnimationRuntime::BlendPosesTogetherIndirect(
 	/*out*/ FCompactPose& ResultPose,
 	/*out*/ FBlendedCurve& ResultCurve)
 {
-	FStackCustomAttributes TempAttributes;
+	UE::Anim::FStackAttributeContainer TempAttributes;
 	FAnimationPoseData AnimationPoseData = { ResultPose, ResultCurve, TempAttributes };
 
 	BlendPosesTogetherIndirect(SourcePoses, SourceCurves, {}, SourceWeights, AnimationPoseData);
 }
 
-void FAnimationRuntime::BlendPosesTogetherIndirect(TArrayView<const FCompactPose* const> SourcePoses, TArrayView<const FBlendedCurve* const> SourceCurves, TArrayView<const FStackCustomAttributes* const> SourceAttributes, TArrayView<const float> SourceWeights, FAnimationPoseData& OutAnimationPoseData)
+void FAnimationRuntime::BlendPosesTogetherIndirect(TArrayView<const FCompactPose* const> SourcePoses, TArrayView<const FBlendedCurve* const> SourceCurves, TArrayView<const UE::Anim::FStackAttributeContainer* const> SourceAttributes, TArrayView<const float> SourceWeights, FAnimationPoseData& OutAnimationPoseData)
 {
 	check(SourcePoses.Num() > 0);
 
 	FCompactPose& OutPose = OutAnimationPoseData.GetPose();
 	FBlendedCurve& OutCurve = OutAnimationPoseData.GetCurve();
-	FStackCustomAttributes& OutAttributes = OutAnimationPoseData.GetAttributes();
+	UE::Anim::FStackAttributeContainer& OutAttributes = OutAnimationPoseData.GetAttributes();
 	
 	BlendPose<ETransformBlendMode::Overwrite>(*SourcePoses[0], OutPose, SourceWeights[0]);
 
@@ -394,7 +394,7 @@ void FAnimationRuntime::BlendPosesTogetherIndirect(TArrayView<const FCompactPose
 
 	if (SourceAttributes.Num() > 0)
 	{
-		FCustomAttributesRuntime::BlendAttributes(SourceAttributes, SourceWeights, OutAttributes);
+		UE::Anim::Attributes::BlendAttributes(SourceAttributes, SourceWeights, OutAttributes);
 	}
 }
 
@@ -407,7 +407,7 @@ void FAnimationRuntime::BlendTwoPosesTogether(
 	/*out*/ FCompactPose& ResultPose,
 	/*out*/ FBlendedCurve& ResultCurve)
 {
-	FStackCustomAttributes TempAttributes;
+	UE::Anim::FStackAttributeContainer TempAttributes;
 	
 	FAnimationPoseData AnimationPoseData = { ResultPose, ResultCurve, TempAttributes };
 	
@@ -421,7 +421,7 @@ void FAnimationRuntime::BlendTwoPosesTogether(const FAnimationPoseData& SourcePo
 {
 	FCompactPose& OutPose = OutAnimationPoseData.GetPose();
 	FBlendedCurve& OutCurve = OutAnimationPoseData.GetCurve();
-	FStackCustomAttributes& OutAttributes = OutAnimationPoseData.GetAttributes();
+	UE::Anim::FStackAttributeContainer& OutAttributes = OutAnimationPoseData.GetAttributes();
 
 	const float WeightOfPoseTwo = 1.f - WeightOfPoseOne;
 
@@ -432,7 +432,7 @@ void FAnimationRuntime::BlendTwoPosesTogether(const FAnimationPoseData& SourcePo
 	OutPose.NormalizeRotations();
 
 	OutCurve.Lerp(SourcePoseOneData.GetCurve(), SourcePoseTwoData.GetCurve(), WeightOfPoseTwo);
-	FCustomAttributesRuntime::BlendAttributes({ SourcePoseOneData.GetAttributes(), SourcePoseTwoData.GetAttributes() }, { WeightOfPoseOne, WeightOfPoseTwo }, { 0, 1 }, OutAttributes);
+	UE::Anim::Attributes::BlendAttributes({ SourcePoseOneData.GetAttributes(), SourcePoseTwoData.GetAttributes() }, { WeightOfPoseOne, WeightOfPoseTwo }, { 0, 1 }, OutAttributes);
 }
 
 void FAnimationRuntime::BlendTwoPosesTogetherPerBone(
@@ -444,7 +444,7 @@ void FAnimationRuntime::BlendTwoPosesTogetherPerBone(
 	/*out*/ FCompactPose& ResultPose,
 	/*out*/ FBlendedCurve& ResultCurve)
 {
-	FStackCustomAttributes TempAttributes;
+	UE::Anim::FStackAttributeContainer TempAttributes;
 
 	FAnimationPoseData AnimationPoseData = { ResultPose, ResultCurve, TempAttributes };
 	const FAnimationPoseData SourceOnePoseData(*const_cast<FCompactPose*>(&SourcePose1), *const_cast<FBlendedCurve*>(&SourceCurve1), TempAttributes);
@@ -457,7 +457,7 @@ void FAnimationRuntime::BlendTwoPosesTogetherPerBone(const FAnimationPoseData& S
 {
 	FCompactPose& OutPose = OutAnimationPoseData.GetPose();
 	FBlendedCurve& OutCurve = OutAnimationPoseData.GetCurve();
-	FStackCustomAttributes& OutAttributes = OutAnimationPoseData.GetAttributes();
+	UE::Anim::FStackAttributeContainer& OutAttributes = OutAnimationPoseData.GetAttributes();
 
 	const FCompactPose& SourcePoseOne = SourcePoseOneData.GetPose();
 	const FCompactPose& SourcePoseTwo = SourcePoseTwoData.GetPose();
@@ -494,7 +494,7 @@ void FAnimationRuntime::BlendTwoPosesTogetherPerBone(const FAnimationPoseData& S
 	OutCurve.Override(SourceCurveOne);
 	OutCurve.Combine(SourceCurveTwo);
 
-	FCustomAttributesRuntime::BlendAttributesPerBone(SourcePoseOneData.GetAttributes(), SourcePoseTwoData.GetAttributes(), WeightsOfSource2, OutAttributes);
+	UE::Anim::Attributes::BlendAttributesPerBone(SourcePoseOneData.GetAttributes(), SourcePoseTwoData.GetAttributes(), WeightsOfSource2, OutAttributes);
 }
 
 template <int32 TRANSFORM_BLEND_MODE>
@@ -542,20 +542,20 @@ void FAnimationRuntime::BlendPosesTogetherPerBone(
 	/*out*/ FCompactPose& ResultPose,
 	/*out*/ FBlendedCurve& ResultCurve)
 {
-	FStackCustomAttributes TempAttributes;
+	UE::Anim::FStackAttributeContainer TempAttributes;
 	FAnimationPoseData AnimationPoseData = { ResultPose, ResultCurve, TempAttributes };
 
 	BlendPosesTogetherPerBone(SourcePoses, SourceCurves, {}, InterpolationIndexProvider, BlendSampleDataCache, AnimationPoseData);	
 }
 
 
-void FAnimationRuntime::BlendPosesTogetherPerBone(TArrayView<const FCompactPose> SourcePoses, TArrayView<const FBlendedCurve> SourceCurves, TArrayView<const FStackCustomAttributes> SourceAttributes, const IInterpolationIndexProvider* InterpolationIndexProvider, TArrayView<const FBlendSampleData> BlendSampleDataCache, /*out*/ FAnimationPoseData& OutAnimationPoseData)
+void FAnimationRuntime::BlendPosesTogetherPerBone(TArrayView<const FCompactPose> SourcePoses, TArrayView<const FBlendedCurve> SourceCurves, TArrayView<const UE::Anim::FStackAttributeContainer> SourceAttributes, const IInterpolationIndexProvider* InterpolationIndexProvider, TArrayView<const FBlendSampleData> BlendSampleDataCache, /*out*/ FAnimationPoseData& OutAnimationPoseData)
 {
 	check(SourcePoses.Num() > 0);
 
 	FCompactPose& OutPose = OutAnimationPoseData.GetPose();
 	FBlendedCurve& OutCurve = OutAnimationPoseData.GetCurve();
-	FStackCustomAttributes& OutAttributes = OutAnimationPoseData.GetAttributes();
+	UE::Anim::FStackAttributeContainer& OutAttributes = OutAnimationPoseData.GetAttributes();
 
 	const TArray<FBoneIndexType>& RequiredBoneIndices = OutPose.GetBoneContainer().GetBoneIndicesArray();
 
@@ -576,7 +576,7 @@ void FAnimationRuntime::BlendPosesTogetherPerBone(TArrayView<const FCompactPose>
 	// Ensure that all of the resulting rotations are normalized
 	OutPose.NormalizeRotations();
 
-	if (SourceCurves.Num() > 0 || SourceAttributes.Num() > 0)
+	if (SourceCurves.Num() > 0)
 	{
 		TArray<float, TInlineAllocator<16>> SourceWeights;
 		SourceWeights.AddUninitialized(BlendSampleDataCache.Num());
@@ -585,16 +585,13 @@ void FAnimationRuntime::BlendPosesTogetherPerBone(TArrayView<const FCompactPose>
 			SourceWeights[CacheIndex] = BlendSampleDataCache[CacheIndex].TotalWeight;
 		}
 
-		if (SourceCurves.Num() > 0)
-		{
-			BlendCurves(SourceCurves, SourceWeights, OutCurve);
-		}
-
-		if (SourceAttributes.Num() > 0)
-		{
-			FCustomAttributesRuntime::BlendAttributes(SourceAttributes, SourceWeights, OutAttributes);
-		}		
+		BlendCurves(SourceCurves, SourceWeights, OutCurve);		
 	}
+
+	if (SourceAttributes.Num() > 0)
+	{
+		UE::Anim::Attributes::BlendAttributesPerBone(SourceAttributes, PerBoneIndices, BlendSampleDataCache, OutAttributes);
+	}	
 }
 
 void FAnimationRuntime::BlendPosesTogetherPerBone(
@@ -606,19 +603,19 @@ void FAnimationRuntime::BlendPosesTogetherPerBone(
 	/*out*/ FCompactPose& ResultPose,
 	/*out*/ FBlendedCurve& ResultCurve)
 {
-	FStackCustomAttributes TempAttributes;
+	UE::Anim::FStackAttributeContainer TempAttributes;
 	FAnimationPoseData AnimationPoseData = { ResultPose, ResultCurve, TempAttributes };
 
 	BlendPosesTogetherPerBone(SourcePoses, SourceCurves, {}, InterpolationIndexProvider, BlendSampleDataCache, BlendSampleDataCacheIndices, AnimationPoseData);
 }
 
-void FAnimationRuntime::BlendPosesTogetherPerBone(TArrayView<const FCompactPose> SourcePoses, TArrayView<const FBlendedCurve> SourceCurves, TArrayView<const FStackCustomAttributes> SourceAttributes, const IInterpolationIndexProvider* InterpolationIndexProvider, TArrayView<const FBlendSampleData> BlendSampleDataCache, TArrayView<const int32> BlendSampleDataCacheIndices, /*out*/ FAnimationPoseData& OutAnimationPoseData)
+void FAnimationRuntime::BlendPosesTogetherPerBone(TArrayView<const FCompactPose> SourcePoses, TArrayView<const FBlendedCurve> SourceCurves, TArrayView<const UE::Anim::FStackAttributeContainer> SourceAttributes, const IInterpolationIndexProvider* InterpolationIndexProvider, TArrayView<const FBlendSampleData> BlendSampleDataCache, TArrayView<const int32> BlendSampleDataCacheIndices, /*out*/ FAnimationPoseData& OutAnimationPoseData)
 {
 	check(SourcePoses.Num() > 0);
 
 	FCompactPose& OutPose = OutAnimationPoseData.GetPose();
 	FBlendedCurve& OutCurve = OutAnimationPoseData.GetCurve();
-	FStackCustomAttributes& OutAttributes = OutAnimationPoseData.GetAttributes();
+	UE::Anim::FStackAttributeContainer& OutAttributes = OutAnimationPoseData.GetAttributes();
 
 	const TArray<FBoneIndexType>& RequiredBoneIndices = OutPose.GetBoneContainer().GetBoneIndicesArray();
 
@@ -639,7 +636,7 @@ void FAnimationRuntime::BlendPosesTogetherPerBone(TArrayView<const FCompactPose>
 	// Ensure that all of the resulting rotations are normalized
 	OutPose.NormalizeRotations();
 
-	if (SourceCurves.Num() > 0 || SourceAttributes.Num() > 0)
+	if (SourceCurves.Num() > 0)
 	{
 		TArray<float, TInlineAllocator<16>> SourceWeights;
 		SourceWeights.AddUninitialized(BlendSampleDataCacheIndices.Num());
@@ -648,16 +645,13 @@ void FAnimationRuntime::BlendPosesTogetherPerBone(TArrayView<const FCompactPose>
 			SourceWeights[CacheIndex] = BlendSampleDataCache[BlendSampleDataCacheIndices[CacheIndex]].TotalWeight;
 		}
 
-		if (SourceCurves.Num() > 0)
-		{
-			BlendCurves(SourceCurves, SourceWeights, OutCurve);
-		}	
-
-		if (SourceAttributes.Num() > 0)
-		{
-			FCustomAttributesRuntime::BlendAttributes(SourceAttributes, SourceWeights, OutAttributes);
-		}
+		BlendCurves(SourceCurves, SourceWeights, OutCurve);		
 	}	
+
+	if (SourceAttributes.Num() > 0)
+	{
+		UE::Anim::Attributes::BlendAttributesPerBone(SourceAttributes, PerBoneIndices, BlendSampleDataCache, BlendSampleDataCacheIndices, OutAttributes);
+	}
 }
 
 void FAnimationRuntime::BlendPosesTogetherPerBoneInMeshSpace(
@@ -668,13 +662,13 @@ void FAnimationRuntime::BlendPosesTogetherPerBoneInMeshSpace(
 	/*out*/ FCompactPose& ResultPose,
 	/*out*/ FBlendedCurve& ResultCurve)
 {
-	FStackCustomAttributes TempAttributes;
+	UE::Anim::FStackAttributeContainer TempAttributes;
 	FAnimationPoseData AnimationPoseData = { ResultPose, ResultCurve, TempAttributes };
 
 	BlendPosesTogetherPerBoneInMeshSpace(SourcePoses, SourceCurves, {}, BlendSpace, BlendSampleDataCache, AnimationPoseData);
 }
 
-void FAnimationRuntime::BlendPosesTogetherPerBoneInMeshSpace(TArrayView<FCompactPose> SourcePoses, TArrayView<const FBlendedCurve> SourceCurves, TArrayView<const FStackCustomAttributes> SourceAttributes, const UBlendSpace* BlendSpace, TArrayView<const FBlendSampleData> BlendSampleDataCache, /*out*/ FAnimationPoseData& OutAnimationPoseData)
+void FAnimationRuntime::BlendPosesTogetherPerBoneInMeshSpace(TArrayView<FCompactPose> SourcePoses, TArrayView<const FBlendedCurve> SourceCurves, TArrayView<const UE::Anim::FStackAttributeContainer> SourceAttributes, const UBlendSpace* BlendSpace, TArrayView<const FBlendSampleData> BlendSampleDataCache, /*out*/ FAnimationPoseData& OutAnimationPoseData)
 {
 	FQuat NewRotation;
 	USkeleton* Skeleton = BlendSpace->GetSkeleton();
@@ -1006,7 +1000,7 @@ void FAnimationRuntime::ConvertMeshRotationPoseToLocalSpace(FCompactPose& Pose)
 
 void FAnimationRuntime::AccumulateAdditivePose(FCompactPose& BasePose, const FCompactPose& AdditivePose, FBlendedCurve& BaseCurve, const FBlendedCurve& AdditiveCurve, float Weight, enum EAdditiveAnimationType AdditiveType)
 {
-	FStackCustomAttributes TempAttributes;
+	UE::Anim::FStackAttributeContainer TempAttributes;
 	FAnimationPoseData BaseAnimationPoseData = { BasePose, BaseCurve, TempAttributes };
 	const FAnimationPoseData AdditiveAnimationPoseData(*const_cast<FCompactPose*>(&AdditivePose), *const_cast<FBlendedCurve*>(&AdditiveCurve), TempAttributes);
 
@@ -1016,7 +1010,7 @@ void FAnimationRuntime::AccumulateAdditivePose(FCompactPose& BasePose, const FCo
 
 void FAnimationRuntime::AccumulateLocalSpaceAdditivePose(FCompactPose& BasePose, const FCompactPose& AdditivePose, FBlendedCurve& BaseCurve, const FBlendedCurve& AdditiveCurve, float Weight)
 {
-	FStackCustomAttributes TempAttributes;
+	UE::Anim::FStackAttributeContainer TempAttributes;
 	FAnimationPoseData BaseAnimationPoseData = { BasePose, BaseCurve, TempAttributes };
 	const FAnimationPoseData AdditiveAnimationPoseData(*const_cast<FCompactPose*>(&AdditivePose), *const_cast<FBlendedCurve*>(&AdditiveCurve), TempAttributes);
 
@@ -1030,7 +1024,7 @@ void FAnimationRuntime::AccumulateLocalSpaceAdditivePose(FAnimationPoseData& Bas
 
 void FAnimationRuntime::AccumulateMeshSpaceRotationAdditiveToLocalPose(FCompactPose& BasePose, const FCompactPose& MeshSpaceRotationAdditive, FBlendedCurve& BaseCurve, const FBlendedCurve& AdditiveCurve, float Weight)
 {
-	FStackCustomAttributes TempAttributes;
+	UE::Anim::FStackAttributeContainer TempAttributes;
 	FAnimationPoseData BaseAnimationPoseData = { BasePose, BaseCurve, TempAttributes };
 	const FAnimationPoseData AdditiveAnimationPoseData(*const_cast<FCompactPose*>(&MeshSpaceRotationAdditive), *const_cast<FBlendedCurve*>(&AdditiveCurve), TempAttributes);
 
@@ -1056,7 +1050,7 @@ void FAnimationRuntime::AccumulateAdditivePose(FAnimationPoseData& BaseAnimation
 	// if curve exists, accumulate with the weight, 
 	BaseAnimationPoseData.GetCurve().Accumulate(AdditiveAnimationPoseData.GetCurve(), Weight);
 
-	FCustomAttributesRuntime::AccumulateAttributes(AdditiveAnimationPoseData.GetAttributes(), BaseAnimationPoseData.GetAttributes(), Weight);
+	UE::Anim::Attributes::AccumulateAttributes(AdditiveAnimationPoseData.GetAttributes(), BaseAnimationPoseData.GetAttributes(), Weight, AdditiveType);
 	
 	// normalize
 	BaseAnimationPoseData.GetPose().NormalizeRotations();
@@ -1412,22 +1406,22 @@ void FAnimationRuntime::BlendPosesPerBoneFilter(
 	EBlendPosesPerBoneFilterFlags BlendFlags,
 	ECurveBlendOption::Type CurveBlendOption)
 {
-	FStackCustomAttributes TempAttributes;
+	UE::Anim::FStackAttributeContainer TempAttributes;
 	FAnimationPoseData AnimationPoseData = {OutPose, OutCurve, TempAttributes };
 	
 	BlendPosesPerBoneFilter(BasePose, BlendPoses, BaseCurve, BlendedCurves, TempAttributes, {}, AnimationPoseData, BoneBlendWeights, BlendFlags, CurveBlendOption);
 }
 
 void FAnimationRuntime::BlendPosesPerBoneFilter(FCompactPose& BasePose, const TArray<FCompactPose>& BlendPoses, FBlendedCurve& BaseCurve, const TArray<FBlendedCurve>& BlendedCurves,
-	FStackCustomAttributes& BaseAttributes,
-	const TArray<FStackCustomAttributes>& BlendAttributes,
+	UE::Anim::FStackAttributeContainer& BaseAttributes,
+	const TArray<UE::Anim::FStackAttributeContainer>& BlendAttributes,
 FAnimationPoseData& OutAnimationPoseData, TArray<FPerBoneBlendWeight>& BoneBlendWeights, EBlendPosesPerBoneFilterFlags BlendFlags, enum ECurveBlendOption::Type CurveBlendOption)
 {
 	SCOPE_CYCLE_COUNTER(STAT_BlendPosesPerBoneFilter);
 
 	FCompactPose& OutPose = OutAnimationPoseData.GetPose();
 	FBlendedCurve& OutCurve = OutAnimationPoseData.GetCurve();
-	FStackCustomAttributes& OutAttributes = OutAnimationPoseData.GetAttributes();
+	UE::Anim::FStackAttributeContainer& OutAttributes = OutAnimationPoseData.GetAttributes();
 
 	
 	// if no blendpose, outpose = basepose
@@ -1770,8 +1764,7 @@ FAnimationPoseData& OutAnimationPoseData, TArray<FPerBoneBlendWeight>& BoneBlend
 	}
 
 	{
-		OutAttributes.CopyFrom(BaseAttributes);
-		FCustomAttributesRuntime::BlendAttributesPerBoneFilter(BlendAttributes, BoneBlendWeights, OutAttributes);
+		UE::Anim::Attributes::BlendAttributesPerBoneFilter(BaseAttributes, BlendAttributes, BoneBlendWeights, OutAttributes);
 	}
 }
 
diff --git a/Engine/Source/Runtime/Engine/Private/Animation/AnimationSettings.cpp b/Engine/Source/Runtime/Engine/Private/Animation/AnimationSettings.cpp
--- a/Engine/Source/Runtime/Engine/Private/Animation/AnimationSettings.cpp
+++ b/Engine/Source/Runtime/Engine/Private/Animation/AnimationSettings.cpp
@@ -12,7 +12,7 @@ UAnimationSettings::UAnimationSettings(const FObjectInitializer& ObjectInitializ
 	, bRaiseMaxErrorToExisting(false)
 	, bEnablePerformanceLog(false)
 	, bTickAnimationOnSkeletalMeshInit(true)
-	, DefaultAttributeBlendMode(ECustomAttributeBlendType::Override)
+	, DefaultAttributeBlendMode(ECustomAttributeBlendType::Blend)
 {
 	SectionName = TEXT("Animation");
 
diff --git a/Engine/Source/Runtime/Engine/Private/Animation/AttributeIdentifier.cpp b/Engine/Source/Runtime/Engine/Private/Animation/AttributeIdentifier.cpp
new file mode 100644
--- /dev/null
+++ b/Engine/Source/Runtime/Engine/Private/Animation/AttributeIdentifier.cpp
@@ -0,0 +1,81 @@
+// Copyright Epic Games, Inc. All Rights Reserved.
+
+#include "Animation/AnimData/AttributeIdentifier.h"
+#include "Animation/AttributeTypes.h"
+#include "Animation/Skeleton.h"
+#include "Animation/AnimSequence.h"
+
+#if WITH_EDITOR
+FAnimationAttributeIdentifier UAnimationAttributeIdentifierExtensions::CreateAttributeIdentifier(UAnimationAsset* AnimationAsset, const FName AttributeName, const FName BoneName, UScriptStruct* AttributeType, bool bValidateExistsOnAsset /*= false*/)
+{
+	FAnimationAttributeIdentifier Identifier;
+
+	// Ensure the type is valid and registered with the attributes system
+	if (AttributeType)
+	{
+		if (UE::Anim::AttributeTypes::IsTypeRegistered(AttributeType))
+		{
+			if (AnimationAsset)
+			{
+				const USkeleton* Skeleton = AnimationAsset->GetSkeleton();
+				if (Skeleton)
+				{
+					// Ensure that the request bone exists on the target asset its skeleton
+					const int32 BoneIndex = Skeleton->GetReferenceSkeleton().FindBoneIndex(BoneName);
+					if (BoneIndex != INDEX_NONE)
+					{
+						Identifier.Name = AttributeName;
+						Identifier.BoneName = BoneName;
+						Identifier.BoneIndex = BoneIndex;
+						Identifier.ScriptStruct = AttributeType;
+						Identifier.ScriptStructPath = AttributeType;
+
+						// If the user requested so, make sure the attribute exists on the asset, and if not reset the identifier to be invalid
+						if (bValidateExistsOnAsset)
+						{
+							UAnimSequence* AnimSequence = Cast<UAnimSequence>(AnimationAsset);
+							if (AnimSequence)
+							{
+								if (AnimSequence->GetDataModel()->FindAttribute(Identifier) == nullptr)
+								{
+									Identifier = FAnimationAttributeIdentifier();
+									UE_LOG(LogAnimation, Warning, TEXT("Attribute %s does not exists on AnimationAsset %s provided for CreateAttributeIdentifier"), *Identifier.ToString(), *AnimationAsset->GetName());
+								}
+							}
+							else
+							{
+								UE_LOG(LogAnimation, Warning, TEXT("Cannot validate if Attribute %s exists on AnimationAsset %s because it is not an AnimationSequence"), *Identifier.ToString(), *AnimationAsset->GetName());
+								Identifier = FAnimationAttributeIdentifier();
+							}
+						}
+
+						return Identifier;
+					}
+					else
+					{
+						UE_LOG(LogAnimation, Warning, TEXT("Bone name %s provided for CreateAttributeIdentifier does not exist on the Skeleton %s"), *BoneName.ToString(), *Skeleton->GetName());
+					}
+				}
+				else
+				{
+					UE_LOG(LogAnimation, Warning, TEXT("Skeleton for provided Animation Asset %s in CreateAttributeIdentifier is null"), *AnimationAsset->GetName());
+				}
+			}
+			else
+			{
+				UE_LOG(LogAnimation, Warning, TEXT("Invalid Animation Asset provided for CreateAttributeIdentifier"));
+			}
+		}
+		else
+		{
+			UE_LOG(LogAnimation, Warning, TEXT("Attribute type %s provided for CreateAttributeIdentifier has not been registered see AttributeTypes::RegisterType"), *AttributeType->GetName());
+		}
+	}
+	else
+	{
+		UE_LOG(LogAnimation, Warning, TEXT("Invalid (null) Attribute type provided for CreateAttributeIdentifier"));
+	}
+
+	return Identifier;
+}
+#endif // WITH_EDITOR
\ No newline at end of file
diff --git a/Engine/Source/Runtime/Engine/Private/Animation/AttributeTypes.cpp b/Engine/Source/Runtime/Engine/Private/Animation/AttributeTypes.cpp
new file mode 100644
--- /dev/null
+++ b/Engine/Source/Runtime/Engine/Private/Animation/AttributeTypes.cpp
@@ -0,0 +1,29 @@
+// Copyright Epic Games, Inc. All Rights Reserved.
+
+#include "Animation/AttributeTypes.h"
+#include "Animation/BuiltInAttributeTypes.h"
+
+namespace UE
+{
+	namespace Anim
+	{		
+		TArray<TWeakObjectPtr<const UScriptStruct>> AttributeTypes::RegisteredTypes;
+		TArray<TUniquePtr<IAttributeBlendOperator>> AttributeTypes::Operators;
+		TArray<TWeakObjectPtr<const UScriptStruct>> AttributeTypes::InterpolatableTypes;
+
+		void AttributeTypes::Initialize()
+		{
+			static bool bInitialized = false;
+			if (!bInitialized)
+			{
+				bInitialized = true;
+
+				RegisterType<FFloatAnimationAttribute>();
+				RegisterType<FIntegerAnimationAttribute>();
+				RegisterType<FStringAnimationAttribute>();
+				RegisterType<FTransformAnimationAttribute>();
+			}
+		}		
+	}
+}
+
diff --git a/Engine/Source/Runtime/Engine/Private/Animation/BlendProfileScratchData.h b/Engine/Source/Runtime/Engine/Private/Animation/BlendProfileScratchData.h
--- a/Engine/Source/Runtime/Engine/Private/Animation/BlendProfileScratchData.h
+++ b/Engine/Source/Runtime/Engine/Private/Animation/BlendProfileScratchData.h
@@ -7,7 +7,7 @@
 
 struct FCompactPose;
 struct FBlendedCurve;
-struct FStackCustomAttributes;
+namespace UE { namespace Anim { struct FStackAttributeContainer; } }
 struct FSlotEvaluationPose;
 
 /**
@@ -30,5 +30,5 @@ struct FBlendProfileScratchData : public TThreadSingleton<FBlendProfileScratchDa
 	TArray<float, TInlineAllocator<8>> BlendingWeights;					// The per pose blend weights.
 	TArray<const FCompactPose*, TInlineAllocator<8>> BlendingPoses;		// The non additive poses to blend.
 	TArray<const FBlendedCurve*, TInlineAllocator<8>> BlendingCurves;	// The curves to blend.
-	TArray<const FStackCustomAttributes*, TInlineAllocator<8>> BlendingAttributes;	// The attributes to blend.
+	TArray<const UE::Anim::FStackAttributeContainer*, TInlineAllocator<8>> BlendingAttributes;	// The attributes to blend.
 };
diff --git a/Engine/Source/Runtime/Engine/Private/Animation/BlendSpace.cpp b/Engine/Source/Runtime/Engine/Private/Animation/BlendSpace.cpp
--- a/Engine/Source/Runtime/Engine/Private/Animation/BlendSpace.cpp
+++ b/Engine/Source/Runtime/Engine/Private/Animation/BlendSpace.cpp
@@ -13,7 +13,7 @@
 #include "Logging/TokenizedMessage.h"
 #include "Logging/MessageLog.h"
 #include "Animation/AnimationPoseData.h"
-#include "Animation/CustomAttributesRuntime.h"
+#include "Animation/AttributesRuntime.h"
 
 #define LOCTEXT_NAMESPACE "BlendSpace"
 
@@ -600,7 +600,7 @@ void UBlendSpace::ResetToRefPose(FCompactPose& OutPose) const
 
 void UBlendSpace::GetAnimationPose(TArray<FBlendSampleData>& BlendSampleDataCache, /*out*/ FCompactPose& OutPose, /*out*/ FBlendedCurve& OutCurve) const
 {
-	FStackCustomAttributes TempAttributes;
+	UE::Anim::FStackAttributeContainer TempAttributes;
 	FAnimationPoseData AnimationPoseData = { OutPose, OutCurve, TempAttributes };
 	GetAnimationPose(BlendSampleDataCache, AnimationPoseData);
 }
@@ -639,7 +639,7 @@ void UBlendSpace::GetAnimationPose_Internal(TArray<FBlendSampleData>& BlendSampl
 	TArray<FBlendedCurve, TInlineAllocator<8>> ChildrenCurves;
 	ChildrenCurves.AddZeroed(NumPoses);
 
-	TArray<FStackCustomAttributes, TInlineAllocator<8>> ChildrenAttributes;
+	TArray<UE::Anim::FStackAttributeContainer, TInlineAllocator<8>> ChildrenAttributes;
 	ChildrenAttributes.AddZeroed(NumPoses);
 
 	TArray<float, TInlineAllocator<8>> ChildrenWeights;
diff --git a/Engine/Source/Runtime/Engine/Private/Animation/PoseAsset.cpp b/Engine/Source/Runtime/Engine/Private/Animation/PoseAsset.cpp
--- a/Engine/Source/Runtime/Engine/Private/Animation/PoseAsset.cpp
+++ b/Engine/Source/Runtime/Engine/Private/Animation/PoseAsset.cpp
@@ -380,7 +380,7 @@ struct FBoneIndices
 
 void UPoseAsset::GetBaseAnimationPose(struct FCompactPose& OutPose, FBlendedCurve& OutCurve) const
 {
-	FStackCustomAttributes TempAttributes;
+	UE::Anim::FStackAttributeContainer TempAttributes;
 	FAnimationPoseData OutPoseData(OutPose, OutCurve, TempAttributes);
 	GetBaseAnimationPose(OutPoseData);
 }
@@ -542,7 +542,7 @@ void UPoseAsset::GetAnimationCurveOnly(TArray<FName>& InCurveNames, TArray<float
 
 bool UPoseAsset::GetAnimationPose(struct FCompactPose& OutPose, FBlendedCurve& OutCurve, const FAnimExtractContext& ExtractionContext) const
 {
-	FStackCustomAttributes TempAttributes;
+	UE::Anim::FStackAttributeContainer TempAttributes;
 	FAnimationPoseData OutPoseData(OutPose, OutCurve, TempAttributes);
 	return GetAnimationPose(OutPoseData, ExtractionContext);
 }
