New Animation Attributes system, replacing Custom Attributes:

diff --git a/Engine/Source/Runtime/Engine/Private/Components/SkeletalMeshComponent.cpp b/Engine/Source/Runtime/Engine/Private/Components/SkeletalMeshComponent.cpp
--- a/Engine/Source/Runtime/Engine/Private/Components/SkeletalMeshComponent.cpp
+++ b/Engine/Source/Runtime/Engine/Private/Components/SkeletalMeshComponent.cpp
@@ -41,6 +41,8 @@
 #include "SkeletalRenderPublic.h"
 #include "ContentStreaming.h"
 #include "Animation/AnimTrace.h"
+#include "Animation/BuiltInAttributeTypes.h"
+
 #if INTEL_ISPC
 #include "SkeletalMeshComponent.ispc.generated.h"
 #endif
@@ -1863,7 +1865,7 @@ bool USkeletalMeshComponent::AreRequiredCurvesUpToDate() const
 	return (!SkeletalMesh || !SkeletalMesh->GetSkeleton() || CachedAnimCurveUidVersion == SkeletalMesh->GetSkeleton()->GetAnimCurveUidVersion());
 }
 
-void USkeletalMeshComponent::EvaluateAnimation(const USkeletalMesh* InSkeletalMesh, UAnimInstance* InAnimInstance, FVector& OutRootBoneTranslation, FBlendedHeapCurve& OutCurve, FCompactPose& OutPose, FHeapCustomAttributes& OutAttributes) const
+void USkeletalMeshComponent::EvaluateAnimation(const USkeletalMesh* InSkeletalMesh, UAnimInstance* InAnimInstance, FVector& OutRootBoneTranslation, FBlendedHeapCurve& OutCurve, FCompactPose& OutPose, UE::Anim::FHeapAttributeContainer& OutAttributes) const
 {
 	ANIM_MT_SCOPE_CYCLE_COUNTER(SkeletalComponentAnimEvaluate, !IsInGameThread());
 
@@ -1925,20 +1927,20 @@ void USkeletalMeshComponent::UpdateSlaveComponent()
 
 #if WITH_EDITOR
 
-void USkeletalMeshComponent::PerformAnimationEvaluation(const USkeletalMesh* InSkeletalMesh, UAnimInstance* InAnimInstance, TArray<FTransform>& OutSpaceBases, TArray<FTransform>& OutBoneSpaceTransforms, FVector& OutRootBoneTranslation, FBlendedHeapCurve& OutCurve, FHeapCustomAttributes& OutAttributes)
+void USkeletalMeshComponent::PerformAnimationEvaluation(const USkeletalMesh* InSkeletalMesh, UAnimInstance* InAnimInstance, TArray<FTransform>& OutSpaceBases, TArray<FTransform>& OutBoneSpaceTransforms, FVector& OutRootBoneTranslation, FBlendedHeapCurve& OutCurve, UE::Anim::FHeapAttributeContainer& OutAttributes)
 {
 	PerformAnimationProcessing(InSkeletalMesh, InAnimInstance, true, OutSpaceBases, OutBoneSpaceTransforms, OutRootBoneTranslation, OutCurve, OutAttributes);
 }
 
 void USkeletalMeshComponent::PerformAnimationEvaluation(const USkeletalMesh* InSkeletalMesh, UAnimInstance* InAnimInstance, TArray<FTransform>& OutSpaceBases, TArray<FTransform>& OutBoneSpaceTransforms, FVector& OutRootBoneTranslation, FBlendedHeapCurve& OutCurve)
 {
-	FHeapCustomAttributes Attributes;
+	UE::Anim::FHeapAttributeContainer Attributes;
 	PerformAnimationEvaluation(InSkeletalMesh, InAnimInstance, OutSpaceBases, OutBoneSpaceTransforms, OutRootBoneTranslation, OutCurve, Attributes);
 }
 
 #endif
 
-void USkeletalMeshComponent::PerformAnimationProcessing(const USkeletalMesh* InSkeletalMesh, UAnimInstance* InAnimInstance, bool bInDoEvaluation, TArray<FTransform>& OutSpaceBases, TArray<FTransform>& OutBoneSpaceTransforms, FVector& OutRootBoneTranslation, FBlendedHeapCurve& OutCurve, FHeapCustomAttributes& OutAttributes)
+void USkeletalMeshComponent::PerformAnimationProcessing(const USkeletalMesh* InSkeletalMesh, UAnimInstance* InAnimInstance, bool bInDoEvaluation, TArray<FTransform>& OutSpaceBases, TArray<FTransform>& OutBoneSpaceTransforms, FVector& OutRootBoneTranslation, FBlendedHeapCurve& OutCurve, UE::Anim::FHeapAttributeContainer& OutAttributes)
 {
 	CSV_SCOPED_TIMING_STAT(Animation, WorkerThreadTickTime);
 	ANIM_MT_SCOPE_CYCLE_COUNTER(PerformAnimEvaluation, !IsInGameThread());
@@ -1982,17 +1984,17 @@ void USkeletalMeshComponent::PerformAnimationProcessing(const USkeletalMesh* InS
 
 void USkeletalMeshComponent::PerformAnimationProcessing(const USkeletalMesh* InSkeletalMesh, UAnimInstance* InAnimInstance, bool bInDoEvaluation, TArray<FTransform>& OutSpaceBases, TArray<FTransform>& OutBoneSpaceTransforms, FVector& OutRootBoneTranslation, FBlendedHeapCurve& OutCurve)
 {
-	FHeapCustomAttributes Attributes;	
+	UE::Anim::FHeapAttributeContainer Attributes;	
 	PerformAnimationProcessing(InSkeletalMesh, InAnimInstance, bInDoEvaluation, OutSpaceBases, OutBoneSpaceTransforms, OutRootBoneTranslation, OutCurve, Attributes);
 }
 
 void USkeletalMeshComponent::EvaluatePostProcessMeshInstance(TArray<FTransform>& OutBoneSpaceTransforms, FCompactPose& InOutPose, FBlendedHeapCurve& OutCurve, const USkeletalMesh* InSkeletalMesh, FVector& OutRootBoneTranslation) const
 {
-	FHeapCustomAttributes Attributes;
+	UE::Anim::FHeapAttributeContainer Attributes;
 	EvaluatePostProcessMeshInstance(OutBoneSpaceTransforms, InOutPose, OutCurve, InSkeletalMesh, OutRootBoneTranslation, Attributes);
 }
 
-void USkeletalMeshComponent::EvaluatePostProcessMeshInstance(TArray<FTransform>& OutBoneSpaceTransforms, FCompactPose& InOutPose, FBlendedHeapCurve& OutCurve, const USkeletalMesh* InSkeletalMesh, FVector& OutRootBoneTranslation, FHeapCustomAttributes& OutAttributes) const
+void USkeletalMeshComponent::EvaluatePostProcessMeshInstance(TArray<FTransform>& OutBoneSpaceTransforms, FCompactPose& InOutPose, FBlendedHeapCurve& OutCurve, const USkeletalMesh* InSkeletalMesh, FVector& OutRootBoneTranslation, UE::Anim::FHeapAttributeContainer& OutAttributes) const
 {
 	if (ShouldEvaluatePostProcessInstance())
 	{
@@ -2568,7 +2570,7 @@ void USkeletalMeshComponent::PostAnimEvaluation(FAnimationEvaluationContext& Eva
 			AnimCurves.LerpTo(CachedCurve, Alpha);
 
 			// Interpolate custom attributes
-			FCustomAttributesRuntime::InterpolateAttributes(CachedAttributes, CustomAttributes, Alpha);
+			UE::Anim::Attributes::InterpolateAttributes(CustomAttributes, CachedAttributes, Alpha);
 		}
 	}
 
@@ -3867,7 +3869,7 @@ void USkeletalMeshComponent::ParallelDuplicateAndInterpolate(FAnimationEvaluatio
 				InAnimEvaluationContext.Curve.LerpTo(InAnimEvaluationContext.CachedCurve, Alpha);
 			}
 
-			FCustomAttributesRuntime::InterpolateAttributes(InAnimEvaluationContext.CachedCustomAttributes, InAnimEvaluationContext.CustomAttributes, Alpha);
+			UE::Anim::Attributes::InterpolateAttributes(InAnimEvaluationContext.CustomAttributes, InAnimEvaluationContext.CachedCustomAttributes, Alpha);
 		}
 	}
 }
@@ -4332,55 +4334,66 @@ const TArray<FTransform>& USkeletalMeshComponent::GetCachedComponentSpaceTransfo
 
 bool USkeletalMeshComponent::GetFloatAttribute_Ref(const FName& BoneName, const FName& AttributeName, float& OutValue, ECustomBoneAttributeLookup LookupType /*= ECustomBoneAttributeLookup::BoneOnly*/)
 {
-	return GetBoneAttribute<float>(BoneName, AttributeName, OutValue, OutValue, LookupType);
+	return FindAttributeChecked<float, FFloatAnimationAttribute>(BoneName, AttributeName, OutValue, OutValue, LookupType);
+}
+
+bool USkeletalMeshComponent::GetTransformAttribute_Ref(const FName& BoneName, const FName& AttributeName, FTransform& OutValue, ECustomBoneAttributeLookup LookupType /*= ECustomBoneAttributeLookup::BoneOnly*/)
+{
+	bool bResult = FindAttributeChecked<FTransform, FTransformAnimationAttribute>(BoneName, AttributeName, OutValue, OutValue, LookupType);
+	return bResult;
 }
 
 bool USkeletalMeshComponent::GetIntegerAttribute_Ref(const FName& BoneName, const FName& AttributeName, int32& OutValue, ECustomBoneAttributeLookup LookupType /*= ECustomBoneAttributeLookup::BoneOnly*/)
 {
-	return GetBoneAttribute<int32>(BoneName, AttributeName, OutValue, OutValue, LookupType);
+	return FindAttributeChecked<int32, FIntegerAnimationAttribute>(BoneName, AttributeName, OutValue, OutValue, LookupType);
 }
 
 bool USkeletalMeshComponent::GetStringAttribute_Ref(const FName& BoneName, const FName& AttributeName, FString& OutValue, ECustomBoneAttributeLookup LookupType /*= ECustomBoneAttributeLookup::BoneOnly*/)
 {
-	return GetBoneAttribute<FString>(BoneName, AttributeName, OutValue, OutValue, LookupType);
+	return FindAttributeChecked<FString, FStringAnimationAttribute>(BoneName, AttributeName, OutValue, OutValue, LookupType);
 }
 
 bool USkeletalMeshComponent::GetFloatAttribute(const FName& BoneName, const FName& AttributeName, float DefaultValue, float& OutValue, ECustomBoneAttributeLookup LookupType /*= ECustomBoneAttributeLookup::BoneOnly*/)
 {
-	return GetBoneAttribute<float>(BoneName, AttributeName, DefaultValue, OutValue, LookupType);
+	return FindAttributeChecked<float, FFloatAnimationAttribute>(BoneName, AttributeName, DefaultValue, OutValue, LookupType);
+}
+
+bool USkeletalMeshComponent::GetTransformAttribute(const FName& BoneName, const FName& AttributeName, FTransform DefaultValue, FTransform& OutValue, ECustomBoneAttributeLookup LookupType /*= ECustomBoneAttributeLookup::BoneOnly*/)
+{
+	return FindAttributeChecked<FTransform, FTransformAnimationAttribute>(BoneName, AttributeName, DefaultValue, OutValue, LookupType);
 }
 
 bool USkeletalMeshComponent::GetIntegerAttribute(const FName& BoneName, const FName& AttributeName, int32 DefaultValue, int32& OutValue, ECustomBoneAttributeLookup LookupType /*= ECustomBoneAttributeLookup::BoneOnly*/)
 {
-	return GetBoneAttribute<int32>(BoneName, AttributeName, DefaultValue, OutValue, LookupType);
+	return FindAttributeChecked<int32, FIntegerAnimationAttribute>(BoneName, AttributeName, DefaultValue, OutValue, LookupType);
 }
 
 bool USkeletalMeshComponent::GetStringAttribute(const FName& BoneName, const FName& AttributeName, FString DefaultValue, FString& OutValue, ECustomBoneAttributeLookup LookupType /*= ECustomBoneAttributeLookup::BoneOnly*/)
 {
-	return GetBoneAttribute<FString>(BoneName, AttributeName, DefaultValue, OutValue, LookupType);
+	return FindAttributeChecked<FString, FStringAnimationAttribute>(BoneName, AttributeName, DefaultValue, OutValue, LookupType);
 }
 
-template<typename DataType>
-bool USkeletalMeshComponent::GetBoneAttribute(const FName& BoneName, const FName& AttributeName, DataType DefaultValue, DataType& OutValue, ECustomBoneAttributeLookup LookupType)
+template<typename DataType, typename CustomAttributeType>
+bool USkeletalMeshComponent::FindAttributeChecked(const FName& BoneName, const FName& AttributeName, DataType DefaultValue, DataType& OutValue, ECustomBoneAttributeLookup LookupType)
 {
 	OutValue = DefaultValue;
 	bool bFound = false;
 
 	if (SkeletalMesh)
 	{
-		const FHeapCustomAttributes& Attributes = GetCustomAttributes();
+		const UE::Anim::FHeapAttributeContainer& Attributes = GetCustomAttributes();
 		const int32 BoneIndex = SkeletalMesh->GetRefSkeleton().FindBoneIndex(BoneName);
 
-		bFound = Attributes.GetBoneAttribute(FCompactPoseBoneIndex(BoneIndex), AttributeName, OutValue);
+		const CustomAttributeType* AttributePtr = Attributes.Find<CustomAttributeType>(UE::Anim::FAttributeId(AttributeName, FCompactPoseBoneIndex(BoneIndex)));
 
-		if (!bFound && LookupType != ECustomBoneAttributeLookup::BoneOnly)
+		if (AttributePtr == nullptr && LookupType != ECustomBoneAttributeLookup::BoneOnly)
 		{
 			if (LookupType == ECustomBoneAttributeLookup::ImmediateParent)
 			{
 				const int32 ParentIndex = SkeletalMesh->GetRefSkeleton().GetParentIndex(BoneIndex);
 				if (ParentIndex != INDEX_NONE)
 				{
-					bFound = Attributes.GetBoneAttribute(FCompactPoseBoneIndex(ParentIndex), AttributeName, OutValue);
+					AttributePtr = Attributes.Find<CustomAttributeType>(UE::Anim::FAttributeId(AttributeName, FCompactPoseBoneIndex(ParentIndex)));
 				}
 			}
 			else if (LookupType == ECustomBoneAttributeLookup::ParentHierarchy)
@@ -4390,8 +4403,8 @@ bool USkeletalMeshComponent::GetBoneAttribute(const FName& BoneName, const FName
 
 				while (ParentIndex != INDEX_NONE)
 				{
-					bFound = Attributes.GetBoneAttribute(FCompactPoseBoneIndex(ParentIndex), AttributeName, OutValue);
-					if (bFound)
+					AttributePtr = Attributes.Find<CustomAttributeType>(UE::Anim::FAttributeId(AttributeName, FCompactPoseBoneIndex(ParentIndex)));
+					if (AttributePtr)
 					{
 						break;
 					}
@@ -4401,6 +4414,12 @@ bool USkeletalMeshComponent::GetBoneAttribute(const FName& BoneName, const FName
 				}
 			}
 		}
+
+		if (AttributePtr != nullptr)
+		{
+			OutValue = AttributePtr->Value;
+			bFound = true;
+		}
 	}
 
 	return bFound;
diff --git a/Engine/Source/Runtime/Engine/Public/Animation/AttributeTraits.h b/Engine/Source/Runtime/Engine/Public/Animation/AttributeTraits.h
new file mode 100644
--- /dev/null
+++ b/Engine/Source/Runtime/Engine/Public/Animation/AttributeTraits.h
@@ -0,0 +1,79 @@
+// Copyright Epic Games, Inc. All Rights Reserved.
+
+#pragma once
+
+#include "CoreMinimal.h"
+
+namespace UE
+{
+	namespace Anim
+	{
+		/** Set of type-traits, used by the Animation Attributes system to verify and implement certain behavior */
+		template <class AttributeType>
+		struct TAttributeTypeTraitsBase
+		{
+			enum
+			{
+				/** Determines whether or not the type should be blended, and is supported to do so. True by default */
+				IsBlendable = true,
+				/** Determines whether or not the type has an associated user-defined implementation of IAttributeBlendOperator. False by default */
+				WithCustomBlendOperator = false,
+				/** Determines whether or not the type should be step-interpolated rather than linearly. False by default */
+				StepInterpolate = false
+			};
+		};
+
+		template <class AttributeType>
+		struct TAttributeTypeTraits : public TAttributeTypeTraitsBase<AttributeType>
+		{
+		};
+		
+		/** Implemented functionality for retrieving a types trait value for platforms without constexpr support */
+#if !PLATFORM_COMPILER_HAS_IF_CONSTEXPR
+		/**
+		 * Selection of IsBlendableType call.
+		 */
+		template<class AttributeType>
+		FORCEINLINE typename TEnableIf<!TAttributeTypeTraits<AttributeType>::IsBlendable, bool>::Type IsBlendableType()
+		{
+			return false;
+		}
+
+		template<class AttributeType>
+		FORCEINLINE typename TEnableIf<TAttributeTypeTraits<AttributeType>::IsBlendable, bool>::Type IsBlendableType()
+		{
+			return true;
+		}
+
+		/**
+		 * Selection of IsStepInterpolatedType call.
+		 */
+		template<class AttributeType>
+		FORCEINLINE typename TEnableIf<!TAttributeTypeTraits<AttributeType>::StepInterpolate, bool>::Type IsStepInterpolatedType()
+		{
+			return false;
+		}
+
+		template<class AttributeType>
+		FORCEINLINE typename TEnableIf<TAttributeTypeTraits<AttributeType>::StepInterpolate, bool>::Type IsStepInterpolatedType()
+		{
+			return true;
+		}
+
+		/**
+		 * Selection of CustomBlendedType call.
+		 */
+		template<class AttributeType>
+		FORCEINLINE typename TEnableIf<!TAttributeTypeTraits<AttributeType>::WithCustomBlendOperator, bool>::Type CustomBlendedType()
+		{
+			return false;
+		}
+
+		template<class AttributeType>
+		FORCEINLINE typename TEnableIf<TAttributeTypeTraits<AttributeType>::WithCustomBlendOperator, bool>::Type CustomBlendedType()
+		{
+			return true;
+		}
+#endif
+	}
+}
diff --git a/Engine/Source/Runtime/Engine/Public/Animation/AttributesRuntime.h b/Engine/Source/Runtime/Engine/Public/Animation/AttributesRuntime.h
new file mode 100644
--- /dev/null
+++ b/Engine/Source/Runtime/Engine/Public/Animation/AttributesRuntime.h
@@ -0,0 +1,95 @@
+// Copyright Epic Games, Inc. All Rights Reserved.
+
+#pragma once
+
+#include "Containers/ArrayView.h"
+#include "Animation/AnimTypes.h"
+#include "UObject/NameTypes.h"
+#include "Containers/ContainersFwd.h"
+
+#include "Containers/Map.h"
+#include "Misc/AssertionMacros.h"
+#include "Misc/CoreMiscDefines.h"
+#include "UObject/StructOnScope.h"
+
+#include "Animation/IAttributeBlendOperator.h"
+
+#include "UObject/Class.h"
+#include "UObject/WeakObjectPtrTemplates.h"
+
+#include "Animation/AttributesContainer.h"
+#include "Animation/WrappedAttribute.h"
+
+struct FCompactPoseBoneIndex;
+struct FBlendSampleData;
+struct FBoneContainer;
+
+enum class ECustomAttributeBlendType : uint8;
+enum EAdditiveAnimationType;
+
+struct FAnimatedBoneAttribute;
+struct FCompactPoseBoneIndex;
+struct FAnimExtractContext;
+
+namespace UE
+{
+	namespace Anim
+	{
+		struct ENGINE_API FStackAttributeContainer : public TAttributeContainer<FCompactPoseBoneIndex, FAnimStackAllocator> {};
+		struct ENGINE_API FHeapAttributeContainer : public TAttributeContainer<FCompactPoseBoneIndex, FDefaultAllocator> {};
+
+		/** Helper functionality for attributes animation runtime */
+		struct ENGINE_API Attributes
+		{
+#if WITH_EDITOR
+			static void GetAttributeValue(FStackAttributeContainer& OutAttributes, const FCompactPoseBoneIndex& PoseBoneIndex, const  FAnimatedBoneAttribute& Attribute, const FAnimExtractContext& ExtractionContext);
+#endif // WITH_EDITOR
+
+			/** Blend custom attribute values from N set of inputs */
+			static void BlendAttributes(const TArrayView<const FStackAttributeContainer> SourceAttributes, const TArrayView<const float> SourceWeights, FStackAttributeContainer& OutAttributes);
+
+			/** Blend custom attribute values from N set of inputs (ptr-values) */
+			static void BlendAttributes(const TArrayView<const FStackAttributeContainer* const> SourceAttributes, const TArrayView<const float> SourceWeights, FStackAttributeContainer& OutAttributes);
+
+			/** Blend custom attribute values from N set of inputs, using input weight remapping */
+			static void BlendAttributes(const TArrayView<const FStackAttributeContainer> SourceAttributes, const TArrayView<const float> SourceWeights, const TArrayView<const int32> SourceWeightsIndices, FStackAttributeContainer& OutAttributes);
+
+			/* Blend custom attribute values from 2 inputs, using per-bone weights */
+			static void BlendAttributesPerBone(const FStackAttributeContainer& SourceAttributes1, const FStackAttributeContainer& SourceAttributes2, const TArrayView<const float> WeightsOfSource2, FStackAttributeContainer& OutAttributes);
+
+			/* Blend custom attribute values from N set of inputs, using N number of blend samples */
+			static void BlendAttributesPerBone(TArrayView<const FStackAttributeContainer> SourceAttributes, TArrayView<const int32> PerBoneInterpolationIndices, TArrayView<const FBlendSampleData> BlendSampleDataCache, FStackAttributeContainer& OutAttributes);
+
+			/* Blend custom attribute values from N set of inputs, using N number of blend samples */
+			static void BlendAttributesPerBone(TArrayView<const FStackAttributeContainer> SourceAttributes, TArrayView<const int32> PerBoneInterpolationIndices, TArrayView<const FBlendSampleData> BlendSampleDataCache, TArrayView<const int32> BlendSampleDataCacheIndices, FStackAttributeContainer& OutAttributes);
+
+			/* Blend custom attribute values from N set of inputs, using bone filter pose weights */
+			static void BlendAttributesPerBoneFilter(const FStackAttributeContainer& BaseAttributes, const TArrayView<const FStackAttributeContainer> BlendAttributes, const TArray<FPerBoneBlendWeight>& BoneBlendWeights, FStackAttributeContainer& OutAttributes);
+
+			/** Add any new or override existing custom attributes */
+			static void OverrideAttributes(const FStackAttributeContainer& SourceAttributes, FStackAttributeContainer& OutAttributes, float Weight);
+
+			/** Add any new or accumulate with existing custom attributes */
+			static void AccumulateAttributes(const FStackAttributeContainer& SourceAttributes, FStackAttributeContainer& OutAttributes, float Weight, EAdditiveAnimationType AdditiveType);
+
+			/** Add (negated) any new or subtract from existing custom attributes */
+			static void ConvertToAdditive(const FStackAttributeContainer& BaseAttributes, FStackAttributeContainer& OutAdditiveAttributes);
+
+			/** Copy attributes from source, and remap the bone indices according to BoneMapToSource */
+			static void CopyAndRemapAttributes(const FHeapAttributeContainer& SourceAttributes, FStackAttributeContainer& OutAttributes, const TMap<int32, int32>& BoneMapToSource, const FBoneContainer& RequiredBones);
+
+			/** Interpolates between two sets of attributes */
+			static void InterpolateAttributes(FHeapAttributeContainer& FromAttributes, const FHeapAttributeContainer& ToAttributes, float Alpha);
+
+			/** Helper functionality to retrieve the correct blend type (from UAnimationSettings) for the provided attribute name */
+			static ECustomAttributeBlendType GetAttributeBlendType(const FName& InName);
+
+			/** Helper functionality to retrieve the correct blend type (from UAnimationSettings) for the provided attribute info */
+			static ECustomAttributeBlendType GetAttributeBlendType(const FAttributeId& Info);
+		};
+	}
+}
+
+struct ENGINE_API UE_DEPRECATED(5.0, "FStackCustomAttributes has been deprecated use UE::Anim::FStackAttributeContainer instead") FStackCustomAttributes : public UE::Anim::FStackAttributeContainer {};
+struct ENGINE_API UE_DEPRECATED(5.0, "FStackCustomAttributes has been deprecated use UE::Anim::FStackAttributeContainer instead") FHeapCustomAttributes : public UE::Anim::FHeapAttributeContainer {};
+
diff --git a/Engine/Source/Runtime/Engine/Public/AnimationRuntime.h b/Engine/Source/Runtime/Engine/Public/AnimationRuntime.h
--- a/Engine/Source/Runtime/Engine/Public/AnimationRuntime.h
+++ b/Engine/Source/Runtime/Engine/Public/AnimationRuntime.h
@@ -24,7 +24,8 @@ struct FA2CSPose;
 struct FA2Pose;
 struct FInputBlendPose;
 struct FAnimationPoseData;
-struct FStackCustomAttributes;
+
+namespace UE { namespace Anim { struct FStackAttributeContainer; }}
 
 typedef TArray<FTransform> FTransformArrayA2;
 
@@ -104,7 +105,7 @@ public:
 	static void BlendPosesTogether(
 		TArrayView<const FCompactPose> SourcePoses,
 		TArrayView<const FBlendedCurve> SourceCurves,
-		TArrayView<const FStackCustomAttributes> SourceAttributes,
+		TArrayView<const UE::Anim::FStackAttributeContainer> SourceAttributes,
 		TArrayView<const float> SourceWeights,
 		FAnimationPoseData& OutAnimationPoseData);
 
@@ -131,7 +132,7 @@ public:
 	static void BlendPosesTogether(
 		TArrayView<const FCompactPose> SourcePoses,
 		TArrayView<const FBlendedCurve> SourceCurves,
-		TArrayView<const FStackCustomAttributes> SourceAttributes,
+		TArrayView<const UE::Anim::FStackAttributeContainer> SourceAttributes,
 		TArrayView<const float> SourceWeights,
 		TArrayView<const int32> SourceWeightsIndices,
 		/*out*/ FAnimationPoseData& OutPoseData
@@ -157,7 +158,7 @@ public:
 	static void BlendPosesTogetherIndirect(
 		TArrayView<const FCompactPose* const> SourcePoses,
 		TArrayView<const FBlendedCurve* const> SourceCurves,
-		TArrayView<const FStackCustomAttributes* const> SourceAttributes,
+		TArrayView<const UE::Anim::FStackAttributeContainer* const> SourceAttributes,
 		TArrayView<const float> SourceWeights,
 		FAnimationPoseData& OutPoseData);
 
@@ -229,7 +230,7 @@ public:
 	static void BlendPosesTogetherPerBone(
 		TArrayView<const FCompactPose> SourcePoses,
 		TArrayView<const FBlendedCurve> SourceCurves,
-		TArrayView<const FStackCustomAttributes> SourceAttributes,
+		TArrayView<const UE::Anim::FStackAttributeContainer> SourceAttributes,
 		const IInterpolationIndexProvider* InterpolationIndexProvider,
 		TArrayView<const FBlendSampleData> BlendSampleDataCache,
 		/*out*/ FAnimationPoseData& OutAnimationPoseData);
@@ -255,7 +256,7 @@ public:
 	static void BlendPosesTogetherPerBone(
 		TArrayView<const FCompactPose> SourcePoses,
 		TArrayView<const FBlendedCurve> SourceCurves,
-		TArrayView<const FStackCustomAttributes> SourceAttributes,
+		TArrayView<const UE::Anim::FStackAttributeContainer> SourceAttributes,
 		const IInterpolationIndexProvider* InterpolationIndexProvider,
 		TArrayView<const FBlendSampleData> BlendSampleDataCache,
 		TArrayView<const int32> BlendSampleDataCacheIndices,
@@ -281,7 +282,7 @@ public:
 	static void BlendPosesTogetherPerBoneInMeshSpace(
 		TArrayView< FCompactPose> SourcePoses,
 		TArrayView<const FBlendedCurve> SourceCurves,
-		TArrayView<const FStackCustomAttributes> SourceAttributes,	
+		TArrayView<const UE::Anim::FStackAttributeContainer> SourceAttributes,	
 		const UBlendSpace* BlendSpace,
 		TArrayView<const FBlendSampleData> BlendSampleDataCache,
 		/*out*/ FAnimationPoseData& OutAnimationPoseData);
@@ -324,8 +325,8 @@ public:
 		const TArray<FCompactPose>& BlendPoses,
 		FBlendedCurve& BaseCurve,
 		const TArray<FBlendedCurve>& BlendCurves,
-		FStackCustomAttributes& CustomAttributes,
-		const TArray<FStackCustomAttributes>& BlendAttributes,
+		UE::Anim::FStackAttributeContainer& CustomAttributes,
+		const TArray<UE::Anim::FStackAttributeContainer>& BlendAttributes,
 		FAnimationPoseData& OutAnimationPoseData,
 		TArray<FPerBoneBlendWeight>& BoneBlendWeights,
 		EBlendPosesPerBoneFilterFlags blendFlags,
diff --git a/Engine/Source/Runtime/MovieSceneTracks/Private/Tracks/MovieSceneSkeletalAnimationTrack.cpp b/Engine/Source/Runtime/MovieSceneTracks/Private/Tracks/MovieSceneSkeletalAnimationTrack.cpp
--- a/Engine/Source/Runtime/MovieSceneTracks/Private/Tracks/MovieSceneSkeletalAnimationTrack.cpp
+++ b/Engine/Source/Runtime/MovieSceneTracks/Private/Tracks/MovieSceneSkeletalAnimationTrack.cpp
@@ -13,12 +13,12 @@
 #include "Animation/AnimInstance.h"
 #include "Animation/AnimSequence.h"
 #include "Animation/AnimationPoseData.h"
-#include "Animation/CustomAttributesRuntime.h"
+#include "Animation/AttributesRuntime.h"
 #include "SkeletalDebugRendering.h"
 #include "Rendering/SkeletalMeshRenderData.h"
 #include "BoneContainer.h"
 #include "Animation/AnimationPoseData.h"
-#include "Animation/CustomAttributesRuntime.h"
+#include "Animation/AttributesRuntime.h"
 
 #define LOCTEXT_NAMESPACE "MovieSceneSkeletalAnimationTrack"
 
@@ -344,7 +344,7 @@ static void CalculateDistanceMap(USkeletalMeshComponent* SkelMeshComp, UAnimSequ
 	SecondAnimPose.ResetToRefPose(SkelMeshComp->GetAnimInstance()->GetRequiredBones());
 
 	FBlendedCurve FirstOutCurve, SecondOutCurve;
-	FStackCustomAttributes FirstTempAttributes, SecondTempAttributes;
+	UE::Anim::FStackAttributeContainer FirstTempAttributes, SecondTempAttributes;
 	FAnimationPoseData FirstAnimPoseData(FirstAnimPose, FirstOutCurve, FirstTempAttributes);
 	FAnimationPoseData SecondAnimPoseData(SecondAnimPose, SecondOutCurve, SecondTempAttributes);
 
@@ -560,7 +560,7 @@ void UMovieSceneSkeletalAnimationTrack::SetUpRootMotions(bool bForce)
 					{
 						UMovieSceneSkeletalAnimationSection* AnimSection = CastChecked<UMovieSceneSkeletalAnimationSection>(Section);
 
-						FStackCustomAttributes TempAttributes;
+						UE::Anim::FStackAttributeContainer TempAttributes;
 						FAnimationPoseData AnimationPoseData(OutPose, OutCurve, TempAttributes);
 						bool bIsAdditive = false;
 						if (AnimSection->GetRootMotionTransform(FrameNumber.FrameNumber, TickResolution, AnimationPoseData, bIsAdditive, CurrentTransform, CurrentWeight))
