Improve the logging when waiting on async compilation of textures and meshes.
Only log something when we have waited for more than 16 ms.
Log the message before/while we are waiting for a job instead of after it (makes the count info accurate).
Log the package names instead of the object names.

diff --git a/Engine/Source/Runtime/Engine/Private/AsyncCompilationHelpers.cpp b/Engine/Source/Runtime/Engine/Private/AsyncCompilationHelpers.cpp
--- a/Engine/Source/Runtime/Engine/Private/AsyncCompilationHelpers.cpp
+++ b/Engine/Source/Runtime/Engine/Private/AsyncCompilationHelpers.cpp
@@ -171,12 +171,18 @@ namespace AsyncCompilationHelpers
 			FText Progress = FormatProgress(NumDone++, Num, Job->GetName());
 
 			// Be nice with the game thread and tick the progress at 60 fps even when no progress is being made...
+			bool bLogSlowProgress = true;
 			while (!PendingTask.Event->Wait(16))
 			{
 				if (SlowTask.IsSet())
 				{
 					SlowTask->EnterProgressFrame(0.0f, Progress);
 				}
+				if (bLogSlowProgress)
+				{
+					UE_LOG_REF(LogCategory, Display, TEXT("%s"), *Progress.ToString());
+					bLogSlowProgress = false;
+				}
 			}
 
 			if (SlowTask.IsSet())
@@ -184,8 +190,6 @@ namespace AsyncCompilationHelpers
 				SlowTask->EnterProgressFrame(1.0f, Progress);
 			}
 
-			UE_LOG_REF(LogCategory, Display, TEXT("%s"), *Progress.ToString());
-
 			PostCompileSingle(Job);
 		}
 
diff --git a/Engine/Source/Runtime/Engine/Private/StaticMeshCompiler.cpp b/Engine/Source/Runtime/Engine/Private/StaticMeshCompiler.cpp
--- a/Engine/Source/Runtime/Engine/Private/StaticMeshCompiler.cpp
+++ b/Engine/Source/Runtime/Engine/Private/StaticMeshCompiler.cpp
@@ -312,7 +312,7 @@ void FStaticMeshCompilingManager::FinishCompilation(TArrayView<UStaticMesh* cons
 
 			TStrongObjectPtr<UStaticMesh> StaticMesh;
 			void EnsureCompletion() override { if (StaticMesh->AsyncTask) { StaticMesh->AsyncTask->EnsureCompletion(); } }
-			FName GetName() override { return StaticMesh->GetFName(); }
+			FName GetName() override { return StaticMesh->GetOutermost()->GetFName(); }
 		};
 
 		TArray<FCompilableStaticMesh> CompilableStaticMeshes(PendingStaticMeshes);
diff --git a/Engine/Source/Runtime/Engine/Private/TextureCompiler.cpp b/Engine/Source/Runtime/Engine/Private/TextureCompiler.cpp
--- a/Engine/Source/Runtime/Engine/Private/TextureCompiler.cpp
+++ b/Engine/Source/Runtime/Engine/Private/TextureCompiler.cpp
@@ -259,7 +259,7 @@ void FTextureCompilingManager::FinishCompilation(TArrayView<UTexture* const> InT
 
 			TStrongObjectPtr<UTexture> Texture;
 			void EnsureCompletion() override { Texture->FinishCachePlatformData(); }
-			FName GetName() override { return Texture->GetFName(); }
+			FName GetName() override { return Texture->GetOutermost()->GetFName(); }
 		};
 
 		TArray<UTexture*> UniqueTextures(PendingTextures.Array());
