Autoderiv: Removed unnecessary SelectElemHelper function

diff --git a/Engine/Shaders/Public/ShaderVersion.ush b/Engine/Shaders/Public/ShaderVersion.ush
--- a/Engine/Shaders/Public/ShaderVersion.ush
+++ b/Engine/Shaders/Public/ShaderVersion.ush
@@ -3,4 +3,4 @@
 // in Platform.ush (which should be included in any shader) it allows to invalidate the shader DDC.
 //
 // If you are merging streams and there is a conflict with this GUID you should make a new GUID rather than taking one or the other.
-// GUID = D4BD6A39-2B15-485A-8DAF-15CA2FCBB666
+// GUID = D4BD6A39-2B29-485A-8DBA-25CA2FCBB666
diff --git a/Engine/Source/Runtime/Engine/Private/Materials/HLSLMaterialDerivativeAutogen.cpp b/Engine/Source/Runtime/Engine/Private/Materials/HLSLMaterialDerivativeAutogen.cpp
--- a/Engine/Source/Runtime/Engine/Private/Materials/HLSLMaterialDerivativeAutogen.cpp
+++ b/Engine/Source/Runtime/Engine/Private/Materials/HLSLMaterialDerivativeAutogen.cpp
@@ -188,18 +188,6 @@ void FMaterialDerivativeAutogen::EnableGeneratedDepencencies()
 		}
 	}
 
-	// min, max, saturate, and PosClampedPow require the select helper
-	for (int32 Index = 0; Index < 4; Index++)
-	{
-		if (bFunc2OpIsEnabled[(int32)EFunc2::Max][Index] ||
-			bFunc2OpIsEnabled[(int32)EFunc2::Min][Index] ||
-			bFunc1OpIsEnabled[(int32)EFunc1::Saturate][Index] ||
-			bFunc2OpIsEnabled[(int32)EFunc2::PowPositiveClamped][Index])
-		{
-			bSelectElemHelperEnabled[Index] = true;
-		}
-	}
-
 	// Dot requires extract, mul1, add1 and FloatDeriv constructor
 	for (int32 Index = 0; Index < 4; Index++)
 	{
@@ -1147,30 +1135,6 @@ FString FMaterialDerivativeAutogen::GenerateUsedFunctions(FHLSLMaterialTranslato
 		}
 	}
 
-
-	// SelectElemHelper
-	for (int32 Index = 0; Index < 4; Index++)
-	{
-		if (bSelectElemHelperEnabled[Index] || IsDebugGenerateAllFunctionsEnabled())
-		{
-			FString FieldName = GetFloatVectorName(Index);
-			FString BoolName = GetBoolVectorName(Index);
-
-			Ret += FieldName + TEXT(" SelectElemHelper(") + BoolName + TEXT(" Cmp, ") + FieldName + TEXT(" Lhs, ") + FieldName + TEXT(" Rhs)") LINE_TERMINATOR;
-			Ret += TEXT("{") LINE_TERMINATOR;
-			Ret += TEXT("\t") + FieldName + TEXT(" Ret = 0.0f;") LINE_TERMINATOR;
-			for (int32 Component = 0; Component <= Index; Component++)
-			{
-				const TCHAR* ChanName = SwizzleList[Component];
-				Ret += FString::Printf(TEXT("\tRet.%s = Cmp.%s ? Lhs.%s : Rhs.%s;"),ChanName,ChanName,ChanName,ChanName) + LINE_TERMINATOR;
-			}
-
-			Ret += TEXT("\treturn Ret;") LINE_TERMINATOR;
-			Ret += TEXT("}") LINE_TERMINATOR;
-			Ret += TEXT("") LINE_TERMINATOR;
-		}
-	}
-
 	// Func2s
 	for (int32 Op = 0; Op < (int32)EFunc2::Num; Op++)
 	{
@@ -1249,9 +1213,9 @@ FString FMaterialDerivativeAutogen::GenerateUsedFunctions(FHLSLMaterialTranslato
 					Ret += TEXT("{") LINE_TERMINATOR;
 					Ret += TEXT("\t") + BaseName + TEXT(" Ret;") LINE_TERMINATOR;
 					Ret += TEXT("\t") + BoolName + TEXT(" Cmp = A.Value < B.Value;") LINE_TERMINATOR;
-					Ret += TEXT("\tRet.Value = SelectElemHelper(Cmp, A.Value, B.Value);") LINE_TERMINATOR;
-					Ret += TEXT("\tRet.Ddx = SelectElemHelper(Cmp, A.Ddx, B.Ddx);") LINE_TERMINATOR;
-					Ret += TEXT("\tRet.Ddy = SelectElemHelper(Cmp, A.Ddy, B.Ddy);") LINE_TERMINATOR;
+					Ret += TEXT("\tRet.Value = Cmp ? A.Value : B.Value;") LINE_TERMINATOR;
+					Ret += TEXT("\tRet.Ddx = Cmp ? A.Ddx : B.Ddx;") LINE_TERMINATOR;
+					Ret += TEXT("\tRet.Ddy = Cmp ? A.Ddy : B.Ddy;") LINE_TERMINATOR;
 					Ret += TEXT("\treturn Ret;") LINE_TERMINATOR;
 					Ret += TEXT("}") LINE_TERMINATOR;
 					Ret += TEXT("") LINE_TERMINATOR;
@@ -1261,9 +1225,9 @@ FString FMaterialDerivativeAutogen::GenerateUsedFunctions(FHLSLMaterialTranslato
 					Ret += TEXT("{") LINE_TERMINATOR;
 					Ret += TEXT("\t") + BaseName + TEXT(" Ret;") LINE_TERMINATOR;
 					Ret += TEXT("\t") + BoolName + TEXT(" Cmp = A.Value > B.Value;") LINE_TERMINATOR;
-					Ret += TEXT("\tRet.Value = SelectElemHelper(Cmp, A.Value, B.Value);") LINE_TERMINATOR;
-					Ret += TEXT("\tRet.Ddx = SelectElemHelper(Cmp, A.Ddx, B.Ddx);") LINE_TERMINATOR;
-					Ret += TEXT("\tRet.Ddy = SelectElemHelper(Cmp, A.Ddy, B.Ddy);") LINE_TERMINATOR;
+					Ret += TEXT("\tRet.Value = Cmp ? A.Value : B.Value;") LINE_TERMINATOR;
+					Ret += TEXT("\tRet.Ddx = Cmp ? A.Ddx : B.Ddx;") LINE_TERMINATOR;
+					Ret += TEXT("\tRet.Ddy = Cmp ? A.Ddy : B.Ddy;") LINE_TERMINATOR;
 					Ret += TEXT("\treturn Ret;") LINE_TERMINATOR;
 					Ret += TEXT("}") LINE_TERMINATOR;
 					Ret += TEXT("") LINE_TERMINATOR;
@@ -1306,8 +1270,8 @@ FString FMaterialDerivativeAutogen::GenerateUsedFunctions(FHLSLMaterialTranslato
 					Ret += TEXT("\tRet.Value = PositiveClampedPow(A.Value, B.Value);") LINE_TERMINATOR;
 					Ret += TEXT("\tRet.Ddx = Ret.Value * (B.Ddx * log(A.Value) + (B.Value/A.Value)*A.Ddx);") LINE_TERMINATOR;
 					Ret += TEXT("\tRet.Ddy = Ret.Value * (B.Ddy * log(A.Value) + (B.Value/A.Value)*A.Ddy);") LINE_TERMINATOR;
-					Ret += TEXT("\tRet.Ddx = SelectElemHelper(InRange,Ret.Ddx,Zero);") LINE_TERMINATOR;
-					Ret += TEXT("\tRet.Ddy = SelectElemHelper(InRange,Ret.Ddy,Zero);") LINE_TERMINATOR;
+					Ret += TEXT("\tRet.Ddx = InRange ? Ret.Ddx : Zero;") LINE_TERMINATOR;
+					Ret += TEXT("\tRet.Ddy = InRange ? Ret.Ddy : Zero;") LINE_TERMINATOR;
 					Ret += TEXT("\treturn Ret;") LINE_TERMINATOR;
 					Ret += TEXT("}") LINE_TERMINATOR;
 					Ret += TEXT("") LINE_TERMINATOR;
@@ -1540,8 +1504,8 @@ FString FMaterialDerivativeAutogen::GenerateUsedFunctions(FHLSLMaterialTranslato
 					Ret += TEXT("\t") + FieldName + TEXT(" Zero = 0.0f;") LINE_TERMINATOR;
 					Ret += TEXT("\t") + BaseName + TEXT(" Ret;") LINE_TERMINATOR;
 					Ret += TEXT("\tRet.Value = saturate(A.Value);") LINE_TERMINATOR;
-					Ret += TEXT("\tRet.Ddx = SelectElemHelper(InRange, A.Ddx, Zero);") LINE_TERMINATOR;
-					Ret += TEXT("\tRet.Ddy = SelectElemHelper(InRange, A.Ddy, Zero);") LINE_TERMINATOR;
+					Ret += TEXT("\tRet.Ddx = InRange ? A.Ddx : Zero;") LINE_TERMINATOR;
+					Ret += TEXT("\tRet.Ddy = InRange ? A.Ddy : Zero;") LINE_TERMINATOR;
 					Ret += TEXT("\treturn Ret;") LINE_TERMINATOR;
 					Ret += TEXT("}") LINE_TERMINATOR;
 					Ret += TEXT("") LINE_TERMINATOR;
diff --git a/Engine/Source/Runtime/Engine/Private/Materials/HLSLMaterialDerivativeAutogen.h b/Engine/Source/Runtime/Engine/Private/Materials/HLSLMaterialDerivativeAutogen.h
--- a/Engine/Source/Runtime/Engine/Private/Materials/HLSLMaterialDerivativeAutogen.h
+++ b/Engine/Source/Runtime/Engine/Private/Materials/HLSLMaterialDerivativeAutogen.h
@@ -129,7 +129,6 @@ private:
 	bool bIf2Enabled[4] = {};
 	bool bLerpEnabled[4] = {};
 	bool bRotateScaleOffsetTexCoords = false;
-	bool bSelectElemHelperEnabled[4] = {};
 	bool bUnMirrorEnabled[2][2] = {};
 	
 };
