Make sure to replicate newly created actor descriptors to their corresponding views when creating new actors during HLODs generation.

diff --git a/Engine/Source/Runtime/Engine/Private/WorldPartition/RuntimeSpatialHash/RuntimeSpatialHashGridHelper.cpp b/Engine/Source/Runtime/Engine/Private/WorldPartition/RuntimeSpatialHash/RuntimeSpatialHashGridHelper.cpp
--- a/Engine/Source/Runtime/Engine/Private/WorldPartition/RuntimeSpatialHash/RuntimeSpatialHashGridHelper.cpp
+++ b/Engine/Source/Runtime/Engine/Private/WorldPartition/RuntimeSpatialHash/RuntimeSpatialHashGridHelper.cpp
@@ -46,7 +46,7 @@ FSquare2DGridHelper::FSquare2DGridHelper(const FBox& InWorldBounds, const FVecto
 		// Also note that, the default origin of each level will always be centered at the middle of the bounds of (level's cellsize * level's grid size).
 		int32 LevelGridSize = (Level == GridLevelCount-1) ? CurrentGridSize : CurrentGridSize + 1;
 		
-		Levels.Emplace(FVector2D(InOrigin), CurrentCellSize, LevelGridSize);
+		Levels.Emplace(FVector2D(InOrigin), CurrentCellSize, LevelGridSize, Level);
 
 		CurrentCellSize <<= 1;
 		CurrentGridSize >>= 1;
@@ -58,11 +58,14 @@ FSquare2DGridHelper::FSquare2DGridHelper(const FBox& InWorldBounds, const FVecto
 #endif
 }
 
-void FSquare2DGridHelper::ForEachCells(TFunctionRef<void(const FIntVector&)> InOperation) const
+void FSquare2DGridHelper::ForEachCells(TFunctionRef<void(const FSquare2DGridHelper::FGridLevel::FGridCell&)> InOperation) const
 {
 	for (int32 Level = 0; Level < Levels.Num(); Level++)
 	{
-		Levels[Level].ForEachCells([Level, &InOperation](const FIntVector2& Coord) { InOperation(FIntVector(Coord.X, Coord.Y, Level)); });
+		for (const FSquare2DGridHelper::FGridLevel::FGridCell& ThisCell : Levels[Level].Cells)
+		{
+			InOperation(ThisCell);
+		}
 	}
 }
 
diff --git a/Engine/Source/Runtime/Engine/Private/WorldPartition/RuntimeSpatialHash/RuntimeSpatialHashGridHelper.h b/Engine/Source/Runtime/Engine/Private/WorldPartition/RuntimeSpatialHash/RuntimeSpatialHashGridHelper.h
--- a/Engine/Source/Runtime/Engine/Private/WorldPartition/RuntimeSpatialHash/RuntimeSpatialHashGridHelper.h
+++ b/Engine/Source/Runtime/Engine/Private/WorldPartition/RuntimeSpatialHash/RuntimeSpatialHashGridHelper.h
@@ -267,6 +267,10 @@ struct FSquare2DGridHelper
 
 		struct FGridCell
 		{
+			FGridCell(const FIntVector& InCoords)
+				: Coords(InCoords)
+			{}
+
 			void AddActor(FActorInstance ActorInstance, const TArray<const UDataLayer*>& InDataLayers)
 			{
 				FDataLayersID DataLayersID = FDataLayersID(InDataLayers);
@@ -300,13 +304,20 @@ struct FSquare2DGridHelper
 				return nullptr;
 			}
 
+			FIntVector GetCoords() const
+			{
+				return Coords;
+			}
+
 		private:
+			FIntVector Coords;
 			TArray<FGridCellDataChunk> DataChunks;
 		};
 #endif
 
-		inline FGridLevel(const FVector2D& InOrigin, int32 InCellSize, int32 InGridSize)
+		inline FGridLevel(const FVector2D& InOrigin, int32 InCellSize, int32 InGridSize, int32 InLevel)
 			: FGrid2D(InOrigin, InCellSize, InGridSize)
+			, Level(InLevel)
 		{}
 
 #if WITH_EDITOR
@@ -330,7 +341,7 @@ struct FSquare2DGridHelper
 			}
 			else
 			{
-				CellIndexMapping = Cells.AddDefaulted();
+				CellIndexMapping = Cells.Emplace(FIntVector(InCoords.X, InCoords.Y, Level));
 				CellsMapping.Add(CellIndex, CellIndexMapping);
 			}
 
@@ -350,9 +361,13 @@ struct FSquare2DGridHelper
 			GetCellIndex(InCoords, CellIndex);
 
 			int32 CellIndexMapping = CellsMapping.FindChecked(CellIndex);
-			return Cells[CellIndexMapping];
+
+			const FGridCell& Cell = Cells[CellIndexMapping];
+			check(Cell.GetCoords() == FIntVector(InCoords.X, InCoords.Y, Level));
+			return Cell;
 		}
 
+		int32 Level;
 		TArray<FGridCell> Cells;
 		TMap<int32, int32> CellsMapping;
 #endif
@@ -389,7 +404,7 @@ struct FSquare2DGridHelper
 	}
 
 	// Runs a function on all cells
-	void ForEachCells(TFunctionRef<void(const FIntVector&)> InOperation) const;
+	void ForEachCells(TFunctionRef<void(const FSquare2DGridHelper::FGridLevel::FGridCell&)> InOperation) const;
 
 	/**
 	 * Runs a function on all intersecting cells for the provided box
diff --git a/Engine/Source/Runtime/Engine/Private/WorldPartition/RuntimeSpatialHash/RuntimeSpatialHashHLOD.cpp b/Engine/Source/Runtime/Engine/Private/WorldPartition/RuntimeSpatialHash/RuntimeSpatialHashHLOD.cpp
--- a/Engine/Source/Runtime/Engine/Private/WorldPartition/RuntimeSpatialHash/RuntimeSpatialHashHLOD.cpp
+++ b/Engine/Source/Runtime/Engine/Private/WorldPartition/RuntimeSpatialHash/RuntimeSpatialHashHLOD.cpp
@@ -141,10 +141,8 @@ static TArray<FGuid> GenerateHLODsForGrid(UWorldPartition* WorldPartition, const
 
 	// Quick pass to compute the number of cells we'll have to process, to provide a meaningful progress display
 	int32 NbCellsToProcess = 0;
-	PartitionedActors.ForEachCells([&](const FIntVector& CellCoord)
+	PartitionedActors.ForEachCells([&](const FSquare2DGridHelper::FGridLevel::FGridCell& GridCell)
 	{
-		const FSquare2DGridHelper::FGridLevel::FGridCell& GridCell = PartitionedActors.GetCell(CellCoord);
-
 		for (const FSquare2DGridHelper::FGridLevel::FGridCellDataChunk& GridCellDataChunk : GridCell.GetDataChunks())
 		{
 			const bool bShouldGenerateHLODs = ShouldGenerateHLODs(GridCell, GridCellDataChunk);
@@ -162,9 +160,9 @@ static TArray<FGuid> GenerateHLODsForGrid(UWorldPartition* WorldPartition, const
 	FAssetRegistryModule& AssetRegistryModule = FModuleManager::LoadModuleChecked<FAssetRegistryModule>(TEXT("AssetRegistry"));
 
 	TArray<FGuid> GridHLODActors;
-	PartitionedActors.ForEachCells([&](const FIntVector& CellCoord)
+	PartitionedActors.ForEachCells([&](const FSquare2DGridHelper::FGridLevel::FGridCell& GridCell)
 	{
-		const FSquare2DGridHelper::FGridLevel::FGridCell& GridCell = PartitionedActors.GetCell(CellCoord);
+		const FIntVector CellCoord = GridCell.GetCoords();
 
 		FBox2D CellBounds2D;
 		PartitionedActors.GetCellBounds(CellCoord, CellBounds2D);
@@ -234,6 +232,12 @@ static TArray<FGuid> GenerateHLODsForGrid(UWorldPartition* WorldPartition, const
 					DirectoryWatcherModule.Get()->Tick(-1.0f);
 					AssetRegistryModule.Get().Tick(-1.0f);
 
+					// Make sure the previous ticks properly registered all our actor descriptors
+					for (AWorldPartitionHLOD* CellHLODActor : CellHLODActors)
+					{
+						check(WorldPartition->GetActorDesc(CellHLODActor->GetActorGuid()));
+					}
+
 					// Update newly created HLOD actors
 					for (AWorldPartitionHLOD* CellHLODActor : NewCellHLODActors)
 					{
@@ -490,9 +494,11 @@ bool UWorldPartitionRuntimeSpatialHash::GenerateHLOD(ISourceControlHelper* Sourc
 
 	// Create actor clusters - ignore HLOD actors
 	FActorClusterContext ClusterContext(WorldPartition, this, TOptional<FActorClusterContext::FFilterPredicate>([](const FWorldPartitionActorDescView& ActorDescView)
-		{
-			return !ActorDescView.GetActorClass()->IsChildOf<AWorldPartitionHLOD>();
-		}), /* bInIncludeChildContainers=*/ false);
+	{
+		return !ActorDescView.GetActorClass()->IsChildOf<AWorldPartitionHLOD>();
+	}), /* bInIncludeChildContainers=*/ false);
+
+	FActorContainerInstance& MainContainerInstance = *ClusterContext.GetClusterInstance(WorldPartition);
 		
 	TArray<TArray<const FActorClusterInstance*>> GridsClusters;
 	GridsClusters.InsertDefaulted(0, Grids.Num());
@@ -514,7 +520,7 @@ bool UWorldPartitionRuntimeSpatialHash::GenerateHLOD(ISourceControlHelper* Sourc
 	// Keep track of all valid HLOD actors, along with which runtime grid they live in
 	TMap<FName, TArray<FGuid>> GridsHLODActors;
 
-	auto GenerateHLODs = [&GridsHLODActors, WorldPartition, &GridsDepth, &Context, SourceControlHelper, bCreateActorsOnly](const FSpatialHashRuntimeGrid& RuntimeGrid, uint32 HLODLevel, const TArray<const FActorClusterInstance*>& ActorClusterInstances)
+	auto GenerateHLODs = [&GridsHLODActors, &MainContainerInstance, WorldPartition, &GridsDepth, &Context, SourceControlHelper, bCreateActorsOnly](const FSpatialHashRuntimeGrid& RuntimeGrid, uint32 HLODLevel, const TArray<const FActorClusterInstance*>& ActorClusterInstances)
 	{
 		// Generate HLODs for this grid
 		TArray<FGuid> HLODActors = GenerateHLODsForGrid(WorldPartition, RuntimeGrid, HLODLevel, Context, SourceControlHelper, bCreateActorsOnly, ActorClusterInstances);
@@ -522,10 +528,12 @@ bool UWorldPartitionRuntimeSpatialHash::GenerateHLOD(ISourceControlHelper* Sourc
 		for (const FGuid& HLODActorGuid : HLODActors)
 		{
 			FWorldPartitionActorDesc* HLODActorDesc = WorldPartition->GetActorDesc(HLODActorGuid);
-			if(ensure(HLODActorDesc))
-			{
-				GridsHLODActors.FindOrAdd(HLODActorDesc->GetRuntimeGrid()).Add(HLODActorGuid);
-			}
+			check(HLODActorDesc);
+
+			MainContainerInstance.ActorDescViewMap.Emplace(HLODActorGuid, HLODActorDesc);
+
+			const FWorldPartitionActorDescView& HLODActorDescView = MainContainerInstance.GetActorDescView(HLODActorGuid);
+			GridsHLODActors.FindOrAdd(HLODActorDescView.GetRuntimeGrid()).Add(HLODActorGuid);
 		}
 	};
 
@@ -537,7 +545,6 @@ bool UWorldPartitionRuntimeSpatialHash::GenerateHLOD(ISourceControlHelper* Sourc
 	}
 
 	// Now, go on and create HLOD actors from HLOD grids (HLOD 1-N)
-	FActorContainerInstance& MainContainerInstance = *ClusterContext.GetClusterInstance(WorldPartition);
 	for (const FName HLODGridName : SortedGrids)
 	{
 		// No need to process empty grids
