AnimInstanceProxy - Small fix to previous change.

diff --git a/Engine/Source/Runtime/Engine/Private/Animation/AnimInstanceProxy.cpp b/Engine/Source/Runtime/Engine/Private/Animation/AnimInstanceProxy.cpp
--- a/Engine/Source/Runtime/Engine/Private/Animation/AnimInstanceProxy.cpp
+++ b/Engine/Source/Runtime/Engine/Private/Animation/AnimInstanceProxy.cpp
@@ -1590,7 +1590,7 @@ void FAnimInstanceProxy::SlotEvaluatePoseWithBlendProfiles(const FName& SlotNode
 		FSlotEvaluationPose& AdditivePose = Poses[PoseIndex];
 		const FAnimationPoseData AdditiveAnimationPoseData(AdditivePose);
 		OutBlendedAnimationPoseData.GetCurve().Accumulate(AdditiveAnimationPoseData.GetCurve(), AdditivePose.Weight);
-		UE::Anim::Attributes::AccumulateAttributes(AdditiveAnimationPoseData.GetAttributes(), OutBlendedAnimationPoseData.GetAttributes(), AdditivePose.Weight, AdditivePoses[PoseIndex].AdditiveType);
+		UE::Anim::Attributes::AccumulateAttributes(AdditiveAnimationPoseData.GetAttributes(), OutBlendedAnimationPoseData.GetAttributes(), AdditivePose.Weight, Poses[PoseIndex].AdditiveType);
 	}
 }
 
