Fix non-editor builds

diff --git a/Engine/Source/Runtime/Engine/Private/WorldPartition/RuntimeSpatialHash/RuntimeSpatialHashGridHelper.cpp b/Engine/Source/Runtime/Engine/Private/WorldPartition/RuntimeSpatialHash/RuntimeSpatialHashGridHelper.cpp
--- a/Engine/Source/Runtime/Engine/Private/WorldPartition/RuntimeSpatialHash/RuntimeSpatialHashGridHelper.cpp
+++ b/Engine/Source/Runtime/Engine/Private/WorldPartition/RuntimeSpatialHash/RuntimeSpatialHashGridHelper.cpp
@@ -58,6 +58,7 @@ FSquare2DGridHelper::FSquare2DGridHelper(const FBox& InWorldBounds, const FVecto
 #endif
 }
 
+#if WITH_EDITOR
 void FSquare2DGridHelper::ForEachCells(TFunctionRef<void(const FSquare2DGridHelper::FGridLevel::FGridCell&)> InOperation) const
 {
 	for (int32 Level = 0; Level < Levels.Num(); Level++)
@@ -68,6 +69,7 @@ void FSquare2DGridHelper::ForEachCells(TFunctionRef<void(const FSquare2DGridHelp
 		}
 	}
 }
+#endif
 
 int32 FSquare2DGridHelper::ForEachIntersectingCells(const FBox& InBox, TFunctionRef<void(const FIntVector&)> InOperation) const
 {
diff --git a/Engine/Source/Runtime/Engine/Private/WorldPartition/RuntimeSpatialHash/RuntimeSpatialHashGridHelper.h b/Engine/Source/Runtime/Engine/Private/WorldPartition/RuntimeSpatialHash/RuntimeSpatialHashGridHelper.h
--- a/Engine/Source/Runtime/Engine/Private/WorldPartition/RuntimeSpatialHash/RuntimeSpatialHashGridHelper.h
+++ b/Engine/Source/Runtime/Engine/Private/WorldPartition/RuntimeSpatialHash/RuntimeSpatialHashGridHelper.h
@@ -317,7 +317,9 @@ struct FSquare2DGridHelper
 
 		inline FGridLevel(const FVector2D& InOrigin, int32 InCellSize, int32 InGridSize, int32 InLevel)
 			: FGrid2D(InOrigin, InCellSize, InGridSize)
+#if WITH_EDITOR
 			, Level(InLevel)
+#endif
 		{}
 
 #if WITH_EDITOR
@@ -403,8 +405,10 @@ struct FSquare2DGridHelper
 		return false;
 	}
 
+#if WITH_EDITOR
 	// Runs a function on all cells
 	void ForEachCells(TFunctionRef<void(const FSquare2DGridHelper::FGridLevel::FGridCell&)> InOperation) const;
+#endif
 
 	/**
 	 * Runs a function on all intersecting cells for the provided box
