Fixed 'stat LEVELS' display to automatically adapt column sizes when displaying levels on multiple columns.

diff --git a/Engine/Source/Runtime/Engine/Classes/Engine/Canvas.h b/Engine/Source/Runtime/Engine/Classes/Engine/Canvas.h
--- a/Engine/Source/Runtime/Engine/Classes/Engine/Canvas.h
+++ b/Engine/Source/Runtime/Engine/Classes/Engine/Canvas.h
@@ -387,7 +387,19 @@ public:
 	 * @param Text The string to calculate for.
 	 * @param XL out Horizontal length of string.
 	 * @param YL out Vertical length of string.
-	 * @param bDPUAware If true measures text considering the current DPI scale factor of the canvas.  Defaults to false for backwards compatibility
+	 * @param bDPIAware If true measures text considering the current DPI scale factor of the canvas.  Defaults to false for backwards compatibility
+	 * @param Canvas Canvas state object
+	 */
+	static void StrLen(const UFont* InFont, const FString& InText, float& XL, float& YL, bool bDPIAware, FCanvas* InCanvas);
+
+	/**
+	 * Calculate the length of a string, taking text wrapping into account.
+	 *
+	 * @param InFont The Font use.
+	 * @param Text The string to calculate for.
+	 * @param XL out Horizontal length of string.
+	 * @param YL out Vertical length of string.
+	 * @param bDPIAware If true measures text considering the current DPI scale factor of the canvas.  Defaults to false for backwards compatibility
 	 */
 	void StrLen(const UFont* InFont, const FString& InText, float& XL, float& YL, bool bDPIAware = false);
 
diff --git a/Engine/Source/Runtime/Engine/Private/UnrealEngine.cpp b/Engine/Source/Runtime/Engine/Private/UnrealEngine.cpp
--- a/Engine/Source/Runtime/Engine/Private/UnrealEngine.cpp
+++ b/Engine/Source/Runtime/Engine/Private/UnrealEngine.cpp
@@ -15293,11 +15293,6 @@ int32 UEngine::RenderStatColorList(UWorld* World, FViewport* Viewport, FCanvas*
 }
 
 // LEVELS
-TAutoConsoleVariable<int32> CVarStatLevelsColumnWidth(
-	TEXT("stats.StatLevelsColumnWidth"),
-	350,
-	TEXT("The width in pixels of a column in stat levels.")
-);
 
 int32 UEngine::RenderStatLevels(UWorld* World, FViewport* Viewport, FCanvas* Canvas, int32 X, int32 Y, const FVector* ViewLocation, const FRotator* ViewRotation)
 {
@@ -15332,6 +15327,10 @@ int32 UEngine::RenderStatLevels(UWorld* World, FViewport* Viewport, FCanvas* Can
 	int32 BaseY = Y;
 
 	// now draw the levels
+	float CurrentColumnWidth = 0.f;
+	float TextWidth = 0.f;
+	float TextHeight = 0.f;
+	
 	for (int32 LevelIdx = 1; LevelIdx < SubLevelsStatusList.Num(); ++LevelIdx)
 	{
 		const FSubLevelStatus& LevelStatus = SubLevelsStatusList[LevelIdx];
@@ -15341,7 +15340,8 @@ int32 UEngine::RenderStatLevels(UWorld* World, FViewport* Viewport, FCanvas* Can
 		{
 			MaxY = FMath::Max(MaxY, Y);
 			Y = BaseY;
-			X += CVarStatLevelsColumnWidth.GetValueOnGameThread();
+			X += CurrentColumnWidth;
+			CurrentColumnWidth = 0.f;
 		}
 
 		FColor	Color = ULevelStreaming::GetLevelStreamingStatusColor(LevelStatus.StreamingStatus);
@@ -15374,7 +15374,9 @@ int32 UEngine::RenderStatLevels(UWorld* World, FViewport* Viewport, FCanvas* Can
 		DisplayName = *FString::Printf(TEXT("%s %s %s (%s)"), (LevelStatus.bPlayerInside ? TEXT("->") : TEXT("  ")), (LevelStatus.bInConsiderList ? TEXT("*") : TEXT(" ")), *DisplayName, StatusDisplayName);
 
 		Canvas->DrawShadowedString(X + 4, Y, *DisplayName, GetSmallFont(), Color);
-		Y += 12;
+		UCanvas::StrLen(GetSmallFont(), *DisplayName, TextWidth, TextHeight, /*bDPIAware*/false, Canvas);
+		CurrentColumnWidth = FMath::Max(CurrentColumnWidth, TextWidth);
+		Y += TextHeight;
 	}
 	return FMath::Max(MaxY, Y);
 }
diff --git a/Engine/Source/Runtime/Engine/Private/UserInterface/Canvas.cpp b/Engine/Source/Runtime/Engine/Private/UserInterface/Canvas.cpp
--- a/Engine/Source/Runtime/Engine/Private/UserInterface/Canvas.cpp
+++ b/Engine/Source/Runtime/Engine/Private/UserInterface/Canvas.cpp
@@ -1705,6 +1705,11 @@ int32 UCanvas::WrappedPrint(bool Draw, float X, float Y, int32& out_XL, int32& o
 }
 
 void UCanvas::StrLen(const UFont* InFont, const FString& InText, float& XL, float& YL, bool bDPIAware)
+{
+	UCanvas::StrLen(InFont, InText, XL, YL, bDPIAware, Canvas);
+}
+
+void UCanvas::StrLen(const UFont* InFont, const FString& InText, float& XL, float& YL, bool bDPIAware, FCanvas* InCanvas)
 {
 	if (InFont == NULL)
 	{
@@ -1723,7 +1728,7 @@ void UCanvas::StrLen(const UFont* InFont, const FString& InText, float& XL, floa
 		{
 			const FSlateFontInfo LegacyFontInfo = InFont->GetLegacySlateFontInfo();
 			const TSharedRef<FSlateFontMeasure> FontMeasure = FSlateApplication::Get().GetRenderer()->GetFontMeasureService();
-			const FVector2D MeasuredTextSize = FontMeasure->Measure(InText, LegacyFontInfo, Canvas->GetDPIScale());
+			const FVector2D MeasuredTextSize = FontMeasure->Measure(InText, LegacyFontInfo, InCanvas->GetDPIScale());
 			XL = MeasuredTextSize.X;
 			YL = MeasuredTextSize.Y;
 		}
