Fix two shader jobs cache issues.

- Not picking up changes after the source file has been edited (UE-109287)
   - Now the "key" includes the source shader hash. Access to the hash is guarded as it now happen through multiple threads since SubmitJobs is called from TryGetShaders. Also, the thread safety check might have been disabled in CL 15558051 by mistake.
- Spurious errors (UE-109638) due to shader pipeline jobs not initializing the success status correctly.

diff --git a/Engine/Source/Runtime/Engine/Private/ShaderCompiler/ShaderCompiler.cpp b/Engine/Source/Runtime/Engine/Private/ShaderCompiler/ShaderCompiler.cpp
--- a/Engine/Source/Runtime/Engine/Private/ShaderCompiler/ShaderCompiler.cpp
+++ b/Engine/Source/Runtime/Engine/Private/ShaderCompiler/ShaderCompiler.cpp
@@ -224,6 +224,12 @@ static FShaderCommonCompileJob* CloneJob_Single(const FShaderCompileJob* SrcJob)
 	Job->ShaderParameters = SrcJob->ShaderParameters;
 	Job->PendingShaderMap = SrcJob->PendingShaderMap;
 	Job->Input = SrcJob->Input;
+	if (SrcJob->bInputHashSet)
+	{
+		Job->InputHash = SrcJob->InputHash;
+		Job->bInputHashSet = true;
+	}
+	ensure(Job->bInputHashSet == SrcJob->bInputHashSet);
 	return Job;
 }
 
@@ -237,6 +243,13 @@ static FShaderCommonCompileJob* CloneJob_Pipeline(const FShaderPipelineCompileJo
 	{
 		Job->StageJobs[i]->Input = SrcJob->StageJobs[i]->Input;
 	}
+
+	if (SrcJob->bInputHashSet)
+	{
+		Job->InputHash = SrcJob->InputHash;
+		Job->bInputHashSet = true;
+	}
+	ensure(Job->bInputHashSet == SrcJob->bInputHashSet);
 	return Job;
 }
 
@@ -264,6 +277,7 @@ void FShaderCompileJobCollection::InternalSetPriority(FShaderCommonCompileJob* J
 		Job->Unlink();
 
 		NumPendingJobs[PriorityIndex]++;
+		ensure(!ShaderCompiler::IsJobCacheEnabled() || Job->bInputHashSet);
 		Job->LinkHead(PendingJobs[PriorityIndex]);
 		Job->Priority = InPriority;
 		Job->PendingPriority = InPriority;
@@ -276,10 +290,16 @@ void FShaderCompileJobCollection::InternalSetPriority(FShaderCommonCompileJob* J
 		NewJob->Priority = InPriority;
 		const int32 NewNumPendingJobs = NewJob->PendingShaderMap->NumPendingJobs.Increment();
 		checkf(NewNumPendingJobs > 1, TEXT("Invalid number of pending jobs %d, should have had at least 1 job previously"), NewNumPendingJobs);
-
 		InternalAddJob(NewJob);
 
 		NumPendingJobs[PriorityIndex]++;
+		ensureMsgf(NewJob->bInputHashSet == Job->bInputHashSet, TEXT("Cloned and original jobs should either both have input hash, or both not have it. Job->bInputHashSet=%d, NewJob->bInputHashSet=%d"),
+			Job->bInputHashSet,
+			NewJob->bInputHashSet
+			);
+		ensureMsgf(!ShaderCompiler::IsJobCacheEnabled() || NewJob->GetInputHash() == Job->GetInputHash(),
+			TEXT("If shader jobs cache is enabled, cloned job should have the same input hash as the original, and it doesn't.")
+			);
 		NewJob->LinkHead(PendingJobs[PriorityIndex]);
 		NewJob->PendingPriority = InPriority;
 		NumOutstandingJobs.Increment();
@@ -452,6 +472,7 @@ void FShaderCompileJobCollection::SubmitJobs(const TArray<FShaderCommonCompileJo
 				// new job
 				if (bNewJob)
 				{
+					ensure(!ShaderCompiler::IsJobCacheEnabled() || Job->bInputHashSet);
 					Job->LinkHead(PendingJobs[PriorityIndex]);
 
 					NumPendingJobs[PriorityIndex]++;
@@ -483,8 +504,6 @@ void FShaderCompileJobCollection::HandleLogJobsCacheStats()
 
 void FShaderCompileJobCollection::ProcessFinishedJob(FShaderCommonCompileJob* FinishedJob, bool bWasCached)
 {
-	const FSHAHash& InputHash = FinishedJob->GetInputHash();
-
 	// TODO: have a pending shader map critical section? not clear at this point if we can be accessing the results on another thread at the same time
 	FShaderMapCompileResults& ShaderMapResults = *(FinishedJob->PendingShaderMap);
 	ShaderMapResults.FinishedJobs.Add(FinishedJob);
@@ -494,10 +513,13 @@ void FShaderCompileJobCollection::ProcessFinishedJob(FShaderCommonCompileJob* Fi
 	checkf(NumPendingJobsForSM >= 0, TEXT("Problem tracking pending jobs for a SM (%d), number of pending jobs (%d) is negative!"), FinishedJob->Id, NumPendingJobsForSM);
 
 	InternalSubtractNumOutstandingJobs(1);
-	if (!bWasCached)
+	if (!bWasCached && ShaderCompiler::IsJobCacheEnabled())
 	{
 		AddToCacheAndProcessPending(FinishedJob);
 
+		ensureMsgf(FinishedJob->bInputHashSet, TEXT("Finished job didn't have input hash set, was shader compiler jobs cache toggled runtime?"));
+		const FSHAHash& InputHash = FinishedJob->GetInputHash();
+
 		// remove ourselves from the jobs in flight, if we were there (if this job is a cloned job it might not have been)
 		JobsInFlight.Remove(InputHash);
 	}
@@ -529,6 +551,7 @@ void FShaderCompileJobCollection::AddToCacheAndProcessPending(FShaderCommonCompi
 
 			FMemoryReader MemReader(Output);
 			CurHead->SerializeOutput(MemReader);
+			checkf(CurHead->bSucceeded == FinishedJob->bSucceeded, TEXT("Different success status for the job with the same ihash"));
 
 			// finish the job instantly
 			ProcessFinishedJob(CurHead, true);
@@ -611,6 +634,7 @@ int32 FShaderCompileJobCollection::GetPendingJobs(EShaderCompilerWorkerType InWo
 			// Temporary commented out until r.ShaderDevelopmentMode=1 shader error retry crash gets fixed
 			//check(Job.CurrentWorker == EShaderCompilerWorkerType::None);
 			//check(Job.PendingPriority == InPriority);
+			ensure(!ShaderCompiler::IsJobCacheEnabled() || Job.bInputHashSet);
 
 			It.Next();
 			Job.Unlink();
@@ -1555,6 +1579,7 @@ void FShaderCompileUtilities::DoReadTaskResults(const TArray<FShaderCommonCompil
 				// cannot reissue a single stage of a pipeline job
 				ReadSingleJob(SingleJob, OutputFile);
 				CurrentJob->bFailedRemovingUnused = CurrentJob->bFailedRemovingUnused | SingleJob->Output.bFailedRemovingUnused;
+				CurrentJob->bSucceeded = CurrentJob->bSucceeded && SingleJob->bSucceeded;
 			}
 		}
 	}
@@ -1981,6 +2006,23 @@ void FShaderCompileThreadRunnable::WriteNewTasks()
 			}
 			delete TransferFile;
 
+#if 0 // debugging code to dump the worker inputs
+			static FCriticalSection ArchiveLock;
+			{
+				FScopeLock Locker(&ArchiveLock);
+				static int ArchivedTransferFileNum = 0;
+				FString JobCacheDir = ShaderCompiler::IsJobCacheEnabled() ? TEXT("JobCache") : TEXT("NoJobCache");
+				FString ArchiveDir = FPaths::ProjectSavedDir() / TEXT("ArchivedWorkerInputs") / JobCacheDir;
+				FString ArchiveName = FString::Printf(TEXT("Input-%d"), ArchivedTransferFileNum++);
+				FString ArchivePath = ArchiveDir / ArchiveName;
+				if (!IFileManager::Get().Copy(*ArchivePath, *TransferFileName))
+				{
+					UE_LOG(LogInit, Error, TEXT("Could not copy file %s to %s"), *TransferFileName, *ArchivePath);
+					ensure(false);
+				}
+			}
+#endif
+
 			// Change the transfer file name to proper one
 			FString ProperTransferFileName = WorkingDirectory / TEXT("WorkerInputOnly.in");
 			if (!IFileManager::Get().Move(*ProperTransferFileName, *TransferFileName))
@@ -6322,10 +6364,15 @@ FSHAHash FShaderCompileJob::GetInputHash()
 
 	auto SerializeInputs = [this](FArchive& Archive)
 	{
-		check(Archive.IsSaving() && !Archive.IsLoading());
+		checkf(Archive.IsSaving() && !Archive.IsLoading(), TEXT("A loading archive is passed to FShaderCompileJob::GetInputHash(), this is not supported as it may corrupt its data"));
 
 		Archive << Input;
 		Archive << Input.Environment;
+
+		// hash the source file so changes to files during the development are picked up
+		const FSHAHash &SourceHash = GetShaderFileHash(*Input.VirtualSourceFilePath, Input.Target.GetPlatform());
+		Archive << const_cast<FSHAHash &>(SourceHash);
+
 		for (TMap<FString, FThreadSafeSharedStringPtr>::TConstIterator It(Input.Environment.IncludeVirtualPathToExternalContentsMap); It; ++It)
 		{
 			const FString& VirtualPath = It.Key();
@@ -6446,9 +6493,17 @@ FSHAHash FShaderPipelineCompileJob::GetInputHash()
 
 void FShaderPipelineCompileJob::SerializeOutput(FArchive& Ar)
 {
+	bool bAllStagesSucceeded = true;
 	for (int32 Index = 0, Num = StageJobs.Num(); Index < Num; ++Index)
 	{
 		StageJobs[Index]->SerializeOutput(Ar);
+		bAllStagesSucceeded = bAllStagesSucceeded && StageJobs[Index]->bSucceeded;
+	}
+
+	if (Ar.IsLoading())
+	{
+		bFinalized = true;
+		bSucceeded = bAllStagesSucceeded;
 	}
 }
 
diff --git a/Engine/Source/Runtime/RenderCore/Private/ShaderCore.cpp b/Engine/Source/Runtime/RenderCore/Private/ShaderCore.cpp
--- a/Engine/Source/Runtime/RenderCore/Private/ShaderCore.cpp
+++ b/Engine/Source/Runtime/RenderCore/Private/ShaderCore.cpp
@@ -226,6 +226,9 @@ private:
 	bool bInitialized;	
 };
 
+/** Protects GShaderHashCache from simultaneous modification by multiple threads. Note that it can cover more than one method of the class, e.g. a block of code doing Find() then Add() can be guarded */
+FCriticalSection GShaderHashAccessGuard;
+
 FShaderHashCache GShaderHashCache;
 
 /** Global map of virtual file path to physical file paths */
@@ -674,7 +677,14 @@ FString ParseVirtualShaderFilename(const FString& InFilename)
 
 bool ReplaceVirtualFilePathForShaderPlatform(FString& InOutVirtualFilePath, EShaderPlatform ShaderPlatform)
 {
-	const FString& PlatformIncludeDirectory = GShaderHashCache.GetPlatformIncludeDirectory(ShaderPlatform);
+	// as of 2021-03-01, it'd be safe to access just the include directory without the lock... but the lock (and copy) is here for the consistency's and future-proofness' sake
+	const FString PlatformIncludeDirectory( [ShaderPlatform]()
+		{
+			FScopeLock ShaderHashAccessLock(&GShaderHashAccessGuard);
+			return GShaderHashCache.GetPlatformIncludeDirectory(ShaderPlatform);
+		}()
+	);
+
 	if (PlatformIncludeDirectory.IsEmpty())
 	{
 		return false;
@@ -932,7 +942,11 @@ static void GetShaderIncludes(const TCHAR* EntryPointVirtualFilePath, const TCHA
 					bIgnoreInclude |= !CheckVirtualShaderFilePath(ExtractedIncludeFilename);
 
 					// Include only platform specific files, which will be used by the target platform.
-					bIgnoreInclude = bIgnoreInclude || GShaderHashCache.ShouldIgnoreInclude(ExtractedIncludeFilename, ShaderPlatform);
+					{
+						FScopeLock ShaderHashAccessLock(&GShaderHashAccessGuard);
+						bIgnoreInclude = bIgnoreInclude || GShaderHashCache.ShouldIgnoreInclude(ExtractedIncludeFilename, ShaderPlatform);
+					}
+
 
 					//vertex factories need to be handled separately
 					if (!bIgnoreInclude)
@@ -1016,6 +1030,7 @@ const FSHAHash& GetShaderFileHash(const TCHAR* VirtualFilePath, EShaderPlatform
 	// Make sure we are only accessing GShaderHashCache from one thread
 	//check(IsInGameThread() || IsAsyncLoading());
 	STAT(double HashTime = 0);
+	FScopeLock ShaderHashAccessLock(&GShaderHashAccessGuard);
 	{
 		SCOPE_SECONDS_COUNTER(HashTime);
 
@@ -1053,6 +1068,7 @@ const FSHAHash& GetShaderFilesHash(const TArray<FString>& VirtualFilePaths, ESha
 	// Make sure we are only accessing GShaderHashCache from one thread
 	//check(IsInGameThread() || IsAsyncLoading());
 	STAT(double HashTime = 0);
+	FScopeLock ShaderHashAccessLock(&GShaderHashAccessGuard);
 	{
 		SCOPE_SECONDS_COUNTER(HashTime);
 
@@ -1144,11 +1160,13 @@ void BuildShaderFileToUniformBufferMap(TMap<FString, TArray<const TCHAR*> >& Sha
 
 void InitializeShaderHashCache()
 {
+	FScopeLock ShaderHashAccessLock(&GShaderHashAccessGuard);
 	GShaderHashCache.Initialize();
 }
 
 void CheckShaderHashCacheInclude(const FString& VirtualFilePath, EShaderPlatform ShaderPlatform)
 {
+	FScopeLock ShaderHashAccessLock(&GShaderHashAccessGuard);
 	bool bIgnoreInclude = GShaderHashCache.ShouldIgnoreInclude(VirtualFilePath, ShaderPlatform);
 
 	checkf(!bIgnoreInclude,
@@ -1194,7 +1212,10 @@ void FlushShaderFileCache()
 {
 	UE_LOG(LogShaders, Log, TEXT("FlushShaderFileCache() begin"));
 
-	GShaderHashCache.Empty();
+	{
+		FScopeLock ShaderHashAccessLock(&GShaderHashAccessGuard);
+		GShaderHashCache.Empty();
+	}
 	GShaderFileCache.Empty();
 
 	if (!FPlatformProperties::RequiresCookedData())
