LevelStreaming: Handle case where level's current state is MakingVisible but ShouldBeVisible() returns false. Now set its TargetState to LoadedNotVisible an let next call to DetermineTargetState handle normal flow (typically for unneeded levels, target state will go back to Unloaded). Without this, unneeded inivisible loaded levels were first made visible before going back to invisible and unloaded.

diff --git a/Engine/Source/Runtime/Engine/Private/LevelStreaming.cpp b/Engine/Source/Runtime/Engine/Private/LevelStreaming.cpp
--- a/Engine/Source/Runtime/Engine/Private/LevelStreaming.cpp
+++ b/Engine/Source/Runtime/Engine/Private/LevelStreaming.cpp
@@ -416,7 +416,17 @@ bool ULevelStreaming::DetermineTargetState()
 	{
 	case ECurrentState::MakingVisible:
 		ensure(LoadedLevel);
-		TargetState = ETargetState::LoadedVisible;
+		if (!ShouldBeVisible() && GetWorld()->GetCurrentLevelPendingVisibility() != LoadedLevel)
+		{
+			// Since level doesn't need to be visible anymore, change TargetState to ETargetState::LoadedNotVisible.
+			// Next UpdateStreamingState will handle switching CurrentState to ECurrentState::LoadedNotVisible.
+			// From there, regular flow will properly handle TargetState.
+			TargetState = ETargetState::LoadedNotVisible;
+		}
+		else
+		{
+			TargetState = ETargetState::LoadedVisible;
+		}
 		break;
 
 	case ECurrentState::MakingInvisible:
@@ -570,7 +580,14 @@ void ULevelStreaming::UpdateStreamingState(bool& bOutUpdateAgain, bool& bOutRede
 	case ECurrentState::MakingVisible:
 		if (ensure(LoadedLevel))
 		{
-			if (!World->GetNextPreferredLevelPendingVisibility() || World->GetNextPreferredLevelPendingVisibility() == LoadedLevel)
+			// Handle case where MakingVisible is not needed anymore
+			if (TargetState == ETargetState::LoadedNotVisible)
+			{
+				CurrentState = ECurrentState::LoadedNotVisible;
+				bOutUpdateAgain = true;
+				bOutRedetermineTarget = true;
+			}
+			else if (!World->GetNextPreferredLevelPendingVisibility() || World->GetNextPreferredLevelPendingVisibility() == LoadedLevel)
 			{
 				World->AddToWorld(LoadedLevel, LevelTransform, !bShouldBlockOnLoad);
 
