Don't build/cache static mesh derived data if loading from a cooked package (as per UE4 behavior).

diff --git a/Engine/Source/Runtime/Engine/Private/StaticMesh.cpp b/Engine/Source/Runtime/Engine/Private/StaticMesh.cpp
--- a/Engine/Source/Runtime/Engine/Private/StaticMesh.cpp
+++ b/Engine/Source/Runtime/Engine/Private/StaticMesh.cpp
@@ -5210,25 +5210,28 @@ void UStaticMesh::ExecutePostLoadInternal(FStaticMeshPostLoadContext& Context)
 #if WITH_EDITOR
 	TRACE_CPUPROFILER_EVENT_SCOPE(UStaticMesh::ExecutePostLoadInternal)
 
-	if (IsBuiltAtRuntime())
+	if (!Context.bIsCookedForEditor)
 	{
-		// If built at runtime, but an editor build, we cache the mesh descriptions so that they can be rebuilt within the editor if necessary.
-		// This is done through the fast build path for consistency
-		TArray<const FMeshDescription*> MeshDescriptions;
-		const int32 NumSourceModels = GetNumSourceModels();
-		MeshDescriptions.Reserve(NumSourceModels);
-		for (int32 SourceModelIndex = 0; SourceModelIndex < NumSourceModels; SourceModelIndex++)
+		if (IsBuiltAtRuntime())
 		{
-			MeshDescriptions.Add(GetMeshDescription(SourceModelIndex));
+			// If built at runtime, but an editor build, we cache the mesh descriptions so that they can be rebuilt within the editor if necessary.
+			// This is done through the fast build path for consistency
+			TArray<const FMeshDescription*> MeshDescriptions;
+			const int32 NumSourceModels = GetNumSourceModels();
+			MeshDescriptions.Reserve(NumSourceModels);
+			for (int32 SourceModelIndex = 0; SourceModelIndex < NumSourceModels; SourceModelIndex++)
+			{
+				MeshDescriptions.Add(GetMeshDescription(SourceModelIndex));
+			}
+			BuildFromMeshDescriptions(MeshDescriptions);
+		}
+		else
+		{
+			// This, among many other things, will build a MeshDescription from the legacy RawMesh if one has not already been serialized,
+			// or, failing that, if there is not already one in the DDC. This will remain cached until the end of PostLoad(), upon which it
+			// is then released, and can be reloaded on demand.
+			CacheDerivedData();
 		}
-		BuildFromMeshDescriptions(MeshDescriptions);
-	}
-	else
-	{
-		// This, among many other things, will build a MeshDescription from the legacy RawMesh if one has not already been serialized,
-		// or, failing that, if there is not already one in the DDC. This will remain cached until the end of PostLoad(), upon which it
-		// is then released, and can be reloaded on demand.
-		CacheDerivedData();
 	}
 
 	GetBodySetup()->CreatePhysicsMeshes();
