Chaos physics - Fix RotationOfMass in all oriented implicit geometries ( Cylinder, Capsule and TaperedCylinder )
It now properly accounts for the orientation of the axis of the geometry compare to the Z axis ( that defines the main axis of the geometry for computing the InertiaTensor )
Had to change OrientationOfMass to be non static in order to have access to the axis properties ( similiar to GetInertiaTensor is implemented )
(including compilation fix)

diff --git a/Engine/Source/Runtime/Engine/Private/PhysicsEngine/Experimental/ChaosInterfaceUtils.cpp b/Engine/Source/Runtime/Engine/Private/PhysicsEngine/Experimental/ChaosInterfaceUtils.cpp
--- a/Engine/Source/Runtime/Engine/Private/PhysicsEngine/Experimental/ChaosInterfaceUtils.cpp
+++ b/Engine/Source/Runtime/Engine/Private/PhysicsEngine/Experimental/ChaosInterfaceUtils.cpp
@@ -103,7 +103,7 @@ namespace ChaosInterface
 			const auto& Vertex = PhysXMesh->getVertices()[j];
 			CollisionMeshParticles.X(j) = Scale * Chaos::FVec3(Vertex.x, Vertex.y, Vertex.z);
 		}
-		Chaos::TAABB<float, 3> BoundingBox(CollisionMeshParticles.X(0), CollisionMeshParticles.X(0));
+		Chaos::FAABB3 BoundingBox(CollisionMeshParticles.X(0), CollisionMeshParticles.X(0));
 		for (uint32 j = 1; j < CollisionMeshParticles.Size(); ++j)
 		{
 			BoundingBox.GrowToInclude(CollisionMeshParticles.X(j));
@@ -247,7 +247,7 @@ namespace ChaosInterface
 
 				const float CollisionMargin = FMath::Min(2.0f * HalfExtents.GetMin() * CollisionMarginFraction, CollisionMarginMax);
 
-				// TAABB can handle translations internally but if we have a rotation we need to wrap it in a transform
+				// AABB can handle translations internally but if we have a rotation we need to wrap it in a transform
 				TUniquePtr<Chaos::FImplicitObject> Implicit;
 				if (!BoxTransform.GetRotation().IsIdentity())
 				{
@@ -435,7 +435,7 @@ namespace ChaosInterface
 					OutMassProperties.Mass = OutMassProperties.Volume * InDensityKGPerCM;
 					OutMassProperties.InertiaTensor = Object.GetInertiaTensor(OutMassProperties.Mass);
 					OutMassProperties.CenterOfMass = LocalTM.TransformPosition(Object.GetCenterOfMass());
-					OutMassProperties.RotationOfMass = LocalTM.GetRotation();
+					OutMassProperties.RotationOfMass = LocalTM.GetRotation() * Object.GetRotationOfMass();
 				});
 			return true;
 		}
diff --git a/Engine/Source/Runtime/Experimental/Chaos/Public/Chaos/Capsule.h b/Engine/Source/Runtime/Experimental/Chaos/Public/Chaos/Capsule.h
--- a/Engine/Source/Runtime/Experimental/Chaos/Public/Chaos/Capsule.h
+++ b/Engine/Source/Runtime/Experimental/Chaos/Public/Chaos/Capsule.h
@@ -375,9 +375,12 @@ namespace Chaos
 			return PMatrix<T, 3, 3>(Diag12, Diag12, Diag3);
 		}
 
-		static TRotation<T, 3> GetRotationOfMass()
+		TRotation<T, 3> GetRotationOfMass() const { return GetRotationOfMass(GetAxis()); }
+		static TRotation<T, 3> GetRotationOfMass(const TVec3<T>& Axis)
 		{
-			return TRotation<T, 3>::FromIdentity(); 
+			// since the capsule stores an axis and the InertiaTensor is assumed to be along the ZAxis
+			// we need to make sure to return the rotation of the axis from Z
+			return TRotation<T, 3>::FromRotatedVector(TVector<T, 3>(0, 0, 1), Axis);
 		}
 
 		virtual uint32 GetTypeHash() const override
diff --git a/Engine/Source/Runtime/Experimental/Chaos/Public/Chaos/Convex.h b/Engine/Source/Runtime/Experimental/Chaos/Public/Chaos/Convex.h
--- a/Engine/Source/Runtime/Experimental/Chaos/Public/Chaos/Convex.h
+++ b/Engine/Source/Runtime/Experimental/Chaos/Public/Chaos/Convex.h
@@ -54,7 +54,7 @@ namespace Chaos
 		    : FImplicitObject(EImplicitObject::IsConvex | EImplicitObject::HasBoundingBox, ImplicitObjectType::Convex)
 			, Planes(MoveTemp(InPlanes))
 		    , Vertices(MoveTemp(InVertices))
-		    , LocalBoundingBox(TAABB<FReal, 3>::EmptyAABB())
+		    , LocalBoundingBox(FAABB3::EmptyAABB())
 		{
 			for (int32 ParticleIndex = 0; ParticleIndex < Vertices.Num(); ++ParticleIndex)
 			{
@@ -114,7 +114,7 @@ namespace Chaos
 			return 0.0f;
 		}
 
-		virtual const TAABB<FReal, 3> BoundingBox() const override
+		virtual const FAABB3 BoundingBox() const override
 		{
 			return LocalBoundingBox;
 		}
@@ -550,6 +550,11 @@ namespace Chaos
 			return LocalBoundingBox.GetInertiaTensor(Mass);
 		}
 
+		FRotation3 GetRotationOfMass() const
+		{
+			return FRotation3::FromIdentity();
+		}
+
 		const FVec3 GetCenterOfMass() const
 		{
 			return CenterOfMass;
@@ -717,7 +722,7 @@ namespace Chaos
 	private:
 		TArray<TPlaneConcrete<FReal, 3>> Planes;
 		TArray<FVec3> Vertices; //copy of the vertices that are just on the convex hull boundary
-		TAABB<FReal, 3> LocalBoundingBox;
+		FAABB3 LocalBoundingBox;
 		FConvexStructureData StructureData;
 		float Volume;
 		FVec3 CenterOfMass;
diff --git a/Engine/Source/Runtime/Experimental/Chaos/Public/Chaos/Cylinder.h b/Engine/Source/Runtime/Experimental/Chaos/Public/Chaos/Cylinder.h
--- a/Engine/Source/Runtime/Experimental/Chaos/Public/Chaos/Cylinder.h
+++ b/Engine/Source/Runtime/Experimental/Chaos/Public/Chaos/Cylinder.h
@@ -207,8 +207,13 @@ namespace Chaos
 			return PMatrix<T, 3, 3>(Diag12, Diag12, Diag3);
 		}
 
-		static TRotation<T, 3> GetRotationOfMass()
-		{ return TRotation<T, 3>(TVector<T, 3>(0), 1); }
+		TRotation<T, 3> GetRotationOfMass() const { return GetRotationOfMass(GetAxis()); }
+		static TRotation<T, 3> GetRotationOfMass(const TVec3<T>& Axis)
+		{ 
+			// since the cylinder stores an axis and the InertiaTensor is assumed to be along the ZAxis
+			// we need to make sure to return the rotation of the axis from Z
+			return TRotation<T, 3>::FromRotatedVector(TVector<T, 3>(0, 0, 1), Axis);
+		}
 
 		virtual void Serialize(FChaosArchive& Ar)
 		{
diff --git a/Engine/Source/Runtime/Experimental/Chaos/Public/Chaos/ImplicitObjectScaled.h b/Engine/Source/Runtime/Experimental/Chaos/Public/Chaos/ImplicitObjectScaled.h
--- a/Engine/Source/Runtime/Experimental/Chaos/Public/Chaos/ImplicitObjectScaled.h
+++ b/Engine/Source/Runtime/Experimental/Chaos/Public/Chaos/ImplicitObjectScaled.h
@@ -747,6 +747,12 @@ public:
 		return MScale * MObject->GetCenterOfMass();
 	}
 
+	FRotation3 GetRotationOfMass() const
+	{
+		// TODO: use the actual underlying object ( to be consistent with current implementation of GetInertiaTensor and GetVolume )
+		return BoundingBox().GetRotationOfMass();
+	}
+
 	const ObjectType Object() const { return MObject; }
 
 	// Only should be retrieved for copy purposes. Do not modify or access.
diff --git a/Engine/Source/Runtime/Experimental/Chaos/Public/Chaos/Sphere.h b/Engine/Source/Runtime/Experimental/Chaos/Public/Chaos/Sphere.h
--- a/Engine/Source/Runtime/Experimental/Chaos/Public/Chaos/Sphere.h
+++ b/Engine/Source/Runtime/Experimental/Chaos/Public/Chaos/Sphere.h
@@ -329,9 +329,9 @@ namespace Chaos
 			return PMatrix<T, d, d>(Diagonal, Diagonal, Diagonal);
 		}
 
-		static TRotation<T, d> GetRotationOfMass()
-		{ 
-			return TRotation<T, d>::FromIdentity(); 
+		TRotation<T, d> GetRotationOfMass() const
+		{
+			return TRotation<T, d>::FromIdentity();
 		}
 
 		virtual uint32 GetTypeHash() const override
diff --git a/Engine/Source/Runtime/Experimental/Chaos/Public/Chaos/TaperedCylinder.h b/Engine/Source/Runtime/Experimental/Chaos/Public/Chaos/TaperedCylinder.h
--- a/Engine/Source/Runtime/Experimental/Chaos/Public/Chaos/TaperedCylinder.h
+++ b/Engine/Source/Runtime/Experimental/Chaos/Public/Chaos/TaperedCylinder.h
@@ -327,8 +327,13 @@ namespace Chaos
 			return PMatrix<T, 3, 3>(Diag12, Diag12, Diag3);
 		}
 
-		static TRotation<T, 3> GetRotationOfMass()
-		{ return TRotation<T, 3>::FromIdentity(); }
+		TRotation<T, 3> GetRotationOfMass() const { return GetRotationOfMass(GetAxis()); }
+		static TRotation<T, 3> GetRotationOfMass(const TVec3<T>& Axis)
+		{ 
+			// since the cylinder stores an axis and the InertiaTensor is assumed to be along the ZAxis
+			// we need to make sure to return the rotation of the axis from Z
+			return TRotation<T, 3>::FromRotatedVector(TVector<T, 3>(0, 0, 1), Axis);
+		}
 
 		virtual uint32 GetTypeHash() const override
 		{
