Avoid redundant hash map lookups and FName construction in FStaticLightingSystemInterface::GetPreferredImplementation()

diff --git a/Engine/Source/Runtime/Engine/Private/Rendering/StaticLightingSystemInterface.cpp b/Engine/Source/Runtime/Engine/Private/Rendering/StaticLightingSystemInterface.cpp
--- a/Engine/Source/Runtime/Engine/Private/Rendering/StaticLightingSystemInterface.cpp
+++ b/Engine/Source/Runtime/Engine/Private/Rendering/StaticLightingSystemInterface.cpp
@@ -116,9 +116,10 @@ void FStaticLightingSystemInterface::UnregisterImplementation(FName Name)
 
 IStaticLightingSystemImpl* FStaticLightingSystemInterface::GetPreferredImplementation()
 {
-	if (Implementations.Find(FName(TEXT("GPULightmass"))) != nullptr)
+	static const FName NameGPULightmass(TEXT("GPULightmass"));
+	if (IStaticLightingSystemImpl** Found = Implementations.Find(NameGPULightmass))
 	{
-		return *Implementations.Find(FName(TEXT("GPULightmass")));
+		return *Found;
 	}
 	
 	if (Implementations.Num() > 0)
