New Animation Mirroring System:

+ UMirrorDataTable stores the bones / notifies / curves to mirror
	* Generates table based on FMirrorFindReplaceExpression in project settings
	* Can be used in any Animation Blueprint with a compatible skeleton
	* Support for self mirroring
	* Shared mirror axis
+ Project Settings for Mirroring Find & Replace Strings
+ FMirrorFindReplaceExpression
	* Find the mirrored name based on a prefix, suffix, or regular expression match
+ FAnimNode_Mirror
	* Inertialization support for blending between mirroring / unmirrored states over a configurable duration
	* Support for different MirrorDataTables (to support partial / full body mirroring)
+ Deprecated old mirroring system (FBoneMirrorExport) in SkeletalMesh

diff --git a/Engine/Source/Developer/AssetTools/Private/AssetTools.cpp b/Engine/Source/Developer/AssetTools/Private/AssetTools.cpp
--- a/Engine/Source/Developer/AssetTools/Private/AssetTools.cpp
+++ b/Engine/Source/Developer/AssetTools/Private/AssetTools.cpp
@@ -83,6 +83,7 @@
 #include "AssetTypeActions/AssetTypeActions_MaterialInstanceConstant.h"
 #include "AssetTypeActions/AssetTypeActions_MaterialInstanceDynamic.h"
 #include "AssetTypeActions/AssetTypeActions_MaterialParameterCollection.h"
+#include "AssetTypeActions/AssetTypeActions_MirrorDataTable.h"
 #include "AssetTypeActions/AssetTypeActions_ObjectLibrary.h"
 #include "AssetTypeActions/AssetTypeActions_ParticleSystem.h"
 #include "AssetTypeActions/AssetTypeActions_PhysicalMaterial.h"
@@ -263,6 +264,7 @@ UAssetToolsImpl::UAssetToolsImpl(const FObjectInitializer& ObjectInitializer)
 	RegisterAssetTypeActions(MakeShareable(new FAssetTypeActions_MaterialInstanceDynamic));
 	RegisterAssetTypeActions(MakeShareable(new FAssetTypeActions_MaterialInterface));
 	RegisterAssetTypeActions(MakeShareable(new FAssetTypeActions_MaterialParameterCollection));
+	RegisterAssetTypeActions(MakeShareable(new FAssetTypeActions_MirrorDataTable));
 	RegisterAssetTypeActions(MakeShareable(new FAssetTypeActions_ObjectLibrary));
 	RegisterAssetTypeActions(MakeShareable(new FAssetTypeActions_ParticleSystem));
 	RegisterAssetTypeActions(MakeShareable(new FAssetTypeActions_PhysicalMaterial));
diff --git a/Engine/Source/Developer/AssetTools/Private/AssetTypeActions/AssetTypeActions_MirrorDataTable.h b/Engine/Source/Developer/AssetTools/Private/AssetTypeActions/AssetTypeActions_MirrorDataTable.h
new file mode 100644
--- /dev/null
+++ b/Engine/Source/Developer/AssetTools/Private/AssetTypeActions/AssetTypeActions_MirrorDataTable.h
@@ -0,0 +1,16 @@
+// Copyright Epic Games, Inc. All Rights Reserved.
+
+#pragma once
+
+#include "AssetTypeActions_DataTable.h"
+#include "Engine/Classes/Animation/MirrorDataTable.h"
+
+class FAssetTypeActions_MirrorDataTable : public FAssetTypeActions_DataTable
+{
+public:
+	// IAssetTypeActions Implementation
+	virtual FText GetName() const override { return NSLOCTEXT("AssetTypeActions", "FAssetTypeActions_MirrorDataTable", "Mirror Data Table"); }
+	virtual UClass* GetSupportedClass() const override { return UMirrorDataTable::StaticClass(); }
+	virtual bool IsImportedAsset() const override { return false; }
+	// End IAssetTypeActions
+};
diff --git a/Engine/Source/Editor/AnimGraph/Classes/AnimGraphNode_Mirror.h b/Engine/Source/Editor/AnimGraph/Classes/AnimGraphNode_Mirror.h
new file mode 100644
--- /dev/null
+++ b/Engine/Source/Editor/AnimGraph/Classes/AnimGraphNode_Mirror.h
@@ -0,0 +1,41 @@
+// Copyright Epic Games, Inc. All Rights Reserved.
+
+#pragma once
+
+#include "CoreMinimal.h"
+#include "UObject/ObjectMacros.h"
+#include "AnimGraphNode_Base.h"
+#include "AnimNodes/AnimNode_Mirror.h"
+#include "AnimGraphNode_Mirror.generated.h"
+
+UCLASS(MinimalAPI)
+class UAnimGraphNode_Mirror : public UAnimGraphNode_Base
+{
+	GENERATED_BODY()
+
+public:
+	UPROPERTY(EditAnywhere, Category = Settings)
+	FAnimNode_Mirror Node;
+
+	virtual FLinearColor GetNodeTitleColor() const override;
+	virtual FText GetTooltipText() const override;
+	virtual FText GetNodeTitle(ENodeTitleType::Type TitleType) const override;
+	virtual FText GetMenuCategory() const override;
+	virtual void PostPlacedNewNode() override;
+	virtual void PostLoad() override;
+	virtual void GetOutputLinkAttributes(FNodeAttributeArray& OutAttributes) const override;
+	virtual void GetMenuActions(FBlueprintActionDatabaseRegistrar& ActionRegistrar) const override;
+	virtual bool IsActionFilteredOut(class FBlueprintActionFilter const& Filter) override;
+	virtual EAnimAssetHandlerType SupportsAssetClass(const UClass* AssetClass) const override;
+	virtual void PreEditChange(FProperty* PropertyThatWillChange) override;
+	virtual void PostEditChangeProperty(struct FPropertyChangedEvent& PropertyChangedEvent) override;
+	void ValidateAnimNodeDuringCompilation(class USkeleton* ForSkeleton, class FCompilerResultsLog& MessageLog) override;
+
+private:
+	static FText GetTitleGivenAssetInfo(const FText& AssetName);
+	void OnMirrorDataTableChanged(); 
+	void BindMirrorDataTableChangedDelegate(); 
+
+	/** Used for filtering in the Blueprint context menu when the sequence asset this node uses is unloaded */
+	FString UnloadedSkeletonName;
+};
diff --git a/Engine/Source/Editor/AnimGraph/Private/AnimGraphNode_Mirror.cpp b/Engine/Source/Editor/AnimGraph/Private/AnimGraphNode_Mirror.cpp
new file mode 100644
--- /dev/null
+++ b/Engine/Source/Editor/AnimGraph/Private/AnimGraphNode_Mirror.cpp
@@ -0,0 +1,277 @@
+// Copyright Epic Games, Inc. All Rights Reserved.UAnimGraphNode_MirrorPose
+
+#include "AnimGraphNode_Mirror.h"
+#include "Animation/AnimNode_Inertialization.h"
+#include "Animation/MirrorDataTable.h"
+#include "ARFilter.h"
+#include "AssetRegistryModule.h"
+#include "BlueprintActionFilter.h"
+#include "BlueprintActionFilter.h"
+#include "BlueprintActionDatabaseRegistrar.h"
+#include "BlueprintNodeSpawner.h"
+#include "Kismet2/CompilerResultsLog.h"
+
+#define LOCTEXT_NAMESPACE "AnimGraphNode_MirrorPose"
+
+
+FLinearColor UAnimGraphNode_Mirror::GetNodeTitleColor() const
+{
+	return FLinearColor(0.7f, 0.7f, 0.7f);
+}
+
+FText UAnimGraphNode_Mirror::GetTooltipText() const
+{
+	return LOCTEXT("NodeToolTip", "Mirror Pose");
+}
+
+FText UAnimGraphNode_Mirror::GetNodeTitle(ENodeTitleType::Type TitleType) const
+{
+	if (Node.GetMirrorDataTable())
+	{
+		FFormatNamedArguments Args;
+		Args.Add(TEXT("AssetName"), FText::FromName(Node.GetMirrorDataTable()->GetFName()));
+		return FText::Format(LOCTEXT("NodeTitle", "Mirror with {AssetName}"), Args);
+	}
+	else
+	{
+		return LOCTEXT("NodeTitle", "Mirror");
+	}
+}
+
+FText UAnimGraphNode_Mirror::GetMenuCategory() const
+{
+	return LOCTEXT("NodeCategory", "Mirroring");
+}
+
+void UAnimGraphNode_Mirror::GetOutputLinkAttributes(FNodeAttributeArray& OutAttributes) const
+{
+	if (Node.GetBlendTimeOnMirrorStateChange() > 0.0)
+	{
+		OutAttributes.Add(UE::Anim::IInertializationRequester::Attribute);
+	}
+}
+
+void UAnimGraphNode_Mirror::ValidateAnimNodeDuringCompilation(class USkeleton* ForSkeleton, class FCompilerResultsLog& MessageLog)
+{
+	Super::ValidateAnimNodeDuringCompilation(ForSkeleton, MessageLog);
+
+	if (Node.GetMirrorDataTable() == nullptr)
+	{
+		MessageLog.Error(TEXT("@@ does not have a mirror data table selected.  Please select a table or delete the node."), this);
+	}
+	BindMirrorDataTableChangedDelegate();
+}
+
+void UAnimGraphNode_Mirror::PreEditChange(FProperty* PropertyThatWillChange)
+{
+	Super::PreEditChange(PropertyThatWillChange);
+
+	if (PropertyThatWillChange && PropertyThatWillChange->GetFName() == GET_MEMBER_NAME_STRING_CHECKED(FAnimNode_Mirror, MirrorDataTable))
+	{
+		if (Node.GetMirrorDataTable())
+		{
+			Node.GetMirrorDataTable()->OnDataTableChanged().RemoveAll(this);
+		}
+	}
+}
+
+void  UAnimGraphNode_Mirror::PostPlacedNewNode()
+{
+	Super::PostPlacedNewNode();
+	BindMirrorDataTableChangedDelegate();
+}
+
+void  UAnimGraphNode_Mirror::PostLoad()
+{
+	Super::PostLoad(); 
+	BindMirrorDataTableChangedDelegate();
+}
+
+void UAnimGraphNode_Mirror::PostEditChangeProperty(struct FPropertyChangedEvent& PropertyChangedEvent)
+{
+	const FName PropertyName = (PropertyChangedEvent.Property ? PropertyChangedEvent.Property->GetFName() : NAME_None);
+
+	if (PropertyName == GET_MEMBER_NAME_STRING_CHECKED(FAnimNode_Mirror, MirrorDataTable))
+	{
+		BindMirrorDataTableChangedDelegate();
+	}
+
+	Super::PostEditChangeProperty(PropertyChangedEvent);
+}
+
+void UAnimGraphNode_Mirror::BindMirrorDataTableChangedDelegate()
+{
+	if (Node.GetMirrorDataTable())
+	{
+		Node.GetMirrorDataTable()->OnDataTableChanged().RemoveAll(this);
+		Node.GetMirrorDataTable()->OnDataTableChanged().AddUObject(this, &UAnimGraphNode_Mirror::OnMirrorDataTableChanged);
+	}
+}
+
+void UAnimGraphNode_Mirror::OnMirrorDataTableChanged()
+{
+	if (HasValidBlueprint())
+	{
+		UBlueprint* BP = GetBlueprint();
+		BP->Status = BS_Dirty;
+	}
+}
+
+void UAnimGraphNode_Mirror::GetMenuActions(FBlueprintActionDatabaseRegistrar& ActionRegistrar) const
+{
+	auto LoadedAssetSetup = [](UEdGraphNode* NewNode, bool /*bIsTemplateNode*/, TWeakObjectPtr<UMirrorDataTable> MirrorDataTablePtr)
+	{
+		UAnimGraphNode_Mirror* MirrorNode = CastChecked<UAnimGraphNode_Mirror>(NewNode);
+		MirrorNode->Node.SetMirrorDataTable(MirrorDataTablePtr.Get());
+	};
+
+	auto UnloadedAssetSetup = [](UEdGraphNode* NewNode, bool bIsTemplateNode, const FAssetData AssetData)
+	{
+		UAnimGraphNode_Mirror* MirrorNode = CastChecked<UAnimGraphNode_Mirror>(NewNode);
+		if (bIsTemplateNode)
+		{
+			AssetData.GetTagValue("Skeleton", MirrorNode->UnloadedSkeletonName);
+		}
+		else
+		{
+			UMirrorDataTable* MirrorTable = Cast<UMirrorDataTable>(AssetData.GetAsset());
+			check(MirrorTable != nullptr);
+			MirrorNode->Node.SetMirrorDataTable(MirrorTable);
+		}
+	};
+
+	const UObject* QueryObject = ActionRegistrar.GetActionKeyFilter();
+	if (QueryObject == nullptr)
+	{
+		FAssetRegistryModule& AssetRegistryModule = FModuleManager::LoadModuleChecked<FAssetRegistryModule>(TEXT("AssetRegistry"));
+		// define a filter to help in pulling UMirrrorDataTable asset data from the registry
+		FARFilter Filter;
+		Filter.ClassNames.Add(UMirrorDataTable::StaticClass()->GetFName());
+
+		Filter.bRecursiveClasses = true;
+		// Find matching assets and add an entry for each one
+		TArray<FAssetData> MirrorDataTableList;
+		AssetRegistryModule.Get().GetAssets(Filter, /*out*/MirrorDataTableList);
+
+		for (auto AssetIt = MirrorDataTableList.CreateConstIterator(); AssetIt; ++AssetIt)
+		{
+			const FAssetData& Asset = *AssetIt;
+
+			UBlueprintNodeSpawner* NodeSpawner = UBlueprintNodeSpawner::Create(GetClass());
+			if (Asset.IsAssetLoaded())
+			{
+				UMirrorDataTable* MirrorDataTable = Cast<UMirrorDataTable>(Asset.GetAsset());
+				if (MirrorDataTable)
+				{
+					NodeSpawner->CustomizeNodeDelegate = UBlueprintNodeSpawner::FCustomizeNodeDelegate::CreateStatic(LoadedAssetSetup, TWeakObjectPtr<UMirrorDataTable>(MirrorDataTable));
+					NodeSpawner->DefaultMenuSignature.MenuName = GetTitleGivenAssetInfo(FText::FromName(MirrorDataTable->GetFName()));
+				}
+			}
+			else
+			{
+				NodeSpawner->CustomizeNodeDelegate = UBlueprintNodeSpawner::FCustomizeNodeDelegate::CreateStatic(UnloadedAssetSetup, Asset);
+				NodeSpawner->DefaultMenuSignature.MenuName = GetTitleGivenAssetInfo(FText::FromName(Asset.AssetName));
+			}
+			ActionRegistrar.AddBlueprintAction(Asset, NodeSpawner);
+		}
+	}
+	else if (const UMirrorDataTable* MirrorDataTable = Cast<UMirrorDataTable>(QueryObject))
+	{
+		UBlueprintNodeSpawner* NodeSpawner = UBlueprintNodeSpawner::Create(GetClass());
+
+		TWeakObjectPtr<UMirrorDataTable> MirrorDataTablePtr = MakeWeakObjectPtr(const_cast<UMirrorDataTable*>(MirrorDataTable));
+		NodeSpawner->CustomizeNodeDelegate = UBlueprintNodeSpawner::FCustomizeNodeDelegate::CreateStatic(LoadedAssetSetup, MirrorDataTablePtr);
+		NodeSpawner->DefaultMenuSignature.MenuName = GetTitleGivenAssetInfo(FText::FromName(MirrorDataTablePtr->GetFName()));
+		NodeSpawner->DefaultMenuSignature.Tooltip = GetTitleGivenAssetInfo(FText::FromString(MirrorDataTablePtr->GetPathName()));
+
+		ActionRegistrar.AddBlueprintAction(QueryObject, NodeSpawner);
+	}
+	else if (QueryObject == GetClass())
+	{
+		FAssetRegistryModule& AssetRegistryModule = FModuleManager::LoadModuleChecked<FAssetRegistryModule>(TEXT("AssetRegistry"));
+		// define a filter to help in pulling UAnimSequence asset data from the registry
+		FARFilter Filter;
+		Filter.ClassNames.Add(UMirrorDataTable::StaticClass()->GetFName());
+		Filter.bRecursiveClasses = true;
+		// Find matching assets and add an entry for each one
+		TArray<FAssetData> MirrorDataTableList;
+		AssetRegistryModule.Get().GetAssets(Filter, /*out*/MirrorDataTableList);
+
+		for (auto AssetIt = MirrorDataTableList.CreateConstIterator(); AssetIt; ++AssetIt)
+		{
+			const FAssetData& Asset = *AssetIt;
+			if (Asset.IsAssetLoaded())
+			{
+				continue;
+			}
+
+			UBlueprintNodeSpawner* NodeSpawner = UBlueprintNodeSpawner::Create(GetClass());
+			NodeSpawner->CustomizeNodeDelegate = UBlueprintNodeSpawner::FCustomizeNodeDelegate::CreateStatic(UnloadedAssetSetup, Asset);
+
+			NodeSpawner->DefaultMenuSignature.MenuName = GetTitleGivenAssetInfo(FText::FromName(Asset.AssetName));
+			NodeSpawner->DefaultMenuSignature.Tooltip = GetTitleGivenAssetInfo(FText::FromName(Asset.ObjectPath));
+			ActionRegistrar.AddBlueprintAction(Asset, NodeSpawner);
+		}
+	}
+}
+
+bool UAnimGraphNode_Mirror::IsActionFilteredOut(class FBlueprintActionFilter const& Filter)
+{
+	bool bIsFilteredOut = false;
+	FBlueprintActionContext const& FilterContext = Filter.Context;
+
+	for (UBlueprint* Blueprint : FilterContext.Blueprints)
+	{
+		if (UAnimBlueprint* AnimBlueprint = Cast<UAnimBlueprint>(Blueprint))
+		{
+			if (Node.GetMirrorDataTable())
+			{
+				if (Node.GetMirrorDataTable()->Skeleton != AnimBlueprint->TargetSkeleton)
+				{
+					// Mirror Data Table does not use the same skeleton as the Blueprint, cannot use
+					bIsFilteredOut = true;
+					break;
+				}
+			}
+			else
+			{
+				FAssetData SkeletonData(AnimBlueprint->TargetSkeleton);
+				if (UnloadedSkeletonName != SkeletonData.GetExportTextName())
+				{
+					bIsFilteredOut = true;
+					break;
+				}
+			}
+		}
+		else
+		{
+			// Not an animation Blueprint, cannot use
+			bIsFilteredOut = true;
+			break;
+		}
+	}
+	return bIsFilteredOut;
+}
+
+FText  UAnimGraphNode_Mirror::GetTitleGivenAssetInfo(const FText& AssetName)
+{
+	FFormatNamedArguments Args;
+	Args.Add(TEXT("AssetName"), AssetName);
+
+	return FText::Format(LOCTEXT("MirrorNodeTitle", "Mirror with {AssetName}"), Args);
+}
+
+EAnimAssetHandlerType UAnimGraphNode_Mirror::SupportsAssetClass(const UClass* AssetClass) const
+{
+	if (AssetClass->IsChildOf(UMirrorDataTable::StaticClass()))
+	{
+		return EAnimAssetHandlerType::PrimaryHandler;
+	}
+	else
+	{
+		return EAnimAssetHandlerType::NotSupported;
+	}
+}
+
+
+#undef LOCTEXT_NAMESPACE
diff --git a/Engine/Source/Editor/UnrealEd/Classes/Factories/MirrorDataTableFactory.h b/Engine/Source/Editor/UnrealEd/Classes/Factories/MirrorDataTableFactory.h
new file mode 100644
--- /dev/null
+++ b/Engine/Source/Editor/UnrealEd/Classes/Factories/MirrorDataTableFactory.h
@@ -0,0 +1,55 @@
+// Copyright Epic Games, Inc. All Rights Reserved.
+
+#pragma once
+
+#include "CoreMinimal.h"
+#include "UObject/ObjectMacros.h"
+#include "Factories/Factory.h"
+#include "Engine/Classes/Animation/MirrorDataTable.h"
+#include "MirrorDataTableFactory.generated.h"
+
+class UMirrorDataTable;
+
+UCLASS(MinimalAPI)
+class UMirrorTableFindReplaceExpressions : public UObject
+{
+	GENERATED_BODY()
+	UMirrorTableFindReplaceExpressions(const FObjectInitializer& ObjectInitializer);
+public:
+	// Collection of animations for motion matching
+	UPROPERTY(EditAnywhere, Category = TableCreation)
+	TArray<FMirrorFindReplaceExpression> FindReplaceExpressions;
+};
+
+UCLASS(hidecategories=Object)
+class UNREALED_API UMirrorDataTableFactory : public UFactory
+{
+	GENERATED_BODY()
+	UMirrorDataTableFactory(const FObjectInitializer& ObjectInitializer);
+
+	UPROPERTY()
+	const class UScriptStruct* Struct;
+
+	UPROPERTY()
+	class USkeleton* Skeleton;
+
+	UPROPERTY()
+	class UMirrorTableFindReplaceExpressions* MirrorFindReplaceExpressions;
+
+public:
+	//~ Begin UFactory Interface
+	virtual bool ConfigureProperties() override;
+	virtual UObject* FactoryCreateNew(UClass* Class, UObject* InParent, FName Name, EObjectFlags Flags, UObject* Context, FFeedbackContext* Warn) override;
+	//~ Begin UFactory Interface
+
+		/** Returns the name of the factory for menus */
+	virtual FText GetDisplayName() const override;
+
+	/** When shown in menus, this is the category containing this factory. Return type is a BitFlag mask using EAssetTypeCategories. */
+	virtual uint32 GetMenuCategories() const override;
+protected:
+
+	void OnWindowUserActionDelegate(bool bCreate, USkeleton* InSkeleton, UMirrorTableFindReplaceExpressions* InFindReplaceExpressions, const UScriptStruct* InResultStruct);
+	virtual UMirrorDataTable* MakeNewMirrorDataTable(UObject* InParent, FName Name, EObjectFlags Flags);
+};
+
diff --git a/Engine/Source/Editor/UnrealEd/Private/DataTableEditorUtils.cpp b/Engine/Source/Editor/UnrealEd/Private/DataTableEditorUtils.cpp
--- a/Engine/Source/Editor/UnrealEd/Private/DataTableEditorUtils.cpp
+++ b/Engine/Source/Editor/UnrealEd/Private/DataTableEditorUtils.cpp
@@ -752,6 +752,7 @@ void FDataTableEditorUtils::BroadcastPostChange(UDataTable* DataTable, EDataTabl
 		}
 	}
 	FDataTableEditorManager::Get().PostChange(DataTable, Info);
+	DataTable->OnDataTableChanged().Broadcast();
 }
 
 void FDataTableEditorUtils::CacheDataTableForEditing(const UDataTable* DataTable, TArray<FDataTableEditorColumnHeaderDataPtr>& OutAvailableColumns, TArray<FDataTableEditorRowListViewDataPtr>& OutAvailableRows)
diff --git a/Engine/Source/Editor/UnrealEd/Private/Factories/MirrorDataTableFactory.cpp b/Engine/Source/Editor/UnrealEd/Private/Factories/MirrorDataTableFactory.cpp
new file mode 100644
--- /dev/null
+++ b/Engine/Source/Editor/UnrealEd/Private/Factories/MirrorDataTableFactory.cpp
@@ -0,0 +1,385 @@
+// Copyright Epic Games, Inc. All Rights Reserved.
+
+#include "Factories/MirrorDataTableFactory.h"
+#include "DataTableEditorUtils.h"
+#include "Engine/Classes/Animation/MirrorDataTable.h"
+#include "Engine/Classes/Animation/AnimationSettings.h"
+#include "StructViewerModule.h"
+#include "StructViewerFilter.h"
+
+#include "PropertyEditorModule.h"
+#include "IContentBrowserSingleton.h"
+#include "ContentBrowserModule.h"
+
+#include "Widgets/Layout/SUniformGridPanel.h"
+#include "Widgets/Layout/SScrollBox.h"
+#include "Widgets/Input/SButton.h"
+#include "Widgets/Layout/SExpandableArea.h"
+#include "Widgets/SWindow.h"
+
+#define LOCTEXT_NAMESPACE "MirrorDataTableFactory"
+
+DECLARE_DELEGATE_FourParams(FOnMirrorTableConfigureUserAction, bool /*bCreate*/, USkeleton* /*InSkeleton*/, UMirrorTableFindReplaceExpressions* /*InRegularExpressions*/, const UScriptStruct* /*InPoseNames*/);
+
+class FDataTableStructFilter : public IStructViewerFilter
+{
+public:
+	virtual bool IsStructAllowed(const FStructViewerInitializationOptions& InInitOptions, const UScriptStruct* InStruct, TSharedRef<FStructViewerFilterFuncs> InFilterFuncs) override
+	{
+		return FDataTableEditorUtils::IsValidTableStruct(InStruct) && InStruct->FindPropertyByName(TEXT("MirroredName")) != nullptr;
+	}
+
+	virtual bool IsUnloadedStructAllowed(const FStructViewerInitializationOptions& InInitOptions, const FName InStructPath, TSharedRef<FStructViewerFilterFuncs> InFilterFuncs) override
+	{
+		// Unloaded structs are always User Defined Structs, and User Defined Structs are always allowed
+		// They will be re-validated by IsStructAllowed once loaded during the pick
+		//return FDataTableEditorUtils::IsValidTableStruct(InStruct) && InStruct->FindPropertyByName(TEXT("MirroredName")) != nullptr;
+
+		return false;
+	}
+};
+
+class SMirrorDataTableFactoryWindow : public SWindow
+{
+public:
+	SLATE_BEGIN_ARGS(SMirrorDataTableFactoryWindow)
+	{}
+
+	SLATE_ARGUMENT(USkeleton*, Skeleton)
+	SLATE_ARGUMENT(UMirrorTableFindReplaceExpressions*, FindReplaceExpressions)
+	SLATE_ARGUMENT(FOnMirrorTableConfigureUserAction, UserActionHandler)
+	SLATE_ARGUMENT(FSimpleDelegate, OnCreateCanceled)
+		
+	SLATE_END_ARGS()
+public:
+	FReply OnCreate()
+	{
+		check(ResultStruct);
+		if (PickerWindow.IsValid())
+		{
+			PickerWindow->RequestDestroyWindow();
+		}
+		return FReply::Handled();
+	}
+
+	FReply OnCancel()
+	{
+		if (UserActionHandler.IsBound())
+		{
+			UserActionHandler.Execute(false, nullptr,nullptr, nullptr);
+		}
+
+		RequestDestroyWindow();
+		return FReply::Handled();
+	}
+
+	bool IsStructSelected() const
+	{
+		return ResultStruct != nullptr;
+	}
+
+	void OnPickedStruct(const UScriptStruct* ChosenStruct)
+	{
+		ResultStruct = ChosenStruct;
+		StructPickerAnchor->SetIsOpen(false);
+	}
+
+	FText OnGetComboTextValue() const
+	{
+		return ResultStruct
+			? FText::AsCultureInvariant(ResultStruct->GetName())
+			: LOCTEXT("None", "None");
+	}
+
+	TSharedRef<SWidget> GenerateStructPicker()
+	{
+		FStructViewerModule& StructViewerModule = FModuleManager::LoadModuleChecked<FStructViewerModule>("StructViewer");
+
+		// Fill in options
+		FStructViewerInitializationOptions Options;
+		Options.Mode = EStructViewerMode::StructPicker;
+		Options.StructFilter = MakeShared<FDataTableStructFilter>();
+		return
+			SNew(SBox)
+			.WidthOverride(330)
+			[
+				SNew(SVerticalBox)
+
+				+ SVerticalBox::Slot()
+				.FillHeight(1.0f)
+				.MaxHeight(500)
+				[
+					SNew(SBorder)
+					.Padding(4)
+					.BorderImage(FEditorStyle::GetBrush("ToolPanel.GroupBorder"))
+					[
+						StructViewerModule.CreateStructViewer(Options, FOnStructPicked::CreateSP(this, &SMirrorDataTableFactoryWindow::OnPickedStruct))
+					]
+				]
+			];
+	}
+
+
+	void Construct(const FArguments& InArgs)
+	{
+		UserActionHandler = InArgs._UserActionHandler;
+		Skeleton = InArgs._Skeleton;
+		FindReplaceExpressions = InArgs._FindReplaceExpressions;
+		ResultStruct = FMirrorTableRow::StaticStruct();
+		FStructViewerModule& StructViewerModule = FModuleManager::LoadModuleChecked<FStructViewerModule>("StructViewer");
+		
+
+		// Load the content browser module to display an asset picker
+		FContentBrowserModule& ContentBrowserModule = FModuleManager::LoadModuleChecked<FContentBrowserModule>("ContentBrowser");
+
+		FAssetPickerConfig AssetPickerConfig;
+		/** The asset picker will only show skeletons */
+		AssetPickerConfig.Filter.ClassNames.Add(USkeleton::StaticClass()->GetFName());
+		/** The delegate that fires when an asset was selected */
+		AssetPickerConfig.OnAssetSelected = FOnAssetSelected::CreateRaw(this, &SMirrorDataTableFactoryWindow::OnSkeletonSelected);
+		if (Skeleton != nullptr)
+		{
+			AssetPickerConfig.InitialAssetSelection = FAssetData(Skeleton);
+		}
+
+		// Create a property view
+		FPropertyEditorModule& EditModule = FModuleManager::Get().GetModuleChecked<FPropertyEditorModule>("PropertyEditor");
+
+		FDetailsViewArgs DetailsViewArgs;
+		DetailsViewArgs.NameAreaSettings = FDetailsViewArgs::HideNameArea;
+		DetailsViewArgs.bAllowSearch = false;
+		DetailsViewArgs.bHideSelectionTip = true;
+		DetailsViewArgs.bShowOptions = false;
+
+		PropertyView = EditModule.CreateDetailView(DetailsViewArgs);
+		PropertyView->SetObject(FindReplaceExpressions);
+		PropertyView->SetClipping(EWidgetClipping::ClipToBounds);
+
+		/** The default view mode should be a list view */
+		AssetPickerConfig.InitialAssetViewType = EAssetViewType::List;
+
+		SWindow::Construct(SWindow::FArguments()
+			.Title(LOCTEXT("CreateMirrorDataTableOptions", "Create Mirror Data Table Asset"))
+			.SizingRule(ESizingRule::UserSized)
+			.ClientSize(FVector2D(500, 600))
+			.SupportsMinimize(false)
+			.SupportsMaximize(false)
+		[
+			SNew(SBorder)
+			.BorderImage( FEditorStyle::GetBrush("Menu.Background") )
+			[
+				SNew(SVerticalBox)
+				+SVerticalBox::Slot()
+				.FillHeight(0.6)
+				.Padding(3, 3)
+				[
+					SNew(SVerticalBox)
+
+					+SVerticalBox::Slot()
+					.AutoHeight()
+					.Padding(2, 2)
+					[
+						SNew(STextBlock)
+						.Text(FText::FromString(TEXT("Select Skeleton")))
+					]
+					+SVerticalBox::Slot()
+					.Padding(2, 2)
+					[
+						ContentBrowserModule.Get().CreateAssetPicker(AssetPickerConfig)
+					]
+				]
+				+ SVerticalBox::Slot()
+				.AutoHeight()
+				.MaxHeight(300)
+				[
+					SNew(SExpandableArea)
+					.AreaTitle(LOCTEXT("CreateMirrorDataTableAdvancedOptions", "Advanced Options"))
+					.Padding(2)
+					.InitiallyCollapsed(true)
+					.BodyContent()
+					[
+						SNew(SVerticalBox)
+						+SVerticalBox::Slot()
+						.AutoHeight()
+						[
+							SNew(SUniformGridPanel)
+							.SlotPadding(FEditorStyle::GetMargin("StandardDialog.SlotPadding"))
+							.MinDesiredSlotWidth(FEditorStyle::GetFloat("StandardDialog.MinDesiredSlotWidth"))
+							.MinDesiredSlotHeight(FEditorStyle::GetFloat("StandardDialog.MinDesiredSlotHeight"))
+							+ SUniformGridPanel::Slot(0, 0)
+							.VAlign(VAlign_Center)
+							[
+								SNew(STextBlock)
+								.Text(LOCTEXT("MirrorRowStructure", "Row Structure"))
+							]
+							+ SUniformGridPanel::Slot(1, 0)
+							[
+								SAssignNew(StructPickerAnchor, SComboButton)
+								.ContentPadding(FMargin(2, 2, 2, 1))
+								.MenuPlacement(MenuPlacement_BelowAnchor)
+								.ButtonContent()
+								[
+									SNew(STextBlock)
+									.Text(this, &SMirrorDataTableFactoryWindow::OnGetComboTextValue)
+								]
+								.OnGetMenuContent(this, &SMirrorDataTableFactoryWindow::GenerateStructPicker)
+							]
+						]
+						+SVerticalBox::Slot()
+						.FillHeight(1.0)
+						[
+							SNew(SScrollBox)
+							.Orientation(Orient_Vertical)
+							+SScrollBox::Slot()
+							[
+								PropertyView.ToSharedRef()
+							]
+						]
+					]
+				]
+				+SVerticalBox::Slot()
+				.AutoHeight()
+				.Padding(3,3)
+				[
+					SNew(SUniformGridPanel)
+					.SlotPadding(FEditorStyle::GetMargin("StandardDialog.SlotPadding"))
+					.MinDesiredSlotWidth(FEditorStyle::GetFloat("StandardDialog.MinDesiredSlotWidth"))
+					.MinDesiredSlotHeight(FEditorStyle::GetFloat("StandardDialog.MinDesiredSlotHeight"))
+					+ SUniformGridPanel::Slot(0, 0)
+					[
+						SNew(SButton)
+						.Text(LOCTEXT("Accept", "Accept"))
+						.HAlign(HAlign_Center)
+						.ContentPadding(FEditorStyle::GetMargin("StandardDialog.ContentPadding"))
+						.IsEnabled(this, &SMirrorDataTableFactoryWindow::CanAccept)
+						.OnClicked_Raw(this, &SMirrorDataTableFactoryWindow::OnAccept)
+					]
+					+ SUniformGridPanel::Slot(1, 0)
+					[
+						SNew(SButton)
+						.Text(LOCTEXT("Cancel", "Cancel"))
+						.HAlign(HAlign_Center)
+						.ContentPadding(FEditorStyle::GetMargin("StandardDialog.ContentPadding"))
+						.OnClicked_Raw(this, &SMirrorDataTableFactoryWindow::OnCancel)
+					]
+				]
+			]
+		]);
+	}
+
+	bool CanAccept() const
+	{
+		return Skeleton != nullptr &&  ResultStruct != nullptr && UserActionHandler.IsBound();
+	}
+
+	FReply OnAccept()
+	{
+		if (CanAccept())
+		{
+			UserActionHandler.Execute(true, Skeleton, FindReplaceExpressions, ResultStruct);
+		}
+
+		RequestDestroyWindow();
+		return FReply::Handled();
+	}
+
+	void OnSkeletonSelected(const FAssetData& SelectedAsset)
+	{
+		Skeleton = Cast<USkeleton>(SelectedAsset.GetAsset());
+	}
+
+
+	USkeleton*					Skeleton;
+	TSharedPtr<SWindow>			PickerWindow;
+	TSharedPtr<SComboButton>	StructPickerAnchor;
+	FOnMirrorTableConfigureUserAction	UserActionHandler;
+	TSharedPtr<class IDetailsView> PropertyView;
+	UMirrorTableFindReplaceExpressions* FindReplaceExpressions;
+	const UScriptStruct* ResultStruct = nullptr;
+};
+
+UMirrorTableFindReplaceExpressions::UMirrorTableFindReplaceExpressions(const FObjectInitializer& ObjectInitializer)
+	: Super(ObjectInitializer)
+{
+	UAnimationSettings* AnimationSettings = UAnimationSettings::Get(); 
+	if (AnimationSettings)
+	{
+		FindReplaceExpressions = AnimationSettings->MirrorFindReplaceExpressions;
+	}
+}
+
+UMirrorDataTableFactory::UMirrorDataTableFactory(const FObjectInitializer& ObjectInitializer)
+	: Super(ObjectInitializer), Struct(nullptr), Skeleton(nullptr)
+{
+	SupportedClass = UMirrorDataTable::StaticClass();
+	bCreateNew = true;
+	bEditAfterNew = true;
+	MirrorFindReplaceExpressions = NewObject<UMirrorTableFindReplaceExpressions>();
+}
+
+UMirrorDataTable* UMirrorDataTableFactory::MakeNewMirrorDataTable(UObject* InParent, FName Name, EObjectFlags Flags)
+{
+	return NewObject<UMirrorDataTable>(InParent, Name, Flags);
+}
+
+UObject* UMirrorDataTableFactory::FactoryCreateNew(UClass* Class, UObject* InParent, FName Name, EObjectFlags Flags, UObject* Context, FFeedbackContext* Warn)
+{
+	UMirrorDataTable* DataTable = nullptr;
+	if (Struct && ensure(SupportedClass == Class) && Skeleton)
+	{
+		ensure(0 != (RF_Public & Flags));
+		DataTable = MakeNewMirrorDataTable(InParent, Name, Flags);
+		if (DataTable)
+		{
+			DataTable->RowStruct = const_cast<UScriptStruct*>(ToRawPtr(Struct));
+			DataTable->Skeleton = Skeleton; 
+			if (MirrorFindReplaceExpressions)
+			{
+				DataTable->MirrorFindReplaceExpressions = MirrorFindReplaceExpressions->FindReplaceExpressions;
+			}
+			if (Skeleton)
+			{
+				DataTable->FindReplaceMirroredNames(); 
+			}
+		}
+	}
+	return DataTable;
+}
+
+FText UMirrorDataTableFactory::GetDisplayName() const
+{
+	return LOCTEXT("MirrorDataTableFactoryDescription", "Mirror Data Table");
+}
+
+uint32 UMirrorDataTableFactory::GetMenuCategories() const
+{
+	return EAssetTypeCategories::Animation;
+}
+
+bool UMirrorDataTableFactory::ConfigureProperties()
+{
+	TSharedPtr<SWindow> PickerWindow = SNew(SMirrorDataTableFactoryWindow)
+		.Skeleton(Skeleton)
+		.FindReplaceExpressions(MirrorFindReplaceExpressions)
+		.UserActionHandler(FOnMirrorTableConfigureUserAction::CreateUObject(this, &UMirrorDataTableFactory::OnWindowUserActionDelegate));
+
+	// have to clear after setting it because you could use close button without clicking cancel
+	Skeleton = nullptr;
+
+	GEditor->EditorAddModalWindow(PickerWindow.ToSharedRef());
+	PickerWindow.Reset();
+
+	return Skeleton != nullptr && Struct != nullptr;
+}
+
+void  UMirrorDataTableFactory::OnWindowUserActionDelegate(bool bCreate, USkeleton* InSkeleton, UMirrorTableFindReplaceExpressions* InFindReplaceExpressions, const UScriptStruct* InResultStruct)
+{
+	if (bCreate)
+	{
+		Skeleton = InSkeleton;
+		MirrorFindReplaceExpressions = InFindReplaceExpressions;
+		Struct = InResultStruct; 
+	}
+}
+
+#undef LOCTEXT_NAMESPACE // "MirrorDataTableFactory"
\ No newline at end of file
diff --git a/Engine/Source/Editor/UnrealEd/Private/Factories/SkeletalMeshImport.cpp b/Engine/Source/Editor/UnrealEd/Private/Factories/SkeletalMeshImport.cpp
--- a/Engine/Source/Editor/UnrealEd/Private/Factories/SkeletalMeshImport.cpp
+++ b/Engine/Source/Editor/UnrealEd/Private/Factories/SkeletalMeshImport.cpp
@@ -468,7 +468,7 @@ TSharedPtr<FExistingSkelMeshData> SkeletalMeshHelper::SaveExistingSkelMeshData(U
 	// since copying back original skeleton, this should be safe to do
 	ExistingMeshDataPtr->ExistingPostProcessAnimBlueprint = SourceSkeletalMesh->GetPostProcessAnimBlueprint();
 	ExistingMeshDataPtr->ExistingLODSettings = SourceSkeletalMesh->GetLODSettings();
-	SourceSkeletalMesh->ExportMirrorTable(ExistingMeshDataPtr->ExistingMirrorTable);
+
 	ExistingMeshDataPtr->ExistingMorphTargets = SourceSkeletalMesh->GetMorphTargets();
 	ExistingMeshDataPtr->ExistingAssetImportData = SourceSkeletalMesh->GetAssetImportData();
 	ExistingMeshDataPtr->ExistingThumbnailInfo = SourceSkeletalMesh->GetThumbnailInfo();
@@ -962,8 +962,6 @@ void SkeletalMeshHelper::RestoreExistingSkelMeshData(TSharedPtr<const FExistingS
 		SkeletalMesh->SetSkeleton(MeshData->ExistingSkeleton);
 		SkeletalMesh->SetPostProcessAnimBlueprint(MeshData->ExistingPostProcessAnimBlueprint);
 
-		// Copy mirror table.
-		SkeletalMesh->ImportMirrorTable(MeshData->ExistingMirrorTable);
 		SkeletalMesh->GetMorphTargets().Empty(MeshData->ExistingMorphTargets.Num());
 		SkeletalMesh->GetMorphTargets().Append(MeshData->ExistingMorphTargets);
 		SkeletalMesh->InitMorphTargets();
diff --git a/Engine/Source/Editor/UnrealEd/Public/SkelImport.h b/Engine/Source/Editor/UnrealEd/Public/SkelImport.h
--- a/Engine/Source/Editor/UnrealEd/Public/SkelImport.h
+++ b/Engine/Source/Editor/UnrealEd/Public/SkelImport.h
@@ -60,7 +60,10 @@ struct FExistingSkelMeshData
 	bool									bExistingUseFullPrecisionUVs;
 	bool									bExistingUseHighPrecisionTangentBasis;
 
+	UE_DEPRECATED(5.0, "Please use UMirrorDataTable for mirroring ")
+	PRAGMA_DISABLE_DEPRECATION_WARNINGS
 	TArray<FBoneMirrorExport>				ExistingMirrorTable;
+	PRAGMA_ENABLE_DEPRECATION_WARNINGS
 
 	TWeakObjectPtr<UAssetImportData>		ExistingAssetImportData;
 	TWeakObjectPtr<UThumbnailInfo>			ExistingThumbnailInfo;
diff --git a/Engine/Source/Runtime/AnimGraphRuntime/Private/AnimNodes/AnimNode_Mirror.cpp b/Engine/Source/Runtime/AnimGraphRuntime/Private/AnimNodes/AnimNode_Mirror.cpp
new file mode 100644
--- /dev/null
+++ b/Engine/Source/Runtime/AnimGraphRuntime/Private/AnimNodes/AnimNode_Mirror.cpp
@@ -0,0 +1,133 @@
+// Copyright Epic Games, Inc. All Rights Reserved.
+
+#include "AnimNodes/AnimNode_Mirror.h"
+#include "Animation/AnimInstanceProxy.h"
+#include "Animation/AnimNode_Inertialization.h"
+#include "AnimationRuntime.h"
+
+#define LOCTEXT_NAMESPACE "AnimNode_Mirror"
+
+FAnimNode_Mirror::FAnimNode_Mirror()
+	: MirrorDataTable(nullptr)
+	, BlendTimeOnMirrorStateChange(0.0f)
+	, bMirror(true)
+	, bResetChildOnMirrorStateChange(false)
+	, bMirrorState(false)
+	, bMirrorStateIsValid(false)
+{
+}
+
+#if WITH_EDITOR
+UMirrorDataTable* FAnimNode_Mirror::GetMirrorDataTable() const 
+{ 
+	return MirrorDataTable.Get();
+}
+
+void FAnimNode_Mirror::SetMirrorDataTable(UMirrorDataTable* MirrorTable)
+{
+	MirrorDataTable = MirrorTable; 
+}
+
+#endif 
+
+void FAnimNode_Mirror::Initialize_AnyThread(const FAnimationInitializeContext& Context)
+{
+	FAnimNode_Base::Initialize_AnyThread(Context);
+	Source.Initialize(Context);
+}
+
+void FAnimNode_Mirror::CacheBones_AnyThread(const FAnimationCacheBonesContext& Context)
+{
+	Source.CacheBones(Context);
+	DECLARE_SCOPE_HIERARCHICAL_COUNTER_ANIMNODE(CacheBones_AnyThread)
+	Super::CacheBones_AnyThread(Context);
+
+	if (MirrorDataTable && Context.AnimInstanceProxy)
+	{
+		const FBoneContainer& BoneContainer = Context.AnimInstanceProxy->GetRequiredBones();
+		bool bSharesSkeleton = BoneContainer.GetSkeletonAsset() == MirrorDataTable->Skeleton;
+		bool bIdenticalJointCount = BoneContainer.GetReferenceSkeleton().GetNum() == MirrorDataTable->BoneToMirrorBoneIndex.Num(); 
+
+		if (bSharesSkeleton && bIdenticalJointCount)
+		{
+			MirrorDataTable->FillCompactPoseMirrorBones(BoneContainer, MirrorDataTable->BoneToMirrorBoneIndex, CompactPoseMirrorBones);
+		}
+		else
+		{
+			TArray<int32> MirrorBoneIndexes;
+			MirrorDataTable->FillMirrorBoneIndexes(BoneContainer.GetReferenceSkeleton(), MirrorBoneIndexes);
+			MirrorDataTable->FillCompactPoseMirrorBones(BoneContainer, MirrorBoneIndexes, CompactPoseMirrorBones);
+		}
+
+		const int32 NumBones = BoneContainer.GetCompactPoseNumBones();
+		ComponentSpaceRefRotations.SetNumUninitialized(NumBones);
+		ComponentSpaceRefRotations[FCompactPoseBoneIndex(0)] = BoneContainer.GetRefPoseTransform(FCompactPoseBoneIndex(0)).GetRotation();
+		for (FCompactPoseBoneIndex BoneIndex(1); BoneIndex < NumBones; ++BoneIndex)
+		{
+			const FCompactPoseBoneIndex ParentBoneIndex = BoneContainer.GetParentBoneIndex(BoneIndex);
+			ComponentSpaceRefRotations[BoneIndex] = ComponentSpaceRefRotations[ParentBoneIndex] * BoneContainer.GetRefPoseTransform(BoneIndex).GetRotation();
+		}
+	}
+	else
+	{
+		CompactPoseMirrorBones.Reset();
+		ComponentSpaceRefRotations.Reset(); 
+	}
+}
+
+void FAnimNode_Mirror::Update_AnyThread(const FAnimationUpdateContext& Context)
+{
+	GetEvaluateGraphExposedInputs().Execute(Context);
+
+	if (bMirrorStateIsValid && bMirrorState != bMirror)
+	{
+		if (BlendTimeOnMirrorStateChange > SMALL_NUMBER)
+		{
+			// Inertialize when switching between mirrored and unmirrored states to smooth out the pose discontinuity
+			UE::Anim::IInertializationRequester* InertializationRequester = Context.GetMessage<UE::Anim::IInertializationRequester>();
+			if (InertializationRequester)
+			{
+				InertializationRequester->RequestInertialization(BlendTimeOnMirrorStateChange);
+				InertializationRequester->AddDebugRecord(*Context.AnimInstanceProxy, Context.GetCurrentNodeId());
+			}
+			else
+			{
+				FAnimNode_Inertialization::LogRequestError(Context, Source);
+			}
+		}
+
+		// Optionally reinitialize the source when the mirror state changes
+		if (bResetChildOnMirrorStateChange)
+		{
+			FAnimationInitializeContext ReinitializeContext(Context.AnimInstanceProxy);
+			Source.Initialize(ReinitializeContext);
+		}
+	}
+
+	bMirrorState = bMirror;
+	bMirrorStateIsValid = true;
+
+	Source.Update(Context);
+}
+
+void FAnimNode_Mirror::Evaluate_AnyThread(FPoseContext& Output)
+{
+	Source.Evaluate(Output);
+
+	if (bMirrorState && MirrorDataTable)
+	{
+		FAnimationRuntime::MirrorPose(Output.Pose, MirrorDataTable->MirrorAxis, CompactPoseMirrorBones, ComponentSpaceRefRotations);
+	}
+}
+
+void FAnimNode_Mirror::GatherDebugData(FNodeDebugData& DebugData)
+{
+	FString DebugLine = DebugData.GetNodeName(this);
+	DebugLine += FString::Printf(TEXT("(Mirrored: %s)"), (bMirrorState) ? TEXT("true") : TEXT("false"));
+	DebugData.AddDebugItem(DebugLine);
+
+	Source.GatherDebugData(DebugData);
+}
+
+
+#undef LOCTEXT_NAMESPACE
diff --git a/Engine/Source/Runtime/AnimGraphRuntime/Public/AnimNodes/AnimNode_Mirror.h b/Engine/Source/Runtime/AnimGraphRuntime/Public/AnimNodes/AnimNode_Mirror.h
new file mode 100644
--- /dev/null
+++ b/Engine/Source/Runtime/AnimGraphRuntime/Public/AnimNodes/AnimNode_Mirror.h
@@ -0,0 +1,62 @@
+// Copyright Epic Games, Inc. All Rights Reserved.
+
+#pragma once
+
+#include "CoreMinimal.h"
+#include "UObject/ObjectMacros.h"
+#include "Animation/AnimNodeBase.h"
+#include "Animation/MirrorDataTable.h"
+#include "AnimNode_Mirror.generated.h"
+
+
+USTRUCT(BlueprintInternalUseOnly)
+struct ANIMGRAPHRUNTIME_API FAnimNode_Mirror : public FAnimNode_Base
+{
+	GENERATED_BODY()
+
+	friend class UAnimGraphNode_Mirror;
+
+public:
+	FAnimNode_Mirror();
+
+	virtual void Initialize_AnyThread(const FAnimationInitializeContext& Context) override;
+	virtual void CacheBones_AnyThread(const FAnimationCacheBonesContext& Context) override;
+	virtual void Update_AnyThread(const FAnimationUpdateContext& Context) override;
+	virtual void Evaluate_AnyThread(FPoseContext& Output) override;
+	virtual void GatherDebugData(FNodeDebugData& DebugData) override;
+
+	float GetBlendTimeOnMirrorStateChange() const { return BlendTimeOnMirrorStateChange; }
+
+#if WITH_EDITOR
+	UMirrorDataTable* GetMirrorDataTable() const;
+	void SetMirrorDataTable(UMirrorDataTable* MirrorTable);
+#endif 
+
+private:
+
+	UPROPERTY(EditAnywhere, Category = Settings)
+	TObjectPtr<UMirrorDataTable> MirrorDataTable;
+
+	UPROPERTY(EditAnywhere, Category = Settings, meta = (PinHiddenByDefault))
+	float BlendTimeOnMirrorStateChange;
+
+	UPROPERTY(EditAnywhere, Category = Links)
+	FPoseLink Source;
+
+	UPROPERTY(EditAnywhere,  Category = Settings, meta = (PinShownByDefault))
+	bool bMirror;
+
+	// Whether to reset (reinitialize) the child (source) pose when the mirror state changes
+	UPROPERTY(EditAnywhere, Category = Settings)
+	bool bResetChildOnMirrorStateChange;
+
+private:
+	bool bMirrorState;
+	bool bMirrorStateIsValid;
+
+	// Compact pose format of Mirror Bone Map
+	TArray<FCompactPoseBoneIndex> CompactPoseMirrorBones;
+
+	// Pre-calculated component space of reference pose, which allows mirror to work with any joint orient 
+	TCustomBoneIndexArray<FQuat, FCompactPoseBoneIndex> ComponentSpaceRefRotations;
+};
diff --git a/Engine/Source/Runtime/Engine/Classes/Animation/AnimationSettings.h b/Engine/Source/Runtime/Engine/Classes/Animation/AnimationSettings.h
--- a/Engine/Source/Runtime/Engine/Classes/Animation/AnimationSettings.h
+++ b/Engine/Source/Runtime/Engine/Classes/Animation/AnimationSettings.h
@@ -12,7 +12,7 @@
 #include "Animation/AnimSequence.h"
 #include "Engine/DeveloperSettings.h"
 #include "CustomAttributes.h"
-
+#include "MirrorDataTable.h"
 #include "AnimationSettings.generated.h"
 
 /**
@@ -77,6 +77,10 @@ class ENGINE_API UAnimationSettings : public UDeveloperSettings
 	UPROPERTY(config, EditAnywhere, Category = CustomAttributes)
 	ECustomAttributeBlendType DefaultAttributeBlendMode;
 
+	/** Find and Replace Expressions used for mirroring  */
+	UPROPERTY(config, EditAnywhere, Category = Mirroring)
+	TArray<FMirrorFindReplaceExpression> MirrorFindReplaceExpressions;
+
 public:
 	static UAnimationSettings * Get() { return CastChecked<UAnimationSettings>(UAnimationSettings::StaticClass()->GetDefaultObject()); }
 
diff --git a/Engine/Source/Runtime/Engine/Classes/Animation/MirrorDataTable.h b/Engine/Source/Runtime/Engine/Classes/Animation/MirrorDataTable.h
new file mode 100644
--- /dev/null
+++ b/Engine/Source/Runtime/Engine/Classes/Animation/MirrorDataTable.h
@@ -0,0 +1,200 @@
+// Copyright Epic Games, Inc. All Rights Reserved.
+
+#pragma once
+
+#include "Engine/DataTable.h"
+#include "MirrorDataTable.generated.h"
+
+/** Type referenced by a row in the mirror data table */
+UENUM()
+enum EMirrorRowType
+{
+	Bone,
+	Notify,
+	Curve
+};
+
+/** Find and Replace Method for FMirrorFindReplaceExpression. */
+UENUM()
+enum EMirrorFindReplaceMethod
+{
+	/** Only find and replace matching strings at the start of the name  */
+	Prefix,
+	/** Only find and replace matching strings at the end of the name  */
+	Suffix,
+	/** Use regular expressions for find and replace, including support for captures $1 - $10 */
+	RegularExpression
+};
+
+/**  Base Mirror Table containing all data required by the animation mirroring system. */
+USTRUCT()
+struct FMirrorTableRow : public FTableRowBase
+{
+	GENERATED_BODY()
+
+	UPROPERTY(EditAnywhere, BlueprintReadOnly, Category = Mirroring)
+	FName Name;
+
+	UPROPERTY(EditAnywhere, BlueprintReadOnly, Category = Mirroring)
+	FName MirroredName;
+
+	UPROPERTY(EditAnywhere, BlueprintReadOnly, Category = Mirroring)
+	TEnumAsByte<EMirrorRowType> MirrorEntryType;
+
+	FMirrorTableRow()
+		: Name(NAME_None)
+		, MirroredName(NAME_None)
+		, MirrorEntryType(EMirrorRowType::Bone) {}
+
+	ENGINE_API FMirrorTableRow(const FMirrorTableRow& Other);
+	ENGINE_API FMirrorTableRow& operator=(FMirrorTableRow const& Other);
+	ENGINE_API bool operator==(FMirrorTableRow const& Other) const;
+	ENGINE_API bool operator!=(FMirrorTableRow const& Other) const;
+	ENGINE_API bool operator<(FMirrorTableRow const& Other) const;
+};
+
+/** Find and Replace expressions used to generate mirror tables*/
+USTRUCT()
+struct FMirrorFindReplaceExpression
+{
+	GENERATED_BODY()
+
+	UPROPERTY(EditAnywhere, Category = Mirroring)
+	FName FindExpression;
+
+	UPROPERTY(EditAnywhere, Category = Mirroring)
+	FName ReplaceExpression;
+
+	UPROPERTY(EditAnywhere, Category = Mirroring)
+	TEnumAsByte<EMirrorFindReplaceMethod> FindReplaceMethod;
+
+	FMirrorFindReplaceExpression() 
+		: FindExpression(NAME_None)
+		, ReplaceExpression(NAME_None)
+		, FindReplaceMethod(EMirrorFindReplaceMethod::Prefix) {}
+
+	FMirrorFindReplaceExpression(FName InFindExpression, FName InReplaceExpression, EMirrorFindReplaceMethod Method)
+		: FindExpression(InFindExpression)
+		, ReplaceExpression(InReplaceExpression)
+		, FindReplaceMethod(Method)
+	{
+	}
+};
+
+/**
+ * Data table for mirroring bones, notifies, and curves.   The mirroring table allows self mirroring with entries where the name and mirrored name are identical
+ */
+UCLASS(MinimalAPI, BlueprintType, hideCategories = (ImportOptions, ImportSource) /* AutoExpandCategories = "MirrorDataTable,ImportOptions"*/)
+class UMirrorDataTable : public UDataTable
+{
+	GENERATED_BODY()
+
+	friend class UMirrorDataTableFactory;
+
+public:
+	UMirrorDataTable(const FObjectInitializer& ObjectInitializer);
+
+	ENGINE_API virtual void GetPreloadDependencies(TArray<UObject*>& OutDeps) override;
+
+	ENGINE_API virtual void Serialize(FArchive& Ar) override;
+
+	ENGINE_API virtual void PostLoad() override;
+
+	ENGINE_API virtual void EmptyTable() override;
+
+#if WITH_EDITOR
+
+	ENGINE_API virtual void CleanBeforeStructChange() override;
+
+	ENGINE_API virtual void RestoreAfterStructChange() override;
+
+	ENGINE_API virtual void PostEditChangeProperty(FPropertyChangedEvent& PropertyChangedEvent) override;
+
+	ENGINE_API virtual void PostEditUndo() override;
+
+#endif // WITH_EDITOR
+
+	/**
+	 * Apply the animation settings mirroring find and replace strings against the given name, returning
+	 * the mirrored name or NAME_None if none of the find strings are found in the name. 
+	 * 
+	 * @param	InName		Name to map against animation settings mirroring find and replace 
+	 * @return				The mirrored name or NAME_None
+	 */
+	ENGINE_API static FName GetSettingsMirrorName(FName InName); 
+
+	/**
+	 * Apply the provided find and replace strings against the given name, returning
+	 * the mirrored name or NAME_None if none of the find strings are found in the name. 
+	 * 
+	 * @param	MirrorFindReplaceExpressions		Find and replace expressions.  The first matching expression will be returned
+	 * @param	InName								Name to find and replace 
+	 * @return										The mirrored name or NAME_None if none of the expressions match
+	 */
+	ENGINE_API static FName GetMirrorName(FName InName, const TArray<FMirrorFindReplaceExpression>& MirrorFindReplaceExpressions);
+
+	/**
+     * Create Mirror Bone Indices for the provided BoneContainer.  The CompactBonePoseMirrorBones provides an index map which can be used to mirror at runtime
+	 *
+	 * @param	BoneContainer					The Bone Container that the OutCompactPaseMirrorBones should match
+	 * @param	MirrorBoneIndexes				Mirror bone indexes created for the ReferenceSkeleton used by the BoneContainer 
+	 * @param	OutCompactPoseMirrorBones		An efficient representation of the bones to mirror which can be used at runtime
+	 */
+	ENGINE_API static void FillCompactPoseMirrorBones(const FBoneContainer& BoneContainer, const TArray<int32>& MirrorBoneIndexes, TArray<FCompactPoseBoneIndex>& OutCompactPoseMirrorBones);
+
+
+	/**
+	 * Converts the mirror data table Name -> MirrorName map into an index map for the given ReferenceSkeleton
+	 *
+	 * @param	ReferenceSkeleton		The ReferenceSkeleton to compute the mirror index against
+	 * @param	OutMirrorBoneIndexes	An array that provides the bone index of the mirror bone, or INDEX_NONE if the bone is not mirrored
+	 */
+	ENGINE_API void FillMirrorBoneIndexes(const FReferenceSkeleton& ReferenceSkeleton, TArray<int32>& OutMirrorBoneIndexes) const;
+
+#if WITH_EDITOR  
+	/**
+	 * Populates the table by running the MirrorFindReplaceExpressions on bone names in the Skeleton.  If the mirrored name is also found 
+	 * on the Skeleton it is added to the table.
+	 */
+	ENGINE_API void FindReplaceMirroredNames();
+#endif // WITH_EDITOR
+
+	/**
+	 * Evaluate the MirrorFindReplaceExpressions on InName and return the replaced value of the first entry that matches
+	 *
+	 * @param	InName		The input string to find & replace
+	 * @return				The replaced result of the first MirrorFindReplaceExpression where the find pattern matched
+	 */
+	ENGINE_API FName FindReplace(FName InName) const;
+
+public:
+
+	UPROPERTY(EditAnywhere, Category = CreateTable)
+	TArray<FMirrorFindReplaceExpression> MirrorFindReplaceExpressions;
+
+	UPROPERTY(EditAnywhere, Category = Mirroring)
+	TEnumAsByte<EAxis::Type> MirrorAxis;
+
+	UPROPERTY(EditAnywhere, AssetRegistrySearchable, Category = Skeleton)
+	TObjectPtr<USkeleton> Skeleton; 
+
+	// Index of the mirror bone for a given bone index in the reference skeleton, or INDEX_NONE if the bone is not mirrored
+	TArray<int32> BoneToMirrorBoneIndex;
+
+	// Array with entries the source UIDs of curves that should be mirrored. 
+	TArray<SmartName::UID_Type> CurveMirrorSourceUIDArray;
+
+	// Array with the target UIDs of curves that should be mirrored. 
+	TArray<SmartName::UID_Type> CurveMirrorTargetUIDArray;
+
+	// Map from notify to mirror notify
+	TMap<FName, FName> NotifyToMirrorNotifyMap;
+
+protected: 
+
+	// Fill BoneToMirrorBoneIndex, CurveMirrorSourceUIDArray, CurveMirrorTargetUIDArray and NotifyToMirrorNotifyIndex based on the Skeleton and Table Contents
+	void FillMirrorArrays(); 
+
+	void HandleDataTableChanged();
+};
+
diff --git a/Engine/Source/Runtime/Engine/Classes/Animation/Skeleton.h b/Engine/Source/Runtime/Engine/Classes/Animation/Skeleton.h
--- a/Engine/Source/Runtime/Engine/Classes/Animation/Skeleton.h
+++ b/Engine/Source/Runtime/Engine/Classes/Animation/Skeleton.h
@@ -442,7 +442,6 @@ private:
  *		- Bone hierarchy for animations
  *		- Bone/track linkup between mesh and animation
  *		- Retargetting related
- *		- Mirror table
  */
 UCLASS(hidecategories=Object, MinimalAPI)
 class USkeleton : public UObject, public IInterface_AssetUserData, public IInterface_PreviewMeshProvider
diff --git a/Engine/Source/Runtime/Engine/Classes/Engine/SkeletalMesh.h b/Engine/Source/Runtime/Engine/Classes/Engine/SkeletalMesh.h
--- a/Engine/Source/Runtime/Engine/Classes/Engine/SkeletalMesh.h
+++ b/Engine/Source/Runtime/Engine/Classes/Engine/SkeletalMesh.h
@@ -243,6 +243,7 @@ enum class ESkinCacheDefaultBehavior : uint8
 	Inclusive = 1,
 };
 
+struct UE_DEPRECATED(5.0, "FBoneMirrorInfo is deprecated. Please use UMirrorDataTable for mirroring support.") FBoneMirrorInfo;
 USTRUCT()
 struct FBoneMirrorInfo
 {
@@ -264,7 +265,7 @@ struct FBoneMirrorInfo
 
 };
 
-/** Structure to export/import bone mirroring information */
+struct UE_DEPRECATED(5.0, "FBoneMirrorExport is deprecated. Please use UMirrorDataTable for mirroring support.") FBoneMirrorExport;
 USTRUCT()
 struct FBoneMirrorExport
 {
@@ -1029,10 +1030,11 @@ public:
 	}
 
 	/** List of bones that should be mirrored. */
-	UE_DEPRECATED(4.27, "Please do not access this member directly; use USkeletalMesh::GetSkelMirrorTable() or USkeletalMesh::SetSkelMirrorTable().")
+	UE_DEPRECATED(4.27, "Please do not access this member directly; Use UMirrorDataTable for mirroring support")
 	UPROPERTY(EditAnywhere, editfixedsize, Category=Mirroring)
 	TArray<struct FBoneMirrorInfo> SkelMirrorTable;
 
+    UE_DEPRECATED(5.0, "Please use UMirrorDataTable for mirroring support")
 	static FName GetSkelMirrorTableMemberName()
 	{
 		PRAGMA_DISABLE_DEPRECATION_WARNINGS
@@ -1040,6 +1042,7 @@ public:
 		PRAGMA_ENABLE_DEPRECATION_WARNINGS
 	}
 
+    UE_DEPRECATED(5.0, "Please use UMirrorDataTable for mirroring support")
 	TArray<struct FBoneMirrorInfo>& GetSkelMirrorTable()
 	{
 		WaitUntilAsyncPropertyReleased(ESkeletalMeshAsyncProperties::SkelMirrorTable);
@@ -1048,6 +1051,7 @@ public:
 		PRAGMA_ENABLE_DEPRECATION_WARNINGS
 	}
 
+    UE_DEPRECATED(5.0, "Please use UMirrorDataTable for mirroring support")
 	const TArray<struct FBoneMirrorInfo>& GetSkelMirrorTable() const
 	{
 		WaitUntilAsyncPropertyReleased(ESkeletalMeshAsyncProperties::SkelMirrorTable, EAsyncPropertyLockType::ReadOnly);
@@ -1056,6 +1060,7 @@ public:
 		PRAGMA_ENABLE_DEPRECATION_WARNINGS
 	}
 
+    UE_DEPRECATED(5.0, "Please use UMirrorDataTable for mirroring support")
 	void SetSkelMirrorTable(const TArray<struct FBoneMirrorInfo>& InSkelMirrorTable)
 	{
 		WaitUntilAsyncPropertyReleased(ESkeletalMeshAsyncProperties::SkelMirrorTable);
@@ -1319,7 +1324,6 @@ public:
 	int32 GetMaxNumOptionalLODs(const class ITargetPlatform* TargetPlatform) const;
 #endif
 
-
 	UFUNCTION(BlueprintSetter)
 	void SetDefaultAnimatingRig(TSoftObjectPtr<UObject> InAnimatingRig);
 
@@ -1336,6 +1340,7 @@ public:
 		PRAGMA_ENABLE_DEPRECATION_WARNINGS
 	}
 
+	UE_DEPRECATED(5.0, "Please use UMirrorDataTable for mirroring support")
 	TEnumAsByte<EAxis::Type> GetSkelMirrorAxis() const
 	{
 		WaitUntilAsyncPropertyReleased(ESkeletalMeshAsyncProperties::SkelMirrorAxis, EAsyncPropertyLockType::ReadOnly);
@@ -1344,6 +1349,7 @@ public:
 		PRAGMA_ENABLE_DEPRECATION_WARNINGS
 	}
 
+	UE_DEPRECATED(5.0, "Please use UMirrorDataTable for mirroring support")
 	void SetSkelMirrorAxis(TEnumAsByte<EAxis::Type> InSkelMirrorAxis)
 	{
 		WaitUntilAsyncPropertyReleased(ESkeletalMeshAsyncProperties::SkelMirrorAxis);
@@ -1352,7 +1358,7 @@ public:
 		PRAGMA_ENABLE_DEPRECATION_WARNINGS
 	}
 
-	UE_DEPRECATED(4.27, "Please do not access this member directly; use USkeletalMesh::GetSkelMirrorFlipAxis() or USkeletalMesh::SetSkelMirrorFlipAxis().")
+	UE_DEPRECATED(4.27, "Please do not access this member directly; Use UMirrorDataTable for mirroring support")
 	UPROPERTY(EditAnywhere, Category=Mirroring)
 	TEnumAsByte<EAxis::Type> SkelMirrorFlipAxis;
 	static FName GetSkelMirrorFlipAxisMemberName()
@@ -1362,6 +1368,7 @@ public:
 		PRAGMA_ENABLE_DEPRECATION_WARNINGS
 	}
 
+    UE_DEPRECATED(5.0, "Please use UMirrorDataTable for mirroring support")
 	TEnumAsByte<EAxis::Type> GetSkelMirrorFlipAxis() const
 	{
 		WaitUntilAsyncPropertyReleased(ESkeletalMeshAsyncProperties::SkelMirrorFlipAxis, EAsyncPropertyLockType::ReadOnly);
@@ -1370,6 +1377,7 @@ public:
 		PRAGMA_ENABLE_DEPRECATION_WARNINGS
 	}
 
+    UE_DEPRECATED(5.0, "Please use UMirrorDataTable for mirroring support")
 	void SetSkelMirrorFlipAxis(TEnumAsByte<EAxis::Type> InSkelMirrorFlipAxis)
 	{
 		WaitUntilAsyncPropertyReleased(ESkeletalMeshAsyncProperties::SkelMirrorFlipAxis);
@@ -2603,20 +2611,23 @@ public:
 	 */
 	FMatrix GetComposedRefPoseMatrix( FName InBoneName ) const;
 	FMatrix GetComposedRefPoseMatrix( int32 InBoneIndex ) const;
-
-	/** Allocate and initialise bone mirroring table for this skeletal mesh. Default is source = destination for each bone. */
+	UE_DEPRECATED(5.0, "Please use UMirrorDataTable for mirroring support.")
 	void InitBoneMirrorInfo();
 
-	/** Utility for copying and converting a mirroring table from another USkeletalMesh. */
+	UE_DEPRECATED(5.0, "Please use UMirrorDataTable for mirroring support.")
 	void CopyMirrorTableFrom(USkeletalMesh* SrcMesh);
+	
+	UE_DEPRECATED(5.0, "Please use UMirrorDataTable for mirroring support.")
+	PRAGMA_DISABLE_DEPRECATION_WARNINGS
 	void ExportMirrorTable(TArray<FBoneMirrorExport> &MirrorExportInfo) const;
+	PRAGMA_ENABLE_DEPRECATION_WARNINGS
+
+	UE_DEPRECATED(5.0, "Please use UMirrorDataTable for mirroring support.")
+	PRAGMA_DISABLE_DEPRECATION_WARNINGS
 	void ImportMirrorTable(const TArray<FBoneMirrorExport> &MirrorExportInfo);
+	PRAGMA_ENABLE_DEPRECATION_WARNINGS
 
-	/** 
-	 *	Utility for checking that the bone mirroring table of this mesh is good.
-	 *	Return true if mirror table is OK, false if there are problems.
-	 *	@param	ProblemBones	Output string containing information on bones that are currently bad.
-	 */
+	UE_DEPRECATED(5.0, "Please use UMirrorDataTable for mirroring support.")
 	bool MirrorTableIsGood(FString& ProblemBones) const;
 
 	/**
diff --git a/Engine/Source/Runtime/Engine/Private/Animation/AnimationRuntime.cpp b/Engine/Source/Runtime/Engine/Private/Animation/AnimationRuntime.cpp
--- a/Engine/Source/Runtime/Engine/Private/Animation/AnimationRuntime.cpp
+++ b/Engine/Source/Runtime/Engine/Private/Animation/AnimationRuntime.cpp
@@ -14,6 +14,7 @@
 #include "GenericPlatform/GenericPlatformCompilerPreSetup.h"
 #include "Animation/AnimationPoseData.h"
 #include "Animation/BlendProfile.h"
+#include "Animation/MirrorDataTable.h"
 #if INTEL_ISPC
 #include "AnimationRuntime.ispc.generated.h"
 #endif
@@ -1112,6 +1113,169 @@ void FAnimationRuntime::AccumulateMeshSpaceRotationAdditiveToLocalPoseInternal(F
 	}
 }
 
+FVector FAnimationRuntime::MirrorVector(const FVector& V, EAxis::Type MirrorAxis)
+{
+	FVector MirrorV(V);
+
+	switch (MirrorAxis)
+	{
+	case EAxis::X:
+		MirrorV.X = -MirrorV.X;
+		break;
+	case EAxis::Y:
+		MirrorV.Y = -MirrorV.Y;
+		break;
+	case EAxis::Z:
+		MirrorV.Z = -MirrorV.Z;
+		break;
+	}
+
+	return MirrorV;
+}
+
+FQuat FAnimationRuntime::MirrorQuat(const FQuat& Q, EAxis::Type MirrorAxis)
+{
+	FQuat MirrorQ(Q);
+
+	// Given an axis V and an angle A, the corresponding unmirrored quaternion Q = { Q.XYZ, Q.W } is:
+	//
+	//		Q = { V * sin(A/2), cos(A/2) }
+	//
+	//  mirror both the axis of rotation and the angle of rotation around that axis.
+	// Therefore, the mirrored quaternion Q' for the axis V and angle A is:
+	//
+	//		Q' = { MirrorVector(V) * sin(-A/2), cos(-A/2) }
+	//		Q' = { -MirrorVector(V) * sin(A/2), cos(A/2) }
+	//		Q' = { -MirrorVector(V * sin(A/2)), cos(A/2) }
+	//		Q' = { -MirrorVector(Q.XYZ), Q.W }
+	//
+	switch (MirrorAxis)
+	{
+	case EAxis::X:
+		MirrorQ.Y = -MirrorQ.Y;
+		MirrorQ.Z = -MirrorQ.Z;
+		break;
+	case EAxis::Y:
+		MirrorQ.X = -MirrorQ.X;
+		MirrorQ.Z = -MirrorQ.Z;
+		break;
+	case EAxis::Z:
+		MirrorQ.X = -MirrorQ.X;
+		MirrorQ.Y = -MirrorQ.Y;
+		break;
+	}
+
+	return MirrorQ;
+}
+
+
+void FAnimationRuntime::MirrorPose(FCompactPose& Pose, EAxis::Type MirrorAxis, const TArray<FCompactPoseBoneIndex>& CompactPoseMirrorBones, const TCustomBoneIndexArray<FQuat, FCompactPoseBoneIndex>& ComponentSpaceRefRotations)
+{
+	const FBoneContainer& BoneContainer = Pose.GetBoneContainer();
+	if (MirrorAxis == EAxis::None)
+	{
+		return;
+	}
+
+	auto MirrorTransform = [&ComponentSpaceRefRotations, MirrorAxis](const FTransform& SourceTransform, const FCompactPoseBoneIndex& SourceParentIndex, const FCompactPoseBoneIndex& SourceBoneIndex, const FCompactPoseBoneIndex& TargetParentIndex, const FCompactPoseBoneIndex& TargetBoneIndex) -> FTransform {
+
+		const FQuat TargetParentRefRotation = ComponentSpaceRefRotations[TargetParentIndex];
+		const FQuat TargetBoneRefRotation = ComponentSpaceRefRotations[TargetBoneIndex];
+		const FQuat SourceParentRefRotation = ComponentSpaceRefRotations[SourceParentIndex];
+		const FQuat SourceBoneRefRotation = ComponentSpaceRefRotations[SourceBoneIndex];
+
+		// Mirror the translation component:  Rotate the translation into the space of the mirror plane,  mirror across the mirror plane, and rotate into the space of its new parent
+
+		FVector T = SourceTransform.GetTranslation();
+		T = SourceParentRefRotation.RotateVector(T);
+		T = MirrorVector(T, MirrorAxis);
+		T = TargetParentRefRotation.UnrotateVector(T);
+
+		// Mirror the rotation component:- Rotate into the space of the mirror plane, mirror across the plane, apply corrective rotation to align result with target space's rest orientation, 
+		// then rotate into the space of its new parent
+
+		FQuat Q = SourceTransform.GetRotation();
+		Q = SourceParentRefRotation * Q;
+		Q = MirrorQuat(Q, MirrorAxis);
+		Q *= MirrorQuat(SourceBoneRefRotation, MirrorAxis).Inverse() * TargetBoneRefRotation;
+		Q = TargetParentRefRotation.Inverse() * Q;
+
+		FVector S = SourceTransform.GetScale3D();
+
+		return FTransform(Q, T, S);
+	};
+
+	// Mirror the root bone
+	{
+		const FCompactPoseBoneIndex RootBoneIndex(0);
+		const FCompactPoseBoneIndex MirrorRootBoneIndex = CompactPoseMirrorBones[RootBoneIndex.GetInt()];
+		if (MirrorRootBoneIndex.IsValid())
+		{
+			const FQuat RootBoneRefRotation = ComponentSpaceRefRotations[RootBoneIndex];
+
+			FVector T = Pose[RootBoneIndex].GetTranslation();
+			T = MirrorVector(T, MirrorAxis);
+
+			FQuat Q = Pose[RootBoneIndex].GetRotation();
+			Q = MirrorQuat(Q, MirrorAxis);
+			Q *= MirrorQuat(RootBoneRefRotation, MirrorAxis).Inverse() * RootBoneRefRotation;
+
+			FVector S = Pose[RootBoneIndex].GetScale3D();
+
+			Pose[RootBoneIndex] = FTransform(Q, T, S);
+		}
+	}
+
+	const int32 NumBones = BoneContainer.GetCompactPoseNumBones();
+
+	// Mirror the non-root bones
+	for (FCompactPoseBoneIndex TargetBoneIndex(1); TargetBoneIndex < NumBones; ++TargetBoneIndex)
+	{
+		const FCompactPoseBoneIndex SourceBoneIndex = CompactPoseMirrorBones[TargetBoneIndex.GetInt()];
+		if (SourceBoneIndex == TargetBoneIndex)
+		{
+			const FCompactPoseBoneIndex TargetParentIndex = BoneContainer.GetParentBoneIndex(TargetBoneIndex);
+			Pose[TargetBoneIndex] = MirrorTransform(Pose[TargetBoneIndex], TargetParentIndex, TargetBoneIndex, TargetParentIndex, TargetBoneIndex);
+		}
+		else if (SourceBoneIndex > TargetBoneIndex)
+		{
+			const FCompactPoseBoneIndex TargetParentIndex = BoneContainer.GetParentBoneIndex(TargetBoneIndex);
+			const FCompactPoseBoneIndex SourceParentIndex = BoneContainer.GetParentBoneIndex(SourceBoneIndex);
+			const FTransform NewTargetBoneTransform = MirrorTransform(Pose[SourceBoneIndex], SourceParentIndex, SourceBoneIndex, TargetParentIndex, TargetBoneIndex);
+			const FTransform NewSourceBoneTransform = MirrorTransform(Pose[TargetBoneIndex], TargetParentIndex, TargetBoneIndex, SourceParentIndex, SourceBoneIndex);
+			Pose[TargetBoneIndex] = NewTargetBoneTransform;
+			Pose[SourceBoneIndex] = NewSourceBoneTransform;
+		}
+	}
+}
+
+void FAnimationRuntime::MirrorPose(FCompactPose& Pose, const UMirrorDataTable& MirrorDataTable)
+{
+	const FBoneContainer& BoneContainer = Pose.GetBoneContainer();
+	USkeleton* Skeleton = BoneContainer.GetSkeletonAsset();
+	if (Skeleton)
+	{
+		TArray<int32> MirrorBoneIndexes;
+		MirrorDataTable.FillMirrorBoneIndexes(BoneContainer.GetReferenceSkeleton(), MirrorBoneIndexes);
+
+		// Compact pose format of Mirror Bone Map
+		TArray<FCompactPoseBoneIndex> CompactPoseMirrorBones;
+		MirrorDataTable.FillCompactPoseMirrorBones(BoneContainer, MirrorBoneIndexes, CompactPoseMirrorBones);
+
+		const int32 NumBones = BoneContainer.GetCompactPoseNumBones();
+
+		TCustomBoneIndexArray<FQuat, FCompactPoseBoneIndex> ComponentSpaceRefRotations;
+		ComponentSpaceRefRotations.SetNumUninitialized(NumBones);
+		ComponentSpaceRefRotations[FCompactPoseBoneIndex(0)] = BoneContainer.GetRefPoseTransform(FCompactPoseBoneIndex(0)).GetRotation();
+		for (FCompactPoseBoneIndex BoneIndex(1); BoneIndex < NumBones; ++BoneIndex)
+		{
+			const FCompactPoseBoneIndex ParentBoneIndex = BoneContainer.GetParentBoneIndex(BoneIndex);
+			ComponentSpaceRefRotations[BoneIndex] = ComponentSpaceRefRotations[ParentBoneIndex] * BoneContainer.GetRefPoseTransform(BoneIndex).GetRotation();
+		}
+		MirrorPose(Pose, MirrorDataTable.MirrorAxis, CompactPoseMirrorBones, ComponentSpaceRefRotations);
+	}
+}
+
 /** 
  * return ETypeAdvanceAnim type
  */
diff --git a/Engine/Source/Runtime/Engine/Private/Animation/AnimationSettings.cpp b/Engine/Source/Runtime/Engine/Private/Animation/AnimationSettings.cpp
--- a/Engine/Source/Runtime/Engine/Private/Animation/AnimationSettings.cpp
+++ b/Engine/Source/Runtime/Engine/Private/Animation/AnimationSettings.cpp
@@ -22,6 +22,16 @@ UAnimationSettings::UAnimationSettings(const FObjectInitializer& ObjectInitializ
 	KeyEndEffectorsMatchNameArray.Add(TEXT("hand"));
 	KeyEndEffectorsMatchNameArray.Add(TEXT("attach"));
 	KeyEndEffectorsMatchNameArray.Add(TEXT("camera"));
+
+	MirrorFindReplaceExpressions = {
+		FMirrorFindReplaceExpression("r_", "l_", EMirrorFindReplaceMethod::Prefix), FMirrorFindReplaceExpression("l_", "r_", EMirrorFindReplaceMethod::Prefix),
+		FMirrorFindReplaceExpression("R_", "L_", EMirrorFindReplaceMethod::Prefix), FMirrorFindReplaceExpression("L_", "R_", EMirrorFindReplaceMethod::Prefix),
+		FMirrorFindReplaceExpression("_l", "_r", EMirrorFindReplaceMethod::Suffix), FMirrorFindReplaceExpression("_r", "_l",  EMirrorFindReplaceMethod::Suffix),
+		FMirrorFindReplaceExpression("_R", "_L", EMirrorFindReplaceMethod::Suffix), FMirrorFindReplaceExpression("_L", "_R", EMirrorFindReplaceMethod::Suffix),
+		FMirrorFindReplaceExpression("right", "left", EMirrorFindReplaceMethod::Prefix), FMirrorFindReplaceExpression("left", "right",  EMirrorFindReplaceMethod::Prefix),
+		FMirrorFindReplaceExpression("Right", "Left", EMirrorFindReplaceMethod::Prefix), FMirrorFindReplaceExpression("Left", "Right",  EMirrorFindReplaceMethod::Prefix),
+		FMirrorFindReplaceExpression("((?:^[sS]pine|^[rR]oot|^[pP]elvis|^[hH]ead).*)", "$1", EMirrorFindReplaceMethod::RegularExpression)
+	};
 }
 
 
diff --git a/Engine/Source/Runtime/Engine/Private/Animation/MirrorDataTable.cpp b/Engine/Source/Runtime/Engine/Private/Animation/MirrorDataTable.cpp
new file mode 100644
--- /dev/null
+++ b/Engine/Source/Runtime/Engine/Private/Animation/MirrorDataTable.cpp
@@ -0,0 +1,421 @@
+// Copyright Epic Games, Inc. All Rights Reserved.
+
+#include "Animation/MirrorDataTable.h"
+#include "Animation/AnimationSettings.h"
+#include "Internationalization/Regex.h"
+
+#define LOCTEXT_NAMESPACE "MirrorDataTables"
+
+
+FMirrorTableRow::FMirrorTableRow(const FMirrorTableRow& Other)
+{
+	*this = Other;
+}
+
+FMirrorTableRow& FMirrorTableRow::operator=(FMirrorTableRow const& Other)
+{
+	if (this == &Other)
+	{
+		return *this;
+	}
+
+	Name = Other.Name;
+	MirroredName = Other.MirroredName;
+	MirrorEntryType = Other.MirrorEntryType;
+	return *this;
+}
+
+bool FMirrorTableRow::operator==(FMirrorTableRow const& Other) const
+{
+	return (Name == Other.Name && MirroredName == Other.MirroredName && MirrorEntryType == Other.MirrorEntryType);
+}
+
+bool FMirrorTableRow::operator!=(FMirrorTableRow const& Other) const
+{
+	return (Name != Other.Name || MirroredName != Other.MirroredName || MirrorEntryType != Other.MirrorEntryType);
+}
+
+bool FMirrorTableRow::operator<(FMirrorTableRow const& Other) const
+{
+	if (MirrorEntryType == Other.MirrorEntryType)
+	{
+		if (Name != Other.Name)
+		{
+			return Name.LexicalLess(Other.Name);
+		}
+		else
+		{
+			return MirroredName.LexicalLess(Other.MirroredName);
+		}
+	}
+	return MirrorEntryType < Other.MirrorEntryType;
+}
+
+UMirrorDataTable::UMirrorDataTable(const FObjectInitializer& ObjectInitializer)
+	: Super(ObjectInitializer), MirrorAxis(EAxis::X)
+{
+#if WITH_EDITORONLY_DATA
+	OnDataTableChanged().AddUObject(this, &UMirrorDataTable::FillMirrorArrays);
+#endif 
+}
+
+void UMirrorDataTable::GetPreloadDependencies(TArray<UObject*>& OutDeps)
+{
+	Super::GetPreloadDependencies(OutDeps);
+}
+
+void UMirrorDataTable::Serialize(FArchive& Ar)
+{
+	Super::Serialize(Ar); 
+}
+
+void UMirrorDataTable::PostLoad()
+{
+	Super::PostLoad();
+	FillMirrorArrays(); 
+}
+
+void UMirrorDataTable::EmptyTable()
+{
+}
+
+#if WITH_EDITOR
+void UMirrorDataTable::CleanBeforeStructChange()
+{
+	Super::CleanBeforeStructChange();
+}
+
+void UMirrorDataTable::RestoreAfterStructChange()
+{
+	Super::RestoreAfterStructChange();
+}
+
+void UMirrorDataTable::PostEditChangeProperty(FPropertyChangedEvent& PropertyChangedEvent)
+{
+	Super::PostEditChangeProperty(PropertyChangedEvent);
+	FillMirrorArrays(); 
+}
+
+void UMirrorDataTable::PostEditUndo()
+{
+	Super::PostEditUndo();
+}
+
+#endif // WITH_EDITOR
+
+
+FName UMirrorDataTable::GetSettingsMirrorName(FName InName)
+{
+	UAnimationSettings* AnimationSettings = UAnimationSettings::Get();
+	FName MirrorName;
+	if (AnimationSettings)
+	{
+		MirrorName = GetMirrorName(InName, AnimationSettings->MirrorFindReplaceExpressions);
+	}
+	return MirrorName;
+}
+
+FName UMirrorDataTable::GetMirrorName(FName InName, const TArray<FMirrorFindReplaceExpression>& MirrorFindReplaceExpressions)
+{
+	FName ReplacedName;
+	FString InNameString = InName.ToString();
+	for (const FMirrorFindReplaceExpression& regExStr : MirrorFindReplaceExpressions)
+	{
+		FString FindString = regExStr.FindExpression.ToString();
+		FString ReplaceString = regExStr.ReplaceExpression.ToString();
+		if (regExStr.FindReplaceMethod == EMirrorFindReplaceMethod::Prefix)
+		{
+			// convert prefix expression to regex that matches start of string, prefix, and any number of characters
+			FindString = FindString + TEXT("([^}]*)");
+			FindString = TEXT("^") + FindString;
+			ReplaceString = ReplaceString + TEXT("$1");
+
+		}
+		else if (regExStr.FindReplaceMethod == EMirrorFindReplaceMethod::Suffix)
+		{
+			// convert suffix expression to regex that matches any number of characters start, the suffix, and end of string
+			FindString = TEXT("([^}]*)") + FindString + TEXT("$");
+			ReplaceString = TEXT("$1") + ReplaceString;
+		}
+
+		FRegexPattern MatherPatter(FindString);
+		FRegexMatcher Matcher(MatherPatter, InNameString);
+		bool bFound = false;
+		while (Matcher.FindNext())
+		{
+			for (int32 CaptureIndex = 1; CaptureIndex < 10; CaptureIndex++)
+			{
+				FString CaptureResult = Matcher.GetCaptureGroup(CaptureIndex);
+				int32 CaptureBegin = Matcher.GetCaptureGroupBeginning(CaptureIndex);
+				int32 CaptureEnd = Matcher.GetCaptureGroupEnding(CaptureIndex);
+				FString CaptureRegion = CaptureResult.Mid(CaptureBegin, CaptureEnd - CaptureBegin);
+				if (CaptureResult.IsEmpty())
+				{
+					break;
+				}
+				FString MatchString = FString::Printf(TEXT("$%i"), CaptureIndex);
+				ReplaceString = ReplaceString.Replace(*MatchString, *CaptureResult);
+			}
+			bFound = true;
+		}
+		if (bFound)
+		{
+			ReplacedName = *ReplaceString;
+			break;
+		}
+	}
+	return ReplacedName;
+}
+
+FName UMirrorDataTable::FindReplace(FName InName) const
+{
+	return GetMirrorName(InName, MirrorFindReplaceExpressions); 
+}
+
+#if WITH_EDITOR  
+
+void UMirrorDataTable::FindReplaceMirroredNames()
+{
+	if (!Skeleton)
+	{
+		return; 
+	}
+
+	const FReferenceSkeleton& RefSkeleton = Skeleton->GetReferenceSkeleton();
+
+	TSet<FName> NamesByCategory[3];
+	ForeachRow<FMirrorTableRow>(TEXT("UMirrorDataTable::FindReplaceMirroredNames"), [&NamesByCategory] (const FName& Key, const FMirrorTableRow& Value) mutable
+		{
+			NamesByCategory[Value.MirrorEntryType].Add(Value.Name);
+		}
+	);
+
+	bool bChangedTable = false;
+
+	auto AddMirrorRow = [this, &bChangedTable, &NamesByCategory](const FName& Name, const FName& MirroredName, EMirrorRowType RowType)
+	{
+		// directly add rows to avoid using FDataTableEditorUtils, which is not appropriate at this point
+		// equivalent to FDataTableEditorUtils::AddRow(DataTable, BoneName);
+		static const FString CategoryName[] = { TEXT(":Bone"), TEXT(":Notify"), TEXT(":Curve") };
+		if (!NamesByCategory[RowType].Contains(Name))
+		{
+			FName RowName = Name;
+			FMirrorTableRow* ExistingRow = FindRow<FMirrorTableRow>(RowName, TEXT("UMirrorDataTable::FindReplaceMirroredNames"), false);
+			uint32 RenameAttempts = 0; 
+			while (ExistingRow)
+			{
+				// Row names must be unique - in this case append a category name
+				FString RowString = Name.ToString() + CategoryName[RowType];
+				if (RenameAttempts > 0)
+				{
+					RowString.Appendf(TEXT("%d"), RenameAttempts);
+				}
+				RowName = *RowString;
+				ExistingRow = FindRow<FMirrorTableRow>(RowName, TEXT("UMirrorDataTable::FindReplaceMirroredNames"), false);
+				RenameAttempts++;
+			}
+			
+			if (RowStruct)
+			{
+				Modify();
+				// Allocate data to store information, using UScriptStruct to know its size
+				uint8* RowData = (uint8*)FMemory::Malloc(RowStruct->GetStructureSize());
+				RowStruct->InitializeStruct(RowData);
+				// Add to row map
+				AddRowInternal(RowName, RowData);
+				bChangedTable = true;
+			}
+			FMirrorTableRow* MirrorRow = FindRow<FMirrorTableRow>(RowName, TEXT("UMirrorDataTable::FindReplaceMirroredNames"), false);
+			if (MirrorRow)
+			{
+				MirrorRow->MirroredName = MirroredName;
+				MirrorRow->Name = Name;
+				MirrorRow->MirrorEntryType = RowType;
+			}
+		}
+	};
+
+	for (int32 BoneIndex = 0; BoneIndex < RefSkeleton.GetNum(); ++BoneIndex)
+	{
+		FName BoneName = RefSkeleton.GetBoneName(BoneIndex);
+		FName MirroredName = FindReplace(BoneName);
+		if (!MirroredName.IsNone() && RefSkeleton.FindBoneIndex(MirroredName) != INDEX_NONE)
+		{
+			AddMirrorRow(BoneName, MirroredName, EMirrorRowType::Bone);
+		}
+	}
+
+	for (const FName& Notify : Skeleton->AnimationNotifies)
+	{
+		FName MirroredName = FindReplace(Notify);
+		if (!MirroredName.IsNone() && Skeleton->AnimationNotifies.Contains(MirroredName))
+		{
+			AddMirrorRow(Notify, MirroredName, EMirrorRowType::Notify);
+		}
+	}
+
+	const FSmartNameMapping* CurveSmartNames = Skeleton->GetSmartNameContainer(USkeleton::AnimCurveMappingName);
+	if (CurveSmartNames)
+	{
+		const SmartName::UID_Type NumItems = CurveSmartNames->GetMaxUID() + 1;
+		TSet<FName> CurveNames; 
+		// For every source curve, try to find the curve with the same name in the target.
+		for (SmartName::UID_Type Index = 0; Index < NumItems; ++Index)
+		{
+			FName CurveName;
+			CurveSmartNames->GetName(Index, CurveName);
+			CurveNames.Add(CurveName);
+		}
+
+		for (const FName& CurveName : CurveNames)
+		{
+			FName MirroredName = FindReplace(CurveName);
+			if (!MirroredName.IsNone() && CurveNames.Contains(MirroredName))
+			{
+				AddMirrorRow(CurveName, MirroredName, EMirrorRowType::Curve);
+			}
+		}
+	}
+
+	if (bChangedTable)
+	{
+		OnDataTableChanged().Broadcast();
+	}
+	FillMirrorArrays(); 
+}
+
+#endif // WITH_EDITOR
+void UMirrorDataTable::FillCompactPoseMirrorBones(const FBoneContainer& BoneContainer, const TArray<int32>& MirrorBoneIndexes,  TArray<FCompactPoseBoneIndex>& OutCompactPoseMirrorBones)
+{
+	const TArray<FBoneIndexType>& RequiredBoneIndices = BoneContainer.GetBoneIndicesArray();
+	int32 NumReqBones = RequiredBoneIndices.Num();
+	OutCompactPoseMirrorBones.Reset(NumReqBones);
+
+	TArray<FBoneIndexType> const& BoneIndicesArray = BoneContainer.GetBoneIndicesArray();
+	TArray<int32> const& PoseToSkeletonBoneIndexArray = BoneContainer.GetPoseToSkeletonBoneIndexArray();
+	TArray<int32> MeshIndexToCompactPoseIndex;
+	MeshIndexToCompactPoseIndex.Reset(PoseToSkeletonBoneIndexArray.Num());
+	MeshIndexToCompactPoseIndex.AddUninitialized(PoseToSkeletonBoneIndexArray.Num());
+
+	for (int32 CompactBoneIndex = 0; CompactBoneIndex < NumReqBones; ++CompactBoneIndex)
+	{
+		FBoneIndexType MeshPoseIndex = BoneIndicesArray[CompactBoneIndex];
+		MeshIndexToCompactPoseIndex[MeshPoseIndex] = CompactBoneIndex;
+	}
+
+	if (MirrorBoneIndexes.Num() > 0)
+	{
+		for (int32 CompactBoneIndex = 0; CompactBoneIndex < NumReqBones; ++CompactBoneIndex)
+		{
+			FBoneIndexType MeshPoseIndex = BoneIndicesArray[CompactBoneIndex];
+
+			//Mirror Bone
+			const int32 MirrorIndex = MirrorBoneIndexes.IsValidIndex(MeshPoseIndex) ? MirrorBoneIndexes[MeshPoseIndex] : INDEX_NONE;
+			const int32 CompactMirrorIndex = MeshIndexToCompactPoseIndex.IsValidIndex(MirrorIndex) ? MeshIndexToCompactPoseIndex[MirrorIndex] : INDEX_NONE;
+
+			OutCompactPoseMirrorBones.Add(FCompactPoseBoneIndex(CompactMirrorIndex));
+		}
+	}
+	else
+	{
+		for (int32 CompactBoneIndex = 0; CompactBoneIndex < NumReqBones; ++CompactBoneIndex)
+		{
+			OutCompactPoseMirrorBones.Add(FCompactPoseBoneIndex(INDEX_NONE));
+		}
+	}
+}
+
+void UMirrorDataTable::FillMirrorBoneIndexes(const FReferenceSkeleton& ReferenceSkeleton, TArray<int32>& OutMirrorBoneIndexes) const
+{
+	// Reset the mirror table to defaults (no mirroring)
+	OutMirrorBoneIndexes.SetNumUninitialized(ReferenceSkeleton.GetNum());
+	FMemory::Memset(OutMirrorBoneIndexes.GetData(), INDEX_NONE, OutMirrorBoneIndexes.Num() * OutMirrorBoneIndexes.GetTypeSize());
+
+	TMap<FName, FName> NameToMirrorNameBoneMap; 
+	ForeachRow<FMirrorTableRow>(TEXT("UMirrorDataTable::FillMirrorBoneIndexes"), [&NameToMirrorNameBoneMap](const FName& Key, const FMirrorTableRow& Value) mutable
+		{
+			if (Value.MirrorEntryType == EMirrorRowType::Bone)
+			{
+				NameToMirrorNameBoneMap.Add(Value.Name, Value.MirroredName);
+			}
+		}
+	);
+
+	if (MirrorAxis != EAxis::None)
+	{
+		for (int32 BoneIndex = 0; BoneIndex < OutMirrorBoneIndexes.Num(); ++BoneIndex)
+		{
+			if (OutMirrorBoneIndexes[BoneIndex] == INDEX_NONE)
+			{
+				// Find the candidate mirror partner for this bone (falling back to mirroring to self)
+				FName SourceBoneName = ReferenceSkeleton.GetBoneName(BoneIndex);
+				int32 MirrorBoneIndex = INDEX_NONE;
+
+				FName* MirroredBoneName = NameToMirrorNameBoneMap.Find(SourceBoneName);
+				if (!SourceBoneName.IsNone() && MirroredBoneName)
+				{
+					MirrorBoneIndex = ReferenceSkeleton.FindBoneIndex(*MirroredBoneName);
+				}
+
+				OutMirrorBoneIndexes[BoneIndex] = MirrorBoneIndex;
+				if (MirrorBoneIndex != INDEX_NONE)
+				{
+					OutMirrorBoneIndexes[MirrorBoneIndex] = BoneIndex;
+				}
+			}
+		}
+	}
+}
+
+void UMirrorDataTable::FillMirrorArrays()
+{
+	NotifyToMirrorNotifyMap.Empty();
+	if (!Skeleton)
+	{
+		BoneToMirrorBoneIndex.Empty();
+		CurveMirrorSourceUIDArray.Empty(); 
+		CurveMirrorTargetUIDArray.Empty(); 
+		return; 
+	}
+
+	FillMirrorBoneIndexes(Skeleton->GetReferenceSkeleton(), BoneToMirrorBoneIndex);
+
+	TMap<FName, FName> CurveToMirrorCurveMap;
+	
+
+	ForeachRow<FMirrorTableRow>(TEXT("UMirrorDataTable::FillMirrorArrays"), [this, &CurveToMirrorCurveMap](const FName& Key, const FMirrorTableRow& Value) mutable
+		{
+			if (Value.MirrorEntryType == EMirrorRowType::Curve)
+			{
+				CurveToMirrorCurveMap.Add(Value.Name, Value.MirroredName);
+			}
+			else if (Value.MirrorEntryType == EMirrorRowType::Notify)
+			{
+				NotifyToMirrorNotifyMap.Add(Value.Name, Value.MirroredName);
+			}
+		}
+	);
+
+	const FSmartNameMapping* CurveSmartNames = Skeleton->GetSmartNameContainer(USkeleton::AnimCurveMappingName);
+	if (CurveSmartNames)
+	{
+		CurveMirrorSourceUIDArray.Reset(CurveToMirrorCurveMap.Num());
+		CurveMirrorTargetUIDArray.Reset(CurveToMirrorCurveMap.Num());
+		for (auto& Elem : CurveToMirrorCurveMap)
+		{
+			SmartName::UID_Type SourceCurveUID = CurveSmartNames->FindUID(Elem.Key);
+			SmartName::UID_Type TargetCurveUID = CurveSmartNames->FindUID(Elem.Value);
+			if (SourceCurveUID != INDEX_NONE && TargetCurveUID != INDEX_NONE)
+			{
+				CurveMirrorSourceUIDArray.Add(SourceCurveUID);
+				CurveMirrorTargetUIDArray.Add(TargetCurveUID);
+			}
+		}
+		CurveMirrorSourceUIDArray.Shrink(); 
+		CurveMirrorTargetUIDArray.Shrink();
+	}
+}
+
+
+
+#undef LOCTEXT_NAMESPACE
\ No newline at end of file
diff --git a/Engine/Source/Runtime/Engine/Private/PhysicsEngine/PhysicsAsset.cpp b/Engine/Source/Runtime/Engine/Private/PhysicsEngine/PhysicsAsset.cpp
--- a/Engine/Source/Runtime/Engine/Private/PhysicsEngine/PhysicsAsset.cpp
+++ b/Engine/Source/Runtime/Engine/Private/PhysicsEngine/PhysicsAsset.cpp
@@ -5,6 +5,7 @@
 =============================================================================*/ 
 
 #include "PhysicsEngine/PhysicsAsset.h"
+#include "Animation/MirrorDataTable.h"
 #include "UObject/FrameworkObjectVersion.h"
 #include "Serialization/ObjectWriter.h"
 #include "Serialization/ObjectReader.h"
@@ -502,31 +503,16 @@ FName UPhysicsAsset::FindConstraintBoneName(int32 ConstraintIndex)
 	return ConstraintSetup[ConstraintIndex]->DefaultInstance.GetChildBoneName();
 }
 
-int32 UPhysicsAsset::FindMirroredBone(class USkeletalMesh* skelMesh, int32 BoneIndex)
+int32 UPhysicsAsset::FindMirroredBone(USkeletalMesh* SkelMesh,  int32 BoneIndex)
 {
-	if (BoneIndex == INDEX_NONE)
+	if (SkelMesh)
 	{
-		return INDEX_NONE;
+		FName BoneName = SkelMesh->GetRefSkeleton().GetBoneName(BoneIndex);
+		FName MirrorName = UMirrorDataTable::GetSettingsMirrorName(BoneName);
+		return SkelMesh->GetRefSkeleton().FindBoneIndex(MirrorName); 
 	}
-	//we try to find the mirroed bone using several approaches. The first is to look for the same name but with _R instead of _L or vise versa
-	FName BoneName = skelMesh->GetRefSkeleton().GetBoneName(BoneIndex);
-	FString BoneNameString = BoneName.ToString();
 
-	bool bIsLeft = BoneNameString.Find("_L", ESearchCase::IgnoreCase, ESearchDir::FromEnd) == (BoneNameString.Len() - 2);	//has _L at the end
-	bool bIsRight = BoneNameString.Find("_R", ESearchCase::IgnoreCase, ESearchDir::FromEnd) == (BoneNameString.Len() - 2);	//has _R at the end
-	bool bMirrorConvention = bIsLeft || bIsRight;
-
-	//if the bone follows our left right naming convention then let's try and find its mirrored bone
-	int32 BoneIndexMirrored = INDEX_NONE;
-	if (bMirrorConvention)
-	{
-		FString BoneNameMirrored = BoneNameString.LeftChop(2);
-		BoneNameMirrored.Append(bIsLeft ? "_R" : "_L");
-
-		BoneIndexMirrored = skelMesh->GetRefSkeleton().FindBoneIndex(FName(*BoneNameMirrored));
-	}
-
-	return BoneIndexMirrored;
+	return INDEX_NONE; 
 }
 
 void UPhysicsAsset::GetBodyIndicesBelow(TArray<int32>& OutBodyIndices, FName InBoneName, USkeletalMesh* SkelMesh, bool bIncludeParent /*= true*/)
diff --git a/Engine/Source/Runtime/Engine/Private/SkeletalMesh.cpp b/Engine/Source/Runtime/Engine/Private/SkeletalMesh.cpp
--- a/Engine/Source/Runtime/Engine/Private/SkeletalMesh.cpp
+++ b/Engine/Source/Runtime/Engine/Private/SkeletalMesh.cpp
@@ -483,8 +483,11 @@ ENUM_CLASS_FLAGS(USkeletalMesh::EAsyncPropertyLockType);
 USkeletalMesh::USkeletalMesh(const FObjectInitializer& ObjectInitializer)
 	: Super(ObjectInitializer)
 {
+PRAGMA_DISABLE_DEPRECATION_WARNINGS
 	SetSkelMirrorAxis(EAxis::X);
 	SetSkelMirrorFlipAxis(EAxis::Z);
+PRAGMA_ENABLE_DEPRECATION_WARNINGS
+
 #if WITH_EDITORONLY_DATA
 	SetImportedModel(MakeShareable(new FSkeletalMeshModel()));
 	SetVertexColorGuid(FGuid());
@@ -4330,6 +4333,7 @@ void USkeletalMesh::ReleaseCPUResources()
 /** Allocate and initialise bone mirroring table for this skeletal mesh. Default is source = destination for each bone. */
 void USkeletalMesh::InitBoneMirrorInfo()
 {
+PRAGMA_DISABLE_DEPRECATION_WARNINGS
 	TArray<FBoneMirrorInfo>& LocalSkelMirrorTable = GetSkelMirrorTable();
 	LocalSkelMirrorTable.Empty(GetRefSkeleton().GetNum());
 	LocalSkelMirrorTable.AddZeroed(GetRefSkeleton().GetNum());
@@ -4339,11 +4343,13 @@ void USkeletalMesh::InitBoneMirrorInfo()
 	{
 		LocalSkelMirrorTable[i].SourceIndex = i;
 	}
+PRAGMA_ENABLE_DEPRECATION_WARNINGS
 }
 
 /** Utility for copying and converting a mirroring table from another SkeletalMesh. */
 void USkeletalMesh::CopyMirrorTableFrom(USkeletalMesh* SrcMesh)
 {
+PRAGMA_DISABLE_DEPRECATION_WARNINGS
 	TArray<FBoneMirrorInfo>& SrcSkelMirrorTable = SrcMesh->GetSkelMirrorTable();
 	TArray<FBoneMirrorInfo>& LocalSkelMirrorTable = GetSkelMirrorTable();
 	// Do nothing if no mirror table in source mesh
@@ -4394,8 +4400,10 @@ void USkeletalMesh::CopyMirrorTableFrom(USkeletalMesh* SrcMesh)
 			}
 		}
 	}
+PRAGMA_ENABLE_DEPRECATION_WARNINGS
 }
 
+PRAGMA_DISABLE_DEPRECATION_WARNINGS
 /** Utility for copying and converting a mirroring table from another SkeletalMesh. */
 void USkeletalMesh::ExportMirrorTable(TArray<FBoneMirrorExport> &MirrorExportInfo) const
 {
@@ -4421,11 +4429,13 @@ void USkeletalMesh::ExportMirrorTable(TArray<FBoneMirrorExport> &MirrorExportInf
 		MirrorExportInfo[i].BoneFlipAxis	= LocalSkelMirrorTable[i].BoneFlipAxis;
 	}
 }
+PRAGMA_ENABLE_DEPRECATION_WARNINGS
 
-
+PRAGMA_DISABLE_DEPRECATION_WARNINGS
 /** Utility for copying and converting a mirroring table from another SkeletalMesh. */
 void USkeletalMesh::ImportMirrorTable(const TArray<FBoneMirrorExport> &MirrorExportInfo)
 {
+
 	// Do nothing if no mirror table in source mesh
 	if( MirrorExportInfo.Num() == 0 )
 	{
@@ -4472,6 +4482,7 @@ void USkeletalMesh::ImportMirrorTable(const TArray<FBoneMirrorExport> &MirrorExp
 		}
 	}
 }
+PRAGMA_ENABLE_DEPRECATION_WARNINGS
 
 /** 
  *	Utility for checking that the bone mirroring table of this mesh is good.
@@ -4480,6 +4491,7 @@ void USkeletalMesh::ImportMirrorTable(const TArray<FBoneMirrorExport> &MirrorExp
  */
 bool USkeletalMesh::MirrorTableIsGood(FString& ProblemBones) const
 {
+PRAGMA_DISABLE_DEPRECATION_WARNINGS
 	TArray<int32>	BadBoneMirror;
 
 	const TArray<FBoneMirrorInfo>& LocalSkelMirrorTable = GetSkelMirrorTable();
@@ -4509,6 +4521,7 @@ bool USkeletalMesh::MirrorTableIsGood(FString& ProblemBones) const
 	{
 		return true;
 	}
+PRAGMA_ENABLE_DEPRECATION_WARNINGS
 }
 
 void USkeletalMesh::CreateBodySetup()
diff --git a/Engine/Source/Runtime/Engine/Private/SkeletalMeshMerge.cpp b/Engine/Source/Runtime/Engine/Private/SkeletalMeshMerge.cpp
--- a/Engine/Source/Runtime/Engine/Private/SkeletalMeshMerge.cpp
+++ b/Engine/Source/Runtime/Engine/Private/SkeletalMeshMerge.cpp
@@ -757,7 +757,8 @@ bool FSkeletalMeshMerge::ProcessMergeMesh()
 	
 	// copy settings and bone info from src meshes
 	bool bNeedsInit=true;
-
+	
+PRAGMA_DISABLE_DEPRECATION_WARNINGS
 	MergeMesh->GetSkelMirrorTable().Empty();
 
 	for( int32 MeshIdx=0; MeshIdx < SrcMeshList.Num(); MeshIdx++ )
@@ -783,7 +784,7 @@ bool FSkeletalMeshMerge::ProcessMergeMesh()
 			}
 		}
 	}
-
+PRAGMA_ENABLE_DEPRECATION_WARNINGS
 	// Rebuild inverse ref pose matrices.
 	MergeMesh->GetRefBasesInvMatrix().Empty();
 	MergeMesh->CalculateInvRefMatrices();
diff --git a/Engine/Source/Runtime/Engine/Public/AnimationRuntime.h b/Engine/Source/Runtime/Engine/Public/AnimationRuntime.h
--- a/Engine/Source/Runtime/Engine/Public/AnimationRuntime.h
+++ b/Engine/Source/Runtime/Engine/Public/AnimationRuntime.h
@@ -20,6 +20,7 @@
 
 class UBlendSpace;
 class USkeletalMeshComponent;
+class UMirrorDataTable;
 struct FA2CSPose;
 struct FA2Pose;
 struct FInputBlendPose;
@@ -443,6 +444,41 @@ public:
 	UE_DEPRECATED(4.26.0, "Please use BlendTransform() for weighted blending")
 	static void BlendTransformsByWeight(FTransform& OutTransform, const TArray<FTransform>& Transforms, const TArray<float>& Weights);
 
+	/**
+	 * Mirror a vector across the specified mirror axis 
+	 * @param	V			The vector to mirror
+	 * @param	MirrorAxis	The axis to mirror across
+	 * @return				The vector mirrored across the specified axis
+	 */
+	static FVector MirrorVector(const FVector& V, EAxis::Type MirrorAxis);
+
+	/** 
+	 * Mirror a quaternion across the specified mirror axis 
+	 * @param	Q			The quaternion to mirror
+	 * @param	MirrorAxis	The axis to mirror across
+	 * @return				The quaternion mirrored across the specified axis
+	 */
+	static FQuat MirrorQuat(const FQuat& Q, EAxis::Type MirrorAxis);
+
+	/** 
+	 * Mirror a pose with the specified MirrorDataTable.  
+	 * This method computes the required compact mirror pose and component space reference rotations each call
+	 * and should not be used for repeated calculations
+	 * 
+	 * @param	Pose			The pose which is mirrored in place
+	 * @param	MirrorDataTable	A UMirrorDataTable for the same Skeleton as the Pose 
+	 */
+	static void MirrorPose(FCompactPose& Pose, const UMirrorDataTable& MirrorDataTable);
+
+	/** Mirror Pose using cached mirror bones and components space arrays.   
+	 * 
+	 * @param	Pose						The pose which is mirrored in place
+	 * @param	MirrorAxis					The axis that all bones are mirrored across 
+	 * @param	CompactPoseMirrorBones		Compact array of bone indices. Each index contains the bone to mirror or -1 to indicate mirroring should not apply to that bone.
+	 * @param	ComponentSpaceRefRotations	Compoenent space rotations of the reference pose for each bone. 
+	 */
+	static void MirrorPose(FCompactPose& Pose, EAxis::Type MirrorAxis, const TArray<FCompactPoseBoneIndex>& CompactPoseMirrorBones, const TCustomBoneIndexArray<FQuat, FCompactPoseBoneIndex>& ComponentSpaceRefRotations);
+
 	/** 
 	 * Advance CurrentTime to CurrentTime + MoveDelta. 
 	 * It will handle wrapping if bAllowLooping is true
