Properly handle deleted actors when creating actor clusters for PIE.

diff --git a/Engine/Source/Runtime/Engine/Private/WorldPartition/WorldPartitionActorCluster.cpp b/Engine/Source/Runtime/Engine/Private/WorldPartition/WorldPartitionActorCluster.cpp
--- a/Engine/Source/Runtime/Engine/Private/WorldPartition/WorldPartitionActorCluster.cpp
+++ b/Engine/Source/Runtime/Engine/Private/WorldPartition/WorldPartitionActorCluster.cpp
@@ -207,34 +207,35 @@ void CreateActorCluster(const FWorldPartitionActorDescView& ActorDescView, EActo
 	{
 		for (const FGuid& ReferenceGuid : ActorDescView.GetReferences())
 		{
-			const FWorldPartitionActorDescView& ReferenceActorDescView = ContainerInstance.ActorDescViewMap.FindChecked(ReferenceGuid);
-
-			// Don't include references to editor-only actors
-			if (!ReferenceActorDescView.GetActorIsEditorOnly())
+			if (const FWorldPartitionActorDescView* ReferenceActorDescView = ContainerInstance.ActorDescViewMap.Find(ReferenceGuid))
 			{
-				FActorCluster* ReferenceCluster = ActorToActorCluster.FindRef(ReferenceGuid);
-				if (ReferenceCluster)
+				// Don't include references to editor-only actors
+				if (!ReferenceActorDescView->GetActorIsEditorOnly())
 				{
-					if (ReferenceCluster != ActorCluster)
+					FActorCluster* ReferenceCluster = ActorToActorCluster.FindRef(ReferenceGuid);
+					if (ReferenceCluster)
 					{
-						// Merge reference cluster in Actor's cluster
-						ActorCluster->Add(*ReferenceCluster);
-						for (const FGuid& ReferenceClusterActorGuid : ReferenceCluster->Actors)
+						if (ReferenceCluster != ActorCluster)
 						{
-							ActorToActorCluster[ReferenceClusterActorGuid] = ActorCluster;
+							// Merge reference cluster in Actor's cluster
+							ActorCluster->Add(*ReferenceCluster);
+							for (const FGuid& ReferenceClusterActorGuid : ReferenceCluster->Actors)
+							{
+								ActorToActorCluster[ReferenceClusterActorGuid] = ActorCluster;
+							}
+							ActorClustersSet.Remove(ReferenceCluster);
+							delete ReferenceCluster;
 						}
-						ActorClustersSet.Remove(ReferenceCluster);
-						delete ReferenceCluster;
 					}
-				}
-				else
-				{
-					// Put Reference in Actor's cluster
-					ActorCluster->Add(FActorCluster(World, ReferenceActorDescView, GridPlacement));
-				}
+					else
+					{
+						// Put Reference in Actor's cluster
+						ActorCluster->Add(FActorCluster(World, *ReferenceActorDescView, GridPlacement));
+					}
 
-				// Map its cluster
-				ActorToActorCluster.Add(ReferenceGuid, ActorCluster);
+					// Map its cluster
+					ActorToActorCluster.Add(ReferenceGuid, ActorCluster);
+				}
 			}
 		}
 	}
