Add definition for AssetCompilation LLM tag to avoid error log when cooking

diff --git a/Engine/Source/Runtime/Engine/Private/AssetCompilingManager.cpp b/Engine/Source/Runtime/Engine/Private/AssetCompilingManager.cpp
--- a/Engine/Source/Runtime/Engine/Private/AssetCompilingManager.cpp
+++ b/Engine/Source/Runtime/Engine/Private/AssetCompilingManager.cpp
@@ -2,6 +2,13 @@
 
 #include "AssetCompilingManager.h"
 
+#include "HAL/LowLevelMemStats.h"
+#include "HAL/LowLevelMemTracker.h"
+
+DECLARE_LLM_MEMORY_STAT(TEXT("AssetCompilation"), STAT_AssetCompilationLLM, STATGROUP_LLMFULL);
+DECLARE_LLM_MEMORY_STAT(TEXT("AssetCompilation"), STAT_AssetCompilationSummaryLLM, STATGROUP_LLM);
+LLM_DEFINE_TAG(AssetCompilation, NAME_None, NAME_None, GET_STATFNAME(STAT_AssetCompilationLLM), GET_STATFNAME(STAT_AssetCompilationSummaryLLM));
+
 #if WITH_EDITOR
 
 #include "Misc/QueuedThreadPoolWrapper.h"
