Fix CIS compilation error in shipping in AssetCompilingManager

diff --git a/Engine/Source/Runtime/Engine/Private/AssetCompilingManager.cpp b/Engine/Source/Runtime/Engine/Private/AssetCompilingManager.cpp
--- a/Engine/Source/Runtime/Engine/Private/AssetCompilingManager.cpp
+++ b/Engine/Source/Runtime/Engine/Private/AssetCompilingManager.cpp
@@ -5,9 +5,11 @@
 #include "HAL/LowLevelMemStats.h"
 #include "HAL/LowLevelMemTracker.h"
 
+#if ENABLE_LOW_LEVEL_MEM_TRACKER
 DECLARE_LLM_MEMORY_STAT(TEXT("AssetCompilation"), STAT_AssetCompilationLLM, STATGROUP_LLMFULL);
 DECLARE_LLM_MEMORY_STAT(TEXT("AssetCompilation"), STAT_AssetCompilationSummaryLLM, STATGROUP_LLM);
 LLM_DEFINE_TAG(AssetCompilation, NAME_None, NAME_None, GET_STATFNAME(STAT_AssetCompilationLLM), GET_STATFNAME(STAT_AssetCompilationSummaryLLM));
+#endif
 
 #if WITH_EDITOR
 
