[ODSC] - Properly initialize compiled shader maps with the editor only flag.

GetCurrentShaderPermutationFlags was being called which will always return the editor only data flag regardless of what platform we are cooking for which is the wrong behavior.  I've changed this call to be GetShaderPermutationFlags which will check the layout parameters's flag for editor only data which will have per-platform information.

- Pass in TargetPlatform so we don't create the MaterialID as the current platform (in this case the cooker) which would be wrong.
- Pass LayoutParams to GetDependentShaderAndVFTypes so it can pass the appropriate editor only data flag into ShouldCompilePermutation.

diff --git a/Engine/Source/Developer/MaterialBaking/Private/ExportMaterialProxy.h b/Engine/Source/Developer/MaterialBaking/Private/ExportMaterialProxy.h
--- a/Engine/Source/Developer/MaterialBaking/Private/ExportMaterialProxy.h
+++ b/Engine/Source/Developer/MaterialBaking/Private/ExportMaterialProxy.h
@@ -258,7 +258,7 @@ public:
 			TArray<FShaderType*> ShaderTypes;
 			TArray<FVertexFactoryType*> VFTypes;
 			TArray<const FShaderPipelineType*> ShaderPipelineTypes;
-			GetDependentShaderAndVFTypes(GMaxRHIShaderPlatform, ShaderTypes, ShaderPipelineTypes, VFTypes);
+			GetDependentShaderAndVFTypes(GMaxRHIShaderPlatform, ResourceId.LayoutParams, ShaderTypes, ShaderPipelineTypes, VFTypes);
 
 			// Overwrite the shader map Id's dependencies with ones that came from the FMaterial actually being compiled (this)
 			// This is necessary as we change FMaterial attributes like GetShadingModels(), which factor into the ShouldCache functions that determine dependent shader types
diff --git a/Engine/Source/Developer/MaterialUtilities/Private/MaterialUtilities.cpp b/Engine/Source/Developer/MaterialUtilities/Private/MaterialUtilities.cpp
--- a/Engine/Source/Developer/MaterialUtilities/Private/MaterialUtilities.cpp
+++ b/Engine/Source/Developer/MaterialUtilities/Private/MaterialUtilities.cpp
@@ -433,7 +433,7 @@ public:
 			TArray<FShaderType*> ShaderTypes;
 			TArray<FVertexFactoryType*> VFTypes;
 			TArray<const FShaderPipelineType*> ShaderPipelineTypes;
-			GetDependentShaderAndVFTypes(GMaxRHIShaderPlatform, ShaderTypes, ShaderPipelineTypes, VFTypes);
+			GetDependentShaderAndVFTypes(GMaxRHIShaderPlatform, ResourceId.LayoutParams, ShaderTypes, ShaderPipelineTypes, VFTypes);
 
 			// Overwrite the shader map Id's dependencies with ones that came from the FMaterial actually being compiled (this)
 			// This is necessary as we change FMaterial attributes like GetShadingModels(), which factor into the ShouldCache functions that determine dependent shader types
diff --git a/Engine/Source/Editor/UnrealEd/Private/Lightmass/LightmassRender.cpp b/Engine/Source/Editor/UnrealEd/Private/Lightmass/LightmassRender.cpp
--- a/Engine/Source/Editor/UnrealEd/Private/Lightmass/LightmassRender.cpp
+++ b/Engine/Source/Editor/UnrealEd/Private/Lightmass/LightmassRender.cpp
@@ -261,7 +261,7 @@ public:
 					TArray<FShaderType*> ShaderTypes;
 					TArray<FVertexFactoryType*> VFTypes;
 					TArray<const FShaderPipelineType*> ShaderPipelineTypes;
-					GetDependentShaderAndVFTypes(GMaxRHIShaderPlatform, ShaderTypes, ShaderPipelineTypes, VFTypes);
+					GetDependentShaderAndVFTypes(GMaxRHIShaderPlatform, ResourceId.LayoutParams, ShaderTypes, ShaderPipelineTypes, VFTypes);
 
 					// Overwrite the shader map Id's dependencies with ones that came from the FMaterial actually being compiled (this)
 					// This is necessary as we change FMaterial attributes like GetShadingModels(), which factor into the ShouldCache functions that determine dependent shader types
diff --git a/Engine/Source/Runtime/Engine/Classes/Materials/Material.h b/Engine/Source/Runtime/Engine/Classes/Materials/Material.h
--- a/Engine/Source/Runtime/Engine/Classes/Materials/Material.h
+++ b/Engine/Source/Runtime/Engine/Classes/Materials/Material.h
@@ -1552,7 +1552,8 @@ public:
 	/** Builds a map from UMaterialInterface name to the shader maps that are needed for rendering on the given platform. */
 	ENGINE_API static void CompileMaterialsForRemoteRecompile(
 		const TArray<UMaterialInterface*>& MaterialsToCompile,
-		EShaderPlatform ShaderPlatform, 
+		EShaderPlatform ShaderPlatform,
+		class ITargetPlatform* TargetPlatform,
 		TMap<FString, TArray<TRefCountPtr<class FMaterialShaderMap> > >& OutShaderMaps);
 
 #if WITH_EDITOR
diff --git a/Engine/Source/Runtime/Engine/Private/Materials/Material.cpp b/Engine/Source/Runtime/Engine/Private/Materials/Material.cpp
--- a/Engine/Source/Runtime/Engine/Private/Materials/Material.cpp
+++ b/Engine/Source/Runtime/Engine/Private/Materials/Material.cpp
@@ -5516,7 +5516,8 @@ void UMaterial::RestoreMaterialShadersFromMemory(const TMap<FMaterialShaderMap*,
 
 void UMaterial::CompileMaterialsForRemoteRecompile(
 	const TArray<UMaterialInterface*>& MaterialsToCompile,
-	EShaderPlatform ShaderPlatform, 
+	EShaderPlatform ShaderPlatform,
+	ITargetPlatform* TargetPlatform,
 	TMap<FString, TArray<TRefCountPtr<FMaterialShaderMap> > >& OutShaderMaps)
 {
 	// Build a map from UMaterial / UMaterialInstance to the resources which are being compiled
@@ -5533,12 +5534,12 @@ void UMaterial::CompileMaterialsForRemoteRecompile(
 		if (MaterialInstance && MaterialInstance->bHasStaticPermutationResource)
 		{
 			TArray<FMaterialResource*>& ResourceArray = CompilingResources.Add(Material->GetPathName(), TArray<FMaterialResource*>());
-			MaterialInstance->CacheResourceShadersForCooking(ShaderPlatform, ResourceArray);
+			MaterialInstance->CacheResourceShadersForCooking(ShaderPlatform, ResourceArray, EMaterialShaderPrecompileMode::Default, TargetPlatform);
 		}
 		else if (BaseMaterial)
 		{
 			TArray<FMaterialResource*>& ResourceArray = CompilingResources.Add(Material->GetPathName(), TArray<FMaterialResource*>());
-			BaseMaterial->CacheResourceShadersForCooking(ShaderPlatform, ResourceArray);
+			BaseMaterial->CacheResourceShadersForCooking(ShaderPlatform, ResourceArray, TargetPlatform);
 		}
 	}
 
diff --git a/Engine/Source/Runtime/Engine/Private/Materials/MaterialShared.cpp b/Engine/Source/Runtime/Engine/Private/Materials/MaterialShared.cpp
--- a/Engine/Source/Runtime/Engine/Private/Materials/MaterialShared.cpp
+++ b/Engine/Source/Runtime/Engine/Private/Materials/MaterialShared.cpp
@@ -630,11 +630,13 @@ void FMaterial::GetShaderMapId(EShaderPlatform Platform, const ITargetPlatform*
 	else
 	{
 #if WITH_EDITOR
+		OutId.LayoutParams.InitializeForPlatform(TargetPlatform);
+
 		TArray<FShaderType*> ShaderTypes;
 		TArray<FVertexFactoryType*> VFTypes;
 		TArray<const FShaderPipelineType*> ShaderPipelineTypes;
 
-		GetDependentShaderAndVFTypes(Platform, ShaderTypes, ShaderPipelineTypes, VFTypes);
+		GetDependentShaderAndVFTypes(Platform, OutId.LayoutParams, ShaderTypes, ShaderPipelineTypes, VFTypes);
 
 		OutId.Usage = GetShaderMapUsage();
 		OutId.bUsingNewHLSLGenerator = IsUsingNewHLSLGenerator();
@@ -644,7 +646,6 @@ void FMaterial::GetShaderMapId(EShaderPlatform Platform, const ITargetPlatform*
 		OutId.SetShaderDependencies(ShaderTypes, ShaderPipelineTypes, VFTypes, Platform);
 		GetReferencedTexturesHash(Platform, OutId.TextureReferencesHash);
 
-		OutId.LayoutParams.InitializeForPlatform(TargetPlatform);
 #else
 		OutId.QualityLevel = GetQualityLevel();
 		OutId.FeatureLevel = GetFeatureLevel();
@@ -3532,10 +3533,10 @@ static void AddSortedShaderPipeline(TArray<const FShaderPipelineType*>& Pipeline
 	}
 }
 
-void FMaterial::GetDependentShaderAndVFTypes(EShaderPlatform Platform, TArray<FShaderType*>& OutShaderTypes, TArray<const FShaderPipelineType*>& OutShaderPipelineTypes, TArray<FVertexFactoryType*>& OutVFTypes) const
+void FMaterial::GetDependentShaderAndVFTypes(EShaderPlatform Platform, const FPlatformTypeLayoutParameters& LayoutParams, TArray<FShaderType*>& OutShaderTypes, TArray<const FShaderPipelineType*>& OutShaderPipelineTypes, TArray<FVertexFactoryType*>& OutVFTypes) const
 {
 	const FMaterialShaderParameters MaterialParameters(this);
-	const FMaterialShaderMapLayout& Layout = AcquireMaterialShaderMapLayout(Platform, GetCurrentShaderPermutationFlags(), MaterialParameters);
+	const FMaterialShaderMapLayout& Layout = AcquireMaterialShaderMapLayout(Platform, GetShaderPermutationFlags(LayoutParams), MaterialParameters);
 
 	for (const FShaderLayoutEntry& Shader : Layout.Shaders)
 	{
diff --git a/Engine/Source/Runtime/Engine/Private/ShaderCompiler/ShaderCompiler.cpp b/Engine/Source/Runtime/Engine/Private/ShaderCompiler/ShaderCompiler.cpp
--- a/Engine/Source/Runtime/Engine/Private/ShaderCompiler/ShaderCompiler.cpp
+++ b/Engine/Source/Runtime/Engine/Private/ShaderCompiler/ShaderCompiler.cpp
@@ -6176,7 +6176,7 @@ void RecompileShadersForRemote(
 				if (MeshMaterialMaps != NULL && (OutdatedShaderTypes.Num() || OutdatedFactoryTypes.Num() || bCompileChangedShaders == false))
 				{
 					TMap<FString, TArray<TRefCountPtr<FMaterialShaderMap> > > CompiledShaderMaps;
-					UMaterial::CompileMaterialsForRemoteRecompile(MaterialsToCompile, ShaderPlatform, CompiledShaderMaps);
+					UMaterial::CompileMaterialsForRemoteRecompile(MaterialsToCompile, ShaderPlatform, TargetPlatform, CompiledShaderMaps);
 
 					// write the shader compilation info to memory, converting fnames to strings
 					FMemoryWriter MemWriter(*MeshMaterialMaps, true);
diff --git a/Engine/Source/Runtime/Engine/Public/MaterialShared.h b/Engine/Source/Runtime/Engine/Public/MaterialShared.h
--- a/Engine/Source/Runtime/Engine/Public/MaterialShared.h
+++ b/Engine/Source/Runtime/Engine/Public/MaterialShared.h
@@ -2103,7 +2103,7 @@ protected:
 	virtual FGuid GetMaterialId() const = 0;
 	
 	/** Produces arrays of any shader and vertex factory type that this material is dependent on. */
-	ENGINE_API void GetDependentShaderAndVFTypes(EShaderPlatform Platform, TArray<FShaderType*>& OutShaderTypes, TArray<const FShaderPipelineType*>& OutShaderPipelineTypes, TArray<FVertexFactoryType*>& OutVFTypes) const;
+	ENGINE_API void GetDependentShaderAndVFTypes(EShaderPlatform Platform, const FPlatformTypeLayoutParameters& LayoutParams, TArray<FShaderType*>& OutShaderTypes, TArray<const FShaderPipelineType*>& OutShaderPipelineTypes, TArray<FVertexFactoryType*>& OutVFTypes) const;
 
 	bool GetLoadedCookedShaderMapId() const { return bLoadedCookedShaderMapId; }
 
