Chaos - fix scaled and transformed convex collision

diff --git a/Engine/Source/Runtime/Engine/Private/PhysicsEngine/Experimental/ChaosInterfaceUtils.cpp b/Engine/Source/Runtime/Engine/Private/PhysicsEngine/Experimental/ChaosInterfaceUtils.cpp
--- a/Engine/Source/Runtime/Engine/Private/PhysicsEngine/Experimental/ChaosInterfaceUtils.cpp
+++ b/Engine/Source/Runtime/Engine/Private/PhysicsEngine/Experimental/ChaosInterfaceUtils.cpp
@@ -321,38 +321,35 @@ namespace ChaosInterface
 			for (uint32 i = 0; i < static_cast<uint32>(InParams.Geometry->ConvexElems.Num()); ++i)
 			{
 				const FKConvexElem& CollisionBody = InParams.Geometry->ConvexElems[i];
-				const FTransform& ConvexTransform = InParams.LocalTransform;
 				if (const auto& ConvexImplicit = CollisionBody.GetChaosConvexMesh())
 				{
-					const FVector ScaledSize = (Scale.GetAbs() * CollisionBody.ElemBox.GetSize());	// Note: Scale can be negative
+					// Extract the scale from the transform - we have separate wrapper classes for scale versus translate/rotate 
+					const FVector NetScale = Scale * InParams.LocalTransform.GetScale3D();
+					FTransform ConvexTransform = FTransform(InParams.LocalTransform.GetRotation(), Scale * InParams.LocalTransform.GetLocation(), FVector(1, 1, 1));
+					const FVector ScaledSize = (NetScale.GetAbs() * CollisionBody.ElemBox.GetSize());	// Note: Scale can be negative
 					const float CollisionMargin = FMath::Min(ScaledSize.GetMin() * CollisionMarginFraction, CollisionMarginMax);
 
-					if (!ConvexTransform.GetTranslation().IsNearlyZero() || !ConvexTransform.GetRotation().IsIdentity())
+					// Wrap the convex in a scaled or instanced wrapper depending on scale value, and add a margin
+					// NOTE: CollisionMargin is on the Instance/Scaled wrapper, not the inner convex (which is shared and should not have a margin).
+					TUniquePtr<Chaos::FImplicitObject> Implicit;
+					if (NetScale == FVector(1))
 					{
-						// this path is taken when objects are welded
-						//TUniquePtr<Chaos::FImplicitObject> ScaledImplicit = TUniquePtr<Chaos::FImplicitObject>(new Chaos::TImplicitObjectScaled<Chaos::FImplicitObject>(MakeSerializable(ConvexImplicit), Scale, CollisionMargin));
-						TUniquePtr<Chaos::FImplicitObject> Implicit = TUniquePtr<Chaos::FImplicitObject>(new Chaos::TImplicitObjectTransformed<float, 3>(MakeSerializable(ConvexImplicit), ConvexTransform));
-						TUniquePtr<Chaos::FPerShapeData> NewShape = NewShapeHelper(MakeSerializable(Implicit), Shapes.Num(), (void*)CollisionBody.GetUserData(), CollisionBody.GetCollisionEnabled());
-						Shapes.Emplace(MoveTemp(NewShape));
-						Geoms.Add(MoveTemp(Implicit));
+						Implicit = TUniquePtr<Chaos::FImplicitObject>(new Chaos::TImplicitObjectInstanced<Chaos::FConvex>(ConvexImplicit, CollisionMargin));
 					}
 					else
 					{
-						// NOTE: CollisionMargin is on the Instance/Scaled wrapper, not the inner convex (which has no margin). This means that convex shapes grow by the margion size...
-						TUniquePtr<Chaos::FImplicitObject> Implicit;
-						if (Scale == FVector(1))
-						{
-							Implicit = TUniquePtr<Chaos::FImplicitObject>(new Chaos::TImplicitObjectInstanced<Chaos::FConvex>(ConvexImplicit, CollisionMargin));
-						}
-						else
-						{
-							Implicit = TUniquePtr<Chaos::FImplicitObject>(new Chaos::TImplicitObjectScaled<Chaos::FConvex>(ConvexImplicit, Scale, CollisionMargin));
-						}
-
-						TUniquePtr<Chaos::FPerShapeData> NewShape = NewShapeHelper(MakeSerializable(Implicit),Shapes.Num(), (void*)CollisionBody.GetUserData(), CollisionBody.GetCollisionEnabled());
-						Shapes.Emplace(MoveTemp(NewShape));
-						Geoms.Add(MoveTemp(Implicit));
+						Implicit = TUniquePtr<Chaos::FImplicitObject>(new Chaos::TImplicitObjectScaled<Chaos::FConvex>(ConvexImplicit, NetScale, CollisionMargin));
 					}
+
+					// Wrap the convex in a non-scaled transform if necessary (the scale is pulled out above)
+					if (!ConvexTransform.GetTranslation().IsNearlyZero() || !ConvexTransform.GetRotation().IsIdentity())
+					{
+						Implicit = TUniquePtr<Chaos::FImplicitObject>(new Chaos::TImplicitObjectTransformed<float, 3>(MoveTemp(Implicit), ConvexTransform));
+					}
+
+					TUniquePtr<Chaos::FPerShapeData> NewShape = NewShapeHelper(MakeSerializable(Implicit), Shapes.Num(), (void*)CollisionBody.GetUserData(), CollisionBody.GetCollisionEnabled());
+					Shapes.Emplace(MoveTemp(NewShape));
+					Geoms.Add(MoveTemp(Implicit));
 				}
 			}
 		}
