Skip editor-only actors when running the game in PIE standalone.
Also forced LandmassActor class to be editor-only in order to support PIE standalone.

diff --git a/Engine/Plugins/Experimental/Landmass/Source/Editor/Private/LandmassActor.cpp b/Engine/Plugins/Experimental/Landmass/Source/Editor/Private/LandmassActor.cpp
--- a/Engine/Plugins/Experimental/Landmass/Source/Editor/Private/LandmassActor.cpp
+++ b/Engine/Plugins/Experimental/Landmass/Source/Editor/Private/LandmassActor.cpp
@@ -22,7 +22,6 @@ ALandmassActor::ALandmassActor(const FObjectInitializer& ObjectInitializer)
 		FLevelEditorModule& LevelEditor = FModuleManager::GetModuleChecked<FLevelEditorModule>("LevelEditor");
 		OnActorSelectionChangedHandle = LevelEditor.OnActorSelectionChanged().AddUObject(this, &ALandmassActor::HandleActorSelectionChanged);
 	}
-	bIsEditorOnlyActor = true;
 }
 
 void ALandmassActor::Tick(float DeltaSeconds)
diff --git a/Engine/Plugins/Experimental/Landmass/Source/Editor/Public/LandmassActor.h b/Engine/Plugins/Experimental/Landmass/Source/Editor/Public/LandmassActor.h
--- a/Engine/Plugins/Experimental/Landmass/Source/Editor/Public/LandmassActor.h
+++ b/Engine/Plugins/Experimental/Landmass/Source/Editor/Public/LandmassActor.h
@@ -21,6 +21,8 @@ public:
 	UFUNCTION(BlueprintNativeEvent, CallInEditor, BlueprintCallable, Category = "Tick")
 	void CustomTick(float DeltaSeconds);
 
+	virtual bool IsEditorOnly() const override { return true; }
+
 	virtual bool ShouldTickIfViewportsOnly() const override;
 	virtual void Tick(float DeltaSeconds) override;
 
diff --git a/Engine/Source/Runtime/Engine/Private/WorldPartition/ActorDescContainer.cpp b/Engine/Source/Runtime/Engine/Private/WorldPartition/ActorDescContainer.cpp
--- a/Engine/Source/Runtime/Engine/Private/WorldPartition/ActorDescContainer.cpp
+++ b/Engine/Source/Runtime/Engine/Private/WorldPartition/ActorDescContainer.cpp
@@ -101,23 +101,40 @@ void UActorDescContainer::BeginDestroy()
 #if WITH_EDITOR
 TUniquePtr<FWorldPartitionActorDesc> UActorDescContainer::GetActorDescriptor(const FAssetData& InAssetData)
 {
-	FString ActorClass;
+	FString ActorClassName;
 	static FName NAME_ActorMetaDataClass(TEXT("ActorMetaDataClass"));
-	if (InAssetData.GetTagValue(NAME_ActorMetaDataClass, ActorClass))
+	if (InAssetData.GetTagValue(NAME_ActorMetaDataClass, ActorClassName))
 	{
 		FString ActorMetaDataStr;
 		static FName NAME_ActorMetaData(TEXT("ActorMetaData"));
 		if (InAssetData.GetTagValue(NAME_ActorMetaData, ActorMetaDataStr))
 		{
+			bool bIsValidClass = true;
+			UClass* ActorClass = FindObject<UClass>(ANY_PACKAGE, *ActorClassName, true);
+
+			if (!ActorClass)
+			{
+				ActorClass = AActor::StaticClass();
+				bIsValidClass = false;
+			}
+
 			FWorldPartitionActorDescInitData ActorDescInitData;
-			ActorDescInitData.NativeClass = FindObjectChecked<UClass>(ANY_PACKAGE, *ActorClass, true);
+			ActorDescInitData.NativeClass = ActorClass;
 			ActorDescInitData.PackageName = InAssetData.PackageName;
 			ActorDescInitData.ActorPath = InAssetData.ObjectPath;
 			FBase64::Decode(ActorMetaDataStr, ActorDescInitData.SerializedData);
 
 			TUniquePtr<FWorldPartitionActorDesc> NewActorDesc(AActor::CreateClassActorDesc(ActorDescInitData.NativeClass));
+
 			NewActorDesc->Init(this, ActorDescInitData);
-			return NewActorDesc;
+			
+			check(bIsValidClass || (NewActorDesc->GetActorIsEditorOnly() && IsRunningGame()));
+			if (!NewActorDesc->GetActorIsEditorOnly() || !IsRunningGame())
+			{				
+				return NewActorDesc;
+			}
+
+			return nullptr;
 		}
 	}
 
