Just a small change, to use a reference instead of copy.

#rnx

diff --git a/Engine/Source/Runtime/Engine/Private/Animation/Skeleton.cpp b/Engine/Source/Runtime/Engine/Private/Animation/Skeleton.cpp
--- a/Engine/Source/Runtime/Engine/Private/Animation/Skeleton.cpp
+++ b/Engine/Source/Runtime/Engine/Private/Animation/Skeleton.cpp
@@ -146,7 +146,7 @@ bool USkeleton::IsCompatibleSkeletonByAssetString(const FString& SkeletonAssetSt
 
 	// Now check against the list of compatible skeletons.
 	FString CompatibleName;
-	for (auto& CompatibleSkeleton : CompatibleSkeletons)
+	for (const auto& CompatibleSkeleton : CompatibleSkeletons)
 	{
 		CompatibleName = FString::Printf(TEXT("%s'%s'"), *GetClass()->GetName(), *CompatibleSkeleton.ToString());
 		if (CompatibleName == SkeletonAssetString)
