Enabled analytic derivatives.
Tested by [REDACTED_PERSON]

diff --git a/Engine/Shaders/Public/ShaderVersion.ush b/Engine/Shaders/Public/ShaderVersion.ush
--- a/Engine/Shaders/Public/ShaderVersion.ush
+++ b/Engine/Shaders/Public/ShaderVersion.ush
@@ -3,4 +3,4 @@
 // in Platform.ush (which should be included in any shader) it allows to invalidate the shader DDC.
 //
 // If you are merging streams and there is a conflict with this GUID you should make a new GUID rather than taking one or the other.
-// GUID = D4BD6A39-2B29-485A-8DBA-25CA2FCBB666
+// GUID = E4BD6A49-2C29-4A5A-8FBD-55CA2FCBB669
diff --git a/Engine/Source/Runtime/Engine/Private/Materials/HLSLMaterialTranslator.cpp b/Engine/Source/Runtime/Engine/Private/Materials/HLSLMaterialTranslator.cpp
--- a/Engine/Source/Runtime/Engine/Private/Materials/HLSLMaterialTranslator.cpp
+++ b/Engine/Source/Runtime/Engine/Private/Materials/HLSLMaterialTranslator.cpp
@@ -6,7 +6,7 @@
 
 #if WITH_EDITORONLY_DATA
 
-static int32 GAnalyticDerivEnabled = 0;
+static int32 GAnalyticDerivEnabled = 1;
 static FAutoConsoleVariableRef CVarAnalyticDerivEnabled(
 	TEXT("r.MaterialEditor.AnalyticDeriv"),
 	GAnalyticDerivEnabled,
