Reduce Log Verbosity for Physics and Spawn strictly to Error while spawning Mass Actors as we can expect numerous "intended" failure, temporarily.
#rnx

diff --git a/Engine/Source/Runtime/Engine/Classes/Engine/World.h b/Engine/Source/Runtime/Engine/Classes/Engine/World.h
--- a/Engine/Source/Runtime/Engine/Classes/Engine/World.h
+++ b/Engine/Source/Runtime/Engine/Classes/Engine/World.h
@@ -140,7 +140,7 @@ private:
 };
 
 
-DECLARE_LOG_CATEGORY_EXTERN(LogSpawn, Warning, All);
+ENGINE_API DECLARE_LOG_CATEGORY_EXTERN(LogSpawn, Warning, All);
 
 DECLARE_MULTICAST_DELEGATE_OneParam(FOnActorSpawned, AActor*);
 DECLARE_MULTICAST_DELEGATE_OneParam(FOnFeatureLevelChanged, ERHIFeatureLevel::Type);
