Pose asset optimizations

Data:
* TArrays are now stored in TLS object to reduce re-allocations
* Replaced TrackMap with TrackBoneIndices, matching # or entries with TrackNames
* Added FPoseAssetInfluences wrapper struct containing indices into Poses array for each TrackName entry - mapping from a track to each pose that influences its final transform

Behaviour:
* Rather than building a map of weights to FPoseData* we simply set the weight in an array of size Poses.Num() and store any pose indices that contain any weighting
* If normalizing pose weights, do curve blending inline to prevent looping over weights twice
* Rather than looping over all Poses with weighting, loop over known pose indices influencing a Tracks final transform, and if weighted process it
* Rather than storing per-pose blend transforms and weights, apply blending when looping over in0fluencing poses for track. removing an unnecessary loop
* Retarget final bone transform within blending loop, removing another unnecessary loop

Test - 16 instances running AnimBP driving a pose asset with 79 poses on the Mannequin (68 bones) with all poses weighted [worker thread time]
before: 9.7ms
after: 3.5ms

diff --git a/Engine/Source/Runtime/Core/Public/UObject/UE5ReleaseStreamObjectVersion.h b/Engine/Source/Runtime/Core/Public/UObject/UE5ReleaseStreamObjectVersion.h
--- a/Engine/Source/Runtime/Core/Public/UObject/UE5ReleaseStreamObjectVersion.h
+++ b/Engine/Source/Runtime/Core/Public/UObject/UE5ReleaseStreamObjectVersion.h
@@ -24,6 +24,9 @@ struct CORE_API FUE5ReleaseStreamObjectVersion
 		// LevelInstance serialize runtime behavior
 		LevelInstanceSerializeRuntimeBehavior,
 
+		// Refactoring Pose Asset runtime data structures
+        PoseAssetRuntimeRefactor,
+
 		// -----<new versions can be added above this line>-------------------------------------------------
 		VersionPlusOne,
 		LatestVersion = VersionPlusOne - 1
diff --git a/Engine/Source/Runtime/Engine/Classes/Animation/PoseAsset.h b/Engine/Source/Runtime/Engine/Classes/Animation/PoseAsset.h
--- a/Engine/Source/Runtime/Engine/Classes/Animation/PoseAsset.h
+++ b/Engine/Source/Runtime/Engine/Classes/Animation/PoseAsset.h
@@ -48,16 +48,33 @@ struct ENGINE_API FPoseData
 	UPROPERTY()
 	TArray<FTransform>		LocalSpacePose;
 
-	// this is PoseContainer.Tracks to Buffer Index of LocalSpacePose
-	UPROPERTY()
-	TMap<int32, int32>		TrackToBufferIndex;
-
 	// # of array match with # of Curves in PoseDataContainer
 	// curve data is not compressed
  	UPROPERTY()
  	TArray<float>			CurveData;
 };
 
+USTRUCT()
+struct FPoseAssetInfluence
+{
+	GENERATED_USTRUCT_BODY()
+
+	UPROPERTY()
+	int32 PoseIndex;
+
+	UPROPERTY()
+	int32 BoneTransformIndex;
+};
+
+USTRUCT()
+struct FPoseAssetInfluences
+{
+	GENERATED_USTRUCT_BODY()
+
+    UPROPERTY()
+	TArray<FPoseAssetInfluence> Influences;
+};
+
 /**
 * Pose data container
 * 
@@ -77,9 +94,12 @@ private:
 	UPROPERTY()
 	TArray<FName>							Tracks;
 
-	// cache for the track names to skeleton index 
+	// cache containting the skeleton indices for FName in Tracks array
 	UPROPERTY(transient)
-	TMap<FName, int32>						TrackMap;
+	TArray<int32>						TrackBoneIndices;
+
+	UPROPERTY()
+	TArray<FPoseAssetInfluences>		TrackPoseInfluenceIndices;
 	
 	// this is list of poses
 	UPROPERTY()
@@ -98,7 +118,7 @@ private:
 	int32 GetNumPoses() const { return Poses.Num();  }
 	bool Contains(FSmartName PoseName) const { return PoseNames.Contains(PoseName); }
 
-	bool IsValid() const { return PoseNames.Num() == Poses.Num() && Tracks.Num() == TrackMap.Num(); }
+	bool IsValid() const { return PoseNames.Num() == Poses.Num() && Tracks.Num() == TrackBoneIndices.Num(); }
 	void GetPoseCurve(const FPoseData* PoseData, FBlendedCurve& OutCurve) const;
 	void BlendPoseCurve(const FPoseData* PoseData, FBlendedCurve& OutCurve, float Weight) const;
 
@@ -294,7 +314,8 @@ private:
 #endif // WITH_EDITOR	
 
 private:
-	void RecacheTrackmap();
+	void UpdateTrackBoneIndices();
+	bool RemoveInvalidTracks();
 
 #if WITH_EDITORONLY_DATA
 	void UpdateRetargetSourceAsset();
diff --git a/Engine/Source/Runtime/Engine/Private/Animation/PoseAsset.cpp b/Engine/Source/Runtime/Engine/Private/Animation/PoseAsset.cpp
--- a/Engine/Source/Runtime/Engine/Private/Animation/PoseAsset.cpp
+++ b/Engine/Source/Runtime/Engine/Private/Animation/PoseAsset.cpp
@@ -10,6 +10,7 @@
 #include "Animation/AnimationPoseData.h"
 #include "Animation/AnimData/AnimDataModel.h"
 #include "Animation/AnimSequenceHelpers.h"
+#include "UObject/UE5ReleaseStreamObjectVersion.h"
 
 #define LOCTEXT_NAMESPACE "PoseAsset"
 
@@ -48,7 +49,7 @@ void FPoseDataContainer::Reset()
 	PoseNames.Reset();
 	Poses.Reset();
 	Tracks.Reset();
-	TrackMap.Reset();
+	TrackBoneIndices.Reset();
 	Curves.Reset();
 }
 
@@ -159,13 +160,12 @@ bool FPoseDataContainer::InsertTrack(const FName& InTrackName, USkeleton* InSkel
 		if (SkeletonIndex != INDEX_NONE)
 		{
 			Tracks.Add(InTrackName);
-			TrackMap.Add(InTrackName, SkeletonIndex);
 			TrackIndex = Tracks.Num() - 1;
 
 			// now insert default refpose
 			const FTransform DefaultPose = GetDefaultTransform(SkeletonIndex, RefPose);
 
-			for (auto& PoseData : Poses)
+			for (FPoseData& PoseData : Poses)
 			{
 				ensureAlways(PoseData.SourceLocalSpacePose.Num() == TrackIndex);
 
@@ -190,11 +190,9 @@ bool FPoseDataContainer::FillUpSkeletonPose(FPoseData* PoseData, USkeleton* InSk
 	{
 		int32 TrackIndex = 0;
 		const TArray<FTransform>& RefPose = InSkeleton->GetRefLocalPoses();
-		for (TPair<FName, int32>& TrackIter : TrackMap)
+		for (const int32& SkeletonIndex : TrackBoneIndices)
 		{
-			int32 SkeletonIndex = TrackIter.Value;
 			PoseData->SourceLocalSpacePose[TrackIndex] = RefPose[SkeletonIndex];
-
 			++TrackIndex;
 		}
 
@@ -293,12 +291,15 @@ void FPoseDataContainer::RetrieveSourcePoseFromExistingPose(bool bAdditive, int3
 // this marks dirty tracks for each pose 
 void FPoseDataContainer::ConvertToFullPose(USkeleton* InSkeleton, const TArray<FTransform>& RefPose)
 {
+	TrackPoseInfluenceIndices.Reset();
+	TrackPoseInfluenceIndices.SetNum(Tracks.Num());
+	
 	// first create pose buffer that only has valid data
-	for (auto& Pose : Poses)
+	for (int32 PoseIndex = 0; PoseIndex < Poses.Num(); ++PoseIndex)
 	{
+		FPoseData& Pose = Poses[PoseIndex];
 		check(Pose.SourceLocalSpacePose.Num() == Tracks.Num());
 		Pose.LocalSpacePose.Reset();
-		Pose.TrackToBufferIndex.Reset();
 		if (InSkeleton)
 		{
 			for (int32 TrackIndex = 0; TrackIndex < Tracks.Num(); ++TrackIndex)
@@ -308,7 +309,7 @@ void FPoseDataContainer::ConvertToFullPose(USkeleton* InSkeleton, const TArray<F
 				if (!Pose.SourceLocalSpacePose[TrackIndex].Equals(DefaultTransform, KINDA_SMALL_NUMBER))
 				{
 					int32 NewIndex = Pose.LocalSpacePose.Add(Pose.SourceLocalSpacePose[TrackIndex]);
-					Pose.TrackToBufferIndex.Add(TrackIndex, NewIndex);
+					TrackPoseInfluenceIndices[TrackIndex].Influences.Emplace(FPoseAssetInfluence{PoseIndex, NewIndex});
 				}
 			}
 		}
@@ -323,12 +324,14 @@ void FPoseDataContainer::ConvertToAdditivePose(const TArray<FTransform>& InBaseP
 	check(InBaseCurve.Num() == Curves.Num());
 	const FTransform AdditiveIdentity(FQuat::Identity, FVector::ZeroVector, FVector::ZeroVector);
 
+	TrackPoseInfluenceIndices.Reset();
+	TrackPoseInfluenceIndices.SetNum(Tracks.Num());
+	
 	for (int32 PoseIndex = 0; PoseIndex < Poses.Num(); ++PoseIndex)
 	{
 		FPoseData& PoseData = Poses[PoseIndex];
 		// set up buffer
 		PoseData.LocalSpacePose.Reset();
-		PoseData.TrackToBufferIndex.Reset();
 		PoseData.CurveData.Reset(PoseData.SourceCurveData.Num());
 		PoseData.CurveData.AddUninitialized(PoseData.SourceCurveData.Num());
 
@@ -340,8 +343,8 @@ void FPoseDataContainer::ConvertToAdditivePose(const TArray<FTransform>& InBaseP
 			FAnimationRuntime::ConvertTransformToAdditive(NewTransform, InBasePose[BoneIndex]);
 			if (!NewTransform.Equals(AdditiveIdentity))
 			{
-				int32& NewValue = PoseData.TrackToBufferIndex.Add(BoneIndex);
-				NewValue = PoseData.LocalSpacePose.Add(NewTransform);
+				const int32 Index = PoseData.LocalSpacePose.Add(NewTransform);
+				TrackPoseInfluenceIndices[BoneIndex].Influences.Emplace(FPoseAssetInfluence{PoseIndex, Index});
 			}
 		}
 
@@ -378,6 +381,14 @@ struct FBoneIndices
 	{}
 };
 
+struct FPoseAssetEvalData : public TThreadSingleton<FPoseAssetEvalData>
+{
+	TArray<FBoneIndices> BoneIndices;
+	TArray<int32> PoseWeightedIndices;
+	TArray<float> PoseWeights;
+	TArray<bool> WeightedPoses;
+};
+
 void UPoseAsset::GetBaseAnimationPose(struct FCompactPose& OutPose, FBlendedCurve& OutCurve) const
 {
 	UE::Anim::FStackAttributeContainer TempAttributes;
@@ -398,16 +409,18 @@ void UPoseAsset::GetBaseAnimationPose(FAnimationPoseData& OutAnimationPoseData)
 		OutPose.ResetToRefPose();
 
 		// this contains compact bone pose list that this pose cares
-		TArray<FBoneIndices> BoneIndices;
-
-		const int32 TrackNum = PoseContainer.TrackMap.Num();
+		FPoseAssetEvalData& EvalData = FPoseAssetEvalData::Get();
+		TArray<FBoneIndices>& BoneIndices = EvalData.BoneIndices;
+        const int32 TrackNum = PoseContainer.Tracks.Num();
+		BoneIndices.SetNumUninitialized(TrackNum, false);
 
-		for (const TPair<FName, int32>& TrackPair : PoseContainer.TrackMap)
+		for(int32 TrackIndex = 0; TrackIndex < TrackNum; ++TrackIndex)
 		{
-			const int32 SkeletonBoneIndex = TrackPair.Value;
+			const int32 SkeletonBoneIndex = PoseContainer.TrackBoneIndices[TrackIndex];
 			const FCompactPoseBoneIndex PoseBoneIndex = RequiredBones.GetCompactPoseIndexFromSkeletonIndex(SkeletonBoneIndex);
 			// we add even if it's invalid because we want it to match with track index
-			BoneIndices.Add(FBoneIndices(SkeletonBoneIndex, PoseBoneIndex));
+			BoneIndices[TrackIndex].SkeletonBoneIndex = SkeletonBoneIndex;
+			BoneIndices[TrackIndex].CompactBoneIndex = PoseBoneIndex;
 		}
 
 		const TArray<FTransform>& PoseTransform = PoseContainer.Poses[BasePoseIndex].LocalSpacePose;
@@ -436,20 +449,21 @@ void UPoseAsset::GetBaseAnimationPose(FAnimationPoseData& OutAnimationPoseData)
  * The difference between BlendFromIdentityAndAccumulcate is scale
  * This ADDS scales to the FinalAtom. We use additive identity as final atom, so can't use
  */
-FORCEINLINE void BlendFromIdentityAndAccumulateAdditively(FTransform& FinalAtom, FTransform& SourceAtom, float BlendWeight)
+FORCEINLINE void BlendFromIdentityAndAccumulateAdditively_Custom(FTransform& FinalAtom, const FTransform& SourceAtom, float BlendWeight)
 {
 	const  FTransform AdditiveIdentity(FQuat::Identity, FVector::ZeroVector, FVector::ZeroVector);
 
+	FTransform Delta = SourceAtom;
 	// Scale delta by weight
 	if (BlendWeight < (1.f - ZERO_ANIMWEIGHT_THRESH))
 	{
-		SourceAtom.Blend(AdditiveIdentity, SourceAtom, BlendWeight);
+		Delta.Blend(AdditiveIdentity, Delta, BlendWeight);
 	}
 
-	FinalAtom.SetRotation(SourceAtom.GetRotation() * FinalAtom.GetRotation());
-	FinalAtom.SetTranslation(FinalAtom.GetTranslation() + SourceAtom.GetTranslation());
+	FinalAtom.SetRotation(Delta.GetRotation() * FinalAtom.GetRotation());
+	FinalAtom.SetTranslation(FinalAtom.GetTranslation() + Delta.GetTranslation());
 	// this ADDS scale
-	FinalAtom.SetScale3D(FinalAtom.GetScale3D() + SourceAtom.GetScale3D());
+	FinalAtom.SetScale3D(FinalAtom.GetScale3D() + Delta.GetScale3D());
 
 	FinalAtom.DiagnosticCheckNaN_All();
 
@@ -464,8 +478,16 @@ void UPoseAsset::GetAnimationCurveOnly(TArray<FName>& InCurveNames, TArray<float
 		USkeleton* MySkeleton = GetSkeleton();
 		check(MySkeleton);
 
+		FPoseAssetEvalData& EvalData = FPoseAssetEvalData::Get();
+		const int32 NumPoses = PoseContainer.Poses.Num();
+		TArray<float>& PoseWeights = EvalData.PoseWeights;		
+		PoseWeights.Reset();
+		PoseWeights.SetNumZeroed(NumPoses, false);
+
+		TArray<int32>& WeightedPoseIndices = EvalData.PoseWeightedIndices;	
+		WeightedPoseIndices.Reset();
+
 		bool bNormalizeWeight = bAdditivePose == false;
-		TMap<const FPoseData*, float> IndexToWeightMap;
 		float TotalWeight = 0.f;
 		// we iterate through to see if we have that corresponding pose
 		for (int32 CurveIndex = 0; CurveIndex < InCurveNames.Num(); ++CurveIndex)
@@ -483,25 +505,17 @@ void UPoseAsset::GetAnimationCurveOnly(TArray<FName>& InCurveNames, TArray<float
 					// and has weight
 					if ((!bAdditivePose || PoseIndex != BasePoseIndex) && FAnimationRuntime::HasWeight(Value))
 					{
-						IndexToWeightMap.Add(&PoseData, Value);
-						TotalWeight += Value;
+						TotalWeight += Value;												
+						PoseWeights[PoseIndex] = Value;
+						WeightedPoseIndices.Add(PoseIndex);
 					}
 				}
 			}
 		}
 
-		const int32 TotalNumberOfValidPoses = IndexToWeightMap.Num();
+		const int32 TotalNumberOfValidPoses = WeightedPoseIndices.Num();
 		if (TotalNumberOfValidPoses > 0)
 		{
-			//if full pose, we'll have to normalize by weight
-			if (bNormalizeWeight && TotalWeight > 1.f)
-			{
-				for (TPair<const FPoseData*, float>& WeightPair : IndexToWeightMap)
-				{
-					WeightPair.Value /= TotalWeight;
-				}
-			}
-
 			// collect curves
 			TArray<uint16> CurveUIDList;
 			CurveUIDList.AddUninitialized(PoseContainer.Curves.Num());
@@ -513,12 +527,28 @@ void UPoseAsset::GetAnimationCurveOnly(TArray<FName>& InCurveNames, TArray<float
 			// blend curves
 			FBlendedCurve BlendedCurve;
 			BlendedCurve.InitFrom(&CurveUIDList);
-			for (const TPair<const FPoseData*, float>& ActivePosePair : IndexToWeightMap)
+
+			//if full pose, we'll have to normalize by weight
+			if (bNormalizeWeight && TotalWeight > 1.f)
+			{
+				for (const int32& WeightedPoseIndex : WeightedPoseIndices)
+				{
+					float& PoseWeight = PoseWeights[WeightedPoseIndex];
+					PoseWeight /= TotalWeight;
+
+					const FPoseData& Pose = PoseContainer.Poses[WeightedPoseIndex];
+					PoseContainer.BlendPoseCurve(&Pose, BlendedCurve, PoseWeight);
+				}
+			}
+			else
 			{
-				const FPoseData* Pose = ActivePosePair.Key;
-				const float Weight = ActivePosePair.Value;
+				for (const int32& WeightedPoseIndex : WeightedPoseIndices)
+				{				
+					const FPoseData& Pose = PoseContainer.Poses[WeightedPoseIndex];
+					const float& PoseWeight = PoseWeights[WeightedPoseIndex];
 
-				PoseContainer.BlendPoseCurve(Pose, BlendedCurve, Weight);
+					PoseContainer.BlendPoseCurve(&Pose, BlendedCurve, PoseWeight);
+				}
 			}
 
 			OutCurveNames.Reset();
@@ -560,18 +590,21 @@ bool UPoseAsset::GetAnimationPose(struct FAnimationPoseData& OutAnimationPoseDat
 		const FBoneContainer& RequiredBones = OutPose.GetBoneContainer();
 		USkeleton* MySkeleton = GetSkeleton();
 
+		// TLS storage for working data
+		FPoseAssetEvalData& EvalData = FPoseAssetEvalData::Get();
 
 		// this contains compact bone pose list that this pose cares
-		TArray<FBoneIndices> BoneIndices;
-
-		const int32 TrackNum = PoseContainer.TrackMap.Num();
+		TArray<FBoneIndices>& BoneIndices = EvalData.BoneIndices;
+		const int32 TrackNum = PoseContainer.Tracks.Num();
+		BoneIndices.SetNumUninitialized(TrackNum, false);
 
-		for (const TPair<FName, int32>& TrackPair : PoseContainer.TrackMap)
+		for(int32 TrackIndex = 0; TrackIndex < TrackNum; ++TrackIndex)
 		{
-			const int32 SkeletonBoneIndex = TrackPair.Value;
+			const int32 SkeletonBoneIndex = PoseContainer.TrackBoneIndices[TrackIndex];
 			const FCompactPoseBoneIndex PoseBoneIndex = RequiredBones.GetCompactPoseIndexFromSkeletonIndex(SkeletonBoneIndex);
 			// we add even if it's invalid because we want it to match with track index
-			BoneIndices.Add(FBoneIndices(SkeletonBoneIndex, PoseBoneIndex));
+			BoneIndices[TrackIndex].SkeletonBoneIndex = SkeletonBoneIndex;
+			BoneIndices[TrackIndex].CompactBoneIndex = PoseBoneIndex;
 		}
 
 		// you could only have morphtargets
@@ -587,12 +620,24 @@ bool UPoseAsset::GetAnimationPose(struct FAnimationPoseData& OutAnimationPoseDat
 		{
 			OutPose.ResetToRefPose();
 		}
+		
+		const bool bNormalizeWeight = bAdditivePose == false;
+		const int32 NumPoses = PoseContainer.Poses.Num();
+		TArray<float>& PoseWeights = EvalData.PoseWeights;
+		PoseWeights.Reset();
+		PoseWeights.SetNumZeroed(NumPoses, false);
+
+		TArray<int32>& WeightedPoseIndices = EvalData.PoseWeightedIndices;
+		WeightedPoseIndices.Reset();
+
+		TArray<bool>& WeightedPoses = EvalData.WeightedPoses;
+		WeightedPoses.SetNumZeroed(NumPoses, false);
 
-		bool bNormalizeWeight = bAdditivePose == false;
-		TMap<const FPoseData*, float> IndexToWeightMap;
 		float TotalWeight = 0.f;
 		// we iterate through to see if we have that corresponding pose
-		for (int32 CurveIndex = 0; CurveIndex < ExtractionContext.PoseCurves.Num(); ++CurveIndex)
+
+		const int32 NumPoseCurves = ExtractionContext.PoseCurves.Num();
+		for (int32 CurveIndex = 0; CurveIndex < NumPoseCurves; ++CurveIndex)
 		{
 			const FPoseCurve& Curve = ExtractionContext.PoseCurves[CurveIndex];
 			const int32 PoseIndex = Curve.PoseIndex; 
@@ -605,118 +650,121 @@ bool UPoseAsset::GetAnimationPose(struct FAnimationPoseData& OutAnimationPoseDat
 				// and has weight
 				if ((!bAdditivePose || PoseIndex != BasePoseIndex) && FAnimationRuntime::HasWeight(Value))
 				{
-					IndexToWeightMap.Add(&PoseData, Value);
 					TotalWeight += Value;
+
+					// Set pose weight and bit, and add weighted pose index
+					PoseWeights[PoseIndex] = Value;
+					WeightedPoseIndices.Add(PoseIndex);
+					WeightedPoses[PoseIndex] = true;
 				}
 			}
 		}
 
-		const int32 TotalNumberOfValidPoses = IndexToWeightMap.Num();
+		const int32 TotalNumberOfValidPoses = WeightedPoseIndices.Num();
 		if (TotalNumberOfValidPoses > 0)
 		{
-			TArray<FTransform> BlendedBoneTransform;
-
 			//if full pose, we'll have to normalize by weight
 			if (bNormalizeWeight && TotalWeight > 1.f)
 			{
-				for (TPair<const FPoseData*, float>& WeightPair : IndexToWeightMap)
+				for (const int32& WeightedPoseIndex : WeightedPoseIndices)
 				{
-					WeightPair.Value /= TotalWeight;
+					float& PoseWeight = PoseWeights[WeightedPoseIndex];
+					PoseWeight /= TotalWeight;
+
+					// Do curve blending inline as we are looping over weights anyway
+					const FPoseData& Pose = PoseContainer.Poses[WeightedPoseIndex];
+					PoseContainer.BlendPoseCurve(&Pose, OutCurve, PoseWeight);
 				}
 			}
+            else
+            {
+ 				// Take the matching curve weights from the selected poses, and blend them using the
+				// weighting that we need from each pose. This is much faster than grabbing each
+				// blend curve and blending them in their entirety, especially when there are very
+				// few active curves for each pose and many curves for the entire pose asset.
+				for (const int32& WeightedPoseIndex : WeightedPoseIndices)
+				{
+					const FPoseData& Pose = PoseContainer.Poses[WeightedPoseIndex];
+					const float& Weight = PoseWeights[WeightedPoseIndex];
+					PoseContainer.BlendPoseCurve(&Pose, OutCurve, Weight);
+				}
+            }
 
-			BlendedBoneTransform.AddUninitialized(TrackNum);
+			// Final per-track (bone) transform
+			FTransform OutBoneTransform;
 			for (int32 TrackIndex = 0; TrackIndex < TrackNum; ++TrackIndex)
 			{
-				// If invalid compact bone index, BlendedBoneTransform[TrackIndex] won't be used (see 'blend curves' below), so don't bother filling it in
-				const FCompactPoseBoneIndex CompactIndex = BoneIndices[TrackIndex].CompactBoneIndex;
-				if (CompactIndex != INDEX_NONE)
+				const FCompactPoseBoneIndex& CompactIndex = BoneIndices[TrackIndex].CompactBoneIndex;
+
+				// If bone index is invalid, or not required for the pose - skip
+				if (CompactIndex == INDEX_NONE || !ExtractionContext.IsBoneRequired(CompactIndex.GetInt()))
 				{
-					if (!ExtractionContext.IsBoneRequired(CompactIndex.GetInt()))
-					{
-						continue;
-					}
+					continue;
+				}
+				
+				const TArray<FPoseAssetInfluence>& PoseInfluences = PoseContainer.TrackPoseInfluenceIndices[TrackIndex].Influences;
+
+				// When additive, or for any bone that has no pose influences. Set transform to input.
+				if (bAdditivePose || PoseInfluences.Num() == 0)
+				{
+					OutBoneTransform = OutPose[CompactIndex];
+				}
 
-					TArray<FTransform> BlendingTransform;
-					TArray<float> BlendingWeights;
+				const int32 NumInfluences = PoseInfluences.Num();
+				if (NumInfluences)
+				{
 					float TotalLocalWeight = 0.f;
-					for (const TPair<const FPoseData*, float>& ActivePosePair : IndexToWeightMap)
+					// Only loop over poses known to influence this track its final transform
+					for (int32 Index = 0; Index < NumInfluences; ++Index)
 					{
-						const FPoseData* Pose = ActivePosePair.Key;
-						const float Weight = ActivePosePair.Value;
+						const FPoseAssetInfluence& Influence = PoseInfluences[Index];
+						const int32& PoseIndex = Influence.PoseIndex;
+						const int32& BonePoseIndex = Influence.BoneTransformIndex;
 						
-						// find buffer index from track index
-						const int32* BufferIndex = Pose->TrackToBufferIndex.Find(TrackIndex);
-						if (BufferIndex)
+						// Only processs pose if its weighted
+						if(WeightedPoses[PoseIndex])
 						{
-							BlendingTransform.Add(Pose->LocalSpacePose[*BufferIndex]);
-							BlendingWeights.Add(Weight);
+							const float& Weight = PoseWeights[PoseIndex];
 							TotalLocalWeight += Weight;
+						
+							const FPoseData& PoseData = PoseContainer.Poses[PoseIndex];
+							const FTransform& BonePose = PoseData.LocalSpacePose[BonePoseIndex];
+
+							// Set weighted value For first pose, if applicable
+							if(Index == 0 && !bAdditivePose)
+							{								
+								OutBoneTransform = BonePose * ScalarRegister(Weight);
+							}
+							else
+							{
+								if (bAdditivePose)
+								{
+									BlendFromIdentityAndAccumulateAdditively_Custom(OutBoneTransform, BonePose, Weight);
+								}
+								else
+								{
+									OutBoneTransform.AccumulateWithShortestRotation(BonePose, ScalarRegister(Weight));
+								}
+							}
 						}
 					}
 
-					const int32 StartBlendLoopIndex = (bAdditivePose || TotalLocalWeight < 1.f) ? 0 : 1;
-
-					if (BlendingTransform.Num() == 0)
-					{
-						// copy from out default pose
-						BlendedBoneTransform[TrackIndex] = OutPose[CompactIndex];
-					}
-					else
+					// In case no influencing poses had any weight, set transform to input
+					if(!FAnimWeight::IsRelevant(TotalLocalWeight))
 					{
-						if (bAdditivePose)
-						{
-							BlendedBoneTransform[TrackIndex] = OutPose[CompactIndex];
-						}
-						else if (StartBlendLoopIndex == 0)
-						{
-							BlendedBoneTransform[TrackIndex] = OutPose[CompactIndex] * ScalarRegister(1.f - TotalLocalWeight);
-						}
-						else
-						{
-							BlendedBoneTransform[TrackIndex] = BlendingTransform[0] * ScalarRegister(BlendingWeights[0]);
-						}
+						OutBoneTransform = OutPose[CompactIndex];
 					}
-
-					for (int32 BlendIndex = StartBlendLoopIndex; BlendIndex < BlendingTransform.Num(); ++BlendIndex)
+					else if (!FAnimWeight::IsFullWeight(TotalLocalWeight) && !bAdditivePose)
 					{
-						if (bAdditivePose)
-						{
-							BlendFromIdentityAndAccumulateAdditively(BlendedBoneTransform[TrackIndex], BlendingTransform[BlendIndex], BlendingWeights[BlendIndex]);
-						}
-						else
-						{
-							BlendedBoneTransform[TrackIndex].AccumulateWithShortestRotation(BlendingTransform[BlendIndex], ScalarRegister(BlendingWeights[BlendIndex]));
-						}
+						OutBoneTransform.AccumulateWithShortestRotation(OutPose[CompactIndex], ScalarRegister(1.f - TotalLocalWeight));
 					}
 				}
-			}
-
-			// Take the matching curve weights from the selected poses, and blend them using the
-			// weighting that we need from each pose. This is much faster than grabbing each
-			// blend curve and blending them in their entirety, especially when there are very
-			// few active curves for each pose and many curves for the entire pose asset.
-			for (const TPair<const FPoseData*, float>& ActivePosePair : IndexToWeightMap)
-			{
-				const FPoseData* Pose = ActivePosePair.Key;
-				const float Weight = ActivePosePair.Value;
 
-				PoseContainer.BlendPoseCurve(Pose, OutCurve, Weight);
-			}
+				// Retarget the blended transform, and copy to output pose
+				FAnimationRuntime::RetargetBoneTransform(MySkeleton, GetRetargetTransformsSourceName(), GetRetargetTransforms(), OutBoneTransform,  BoneIndices[TrackIndex].SkeletonBoneIndex, CompactIndex, RequiredBones, bAdditivePose);
 
-			for (int32 TrackIndex = 0; TrackIndex < TrackNum; ++TrackIndex)
-			{
-				const FBoneIndices& LocalBoneIndices = BoneIndices[TrackIndex];
-				if (LocalBoneIndices.CompactBoneIndex != INDEX_NONE)
-				{
-					if (!ExtractionContext.IsBoneRequired(LocalBoneIndices.CompactBoneIndex.GetInt()))
-					{
-						continue;
-					}
-					FAnimationRuntime::RetargetBoneTransform(MySkeleton, GetRetargetTransformsSourceName(), GetRetargetTransforms(), BlendedBoneTransform[TrackIndex], LocalBoneIndices.SkeletonBoneIndex, LocalBoneIndices.CompactBoneIndex, RequiredBones, bAdditivePose);
-					OutPose[LocalBoneIndices.CompactBoneIndex] = BlendedBoneTransform[TrackIndex];
-					OutPose[LocalBoneIndices.CompactBoneIndex].NormalizeRotation();
-				}
+				OutPose[CompactIndex] = OutBoneTransform;
+				OutPose[CompactIndex].NormalizeRotation();
 			}
 
 			return true;
@@ -742,7 +790,7 @@ void UPoseAsset::PostLoad()
 	{
 		// fix curve names
 		// copy to source local data FIRST
-		for (auto& Pose : PoseContainer.Poses)
+		for (FPoseData& Pose : PoseContainer.Poses)
 		{
 			Pose.SourceCurveData = Pose.CurveData;
 			Pose.SourceLocalSpacePose = Pose.LocalSpacePose;
@@ -777,6 +825,11 @@ void UPoseAsset::PostLoad()
 		// fix curve names
 		PostProcessData();
 	}
+
+  	if(GetLinkerCustomVersion(FUE5ReleaseStreamObjectVersion::GUID) < FUE5ReleaseStreamObjectVersion::PoseAssetRuntimeRefactor)
+    {
+		PostProcessData();
+    }
 #endif // WITH_EDITOR
 
 	// fix curve names
@@ -785,7 +838,7 @@ void UPoseAsset::PostLoad()
 	{
 		MySkeleton->VerifySmartNames(USkeleton::AnimCurveMappingName, PoseContainer.PoseNames);
 
-		for (auto& Curve : PoseContainer.Curves)
+		for (FAnimCurveBase& Curve : PoseContainer.Curves)
 		{
 			MySkeleton->VerifySmartName(USkeleton::AnimCurveMappingName, Curve.Name);
 		}
@@ -796,7 +849,7 @@ void UPoseAsset::PostLoad()
 			|| MySkeleton->GetLinkerCustomVersion(FFrameworkObjectVersion::GUID) < FFrameworkObjectVersion::MoveCurveTypesToSkeleton)
 		{
 			// fix up curve flags to skeleton
-			for (auto& Curve : PoseContainer.Curves)
+			for (FAnimCurveBase& Curve : PoseContainer.Curves)
 			{
 				bool bMorphtargetSet = Curve.GetCurveTypeFlag(AACF_DriveMorphTarget_DEPRECATED);
 				bool bMaterialSet = Curve.GetCurveTypeFlag(AACF_DriveMaterial_DEPRECATED);
@@ -810,8 +863,17 @@ void UPoseAsset::PostLoad()
 		}
 	}
 
-	// I  have to fix pose names
-	RecacheTrackmap();
+	// I have to fix pose names
+#if WITH_EDITOR
+	if(RemoveInvalidTracks())
+	{
+		PostProcessData();
+	}
+	else
+#endif // WITH_EDITOR
+	{
+		UpdateTrackBoneIndices();
+	}
 }
 
 void UPoseAsset::Serialize(FArchive& Ar)
@@ -819,6 +881,7 @@ void UPoseAsset::Serialize(FArchive& Ar)
  	Ar.UsingCustomVersion(FFrameworkObjectVersion::GUID);
 	Ar.UsingCustomVersion(FAnimPhysObjectVersion::GUID);
 	Ar.UsingCustomVersion(FFortniteMainBranchObjectVersion::GUID);
+	Ar.UsingCustomVersion(FUE5ReleaseStreamObjectVersion::GUID);
 
 	Super::Serialize(Ar);
 }
@@ -951,7 +1014,7 @@ const int32 UPoseAsset::GetTrackIndexByName(const FName& InTrackName) const
 
 bool UPoseAsset::ContainsPose(const FName& InPoseName) const
 {
-	for (const auto& PoseName : PoseContainer.PoseNames)
+	for (const FSmartName& PoseName : PoseContainer.PoseNames)
 	{
 		if (PoseName.DisplayName == InPoseName)
 		{
@@ -966,6 +1029,8 @@ bool UPoseAsset::ContainsPose(const FName& InPoseName) const
 // whenever you change SourceLocalPoses, or SourceCurves, we should call this to update runtime dataa
 void UPoseAsset::PostProcessData()
 {
+	RemoveInvalidTracks();
+	
 	// convert back to additive if it was that way
 	if (bAdditivePose)
 	{
@@ -976,7 +1041,7 @@ void UPoseAsset::PostProcessData()
 		ConvertToFullPose();
 	}
 
-	RecacheTrackmap();
+	UpdateTrackBoneIndices();
 }
 
 bool UPoseAsset::AddOrUpdatePoseWithUniqueName(USkeletalMeshComponent* MeshComponent, FSmartName* OutPoseName /*= nullptr*/)
@@ -1125,7 +1190,7 @@ void UPoseAsset::CombineTracks(const TArray<FName>& NewTracks)
 	USkeleton* MySkeleton = GetSkeleton();
 	if (MySkeleton)
 	{
-		for (const auto& NewTrack : NewTracks)
+		for (const FName& NewTrack : NewTracks)
 		{
 			if (PoseContainer.Tracks.Contains(NewTrack) == false)
 			{
@@ -1133,6 +1198,7 @@ void UPoseAsset::CombineTracks(const TArray<FName>& NewTracks)
 				// right now it doesn't have to be in the hierarchy
 				// @todo: it is probably best to keep the hierarchy of the skeleton, so in the future, we might like to sort this by track after
 				PoseContainer.InsertTrack(NewTrack, MySkeleton, GetRetargetTransforms());
+				UpdateTrackBoneIndices();
 			}
 		}
 	}
@@ -1342,7 +1408,7 @@ int32 UPoseAsset::DeletePoses(TArray<FName> PoseNamesToDelete)
 
 	USkeleton* MySkeleton = GetSkeleton();
 
-	for (const auto& PoseName : PoseNamesToDelete)
+	for (const FName& PoseName : PoseNamesToDelete)
 	{
 		FSmartName PoseSmartName;
 		if (MySkeleton->GetSmartNameByName(USkeleton::AnimCurveMappingName, PoseName, PoseSmartName))
@@ -1377,7 +1443,7 @@ int32 UPoseAsset::DeleteCurves(TArray<FName> CurveNamesToDelete)
 
 	USkeleton* MySkeleton = GetSkeleton();
 
-	for (const auto& CurveName : CurveNamesToDelete)
+	for (const FName& CurveName : CurveNamesToDelete)
 	{
 		FSmartName CurveSmartName;
 		if (MySkeleton->GetSmartNameByName(USkeleton::AnimCurveMappingName, CurveName, CurveSmartName))
@@ -1464,10 +1530,27 @@ const int32 UPoseAsset::GetCurveIndexByName(const FName& InCurveName) const
 }
 
 
-void UPoseAsset::RecacheTrackmap()
+void UPoseAsset::UpdateTrackBoneIndices()
 {
 	USkeleton* MySkeleton = GetSkeleton();
-	PoseContainer.TrackMap.Reset();
+	PoseContainer.TrackBoneIndices.Reset();
+	if (MySkeleton)
+	{
+		const FReferenceSkeleton& RefSkeleton = MySkeleton->GetReferenceSkeleton();
+
+		PoseContainer.TrackBoneIndices.SetNumZeroed(PoseContainer.Tracks.Num());
+		for (int32 TrackIndex = 0; TrackIndex < PoseContainer.Tracks.Num(); ++TrackIndex)
+		{
+			const FName& TrackName = PoseContainer.Tracks[TrackIndex];
+			PoseContainer.TrackBoneIndices[TrackIndex] = RefSkeleton.FindBoneIndex(TrackName);
+		}
+	}
+}
+
+bool UPoseAsset::RemoveInvalidTracks()
+{
+	const USkeleton* MySkeleton = GetSkeleton();
+	const int32 InitialNumTracks = PoseContainer.Tracks.Num();
 
 	if (MySkeleton)
 	{
@@ -1478,11 +1561,7 @@ void UPoseAsset::RecacheTrackmap()
 		{
 			const FName& TrackName = PoseContainer.Tracks[TrackIndex];
 			const int32 SkeletonTrackIndex = RefSkeleton.FindBoneIndex(TrackName);
-			if (SkeletonTrackIndex != INDEX_NONE)
-			{
-				PoseContainer.TrackMap.Add(TrackName, SkeletonTrackIndex);
-			}
-			else
+			if (SkeletonTrackIndex == INDEX_NONE)
 			{
 				// delete this track. It's missing now
 				PoseContainer.DeleteTrack(TrackIndex);
@@ -1490,6 +1569,8 @@ void UPoseAsset::RecacheTrackmap()
 			}
 		}
 	}
+
+	return InitialNumTracks != PoseContainer.Tracks.Num();
 }
 
 #if WITH_EDITORONLY_DATA
@@ -1543,21 +1624,9 @@ FName UPoseAsset::GetRetargetTransformsSourceName() const
 
 void FPoseDataContainer::DeleteTrack(int32 TrackIndex)
 {
-	if (TrackMap.Contains(Tracks[TrackIndex]))
-	{
-		TrackMap.Remove(Tracks[TrackIndex]);
-	}
-
 	Tracks.RemoveAt(TrackIndex);
-	for (auto& Pose : Poses)
+	for (FPoseData& Pose : Poses)
 	{
-		int32* BufferIndex = Pose.TrackToBufferIndex.Find(TrackIndex);
-		if (BufferIndex)
-		{
-			Pose.LocalSpacePose.RemoveAt(*BufferIndex);
-			Pose.TrackToBufferIndex.Remove(TrackIndex);
-		}
-
 #if WITH_EDITOR
 		// if not editor, they can't save this data, so it will run again when editor runs
 		Pose.SourceLocalSpacePose.RemoveAt(TrackIndex);
@@ -1574,37 +1643,12 @@ void UPoseAsset::RemapTracksToNewSkeleton(USkeleton* NewSkeleton, bool bConvertS
 	{
 		NewSkeleton->VerifySmartNames(USkeleton::AnimCurveMappingName, PoseContainer.PoseNames);
 
-		for (auto& Curve : PoseContainer.Curves)
+		for (FAnimCurveBase& Curve : PoseContainer.Curves)
 		{
 			NewSkeleton->VerifySmartName(USkeleton::AnimCurveMappingName, Curve.Name);
 		}
 	}
 
-	
-	USkeleton* MySkeleton = GetSkeleton();
-	PoseContainer.TrackMap.Reset();
-
-	if (MySkeleton)
-	{
-		const FReferenceSkeleton& RefSkeleton = MySkeleton->GetReferenceSkeleton();
-
-		// set up track data 
-		for (int32 TrackIndex = 0; TrackIndex < PoseContainer.Tracks.Num(); ++TrackIndex)
-		{
-			const FName& TrackName = PoseContainer.Tracks[TrackIndex];
-			const int32 SkeletonTrackIndex = RefSkeleton.FindBoneIndex(TrackName);
-			if (SkeletonTrackIndex != INDEX_NONE)
-			{
-				PoseContainer.TrackMap.Add(TrackName, SkeletonTrackIndex);
-			}
-			else
-			{
-				// delete this track. It's missing now
-				PoseContainer.DeleteTrack(TrackIndex);
-				--TrackIndex;
-			}
-		}
-	}
 	PostProcessData();
 }
 
