Fix NonUnity Compile UnrealEditor Mac.

diff --git a/Engine/Source/Runtime/Engine/Public/WorldPartition/WorldPartitionActorDescView.h b/Engine/Source/Runtime/Engine/Public/WorldPartition/WorldPartitionActorDescView.h
--- a/Engine/Source/Runtime/Engine/Public/WorldPartition/WorldPartitionActorDescView.h
+++ b/Engine/Source/Runtime/Engine/Public/WorldPartition/WorldPartitionActorDescView.h
@@ -23,9 +23,9 @@ public:
 	FWorldPartitionActorDescView(const FWorldPartitionActorDesc* InActorDesc);
 
 	const FGuid& GetGuid() const;
-	inline FName GetClass() const;
-	inline UClass* GetActorClass() const;
-	inline FVector GetOrigin() const;
+	FName GetClass() const;
+	UClass* GetActorClass() const;
+	FVector GetOrigin() const;
 	EActorGridPlacement GetGridPlacement() const;
 	FName GetRuntimeGrid() const;
 	bool GetActorIsEditorOnly() const;
