Add function to gather referenced material parameter collections for material instances. This fixes crashes when failing to find material parameter caches by GUID.

- Adds AppendReferencedParameterCollectionIdsTo function to MaterialInstances similar to UMaterial.
- When we create a MaterialID be sure to call this function for MaterialInstances to gather referenced material parameter collections.

diff --git a/Engine/Source/Runtime/Engine/Classes/Materials/MaterialInstance.h b/Engine/Source/Runtime/Engine/Classes/Materials/MaterialInstance.h
--- a/Engine/Source/Runtime/Engine/Classes/Materials/MaterialInstance.h
+++ b/Engine/Source/Runtime/Engine/Classes/Materials/MaterialInstance.h
@@ -715,6 +715,7 @@ public:
 #if WITH_EDITOR
 	void BeginAllowCachingStaticParameterValues();
 	void EndAllowCachingStaticParameterValues();
+	void AppendReferencedParameterCollectionIdsTo(TArray<FGuid>& OutIds) const;
 #endif // WITH_EDITOR
 
 protected:
diff --git a/Engine/Source/Runtime/Engine/Private/Materials/Material.cpp b/Engine/Source/Runtime/Engine/Private/Materials/Material.cpp
--- a/Engine/Source/Runtime/Engine/Private/Materials/Material.cpp
+++ b/Engine/Source/Runtime/Engine/Private/Materials/Material.cpp
@@ -344,6 +344,7 @@ void FMaterialResource::GetShaderMapId(EShaderPlatform Platform, const ITargetPl
 	if(MaterialInstance)
 	{
 		MaterialInstance->GetBasePropertyOverridesHash(OutId.BasePropertyOverridesHash);
+		MaterialInstance->AppendReferencedParameterCollectionIdsTo(OutId.ReferencedParameterCollections);
 
 		FStaticParameterSet CompositedStaticParameters;
 		MaterialInstance->GetStaticParameterValues(CompositedStaticParameters);
diff --git a/Engine/Source/Runtime/Engine/Private/Materials/MaterialInstance.cpp b/Engine/Source/Runtime/Engine/Private/Materials/MaterialInstance.cpp
--- a/Engine/Source/Runtime/Engine/Private/Materials/MaterialInstance.cpp
+++ b/Engine/Source/Runtime/Engine/Private/Materials/MaterialInstance.cpp
@@ -26,6 +26,8 @@
 #include "Materials/MaterialInstanceConstant.h"
 #include "Materials/MaterialUniformExpressions.h"
 #include "Materials/MaterialInstanceSupport.h"
+#include "Materials/MaterialExpressionCollectionParameter.h"
+#include "Materials/MaterialParameterCollection.h"
 #include "Engine/SubsurfaceProfile.h"
 #include "ProfilingDebugging/LoadTimeTracker.h"
 #include "ProfilingDebugging/CookStats.h"
@@ -5070,6 +5072,96 @@ void UMaterialInstance::CopyMaterialUniformParametersInternal(UMaterialInterface
 #endif
 }
 
+#if WITH_EDITOR
+
+void FindCollectionExpressionRecursive(TArray<FGuid>& OutGuidList, const TArray<UMaterialExpression*>& InMaterialExpression)
+{
+	for (int32 ExpressionIndex = 0; ExpressionIndex < InMaterialExpression.Num(); ExpressionIndex++)
+	{
+		const UMaterialExpression* ExpressionPtr = InMaterialExpression[ExpressionIndex];
+		const UMaterialExpressionCollectionParameter* CollectionPtr = Cast<UMaterialExpressionCollectionParameter>(ExpressionPtr);
+		const UMaterialExpressionMaterialFunctionCall* MaterialFunctionCall = Cast<UMaterialExpressionMaterialFunctionCall>(ExpressionPtr);
+		const UMaterialExpressionMaterialAttributeLayers* MaterialLayers = Cast<UMaterialExpressionMaterialAttributeLayers>(ExpressionPtr);
+
+		if (CollectionPtr)
+		{
+			if (CollectionPtr->Collection)
+			{
+				OutGuidList.Add(CollectionPtr->Collection->StateId);
+			}
+			return;
+		}
+		else if (MaterialFunctionCall && MaterialFunctionCall->MaterialFunction)
+		{
+			if (const TArray<UMaterialExpression*>* FunctionExpressions = MaterialFunctionCall->MaterialFunction->GetFunctionExpressions())
+			{
+				FindCollectionExpressionRecursive(OutGuidList, *FunctionExpressions);
+			}
+		}
+		else if (MaterialLayers)
+		{
+			const TArray<UMaterialFunctionInterface*>& Layers = MaterialLayers->GetLayers();
+			const TArray<UMaterialFunctionInterface*>& Blends = MaterialLayers->GetBlends();
+
+			for (const UMaterialFunctionInterface* Layer : Layers)
+			{
+				if (Layer)
+				{
+					if (const TArray<UMaterialExpression*>* FunctionExpressions = Layer->GetFunctionExpressions())
+					{
+						FindCollectionExpressionRecursive(OutGuidList, *FunctionExpressions);
+					}
+				}
+			}
+
+			for (const UMaterialFunctionInterface* Blend : Blends)
+			{
+				if (Blend)
+				{
+					if (const TArray<UMaterialExpression*>* FunctionExpressions = Blend->GetFunctionExpressions())
+					{
+						FindCollectionExpressionRecursive(OutGuidList, *FunctionExpressions);
+					}
+				}
+			}
+		}
+	}
+}
+
+void UMaterialInstance::AppendReferencedParameterCollectionIdsTo(TArray<FGuid>& OutIds) const
+{
+	const UMaterial* Material = GetMaterial();
+	if (!Material)
+	{
+		return;
+	}
+
+	for (const FStaticMaterialLayersParameter& LayerParameter : StaticParameters.MaterialLayersParameters)
+	{
+		for (const TObjectPtr<UMaterialFunctionInterface> Layer : LayerParameter.Value.Layers)
+		{
+			if (Layer)
+			{
+				if (const TArray<UMaterialExpression*>* FunctionExpressions = Layer->GetFunctionExpressions())
+				{
+					FindCollectionExpressionRecursive(OutIds, *FunctionExpressions);
+				}
+			}
+		}
+		for (const TObjectPtr<UMaterialFunctionInterface> Blend : LayerParameter.Value.Blends)
+		{
+			if (Blend)
+			{
+				if (const TArray<UMaterialExpression*>* FunctionExpressions = Blend->GetFunctionExpressions())
+				{
+					FindCollectionExpressionRecursive(OutIds, *FunctionExpressions);
+				}
+			}
+		}
+	}
+}
+#endif // WITH_EDITOR
+
 #if WITH_EDITORONLY_DATA
 UMaterialInstance::FCustomStaticParametersGetterDelegate UMaterialInstance::CustomStaticParametersGetters;
 TArray<UMaterialInstance::FCustomParameterSetUpdaterDelegate> UMaterialInstance::CustomParameterSetUpdaters;
