Camera Shake - Fixed an issue with CameraShakeBase where infinite Shake Instances were prematurely being considered finished.

diff --git a/Engine/Source/Runtime/Engine/Classes/Camera/CameraShakeBase.h b/Engine/Source/Runtime/Engine/Classes/Camera/CameraShakeBase.h
--- a/Engine/Source/Runtime/Engine/Classes/Camera/CameraShakeBase.h
+++ b/Engine/Source/Runtime/Engine/Classes/Camera/CameraShakeBase.h
@@ -334,6 +334,9 @@ struct ENGINE_API FCameraShakeState
 	/** Helper method to get GetShakeInfo().Duration.Get() */
 	float GetDuration() const { return ShakeInfo.Duration.Get(); }
 
+	/** Helper method to get GetShakeInfo().Duration.IsInifnite() */
+	bool IsInfinite() const { return ShakeInfo.Duration.IsInfinite(); }
+
 private:
 
 	FCameraShakeInfo ShakeInfo;
diff --git a/Engine/Source/Runtime/Engine/Private/Camera/CameraShakeBase.cpp b/Engine/Source/Runtime/Engine/Private/Camera/CameraShakeBase.cpp
--- a/Engine/Source/Runtime/Engine/Private/Camera/CameraShakeBase.cpp
+++ b/Engine/Source/Runtime/Engine/Private/Camera/CameraShakeBase.cpp
@@ -304,7 +304,12 @@ bool UCameraShakeBase::IsFinished() const
 {
 	if (State.IsActive())
 	{
-		if (State.HasDuration())
+		if (State.IsInfinite())
+		{
+			// Since the state is infinite, it's not considered "finished" as long as it is active
+			return false;
+		}		
+		else if (State.HasDuration())
 		{
 			// If we have duration information, we can simply figure out ourselves if
 			// we are finished.
