Fixed Lumen mesh cards sometimes missing from cooked builds

diff --git a/Engine/Source/Runtime/Engine/Private/StaticMesh.cpp b/Engine/Source/Runtime/Engine/Private/StaticMesh.cpp
--- a/Engine/Source/Runtime/Engine/Private/StaticMesh.cpp
+++ b/Engine/Source/Runtime/Engine/Private/StaticMesh.cpp
@@ -1490,6 +1490,8 @@ void FStaticMeshRenderData::SerializeInlineDataRepresentations(FArchive& Ar, USt
 	{
 		if (Ar.IsSaving())
 		{
+			// Need to block on GDistanceFieldAsyncQueue as it can generate GCardRepresentationAsyncQueue tasks
+			GDistanceFieldAsyncQueue->BlockUntilBuildComplete(Owner, false);
 			GCardRepresentationAsyncQueue->BlockUntilBuildComplete(Owner, false);
 		}
 
