Editor splash screen improvements.

- Fixed a bug in the splash screen that would only show the root most scoped slow task resulting in only "Initializing..." displaying.  Now we work from the back of the array (top of the stack) to pull the innermost scoped slow task message.  This results in more log messages displaying.
- Display the splash screen percentage leftmost justified to avoid it jumping around due to the length of the message.
- Adding messages to the global shader scoped slow tasks.
- Fixed a bug where blocking material compilation wasn't displaying a message to the splash screen.  Simply adding a FScopedSlowTask will not update the UI's display message.  The UI is only updated when EnterProgressFrame is called.

diff --git a/Engine/Source/Editor/UnrealEd/Private/FeedbackContextEditor.cpp b/Engine/Source/Editor/UnrealEd/Private/FeedbackContextEditor.cpp
--- a/Engine/Source/Editor/UnrealEd/Private/FeedbackContextEditor.cpp
+++ b/Engine/Source/Editor/UnrealEd/Private/FeedbackContextEditor.cpp
@@ -557,9 +557,9 @@ void FFeedbackContextEditor::ProgressReported( const float TotalProgressInterp,
 	else if (FPlatformSplash::IsShown())
 	{
 		// Always show the top-most message
-		for (auto& Scope : ScopeStack)
+		for (int i = ScopeStack.Num() - 1; i > -1; --i)
 		{
-			const FText ThisMessage = Scope->GetCurrentMessage();
+			const FText ThisMessage = ScopeStack[i]->GetCurrentMessage();
 			if (!ThisMessage.IsEmpty())
 			{
 				DisplayMessage = ThisMessage;
@@ -593,8 +593,8 @@ void FFeedbackContextEditor::ProgressReported( const float TotalProgressInterp,
 					NewDisplayMessage.AppendChar( TCHAR( ' ' ) );
 				}
 			}
-			NewDisplayMessage.Append( FString::Printf( TEXT( " %i%%" ), int(TotalProgressInterp * 100.f) ) );
-			DisplayMessage = FText::FromString( NewDisplayMessage );
+
+			DisplayMessage = FText::FromString(FString::Printf(TEXT("%3i%% - %s"), int(TotalProgressInterp * 100.f), *NewDisplayMessage));
 		}
 
 		FPlatformSplash::SetSplashText(SplashTextType::StartupProgress, *DisplayMessage.ToString());
diff --git a/Engine/Source/Runtime/Engine/Private/ShaderCompiler/ShaderCompiler.cpp b/Engine/Source/Runtime/Engine/Private/ShaderCompiler/ShaderCompiler.cpp
--- a/Engine/Source/Runtime/Engine/Private/ShaderCompiler/ShaderCompiler.cpp
+++ b/Engine/Source/Runtime/Engine/Private/ShaderCompiler/ShaderCompiler.cpp
@@ -4011,7 +4011,8 @@ void FShaderCompilingManager::FinishCompilation(const TCHAR* MaterialName, const
 		StatusUpdate = NSLOCTEXT("ShaderCompilingManager", "CompilingShadersStatus", "Compiling shaders...");
 	}
 
-	FScopedSlowTask SlowTask(0, StatusUpdate, GIsEditor && !IsRunningCommandlet());
+	FScopedSlowTask SlowTask(1, StatusUpdate, GIsEditor && !IsRunningCommandlet());
+	SlowTask.EnterProgressFrame(1);
 
 	TMap<int32, FShaderMapFinalizeResults> CompiledShaderMaps;
 	CompiledShaderMaps.Append( PendingFinalizeShaderMaps );
@@ -5815,10 +5816,10 @@ void CompileGlobalShaderMap(EShaderPlatform Platform, const ITargetPlatform* Tar
 		// GetGlobalShaderMap is called the first time during startup in the main thread.
 		check(IsInGameThread());
 
-		FScopedSlowTask SlowTask(70);
+		FScopedSlowTask SlowTask(70, LOCTEXT("CreateGlobalShaderMap", "Creating Global Shader Map..."));
 
 		// verify that all shader source files are intact
-		SlowTask.EnterProgressFrame(20);
+		SlowTask.EnterProgressFrame(20, LOCTEXT("VerifyShaderSourceFiles", "Verifying Global Shader source files..."));
 		VerifyShaderSourceFiles(Platform);
 
 		GGlobalShaderMap[Platform] = new FGlobalShaderMap(Platform);
@@ -5902,7 +5903,7 @@ void CompileGlobalShaderMap(EShaderPlatform Platform, const ITargetPlatform* Tar
 			// Submit DDC requests.
 			for (const auto& ShaderFilenameDependencies : ShaderMapId.GetShaderFilenameToDependeciesMap())
 			{
-				SlowTask.EnterProgressFrame(ProgressStep);
+				SlowTask.EnterProgressFrame(ProgressStep, LOCTEXT("SubmitDDCRequests", "Submitting global shader DDC Requests..."));
 
 				const FString DataKey = GetGlobalShaderMapKeyString(ShaderMapId, Platform, TargetPlatform, ShaderFilenameDependencies.Value);
 
@@ -5919,7 +5920,7 @@ void CompileGlobalShaderMap(EShaderPlatform Platform, const ITargetPlatform* Tar
 			// Process finished DDC requests.
 			for (const auto& ShaderFilenameDependencies : ShaderMapId.GetShaderFilenameToDependeciesMap())
 			{
-				SlowTask.EnterProgressFrame(ProgressStep);
+				SlowTask.EnterProgressFrame(ProgressStep, LOCTEXT("ProcessDDCRequests", "Processing global shader DDC requests..."));
 				CachedData.Reset();
 				COOK_STAT(auto Timer = GlobalShaderCookStats::UsageStats.TimeSyncWork());
 
