Use triangulation at runtime in BlendSpace. This is coupled with changes to the BlendSpace UI.

diff --git a/Engine/Source/Editor/EditorStyle/Private/SlateEditorStyle.cpp b/Engine/Source/Editor/EditorStyle/Private/SlateEditorStyle.cpp
--- a/Engine/Source/Editor/EditorStyle/Private/SlateEditorStyle.cpp
+++ b/Engine/Source/Editor/EditorStyle/Private/SlateEditorStyle.cpp
@@ -6049,7 +6049,6 @@ void FSlateEditorStyle::FStyle::SetupPersonaStyle()
 		Set("BlendSpaceKey.Drop", SelectionColor_Inactive);
 		Set("BlendSpaceKey.Invalid", LogColor_Error);
 		Set("BlendSpaceKey.Preview", LogColor_Command);
-		Set("BlendSpaceKey.UnSnapped", LogColor_Warning);
 	}
 
 	// Custom menu style for recent commands list
diff --git a/Engine/Source/Editor/EditorStyle/Private/StarshipStyle.cpp b/Engine/Source/Editor/EditorStyle/Private/StarshipStyle.cpp
--- a/Engine/Source/Editor/EditorStyle/Private/StarshipStyle.cpp
+++ b/Engine/Source/Editor/EditorStyle/Private/StarshipStyle.cpp
@@ -5496,7 +5496,6 @@ void FStarshipEditorStyle::FStyle::SetupPersonaStyle()
 		Set("BlendSpaceKey.Drop", SelectionColor_Inactive);
 		Set("BlendSpaceKey.Invalid", LogColor_Error);
 		Set("BlendSpaceKey.Preview", LogColor_Command);
-		Set("BlendSpaceKey.UnSnapped", LogColor_Warning);
 	}
 
 	// Custom menu style for recent commands list
diff --git a/Engine/Source/Editor/Persona/Private/AnimationBlendSpace1DHelpers.cpp b/Engine/Source/Editor/Persona/Private/AnimationBlendSpace1DHelpers.cpp
deleted file mode 100644
diff --git a/Engine/Source/Editor/Persona/Private/AnimationBlendSpace1DHelpers.h b/Engine/Source/Editor/Persona/Private/AnimationBlendSpace1DHelpers.h
deleted file mode 100644
diff --git a/Engine/Source/Editor/Persona/Private/AnimationBlendSpaceHelpers.cpp b/Engine/Source/Editor/Persona/Private/AnimationBlendSpaceHelpers.cpp
deleted file mode 100644
diff --git a/Engine/Source/Editor/Persona/Private/AnimationBlendSpaceHelpers.h b/Engine/Source/Editor/Persona/Private/AnimationBlendSpaceHelpers.h
deleted file mode 100644
diff --git a/Engine/Source/Editor/Persona/Private/Customization/BlendSampleDetails.cpp b/Engine/Source/Editor/Persona/Private/Customization/BlendSampleDetails.cpp
--- a/Engine/Source/Editor/Persona/Private/Customization/BlendSampleDetails.cpp
+++ b/Engine/Source/Editor/Persona/Private/Customization/BlendSampleDetails.cpp
@@ -114,9 +114,6 @@ void FBlendSampleDetails::CustomizeDetails(class IDetailLayoutBuilder& DetailBui
 		FDetailWidgetRow& GraphRow = CategoryBuilder.AddCustomRow(FText::FromString(TEXT("Graph")));
 		FBlendSampleDetails::GenerateSampleGraphWidget(GraphRow, BlendSpaceNode, SampleIndex);
 	}
-	
-	TSharedPtr<IPropertyHandle> SnappedProperty = DetailBuilder.GetProperty(GET_MEMBER_NAME_CHECKED(FBlendSample, bSnapToGrid), FBlendSample::StaticStruct());
-	CategoryBuilder.AddProperty(SnappedProperty);
 }
 
 void FBlendSampleDetails::GenerateBlendSampleWidget(TFunction<FDetailWidgetRow& (void)> InFunctor, FOnSampleMoved OnSampleMoved, const UBlendSpace* BlendSpace, const int32 SampleIndex, bool bShowLabel)
@@ -129,27 +126,8 @@ void FBlendSampleDetails::GenerateBlendSampleWidget(TFunction<FDetailWidgetRow&
 			const FBlendParameter& BlendParameter = BlendSpace->GetBlendParameter(ParameterIndex);
 			const FBlendSample& Sample = BlendSpace->GetBlendSample(SampleIndex);
 			FVector SampleValue = Sample.SampleValue;
-
-			if(Sample.bSnapToGrid)
-			{
-				const float DeltaStep = (BlendParameter.Max - BlendParameter.Min) / BlendParameter.GridNum;
-				// Calculate snapped value
-				const float MinOffset = NewValue - BlendParameter.Min;
-				float GridSteps = MinOffset / DeltaStep;
-				int32 FlooredSteps = FMath::FloorToInt(GridSteps);
-				GridSteps -= FlooredSteps;
-				FlooredSteps = (GridSteps > .5f) ? FlooredSteps + 1 : FlooredSteps;
-
-				// Temporary snap this value to closest point on grid (since the spin box delta does not provide the desired functionality)
-				SampleValue[ParameterIndex] = BlendParameter.Min + (FlooredSteps * DeltaStep);
-			}
-			else
-			{
-				SampleValue[ParameterIndex] = NewValue;
-			}
-
-			OnSampleMoved.ExecuteIfBound(SampleIndex, SampleValue, bIsInteractive, Sample.bSnapToGrid);
-		
+			SampleValue[ParameterIndex] = NewValue;
+			OnSampleMoved.ExecuteIfBound(SampleIndex, SampleValue, bIsInteractive);		
 		};
 		
 		FDetailWidgetRow& ParameterRow = InFunctor();
diff --git a/Engine/Source/Editor/Persona/Private/Customization/BlendSpaceDetails.cpp b/Engine/Source/Editor/Persona/Private/Customization/BlendSpaceDetails.cpp
--- a/Engine/Source/Editor/Persona/Private/Customization/BlendSpaceDetails.cpp
+++ b/Engine/Source/Editor/Persona/Private/Customization/BlendSpaceDetails.cpp
@@ -62,8 +62,8 @@ void FBlendSpaceDetails::CustomizeDetails(class IDetailLayoutBuilder& DetailBuil
 
 		if (b1DBlendSpace)
 		{
-			TSharedPtr<IPropertyHandle> AxisToScaleAnimation = DetailBuilder.GetProperty(GET_MEMBER_NAME_CHECKED(UBlendSpace, AxisToScaleAnimation), UBlendSpace::StaticClass());
-			DetailBuilder.HideProperty(AxisToScaleAnimation);
+			DetailBuilder.HideProperty(DetailBuilder.GetProperty(GET_MEMBER_NAME_CHECKED(UBlendSpace, AxisToScaleAnimation), UBlendSpace::StaticClass()));
+			DetailBuilder.HideProperty(DetailBuilder.GetProperty(GET_MEMBER_NAME_CHECKED(UBlendSpace, PreferredTriangulationDirection), UBlendSpace::StaticClass()));
 		}
 
 		IDetailCategoryBuilder& CategoryBuilder = DetailBuilder.EditCategory(FName("Axis Settings"));
@@ -121,7 +121,6 @@ void FBlendSpaceDetails::CustomizeDetails(class IDetailLayoutBuilder& DetailBuil
 			TSharedPtr<IPropertyHandle> AnimationProperty = BlendSampleProperty->GetChildHandle(GET_MEMBER_NAME_CHECKED(FBlendSample, Animation));
 			TSharedPtr<IPropertyHandle> SampleValueProperty = BlendSampleProperty->GetChildHandle(GET_MEMBER_NAME_CHECKED(FBlendSample, SampleValue));
 			TSharedPtr<IPropertyHandle> RateScaleProperty = BlendSampleProperty->GetChildHandle(GET_MEMBER_NAME_CHECKED(FBlendSample, RateScale));
-			TSharedPtr<IPropertyHandle> SnappedProperty = BlendSampleProperty->GetChildHandle(GET_MEMBER_NAME_CHECKED(FBlendSample, bSnapToGrid));
 
 			IDetailGroup& Group = SampleCategoryBuilder.AddGroup(FName("GroupName"), FText::GetEmpty());
 			Group.HeaderRow()
@@ -152,13 +151,13 @@ void FBlendSpaceDetails::CustomizeDetails(class IDetailLayoutBuilder& DetailBuil
 				]
 			];
 
-			FBlendSampleDetails::GenerateBlendSampleWidget([&Group]() -> FDetailWidgetRow& { return Group.AddWidgetRow(); }, FOnSampleMoved::CreateLambda([this](const uint32 Index, const FVector& SampleValue, bool bIsInteractive, bool bSnap) 
+			FBlendSampleDetails::GenerateBlendSampleWidget([&Group]() -> FDetailWidgetRow& { return Group.AddWidgetRow(); }, FOnSampleMoved::CreateLambda([this](const uint32 Index, const FVector& SampleValue, bool bIsInteractive) 
 			{
 				if (BlendSpaceBase->IsValidBlendSampleIndex(Index) && BlendSpaceBase->GetBlendSample(Index).SampleValue != SampleValue && !BlendSpaceBase->IsTooCloseToExistingSamplePoint(SampleValue, Index))
 				{
 					BlendSpaceBase->Modify();
 
-					bool bMoveSuccesful = BlendSpaceBase->EditSampleValue(Index, SampleValue, bSnap);
+					bool bMoveSuccesful = BlendSpaceBase->EditSampleValue(Index, SampleValue);
 					if (bMoveSuccesful)
 					{
 						BlendSpaceBase->ValidateSampleData();
@@ -179,8 +178,6 @@ void FBlendSpaceDetails::CustomizeDetails(class IDetailLayoutBuilder& DetailBuil
 				FDetailWidgetRow& GraphRow = Group.AddWidgetRow();
 				FBlendSampleDetails::GenerateSampleGraphWidget(GraphRow, BlendSpaceNode.Get(), SampleIndex);
 			}
-			
-			Group.AddPropertyRow(SnappedProperty.ToSharedRef());
 		}
 	}
 	
diff --git a/Engine/Source/Editor/Persona/Private/SAnimationBlendSpace.cpp b/Engine/Source/Editor/Persona/Private/SAnimationBlendSpace.cpp
--- a/Engine/Source/Editor/Persona/Private/SAnimationBlendSpace.cpp
+++ b/Engine/Source/Editor/Persona/Private/SAnimationBlendSpace.cpp
@@ -18,8 +18,6 @@
 #include "Widgets/Layout/SExpandableArea.h"
 #include "Animation/AimOffsetBlendSpace.h"
 #include "Animation/AimOffsetBlendSpace1D.h"
-#include "AnimationBlendSpaceHelpers.h"
-#include "AnimationBlendSpace1DHelpers.h"
 #include "SAnimationBlendSpaceGridWidget.h"
 
 #define LOCTEXT_NAMESPACE "BlendSpaceEditor"
@@ -92,6 +90,11 @@ void SBlendSpaceEditor::Construct(const FArguments& InArgs)
 
 	OnPropertyChangedHandle = FCoreUObjectDelegates::FOnObjectPropertyChanged::FDelegate::CreateRaw(this, &SBlendSpaceEditor::OnPropertyChanged);
 	OnPropertyChangedHandleDelegateHandle = FCoreUObjectDelegates::OnObjectPropertyChanged.Add(OnPropertyChangedHandle);
+
+	// Force a resampling of the data on construction - it should be fast, and ensures that the
+	// runtime data are in sync with what is in the editor, so there's no surprises if the user
+	// changes something. 
+	ResampleData();
 }
 
 void SBlendSpaceEditor::Construct(const FArguments& InArgs, const TSharedRef<class IPersonaPreviewScene>& InPreviewScene)
@@ -101,7 +104,7 @@ void SBlendSpaceEditor::Construct(const FArguments& InArgs, const TSharedRef<cla
 	Construct(InArgs);
 }
 
-void SBlendSpaceEditor::OnSampleMoved(const int32 SampleIndex, const FVector& NewValue, bool bIsInteractive, bool bSnap)
+void SBlendSpaceEditor::OnSampleMoved(const int32 SampleIndex, const FVector& NewValue, bool bIsInteractive)
 {
 	bool bMoveSuccesful = true;
 	if (BlendSpace->IsValidBlendSampleIndex(SampleIndex) && BlendSpace->GetBlendSample(SampleIndex).SampleValue != NewValue && !BlendSpace->IsTooCloseToExistingSamplePoint(NewValue, SampleIndex))
@@ -109,7 +112,7 @@ void SBlendSpaceEditor::OnSampleMoved(const int32 SampleIndex, const FVector& Ne
 		FScopedTransaction ScopedTransaction(LOCTEXT("MoveSample", "Moving Blend Grid Sample"));
 		BlendSpace->Modify();
 
-		bMoveSuccesful = BlendSpace->EditSampleValue(SampleIndex, NewValue, bSnap);
+		bMoveSuccesful = BlendSpace->EditSampleValue(SampleIndex, NewValue);
 		if (bMoveSuccesful)
 		{
 			BlendSpace->ValidateSampleData();
@@ -289,146 +292,9 @@ void SBlendSpaceEditor::NotifyPostChange(const FPropertyChangedEvent& PropertyCh
 
 void SBlendSpaceEditor::ResampleData()
 {
-	const TArray<FBlendSample>& BlendSamples = BlendSpace->GetBlendSamples();
-	FBox AABB(ForceInit);
-	for (const FBlendSample& Sample : BlendSamples)
-	{
-		// Add X value from sample (this is the only valid value to be set for 1D blend spaces / aim offsets
-		if (Sample.bIsValid)
-		{
-			AABB += Sample.SampleValue;
-		}
-	}
-
-	TArray<int32> DimensionIndices;
-	DimensionIndices.Reset(3);
-	if (AABB.GetExtent().X > 0.0f)
-	{
-		DimensionIndices.Push(0);
-	}
-	if (AABB.GetExtent().Y > 0.0f)
-	{
-		DimensionIndices.Push(1);
-	}
-	if (AABB.GetExtent().Z > 0.0f)
-	{
-		DimensionIndices.Push(2);
-	}
-
-	if (DimensionIndices.Num() == 1)
-	{
-		ResampleData1D(DimensionIndices);
-	}
-	else if (DimensionIndices.Num() == 2)
-	{
-		ResampleData2D(DimensionIndices);
-	}
-	else
-	{
-		UE_LOG(LogAnimation, Warning, TEXT("Found %d dimensions from the samples"), DimensionIndices.Num());
-	}
-}
-
-
-void SBlendSpaceEditor::ResampleData1D(const TArray<int32>& DimensionIndices)
-{
-	check(DimensionIndices.Num() == 1);
-	FLineElementGenerator LineElementGenerator;
-
-	int32 Index0 = DimensionIndices[0];
-
-	const FBlendParameter& BlendParameter = BlendSpace->GetBlendParameter(Index0);
-	LineElementGenerator.Init(BlendParameter);
-
-	const TArray<FBlendSample>& BlendSamples = BlendSpace->GetBlendSamples();
-	UE_LOG(LogAnimation, Log, TEXT("Resampling data in 1D - %d samples"), BlendSamples.Num());
-
-	if (BlendSamples.Num())
-	{
-		for (const FBlendSample& Sample : BlendSamples)
-		{
-			// Add X value from sample (this is the only valid value to be set for 1D blend spaces / aim offsets
-			if (Sample.bIsValid)
-			{
-				LineElementGenerator.SamplePointList.Add(Sample.SampleValue[Index0]);
-			}
-		}
-		LineElementGenerator.CalculateEditorElements();
-
-		// Create point to sample index list
-		TArray<int32> PointListToSampleIndices;
-		PointListToSampleIndices.Init(INDEX_NONE, LineElementGenerator.SamplePointList.Num());
-		for (int32 PointIndex = 0; PointIndex < LineElementGenerator.SamplePointList.Num(); ++PointIndex)
-		{
-			const float Point = LineElementGenerator.SamplePointList[PointIndex];
-			for (int32 SampleIndex = 0; SampleIndex < BlendSamples.Num(); ++SampleIndex)
-			{
-				if (BlendSamples[SampleIndex].SampleValue[Index0] == Point)
-				{
-					PointListToSampleIndices[PointIndex] = SampleIndex;
-					break;
-				}
-			}
-		}
-
-		BlendSpace->FillupGridElements(PointListToSampleIndices, LineElementGenerator.EditorElements, DimensionIndices);
-	}
-}
-
-
-void SBlendSpaceEditor::ResampleData2D(const TArray<int32>& DimensionIndices)
-{
-	check(DimensionIndices.Num() == 2);
-
-	// TODO for now, Index0 will always be 0 and Index1 will always be 1, since we don't support 3D.
-	// However, if/when we support authoring using a third dimension, we could get here with any 2D
-	// combination - e.g. XY, XZ or YZ. Then we can either make the triangulation code handle the
-	// indexing, or tweak the triangles going into and out of it.
-	int32 Index0 = DimensionIndices[0];
-	int32 Index1 = DimensionIndices[1];
-
-	// clear first
-	FDelaunayTriangleGenerator DelaunayTriangleGenerator;
-
-	// you don't like to overwrite the link here (between visible points vs sample points, 
-	// so allow this if no triangle is generated
-	const FBlendParameter& BlendParamX = BlendSpace->GetBlendParameter(Index0);
-	const FBlendParameter& BlendParamY = BlendSpace->GetBlendParameter(Index1);
-	FBlendSpaceGrid	BlendSpaceGrid;
-	BlendSpaceGrid.SetGridInfo(BlendParamX, BlendParamY);
-	DelaunayTriangleGenerator.SetGridBox(BlendParamX, BlendParamY);
-
-	BlendSpace->EmptyGridElements();
-
-	UE_LOG(LogAnimation, Log, TEXT("Resampling data in 2D - %d samples"), BlendSpace->GetNumberOfBlendSamples());
-	if (BlendSpace->GetNumberOfBlendSamples())
+	if (BlendSpace)
 	{
-		bool bAllSamplesValid = true;
-		for (int32 SampleIndex = 0; SampleIndex < BlendSpace->GetNumberOfBlendSamples(); ++SampleIndex)
-		{
-			const FBlendSample& Sample = BlendSpace->GetBlendSample(SampleIndex);
-
-			// Do not add invalid sample points (user will need to correct them to be incorporated into the blendspace)
-			if (Sample.bIsValid)
-			{
-				DelaunayTriangleGenerator.AddSamplePoint(Sample.SampleValue, SampleIndex);
-			}
-		}
-
-		// triangulate
-		DelaunayTriangleGenerator.Triangulate();
-
-		// once triangulated, generate grid
-		const TArray<FPoint>& Points = DelaunayTriangleGenerator.GetSamplePointList();
-		const TArray<FTriangle*>& Triangles = DelaunayTriangleGenerator.GetTriangleList();
-		BlendSpaceGrid.GenerateGridElements(Points, Triangles);
-
-		// now fill up grid elements in BlendSpace using this Element information
-		if (Triangles.Num() > 0)
-		{
-			const TArray<FEditorElement>& GridElements = BlendSpaceGrid.GetElements();
-			BlendSpace->FillupGridElements(DelaunayTriangleGenerator.GetIndiceMapping(), GridElements, DimensionIndices);
-		}
+		BlendSpace->ResampleData();
 	}
 }
 
diff --git a/Engine/Source/Editor/Persona/Private/SAnimationBlendSpace.h b/Engine/Source/Editor/Persona/Private/SAnimationBlendSpace.h
--- a/Engine/Source/Editor/Persona/Private/SAnimationBlendSpace.h
+++ b/Engine/Source/Editor/Persona/Private/SAnimationBlendSpace.h
@@ -66,7 +66,7 @@ public:
 	// End FNotifyHook overrides
 
 protected:
-	void OnSampleMoved(const int32 SampleIndex, const FVector& NewValue, bool bIsInteractive, bool bSnap);
+	void OnSampleMoved(const int32 SampleIndex, const FVector& NewValue, bool bIsInteractive);
 	void OnSampleRemoved(const int32 SampleIndex);
 	void OnSampleAdded(UAnimSequence* Animation, const FVector& Value);
 	void OnSampleReplaced(const int32 SampleIndex, UAnimSequence* Animation);
@@ -95,8 +95,6 @@ protected:
 	void OnPropertyChanged(UObject* ObjectBeingModified, FPropertyChangedEvent& PropertyChangedEvent);
 
 	void ResampleData();
-	void ResampleData1D(const TArray<int32>& DimensionIndices);
-	void ResampleData2D(const TArray<int32>& DimensionIndices);
 
 protected:
 	/** The blend space being edited */
diff --git a/Engine/Source/Editor/Persona/Private/SAnimationBlendSpaceGridWidget.cpp b/Engine/Source/Editor/Persona/Private/SAnimationBlendSpaceGridWidget.cpp
--- a/Engine/Source/Editor/Persona/Private/SAnimationBlendSpaceGridWidget.cpp
+++ b/Engine/Source/Editor/Persona/Private/SAnimationBlendSpaceGridWidget.cpp
@@ -6,6 +6,7 @@
 #include "Animation/BlendSpace.h"
 #include "Animation/BlendSpace1D.h"
 
+#include "Widgets/Input/SButton.h"
 #include "Widgets/SBoxPanel.h"
 #include "Widgets/Layout/SBorder.h"
 #include "Widgets/Text/STextBlock.h"
@@ -36,6 +37,9 @@
 
 #define LOCTEXT_NAMESPACE "SAnimationBlendSpaceGridWidget"
 
+// Draws additional data on the triangulation to help debugging
+//#define DEBUG_BLENDSPACE_TRIANGULATION
+
 void SBlendSpaceGridWidget::Construct(const FArguments& InArgs)
 {
 	BlendSpaceBase = InArgs._BlendSpaceBase;
@@ -70,7 +74,7 @@ void SBlendSpaceGridWidget::Construct(const FArguments& InArgs)
 
 	PreviewFilteredPosition = PreviewPosition;
 
-	bShowTriangulation = false;
+	bShowTriangulation = true;
 	bMouseIsOverGeometry = false;
 	bRefreshCachedData = true;
 	bStretchToFit = true;
@@ -87,11 +91,11 @@ void SBlendSpaceGridWidget::Construct(const FArguments& InArgs)
 	InvalidColor = FEditorStyle::GetSlateColor("BlendSpaceKey.Invalid");
 	DropKeyColor = FEditorStyle::GetSlateColor("BlendSpaceKey.Drop");
 	PreviewKeyColor = FEditorStyle::GetSlateColor("BlendSpaceKey.Preview");
-	UnSnappedColor = FEditorStyle::GetSlateColor("BlendSpaceKey.UnSnapped");
 	GridLinesColor = GetDefault<UEditorStyleSettings>()->RegularColor;
 	GridOutlineColor = GetDefault<UEditorStyleSettings>()->RuleColor;
 	TriangulationColor = FSlateColor(EStyleColor::Foreground);
-	
+	TriangulationCurrentColor = FSlateColor(EStyleColor::Highlight);
+
 	// Retrieve background and sample key brushes 
 	BackgroundImage = FEditorStyle::GetBrush(TEXT("Graph.Panel.SolidBackground"));
 	KeyBrush = FEditorStyle::GetBrush("CurveEd.CurveKey");
@@ -248,7 +252,7 @@ void SBlendSpaceGridWidget::Construct(const FArguments& InArgs)
 						.Padding(FMargin(2.0f, 3.0f, 0.0f, 0.0f))
 						[
 							SNew(STextBlock)
-							.Text(LOCTEXT("BlendspacePreviewToolTip", "Hold Shift to move the Preview Point (Green)" ))
+							.Text(LOCTEXT("BlendspacePreviewToolTip", "Hold Control to set the Preview Point (Green)" ))
 							.Font(FEditorStyle::GetFontStyle(TEXT("AnimViewport.MessageFont")))
 							.ColorAndOpacity(FLinearColor(1.0f, 1.0f, 1.0f, 0.7f))
 							.Visibility(TAttribute<EVisibility>::Create(TAttribute<EVisibility>::FGetter::CreateSP(this, &SBlendSpaceGridWidget::GetPreviewToolTipVisibility)))
@@ -404,7 +408,9 @@ void SBlendSpaceGridWidget::PaintBackgroundAndGrid(const FGeometry& AllottedGeom
 	DrawLayerId += 3;
 }
 
-void SBlendSpaceGridWidget::PaintSampleKeys(const FGeometry& AllottedGeometry, const FSlateRect& MyCullingRect, FSlateWindowElementList& OutDrawElements, int32& DrawLayerId) const
+void SBlendSpaceGridWidget::PaintSampleKeys(
+	const FGeometry& AllottedGeometry, const FSlateRect& MyCullingRect, 
+	FSlateWindowElementList& OutDrawElements, int32& DrawLayerId) const
 {
 	const int32 FilteredPositionLayer = DrawLayerId + 1;
 	const int32 PreviewPositionLayer = DrawLayerId + 2;
@@ -435,33 +441,66 @@ void SBlendSpaceGridWidget::PaintSampleKeys(const FGeometry& AllottedGeometry, c
 			{
 				DrawColor = InvalidColor.GetSpecifiedColor();
 			}
-			else
-			{
-				DrawColor = Sample.bSnapToGrid ? DrawColor : UnSnappedColor.GetSpecifiedColor();
-			}
 
-			const FVector2D GridPosition = SampleValueToGridPosition(Sample.SampleValue) - (KeySize * 0.5f);
-			FSlateDrawElement::MakeBox( OutDrawElements, SampleLayer, AllottedGeometry.ToPaintGeometry(GridPosition, KeySize), KeyBrush, ESlateDrawEffect::None, DrawColor );
+			const FVector2D GridPosition = SampleValueToScreenPosition(Sample.SampleValue) - (KeySize * 0.5f);
+			FSlateDrawElement::MakeBox(
+				OutDrawElements, SampleLayer, AllottedGeometry.ToPaintGeometry(GridPosition, KeySize), 
+				KeyBrush, ESlateDrawEffect::None, DrawColor );
 		}
 
 		// Always draw the filtered position which comes back from whatever is running
 		{
-			FVector2D GridPosition = SampleValueToGridPosition(PreviewFilteredPosition) - (PreviewSize * .5f);
-			FSlateDrawElement::MakeBox(OutDrawElements, FilteredPositionLayer, AllottedGeometry.ToPaintGeometry(GridPosition, PreviewSize), PreviewBrush, ESlateDrawEffect::None, PreviewKeyColor.GetSpecifiedColor() * 0.5);
+			FVector2D GridPosition = SampleValueToScreenPosition(PreviewFilteredPosition) - (PreviewSize * .5f);
+			FSlateDrawElement::MakeBox(
+				OutDrawElements, FilteredPositionLayer, AllottedGeometry.ToPaintGeometry(GridPosition, PreviewSize), 
+				PreviewBrush, ESlateDrawEffect::None, PreviewKeyColor.GetSpecifiedColor() * 0.5);
 		}
 
 		if (bPreviewPositionSet)
 		{
-			FVector2D GridPosition = SampleValueToGridPosition(PreviewPosition) - (PreviewSize * .5f);
-			FSlateDrawElement::MakeBox(OutDrawElements, PreviewPositionLayer, AllottedGeometry.ToPaintGeometry(GridPosition, PreviewSize), PreviewBrush, ESlateDrawEffect::None, PreviewKeyColor.GetSpecifiedColor());
+			FVector2D GridPosition = SampleValueToScreenPosition(PreviewPosition) - (PreviewSize * .5f);
+			FSlateDrawElement::MakeBox(
+				OutDrawElements, PreviewPositionLayer, AllottedGeometry.ToPaintGeometry(GridPosition, PreviewSize), 
+				PreviewBrush, ESlateDrawEffect::None, PreviewKeyColor.GetSpecifiedColor());
 		}
 
 		if (DragState == EDragState::DragDrop || DragState == EDragState::InvalidDragDrop)
 		{
-			const FVector2D GridPoint = SnapToClosestGridPoint(LocalMousePosition) - (KeySize * .5f);
-			FSlateDrawElement::MakeBox( OutDrawElements, SampleLayer, AllottedGeometry.ToPaintGeometry(GridPoint, KeySize), KeyBrush, ESlateDrawEffect::None,
+			const FVector2D GridPoint = SnapScreenPositionToGrid(
+				LocalMousePosition, FSlateApplication::Get().GetModifierKeys().IsShiftDown()) - (KeySize * .5f);
+			FSlateDrawElement::MakeBox(
+				OutDrawElements, SampleLayer, AllottedGeometry.ToPaintGeometry(GridPoint, KeySize), KeyBrush, ESlateDrawEffect::None,
 				(DragState == EDragState::DragDrop) ? DropKeyColor.GetSpecifiedColor() : InvalidColor.GetSpecifiedColor() );
 		}
+
+		// Also show the weights that are getting picked up as bars, using two overlaid boxes
+		if (bSamplePreviewing && FSlateApplication::Get().GetModifierKeys().IsAltDown())
+		{
+			for (const FBlendSampleData& PreviewedSample : PreviewedSamples)
+			{
+				float Weight = PreviewedSample.TotalWeight;
+				int32 SampleIndex = PreviewedSample.SampleDataIndex;
+				FVector2D Point = SampleValueToScreenPosition(Samples[SampleIndex].SampleValue);
+
+				float MaxWeightWidth = 48;
+				float WeightHeight = 6;
+
+				Point.Y -= KeySize.Y / 2 + WeightHeight * 1.25;
+				Point.X -= MaxWeightWidth * 0.5f;
+
+				FSlateDrawElement::MakeBox(
+					OutDrawElements, DrawLayerId + 1, AllottedGeometry.MakeChild(
+						FVector2D(Point.X, Point.Y), FVector2D(MaxWeightWidth, WeightHeight)).ToPaintGeometry(),
+					LabelBrush, ESlateDrawEffect::None, FLinearColor::Black);
+
+				float Border = 1.0f;
+				FSlateDrawElement::MakeBox(
+					OutDrawElements, DrawLayerId + 2, AllottedGeometry.MakeChild(
+						FVector2D(Point.X + Border, Point.Y + Border), 
+						FVector2D(Weight * (MaxWeightWidth - 2 * Border), WeightHeight - 2 * Border)).ToPaintGeometry(),
+					LabelBrush, ESlateDrawEffect::None, FLinearColor::Gray);
+			}
+		}
 	}
 
 	DrawLayerId += 3;
@@ -540,38 +579,253 @@ void SBlendSpaceGridWidget::PaintAxisText(const FGeometry& AllottedGeometry, con
 	DrawLayerId += 1;
 }
 
+// Paint a filled triangle
+static void PaintTriangle(
+	const FVector2D&         P0,
+	const FVector2D&         P1,
+	const FVector2D&         P2,
+	const FGeometry&         AllottedGeometry,
+	FLinearColor             Color,
+	const FSlateBrush*       Brush,
+	FSlateWindowElementList& OutDrawElements,
+	int32                    DrawLayerId)
+{
+	const FVector2D* Points[3] = { &P0, &P1, &P2 };
+
+	TArray<FSlateVertex> Vertices;
+	Vertices.Reserve(3);
+
+	for (int32 PointIndex = 0; PointIndex != 3; ++PointIndex)
+	{
+		Vertices.AddZeroed();
+		FSlateVertex& NewVert = Vertices.Last();
+		NewVert.Position = AllottedGeometry.LocalToAbsolute(*Points[PointIndex]);
+		NewVert.Color = Color.ToFColor(false);
+	}
+
+	// Fill by making triangles
+	TArray<SlateIndex> VertexIndices = { 0, 1, 2 };
+	FSlateDrawElement::MakeCustomVerts(
+		OutDrawElements, DrawLayerId, Brush->GetRenderingResource(), Vertices, VertexIndices, nullptr, 0, 0);
+}
+
+// Paints a filled polygon with outline, defined by a set of points which don't need to be sorted.
+// This will handle concave polygons, but only if the centroid lies inside the polygon.
+static void PaintPolygon(
+	TArray<FVector2D>&       Points,
+	const FGeometry&         AllottedGeometry,
+	FLinearColor             FillColor,
+	FLinearColor             OutlineColor,
+	const FSlateBrush*       Brush,
+	FSlateWindowElementList& OutDrawElements,
+	int32                    DrawLayerId)
+{
+	TArray<FSlateVertex> Vertices;
+	Vertices.Reserve(Points.Num() + 1);
+
+	// Add a mid-position vertex so that we handle polygons that aren't completely convex
+	Vertices.AddZeroed();
+
+	FSlateVertex& MidVertex = Vertices.Last();
+	for (int32 PointIndex = 0; PointIndex != Points.Num(); ++PointIndex)
+	{
+		Vertices.AddZeroed();
+		FSlateVertex& NewVert = Vertices.Last();
+		NewVert.Position = AllottedGeometry.LocalToAbsolute(Points[PointIndex]);
+		NewVert.Color = FillColor.ToFColor(false);
+		MidVertex.Position += NewVert.Position;
+	}
+	MidVertex.Position /= Points.Num();
+	MidVertex.Color = FillColor.ToFColor(false);
+
+	// Make sure the points all wind correctly relative to the mid point
+	struct FComparePoints
+	{
+		FComparePoints(const FSlateVertex& Mid) : MidPoint(Mid) {}
+		bool operator()(const FSlateVertex& A, const FSlateVertex& B) const
+		{
+			FVector2D DeltaA = A.Position - MidPoint.Position;
+			FVector2D DeltaB = B.Position - MidPoint.Position;
+			float AngleA = FMath::Atan2(DeltaA.Y, DeltaA.X);
+			float AngleB = FMath::Atan2(DeltaB.Y, DeltaB.X);
+			return AngleA < AngleB;
+		}
+		FSlateVertex MidPoint;
+	};
+	Sort(Vertices.GetData() + 1, Vertices.Num() - 1, FComparePoints(MidVertex));
+
+	if (FillColor.A > 0)
+	{
+		// Fill by making triangles
+		TArray<SlateIndex> VertexIndices;
+		for (int VertexIndex = 1; VertexIndex < Vertices.Num(); ++VertexIndex)
+		{
+			VertexIndices.Add(0);
+			VertexIndices.Add(VertexIndex);
+			VertexIndices.Add(VertexIndex + 1 >= Vertices.Num() ? 1 : VertexIndex + 1);
+		}
+
+		FSlateDrawElement::MakeCustomVerts(
+			OutDrawElements, DrawLayerId, Brush->GetRenderingResource(), Vertices, VertexIndices, nullptr, 0, 0);
+	}
+
+	if (OutlineColor.A > 0)
+	{
+		TArray<FVector2D> LinePoints;
+		LinePoints.Reserve(Points.Num() + 1);
+		for (int VertexIndex = 1; VertexIndex <= Vertices.Num(); ++VertexIndex)
+		{
+			LinePoints.Add(VertexIndex < Vertices.Num() ? Vertices[VertexIndex].Position : Vertices[1].Position);
+		}
+		FSlateDrawElement::MakeLines(
+			OutDrawElements, DrawLayerId + 1, FPaintGeometry(), LinePoints, ESlateDrawEffect::None, OutlineColor, true, 1.0f);
+	}
+}
+
 void SBlendSpaceGridWidget::PaintTriangulation(const FGeometry& AllottedGeometry, const FSlateRect& MyCullingRect, FSlateWindowElementList& OutDrawElements, int32& DrawLayerId) const
 {
+	TArray<FVector2D> PolygonPoints;
+
 	if(const UBlendSpace* BlendSpace = BlendSpaceBase.Get())
 	{
 		const TArray<FBlendSample>& Samples = BlendSpace->GetBlendSamples();
-		const TArray<FEditorElement>& EditorElements = BlendSpace->GetGridSamples();
-		
-		for (const FEditorElement& Element : EditorElements)
+		if (BlendSpace->bInterpolateUsingGrid)
 		{
-			for (int32 SourceIndex = 0; SourceIndex < 3; ++SourceIndex)
+			// Use the Grid
+			const TArray<FEditorElement>& EditorElements = BlendSpace->GetGridSamples();
+			for (const FEditorElement& Element : EditorElements)
 			{
-				if (Element.Indices[SourceIndex] != INDEX_NONE)
+				for (int32 SourceIndex = 0; SourceIndex < 3; ++SourceIndex)
 				{
-					const FBlendSample& SourceSample = Samples[Element.Indices[SourceIndex]];
-					for (int32 TargetIndex = 0; TargetIndex < 3; ++TargetIndex)
+					if (Element.Indices[SourceIndex] != INDEX_NONE)
 					{
-						if (Element.Indices[TargetIndex] != INDEX_NONE)
+						const FBlendSample& SourceSample = Samples[Element.Indices[SourceIndex]];
+						for (int32 TargetIndex = 0; TargetIndex < 3; ++TargetIndex)
 						{
-							if (TargetIndex != SourceIndex)
+							if (Element.Indices[TargetIndex] != INDEX_NONE)
 							{
-								const FBlendSample& TargetSample = Samples[Element.Indices[TargetIndex]];
-								TArray<FVector2D> Points;
+								if (TargetIndex != SourceIndex)
+								{
+									const FBlendSample& TargetSample = Samples[Element.Indices[TargetIndex]];
+									TArray<FVector2D> Points;
+
+									Points.Add(SampleValueToScreenPosition(SourceSample.SampleValue));
+									Points.Add(SampleValueToScreenPosition(TargetSample.SampleValue));
+
+									// Draw line from and to element
+									FSlateDrawElement::MakeLines(
+										OutDrawElements, DrawLayerId + 1, AllottedGeometry.ToPaintGeometry(), Points, 
+										ESlateDrawEffect::None,  TriangulationColor.GetSpecifiedColor(), true, 0.5f);
+								}
+							}
+						}
+					}
+				}
+			}
+		}
+		else
+		{
+			// Use runtime triangulation
+			const FBlendSpaceData& BlendSpaceData = BlendSpace->GetBlendSpaceData();
+			for (const FBlendSpaceSegment& Segment : BlendSpaceData.Segments)
+			{
+				int32 SampleIndex = Segment.SampleIndices[0];
+				int32 SampleIndexNext = Segment.SampleIndices[1];
+				TArray<FVector2D> Points;
+				Points.Add(SampleValueToScreenPosition(Samples[SampleIndex].SampleValue));
+				Points.Add(SampleValueToScreenPosition(Samples[SampleIndexNext].SampleValue));
+				FSlateDrawElement::MakeLines(
+					OutDrawElements, DrawLayerId + 1, AllottedGeometry.ToPaintGeometry(), Points,
+					ESlateDrawEffect::None, TriangulationColor.GetSpecifiedColor(), true, 0.5f);
+			}
 
-								Points.Add(SampleValueToGridPosition(SourceSample.SampleValue));
-								Points.Add(SampleValueToGridPosition(TargetSample.SampleValue));
+			for (const FBlendSpaceTriangle& Triangle : BlendSpaceData.Triangles)
+			{
+				FVector2D MidPoint(0, 0);
+				PolygonPoints.Empty(3);
+				for (int32 Index = 0; Index != FBlendSpaceTriangle::NUM_VERTICES; ++Index)
+				{
+					int32 IndexNext = (Index + 1) % FBlendSpaceTriangle::NUM_VERTICES;
+					int32 SampleIndex = Triangle.SampleIndices[Index];
+					int32 SampleIndexNext = Triangle.SampleIndices[IndexNext];
+					FVector2D P1 = SampleValueToScreenPosition(Samples[SampleIndex].SampleValue);
+					FVector2D P2 = SampleValueToScreenPosition(Samples[SampleIndexNext].SampleValue);
+					TArray<FVector2D> Points = { P1, P2 };
+					MidPoint += SampleValueToScreenPosition(Samples[SampleIndex].SampleValue) / 3.0f;
+
+					FSlateDrawElement::MakeLines(
+						OutDrawElements, DrawLayerId + 1, AllottedGeometry.ToPaintGeometry(), Points,
+						ESlateDrawEffect::None, TriangulationColor.GetSpecifiedColor(), true, 0.5f);
+					PolygonPoints.Push(P1);
+				}
 
-								// Draw line from and to element
-								FSlateDrawElement::MakeLines(OutDrawElements, DrawLayerId + 1, AllottedGeometry.ToPaintGeometry(), Points, ESlateDrawEffect::None, TriangulationColor.GetSpecifiedColor(), true, 0.5f);
-							}
+				FLinearColor Color = TriangulationCurrentColor.GetSpecifiedColor();
+				Color.A = 0.03f; // Alpha for tinting the triangulation background
+				PaintTriangle(PolygonPoints[0], PolygonPoints[1], PolygonPoints[2], AllottedGeometry, Color, LabelBrush, OutDrawElements, DrawLayerId);
+
+#ifdef DEBUG_BLENDSPACE_TRIANGULATION
+				// Draw the adjacent triangle indices around the perimeter
+				if (!bSamplePreviewing)
+				{
+					for (int32 Index = 0; Index != FBlendSpaceTriangle::NUM_VERTICES; ++Index)
+					{
+						if (Triangle.EdgeInfo[Index].NeighbourTriangleIndex < 0)
+						{
+							int32 IndexNext = (Index + 1) % FBlendSpaceTriangle::NUM_VERTICES;
+							int32 SampleIndex = Triangle.SampleIndices[Index];
+							int32 SampleIndexNext = Triangle.SampleIndices[IndexNext];
+							FVector2D Point0 = SampleValueToScreenPosition(Samples[SampleIndex].SampleValue);
+							FVector2D Point1 = SampleValueToScreenPosition(Samples[SampleIndexNext].SampleValue);
+							FVector2D MidEdge = (Point0 + Point1) * 0.5;
+							float PullInAmount = 0.2;
+							FSlateDrawElement::MakeText(
+								OutDrawElements, DrawLayerId + 1, AllottedGeometry.MakeChild(
+									FMath::Lerp(Point0, MidEdge, PullInAmount), FVector2D(1.0f, 1.0f)).ToPaintGeometry(),
+								FText::AsNumber(Triangle.EdgeInfo[Index].AdjacentPerimeterTriangleIndices[0]),
+								FontInfo, ESlateDrawEffect::None, FLinearColor::Red);
+
+							FSlateDrawElement::MakeText(
+								OutDrawElements, DrawLayerId + 1, AllottedGeometry.MakeChild(
+									FMath::Lerp(Point1, MidEdge, PullInAmount), FVector2D(1.0f, 1.0f)).ToPaintGeometry(),
+								FText::AsNumber(Triangle.EdgeInfo[Index].AdjacentPerimeterTriangleIndices[1]),
+								FontInfo, ESlateDrawEffect::None, FLinearColor::Red);
 						}
 					}
 				}
+#endif
+
+#ifdef DEBUG_BLENDSPACE_TRIANGULATION
+				// Draw the triangle indices for debugging
+				FText Text = FText::AsNumber(&Triangle - &BlendSpaceData.Triangles[0]);
+				FSlateDrawElement::MakeText(
+					OutDrawElements, DrawLayerId + 1, AllottedGeometry.MakeChild(
+						FVector2D(MidPoint.X, MidPoint.Y), FVector2D(1.0f, 1.0f)).ToPaintGeometry(),
+					Text, FontInfo, ESlateDrawEffect::None, FLinearColor::Gray);
+#endif
+			}
+		}
+
+		// Draw the current triangle (or polygon)
+		if (bSamplePreviewing && FSlateApplication::Get().GetModifierKeys().IsAltDown())
+		{
+			PolygonPoints.Empty(3);
+			for (const FBlendSampleData& PreviewedSample : PreviewedSamples)
+			{
+				float Weight = PreviewedSample.TotalWeight;
+				if (Weight)
+				{
+					int32 SampleIndex = PreviewedSample.SampleDataIndex;
+					FVector2D Point = SampleValueToScreenPosition(Samples[SampleIndex].SampleValue);
+					PolygonPoints.Push(Point);
+				}
+			}
+			if (PolygonPoints.Num())
+			{
+				FLinearColor FillColor = TriangulationCurrentColor.GetSpecifiedColor();
+				FillColor.A = 0.2f; // Alpha for the current triangulation triangle
+				FLinearColor OutlineColor = FillColor;
+				OutlineColor.A = 0.5f;
+				PaintPolygon(PolygonPoints, AllottedGeometry, FillColor, OutlineColor, LabelBrush, OutDrawElements, DrawLayerId);
 			}
 		}
 	}
@@ -609,20 +863,19 @@ void SBlendSpaceGridWidget::PaintAnimationNames(const FGeometry& AllottedGeometr
 
 			const FText Name = FText::Format(LOCTEXT("SampleNameFormat", "{0} ({1})"), GetSampleName(Sample, SampleIndex), FText::AsNumber(SampleIndex));
 			const FVector2D TextSize = FontMeasure->Measure(Name, FontInfo);
-
-			FVector2D GridPosition = SampleValueToGridPosition(Sample.SampleValue);
-			// Check on which side of the sample the text should be positioned so that we don't run out of geometry space
-			if ((GridPosition + TextSize).X > AllottedGeometry.GetLocalSize().X)
-			{
-				GridPosition -= FVector2D(TextSize.X + KeySize.X, KeySize.X * .5f);
-			}
-			else
-			{
-				GridPosition += FVector2D(KeySize.X, -KeySize.X * .5f);
-			}
-
-			FSlateDrawElement::MakeBox(OutDrawElements, DrawLayerId + 1, AllottedGeometry.MakeChild(FVector2D(GridPosition.X - 6, GridPosition.Y - 2), TextSize + FVector2D(8.0f, 4.0f)).ToPaintGeometry(), LabelBrush, ESlateDrawEffect::None, FLinearColor::Black);
-			FSlateDrawElement::MakeText(OutDrawElements, DrawLayerId + 2, AllottedGeometry.MakeChild(FVector2D(GridPosition.X, GridPosition.Y), FVector2D(1.0f, 1.0f)).ToPaintGeometry(), Name, FontInfo, ESlateDrawEffect::None, FLinearColor::White);
+			const FVector2D Padding = FVector2D(12.0f, 4.0f);
+
+			FVector2D GridPosition = SampleValueToScreenPosition(Sample.SampleValue);
+			GridPosition += FVector2D(-TextSize.X / 2, KeySize.Y / 2);
+			GridPosition.X -= Padding.X / 2;
+			GridPosition.Y += Padding.Y / 2;
+
+			FSlateDrawElement::MakeBox(
+				OutDrawElements, DrawLayerId + 1, AllottedGeometry.MakeChild(GridPosition, 
+				TextSize + Padding).ToPaintGeometry(), LabelBrush, ESlateDrawEffect::None, FLinearColor::Black);
+			FSlateDrawElement::MakeText(
+				OutDrawElements, DrawLayerId + 2, AllottedGeometry.MakeChild(GridPosition + Padding/2,
+					FVector2D(1.0f, 1.0f)).ToPaintGeometry(), Name, FontInfo, ESlateDrawEffect::None, FLinearColor::White);
 		}
 	}
 
@@ -640,7 +893,7 @@ FReply SBlendSpaceGridWidget::OnDrop(const FGeometry& MyGeometry, const FDragDro
 			{
 				if (DragState == EDragState::DragDrop)
 				{
-					const FVector SampleValue = SnapToClosestSamplePoint(LocalMousePosition);
+					const FVector SampleValue = ScreenPositionToSampleValueWithSnapping(LocalMousePosition, FSlateApplication::Get().GetModifierKeys().IsShiftDown());
 					TSharedPtr<FAssetDragDropOp> DragDropOperation = DragDropEvent.GetOperationAs<FAssetDragDropOp>();
 					if (DragDropOperation.IsValid())
 					{
@@ -851,7 +1104,7 @@ FReply SBlendSpaceGridWidget::OnMouseMove(const FGeometry& MyGeometry, const FPo
 			}
 			else if (IsHovered() && bMouseIsOverGeometry)
 			{
-				if (MouseEvent.IsLeftShiftDown() || MouseEvent.IsRightShiftDown())
+				if (MouseEvent.IsControlDown())
 				{
 					StartPreviewing();
 					DragState = EDragState::Preview;
@@ -859,7 +1112,7 @@ FReply SBlendSpaceGridWidget::OnMouseMove(const FGeometry& MyGeometry, const FPo
 					ShowToolTip();			
 
 					// Set flag for showing advanced preview info in tooltip
-					bAdvancedPreview = MouseEvent.IsLeftControlDown() || MouseEvent.IsRightControlDown();
+					bAdvancedPreview = MouseEvent.IsAltDown();
 					return FReply::Handled();
 				}
 				else if(Position.IsSet())
@@ -869,7 +1122,7 @@ FReply SBlendSpaceGridWidget::OnMouseMove(const FGeometry& MyGeometry, const FPo
 					ShowToolTip();
 
 					// Set flag for showing advanced preview info in tooltip
-					bAdvancedPreview = MouseEvent.IsLeftControlDown() || MouseEvent.IsRightControlDown();
+					bAdvancedPreview = MouseEvent.IsAltDown();
 					return FReply::Handled();
 				}
 				else if (bSamplePreviewing)
@@ -969,7 +1222,7 @@ FReply SBlendSpaceGridWidget::OnKeyDown(const FGeometry& MyGeometry, const FKeyE
 			// Start previewing when either one of the shift keys is pressed
 			if (IsHovered() && bMouseIsOverGeometry)
 			{
-				if ((InKeyEvent.GetKey() == EKeys::LeftShift) || (InKeyEvent.GetKey() == EKeys::RightShift))
+				if ((InKeyEvent.GetKey() == EKeys::LeftControl) || (InKeyEvent.GetKey() == EKeys::RightControl))
 				{
 					StartPreviewing();
 					DragState = EDragState::Preview;
@@ -979,7 +1232,7 @@ FReply SBlendSpaceGridWidget::OnKeyDown(const FGeometry& MyGeometry, const FKeyE
 				}
 		
 				// Set flag for showing advanced preview info in tooltip
-				if ((InKeyEvent.GetKey() == EKeys::LeftControl) || (InKeyEvent.GetKey() == EKeys::RightControl))
+				if ((InKeyEvent.GetKey() == EKeys::LeftAlt) || (InKeyEvent.GetKey() == EKeys::RightAlt))
 				{
 					bAdvancedPreview = true;
 					return FReply::Handled();
@@ -998,7 +1251,7 @@ FReply SBlendSpaceGridWidget::OnKeyUp(const FGeometry& MyGeometry, const FKeyEve
 		if(const UBlendSpace* BlendSpace = BlendSpaceBase.Get())
 		{
 			// Stop previewing when shift keys are released 
-			if ((InKeyEvent.GetKey() == EKeys::LeftShift) || (InKeyEvent.GetKey() == EKeys::RightShift))
+			if ((InKeyEvent.GetKey() == EKeys::LeftControl) || (InKeyEvent.GetKey() == EKeys::RightControl))
 			{
 				StopPreviewing();
 				DragState = EDragState::None;
@@ -1006,7 +1259,7 @@ FReply SBlendSpaceGridWidget::OnKeyUp(const FGeometry& MyGeometry, const FKeyEve
 				return FReply::Handled();
 			}
 
-			if((InKeyEvent.GetKey() == EKeys::LeftControl) || (InKeyEvent.GetKey() == EKeys::RightControl))
+			if((InKeyEvent.GetKey() == EKeys::LeftAlt) || (InKeyEvent.GetKey() == EKeys::RightAlt))
 			{
 				bAdvancedPreview = false;
 				return FReply::Handled();
@@ -1144,15 +1397,15 @@ TSharedPtr<SWidget> SBlendSpaceGridWidget::CreateNewBlendSampleContextMenu(const
 	FMenuBuilder MenuBuilder(bShouldCloseWindowAfterMenuSelection, nullptr);
 
 	FVector NewSampleValue;
-	if(FSlateApplication::Get().GetModifierKeys().IsAltDown())
+	if(FSlateApplication::Get().GetModifierKeys().IsShiftDown())
 	{
-		const FVector2D GridPosition(FMath::Clamp(InMousePosition.X, CachedGridRectangle.Left, CachedGridRectangle.Right),
-									 FMath::Clamp(InMousePosition.Y, CachedGridRectangle.Top, CachedGridRectangle.Bottom));
-		NewSampleValue = GridPositionToSampleValue(GridPosition, false);
+		NewSampleValue = ScreenPositionToSampleValue(SnapScreenPositionToGrid(InMousePosition, FSlateApplication::Get().GetModifierKeys().IsShiftDown()), false);
 	}
 	else
 	{
-		NewSampleValue = GridPositionToSampleValue(SnapToClosestGridPoint(InMousePosition), false);
+		const FVector2D GridPosition(FMath::Clamp(InMousePosition.X, CachedGridRectangle.Left, CachedGridRectangle.Right),
+			FMath::Clamp(InMousePosition.Y, CachedGridRectangle.Top, CachedGridRectangle.Bottom));
+		NewSampleValue = ScreenPositionToSampleValue(GridPosition, false);
 	}
 
 	if(const UBlendSpace* BlendSpace = BlendSpaceBase.Get())
@@ -1198,37 +1451,61 @@ void SBlendSpaceGridWidget::CalculateGridPoints()
 		for (int32 GridX = 0; GridX < SampleGridDivisions.X + 1; ++GridX)
 		{
 			// Calculate grid point in 0-1 form
-			FVector2D GridPoint = FVector2D(GridX * (1.0f / SampleGridDivisions.X), (GridType == EGridType::TwoAxis) ? GridY * (1.0f /
-				SampleGridDivisions.Y) : 0.5f);
+			FVector2D GridPoint(
+				GridX * (1.0f / SampleGridDivisions.X), 
+				(GridType == EGridType::TwoAxis) ? GridY * (1.0f / SampleGridDivisions.Y) : 0.5f);
 
 			// Multiply with size and offset according to the grid layout
 			GridPoint *= CachedGridRectangle.GetSize();
 			GridPoint += CachedGridRectangle.GetTopLeft();
 			CachedGridPoints.Add(GridPoint);
 
-			CachedSamplePoints.Add(FVector(SampleValueMin.X + (GridX * (SampleValueRange.X / SampleGridDivisions.X)),
-				(GridType == EGridType::TwoAxis) ? SampleValueMax.Y - (GridY * (SampleValueRange.Y / SampleGridDivisions.Y)) : 0.0f, 0.0f));
+			CachedSamplePoints.Add(FVector(
+				SampleValueMin.X + (GridX * (SampleValueRange.X / SampleGridDivisions.X)),
+				(GridType == EGridType::TwoAxis) ? SampleValueMax.Y - (GridY * (SampleValueRange.Y / SampleGridDivisions.Y)) : 0.0f, 
+				0.0f));
 		}
 	}
 }
 
-const FVector2D SBlendSpaceGridWidget::SnapToClosestGridPoint(const FVector2D& InPosition) const
+const FVector2D SBlendSpaceGridWidget::SnapScreenPositionToGrid(const FVector2D& InPosition, bool bForceSnap) const
 {
-	const int32 GridPointIndex = FindClosestGridPointIndex(InPosition);
-	return CachedGridPoints[GridPointIndex];
+	const int32 GridPointIndex = FindClosestGridPointIndexFromScreenPosition(InPosition);
+	FVector2D GridPoint = CachedGridPoints[GridPointIndex];
+
+	bool bSnapX = bForceSnap || bSampleSnapToGrid[0];
+	bool bSnapY = bForceSnap || bSampleSnapToGrid[1];
+
+	return FVector2D
+	{
+		bSnapX ? GridPoint.X : InPosition.X,
+		bSnapY ? GridPoint.Y : InPosition.Y
+	};
 }
 
-const FVector SBlendSpaceGridWidget::SnapToClosestSamplePoint(const FVector2D& InPosition) const
+const FVector SBlendSpaceGridWidget::ScreenPositionToSampleValueWithSnapping(const FVector2D& InPosition, bool bForceSnap) const
 {
-	const int32 GridPointIndex = FindClosestGridPointIndex(InPosition);
-	return CachedSamplePoints[GridPointIndex];
+	FVector SampleValue = ScreenPositionToSampleValue(InPosition, true);
+
+	const int32 GridPointIndex = FindClosestGridPointIndexFromScreenPosition(InPosition);
+	FVector GridPos = CachedSamplePoints[GridPointIndex];
+
+	bool bSnapX = bForceSnap || bSampleSnapToGrid[0];
+	bool bSnapY = bForceSnap || bSampleSnapToGrid[1];
+	return FVector
+	{
+		bSnapX ? GridPos.X : SampleValue.X,
+		bSnapY ? GridPos.Y : SampleValue.Y,
+		GridPos.Z,
+	};
 }
 
-int32 SBlendSpaceGridWidget::FindClosestGridPointIndex(const FVector2D& InPosition) const
+int32 SBlendSpaceGridWidget::FindClosestGridPointIndexFromScreenPosition(const FVector2D& InPosition) const
 {
 	// Clamp the screen position to the grid
-	const FVector2D GridPosition(FMath::Clamp(InPosition.X, CachedGridRectangle.Left, CachedGridRectangle.Right),
-								  FMath::Clamp(InPosition.Y, CachedGridRectangle.Top, CachedGridRectangle.Bottom));
+	const FVector2D GridPosition(
+		FMath::Clamp(InPosition.X, CachedGridRectangle.Left, CachedGridRectangle.Right),
+		FMath::Clamp(InPosition.Y, CachedGridRectangle.Top, CachedGridRectangle.Bottom));
 	// Find the closest grid point
 	float Distance = FLT_MAX;
 	int32 GridPointIndex = INDEX_NONE;
@@ -1247,7 +1524,7 @@ int32 SBlendSpaceGridWidget::FindClosestGridPointIndex(const FVector2D& InPositi
 
 	return GridPointIndex;
 }
-const FVector2D SBlendSpaceGridWidget::SampleValueToGridPosition(const FVector& SampleValue) const
+const FVector2D SBlendSpaceGridWidget::SampleValueToScreenPosition(const FVector& SampleValue) const
 {
 	const FVector2D GridSize = CachedGridRectangle.GetSize();	
 	const FVector2D GridCenter = GridSize * 0.5f;
@@ -1264,9 +1541,9 @@ const FVector2D SBlendSpaceGridWidget::SampleValueToGridPosition(const FVector&
 	return SamplePosition2D;	
 }
 
-const FVector SBlendSpaceGridWidget::GridPositionToSampleValue(const FVector2D& GridPosition, bool bClamp) const
+const FVector SBlendSpaceGridWidget::ScreenPositionToSampleValue(const FVector2D& ScreenPosition, bool bClamp) const
 {
-	FVector2D LocalGridPosition = GridPosition;
+	FVector2D LocalGridPosition = ScreenPosition;
 	// Move to center of grid and convert to 0 - 1 form
 	LocalGridPosition -= CachedGridRectangle.GetCenter();
 	LocalGridPosition /= (CachedGridRectangle.GetSize() * 0.5f);
@@ -1278,7 +1555,7 @@ const FVector SBlendSpaceGridWidget::GridPositionToSampleValue(const FVector2D&
 	(		
 		(LocalGridPosition.X * SampleValueRange.X) + SampleValueMin.X,
 		(GridType == EGridType::TwoAxis) ? SampleValueMax.Y - (LocalGridPosition.Y * SampleValueRange.Y) : 0.0f,
-		0.f	
+		0.f
 	);
 	if (bClamp)
 	{
@@ -1301,7 +1578,7 @@ const FSlateRect SBlendSpaceGridWidget::GetGridRectangleFromGeometry(const FGeom
 
 bool SBlendSpaceGridWidget::IsSampleValueWithinMouseRange(const FVector& SampleValue, float& OutDistance) const
 {
-	const FVector2D GridPosition = SampleValueToGridPosition(SampleValue);
+	const FVector2D GridPosition = SampleValueToScreenPosition(SampleValue);
 	OutDistance = FVector2D::Distance(LocalMousePosition, GridPosition);
 	return (OutDistance < ClickAndHighlightThreshold);
 }
@@ -1337,8 +1614,8 @@ void SBlendSpaceGridWidget::StartPreviewing()
 	bSamplePreviewing = true;
 	LastPreviewingMousePosition = LocalMousePosition;
 	FModifierKeysState ModifierKeyState = FSlateApplication::Get().GetModifierKeys();
-	bool bIsManualPreviewing = !bReadOnly && IsHovered() && bMouseIsOverGeometry && (ModifierKeyState.IsLeftShiftDown() || ModifierKeyState.IsRightShiftDown());
-	PreviewPosition = Position.IsSet() && !bIsManualPreviewing ? Position.Get() : GridPositionToSampleValue(LastPreviewingMousePosition, false);
+	bool bIsManualPreviewing = !bReadOnly && IsHovered() && bMouseIsOverGeometry && ModifierKeyState.IsControlDown();
+	PreviewPosition = Position.IsSet() && !bIsManualPreviewing ? Position.Get() : ScreenPositionToSampleValue(LastPreviewingMousePosition, false);
 	PreviewFilteredPosition = FilteredPosition.IsSet() ? FilteredPosition.Get() : PreviewPosition;
 	bPreviewPositionSet = true;	
 	bPreviewToolTipHidden = true;
@@ -1525,7 +1802,7 @@ FText SBlendSpaceGridWidget::GetToolTipSampleValue() const
 				// If we are performing a drag and drop operation return the current sample value it is hovered at
 				case EDragState::DragDrop:
 				{
-					const FVector SampleValue = SnapToClosestSamplePoint(LocalMousePosition);
+					const FVector SampleValue = ScreenPositionToSampleValueWithSnapping(LocalMousePosition, FSlateApplication::Get().GetModifierKeys().IsShiftDown());
 
 					ToolTipText = FText::Format(ValueFormattingText, ParameterXName, FText::FromString(FString::SanitizeFloat(SampleValue.X)), ParameterYName, FText::FromString(FString::SanitizeFloat(SampleValue.Y)));
 
@@ -1582,7 +1859,7 @@ void SBlendSpaceGridWidget::EnableStatusBarMessage(bool bEnable)
 				{
 					StatusBarMessageHandle = StatusBarSubsystem->PushStatusBarMessage(StatusBarName, MakeAttributeLambda([]()
 					{
-						return LOCTEXT("StatusBarMssage", "Hold Ctrl for weight details, hold Shift to move preview value");
+						return LOCTEXT("StatusBarMssage", "Hold Ctrl to move preview value, and Alt to show weight details. Click and drag sample points to move them, with Shift to snap to the grid.");
 					}));
 				}
 			}
@@ -1603,7 +1880,7 @@ void SBlendSpaceGridWidget::EnableStatusBarMessage(bool bEnable)
 
 FText SBlendSpaceGridWidget::GetSampleErrorMessage(const FBlendSample &BlendSample) const
 {
-	const FVector2D GridPosition = SampleValueToGridPosition(BlendSample.SampleValue);
+	const FVector2D GridPosition = SampleValueToScreenPosition(BlendSample.SampleValue);
 	// Either an invalid animation asset set
 	if (BlendSample.Animation == nullptr)
 	{
@@ -1704,7 +1981,7 @@ void SBlendSpaceGridWidget::OnInputBoxValueChanged(const float NewValue, const i
 		FVector SampleValue = Sample.SampleValue;
 
 		// Calculate snapped value
-		if(Sample.bSnapToGrid)
+		if (bSampleSnapToGrid[ParameterIndex])
 		{
 			const float MinOffset = NewValue - SampleValueMin[ParameterIndex];
 			float GridSteps = MinOffset / SampleGridDelta[ParameterIndex];
@@ -1720,7 +1997,7 @@ void SBlendSpaceGridWidget::OnInputBoxValueChanged(const float NewValue, const i
 			SampleValue[ParameterIndex] = NewValue;
 		}
 
-		OnSampleMoved.ExecuteIfBound(SelectedSampleIndex, SampleValue, bIsInteractive, Sample.bSnapToGrid);
+		OnSampleMoved.ExecuteIfBound(SelectedSampleIndex, SampleValue, bIsInteractive);
 	}
 }
 
@@ -1771,22 +2048,15 @@ void SBlendSpaceGridWidget::UpdateGridRatioMargin(const FVector2D& GeometrySize)
 	{
 		// Reset values first
 		GridRatioMargin.Top = GridRatioMargin.Bottom = GridRatioMargin.Left = GridRatioMargin.Right = 0.0f;
-
-		if (SampleValueRange.X >= SampleValueRange.Y)
+		if (GeometrySize.Y > GeometrySize.X)
 		{
-			if (GeometrySize.Y > GeometrySize.X)
-			{
-				const float Difference = GeometrySize.Y - GeometrySize.X;
-				GridRatioMargin.Top = GridRatioMargin.Bottom = Difference * 0.5f;
-			}
+			const float Difference = GeometrySize.Y - GeometrySize.X;
+			GridRatioMargin.Top = GridRatioMargin.Bottom = Difference * 0.5f;
 		}
-		else if (SampleValueRange.X < SampleValueRange.Y)
+		else if (GeometrySize.X > GeometrySize.Y)
 		{
-			if (GeometrySize.X > GeometrySize.Y)
-			{
-				const float Difference = GeometrySize.X - GeometrySize.Y;
-				GridRatioMargin.Left = GridRatioMargin.Right = Difference * 0.5f;
-			}
+			const float Difference = GeometrySize.X - GeometrySize.Y;
+			GridRatioMargin.Left = GridRatioMargin.Right = Difference * 0.5f;
 		}
 	}
 }
@@ -1822,6 +2092,9 @@ void SBlendSpaceGridWidget::UpdateCachedBlendParameterData()
 		SampleGridDelta.X /= (BlendParameterX.GridNum);
 		SampleGridDelta.Y /= (BlendParameterY.GridNum);
 
+		bSampleSnapToGrid[0] = BlendParameterX.bSnapToGrid;
+		bSampleSnapToGrid[1] = BlendParameterY.bSnapToGrid;
+
 		SampleGridDivisions.X = BlendParameterX.GridNum;
 		SampleGridDivisions.Y = BlendParameterY.GridNum;
 
@@ -1937,35 +2210,24 @@ void SBlendSpaceGridWidget::Tick(const FGeometry& AllottedGeometry, const double
 		}
 			else if (DragState == EDragState::DragSample)
 			{
-				// If we are dragging a sample, find out whether or not it has actually moved to a different grid position since the last tick and update the blend space accordingly
+				// If we are dragging a sample, find out whether or not it has actually moved to a
+				// different grid position since the last tick and update the blend space accordingly
 				const FBlendSample& BlendSample = BlendSpace->GetBlendSample(DraggedSampleIndex);
 
-				FVector SampleValue;
-				bool bSnap;
-				if(FSlateApplication::Get().GetModifierKeys().IsAltDown() || !BlendSample.bSnapToGrid)
-				{
-					const FVector2D GridPosition(FMath::Clamp(LocalMousePosition.X, CachedGridRectangle.Left, CachedGridRectangle.Right),
-													FMath::Clamp(LocalMousePosition.Y, CachedGridRectangle.Top, CachedGridRectangle.Bottom));
- 
-					SampleValue = GridPositionToSampleValue(GridPosition, true);
-					bSnap = false;
-				}
-				else
-				{
-					SampleValue = SnapToClosestSamplePoint(LocalMousePosition);
-					bSnap = true;
-				} 
+				FVector SampleValue = ScreenPositionToSampleValueWithSnapping(LocalMousePosition, FSlateApplication::Get().GetModifierKeys().IsShiftDown());
 
 				if (SampleValue != LastDragPosition)
 				{
 					LastDragPosition = SampleValue;
-					OnSampleMoved.ExecuteIfBound(DraggedSampleIndex, SampleValue, false, bSnap);
+					OnSampleMoved.ExecuteIfBound(DraggedSampleIndex, SampleValue, false);
 				}
 			}
 			else if (DragState == EDragState::DragDrop || DragState == EDragState::InvalidDragDrop || DragState == EDragState::DragDropOverride)
 			{
-				// Validate that the sample is not overlapping with a current sample when doing a drag/drop operation and that we are dropping a valid animation for the blend space (type)
-				const FVector DropSampleValue = SnapToClosestSamplePoint(LocalMousePosition);
+				// Validate that the sample is not overlapping with a current sample when doing a
+				// drag/drop operation and that we are dropping a valid animation for the blend
+				// space (type)
+				const FVector DropSampleValue = ScreenPositionToSampleValueWithSnapping(LocalMousePosition, FSlateApplication::Get().GetModifierKeys().IsShiftDown());
 				const bool bValidPosition = BlendSpace->IsSampleWithinBounds(DropSampleValue);
 				const bool bExistingSample = BlendSpace->IsTooCloseToExistingSamplePoint(DropSampleValue, INDEX_NONE);
 				const bool bValidSequence = ValidateAnimationSequence(DragDropAnimationSequence, InvalidDragDropText);
@@ -2008,8 +2270,8 @@ void SBlendSpaceGridWidget::Tick(const FGeometry& AllottedGeometry, const double
 			LastPreviewingMousePosition.X = LocalMousePosition.X;
 			LastPreviewingMousePosition.Y = LocalMousePosition.Y;
 			FModifierKeysState ModifierKeyState = FSlateApplication::Get().GetModifierKeys();
-			bool bIsManualPreviewing = !bReadOnly && IsHovered() && bMouseIsOverGeometry && (ModifierKeyState.IsLeftShiftDown() || ModifierKeyState.IsRightShiftDown());
-			PreviewPosition = Position.IsSet() && !bIsManualPreviewing ? Position.Get() : GridPositionToSampleValue(LastPreviewingMousePosition, false);
+			bool bIsManualPreviewing = !bReadOnly && IsHovered() && bMouseIsOverGeometry && ModifierKeyState.IsControlDown();
+			PreviewPosition = Position.IsSet() && !bIsManualPreviewing ? Position.Get() : ScreenPositionToSampleValue(LastPreviewingMousePosition, false);
 			PreviewPosition = BlendSpace->GetClampedAndWrappedBlendInput(PreviewPosition);
 
 			if (FilteredPosition.IsSet())
@@ -2019,7 +2281,7 @@ void SBlendSpaceGridWidget::Tick(const FGeometry& AllottedGeometry, const double
 
 			// Retrieve and cache weighted samples
 			PreviewedSamples.Empty(4);
-			BlendSpace->GetSamplesFromBlendInput(PreviewPosition, PreviewedSamples);
+			BlendSpace->GetSamplesFromBlendInput(PreviewPosition, PreviewedSamples, CachedTriangulationIndex, false);
 		}
 	}
 
@@ -2160,7 +2422,7 @@ bool SBlendSpaceGridWidget::ValidateAnimationSequence(const UAnimSequence* Anima
 const bool SBlendSpaceGridWidget::IsPreviewing() const 
 { 
 	FModifierKeysState ModifierKeyState = FSlateApplication::Get().GetModifierKeys();
-	bool bIsManualPreviewing = !bReadOnly && IsHovered() && bMouseIsOverGeometry && (ModifierKeyState.IsLeftShiftDown() || ModifierKeyState.IsRightShiftDown());
+	bool bIsManualPreviewing = !bReadOnly && IsHovered() && bMouseIsOverGeometry && ModifierKeyState.IsControlDown();
 	return (bSamplePreviewing && !Position.IsSet()) || (Position.IsSet() && bIsManualPreviewing);
 }
 
diff --git a/Engine/Source/Editor/Persona/Private/SAnimationBlendSpaceGridWidget.h b/Engine/Source/Editor/Persona/Private/SAnimationBlendSpaceGridWidget.h
--- a/Engine/Source/Editor/Persona/Private/SAnimationBlendSpaceGridWidget.h
+++ b/Engine/Source/Editor/Persona/Private/SAnimationBlendSpaceGridWidget.h
@@ -22,7 +22,7 @@ class FSlateWindowElementList;
 class UAnimSequence;
 class UBlendSpace;
 
-DECLARE_DELEGATE_FourParams(FOnSampleMoved, const int32 /*SampleIndex*/, const FVector& /*SampleValue*/, bool /*bIsInteractive*/, bool /*bSnap*/);
+DECLARE_DELEGATE_ThreeParams(FOnSampleMoved, const int32 /*SampleIndex*/, const FVector& /*SampleValue*/, bool /*bIsInteractive*/);
 DECLARE_DELEGATE_TwoParams(FOnSampleAnimationChanged, UAnimSequence*, const FVector&);
 
 class SBlendSpaceGridWidget : public SCompoundWidget, public FNotifyHook
@@ -173,15 +173,15 @@ protected:
 	/** Calculates the screen space grid points */
 	void CalculateGridPoints();
 	/** Snaps the given screen position to the closest grid point */
-	const FVector2D SnapToClosestGridPoint(const FVector2D& InPosition) const;
+	const FVector2D SnapScreenPositionToGrid(const FVector2D& InPosition, bool ForceSnap) const;
 	/** Snaps the given screen position to the sample value on the grid */
-	const FVector SnapToClosestSamplePoint(const FVector2D& InPosition) const;
+	const FVector ScreenPositionToSampleValueWithSnapping(const FVector2D& InPosition, bool bForceSnap) const;
 	/** Returns the index of the closest grid point to the given mouse position */
-	int32 FindClosestGridPointIndex(const FVector2D& InPosition) const;
+	int32 FindClosestGridPointIndexFromScreenPosition(const FVector2D& InPosition) const;
 	/** Converts the given sample value to a screen space position */
-	const FVector2D SampleValueToGridPosition(const FVector& SampleValue) const;
+	const FVector2D SampleValueToScreenPosition(const FVector& SampleValue) const;
 	/** Converst a screen space (grid) position to a sample value, with optional clamping */
-	const FVector GridPositionToSampleValue(const FVector2D& GridPosition, bool bClamp) const;
+	const FVector ScreenPositionToSampleValue(const FVector2D& GridPosition, bool bClamp) const;
 	/** Returns the (calculated) grid rectangle given the supplied geometry, and calls UpdateGridRatioMargin */
 	const FSlateRect GetGridRectangleFromGeometry(const FGeometry& MyGeometry);
 	/** Checks whether or not the blendspace sample value is within the range of the mouse position */
@@ -259,6 +259,8 @@ private:
 	bool bPreviewPositionSet;
 	bool bAdvancedPreview;
 	TArray<FBlendSampleData> PreviewedSamples;
+	/** Previous position in the triangulation/segmentation */
+	int32 CachedTriangulationIndex = -1;
 
 	/** Preview state that may have been returned */
 	FVector PreviewFilteredPosition;
@@ -285,6 +287,7 @@ private:
 	FVector2D SampleValueMax;
 	FVector2D SampleValueRange;
 	FVector2D SampleGridDelta;
+	bool      bSampleSnapToGrid[2];
 	FIntPoint SampleGridDivisions;
 	FText ParameterXName;
 	FText ParameterYName;
@@ -338,7 +341,6 @@ private:
 	FSlateColor PreDragKeyColor;
 	FSlateColor DragKeyColor;
 	FSlateColor InvalidColor;
-	FSlateColor UnSnappedColor;
 	FSlateColor DropKeyColor;
 	FSlateColor PreviewKeyColor;
 
@@ -347,6 +349,7 @@ private:
 	FLinearColor GridLinesColor;
 	FLinearColor GridOutlineColor;
 	FSlateColor TriangulationColor;
+	FSlateColor TriangulationCurrentColor;
 	FSlateFontInfo FontInfo;
 	float TextMargin;
 	bool bShowTriangulation;
diff --git a/Engine/Source/Runtime/AnimGraphRuntime/Private/AnimNodes/AnimNode_AimOffsetLookAt.cpp b/Engine/Source/Runtime/AnimGraphRuntime/Private/AnimNodes/AnimNode_AimOffsetLookAt.cpp
--- a/Engine/Source/Runtime/AnimGraphRuntime/Private/AnimNodes/AnimNode_AimOffsetLookAt.cpp
+++ b/Engine/Source/Runtime/AnimGraphRuntime/Private/AnimNodes/AnimNode_AimOffsetLookAt.cpp
@@ -169,7 +169,7 @@ void FAnimNode_AimOffsetLookAt::UpdateFromLookAtTarget(FPoseContext& LocalPoseCo
 	// Generate BlendSampleDataCache from inputs.
 	if (BlendSpace)
 	{
-		BlendSpace->GetSamplesFromBlendInput(BlendInput, BlendSampleDataCache);
+		BlendSpace->GetSamplesFromBlendInput(BlendInput, BlendSampleDataCache, CachedTriangulationIndex, true);
 	}
 }
 
diff --git a/Engine/Source/Runtime/AnimGraphRuntime/Private/AnimNodes/AnimNode_BlendSpaceGraphBase.cpp b/Engine/Source/Runtime/AnimGraphRuntime/Private/AnimNodes/AnimNode_BlendSpaceGraphBase.cpp
--- a/Engine/Source/Runtime/AnimGraphRuntime/Private/AnimNodes/AnimNode_BlendSpaceGraphBase.cpp
+++ b/Engine/Source/Runtime/AnimGraphRuntime/Private/AnimNodes/AnimNode_BlendSpaceGraphBase.cpp
@@ -49,7 +49,7 @@ void FAnimNode_BlendSpaceGraphBase::UpdateInternal(const FAnimationUpdateContext
 #endif
 	const float DeltaTime = Context.GetDeltaTime();
 	const FVector FilteredBlendParams = BlendSpace->FilterInput(&BlendFilter, BlendParams, DeltaTime);
-	BlendSpace->UpdateBlendSamples(FilteredBlendParams, DeltaTime, BlendSampleDataCache);
+	BlendSpace->UpdateBlendSamples(FilteredBlendParams, DeltaTime, BlendSampleDataCache, CachedTriangulationIndex);
 
 	for(int32 SampleIndex = 0; SampleIndex < BlendSampleDataCache.Num(); ++SampleIndex)
 	{
diff --git a/Engine/Source/Runtime/AnimGraphRuntime/Public/AnimNodes/AnimNode_BlendSpaceGraphBase.h b/Engine/Source/Runtime/AnimGraphRuntime/Public/AnimNodes/AnimNode_BlendSpaceGraphBase.h
--- a/Engine/Source/Runtime/AnimGraphRuntime/Public/AnimNodes/AnimNode_BlendSpaceGraphBase.h
+++ b/Engine/Source/Runtime/AnimGraphRuntime/Public/AnimNodes/AnimNode_BlendSpaceGraphBase.h
@@ -68,6 +68,9 @@ protected:
 	// Cache of sampled data, updated each frame
 	TArray<FBlendSampleData> BlendSampleDataCache;
 
+	/** Previous position in the triangulation/segmentation */
+	int32 CachedTriangulationIndex = -1;
+
 #if WITH_EDITORONLY_DATA
 	// Preview blend params - set in editor only
 	FVector PreviewPosition = FVector::ZeroVector;
diff --git a/Engine/Source/Runtime/AnimGraphRuntime/Public/AnimNodes/AnimNode_BlendSpacePlayer.h b/Engine/Source/Runtime/AnimGraphRuntime/Public/AnimNodes/AnimNode_BlendSpacePlayer.h
--- a/Engine/Source/Runtime/AnimGraphRuntime/Public/AnimNodes/AnimNode_BlendSpacePlayer.h
+++ b/Engine/Source/Runtime/AnimGraphRuntime/Public/AnimNodes/AnimNode_BlendSpacePlayer.h
@@ -62,6 +62,9 @@ protected:
 
 	TArray<FBlendSampleData> BlendSampleDataCache;
 
+	/** Previous position in the triangulation/segmentation */
+	int32 CachedTriangulationIndex = -1;
+
 	UPROPERTY(Transient)
 	TObjectPtr<UBlendSpace> PreviousBlendSpace;
 
diff --git a/Engine/Source/Runtime/Core/Private/Math/UnrealMath.cpp b/Engine/Source/Runtime/Core/Private/Math/UnrealMath.cpp
--- a/Engine/Source/Runtime/Core/Private/Math/UnrealMath.cpp
+++ b/Engine/Source/Runtime/Core/Private/Math/UnrealMath.cpp
@@ -2024,6 +2024,14 @@ FVector FMath::GetBaryCentric2D(const FVector& Point, const FVector& A, const FV
 	return FVector(a, b, 1.0f - a - b);	
 }
 
+FVector FMath::GetBaryCentric2D(const FVector2D& Point, const FVector2D& A, const FVector2D& B, const FVector2D& C)
+{
+	float a = ((B.Y - C.Y) * (Point.X - C.X) + (C.X - B.X) * (Point.Y - C.Y)) / ((B.Y - C.Y) * (A.X - C.X) + (C.X - B.X) * (A.Y - C.Y));
+	float b = ((C.Y - A.Y) * (Point.X - C.X) + (A.X - C.X) * (Point.Y - C.Y)) / ((B.Y - C.Y) * (A.X - C.X) + (C.X - B.X) * (A.Y - C.Y));
+
+	return FVector(a, b, 1.0f - a - b);
+}
+
 FVector FMath::ComputeBaryCentric2D(const FVector& Point, const FVector& A, const FVector& B, const FVector& C)
 {
 	// Compute the normal of the triangle
diff --git a/Engine/Source/Runtime/Core/Public/Math/UnrealMathUtility.h b/Engine/Source/Runtime/Core/Public/Math/UnrealMathUtility.h
--- a/Engine/Source/Runtime/Core/Public/Math/UnrealMathUtility.h
+++ b/Engine/Source/Runtime/Core/Public/Math/UnrealMathUtility.h
@@ -1658,7 +1658,7 @@ public:
 	static CORE_API bool Eval( FString Str, float& OutValue );
 
 	/**
-	 * Computes the barycentric coordinates for a given point in a triangle - simpler version
+	 * Computes the barycentric coordinates for a given point in a triangle, only considering the XY coordinates - simpler than the Compute versions
 	 *
 	 * @param	Point			point to convert to barycentric coordinates (in plane of ABC)
 	 * @param	A,B,C			three non-colinear points defining a triangle in CCW
@@ -1668,6 +1668,17 @@ public:
 	 */
 	static CORE_API FVector GetBaryCentric2D(const FVector& Point, const FVector& A, const FVector& B, const FVector& C);
 
+	/**
+	 * Computes the barycentric coordinates for a given point in a triangle - simpler than the Compute versions
+	 *
+	 * @param	Point			point to convert to barycentric coordinates (in plane of ABC)
+	 * @param	A,B,C			three non-colinear points defining a triangle in CCW
+	 *
+	 * @return Vector containing the three weights a,b,c such that Point = a*A + b*B + c*C
+	 *							                                or Point = A + b*(B-A) + c*(C-A) = (1-b-c)*A + b*B + c*C
+	 */
+	static CORE_API FVector GetBaryCentric2D(const FVector2D& Point, const FVector2D& A, const FVector2D& B, const FVector2D& C);
+
 	/**
 	 * Computes the barycentric coordinates for a given point in a triangle
 	 *
diff --git a/Engine/Source/Runtime/Core/Public/UObject/UE5MainStreamObjectVersion.h b/Engine/Source/Runtime/Core/Public/UObject/UE5MainStreamObjectVersion.h
--- a/Engine/Source/Runtime/Core/Public/UObject/UE5MainStreamObjectVersion.h
+++ b/Engine/Source/Runtime/Core/Public/UObject/UE5MainStreamObjectVersion.h
@@ -100,6 +100,9 @@ struct CORE_API FUE5MainStreamObjectVersion
 		// Moving animation custom attributes from AnimationSequence to UAnimDataModel
 		MoveCustomAttributesToDataModel,
 
+		// Use of triangulation at runtime in BlendSpace
+		BlendSpaceRuntimeTriangulation,
+
 		// -----<new versions can be added above this line>-------------------------------------------------
 		VersionPlusOne,
 		LatestVersion = VersionPlusOne - 1
diff --git a/Engine/Source/Runtime/Engine/Classes/Animation/AnimationAsset.h b/Engine/Source/Runtime/Engine/Classes/Animation/AnimationAsset.h
--- a/Engine/Source/Runtime/Engine/Classes/Animation/AnimationAsset.h
+++ b/Engine/Source/Runtime/Engine/Classes/Animation/AnimationAsset.h
@@ -311,6 +311,7 @@ struct FAnimTickRecord
 		{
 			float  BlendSpacePositionX;
 			float  BlendSpacePositionY;
+			int32  TriangulationIndex;
 			FBlendFilter* BlendFilter;
 			TArray<FBlendSampleData>* BlendSampleDataCache;
 		} BlendSpace;
diff --git a/Engine/Source/Runtime/Engine/Classes/Animation/BlendSpace.h b/Engine/Source/Runtime/Engine/Classes/Animation/BlendSpace.h
--- a/Engine/Source/Runtime/Engine/Classes/Animation/BlendSpace.h
+++ b/Engine/Source/Runtime/Engine/Classes/Animation/BlendSpace.h
@@ -22,14 +22,21 @@ enum EBlendSpaceAxis
 {
 	BSA_None UMETA(DisplayName = "None"),
 	BSA_X UMETA(DisplayName = "Horizontal (X) Axis"),
-	BSA_Y UMETA(DisplayName = "Vertical (Y) Axis"),
-	BSA_Max
+	BSA_Y UMETA(DisplayName = "Vertical (Y) Axis")
+};
+
+UENUM()
+enum class EPreferredTriangulationDirection : uint8
+{
+	None UMETA(DisplayName = "None"),
+	Tangential UMETA(DisplayName = "Tangential"),
+	Radial UMETA(DisplayName = "Radial")
 };
 
 USTRUCT()
 struct FInterpolationParameter
 {
-	GENERATED_USTRUCT_BODY()
+	GENERATED_BODY()
 
 	/** Interpolation Time for input, when it gets input, it will use this time to interpolate to target, used for smoother interpolation. */
 	UPROPERTY(EditAnywhere, Category=Parameter)
@@ -43,7 +50,7 @@ struct FInterpolationParameter
 USTRUCT()
 struct FBlendParameter
 {
-	GENERATED_USTRUCT_BODY()
+	GENERATED_BODY()
 
 	UPROPERTY(EditAnywhere, DisplayName = "Name", Category=BlendParameter)
 	FString DisplayName;
@@ -60,6 +67,10 @@ struct FBlendParameter
 	UPROPERTY(EditAnywhere, DisplayName = "Number of Grid Divisions", Category=BlendParameter, meta=(UIMin="1", ClampMin="1"))
 	int32 GridNum;
 
+	/** If true then samples will always be snapped to the grid on this axis when added, moved, or the axes are changed. */
+	UPROPERTY(EditAnywhere, DisplayName = "Snap to Grid", Category = BlendParameter)
+	bool bSnapToGrid;
+
 	/** If false then input parameters are clamped to the min/max values on this axis. If true then the input can go outside the min/max range and the blend space is treated as being cyclic on this axis. */
 	UPROPERTY(EditAnywhere, DisplayName = "Wrap Input", Category = BlendParameter)
 	bool bWrapInput;
@@ -69,6 +80,7 @@ struct FBlendParameter
 		, Min(0.f)
 		, Max(100.f)
 		, GridNum(4) // TODO when changing GridNum's default value, it breaks all grid samples ATM - provide way to rebuild grid samples during loading
+		, bSnapToGrid(false)
 		, bWrapInput(false)
 	{
 	}
@@ -89,7 +101,7 @@ struct FBlendParameter
 USTRUCT()
 struct FBlendSample
 {
-	GENERATED_USTRUCT_BODY()
+	GENERATED_BODY()
 
 	// For linked animations
 	UPROPERTY(EditAnywhere, Category=BlendSample)
@@ -104,9 +116,6 @@ struct FBlendSample
 	float RateScale;
 
 #if WITH_EDITORONLY_DATA
-	UPROPERTY(EditAnywhere, Category=BlendSample)
-	uint8 bSnapToGrid : 1;
-
 	UPROPERTY(transient)
 	uint8 bIsValid : 1;
 
@@ -120,7 +129,6 @@ struct FBlendSample
 		, SampleValue(0.f)
 		, RateScale(1.0f)
 #if WITH_EDITORONLY_DATA
-		, bSnapToGrid(true)
 		, bIsValid(false)
 		, CachedMarkerDataUpdateCounter(INDEX_NONE)
 #endif // WITH_EDITORONLY_DATA
@@ -132,7 +140,6 @@ struct FBlendSample
 		, SampleValue(InValue)
 		, RateScale(1.0f)
 #if WITH_EDITORONLY_DATA
-		, bSnapToGrid(bInIsSnapped)
 		, bIsValid(bInIsValid)
 		, CachedMarkerDataUpdateCounter(INDEX_NONE)
 #endif // WITH_EDITORONLY_DATA
@@ -145,13 +152,152 @@ struct FBlendSample
 	}
 };
 
+/**
+ * This is the runtime representation of a segment which stores its vertices (start and end) in normalized space.
+ */
+USTRUCT()
+struct FBlendSpaceSegment
+{
+	GENERATED_BODY();
+
+public:
+	// Triangles have three vertices
+	static const int32 NUM_VERTICES = 2;
+
+	/** Indices into the samples */
+	UPROPERTY()
+	int32 SampleIndices[NUM_VERTICES];
+
+	/** The vertices are in the normalized space - i.e. in the range 0-1. */
+	UPROPERTY()
+	float Vertices[NUM_VERTICES];
+};
+
+USTRUCT()
+struct FBlendSpaceTriangleEdgeInfo
+{
+	GENERATED_BODY();
+
+public:
+	/** Edge normal faces out */
+	UPROPERTY()
+	FVector2D Normal;
+
+	UPROPERTY()
+	int32 NeighbourTriangleIndex;
+
+	/**
+	* IF there is no neighbor, then (a) we're on the perimeter and (b) these will be the indices of
+	* triangles along the perimeter (next to the start and end of this edge, respectively) 
+	*/
+	UPROPERTY()
+	int32 AdjacentPerimeterTriangleIndices[2];
+
+	/**
+	 * The vertex index of the associated AdjacentPerimeterTriangle such that the perimeter edge is
+	 * from this vertex to the next.
+	 */
+	UPROPERTY()
+	int32 AdjacentPerimeterVertexIndices[2];
+};
+
+/**
+* This is the runtime representation of a triangle. Each triangle stores its vertices etc in normalized space,
+* with an index to the original samples.
+ */
+USTRUCT()
+struct FBlendSpaceTriangle
+{
+	GENERATED_BODY();
+
+public:
+	// Triangles have three vertices
+	static const int32 NUM_VERTICES = 3;
+
+public:
+
+	/** Indices into the samples */
+	UPROPERTY(EditAnywhere, Category = EditorElement)
+	int32 SampleIndices[NUM_VERTICES];
+
+	/** The vertices are in the normalized space - i.e. in the range 0-1. */
+	UPROPERTY(EditAnywhere, Category = EditorElement)
+	FVector2D Vertices[NUM_VERTICES];
+
+	/** Info for the edge starting at the vertex index and going (anti-clockwise) to the next vertex */
+	UPROPERTY(EditAnywhere, Category = EditorElement)
+	FBlendSpaceTriangleEdgeInfo EdgeInfo[NUM_VERTICES];
+};
+
+USTRUCT()
+struct FWeightedBlendSample
+{
+	GENERATED_BODY();
+
+public:
+	FWeightedBlendSample(int32 Index = INDEX_NONE, float Weight = 0) : SampleIndex(Index), SampleWeight(Weight) {}
+
+public:
+	UPROPERTY()
+	int32 SampleIndex;
+
+	UPROPERTY()
+	float SampleWeight;
+};
+
+/**
+* The runtime data used for interpolating. Note that only one of Segments/Triangles will be in use,
+* depending on the dimensionality of the data.
+*/
+USTRUCT()
+struct FBlendSpaceData
+{
+	GENERATED_BODY();
+public:
+	void GetSamples(
+		TArray<FWeightedBlendSample>& OutWeightedSamples,
+		const TArray<int32>&          InDimensionIndices,
+		const FVector&                InSamplePosition,
+		int32&                        InOutTriangulationIndex) const;
+
+	void Empty()
+	{
+		Segments.Empty();
+		Triangles.Empty();
+	}
+
+	bool IsEmpty() const 
+	{
+		return Segments.Num() == 0 && Triangles.Num() == 0;
+	}
+public:
+	UPROPERTY()
+	TArray<FBlendSpaceSegment> Segments;
+
+	UPROPERTY()
+	TArray<FBlendSpaceTriangle> Triangles;
+
+private:
+	void GetSamples1D(
+		TArray<FWeightedBlendSample>& OutWeightedSamples,
+		const TArray<int32>&          InDimensionIndices,
+		const FVector&                InSamplePosition,
+		int32&                        InOutSegmentIndex) const;
+
+	void GetSamples2D(
+		TArray<FWeightedBlendSample>& OutWeightedSamples,
+		const TArray<int32>&          InDimensionIndices,
+		const FVector&                InSamplePosition,
+		int32&                        InOutTriangleIndex) const;
+};
+
 /**
  * Each elements in the grid
  */
 USTRUCT()
 struct FEditorElement
 {
-	GENERATED_USTRUCT_BODY()
+	GENERATED_BODY()
 
 	// for now we only support triangles
 	static const int32 MAX_VERTICES = 3;
@@ -167,9 +313,6 @@ struct FEditorElement
 		for (int32 ElementIndex = 0; ElementIndex < MAX_VERTICES; ElementIndex++)
 		{
 			Indices[ElementIndex] = INDEX_NONE;
-		}
-		for (int32 ElementIndex = 0; ElementIndex < MAX_VERTICES; ElementIndex++)
-		{
 			Weights[ElementIndex] = 0;
 		}
 	}
@@ -180,7 +323,7 @@ struct FEditorElement
 USTRUCT()
 struct FGridBlendSample
 {
-	GENERATED_USTRUCT_BODY()
+	GENERATED_BODY()
 
 	UPROPERTY()
 	struct FEditorElement GridElement;
@@ -198,7 +341,7 @@ struct FGridBlendSample
 USTRUCT()
 struct FPerBoneInterpolation
 {
-	GENERATED_USTRUCT_BODY()
+	GENERATED_BODY()
 
 	UPROPERTY(EditAnywhere, Category=FPerBoneInterpolation)
 	FBoneReference BoneReference;
@@ -245,6 +388,7 @@ UCLASS(config=Engine, hidecategories=Object, MinimalAPI, BlueprintType)
 class UBlendSpace : public UAnimationAsset, public IInterpolationIndexProvider
 {
 	GENERATED_UCLASS_BODY()
+public:
 
 	/** Required for accessing protected variable names */
 	friend class FBlendSpaceDetails;
@@ -310,10 +454,12 @@ class UBlendSpace : public UAnimationAsset, public IInterpolationIndexProvider
 	* It will return all samples that has weight > KINDA_SMALL_NUMBER
 	*
 	* @param	BlendInput	BlendInput X, Y, Z corresponds to BlendParameters[0], [1], [2]
+	* @param	InOutCachedTriangulationIndex	The previous index into triangulation/segmentation to warm-start the search. 
+	* @param	bCombineAnimations	Will combine samples that point to the same animation. Useful when processing, but confusing when viewing
 	*
 	* @return	true if it has valid OutSampleDataList, false otherwise
 	*/
-	ENGINE_API bool GetSamplesFromBlendInput(const FVector &BlendInput, TArray<FBlendSampleData> & OutSampleDataList) const;
+	ENGINE_API bool GetSamplesFromBlendInput(const FVector &BlendInput, TArray<FBlendSampleData> & OutSampleDataList, int32& InOutCachedTriangulationIndex, bool bCombineAnimations) const;
 
 	/** Initialize BlendSpace for runtime. It needs certain data to be reinitialized per instsance **/
 	ENGINE_API void InitializeFilter(FBlendFilter* Filter) const;
@@ -328,7 +474,7 @@ class UBlendSpace : public UAnimationAsset, public IInterpolationIndexProvider
 	 * @param	InOutSampleDataCache	The sample data cache. Previous frames samples are re-used in the case of target weight interpolation
 	 * @param	InDeltaTime				The tick time for this update
 	 */
-	ENGINE_API bool UpdateBlendSamples(const FVector& InBlendSpacePosition, float InDeltaTime, TArray<FBlendSampleData>& InOutSampleDataCache) const;
+	ENGINE_API bool UpdateBlendSamples(const FVector& InBlendSpacePosition, float InDeltaTime, TArray<FBlendSampleData>& InOutSampleDataCache, int32& InOutCachedTriangulationIndex) const;
 
 	/** Interpolate BlendInput based on Filter data **/
 	ENGINE_API FVector FilterInput(FBlendFilter* Filter, const FVector& BlendInput, float DeltaTime) const;
@@ -345,7 +491,7 @@ class UBlendSpace : public UAnimationAsset, public IInterpolationIndexProvider
 	ENGINE_API bool	AddSample(UAnimSequence* AnimationSequence, const FVector& SampleValue);
 
 	/** edit samples */
-	ENGINE_API bool	EditSampleValue(const int32 BlendSampleIndex, const FVector& NewValue, bool bSnap = true);
+	ENGINE_API bool	EditSampleValue(const int32 BlendSampleIndex, const FVector& NewValue);
 
 	UE_DEPRECATED(5.0, "Please use ReplaceSampleAnimation instead")
 	ENGINE_API bool	UpdateSampleAnimation(UAnimSequence* AnimationSequence, const FVector& SampleValue);
@@ -371,17 +517,21 @@ class UBlendSpace : public UAnimationAsset, public IInterpolationIndexProvider
 	*/
 	ENGINE_API const TArray<FEditorElement>& GetGridSamples() const;
 
-	/** Fill up local GridElements from the grid elements that are created using the sorted points
-	*	This will map back to original index for result
-	*
-	*  @param	SortedPointList		This is the pointlist that are used to create the given GridElements
-	*								This list contains subsets of the points it originally requested for visualization and sorted
-	*
-	*/
-	ENGINE_API void FillupGridElements(const TArray<int32> & PointListToSampleIndices, const TArray<FEditorElement> & GridElements, const TArray<int32>& InDimensionIndices);
-		
-	ENGINE_API void EmptyGridElements();
-	
+	/**
+	 * Returns the runtime triangulation etc data
+	 */
+	ENGINE_API const FBlendSpaceData& GetBlendSpaceData() const;
+
+	/**
+	 * Runs triangulation/segmentation to update our grid and BlendSpaceData structures
+	 */
+	ENGINE_API void ResampleData();
+
+	/**
+	 * Sets up BlendSpaceData based on Line elements
+	 */
+	ENGINE_API void SetBlendSpaceData(const TArray<FBlendSpaceSegment>& Segments);
+
 	/** Validate that the given animation sequence and contained blendspace data */
 	ENGINE_API bool ValidateAnimationSequence(const UAnimSequence* AnimationSequence) const;
 
@@ -430,6 +580,9 @@ protected:
 	/** Returns the blend input clamped to the valid range, unless that axis has been set to wrap in which case no clamping is done **/
 	FVector GetClampedBlendInput(const FVector& BlendInput) const;
 	
+	/** Translates BlendInput to grid space */
+	FVector ConvertBlendInputToGridSpace(const FVector& BlendInput) const;
+
 	/** Translates BlendInput to grid space */
 	FVector GetNormalizedBlendInput(const FVector& BlendInput) const;
 
@@ -449,9 +602,7 @@ protected:
 	bool ContainsNonAdditiveSamples() const;
 	void UpdatePreviewBasePose();
 	/** If around border, snap to the border to avoid empty hole of data that is not valid **/
-	void SnapSamplesToClosestGridPoint();
-
-	void RemapSamplesToNewAxisRange();
+	virtual void SnapSamplesToClosestGridPoint();
 #endif // WITH_EDITOR
 	
 private:
@@ -459,7 +610,7 @@ private:
 	void GetAnimationPose_Internal(TArray<FBlendSampleData>& BlendSampleDataCache, TArrayView<FPoseLink> InPoseLinks, FAnimInstanceProxy* InProxy, bool bInExpectsAdditivePose, /*out*/ FAnimationPoseData& OutAnimationPoseData) const;
 
 	// Internal helper function for UpdateBlendSamples and TickAssetPlayer
-	bool UpdateBlendSamples_Internal(const FVector& InBlendSpacePosition, float InDeltaTime, TArray<FBlendSampleData>& InOutOldSampleDataList, TArray<FBlendSampleData>& InOutSampleDataCache) const;
+	bool UpdateBlendSamples_Internal(const FVector& InBlendSpacePosition, float InDeltaTime, TArray<FBlendSampleData>& InOutOldSampleDataList, TArray<FBlendSampleData>& InOutSampleDataCache, int32& InOutCachedTriangulationIndex) const;
 
 public:
 	/**
@@ -469,6 +620,10 @@ public:
 	UPROPERTY()
 	bool bRotationBlendInMeshSpace;
 
+	/** Input interpolation parameter for all 3 axis, for each axis input, decide how you'd like to interpolate input to*/
+	UPROPERTY(EditAnywhere, Category = InputInterpolation)
+	FInterpolationParameter	InterpolationParam[3];
+
 #if WITH_EDITORONLY_DATA
 	/** Preview Base pose for additive BlendSpace **/
 	UPROPERTY(EditAnywhere, Category = AdditiveSettings)
@@ -479,10 +634,6 @@ public:
 	UPROPERTY(transient)
 	float AnimLength;
 
-	/** Input interpolation parameter for all 3 axis, for each axis input, decide how you'd like to interpolate input to*/
-	UPROPERTY(EditAnywhere, Category = InputInterpolation)
-	FInterpolationParameter	InterpolationParam[3];
-
 	/**
 	* If greater than zero, this is the speed at which the sample weights are allowed to change.
 	* 
@@ -516,6 +667,14 @@ public:
 	UPROPERTY(EditAnywhere, Category = AnimationNotifies)
 	TEnumAsByte<ENotifyTriggerMode::Type> NotifyTriggerMode;
 
+	/** If true then interpolation is done via a grid at runtime. If false the interpolation uses the triangulation. */
+	UPROPERTY(EditAnywhere, Category = InputInterpolation)
+	bool bInterpolateUsingGrid = false;
+
+	/** Preferred edge direction when the triangulation has to make an arbitrary choice */
+	UPROPERTY(EditAnywhere, Category = InputInterpolation)
+	EPreferredTriangulationDirection PreferredTriangulationDirection = EPreferredTriangulationDirection::Tangential;
+
 protected:
 
 	/**
@@ -531,23 +690,27 @@ protected:
 	UPROPERTY()
 	int32 SampleIndexWithMarkers;
 
-	/** Sample animation data **/
+	/** Sample animation data */
 	UPROPERTY(EditAnywhere, Category=BlendSamples)
 	TArray<struct FBlendSample> SampleData;
 
-	/** Grid samples, indexing scheme imposed by subclass **/
+	/** Grid samples, indexing scheme imposed by subclass */
 	UPROPERTY()
 	TArray<struct FEditorElement> GridSamples;
+
+	/** COntainer for the runtime data, which could be line segments, triangulation or tetrahedrons */
+	UPROPERTY()
+	FBlendSpaceData BlendSpaceData;
 	
 	/** Blend Parameters for each axis. **/
 	UPROPERTY(EditAnywhere, Category = BlendParametersTest)
 	struct FBlendParameter BlendParameters[3];
 
-	/** If you have input interpolation, which axis to drive animation speed (scale) - i.e. for locomotion animation, speed axis will drive animation speed (thus scale)**/
+	/** If you have input interpolation, which axis to drive animation speed (scale) - i.e. for locomotion animation, speed axis will drive animation speed (thus scale) */
 	UPROPERTY(EditAnywhere, Category = InputInterpolation)
 	TEnumAsByte<EBlendSpaceAxis> AxisToScaleAnimation;
 
-	/** Reset to reference pose. It does apply different refpose based on additive or not*/
+	/** Reset to reference pose. It does apply different refpose based on additive or not */
 	void ResetToRefPose(FCompactPose& OutPose) const;
 
 	/** The order in which to use the dimensions in the data - e.g. [1, 2] means a 2D blend using Y and Z */
@@ -560,29 +723,34 @@ private:
 	int32 MarkerDataUpdateCounter;
 protected:
 	FVector PreviousAxisMinMaxValues[3];
+	float   PreviousGridSpacings[3];
 #endif	
 
-protected:
+private:
 
-	//================================================================================================
-	// 1D specific functions 
-	//================================================================================================
 	void GetRawSamplesFromBlendInput1D(const FVector& BlendInput, TArray<FGridBlendSample, TInlineAllocator<4> >& OutBlendSamples) const;
-#if WITH_EDITOR
-	void SnapSamplesToClosestGridPoint1D();
-	void RemapSamplesToNewAxisRange1D();
-#endif // WITH_EDITOR
 
-	//================================================================================================
-	// 2D specific functions 
-	//================================================================================================
 	void GetRawSamplesFromBlendInput2D(const FVector& BlendInput, TArray<FGridBlendSample, TInlineAllocator<4> >& OutBlendSamples) const;
-#if WITH_EDITOR
-	void SnapSamplesToClosestGridPoint2D();
-	void RemapSamplesToNewAxisRange2D();
-#endif // WITH_EDITOR
 
-private:
+	/** Fill up local GridElements from the grid elements that are created using the sorted points
+	*	This will map back to original index for result
+	*
+	*  @param	SortedPointList		This is the pointlist that are used to create the given GridElements
+	*								This list contains subsets of the points it originally requested for visualization and sorted
+	*
+	*/
+	void FillupGridElements(const TArray<FEditorElement>& GridElements, const TArray<int32>& InDimensionIndices);
+
+	void EmptyGridElements();
+
+	void ClearBlendSpaceData();
+
+	void SetBlendSpaceData(const TArray<FBlendSpaceTriangle>& Triangles);
+
+	void ResampleData1D();
+	void ResampleData2D();
+
+
 	/** Get the Editor Element from Index
 	*
 	* @param	XIndex	Index of X
diff --git a/Engine/Source/Runtime/Engine/Classes/Animation/BlendSpace1D.h b/Engine/Source/Runtime/Engine/Classes/Animation/BlendSpace1D.h
--- a/Engine/Source/Runtime/Engine/Classes/Animation/BlendSpace1D.h
+++ b/Engine/Source/Runtime/Engine/Classes/Animation/BlendSpace1D.h
@@ -24,12 +24,15 @@ public:
 	bool bDisplayEditorVertically_DEPRECATED;
 #endif
 
-	/** Drive animation speed by blend input position **/
+	/** Drive animation speed by blend input position */
 	UPROPERTY(EditAnywhere, Category = InputInterpolation)
 	bool bScaleAnimation;
 
 protected:
 	//~ Begin UBlendSpace Interface
 	virtual EBlendSpaceAxis GetAxisToScale() const override;
+#if WITH_EDITOR
+	virtual void SnapSamplesToClosestGridPoint() override;
+#endif
 	//~ End UBlendSpace Interface
 };
diff --git a/Engine/Source/Runtime/Engine/Private/Animation/BlendSpace.cpp b/Engine/Source/Runtime/Engine/Private/Animation/BlendSpace.cpp
--- a/Engine/Source/Runtime/Engine/Private/Animation/BlendSpace.cpp
+++ b/Engine/Source/Runtime/Engine/Private/Animation/BlendSpace.cpp
@@ -8,15 +8,21 @@
 #include "Animation/AnimNotifyQueue.h"
 #include "AnimationUtils.h"
 #include "Animation/BlendSpaceUtilities.h"
+#include "Animation/AnimationPoseData.h"
+#include "Animation/AttributesRuntime.h"
+#include "Animation/BlendSpaceHelpers.h"
+#include "Animation/BlendSpace1DHelpers.h"
 #include "UObject/FrameworkObjectVersion.h"
+#include "UObject/UE5MainStreamObjectVersion.h"
 #include "UObject/UObjectIterator.h"
 #include "Logging/TokenizedMessage.h"
 #include "Logging/MessageLog.h"
-#include "Animation/AnimationPoseData.h"
-#include "Animation/AttributesRuntime.h"
 
 #define LOCTEXT_NAMESPACE "BlendSpace"
 
+// Logs the triangulation/segmentation search to look up the blend samples
+//#define DEBUG_LOG_BLENDSPACE_TRIANGULATION
+
 DECLARE_CYCLE_STAT(TEXT("BlendSpace GetAnimPose"), STAT_BlendSpace_GetAnimPose, STATGROUP_Anim);
 
 /** Scratch buffers for multi-threaded usage */
@@ -56,6 +62,7 @@ void UBlendSpace::PostLoad()
 void UBlendSpace::Serialize(FArchive& Ar)
 {
 	Ar.UsingCustomVersion(FFrameworkObjectVersion::GUID);
+	Ar.UsingCustomVersion(FUE5MainStreamObjectVersion::GUID);
 	Super::Serialize(Ar);
 
 #if WITH_EDITOR
@@ -72,6 +79,11 @@ void UBlendSpace::Serialize(FArchive& Ar)
 			Sample.RateScale = 1.0f;
 		}
 	}
+
+	if (Ar.IsLoading() && (Ar.CustomVer(FUE5MainStreamObjectVersion::GUID) < FUE5MainStreamObjectVersion::BlendSpaceRuntimeTriangulation))
+	{
+		ResampleData();
+	}
 #endif // WITH_EDITOR
 }
 
@@ -88,6 +100,7 @@ void UBlendSpace::PreEditChange(FProperty* PropertyAboutToChange)
 		{
 			PreviousAxisMinMaxValues[AxisIndex].X = BlendParameters[AxisIndex].Min;
 			PreviousAxisMinMaxValues[AxisIndex].Y = BlendParameters[AxisIndex].Max;
+			PreviousGridSpacings[AxisIndex] = (BlendParameters[AxisIndex].Max - BlendParameters[AxisIndex].Min) / BlendParameters[AxisIndex].GridNum;
 		}
 	}
 }
@@ -109,10 +122,33 @@ void UBlendSpace::PostEditChangeProperty(struct FPropertyChangedEvent& PropertyC
 			// Tried and snap samples to points on the grid, those who do not fit or cannot be snapped are marked as invalid
 			SnapSamplesToClosestGridPoint();
 		}
-		else if (PropertyName == GET_MEMBER_NAME_CHECKED(FBlendParameter, Min) || PropertyName == GET_MEMBER_NAME_CHECKED(FBlendParameter, Max))
+		else if (PropertyName == GET_MEMBER_NAME_CHECKED(FBlendParameter, Min))
 		{
-			// Remap the samples to the new values by normalizing the axis and applying the new value range		
-			RemapSamplesToNewAxisRange();
+			// Preserve/enforce the previous grid spacing if snapping
+			for (int32 AxisIndex = 0; AxisIndex != 2; ++AxisIndex)
+			{
+				if (BlendParameters[AxisIndex].bSnapToGrid)
+				{
+					float GridDelta = PreviousGridSpacings[AxisIndex];
+					float Range = (BlendParameters[AxisIndex].Max - BlendParameters[AxisIndex].Min);
+					BlendParameters[AxisIndex].GridNum = FMath::Max(1, (int32)(0.5f + Range / GridDelta));
+					BlendParameters[AxisIndex].Min = BlendParameters[AxisIndex].Max - BlendParameters[AxisIndex].GridNum * GridDelta;
+				}
+			}
+		}
+		else if (PropertyName == GET_MEMBER_NAME_CHECKED(FBlendParameter, Max))
+		{
+			// Preserve/enforce the previous grid spacing if snapping
+			for (int32 AxisIndex = 0; AxisIndex != 2; ++AxisIndex)
+			{
+				if (BlendParameters[AxisIndex].bSnapToGrid)
+				{
+					float GridDelta = PreviousGridSpacings[AxisIndex];
+					float Range = (BlendParameters[AxisIndex].Max - BlendParameters[AxisIndex].Min);
+					BlendParameters[AxisIndex].GridNum = FMath::Max(1, (int32)(0.5f + Range / GridDelta));
+					BlendParameters[AxisIndex].Max = BlendParameters[AxisIndex].Min + BlendParameters[AxisIndex].GridNum * GridDelta;
+				}
+			}
 		}
 	}
 
@@ -121,7 +157,12 @@ void UBlendSpace::PostEditChangeProperty(struct FPropertyChangedEvent& PropertyC
 #endif // WITH_EDITOR
 
 
-bool UBlendSpace::UpdateBlendSamples_Internal(const FVector& InBlendSpacePosition, float InDeltaTime, TArray<FBlendSampleData>& InOutOldSampleDataList, TArray<FBlendSampleData>& InOutSampleDataCache) const
+bool UBlendSpace::UpdateBlendSamples_Internal(
+	const FVector&            InBlendSpacePosition, 
+	float                     InDeltaTime, 
+	TArray<FBlendSampleData>& InOutOldSampleDataList, 
+	TArray<FBlendSampleData>& InOutSampleDataCache, 
+	int32&                    InOutCachedTriangulationIndex) const
 {
 	// For Target weight interpolation, we'll need to save old data, and interpolate to new data
 	TArray<FBlendSampleData>& NewSampleDataList = FBlendSpaceScratchData::Get().NewSampleDataList;
@@ -146,7 +187,7 @@ bool UBlendSpace::UpdateBlendSamples_Internal(const FVector& InBlendSpacePositio
 
 	// get sample data from blendspace
 	bool bSuccessfullySampled = false;
-	if (GetSamplesFromBlendInput(InBlendSpacePosition, NewSampleDataList))
+	if (GetSamplesFromBlendInput(InBlendSpacePosition, NewSampleDataList, InOutCachedTriangulationIndex, true))
 	{
 		// if target weight interpolation is set
 		if (TargetWeightInterpolationSpeedPerSec > 0.f || PerBoneBlend.Num() > 0)
@@ -177,11 +218,11 @@ bool UBlendSpace::UpdateBlendSamples_Internal(const FVector& InBlendSpacePositio
 	return bSuccessfullySampled;
 }
 
-bool UBlendSpace::UpdateBlendSamples(const FVector& InBlendSpacePosition, float InDeltaTime, TArray<FBlendSampleData>& InOutSampleDataCache) const
+bool UBlendSpace::UpdateBlendSamples(const FVector& InBlendSpacePosition, float InDeltaTime, TArray<FBlendSampleData>& InOutSampleDataCache, int32& InOutCachedTriangulationIndex) const
 {
 	TArray<FBlendSampleData>& OldSampleDataList = FBlendSpaceScratchData::Get().OldSampleDataList;
 	check(!OldSampleDataList.Num()); // this must be called non-recursively
-	const bool bResult = UpdateBlendSamples_Internal(InBlendSpacePosition, InDeltaTime, OldSampleDataList, InOutSampleDataCache);
+	const bool bResult = UpdateBlendSamples_Internal(InBlendSpacePosition, InDeltaTime, OldSampleDataList, InOutSampleDataCache, InOutCachedTriangulationIndex);
 	OldSampleDataList.Reset();
 	return bResult;
 }
@@ -207,7 +248,7 @@ void UBlendSpace::TickAssetPlayer(FAnimTickRecord& Instance, struct FAnimNotifyQ
 		const FVector BlendSpacePosition(Instance.BlendSpace.BlendSpacePositionX, Instance.BlendSpace.BlendSpacePositionY, 0.f);
 		const FVector FilteredBlendInput = FilterInput(Instance.BlendSpace.BlendFilter, BlendSpacePosition, DeltaTime);
 
-		if (UpdateBlendSamples_Internal(FilteredBlendInput, DeltaTime, OldSampleDataList, SampleDataList))
+		if (UpdateBlendSamples_Internal(FilteredBlendInput, DeltaTime, OldSampleDataList, SampleDataList, Instance.BlendSpace.TriangulationIndex))
 		{
 			float NewAnimLength = 0.f;
 			float PreInterpAnimLength = 0.f;
@@ -744,73 +785,114 @@ const struct FBlendSample& UBlendSpace::GetBlendSample(const int32 SampleIndex)
 	return SampleData[SampleIndex];
 }
 
-bool UBlendSpace::GetSamplesFromBlendInput(const FVector& BlendInput, TArray<FBlendSampleData>& OutSampleDataList) const
+bool UBlendSpace::GetSamplesFromBlendInput(
+	const FVector& BlendInput, TArray<FBlendSampleData>& OutSampleDataList, int32& InOutCachedTriangulationIndex, bool bCombineAnimations) const
 {
-	TArray<FGridBlendSample, TInlineAllocator<4> >& RawGridSamples = FBlendSpaceScratchData::Get().RawGridSamples;
-	check(!RawGridSamples.Num()); // this must be called non-recursively
-	GetRawSamplesFromBlendInput(BlendInput, RawGridSamples);
+	if (!bInterpolateUsingGrid)
+	{
+		OutSampleDataList.Reset(3);
 
-	OutSampleDataList.Reset();
-	OutSampleDataList.Reserve(RawGridSamples.Num() * FEditorElement::MAX_VERTICES);
+		TArray<FWeightedBlendSample> WeightedBlendSamples;
+		FVector NormalizedBlendInput = GetNormalizedBlendInput(BlendInput);
+		BlendSpaceData.GetSamples(WeightedBlendSamples, DimensionIndices, NormalizedBlendInput, InOutCachedTriangulationIndex);
 
-	// Consolidate all samples
-	for (int32 SampleNum = 0; SampleNum < RawGridSamples.Num(); ++SampleNum)
+		float TotalWeight = 0.0f;
+		for (const FWeightedBlendSample& WeightedBlendSample : WeightedBlendSamples)
+		{
+			float SampleWeight = WeightedBlendSample.SampleWeight;
+			if (SampleWeight > ZERO_ANIMWEIGHT_THRESH)
+			{
+				int32 SampleIndex = WeightedBlendSample.SampleIndex;
+				FBlendSampleData BlendSampleData;
+				BlendSampleData.SampleDataIndex = SampleIndex;
+				BlendSampleData.TotalWeight = SampleWeight;
+				BlendSampleData.Animation = SampleData[SampleIndex].Animation;
+				BlendSampleData.SamplePlayRate = SampleData[SampleIndex].RateScale;
+				OutSampleDataList.Push(BlendSampleData);
+				TotalWeight += SampleWeight;
+			}
+		}
+	}
+	else
 	{
-		FGridBlendSample& GridSample = RawGridSamples[SampleNum];
-		float GridWeight = GridSample.BlendWeight;
-		FEditorElement& GridElement = GridSample.GridElement;
+		TArray<FGridBlendSample, TInlineAllocator<4> >& RawGridSamples = FBlendSpaceScratchData::Get().RawGridSamples;
+		check(!RawGridSamples.Num()); // this must be called non-recursively
+		GetRawSamplesFromBlendInput(BlendInput, RawGridSamples);
 
-		for (int32 Ind = 0; Ind < GridElement.MAX_VERTICES; ++Ind)
+		OutSampleDataList.Reset();
+		OutSampleDataList.Reserve(RawGridSamples.Num() * FEditorElement::MAX_VERTICES);
+
+		// Consolidate all samples
+		for (int32 SampleNum = 0; SampleNum < RawGridSamples.Num(); ++SampleNum)
 		{
-			const int32 SampleDataIndex = GridElement.Indices[Ind];
-			if (SampleData.IsValidIndex(SampleDataIndex))
+			FGridBlendSample& GridSample = RawGridSamples[SampleNum];
+			float GridWeight = GridSample.BlendWeight;
+			FEditorElement& GridElement = GridSample.GridElement;
+
+			for (int32 Ind = 0; Ind < GridElement.MAX_VERTICES; ++Ind)
 			{
-				int32 Index = OutSampleDataList.AddUnique(SampleDataIndex);
-				FBlendSampleData& NewSampleData = OutSampleDataList[Index];
+				const int32 SampleDataIndex = GridElement.Indices[Ind];
+				if (SampleData.IsValidIndex(SampleDataIndex))
+				{
+					int32 Index = OutSampleDataList.AddUnique(SampleDataIndex);
+					FBlendSampleData& NewSampleData = OutSampleDataList[Index];
 
-				NewSampleData.AddWeight(GridElement.Weights[Ind] * GridWeight);
-				NewSampleData.Animation = SampleData[SampleDataIndex].Animation;
-				NewSampleData.SamplePlayRate = SampleData[SampleDataIndex].RateScale;
+					NewSampleData.AddWeight(GridElement.Weights[Ind] * GridWeight);
+					NewSampleData.Animation = SampleData[SampleDataIndex].Animation;
+					NewSampleData.SamplePlayRate = SampleData[SampleDataIndex].RateScale;
+				}
 			}
 		}
+		RawGridSamples.Reset();
 	}
 
-	// At this point we'll only have one of each sample, but different samples can point to the same
-	// animation. We can combine those, making sure to interpolate the parameters like play rate too.
-	for (int32 Index1 = 0; Index1 < OutSampleDataList.Num(); ++Index1)
+	if (bCombineAnimations)
 	{
-		FBlendSampleData& FirstSample = OutSampleDataList[Index1];
-		for (int32 Index2 = Index1 + 1; Index2 < OutSampleDataList.Num(); ++Index2)
+		// At this point we'll only have one of each sample, but different samples can point to the same
+		// animation. We can combine those, making sure to interpolate the parameters like play rate too.
+		for (int32 Index1 = 0; Index1 < OutSampleDataList.Num(); ++Index1)
 		{
-			FBlendSampleData& SecondSample = OutSampleDataList[Index2];
-			// if they have same sample, remove the Index2, and get out
-			if (FirstSample.SampleDataIndex == SecondSample.SampleDataIndex || // Shouldn't happen
-				(FirstSample.Animation != nullptr && FirstSample.Animation == SecondSample.Animation))
+			// Use pointers to make it more obvious what happens if we swap the first and second samples
+			FBlendSampleData* FirstSample = &OutSampleDataList[Index1];
+			for (int32 Index2 = Index1 + 1; Index2 < OutSampleDataList.Num(); ++Index2)
 			{
-				//Calc New Sample Playrate
-				const float TotalWeight = FirstSample.GetWeight() + SecondSample.GetWeight();
-
-				// Only combine playrates if total weight > 0
-				if (!FMath::IsNearlyZero(TotalWeight))
+				FBlendSampleData* SecondSample = &OutSampleDataList[Index2];
+				// if they have same sample, remove the Index2, and get out
+				if (FirstSample->SampleDataIndex == SecondSample->SampleDataIndex || // Shouldn't happen
+					(FirstSample->Animation != nullptr && FirstSample->Animation == SecondSample->Animation))
 				{
-					const float OriginalWeightedPlayRate = FirstSample.SamplePlayRate * (FirstSample.GetWeight() / TotalWeight);
-					const float SecondSampleWeightedPlayRate = SecondSample.SamplePlayRate * (SecondSample.GetWeight() / TotalWeight);
-					FirstSample.SamplePlayRate = OriginalWeightedPlayRate + SecondSampleWeightedPlayRate;
+					//Calc New Sample Playrate
+					const float TotalWeight = FirstSample->GetWeight() + SecondSample->GetWeight();
 
-					// add weight
-					FirstSample.AddWeight(SecondSample.GetWeight());
-				}
+					// Only combine playrates if total weight > 0
+					if (!FMath::IsNearlyZero(TotalWeight))
+					{
+						if (FirstSample->GetWeight() < SecondSample->GetWeight())
+						{
+							// Not strictly necessary, but if we swap here then we keep the one that has a higher
+							// weight, which can make debugging/viewing the blend space more intuitive.
+							OutSampleDataList.Swap(Index1, Index2);
+						}
+
+						const float OriginalWeightedPlayRate = FirstSample->SamplePlayRate * (FirstSample->GetWeight() / TotalWeight);
+						const float SecondSampleWeightedPlayRate = SecondSample->SamplePlayRate * (SecondSample->GetWeight() / TotalWeight);
+						FirstSample->SamplePlayRate = OriginalWeightedPlayRate + SecondSampleWeightedPlayRate;
+
+						// add weight
+						FirstSample->AddWeight(SecondSample->GetWeight());
+					}
 
-				// as for time or previous time will be the master one(Index1)
-				OutSampleDataList.RemoveAtSwap(Index2, 1, false);
-				--Index2;
+					// as for time or previous time will be the master one(Index1)
+					OutSampleDataList.RemoveAtSwap(Index2, 1, false);
+					--Index2;
+				}
 			}
 		}
 	}
 
 	OutSampleDataList.Sort([](const FBlendSampleData& A, const FBlendSampleData& B) { return B.TotalWeight < A.TotalWeight; });
 
-	// remove noisy ones
+	// Remove any below a threshold
 	int32 TotalSample = OutSampleDataList.Num();
 	float TotalWeight = 0.f;
 	for (int32 I = 0; I < TotalSample; ++I)
@@ -830,7 +912,6 @@ bool UBlendSpace::GetSamplesFromBlendInput(const FVector& BlendInput, TArray<FBl
 		// normalize to all weights
 		OutSampleDataList[I].TotalWeight /= TotalWeight;
 	}
-	RawGridSamples.Reset();
 	return (OutSampleDataList.Num() != 0);
 }
 
@@ -876,7 +957,11 @@ void UBlendSpace::ValidateSampleData()
 
 		if (IsAsset())
 		{
-			Sample.bIsValid = ValidateSampleValue(Sample.SampleValue, SampleIndex) && (Sample.Animation != nullptr);
+			bool bAnimationExists = Sample.Animation != nullptr;
+			bool bSampleInBounds = bAnimationExists ? IsSampleWithinBounds(Sample.SampleValue) : true;
+			bool bSampleIsUnique = bAnimationExists ? !IsTooCloseToExistingSamplePoint(Sample.SampleValue, SampleIndex) : true;
+
+			Sample.bIsValid = bAnimationExists && bSampleInBounds && bSampleIsUnique;
 
 			if (Sample.bIsValid)
 			{
@@ -931,9 +1016,30 @@ void UBlendSpace::ValidateSampleData()
 					FMessageLog LoadErrors(NAME_LoadErrors);
 
 					TSharedRef<FTokenizedMessage> Message = LoadErrors.Error();
-					Message->AddToken(FTextToken::Create(LOCTEXT("EmptyAnimationData1", "The BlendSpace ")));
+					Message->AddToken(FTextToken::Create(LOCTEXT("EmptyAnimationData1", "BlendSpace")));
 					Message->AddToken(FAssetNameToken::Create(GetPathName(), FText::FromString(GetName())));
-					Message->AddToken(FTextToken::Create(LOCTEXT("EmptyAnimationData2", " has sample with no animation. Recommend to remove sample point or set new animation.")));
+					if (!bAnimationExists)
+					{
+						Message->AddToken(FTextToken::Create(LOCTEXT("EmptyAnimationData2", "has a sample with no/invalid animation. Recommend to remove sample point or set new animation.")));
+					}
+					else if (!bSampleInBounds)
+					{
+						Message->AddToken(FTextToken::Create(LOCTEXT("EmptyAnimationData2", "has sample")));
+						Message->AddToken(FAssetNameToken::Create(Sample.Animation->GetPathName(), FText::FromString(Sample.Animation->GetName())));
+						Message->AddToken(FTextToken::Create(LOCTEXT("EmptyAnimationData3", "that is invalid due to being out of bounds. Recommend adjusting it.")));
+					}
+					else if (!bSampleIsUnique)
+					{
+						Message->AddToken(FTextToken::Create(LOCTEXT("EmptyAnimationData2", "has sample")));
+						Message->AddToken(FAssetNameToken::Create(Sample.Animation->GetPathName(), FText::FromString(Sample.Animation->GetName())));
+						Message->AddToken(FTextToken::Create(LOCTEXT("EmptyAnimationData3", "that is not unique. Recommend adjusting it.")));
+					}
+					else
+					{
+						// Shouldn't get here
+						Message->AddToken(FAssetNameToken::Create(Sample.Animation->GetPathName(), FText::FromString(Sample.Animation->GetName())));
+						Message->AddToken(FTextToken::Create(LOCTEXT("EmptyAnimationData2", "is invalid.")));
+					}
 					LoadErrors.Notify();
 				}
 			}
@@ -985,7 +1091,7 @@ bool UBlendSpace::AddSample(UAnimSequence* AnimationSequence, const FVector& Sam
 	return bValidSampleData;
 }
 
-bool UBlendSpace::EditSampleValue(const int32 BlendSampleIndex, const FVector& NewValue, bool bSnap)
+bool UBlendSpace::EditSampleValue(const int32 BlendSampleIndex, const FVector& NewValue)
 {
 	const bool bValidValue = SampleData.IsValidIndex(BlendSampleIndex) && ValidateSampleValue(NewValue, BlendSampleIndex);
 
@@ -994,7 +1100,6 @@ bool UBlendSpace::EditSampleValue(const int32 BlendSampleIndex, const FVector& N
 		// Set new value if it passes the tests
 		SampleData[BlendSampleIndex].SampleValue = NewValue;
 		SampleData[BlendSampleIndex].bIsValid = bValidValue;
-		SampleData[BlendSampleIndex].bSnapToGrid = bSnap;
 	}
 
 	return bValidValue;
@@ -1054,7 +1159,12 @@ const TArray<FEditorElement>& UBlendSpace::GetGridSamples() const
 	return GridSamples;
 }
 
-void UBlendSpace::FillupGridElements(const TArray<int32>& PointListToSampleIndices, const TArray<FEditorElement>& GridElements, const TArray<int32>& InDimensionIndices)
+const FBlendSpaceData& UBlendSpace::GetBlendSpaceData() const
+{
+	return BlendSpaceData;
+}
+
+void UBlendSpace::FillupGridElements(const TArray<FEditorElement>& GridElements, const TArray<int32>& InDimensionIndices)
 {
 	DimensionIndices = InDimensionIndices;
 
@@ -1069,9 +1179,9 @@ void UBlendSpace::FillupGridElements(const TArray<int32>& PointListToSampleIndic
 		for (int32 VertexIndex = 0; VertexIndex < FEditorElement::MAX_VERTICES; ++VertexIndex)
 		{
 			const int32 SampleIndex = ViewGrid.Indices[VertexIndex];
-			if (SampleIndex != INDEX_NONE && PointListToSampleIndices.IsValidIndex(SampleIndex))
+			if (SampleIndex != INDEX_NONE)
 			{
-				NewGrid.Indices[VertexIndex] = PointListToSampleIndices[SampleIndex];
+				NewGrid.Indices[VertexIndex] = SampleIndex;
 			}
 			else
 			{
@@ -1265,7 +1375,7 @@ FVector UBlendSpace::GetClampedAndWrappedBlendInput(const FVector& BlendInput) c
 	return AdjustedInput;
 }
 
-FVector UBlendSpace::GetNormalizedBlendInput(const FVector& BlendInput) const
+FVector UBlendSpace::ConvertBlendInputToGridSpace(const FVector& BlendInput) const
 {
 	FVector AdjustedInput = GetClampedAndWrappedBlendInput(BlendInput);
 
@@ -1276,6 +1386,17 @@ FVector UBlendSpace::GetNormalizedBlendInput(const FVector& BlendInput) const
 	return NormalizedBlendInput;
 }
 
+FVector UBlendSpace::GetNormalizedBlendInput(const FVector& BlendInput) const
+{
+	FVector AdjustedInput = GetClampedAndWrappedBlendInput(BlendInput);
+
+	const FVector MinBlendInput = FVector(BlendParameters[0].Min, BlendParameters[1].Min, BlendParameters[2].Min);
+	const FVector MaxBlendInput = FVector(BlendParameters[0].Max, BlendParameters[1].Max, BlendParameters[2].Max);
+
+	FVector NormalizedBlendInput = (AdjustedInput - MinBlendInput) / (MaxBlendInput - MinBlendInput);
+	return NormalizedBlendInput;
+}
+
 const FEditorElement* UBlendSpace::GetGridSampleInternal(int32 Index) const
 {
 	return GridSamples.IsValidIndex(Index) ? &GridSamples[Index] : NULL;
@@ -1495,9 +1616,9 @@ bool UBlendSpace::ContainsMatchingSamples(EAdditiveAnimationType AdditiveType) c
 bool UBlendSpace::IsSameSamplePoint(const FVector& SamplePointA, const FVector& SamplePointB) const
 {
 #if 1
-	return FMath::IsNearlyEqual(SamplePointA.X, SamplePointB.X)
-		&& FMath::IsNearlyEqual(SamplePointA.Y, SamplePointB.Y)
-		&& FMath::IsNearlyEqual(SamplePointA.Z, SamplePointB.Z);
+	return FMath::IsNearlyEqual(SamplePointA.X, SamplePointB.X, KINDA_SMALL_NUMBER)
+		&& FMath::IsNearlyEqual(SamplePointA.Y, SamplePointB.Y, KINDA_SMALL_NUMBER)
+		&& FMath::IsNearlyEqual(SamplePointA.Z, SamplePointB.Z, KINDA_SMALL_NUMBER);
 #else
 	if (DimensionIndices.Num() == 0 || DimensionIndices.Num() > 3)
 	{
@@ -1577,43 +1698,6 @@ void UBlendSpace::GetRawSamplesFromBlendInput(const FVector& BlendInput, TArray<
 	}
 }
 
-#if WITH_EDITOR
-void UBlendSpace::SnapSamplesToClosestGridPoint()
-{
-	switch (DimensionIndices.Num())
-	{
-	case 1:
-		SnapSamplesToClosestGridPoint1D();
-		break;
-	case 2:
-		SnapSamplesToClosestGridPoint2D();
-		break;
-	default:
-		UE_LOG(LogAnimation, Warning, TEXT("Unhandled dimensionality in samples: %d"), DimensionIndices.Num());
-		break;
-	}
-}
-
-void UBlendSpace::RemapSamplesToNewAxisRange()
-{
-	switch (DimensionIndices.Num())
-	{
-	case 1:
-		RemapSamplesToNewAxisRange1D();
-		break;
-	case 2:
-		RemapSamplesToNewAxisRange2D();
-		break;
-	default:
-		UE_LOG(LogAnimation, Warning, TEXT("Unhandled dimensionality in samples: %d"), DimensionIndices.Num());
-		break;
-	}
-}
-
-#endif // WITH_EDITOR
-
-
-
 /*-----------------------------------------------------------------------------
 	1D functionality.
 -----------------------------------------------------------------------------*/
@@ -1623,7 +1707,7 @@ void UBlendSpace::GetRawSamplesFromBlendInput1D(const FVector& BlendInput, TArra
 	check(DimensionIndices.Num() == 1);
 	int32 Index0 = DimensionIndices[0];
 
-	FVector NormalizedBlendInput = GetNormalizedBlendInput(BlendInput);
+	FVector NormalizedBlendInput = ConvertBlendInputToGridSpace(BlendInput);
 
 	float GridIndex = FMath::TruncToFloat(NormalizedBlendInput[Index0]);
 	float Remainder = NormalizedBlendInput[Index0] - GridIndex;
@@ -1665,116 +1749,6 @@ void UBlendSpace::GetRawSamplesFromBlendInput1D(const FVector& BlendInput, TArra
 	}
 }
 
-#if WITH_EDITOR
-
-void UBlendSpace::SnapSamplesToClosestGridPoint1D()
-{
-	check(DimensionIndices.Num() == 1);
-	int32 Index0 = DimensionIndices[0];
-
-	TArray<float> GridPoints;
-	TArray<int32> ClosestSampleToGridPoint;
-	TArray<bool> SampleDataOnPoints;
-	TArray<bool> ShouldSnap;
-
-	const float GridMin = BlendParameters[Index0].Min;
-	const float GridMax = BlendParameters[Index0].Max;
-	const float GridRange = GridMax - GridMin;
-	const int32 NumGridPoints = BlendParameters[Index0].GridNum + 1;
-	const float GridStep = GridRange / BlendParameters[Index0].GridNum;
-
-	// First mark all samples as invalid and s unsnapped, recording whether they were snapped before
-	ShouldSnap.SetNumZeroed(SampleData.Num());
-	for (int32 BlendSampleIndex = 0; BlendSampleIndex < SampleData.Num(); ++BlendSampleIndex)
-	{
-		FBlendSample& BlendSample = SampleData[BlendSampleIndex];
-		ShouldSnap[BlendSampleIndex] = BlendSample.bSnapToGrid;
-		BlendSample.bSnapToGrid = false;
-	}
-
-	for (int32 GridPointIndex = 0; GridPointIndex < NumGridPoints; ++GridPointIndex)
-	{
-		const float GridPointValue = (GridPointIndex * GridStep) + GridMin;
-		GridPoints.Add(GridPointValue);
-	}
-
-	ClosestSampleToGridPoint.Init(INDEX_NONE, GridPoints.Num());
-
-	// Find closest sample to grid point
-	for (int32 PointIndex = 0; PointIndex < GridPoints.Num(); ++PointIndex)
-	{
-		const float GridPoint = GridPoints[PointIndex];
-		float SmallestDistance = FLT_MAX;
-		int32 Index = INDEX_NONE;
-
-		for (int32 SampleIndex = 0; SampleIndex < SampleData.Num(); ++SampleIndex)
-		{
-			FBlendSample& BlendSample = SampleData[SampleIndex];
-			const float Distance = FMath::Abs(GridPoint - BlendSample.SampleValue[Index0]);
-			if (Distance < SmallestDistance)
-			{
-				Index = SampleIndex;
-				SmallestDistance = Distance;
-			}
-		}
-
-		ClosestSampleToGridPoint[PointIndex] = Index;
-	}
-
-	// Find closest grid point to sample
-	for (int32 SampleIndex = 0; SampleIndex < SampleData.Num(); ++SampleIndex)
-	{
-		if (ShouldSnap[SampleIndex])
-		{
-			FBlendSample& BlendSample = SampleData[SampleIndex];
-
-			// Find closest grid point
-			float SmallestDistance = FLT_MAX;
-			int32 Index = INDEX_NONE;
-			for (int32 PointIndex = 0; PointIndex < GridPoints.Num(); ++PointIndex)
-			{
-				const float Distance = FMath::Abs(GridPoints[PointIndex] - BlendSample.SampleValue[Index0]);
-				if (Distance < SmallestDistance)
-				{
-					Index = PointIndex;
-					SmallestDistance = Distance;
-				}
-			}
-
-			// Only move the sample if it is also closest to the grid point
-			if (Index != INDEX_NONE && ClosestSampleToGridPoint[Index] == SampleIndex)
-			{
-				BlendSample.SampleValue[Index0] = GridPoints[Index];
-				BlendSample.bSnapToGrid = true;
-			}
-		}
-	}
-}
-
-void UBlendSpace::RemapSamplesToNewAxisRange1D()
-{
-	check(DimensionIndices.Num() == 1);
-	int32 Index0 = DimensionIndices[0];
-
-	const float OldGridMin = PreviousAxisMinMaxValues[Index0].X;
-	const float OldGridMax = PreviousAxisMinMaxValues[Index0].Y;
-	const float OldGridRange = OldGridMax - OldGridMin;
-
-	const float NewGridMin = BlendParameters[Index0].Min;
-	const float NewGridMax = BlendParameters[Index0].Max;
-	const float NewGridRange = NewGridMax - NewGridMin;
-
-	for (FBlendSample& BlendSample : SampleData)
-	{
-		const float NormalizedValue = (BlendSample.SampleValue[Index0] - OldGridMin) / OldGridRange;
-		BlendSample.SampleValue[Index0] = NewGridMin + (NormalizedValue * NewGridRange);
-	}
-}
-
-
-#endif // WITH_EDITOR
-
-
 
 /*-----------------------------------------------------------------------------
 	2D functionality.
@@ -1798,7 +1772,7 @@ void UBlendSpace::GetRawSamplesFromBlendInput2D(const FVector& BlendInput, TArra
 	FGridBlendSample& LeftTop = OutBlendSamples[2];
 	FGridBlendSample& RightTop = OutBlendSamples[3];
 
-	const FVector NormalizedBlendInput = GetNormalizedBlendInput(BlendInput);
+	const FVector NormalizedBlendInput = ConvertBlendInputToGridSpace(BlendInput);
 	const FVector GridIndex(FMath::TruncToFloat(NormalizedBlendInput.X), FMath::TruncToFloat(NormalizedBlendInput.Y), 0.f);
 	const FVector Remainder = NormalizedBlendInput - GridIndex;
 	// bi-linear very simple interpolation
@@ -1855,112 +1829,587 @@ void UBlendSpace::GetRawSamplesFromBlendInput2D(const FVector& BlendInput, TArra
 
 #if WITH_EDITOR
 
-void UBlendSpace::SnapSamplesToClosestGridPoint2D()
+// Note that this runs and is needed if the axes settings are change to enable snapping - then we do
+// a pass over all the samples and make sure everything is in order. When samples are being dragged
+// around/dropped then that moving sample will be handled in SBlendSpaceGridWidget so, whilst this
+// gets called, it shouldn't really need to do anything.
+void UBlendSpace::SnapSamplesToClosestGridPoint()
 {
-	check(DimensionIndices.Num() == 2);
+	if (BlendParameters[0].bSnapToGrid && BlendParameters[1].bSnapToGrid)
+	{
+		TArray<FVector> GridPoints;
+		TArray<int32> ClosestSampleToGridPoint;
 
-	TArray<FVector> GridPoints;
-	TArray<int32> ClosestSampleToGridPoint;
-	TArray<bool> SampleDataOnPoints;
-	TArray<bool> ShouldSnap;
+		const FVector GridMin(BlendParameters[0].Min, BlendParameters[1].Min, 0.0f);
+		const FVector GridMax(BlendParameters[0].Max, BlendParameters[1].Max, 0.0f);
+		const FVector GridRange(GridMax.X - GridMin.X, GridMax.Y - GridMin.Y, 0.0f);
+		const FIntPoint NumGridDivisions(BlendParameters[0].GridNum, BlendParameters[1].GridNum);
+		const FVector GridStep(GridRange.X / BlendParameters[0].GridNum, GridRange.Y / BlendParameters[1].GridNum, 0.0f);
 
-	const FVector GridMin(BlendParameters[0].Min, BlendParameters[1].Min, 0.0f);
-	const FVector GridMax(BlendParameters[0].Max, BlendParameters[1].Max, 0.0f);
-	const FVector GridRange(GridMax.X - GridMin.X, GridMax.Y - GridMin.Y, 0.0f);
-	const FIntPoint NumGridPoints(BlendParameters[0].GridNum + 1, BlendParameters[1].GridNum + 1);
-	const FVector GridStep(GridRange.X / BlendParameters[0].GridNum, GridRange.Y / BlendParameters[1].GridNum, 0.0f);
-	
-	// First mark all samples as unsnapped, recording whether they were snapped before
-	ShouldSnap.SetNumZeroed(SampleData.Num());
-	for (int32 BlendSampleIndex = 0; BlendSampleIndex < SampleData.Num(); ++BlendSampleIndex)
-	{
-		FBlendSample& BlendSample = SampleData[BlendSampleIndex];
-		ShouldSnap[BlendSampleIndex] = BlendSample.bSnapToGrid;
-		BlendSample.bSnapToGrid = false;
-	}
+		// Snap to nearest in normalized space - not depending on the units of the params (which may be completely different).
+		const float GridRatio = GridStep.Y / GridStep.X;
 
-	for (int32 GridY = 0; GridY < NumGridPoints.Y; ++GridY)
-	{
-		for (int32 GridX = 0; GridX < NumGridPoints.X; ++GridX)
+		for (int32 GridY = 0; GridY < NumGridDivisions.Y + 1; ++GridY)
 		{
-			const FVector GridPoint((GridX * GridStep.X) + GridMin.X, (GridY * GridStep.Y) + GridMin.Y, 0.0f);
-			GridPoints.Add(GridPoint);
-		}
-	}
+			for (int32 GridX = 0; GridX < NumGridDivisions.X + 1; ++GridX)
+			{
+				const FVector GridPoint(
+					GridMin.X + (GridX * GridRange.X / NumGridDivisions.X), 
+					GridMin.Y + (GridY * GridRange.Y / NumGridDivisions.Y),
+					0.0f);
 
-	ClosestSampleToGridPoint.Init(INDEX_NONE, GridPoints.Num());
+				//const FVector GridPoint((GridX * GridStep.X) + GridMin.X, (GridY * GridStep.Y) + GridMin.Y, 0.0f);
+				GridPoints.Add(GridPoint);
+			}
+		}
 
-	// Find closest sample to grid point
-	for (int32 PointIndex = 0; PointIndex < GridPoints.Num(); ++PointIndex)
-	{
-		const FVector& GridPoint = GridPoints[PointIndex];
-		float SmallestDistance = FLT_MAX;
-		int32 Index = INDEX_NONE;
+		ClosestSampleToGridPoint.Init(INDEX_NONE, GridPoints.Num());
 
-		for (int32 SampleIndex = 0; SampleIndex < SampleData.Num(); ++SampleIndex)
+		// Find closest sample to grid point
+		for (int32 PointIndex = 0; PointIndex < GridPoints.Num(); ++PointIndex)
 		{
-			FBlendSample& BlendSample = SampleData[SampleIndex];
-			const float Distance = (GridPoint - BlendSample.SampleValue).SizeSquared2D();
-			if (Distance < SmallestDistance)
+			const FVector& GridPoint = GridPoints[PointIndex];
+			float SmallestDistanceSq = FLT_MAX;
+			int32 Index = INDEX_NONE;
+
+			for (int32 SampleIndex = 0; SampleIndex < SampleData.Num(); ++SampleIndex)
 			{
-				Index = SampleIndex;
-				SmallestDistance = Distance;
+				FBlendSample& BlendSample = SampleData[SampleIndex];
+				FVector Delta = GridPoint - BlendSample.SampleValue;
+				Delta.X *= GridRatio;
+				const float DistanceSq = Delta.SizeSquared2D();
+				if (DistanceSq < SmallestDistanceSq)
+				{
+					Index = SampleIndex;
+					SmallestDistanceSq = DistanceSq;
+				}
 			}
-		}
 
-		ClosestSampleToGridPoint[PointIndex] = Index;
-	}
+			ClosestSampleToGridPoint[PointIndex] = Index;
+		}
 
-	// Find closest grid point to sample
-	for (int32 SampleIndex = 0; SampleIndex < SampleData.Num(); ++SampleIndex)
-	{
-		if(ShouldSnap[SampleIndex])
+		// Find closest grid point to sample
+		for (int32 SampleIndex = 0; SampleIndex < SampleData.Num(); ++SampleIndex)
 		{
 			FBlendSample& BlendSample = SampleData[SampleIndex];
 
 			// Find closest grid point
-			float SmallestDistance = FLT_MAX;
+			float SmallestDistanceSq = FLT_MAX;
 			int32 Index = INDEX_NONE;
 			for (int32 PointIndex = 0; PointIndex < GridPoints.Num(); ++PointIndex)
 			{
-				const float Distance = (GridPoints[PointIndex] - BlendSample.SampleValue).SizeSquared2D();
-				if (Distance < SmallestDistance)
+				FVector Delta = GridPoints[PointIndex] - BlendSample.SampleValue;
+				Delta.X *= GridRatio;
+				const float DistanceSq = Delta.SizeSquared2D();
+				if (DistanceSq < SmallestDistanceSq)
 				{
 					Index = PointIndex;
-					SmallestDistance = Distance;
+					SmallestDistanceSq = DistanceSq;
 				}
 			}
 
 			// Only move the sample if it is also closest to the grid point
 			if (Index != INDEX_NONE && ClosestSampleToGridPoint[Index] == SampleIndex)
 			{
-				BlendSample.SampleValue = GridPoints[Index];
-				BlendSample.bSnapToGrid = true;
+				for (int32 AxisIndex = 0; AxisIndex != 2; ++AxisIndex)
+				{
+					BlendSample.SampleValue[AxisIndex] = GridPoints[Index][AxisIndex];
+				}
 			}
 		}
 	}
+	else if (BlendParameters[0].bSnapToGrid || BlendParameters[1].bSnapToGrid)
+	{
+		// We only snap on one axis, but need to make sure we don't collapse samples on top of each
+		// other. Just snap to the nearest grid value, unless it would result in a duplicate
+
+		int32 AxisIndex = BlendParameters[0].bSnapToGrid ? 0 : 1;
+
+		float GridMin = BlendParameters[AxisIndex].Min;
+		float GridMax = BlendParameters[AxisIndex].Max;
+		int32 GridNum = BlendParameters[AxisIndex].GridNum;
+		float GridDelta = (GridMax - GridMin) / GridNum;
+
+		for (int32 SampleIndex = 0; SampleIndex < SampleData.Num(); ++SampleIndex)
+		{
+			FBlendSample& BlendSample = SampleData[SampleIndex];
+			FVector NewSampleValue = BlendSample.SampleValue;
+			float FGridPosition = (NewSampleValue[AxisIndex] - GridMin) / GridDelta;
+			int32 IGridPosition = FMath::Clamp((int)(FGridPosition + 0.5f), 0, GridNum + 1);
+			NewSampleValue[AxisIndex] = GridMin + IGridPosition * GridDelta;
+
+			bool bFoundOverlap = false;
+			for (int32 SampleIndex1 = 0; SampleIndex1 != SampleData.Num(); ++SampleIndex1)
+			{
+				if (IsSameSamplePoint(NewSampleValue, SampleData[SampleIndex1].SampleValue))
+				{
+					bFoundOverlap = true;
+					break;
+				}
+			}
+			if (!bFoundOverlap)
+			{
+				BlendSample.SampleValue = NewSampleValue;
+			}
+		}
+	}
+}
+
+void UBlendSpace::ClearBlendSpaceData()
+{
+	BlendSpaceData.Empty();
+}
+
+void UBlendSpace::SetBlendSpaceData(const TArray<FBlendSpaceTriangle>& Triangles)
+{
+	ClearBlendSpaceData();
+	BlendSpaceData.Triangles = Triangles;
+}
+
+void UBlendSpace::SetBlendSpaceData(const TArray<FBlendSpaceSegment>& Segments)
+{
+	ClearBlendSpaceData();
+	BlendSpaceData.Segments = Segments;
+}
+
+void UBlendSpace::ResampleData()
+{
+	ClearBlendSpaceData();
+
+	ValidateSampleData();
+
+	FBox AABB(ForceInit);
+	for (const FBlendSample& Sample : SampleData)
+	{
+		// Add X value from sample (this is the only valid value to be set for 1D blend spaces / aim offsets
+		if (Sample.bIsValid)
+		{
+			AABB += Sample.SampleValue;
+		}
+	}
+
+	DimensionIndices.Reset(3);
+	if (AABB.GetExtent().X > 0.0f)
+	{
+		DimensionIndices.Push(0);
+	}
+	if (AABB.GetExtent().Y > 0.0f)
+	{
+		DimensionIndices.Push(1);
+	}
+	if (AABB.GetExtent().Z > 0.0f)
+	{
+		DimensionIndices.Push(2);
+	}
+
+	// Handle the situation where there is just one point
+	if (DimensionIndices.Num() == 0)
+	{
+		DimensionIndices.Push(0);
+	}
+
+	if (DimensionIndices.Num() == 1)
+	{
+		ResampleData1D();
+	}
+	else if (DimensionIndices.Num() == 2)
+	{
+		ResampleData2D();
+	}
+	else
+	{
+		UE_LOG(LogAnimation, Warning, TEXT("Found %d dimensions from the samples"), DimensionIndices.Num());
+	}
+}
+
+
+void UBlendSpace::ResampleData1D()
+{
+	check(DimensionIndices.Num() == 1);
+	FLineElementGenerator LineElementGenerator;
+
+	int32 Index0 = DimensionIndices[0];
+
+	const FBlendParameter& BlendParameter = GetBlendParameter(Index0);
+	LineElementGenerator.Init(BlendParameter);
+
+	UE_LOG(LogAnimation, Log, TEXT("Resampling data in 1D - %d samples"), SampleData.Num());
+
+	if (SampleData.Num())
+	{
+		for (int32 SampleIndex = 0; SampleIndex < SampleData.Num(); ++SampleIndex)
+		{
+			const FBlendSample& Sample = SampleData[SampleIndex];
+			// Add X value from sample (this is the only valid value to be set for 1D blend spaces / aim offsets
+			if (Sample.bIsValid)
+			{
+				LineElementGenerator.SamplePointList.Add(FLineVertex(Sample.SampleValue[Index0], SampleIndex));
+			}
+		}
+
+		LineElementGenerator.Process();
+
+		if (bInterpolateUsingGrid)
+		{
+			LineElementGenerator.CalculateEditorElements();
+			FillupGridElements(LineElementGenerator.EditorElements, DimensionIndices);
+			ClearBlendSpaceData();
+		}
+		else
+		{
+			SetBlendSpaceData(LineElementGenerator.CalculateSegments());
+			EmptyGridElements();
+		}
+	}
 }
 
 
-void UBlendSpace::RemapSamplesToNewAxisRange2D()
+void UBlendSpace::ResampleData2D()
 {
 	check(DimensionIndices.Num() == 2);
 
-	const FVector OldGridMin(PreviousAxisMinMaxValues[0].X, PreviousAxisMinMaxValues[1].X, 0.0f);
-	const FVector OldGridMax(PreviousAxisMinMaxValues[0].Y, PreviousAxisMinMaxValues[1].Y, 1.0f);
-	const FVector OldGridRange = OldGridMax - OldGridMin;
+	// TODO for now, Index0 will always be 0 and Index1 will always be 1, since we don't support 3D.
+	// However, if/when we support authoring using a third dimension, we could get here with any 2D
+	// combination - e.g. XY, XZ or YZ. Then we can either make the triangulation code handle the
+	// indexing, or tweak the triangles going into and out of it.
+	int32 Index0 = DimensionIndices[0];
+	int32 Index1 = DimensionIndices[1];
+
+	// clear first
+	FDelaunayTriangleGenerator DelaunayTriangleGenerator;
 
-	const FVector NewGridMin(BlendParameters[0].Min, BlendParameters[1].Min, 0.0f);
-	const FVector NewGridMax(BlendParameters[0].Max, BlendParameters[1].Max, 1.0f);
-	const FVector NewGridRange = NewGridMax - NewGridMin;
+	// you don't like to overwrite the link here (between visible points vs sample points, 
+	// so allow this if no triangle is generated
+	const FBlendParameter& BlendParamX = GetBlendParameter(Index0);
+	const FBlendParameter& BlendParamY = GetBlendParameter(Index1);
+	FBlendSpaceGrid	BlendSpaceGrid;
+	BlendSpaceGrid.SetGridInfo(BlendParamX, BlendParamY);
+	DelaunayTriangleGenerator.SetGridBox(BlendParamX, BlendParamY);
 
-	for (FBlendSample& BlendSample : SampleData)
+	EmptyGridElements();
+
+	UE_LOG(LogAnimation, Log, TEXT("Resampling data in 2D - %d samples"), GetNumberOfBlendSamples());
+	if (GetNumberOfBlendSamples())
 	{
-		const FVector NormalizedValue = (BlendSample.SampleValue - OldGridMin) / OldGridRange;
-		BlendSample.SampleValue = NewGridMin + (NormalizedValue * NewGridRange);
+		bool bAllSamplesValid = true;
+		for (int32 SampleIndex = 0; SampleIndex < GetNumberOfBlendSamples(); ++SampleIndex)
+		{
+			const FBlendSample& Sample = GetBlendSample(SampleIndex);
+
+			// Do not add invalid sample points (user will need to correct them to be incorporated into the blendspace)
+			if (Sample.bIsValid)
+			{
+				DelaunayTriangleGenerator.AddSamplePoint(FVector2D(Sample.SampleValue[Index0], Sample.SampleValue[Index1]), SampleIndex);
+			}
+		}
+
+		// triangulate
+		DelaunayTriangleGenerator.Triangulate(PreferredTriangulationDirection);
+
+		if (bInterpolateUsingGrid)
+		{
+			// once triangulated, generate grid
+			const TArray<FVertex>& Points = DelaunayTriangleGenerator.GetSamplePointList();
+			const TArray<FTriangle*>& Triangles = DelaunayTriangleGenerator.GetTriangleList();
+			BlendSpaceGrid.GenerateGridElements(Points, Triangles);
+
+			// now fill up grid elements in BlendSpace using this Element information
+			if (Triangles.Num() > 0)
+			{
+				const TArray<FEditorElement>& GridElements = BlendSpaceGrid.GetElements();
+				FillupGridElements(GridElements, DimensionIndices);
+			}
+
+			ClearBlendSpaceData();
+		}
+		else
+		{
+			SetBlendSpaceData(DelaunayTriangleGenerator.CalculateTriangles());
+			EmptyGridElements();
+		}
+
 	}
 }
+
 #endif // WITH_EDITOR
 
+void FBlendSpaceData::GetSamples1D(
+	TArray<FWeightedBlendSample>& OutWeightedSamples,
+	const TArray<int32>&          InDimensionIndices,
+	const FVector&                InNormalizedSamplePosition,
+	int32&                        InOutSegmentIndex) const
+{
+	check(InDimensionIndices.Num() == 1);
+	OutWeightedSamples.Reset(2);
+
+	if (!Segments.Num())
+	{
+		return;
+	}
+	else if (Segments.Num() == 1 && Segments[0].SampleIndices[0] == Segments[0].SampleIndices[1])
+	{
+		OutWeightedSamples.Push(FWeightedBlendSample(Segments[0].SampleIndices[0], 1.0f));
+		return;
+	}
+
+	int32 Index0 = InDimensionIndices[0];
+	float P = InNormalizedSamplePosition[Index0];
+
+#ifdef DEBUG_LOG_BLENDSPACE_TRIANGULATION
+	static int32 Count = 0;
+	UE_LOG(LogAnimation, Warning, TEXT("%d Starting segment search with cached index %d"), Count++, InOutSegmentIndex);
+#endif
+
+	if (InOutSegmentIndex < 0 || InOutSegmentIndex >= Segments.Num())
+	{
+		InOutSegmentIndex = Segments.Num() / 2;
+	}
+
+	for (int32 Attempt = 0; Attempt != Segments.Num(); ++Attempt)
+	{
+#ifdef DEBUG_LOG_BLENDSPACE_TRIANGULATION
+		UE_LOG(LogAnimation, Warning, TEXT("Segment index: %d"), InOutSegmentIndex);
+#endif
+		const FBlendSpaceSegment& Segment = Segments[InOutSegmentIndex];
+		if (P < Segment.Vertices[0])
+		{
+			// Need to go left
+			if (InOutSegmentIndex > 0)
+			{
+				--InOutSegmentIndex;
+			}
+			else
+			{
+				OutWeightedSamples.Push(FWeightedBlendSample(Segment.SampleIndices[0], 1.0f));
+				return;
+			}
+		}
+		else if (P > Segment.Vertices[1])
+		{
+			// Need to go right
+			if (InOutSegmentIndex < Segments.Num() - 1)
+			{
+				++InOutSegmentIndex;
+			}
+			else
+			{
+				OutWeightedSamples.Push(FWeightedBlendSample(Segment.SampleIndices[1], 1.0f));
+				return;
+			}
+		}
+		else
+		{
+			// We're in the segment
+			float P0 = Segment.Vertices[0];
+			float P1 = Segment.Vertices[1];
+			float Frac = FMath::Clamp((P - P0) / (P1 - P0), 0.0f, 1.0f); // Robust to dividing by zero
+			OutWeightedSamples.Push(FWeightedBlendSample(Segment.SampleIndices[0], 1.0f - Frac));
+			OutWeightedSamples.Push(FWeightedBlendSample(Segment.SampleIndices[1], Frac));
+			return;
+		}
+	}
+	UE_LOG(LogAnimation, Warning, TEXT("Unable to find BlendSpace segment"));
+}
+
+void FBlendSpaceData::GetSamples2D(
+	TArray<FWeightedBlendSample>& OutWeightedSamples,
+	const TArray<int32>&          InDimensionIndices,
+	const FVector&                InNormalizedSamplePosition,
+	int32&                        InOutTriangleIndex) const
+{
+	if (Triangles.Num() == 0)
+	{
+		return;
+	}
+
+	check(InDimensionIndices.Num() == 2);
+	OutWeightedSamples.Reset(3);
+
+	// Do an incremental search by tracking through the triangle edges (with the normals). This will
+	// be guaranteed to work since the overall triangulation region is convex. Note that the
+	// triangulation may not include the query point, so if we're not in any triangle we need to
+	// return the closest point in the triangulated region.
+	int32 Index0 = InDimensionIndices[0];
+	int32 Index1 = InDimensionIndices[1];
+	FVector2D P(InNormalizedSamplePosition[Index0], InNormalizedSamplePosition[Index1]);
+
+	// Special case for when there's a single triangle and it is degenerate
+	if (Triangles.Num() == 1)
+	{
+		const FBlendSpaceTriangle& Triangle = Triangles[0];
+		if (Triangle.SampleIndices[0] == Triangle.SampleIndices[1])
+		{
+			// Single point
+			OutWeightedSamples.Push(FWeightedBlendSample(Triangle.SampleIndices[0], 1.0f));
+			return;
+		}
+		else if (Triangle.SampleIndices[1] == Triangle.SampleIndices[2])
+		{
+			// Two points - blend linearly
+			FVector2D Delta = Triangle.Vertices[1] - Triangle.Vertices[0];
+			float T = ((P - Triangle.Vertices[0]) | Delta) / Delta.SizeSquared();
+			if (T < 0)
+			{
+				OutWeightedSamples.Push(FWeightedBlendSample(Triangle.SampleIndices[0], 1.0f));
+			}
+			else if (T > 1.0f)
+			{
+				OutWeightedSamples.Push(FWeightedBlendSample(Triangle.SampleIndices[1], 1.0f));
+			}
+			else
+			{
+				OutWeightedSamples.Push(FWeightedBlendSample(Triangle.SampleIndices[0], 1.0f - T));
+				OutWeightedSamples.Push(FWeightedBlendSample(Triangle.SampleIndices[1], T));
+			}
+			return;
+		}
+	}
+
+#ifdef DEBUG_LOG_BLENDSPACE_TRIANGULATION
+	static int32 Count = 0;
+	UE_LOG(LogAnimation, Warning, TEXT("%d Starting triangulation search with cached index %d"), Count++, InOutTriangleIndex);
+#endif
+
+	// Where available, start from the the previous/cached result. Also see
+	// https://www.cl.cam.ac.uk/techreports/UCAM-CL-TR-728.pdf for a method which also works for
+	// non-Delaunay triangulations (might be useful if there are precision problems).
+	if (InOutTriangleIndex < 0 || InOutTriangleIndex >= Triangles.Num())
+	{
+		InOutTriangleIndex = Triangles.Num() / 2;
+	}
+	// We should never need more than the number of triangles - if we do then we got caught in a
+	// loop somehow
+	for (int32 Attempt = 0 ; Attempt != Triangles.Num() ; ++Attempt)
+	{
+#ifdef DEBUG_LOG_BLENDSPACE_TRIANGULATION
+		UE_LOG(LogAnimation, Warning, TEXT("Triangle index: %d"), InOutTriangleIndex);
+#endif
+		const FBlendSpaceTriangle* Triangle = &Triangles[InOutTriangleIndex];
+		// Look for the edge which has the target point most outside it
+		float LargestDistance = KINDA_SMALL_NUMBER;
+		int32 LargestEdgeIndex = -1;
+		for (int32 VertexIndex = 0; VertexIndex != FBlendSpaceTriangle::NUM_VERTICES; ++VertexIndex)
+		{
+			FVector2D Corner = Triangle->Vertices[VertexIndex];
+			FVector2D EdgeNormal = Triangle->EdgeInfo[VertexIndex].Normal;
+			float Distance = (P - Corner) | EdgeNormal;
+			if (Distance > LargestDistance)
+			{
+				LargestDistance = Distance;
+				LargestEdgeIndex = VertexIndex;
+			}
+		}
+		if (LargestEdgeIndex < 0)
+		{
+			// Point is inside this triangle
+			FVector Weights = FMath::GetBaryCentric2D(P, Triangle->Vertices[0], Triangle->Vertices[1], Triangle->Vertices[2]);
+			OutWeightedSamples.Push(FWeightedBlendSample(Triangle->SampleIndices[0], Weights[0]));
+			OutWeightedSamples.Push(FWeightedBlendSample(Triangle->SampleIndices[1], Weights[1]));
+			OutWeightedSamples.Push(FWeightedBlendSample(Triangle->SampleIndices[2], Weights[2]));
+			return;
+		}
+		if (Triangle->EdgeInfo[LargestEdgeIndex].NeighbourTriangleIndex < 0)
+		{
+			// We're being directed to go outside the perimeter (which is convex). We will "walk"
+			// around the perimeter until either:
+			// 1. The point projected onto the edge is not at an extreme, or
+			// 2. The direction changes
+			// Note that the Triangle pointer will be updated during this walk
+			int32 OrigDir = -1; // will be 0 or 1 when we get going
+			int32 StartIndex = LargestEdgeIndex;
+			int32 EndIndex = (StartIndex + 1) % FBlendSpaceTriangle::NUM_VERTICES;
+
+			// Should only need a few steps
+			int PrevTriangleIndex = InOutTriangleIndex;
+			for (int32 PerimeterAttempt = 0; PerimeterAttempt != Triangles.Num(); ++PerimeterAttempt)
+			{
+#ifdef DEBUG_LOG_BLENDSPACE_TRIANGULATION
+				UE_LOG(LogAnimation, Warning, TEXT("Perimeter Triangle index: %d"), InOutTriangleIndex);
+#endif
+				const FBlendSpaceTriangleEdgeInfo& EdgeInfo = Triangle->EdgeInfo[StartIndex];
+				FVector2D Start = Triangle->Vertices[StartIndex];
+				FVector2D End = Triangle->Vertices[EndIndex];
+				FVector2D StartToEnd = End - Start;
+				FVector2D StartToPoint = P - Start;
+				// Parametric distance along the StartToEnd line
+				float T = (StartToEnd | StartToPoint) / (StartToEnd | StartToEnd);
+
+				// Check to see if we're (projected) on a perimeter segment
+				if (T >= 0.0f && T <= 1.0f)
+				{
+					OutWeightedSamples.Push(FWeightedBlendSample(Triangle->SampleIndices[StartIndex], 1.0f - T));
+					OutWeightedSamples.Push(FWeightedBlendSample(Triangle->SampleIndices[EndIndex], T));
+					return;
+				}
+
+				// Check for a change in direction
+				if (OrigDir == -1)
+				{
+					OrigDir = T > 1.0f ? 1 : 0;
+				}
+				int32 Dir = T > 1.0f ? 1 : 0;
+				if (OrigDir != Dir)
+				{
+					if (Dir == 0)
+						OutWeightedSamples.Push(FWeightedBlendSample(Triangle->SampleIndices[StartIndex], 1.0f - T));
+					else
+						OutWeightedSamples.Push(FWeightedBlendSample(Triangle->SampleIndices[EndIndex], T));
+					// If we flipped direction then we only found out by moving to a new triangle.
+					// If we cache the final triangle then next time we'll flip to the other
+					// triangle. That's fine, but confusing when debugging! Instead, use the
+					// previous triangle as the cache.
+					InOutTriangleIndex = PrevTriangleIndex;
+					return;
+				}
+
+				// Walk to the next triangle around the perimeter, but keep track of where we came from
+				PrevTriangleIndex = InOutTriangleIndex;
+				InOutTriangleIndex = EdgeInfo.AdjacentPerimeterTriangleIndices[Dir];
+				if (InOutTriangleIndex < 0)
+				{
+					// This can happen if there weren't many triangles, in which case there won't be
+					// another triangle in this direction.
+					if (Dir == 0)
+						OutWeightedSamples.Push(FWeightedBlendSample(Triangle->SampleIndices[StartIndex], 1.0f - T));
+					else
+						OutWeightedSamples.Push(FWeightedBlendSample(Triangle->SampleIndices[EndIndex], T));
+					InOutTriangleIndex = PrevTriangleIndex;
+					return;
+				}
+				Triangle = &Triangles[InOutTriangleIndex];
+				StartIndex = EdgeInfo.AdjacentPerimeterVertexIndices[Dir];
+				EndIndex = (StartIndex + 1) % FBlendSpaceTriangle::NUM_VERTICES;
+			}
+		}
+		else
+		{
+			// Investigate the new triangle
+			const FBlendSpaceTriangleEdgeInfo& EdgeInfo = Triangle->EdgeInfo[LargestEdgeIndex];
+			InOutTriangleIndex = EdgeInfo.NeighbourTriangleIndex;
+		}
+	}
+	UE_LOG(LogAnimation, Warning, TEXT("Unable to find BlendSpace triangle"));
+}
+
+void FBlendSpaceData::GetSamples(
+	TArray<FWeightedBlendSample>& OutWeightedSamples,
+	const TArray<int32>&          InDimensionIndices,
+	const FVector&                InNormalizedSamplePosition,
+	int32&                        InOutTriangulationIndex) const
+{
+	switch (InDimensionIndices.Num())
+	{
+	case 1:
+		return GetSamples1D(OutWeightedSamples, InDimensionIndices, InNormalizedSamplePosition, InOutTriangulationIndex);
+	case 2:
+		return GetSamples2D(OutWeightedSamples, InDimensionIndices, InNormalizedSamplePosition, InOutTriangulationIndex);
+	default:
+		UE_LOG(LogAnimation, Warning, TEXT("Unhandled number of dimensions in BlendSpace: %d"), InDimensionIndices.Num());
+		OutWeightedSamples.Reset();
+		break;
+	}
+}
+
+
 #undef LOCTEXT_NAMESPACE 
 
diff --git a/Engine/Source/Runtime/Engine/Private/Animation/BlendSpace1D.cpp b/Engine/Source/Runtime/Engine/Private/Animation/BlendSpace1D.cpp
--- a/Engine/Source/Runtime/Engine/Private/Animation/BlendSpace1D.cpp
+++ b/Engine/Source/Runtime/Engine/Private/Animation/BlendSpace1D.cpp
@@ -17,3 +17,77 @@ EBlendSpaceAxis UBlendSpace1D::GetAxisToScale() const
 	return bScaleAnimation ? BSA_X : BSA_None;
 }
 
+#if WITH_EDITOR
+
+void UBlendSpace1D::SnapSamplesToClosestGridPoint()
+{
+	if (!BlendParameters[0].bSnapToGrid)
+		return;
+
+	TArray<float> GridPoints;
+	TArray<int32> ClosestSampleToGridPoint;
+
+	const float GridMin = BlendParameters[0].Min;
+	const float GridMax = BlendParameters[0].Max;
+	const float GridRange = GridMax - GridMin;
+	const int32 NumGridPoints = BlendParameters[0].GridNum + 1;
+	const float GridStep = GridRange / BlendParameters[0].GridNum;
+
+	for (int32 GridPointIndex = 0; GridPointIndex < NumGridPoints; ++GridPointIndex)
+	{
+		const float GridPointValue = (GridPointIndex * GridStep) + GridMin;
+		GridPoints.Add(GridPointValue);
+	}
+
+	ClosestSampleToGridPoint.Init(INDEX_NONE, GridPoints.Num());
+
+	// Find closest sample to grid point
+	for (int32 PointIndex = 0; PointIndex < GridPoints.Num(); ++PointIndex)
+	{
+		const float GridPoint = GridPoints[PointIndex];
+		float SmallestDistance = FLT_MAX;
+		int32 Index = INDEX_NONE;
+
+		for (int32 SampleIndex = 0; SampleIndex < SampleData.Num(); ++SampleIndex)
+		{
+			FBlendSample& BlendSample = SampleData[SampleIndex];
+			const float Distance = FMath::Abs(GridPoint - BlendSample.SampleValue[0]);
+			if (Distance < SmallestDistance)
+			{
+				Index = SampleIndex;
+				SmallestDistance = Distance;
+			}
+		}
+
+		ClosestSampleToGridPoint[PointIndex] = Index;
+	}
+
+	// Find closest grid point to sample
+	for (int32 SampleIndex = 0; SampleIndex < SampleData.Num(); ++SampleIndex)
+	{
+		FBlendSample& BlendSample = SampleData[SampleIndex];
+
+		// Find closest grid point
+		float SmallestDistance = FLT_MAX;
+		int32 Index = INDEX_NONE;
+		for (int32 PointIndex = 0; PointIndex < GridPoints.Num(); ++PointIndex)
+		{
+			const float Distance = FMath::Abs(GridPoints[PointIndex] - BlendSample.SampleValue[0]);
+			if (Distance < SmallestDistance)
+			{
+				Index = PointIndex;
+				SmallestDistance = Distance;
+			}
+		}
+
+		// Only move the sample if it is also closest to the grid point
+		if (Index != INDEX_NONE && ClosestSampleToGridPoint[Index] == SampleIndex)
+		{
+			BlendSample.SampleValue[0] = GridPoints[Index];
+		}
+	}
+}
+
+#endif // WITH_EDITOR
+
+
diff --git a/Engine/Source/Runtime/Engine/Private/Animation/BlendSpace1DHelpers.cpp b/Engine/Source/Runtime/Engine/Private/Animation/BlendSpace1DHelpers.cpp
new file mode 100644
--- /dev/null
+++ b/Engine/Source/Runtime/Engine/Private/Animation/BlendSpace1DHelpers.cpp
@@ -0,0 +1,162 @@
+// Copyright Epic Games, Inc. All Rights Reserved.
+
+#include "Animation/BlendSpace1DHelpers.h"
+#include "Animation/BlendSpace.h"
+
+#define LOCTEXT_NAMESPACE "AnimationBlendSpace1DHelpers"
+
+bool FLineElement::PopulateElement(const float ElementPosition, FEditorElement& InOutElement) const
+{	
+	// If the element is left of the line element
+	if (ElementPosition < Start.Position)
+	{
+		// Element can only be left of a point if it is the first one (otherwise line is incorrect)
+		if (!bIsFirst)
+		{
+			return false;
+		}
+
+		InOutElement.Indices[0] = Start.Index;
+		InOutElement.Weights[0] = 1.0f;
+		return true;
+	}
+	// If the element is right of the line element
+	else if (ElementPosition > End.Position)
+	{
+		// Element can only be right of a point if it is the last one (otherwise line is incorrect)
+		if (!bIsLast)
+		{
+			return false;
+		}
+		InOutElement.Indices[0] = End.Index;
+		InOutElement.Weights[0] = 1.0f;
+		return true;
+	}
+	else
+	{
+		// If the element is between the start/end point of the line weight according to where it's closest to
+		InOutElement.Indices[0] = End.Index;
+		InOutElement.Weights[0] = (ElementPosition - Start.Position) / Range;
+
+		InOutElement.Indices[1] = Start.Index;
+		InOutElement.Weights[1] = (1.0f - InOutElement.Weights[0]);		
+		return true;
+	}
+}
+
+void FLineElementGenerator::Init(const FBlendParameter& BlendParameter)
+{
+	SamplePointList.Reset();
+	MinGridValue = BlendParameter.Min;
+	MaxGridValue = BlendParameter.Max;
+	NumGridPoints = BlendParameter.GridNum + 1;
+	NumGridDivisions = BlendParameter.GridNum;
+}
+
+void FLineElementGenerator::Process()
+{
+	if (SamplePointList.Num() > 1)
+	{
+		// Sort points according to position value
+		SamplePointList.Sort([](const FLineVertex& A, const FLineVertex& B) { return A.Position < B.Position; });
+	}
+}
+
+TArray<FBlendSpaceSegment> FLineElementGenerator::CalculateSegments() const
+{
+	TArray<FBlendSpaceSegment> Segments;
+	Segments.Reserve(SamplePointList.Num() > 1 ? SamplePointList.Num() - 1 : 0);
+
+	// Only create lines if we have more than one point to draw between
+	if (SamplePointList.Num() == 0)
+	{
+		return Segments;
+	}
+	else if (SamplePointList.Num() == 1)
+	{
+		FBlendSpaceSegment NewSegment;
+		NewSegment.Vertices[0] = 0.0f;
+		NewSegment.Vertices[1] = 1.0f;
+		NewSegment.SampleIndices[0] = NewSegment.SampleIndices[1] = SamplePointList[0].Index;
+		Segments.Add(NewSegment);
+	}
+	else
+	{
+		// Generate lines between sampling points from start to end (valid since they were sorted)	
+		for (int32 PointIndex = 0; PointIndex < SamplePointList.Num() - 1; ++PointIndex)
+		{
+			const int32 EndPointIndex = PointIndex + 1;
+			FBlendSpaceSegment NewSegment;
+			NewSegment.Vertices[0] = (SamplePointList[PointIndex].Position - MinGridValue) / (MaxGridValue - MinGridValue);
+			NewSegment.Vertices[1] = (SamplePointList[EndPointIndex].Position - MinGridValue) / (MaxGridValue - MinGridValue);
+			NewSegment.SampleIndices[0] = SamplePointList[PointIndex].Index;
+			NewSegment.SampleIndices[1] = SamplePointList[EndPointIndex].Index;
+			Segments.Add(NewSegment);
+		}
+	}
+	return Segments;
+}
+
+
+void FLineElementGenerator::CalculateEditorElements()
+{	
+	// Always clear line elements
+	LineElements.Empty(SamplePointList.Num() > 1 ? SamplePointList.Num() - 1 : 0);
+	
+	// Only create lines if we have more than one point to draw between
+	if (SamplePointList.Num() > 1)
+	{
+		// Generate lines between sampling points from start to end (valid since they were sorted)	
+		for (int32 PointIndex = 0; PointIndex < SamplePointList.Num() - 1; ++PointIndex)
+		{
+			const int32 EndPointIndex = PointIndex + 1;
+			LineElements.Add(FLineElement(SamplePointList[PointIndex], SamplePointList[EndPointIndex]));
+		}
+
+		// Set first and last sample flags (safe because at this point there always at least one sample)
+		LineElements[0].bIsFirst = true;
+		LineElements.Last().bIsLast = true;
+	}	
+	
+	// Number of division between grid edges
+	const float GridRange = MaxGridValue - MinGridValue;
+	const float GridStep = GridRange / NumGridDivisions;
+
+	// Initialize editor elements to required number of points
+	EditorElements.Empty(NumGridPoints);
+	EditorElements.AddDefaulted(NumGridPoints);
+
+	if (LineElements.Num() == 0)
+	{
+		// Since we did not generate any lines all the samples should correspond to the first sample and fully weighted to 1
+		for (FEditorElement& Element : EditorElements)
+		{
+			Element.Indices[0] = 0;
+			Element.Weights[0] = 1.0f;
+		}
+	}
+	else
+	{
+		for (int32 ElementIndex = 0; ElementIndex < NumGridPoints; ++ElementIndex)
+		{
+			FEditorElement& Element = EditorElements[ElementIndex];
+			const float ElementGridPosition = (GridStep * ElementIndex) + MinGridValue;
+
+			// Try and populate the editor element
+			bool bPopulatedElement = false;
+			for (const FLineElement& LineElement : LineElements)
+			{
+				bPopulatedElement |= LineElement.PopulateElement(ElementGridPosition, Element);
+				if (bPopulatedElement)
+				{
+					break;
+				}
+			}
+			
+			// Ensure that the editor element is populated using the available sample data
+			check(bPopulatedElement);
+		}
+	}
+}
+
+#undef LOCTEXT_NAMESPACE
diff --git a/Engine/Source/Runtime/Engine/Private/Animation/BlendSpace1DHelpers.h b/Engine/Source/Runtime/Engine/Private/Animation/BlendSpace1DHelpers.h
new file mode 100644
--- /dev/null
+++ b/Engine/Source/Runtime/Engine/Private/Animation/BlendSpace1DHelpers.h
@@ -0,0 +1,73 @@
+// Copyright Epic Games, Inc. All Rights Reserved.
+
+#pragma once
+
+#include "CoreMinimal.h"
+
+struct FLineVertex
+{
+	float Position;
+	int32 Index;
+	FLineVertex(const float InPosition, const int32 InIndex) : Position(InPosition), Index(InIndex) {}
+};
+
+struct FLineElement
+{
+	// Explicit constructor since we need to populate the Range value
+	explicit FLineElement(const FLineVertex& InStart, const FLineVertex& InEnd) : Start(InStart), End(InEnd), bIsFirst(false), bIsLast(false)
+	{
+		Range = End.Position - Start.Position;
+	}
+
+	bool PopulateElement(const float ElementPosition, struct FEditorElement& InOutElement) const;
+	bool IsBlendInputOnLine(const FVector& BlendInput) const
+	{
+		return (BlendInput.X >= Start.Position) && (BlendInput.X <= End.Position);
+	}
+
+	const FLineVertex Start;
+	const FLineVertex End;
+	bool bIsFirst;
+	bool bIsLast;
+	float Range;
+};
+
+/** Generates a line list between the supplied sample points to
+aid blend space sample generation */
+class FLineElementGenerator
+{
+public:
+	void Init(const struct FBlendParameter& BlendParameter);
+
+	void Process();
+
+	/** Populates EditorElements based on the Sample points previously supplied to AddSamplePoint */
+	void CalculateEditorElements();
+	
+	/**
+	 * Converts from our internal representation to the runtime segments
+	 */
+	TArray<struct FBlendSpaceSegment> CalculateSegments() const;
+
+	/**
+	* Data Structure for line generation
+	* SamplePointList is the input data
+	*/
+	TArray<FLineVertex> SamplePointList;
+
+	/** Editor elements generated by CalculateEditorElements */
+	TArray<struct FEditorElement> EditorElements;
+
+private:
+
+	/** Defines the range of the editor */
+	float MinGridValue;
+	float MaxGridValue;
+
+	/** Number of points that we have to generate FEditorElements for */
+	int32 NumGridPoints;
+	int32 NumGridDivisions;
+
+	/** Line elements generated from the given samples */
+	TArray<FLineElement> LineElements;
+};
diff --git a/Engine/Source/Runtime/Engine/Private/Animation/BlendSpaceHelpers.cpp b/Engine/Source/Runtime/Engine/Private/Animation/BlendSpaceHelpers.cpp
new file mode 100644
--- /dev/null
+++ b/Engine/Source/Runtime/Engine/Private/Animation/BlendSpaceHelpers.cpp
@@ -0,0 +1,764 @@
+// Copyright Epic Games, Inc. All Rights Reserved.
+
+#include "Animation/BlendSpaceHelpers.h"
+
+#define BLENDSPACE_MINSAMPLE	3
+
+#define LOCTEXT_NAMESPACE "AnimationBlendSpaceHelpers"
+
+struct FSortByDistance
+{
+	int32 Index;
+	float Distance;
+	FSortByDistance(int32 InIndex, float InDistance) 
+		: Index(InIndex), Distance(InDistance) {}
+};
+
+void FDelaunayTriangleGenerator::SetGridBox(const FBlendParameter& BlendParamX, const FBlendParameter& BlendParamY)
+{
+	FVector2D Min(BlendParamX.Min, BlendParamY.Min);
+	FVector2D Max(BlendParamX.Max, BlendParamY.Max);
+	FVector2D Mid = (Min + Max) * 0.5f;
+	FVector2D Range = Max - Min;
+	Range.X = FMath::Max(Range.X, DELTA);
+	Range.Y = FMath::Max(Range.Y, DELTA);
+
+	GridMin = Mid - Range * 0.5f;
+	RecipGridSize = FVector2D(1.0f, 1.0f) / Range;
+}
+
+void FDelaunayTriangleGenerator::EmptyTriangles()
+{
+	for (int32 TriangleIndex = 0; TriangleIndex < TriangleList.Num(); ++TriangleIndex)
+	{
+		delete TriangleList[TriangleIndex];
+	}
+
+	TriangleList.Empty();
+}
+void FDelaunayTriangleGenerator::EmptySamplePoints()
+{
+	SamplePointList.Empty();
+}
+
+void FDelaunayTriangleGenerator::Reset()
+{
+	EmptyTriangles();
+	EmptySamplePoints();
+}
+
+void FDelaunayTriangleGenerator::AddSamplePoint(const FVector2D& NewPoint, const int32 SampleIndex)
+{
+	checkf(!SamplePointList.Contains(NewPoint), TEXT("Found duplicate points in blendspace"));
+	SamplePointList.Add(FVertex(NewPoint, SampleIndex));
+}
+
+void FDelaunayTriangleGenerator::Triangulate(EPreferredTriangulationDirection PreferredTriangulationDirection)
+{
+	if (SamplePointList.Num() == 0)
+	{
+		return;
+	}
+	else if (SamplePointList.Num() == 1)
+	{
+		// degenerate case 1
+		FTriangle Triangle(&SamplePointList[0]);
+		AddTriangle(Triangle);
+	}
+	else if (SamplePointList.Num() == 2)
+	{
+		// degenerate case 2
+		FTriangle Triangle(&SamplePointList[0], &SamplePointList[1]);
+		AddTriangle(Triangle);
+	}
+	else
+	{
+		SortSamples();
+
+		// first choose first 3 points
+		for (int32 I = 2; I<SamplePointList.Num(); ++I)
+		{
+			GenerateTriangles(SamplePointList, I + 1);
+		}
+
+		// degenerate case 3: many points all collinear or coincident
+		if (TriangleList.Num() == 0)
+		{
+			if (AllCoincident(SamplePointList))
+			{
+				// coincident case - just create one triangle
+				FTriangle Triangle(&SamplePointList[0]);
+				AddTriangle(Triangle);
+			}
+			else
+			{
+				// collinear case: create degenerate triangles between pairs of points
+				for (int32 PointIndex = 0; PointIndex < SamplePointList.Num() - 1; ++PointIndex)
+				{
+					FTriangle Triangle(&SamplePointList[PointIndex], &SamplePointList[PointIndex + 1]);
+					AddTriangle(Triangle);
+				}
+			}
+		}
+	}
+	AdjustEdgeDirections(PreferredTriangulationDirection);
+}
+
+static int32 GetCandidateEdgeIndex(const FTriangle* Triangle)
+{
+	bool bFoundVert = false;
+	bool bFoundHor = false;
+	bool bFoundSlope = false;
+	int32 SlopeIndex = 0;
+	for (int32 VertexIndex = 0; VertexIndex != 3; ++VertexIndex)
+	{
+		int32 VertexIndexNext = (VertexIndex + 1) % 3;
+		if (Triangle->Vertices[VertexIndex]->Position.X == Triangle->Vertices[VertexIndexNext]->Position.X)
+		{
+			bFoundVert = true;
+		}
+		else if (Triangle->Vertices[VertexIndex]->Position.Y == Triangle->Vertices[VertexIndexNext]->Position.Y)
+		{
+			bFoundHor = true;
+		}
+		else
+		{
+			bFoundSlope = true;
+			SlopeIndex = VertexIndex;
+		}
+	}
+	if (bFoundVert && bFoundHor && bFoundSlope)
+	{
+		return SlopeIndex;
+	}
+	else
+	{
+		return -1;
+	}
+}
+
+void FDelaunayTriangleGenerator::AdjustEdgeDirections(EPreferredTriangulationDirection PreferredTriangulationDirection)
+{
+	if (PreferredTriangulationDirection == EPreferredTriangulationDirection::None)
+	{
+		return;
+	}
+
+	for (int32 TriangleIndex0 = 0; TriangleIndex0 != TriangleList.Num(); ++TriangleIndex0)
+	{
+		FTriangle* Triangle0 = TriangleList[TriangleIndex0];
+
+		// Check if it is axis aligned and if so which is the hypotenuse edge
+		int32 EdgeIndex0 = GetCandidateEdgeIndex(Triangle0);
+		if (EdgeIndex0 >= 0)
+		{
+			int32 EdgeIndexNext0 = (EdgeIndex0 + 1) % 3;
+			// Only continue if we would want to flip this edge
+			FVector2D EdgeDir = Triangle0->Vertices[EdgeIndexNext0]->Position - Triangle0->Vertices[EdgeIndex0]->Position;
+			FVector2D EdgePos = (Triangle0->Vertices[EdgeIndexNext0]->Position + Triangle0->Vertices[EdgeIndex0]->Position) * 0.5f;
+
+			// Slope is true if we go up from left to right
+			bool bUpSlope = EdgeDir.X * EdgeDir.Y > 0.0f;
+			// Desired slope depends on the quadrant
+			bool bDesiredUpSlope = PreferredTriangulationDirection == EPreferredTriangulationDirection::Tangential ? EdgePos.X * EdgePos.Y < 0.0f : EdgePos.X * EdgePos.Y >= 0.0f;
+
+			if (bUpSlope == bDesiredUpSlope)
+			{
+				continue;
+			}
+
+			// Look for a triangle that is also axis aligned and facing this sloping edge
+			int32 SampleIndexStart = Triangle0->Vertices[EdgeIndex0]->SampleIndex;
+			int32 SampleIndexEnd = Triangle0->Vertices[(EdgeIndex0 + 1) % 3]->SampleIndex;
+			int32 EdgeIndex1;
+			int32 TriangleIndex1 = FindTriangleIndexWithEdge(SampleIndexEnd, SampleIndexStart, &EdgeIndex1);
+			if (TriangleIndex1 < 0)
+			{
+				continue;
+			}
+			check(EdgeIndex1 >= 0);
+			FTriangle* Triangle1 = TriangleList[TriangleIndex1];
+			if (GetCandidateEdgeIndex(Triangle1) == EdgeIndex1)
+			{
+				// Flip the edge across the pair of triangles. Note that after this the triangles
+				// will contain incorrect edge info, but we don't need that any more.
+				int32 EdgeIndexPrev1 = (EdgeIndex1 + 2) % 3;
+				int32 EdgeIndexNext1 = (EdgeIndex1 + 1) % 3;
+				int32 EdgeIndexPrev0 = (EdgeIndex0 + 2) % 3;
+
+				Triangle0->Vertices[EdgeIndexNext0] = Triangle1->Vertices[EdgeIndexPrev1];
+				Triangle1->Vertices[EdgeIndexNext1] = Triangle0->Vertices[EdgeIndexPrev0];
+			}
+		}
+
+	}
+
+}
+
+void FDelaunayTriangleGenerator::SortSamples()
+{
+	// Populate sorting array with sample points and their original (blend space -> sample data) indices
+
+	struct FComparePoints
+	{
+		FORCEINLINE bool operator()( const FVertex& A, const FVertex& B ) const
+		{
+			// the sorting happens from -> +X, -> +Y,  -> for now ignore Z ->+Z
+			if( A.Position.Y == B.Position.Y ) // same, then compare Y
+			{
+				return A.Position.X < B.Position.X;
+			}
+			return A.Position.Y < B.Position.Y;			
+		}
+	};
+	// sort all points
+	SamplePointList.Sort( FComparePoints() );
+}
+
+/** 
+* The key function in Delaunay Triangulation
+* return true if the TestPoint is WITHIN the triangle circumcircle
+*	http://en.wikipedia.org/wiki/Delaunay_triangulation 
+*/
+FDelaunayTriangleGenerator::ECircumCircleState FDelaunayTriangleGenerator::GetCircumcircleState(const FTriangle* T, const FVertex& TestPoint)
+{
+	const int32 NumPointsPerTriangle = 3;
+
+	// First off, normalize all the points
+	FVector2D NormalizedPositions[NumPointsPerTriangle];
+	
+	// Unrolled loop
+	NormalizedPositions[0] = (T->Vertices[0]->Position - GridMin) * RecipGridSize;
+	NormalizedPositions[1] = (T->Vertices[1]->Position - GridMin) * RecipGridSize;
+	NormalizedPositions[2] = (T->Vertices[2]->Position - GridMin) * RecipGridSize;
+
+	const FVector2D NormalizedTestPoint = ( TestPoint.Position - GridMin ) * RecipGridSize;
+
+	// ignore Z, eventually this has to be on plane
+	// http://en.wikipedia.org/wiki/Delaunay_triangulation - determinant
+	const float M00 = NormalizedPositions[0].X - NormalizedTestPoint.X;
+	const float M01 = NormalizedPositions[0].Y - NormalizedTestPoint.Y;
+	const float M02 = NormalizedPositions[0].X * NormalizedPositions[0].X - NormalizedTestPoint.X * NormalizedTestPoint.X
+		+ NormalizedPositions[0].Y*NormalizedPositions[0].Y - NormalizedTestPoint.Y * NormalizedTestPoint.Y;
+
+	const float M10 = NormalizedPositions[1].X - NormalizedTestPoint.X;
+	const float M11 = NormalizedPositions[1].Y - NormalizedTestPoint.Y;
+	const float M12 = NormalizedPositions[1].X * NormalizedPositions[1].X - NormalizedTestPoint.X * NormalizedTestPoint.X
+		+ NormalizedPositions[1].Y * NormalizedPositions[1].Y - NormalizedTestPoint.Y * NormalizedTestPoint.Y;
+
+	const float M20 = NormalizedPositions[2].X - NormalizedTestPoint.X;
+	const float M21 = NormalizedPositions[2].Y - NormalizedTestPoint.Y;
+	const float M22 = NormalizedPositions[2].X * NormalizedPositions[2].X - NormalizedTestPoint.X * NormalizedTestPoint.X
+		+ NormalizedPositions[2].Y * NormalizedPositions[2].Y - NormalizedTestPoint.Y * NormalizedTestPoint.Y;
+
+	const float Det = M00*M11*M22+M01*M12*M20+M02*M10*M21 - (M02*M11*M20+M01*M10*M22+M00*M12*M21);
+	
+	// When the vertices are sorted in a counterclockwise order, the determinant is positive if and only if Testpoint lies inside the circumcircle of T.
+	if (FMath::IsNegativeFloat(Det))
+	{
+		return ECCS_Outside;
+	}
+	else
+	{
+		// On top of the triangle edge
+		if (FMath::IsNearlyZero(Det, KINDA_SMALL_NUMBER))
+		{
+			return ECCS_On;
+		}
+		else
+		{
+			return ECCS_Inside;
+		}
+	}
+}
+
+int32 FDelaunayTriangleGenerator::FindTriangleIndexWithEdge(int32 SampleIndex0, int32 SampleIndex1, int32* VertexIndex) const
+{
+	for (int32 TriangleIndex = 0; TriangleIndex != TriangleList.Num(); ++TriangleIndex)
+	{
+		const FTriangle* Triangle = TriangleList[TriangleIndex];
+		if (Triangle->Vertices[0]->SampleIndex == SampleIndex0 && Triangle->Vertices[1]->SampleIndex == SampleIndex1)
+		{
+			if (VertexIndex)
+			{
+				*VertexIndex = 0;
+			}
+			return TriangleIndex;
+		}
+		if (Triangle->Vertices[1]->SampleIndex == SampleIndex0 && Triangle->Vertices[2]->SampleIndex == SampleIndex1)
+		{
+			if (VertexIndex)
+			{
+				*VertexIndex = 1;
+			}
+			return TriangleIndex;
+		}
+		if (Triangle->Vertices[2]->SampleIndex == SampleIndex0 && Triangle->Vertices[0]->SampleIndex == SampleIndex1)
+		{
+			if (VertexIndex)
+			{
+				*VertexIndex = 2;
+			}
+			return TriangleIndex;
+		}
+	}
+	// Perfectly normal to get here - when looking for a triangle that is off the outside edge of
+	// the graph.
+	return INDEX_NONE;
+}
+
+TArray<struct FBlendSpaceTriangle> FDelaunayTriangleGenerator::CalculateTriangles() const
+{
+	TArray<struct FBlendSpaceTriangle> Triangles;
+	Triangles.Reserve(TriangleList.Num());
+	for (int32 TriangleIndex = 0; TriangleIndex != TriangleList.Num(); ++TriangleIndex)
+	{
+		const FTriangle* Triangle = TriangleList[TriangleIndex];
+
+		FBlendSpaceTriangle NewTriangle;
+		
+		NewTriangle.Vertices[0] = (Triangle->Vertices[0]->Position - GridMin) * RecipGridSize;
+		NewTriangle.Vertices[1] = (Triangle->Vertices[1]->Position - GridMin) * RecipGridSize;
+		NewTriangle.Vertices[2] = (Triangle->Vertices[2]->Position - GridMin) * RecipGridSize;
+		NewTriangle.SampleIndices[0] = Triangle->Vertices[0]->SampleIndex;
+		NewTriangle.SampleIndices[1] = Triangle->Vertices[1]->SampleIndex;
+		NewTriangle.SampleIndices[2] = Triangle->Vertices[2]->SampleIndex;
+
+		for (int32 Index = 0; Index != FBlendSpaceTriangle::NUM_VERTICES; ++Index)
+		{
+			FBlendSpaceTriangleEdgeInfo& EdgeInfo = NewTriangle.EdgeInfo[Index];
+			int32 IndexNext = (Index + 1) % FBlendSpaceTriangle::NUM_VERTICES;
+
+			EdgeInfo.NeighbourTriangleIndex = FindTriangleIndexWithEdge(
+				Triangle->Vertices[IndexNext]->SampleIndex, Triangle->Vertices[Index]->SampleIndex);
+			FVector2D EdgeDir = NewTriangle.Vertices[IndexNext] - NewTriangle.Vertices[Index];
+
+			// Triangles are wound anticlockwise as viewed from above - rotate the edge 90 deg clockwise
+			// to make it be the outwards pointing normal.
+			EdgeInfo.Normal = FVector2D(EdgeDir.Y, -EdgeDir.X).GetSafeNormal();
+
+			// Update these if necessary when all triangles have been processed
+			EdgeInfo.AdjacentPerimeterTriangleIndices[0] = INDEX_NONE;
+			EdgeInfo.AdjacentPerimeterTriangleIndices[1] = INDEX_NONE;
+			EdgeInfo.AdjacentPerimeterVertexIndices[0] = INDEX_NONE;
+			EdgeInfo.AdjacentPerimeterVertexIndices[1] = INDEX_NONE;
+		}
+		Triangles.Add(NewTriangle);
+	}
+
+	for (int32 TriangleIndex = 0; TriangleIndex != Triangles.Num(); ++TriangleIndex)
+	{
+		FBlendSpaceTriangle& Triangle = Triangles[TriangleIndex];
+		for (int32 Index = 0; Index != FBlendSpaceTriangle::NUM_VERTICES; ++Index)
+		{
+			FBlendSpaceTriangleEdgeInfo& EdgeInfo = Triangle.EdgeInfo[Index];
+			if (EdgeInfo.NeighbourTriangleIndex == INDEX_NONE)
+			{
+				int32 IndexNext = (Index + 1) % FBlendSpaceTriangle::NUM_VERTICES;
+
+				int32 SampleIndex = Triangle.SampleIndices[Index];
+				int32 SampleIndexNext = Triangle.SampleIndices[IndexNext];
+
+				// Update the triangle info to allow traversal around the perimeter of the
+				// triangulated region. Iterate through all the other triangles...
+				for (int32 OtherTriangleIndex = 0; OtherTriangleIndex != Triangles.Num(); ++OtherTriangleIndex)
+				{
+					if (OtherTriangleIndex == TriangleIndex)
+						continue;
+
+					const FBlendSpaceTriangle& OtherTriangle = Triangles[OtherTriangleIndex];
+					// ... then check for a vertex that matches the edge we're considering
+					for (int32 VertexIndex = 0; VertexIndex != FBlendSpaceTriangle::NUM_VERTICES; ++VertexIndex)
+					{
+						int32 VertexIndexPrev = (VertexIndex + FBlendSpaceTriangle::NUM_VERTICES - 1) % FBlendSpaceTriangle::NUM_VERTICES;
+						if (OtherTriangle.SampleIndices[VertexIndex] == SampleIndex
+							&& OtherTriangle.EdgeInfo[VertexIndexPrev].NeighbourTriangleIndex == INDEX_NONE)
+						{
+							// We found a perimeter triangle that comes before our edge
+							EdgeInfo.AdjacentPerimeterTriangleIndices[0] = OtherTriangleIndex;
+							EdgeInfo.AdjacentPerimeterVertexIndices[0] = VertexIndexPrev;
+						}
+						if (OtherTriangle.SampleIndices[VertexIndex] == SampleIndexNext
+							&& OtherTriangle.EdgeInfo[VertexIndex].NeighbourTriangleIndex == INDEX_NONE)
+						{
+							// We found a perimeter triangle that comes after our edge
+							EdgeInfo.AdjacentPerimeterTriangleIndices[1] = OtherTriangleIndex;
+							EdgeInfo.AdjacentPerimeterVertexIndices[1] = VertexIndex ;
+						}
+					}
+				}
+			}
+		}
+	}
+	return Triangles;
+}
+
+bool FDelaunayTriangleGenerator::IsCollinear(const FVertex* A, const FVertex* B, const FVertex* C)
+{
+	const FVector2D Diff1 = B->Position - A->Position;
+	const FVector2D Diff2 = C->Position - A->Position;
+	float Cross = Diff1 ^ Diff2;
+	return (Cross == 0.f);
+}
+
+bool FDelaunayTriangleGenerator::AllCoincident(const TArray<FVertex>& InPoints)
+{
+	if (InPoints.Num() > 0)
+	{
+		const FVertex& FirstPoint = InPoints[0];
+		for (int32 PointIndex = 0; PointIndex < InPoints.Num(); ++PointIndex)
+		{
+			const FVertex& Point = InPoints[PointIndex];
+			if (Point.Position != FirstPoint.Position)
+			{
+				return false;
+			}
+		}
+
+		return true;
+	}
+
+	return false;
+}
+
+bool FDelaunayTriangleGenerator::FlipTriangles(const int32 TriangleIndexOne, const int32 TriangleIndexTwo)
+{
+	const FTriangle* A = TriangleList[TriangleIndexOne];
+	const FTriangle* B = TriangleList[TriangleIndexTwo];
+
+	// if already optimized, don't have to do any
+	FVertex* TestPt = A->FindNonSharingPoint(B);
+
+	// If it's not inside, we don't have to do any
+	if (GetCircumcircleState(A, *TestPt) != ECCS_Inside)
+	{
+		return false;
+	}
+
+	FTriangle NewTriangles[2];
+	int32 TrianglesMade = 0;
+
+	for (int32 VertexIndexOne = 0; VertexIndexOne < 2; ++VertexIndexOne)
+	{
+		for (int32 VertexIndexTwo = VertexIndexOne + 1; VertexIndexTwo < 3; ++VertexIndexTwo)
+		{
+			// Check if these vertices form a valid triangle (should be non-colinear)
+			if (IsEligibleForTriangulation(A->Vertices[VertexIndexOne], A->Vertices[VertexIndexTwo], TestPt))
+			{
+				// Create the new triangle and check if the final (original) vertex falls inside or outside of it's circumcircle
+				const FTriangle NewTriangle(A->Vertices[VertexIndexOne], A->Vertices[VertexIndexTwo], TestPt);
+				const int32 VertexIndexThree = 3 - (VertexIndexTwo + VertexIndexOne);
+				if (GetCircumcircleState(&NewTriangle, *A->Vertices[VertexIndexThree]) == ECCS_Outside)
+				{
+					// If so store the triangle and increment the number of triangles
+					checkf(TrianglesMade < 2, TEXT("Incorrect number of triangles created"));
+					NewTriangles[TrianglesMade] = NewTriangle;
+					++TrianglesMade;
+				}
+			}
+		}
+	}
+	
+	// In case two triangles were generated the flip was successful so we can add them to the list
+	if (TrianglesMade == 2)
+	{
+		AddTriangle(NewTriangles[0], false);
+		AddTriangle(NewTriangles[1], false);
+	}
+
+	return TrianglesMade == 2;
+}
+
+void FDelaunayTriangleGenerator::AddTriangle(FTriangle& newTriangle, bool bCheckHalfEdge/*=true*/)
+{
+	// see if it's same vertices
+	for (int32 I=0;I<TriangleList.Num(); ++I)
+	{
+		if (newTriangle == *TriangleList[I])
+		{
+			return;
+		}
+
+		if (bCheckHalfEdge && newTriangle.HasSameHalfEdge(TriangleList[I]))
+		{
+			return;
+		}
+	}
+
+	TriangleList.Add(new FTriangle(newTriangle));
+}
+
+int32 FDelaunayTriangleGenerator::GenerateTriangles(TArray<FVertex>& PointList, const int32 TotalNum)
+{
+	if (TotalNum == BLENDSPACE_MINSAMPLE)
+	{
+		if (IsEligibleForTriangulation(&PointList[0], &PointList[1], &PointList[2]))
+		{
+			FTriangle Triangle(&PointList[0], &PointList[1], &PointList[2]);
+			AddTriangle(Triangle);
+		}
+	}
+	else if (TriangleList.Num() == 0)
+	{
+		FVertex * TestPoint = &PointList[TotalNum-1];
+
+		// so far no triangle is made, try to make it with new points that are just entered
+		for (int32 I=0; I<TotalNum-2; ++I)
+		{
+			if (IsEligibleForTriangulation(&PointList[I], &PointList[I+1], TestPoint))
+			{
+				FTriangle NewTriangle (&PointList[I], &PointList[I+1], TestPoint);
+				AddTriangle(NewTriangle);
+			}
+		}
+	}
+	else
+	{
+		// get the last addition
+		FVertex * TestPoint = &PointList[TotalNum-1];
+		int32 TriangleNum = TriangleList.Num();
+	
+		for (int32 I=0; I<TriangleList.Num(); ++I)
+		{
+			FTriangle * Triangle = TriangleList[I];
+			if (IsEligibleForTriangulation(Triangle->Vertices[0], Triangle->Vertices[1], TestPoint))
+			{
+				FTriangle NewTriangle (Triangle->Vertices[0], Triangle->Vertices[1], TestPoint);
+				AddTriangle(NewTriangle);
+			}
+
+			if (IsEligibleForTriangulation(Triangle->Vertices[0], Triangle->Vertices[2], TestPoint))
+			{
+				FTriangle NewTriangle (Triangle->Vertices[0], Triangle->Vertices[2], TestPoint);
+				AddTriangle(NewTriangle);
+			}
+
+			if (IsEligibleForTriangulation(Triangle->Vertices[1], Triangle->Vertices[2], TestPoint))
+			{
+				FTriangle NewTriangle (Triangle->Vertices[1], Triangle->Vertices[2], TestPoint);
+				AddTriangle(NewTriangle);
+			}
+		}
+
+		// this is locally optimization part
+		// we need to make sure all triangles are locally optimized. If not optimize it. 
+		for (int32 I=0; I<TriangleList.Num(); ++I)
+		{
+			FTriangle * A = TriangleList[I];
+			for (int32 J=I+1; J<TriangleList.Num(); ++J)
+			{
+				FTriangle * B = TriangleList[J];
+
+				// does share same edge
+				if (A->DoesShareSameEdge(B))
+				{
+					// then test to see if locally optimized
+					if (FlipTriangles(I, J))
+					{
+						// if this flips, remove current triangle
+						delete TriangleList[I];
+						delete TriangleList[J];
+						//I need to remove J first because other wise, 
+						//  index J isn't valid anymore
+						TriangleList.RemoveAt(J);
+						TriangleList.RemoveAt(I);
+						// start over since we modified triangle
+						// once we don't have any more to flip, we're good to go!
+						I=-1;
+						break;
+					}
+				}
+			}
+		}
+	}
+
+	return TriangleList.Num();
+}
+
+bool FBlendSpaceGrid::FindTriangleThisPointBelongsTo(const FVector2D& TestPoint, FVector& OutBarycentricCoords, FTriangle*& OutTriangle, const TArray<FTriangle*>& TriangleList) const
+{
+	// Calculate distance from point to triangle and sort the triangle list accordingly
+	TArray<FSortByDistance> SortedTriangles;
+	SortedTriangles.AddUninitialized(TriangleList.Num());
+	for (int32 TriangleIndex=0; TriangleIndex<TriangleList.Num(); ++TriangleIndex)
+	{
+		SortedTriangles[TriangleIndex].Index = TriangleIndex;
+		SortedTriangles[TriangleIndex].Distance = TriangleList[TriangleIndex]->GetDistance(TestPoint);
+	}
+	SortedTriangles.Sort([](const FSortByDistance &A, const FSortByDistance &B) { return A.Distance < B.Distance; });
+
+	// Now loop over the sorted triangles and test the barycentric coordinates with the point
+	for (const FSortByDistance& SortedTriangle : SortedTriangles)
+	{
+		FTriangle* Triangle = TriangleList[SortedTriangle.Index];
+
+		FVector Coords = FMath::GetBaryCentric2D(TestPoint, Triangle->Vertices[0]->Position, Triangle->Vertices[1]->Position, Triangle->Vertices[2]->Position);
+
+		// Z coords often has precision error because it's derived from 1-A-B, so do more precise check
+		if (FMath::Abs(Coords.Z) < KINDA_SMALL_NUMBER)
+		{
+			Coords.Z = 0.f;
+		}
+
+		// Is the point inside of the triangle, or on it's edge (Z coordinate should always match since the blend samples are set in 2D)
+		if ( 0.f <= Coords.X && Coords.X <= 1.0 && 0.f <= Coords.Y && Coords.Y <= 1.0 && 0.f <= Coords.Z && Coords.Z <= 1.0 )
+		{
+			OutBarycentricCoords = Coords;
+			OutTriangle = Triangle;
+			return true;
+		}
+	}
+
+	return false;
+}
+
+void FBlendSpaceGrid::GenerateGridElements(const TArray<FVertex>& SamplePoints, const TArray<FTriangle*>& TriangleList)
+{
+	check (NumGridDivisions.X > 0 && NumGridDivisions.Y > 0 );
+	check (GridMax > GridMin);
+
+	const int32 TotalNumGridPoints = NumGridPointsForAxis.X * NumGridPointsForAxis.Y;
+
+	GridPoints.Empty(TotalNumGridPoints);
+
+	if (SamplePoints.Num() == 0 || TriangleList.Num() == 0)
+	{
+		return;
+	}
+
+	GridPoints.AddDefaulted(TotalNumGridPoints);
+
+	FVector2D GridPointPosition;		
+	for (int32 GridPositionX = 0; GridPositionX < NumGridPointsForAxis.X; ++GridPositionX)
+	{
+		for (int32 GridPositionY = 0; GridPositionY < NumGridPointsForAxis.Y; ++GridPositionY)
+		{
+			FTriangle * SelectedTriangle = NULL;
+			FEditorElement& GridPoint = GridPoints[ GridPositionX*NumGridPointsForAxis.Y + GridPositionY ];
+
+			GridPointPosition = GetPosFromIndex(GridPositionX, GridPositionY);
+
+			FVector Weights;
+			if ( FindTriangleThisPointBelongsTo(GridPointPosition, Weights, SelectedTriangle, TriangleList) )
+			{
+				// found it
+				GridPoint.Weights[0] = Weights.X;
+				GridPoint.Weights[1] = Weights.Y;
+				GridPoint.Weights[2] = Weights.Z;
+				GridPoint.Indices[0] = SelectedTriangle->Vertices[0]->SampleIndex;
+				GridPoint.Indices[1] = SelectedTriangle->Vertices[1]->SampleIndex;
+				GridPoint.Indices[2] = SelectedTriangle->Vertices[2]->SampleIndex;
+				check(GridPoint.Indices[0] != INDEX_NONE);
+				check(GridPoint.Indices[1] != INDEX_NONE);
+				check(GridPoint.Indices[2] != INDEX_NONE);
+			}
+			else
+			{
+				TArray<FSortByDistance> SortedTriangles;
+				for (int32 TriangleIndex = 0; TriangleIndex < TriangleList.Num(); ++TriangleIndex )
+				{
+					// Check if points are collinear
+					const FTriangle* Triangle = TriangleList[TriangleIndex];
+					const FVector2D EdgeA = Triangle->Vertices[1]->Position - Triangle->Vertices[0]->Position;
+					const FVector2D EdgeB = Triangle->Vertices[2]->Position - Triangle->Vertices[0]->Position;
+					const float Result = EdgeA.X * EdgeB.Y - EdgeA.Y * EdgeB.X;
+					// Only add valid triangles
+					if (Result > 0.0f)
+					{ 
+						SortedTriangles.Add(FSortByDistance(TriangleIndex, Triangle->GetDistance(GridPointPosition)));
+					}					
+				}
+
+				if (SortedTriangles.Num())
+				{
+					SortedTriangles.Sort([](const FSortByDistance &A, const FSortByDistance &B) { return A.Distance < B.Distance; });
+					FTriangle* ClosestTriangle = TriangleList[SortedTriangles[0].Index];
+
+					// For the closest triangle, determine which of its edges is closest to the grid point
+					TArray<FSortByDistance> Edges;
+					TArray<FVector2D> PointsOnEdges;
+					for (int32 EdgeIndex = 0; EdgeIndex < 3; ++EdgeIndex)
+					{
+						const FVector2D ClosestPoint = FMath::ClosestPointOnSegment2D(
+							GridPointPosition,
+							ClosestTriangle->Edges[EdgeIndex].Vertices[0]->Position,
+							ClosestTriangle->Edges[EdgeIndex].Vertices[1]->Position);
+						Edges.Add(FSortByDistance(EdgeIndex, (ClosestPoint - GridPointPosition).SizeSquared()));
+						PointsOnEdges.Add(ClosestPoint);
+					}
+					Edges.Sort([](const FSortByDistance &A, const FSortByDistance &B) { return A.Distance < B.Distance; });
+					
+					// Calculate weighting using the closest edge points and the clamped grid position on the line
+					const FVector GridWeights = FMath::GetBaryCentric2D(PointsOnEdges[Edges[0].Index], ClosestTriangle->Vertices[0]->Position, ClosestTriangle->Vertices[1]->Position, ClosestTriangle->Vertices[2]->Position);
+					
+					for (int32 Index = 0; Index < 3; ++Index)
+					{
+						GridPoint.Weights[Index] = GridWeights[Index];
+						GridPoint.Indices[Index] = ClosestTriangle->Vertices[Index]->SampleIndex;
+					}
+				}
+				else
+				{
+					// This means that there is either one point, two points or collinear triangles on the grid
+					if (SamplePoints.Num() == 1)
+					{
+						// Just one, fill all grid points to the single sample
+						GridPoint.Weights[0] = 1.0f;
+						GridPoint.Indices[0] = 0;
+					}
+					else
+					{
+						// Two points or co-linear triangles, first find the two closest samples
+						TArray<FSortByDistance> SampleDistances;
+						for (int32 PointIndex = 0; PointIndex < SamplePoints.Num(); ++PointIndex)
+						{
+							const float DistanceFromSampleToPoint= (SamplePoints[PointIndex].Position - GridPointPosition).SizeSquared();
+							SampleDistances.Add(FSortByDistance(PointIndex, DistanceFromSampleToPoint));
+						}
+						SampleDistances.Sort([](const FSortByDistance &A, const FSortByDistance &B) { return A.Distance < B.Distance; });
+
+						// Find closest point on line between the two samples (clamping the grid position to the line, just like clamping to the triangle edges)
+						const FVertex Samples[2] = { SamplePoints[SampleDistances[0].Index], SamplePoints[SampleDistances[1].Index] };
+						const FVector2D ClosestPoint = FMath::ClosestPointOnSegment2D(GridPointPosition, Samples[0].Position, Samples[1].Position);
+						const float LineLength = (Samples[0].Position - Samples[1].Position).SizeSquared();
+
+						// Weight the samples according to the distance from the grid point on the line to the samples
+						for (int32 Index = 0; Index < 2; ++Index)
+						{
+							GridPoint.Weights[Index] = (LineLength - (Samples[Index].Position - ClosestPoint).SizeSquared()) / LineLength;
+							GridPoint.Indices[Index] = Samples[Index].SampleIndex;
+						}
+					}
+				}
+			}
+		}
+	}
+}
+
+/** 
+* Convert grid index (GridX, GridY) to triangle coords and returns FVector2D
+*/
+const FVector2D FBlendSpaceGrid::GetPosFromIndex(const int32 GridX, const int32 GridY) const
+{
+	// grid X starts from 0 -> N when N == GridSizeX
+	// grid Y starts from 0 -> N when N == GridSizeY
+	// LeftBottom will map to Grid 0, 0
+	// RightTop will map to Grid N, N
+
+	FVector2D CoordDim = GridMax - GridMin;
+	FVector2D EachGridSize = CoordDim / NumGridDivisions;
+
+	// for now only 2D
+	return FVector2D(GridX*EachGridSize.X+GridMin.X, GridY*EachGridSize.Y+GridMin.Y);
+}
+
+const FEditorElement& FBlendSpaceGrid::GetElement(const int32 GridX, const int32 GridY) const
+{
+	check (NumGridPointsForAxis.X >= GridX);
+	check (NumGridPointsForAxis.Y >= GridY);
+
+	check (GridPoints.Num() > 0 );
+	return GridPoints[GridX * NumGridPointsForAxis.Y + GridY];
+}
+
+#undef LOCTEXT_NAMESPACE
diff --git a/Engine/Source/Runtime/Engine/Private/Animation/BlendSpaceHelpers.h b/Engine/Source/Runtime/Engine/Private/Animation/BlendSpaceHelpers.h
new file mode 100644
--- /dev/null
+++ b/Engine/Source/Runtime/Engine/Private/Animation/BlendSpaceHelpers.h
@@ -0,0 +1,499 @@
+// Copyright Epic Games, Inc. All Rights Reserved.
+
+#pragma once
+
+#include "CoreMinimal.h"
+#include "Animation/BlendSpace.h"
+
+struct FTriangle;
+
+//////////////////////////////////////////////////////////////////////////
+// SAnimationBlendSpace
+
+struct FTriangle;
+
+/**
+ * Utility class for FDelaunayTriangleGenerator 
+ * This represents Points in FDelaunayTriangleGenerator
+ */
+struct FVertex
+{
+	// position of Point
+	FVector2D Position;
+
+	// Triangles this point belongs to
+	TArray<FTriangle *> Triangles;
+
+	// The original animation sample associated with this point
+	int32 SampleIndex;
+	
+	FVertex(FVector2D InPosition, int32 InSampleIndex) : Position(InPosition), SampleIndex(InSampleIndex) {}
+
+	bool operator==(const FVector2D& Other) const
+	{
+		// if same position, it's same point
+		return (Other == Position);
+	}
+
+	bool operator==( const FVertex& Other ) const
+	{
+		// if same position, it's same point
+		return (Other.Position == Position);
+	}
+
+	void AddTriangle(FTriangle * NewTriangle)
+	{
+		Triangles.AddUnique(NewTriangle);
+	}
+
+	void RemoveTriangle(FTriangle * TriangleToRemove)
+	{
+		Triangles.Remove(TriangleToRemove);
+	}
+
+	float GetDistance(const FVertex& Other)
+	{
+		return (Other.Position-Position).Size();
+	}
+};
+
+struct FHalfEdge
+{
+	// 3 vertices in CCW order
+	FVertex* Vertices[2]; 
+
+	FHalfEdge(){};
+	FHalfEdge(FVertex * A, FVertex * B)
+	{
+		Vertices[0] = A;
+		Vertices[1] = B;
+	}
+
+	bool DoesShare(const FHalfEdge& A) const
+	{
+		return (Vertices[0] == A.Vertices[1] && Vertices[1] == A.Vertices[0]);
+	}
+
+	bool operator==( const FHalfEdge& Other ) const 
+	{
+		// if same position, it's same point
+		return FMemory::Memcmp(Other.Vertices, Vertices, sizeof(Vertices)) == 0;
+	}
+};
+/**
+ * Utility class for FDelaunayTriangleGenerator 
+ * This represents Triangle in FDelaunayTriangleGenerator
+ */
+struct FTriangle
+{
+	// 3 vertices in CCW order
+	FVertex* Vertices[3]; 
+	// average points for Vertices
+	FVector2D	Center;
+	// FEdges
+	FHalfEdge Edges[3];
+
+	bool operator==( const FTriangle& Other ) const 
+	{
+		// if same position, it's same point
+		return FMemory::Memcmp(Other.Vertices, Vertices, sizeof(Vertices)) == 0;
+	}
+
+	FTriangle(FTriangle & Copy)
+	{
+		FMemory::Memcpy(Vertices, Copy.Vertices);
+		FMemory::Memcpy(Edges, Copy.Edges);
+		Center = Copy.Center;
+
+		Vertices[0]->AddTriangle(this);
+		Vertices[1]->AddTriangle(this);
+		Vertices[2]->AddTriangle(this);
+	}
+
+	FTriangle(FVertex * A, FVertex * B, FVertex * C)
+	{
+		Vertices[0] = A;
+		Vertices[1] = B;
+		Vertices[2] = C;
+		Center = (A->Position + B->Position + C->Position)/3.0f;
+
+		Vertices[0]->AddTriangle(this);
+		Vertices[1]->AddTriangle(this);
+		Vertices[2]->AddTriangle(this);
+		// when you make triangle first time, make sure it stays in CCW
+		MakeCCW();
+
+		// now create edges, this should be in the CCW order
+		Edges[0] = FHalfEdge(Vertices[0], Vertices[1]);
+		Edges[1] = FHalfEdge(Vertices[1], Vertices[2]);
+		Edges[2] = FHalfEdge(Vertices[2], Vertices[0]);
+	}
+
+	FTriangle(FVertex * A)
+	{
+		Vertices[0] = A;
+		Vertices[1] = A;
+		Vertices[2] = A;
+		Center = A->Position;
+
+		Vertices[0]->AddTriangle(this);
+		Vertices[1]->AddTriangle(this);
+		Vertices[2]->AddTriangle(this);
+
+		// now create edges, this should be in the CCW order
+		Edges[0] = FHalfEdge(Vertices[0], Vertices[1]);
+		Edges[1] = FHalfEdge(Vertices[1], Vertices[2]);
+		Edges[2] = FHalfEdge(Vertices[2], Vertices[0]);
+	}
+
+	FTriangle(FVertex * A, FVertex * B)
+	{
+		Vertices[0] = A;
+		Vertices[1] = B;
+		Vertices[2] = B;
+		Center = (A->Position + B->Position) / 2.0f;
+
+		Vertices[0]->AddTriangle(this);
+		Vertices[1]->AddTriangle(this);
+		Vertices[2]->AddTriangle(this);
+
+		// now create edges, this should be in the CCW order
+		Edges[0] = FHalfEdge(Vertices[0], Vertices[1]);
+		Edges[1] = FHalfEdge(Vertices[1], Vertices[2]);
+		Edges[2] = FHalfEdge(Vertices[2], Vertices[0]);
+	}
+	
+	FTriangle()
+	{
+		Vertices[0] = nullptr;
+		Vertices[1] = nullptr;
+		Vertices[2] = nullptr;
+	}
+
+	~FTriangle()
+	{
+		for (int32 VertexIndex = 0; VertexIndex < 3; ++VertexIndex)
+		{
+			if (Vertices[VertexIndex])
+			{
+				Vertices[VertexIndex]->RemoveTriangle(this);
+			}
+		}
+	}
+
+	bool Contains (const FVertex& Other) const
+	{
+		return (Other == *Vertices[0] || Other == *Vertices[1] || Other == *Vertices[2]);
+	}
+
+	float GetDistance(const FVertex& Other) const
+	{
+		return (Other.Position-Center).Size();
+	}
+
+	float GetDistance(const FVector2D& Other) const
+	{
+		return (Other-Center).Size();
+	}
+
+	bool HasSameHalfEdge(const FTriangle* Other) const
+	{
+		for (int32 I=0; I<3; ++I)
+		{
+			for (int32 J=0; J<3; ++J)
+			{
+				if (Other->Edges[I] == Edges[J])
+				{
+					return true;
+				}
+			}
+		}
+
+		return false;
+	}
+
+	bool DoesShareSameEdge(const FTriangle* Other) const
+	{
+		for (int32 I=0; I<3; ++I)
+		{
+			for (int32 J=0; J<3; ++J)
+			{
+				if (Other->Edges[I].DoesShare(Edges[J]))
+				{
+					return true;
+				}
+			}
+		}
+
+		return false;
+	}
+
+	// find point that doesn't share with this
+	// this should only get called if it shares same edge
+	FVertex * FindNonSharingPoint(const FTriangle* Other) const
+	{
+		if (!Contains(*Other->Vertices[0]))
+		{
+			return Other->Vertices[0];
+		}
+
+		if (!Contains(*Other->Vertices[1]))
+		{
+			return Other->Vertices[1];
+		}
+
+		if (!Contains(*Other->Vertices[2]))
+		{
+			return Other->Vertices[2];
+		}
+
+		return NULL;
+	}
+
+private:
+	void MakeCCW()
+	{
+		// this eventually has to happen on the plane that contains this 3 points
+		// for now we ignore Z
+		FVector2D Diff1 = Vertices[1]->Position-Vertices[0]->Position;
+		FVector2D Diff2 = Vertices[2]->Position-Vertices[0]->Position;
+
+		float Result = Diff1.X*Diff2.Y - Diff1.Y*Diff2.X;
+
+		check (Result != 0.f);
+
+		// it's in left side, we need this to be right side
+		if (Result < 0.f)
+		{
+			// swap 1&2 
+			FVertex * TempPt = Vertices[2];
+			Vertices[2] = Vertices[1];
+			Vertices[1] = TempPt;
+		}
+	}
+};
+
+/**
+ * Generates triangles from sample point using Delaunay Triangulation, 
+ * TriangleList contains the list of triangle after generated.
+ */
+class FDelaunayTriangleGenerator
+{
+public:
+	enum ECircumCircleState
+	{
+		ECCS_Outside = -1,
+		ECCS_On = 0,
+		ECCS_Inside = 1,
+	};
+	/** 
+	 * Reset all data
+	 */
+	void Reset();
+	void EmptyTriangles();
+	void EmptySamplePoints();
+
+	/**
+	 * Generate triangles from SamplePointList
+	 */
+	void Triangulate(EPreferredTriangulationDirection PreferredTriangulationDirection);
+
+	/**
+	 * Converts from our internal representation to the runtime triangles
+	 */
+	TArray<struct FBlendSpaceTriangle> CalculateTriangles() const;
+
+	/** 
+	 * Add new sample point to SamplePointList and its corresponding sample index in the blendspace
+	 * It won't be added if already exists
+	 */
+	void AddSamplePoint(const FVector2D& NewPoint, const int32 SampleIndex);
+
+	/** 
+	 * This is for debug purpose to step only one to triangulate
+	 * StartIndex should increase manually
+	 */
+	void Step(int32 StartIndex);
+
+	~FDelaunayTriangleGenerator()
+	{
+		Reset();
+	}
+
+	/** 
+	 * Get TriangleList
+	 */
+	const TArray<FTriangle*> & GetTriangleList() const { return TriangleList; };
+
+	/** 
+	 * Get SamplePointList
+	 */
+	const TArray<FVertex> & GetSamplePointList() const { return SamplePointList; };
+
+	/* Set the grid box, so we can normalize the sample points */
+	void SetGridBox(const FBlendParameter& BlendParamX, const FBlendParameter& BlendParamY);
+
+private:
+	/**
+	 * Sort SamplePointList in the order of 1) -X->+X 2) -Y->+Y 3) -Z->+Z
+	 * Sort Point of the input point in the order of
+	 * 1) -> +x. 2) -> +y, 3) -> +z
+	 * by that, it will sort from -X to +X, if same is found, next will be Y, and if same, next will be Z
+	 */
+	void SortSamples();
+
+	/**
+	 * Tries to make the split of ambiguous triangles that are aligned with the grid such that they are symmetric.
+	 */
+	void AdjustEdgeDirections(EPreferredTriangulationDirection PreferredTriangulationDirection);
+
+	/**
+	 * Finds the triangle index that contains an edge referencing the samples in that order.
+	 * Optionally returns the index of the vertex in the found triangle that is the start of the edge.
+	 */
+	int32 FindTriangleIndexWithEdge(int32 SampleIndex0, int32 SampleIndex1, int32* VertexIndex = nullptr) const;
+
+	/**
+	 * The key function in Delaunay Triangulation
+	 * return true if the TestPoint is WITHIN the triangle circumcircle
+	 *	http://en.wikipedia.org/wiki/Delaunay_triangulation 
+	 */
+	ECircumCircleState GetCircumcircleState(const FTriangle* T, const FVertex& TestPoint);
+
+	/**
+	 * return true if they can make triangle
+	 */
+	bool IsEligibleForTriangulation(const FVertex* A, const FVertex* B, const FVertex* C)
+	{
+		return (IsCollinear(A, B, C)==false);
+	}
+
+	/** 
+	 * return true if 3 points are collinear
+	 * by that if those 3 points create straight line
+	 */
+	bool IsCollinear(const FVertex* A, const FVertex* B, const FVertex* C);
+
+	/** 
+	 * return true if all points are coincident
+	 * (i.e. if all points are the same)
+	 */
+	bool AllCoincident(const TArray<FVertex>& InPoints);
+
+	/**
+	 * Flip TriangleList(I) with TriangleList(J). 
+	 */
+	bool FlipTriangles(const int32 TriangleIndexOne, const int32 TriangleIndexTwo);
+
+	/** 
+	 * Add new Triangle
+	 */
+	void AddTriangle(FTriangle & newTriangle, bool bCheckHalfEdge=true);
+
+	/** 
+	 * Used as incremental step to triangulate all points
+	 * Create triangles TotalNum of PointList
+	 */
+	int32 GenerateTriangles(TArray<FVertex> & PointList, const int32 TotalNum);
+
+private:
+	/**
+	 * Data Structure for triangulation
+	 * SamplePointList is the input data
+	 */
+	TArray<FVertex>		SamplePointList;
+
+	/**
+	 * Output data for this generator
+	 */
+	TArray<FTriangle*>	TriangleList;
+
+	/**
+	 * We need to normalize the points before the circumcircle test, so we need the extents of the grid.
+	 * Store 1 / GridSize to avoid recalculating it every time through GetCircumcircleState
+	*/
+	FVector2D				GridMin;
+	FVector2D				RecipGridSize;
+};
+
+// @todo fixmelh : this code is mess between fvector2D and fvector
+// ideally FBlendSpaceGrid will be handled in 3D, and SBlendSpaceGridWidget only knows about 2D
+// in the future this should all change to 3D
+
+/**
+ * BlendSpace Grid
+ * Using triangulated space, create FEditorElement of each grid point using SamplePoint
+ */
+class FBlendSpaceGrid
+{
+public:
+	/**
+	 * Find Triangle this TestPoint is within
+	 * 
+	 * @param	TestPoint				Point to test
+	 * @param	OutBaryCentricCoords	Output BaryCentricCoords2D of the point in the triangle // for now it's only 2D
+	 * @param	OutTriangle				The triangle that this point is within or lie
+	 * @param	TriangleList			TriangleList to test
+	 * 
+	 * @return	true if successfully found the triangle this point is within
+	 */
+	bool FindTriangleThisPointBelongsTo(const FVector2D& TestPoint, FVector& OutBarycentricCoords, FTriangle* & OutTriangle, const TArray<FTriangle*> & TriangleList) const;
+
+	/**  
+	 * Fill up Grid GridPoints using TriangleList input - Grid information should have been set by SetGridInfo
+	 * 
+	 * @param	SamplePoints		: Sample Point List
+	 * @param	TriangleList		: List of triangles
+	 */
+	void GenerateGridElements(const TArray<FVertex>& SamplePoints, const TArray<FTriangle*> & TriangleList);
+
+	/** 
+	 * default value 
+	 */
+	FBlendSpaceGrid()
+		: GridMin(0, 0)
+		, GridMax(100, 100)
+		, NumGridPointsForAxis(5, 5)
+	{
+	}
+
+	void Reset()
+	{
+		GridPoints.Empty();
+	}
+
+	void SetGridInfo(const FBlendParameter& BlendParamX, const FBlendParameter& BlendParamY)
+	{
+		NumGridPointsForAxis.X = BlendParamX.GridNum + 1;
+		NumGridPointsForAxis.Y = BlendParamY.GridNum + 1;
+
+		NumGridDivisions.X = BlendParamX.GridNum;
+		NumGridDivisions.Y = BlendParamY.GridNum;
+
+		GridMin.X = BlendParamX.Min;
+		GridMax.X = BlendParamX.Max;
+
+		GridMin.Y = BlendParamY.Min;
+		GridMax.Y = BlendParamY.Max;
+	}
+	
+	const FEditorElement& GetElement(const int32 GridX, const int32 GridY) const;
+	const TArray<FEditorElement>& GetElements() const { return GridPoints;}
+	
+	/** 
+	 * Convert grid index (GridX, GridY) to triangle coords and returns FVector2D
+	*/
+	const FVector2D GetPosFromIndex(const int32 GridX, const int32 GridY) const;
+
+private:
+	// Grid Dimension
+	FVector2D GridMin;
+	FVector2D GridMax;
+
+	// how many rows/cols for each axis
+	FIntPoint NumGridPointsForAxis;
+	FIntPoint NumGridDivisions;
+
+	// Each point data -output data
+	TArray<FEditorElement> GridPoints; // 2D array saved in 1D array, to search (x, y), x*GridSizeX+y;
+};
diff --git a/Engine/Source/Runtime/Slate/Private/Framework/Commands/UICommandList.cpp b/Engine/Source/Runtime/Slate/Private/Framework/Commands/UICommandList.cpp
--- a/Engine/Source/Runtime/Slate/Private/Framework/Commands/UICommandList.cpp
+++ b/Engine/Source/Runtime/Slate/Private/Framework/Commands/UICommandList.cpp
@@ -58,9 +58,6 @@ void FUICommandList::MapAction( const TSharedPtr< const FUICommandInfo > InUICom
 		return;
 	}
 
-	// Check against already mapped actions
-	checkfSlow(!UICommandBindingMap.Contains(InUICommandInfo), TEXT("Command list already contains a command named '%s'"), *InUICommandInfo->GetCommandName().ToString());
-
 	ContextsInList.Add( InUICommandInfo->GetBindingContext() );
 	UICommandBindingMap.Add( InUICommandInfo, InUIAction );
 }
