Fix for linked anim layers not unlinking.
#rnx

diff --git a/Engine/Source/Runtime/Engine/Private/Animation/AnimInstance.cpp b/Engine/Source/Runtime/Engine/Private/Animation/AnimInstance.cpp
--- a/Engine/Source/Runtime/Engine/Private/Animation/AnimInstance.cpp
+++ b/Engine/Source/Runtime/Engine/Private/Animation/AnimInstance.cpp
@@ -2680,12 +2680,9 @@ void UAnimInstance::PerformLinkedLayerOverlayOperation(TSubclassOf<UAnimInstance
 					if(FoundFunction->bImplemented)
 					{
 						UClass* ClassToSet = InClassSelectorFunction(NewClass, Layer);
-						if (ClassToSet != nullptr)
-						{
-							TMap<FName, TArray<FAnimNode_LinkedAnimLayer*, TInlineAllocator<4>>, TInlineSetAllocator<4>>& ClassLayerNodesToSet = LayerNodesToSet.FindOrAdd(ClassToSet);
-							TArray<FAnimNode_LinkedAnimLayer*, TInlineAllocator<4>>& LayerNodes = ClassLayerNodesToSet.FindOrAdd(FoundFunction->Group);
-							LayerNodes.Add(Layer);
-						}
+						TMap<FName, TArray<FAnimNode_LinkedAnimLayer*, TInlineAllocator<4>>, TInlineSetAllocator<4>>& ClassLayerNodesToSet = LayerNodesToSet.FindOrAdd(ClassToSet);
+						TArray<FAnimNode_LinkedAnimLayer*, TInlineAllocator<4>>& LayerNodes = ClassLayerNodesToSet.FindOrAdd(FoundFunction->Group);
+						LayerNodes.Add(Layer);
 					}
 				}
 			}
@@ -2693,12 +2690,9 @@ void UAnimInstance::PerformLinkedLayerOverlayOperation(TSubclassOf<UAnimInstance
 			{
 				// Add null classes so we clear the node's instance below
 				UClass* ClassToSet = InClassSelectorFunction(NewClass, Layer);
-				if (ClassToSet != nullptr)
-				{
-					TMap<FName, TArray<FAnimNode_LinkedAnimLayer*, TInlineAllocator<4>>, TInlineSetAllocator<4>>& ClassLayerNodesToSet = LayerNodesToSet.FindOrAdd(ClassToSet);
-					TArray<FAnimNode_LinkedAnimLayer*, TInlineAllocator<4>>& LayerNodes = ClassLayerNodesToSet.FindOrAdd(NAME_None);
-					LayerNodes.Add(Layer);
-				}
+				TMap<FName, TArray<FAnimNode_LinkedAnimLayer*, TInlineAllocator<4>>, TInlineSetAllocator<4>>& ClassLayerNodesToSet = LayerNodesToSet.FindOrAdd(ClassToSet);
+				TArray<FAnimNode_LinkedAnimLayer*, TInlineAllocator<4>>& LayerNodes = ClassLayerNodesToSet.FindOrAdd(NAME_None);
+				LayerNodes.Add(Layer);
 			}
 		}
 
