Fix hardcoded ERHIFeatureLevel::SM5 (code by [REDACTED_PERSON])

diff --git a/Engine/Source/Runtime/Engine/Private/Rendering/SkyLightImportanceSampling.cpp b/Engine/Source/Runtime/Engine/Private/Rendering/SkyLightImportanceSampling.cpp
--- a/Engine/Source/Runtime/Engine/Private/Rendering/SkyLightImportanceSampling.cpp
+++ b/Engine/Source/Runtime/Engine/Private/Rendering/SkyLightImportanceSampling.cpp
@@ -225,6 +225,7 @@ IMPLEMENT_SHADER_TYPE(, FBuildSolidAnglePdfCS, TEXT("/Engine/Private/Raytracing/
 
 void BuildSkyLightMipTree(
 	FRHICommandListImmediate& RHICmdList,
+	ERHIFeatureLevel::Type FeatureLevel,
 	FTextureRHIRef SkyLightTexture,
 	FRWBuffer& SkyLightMipTreePosX,
 	FRWBuffer& SkyLightMipTreeNegX,
@@ -235,7 +236,7 @@ void BuildSkyLightMipTree(
 	FIntVector& SkyLightMipTreeDimensions
 )
 {
-	const auto ShaderMap = GetGlobalShaderMap(ERHIFeatureLevel::SM5);
+	const auto ShaderMap = GetGlobalShaderMap(FeatureLevel);
 	TShaderMapRef<FBuildMipTreeCS> BuildSkyLightMipTreeComputeShader(ShaderMap);
 	RHICmdList.SetComputeShader(BuildSkyLightMipTreeComputeShader.GetComputeShader());
 
@@ -287,11 +288,12 @@ void BuildSkyLightMipTree(
 
 void BuildSolidAnglePdf(
 	FRHICommandListImmediate& RHICmdList,
+	ERHIFeatureLevel::Type FeatureLevel,
 	const FIntVector& Dimensions,
 	FRWBuffer& SolidAnglePdf
 )
 {
-	const auto ShaderMap = GetGlobalShaderMap(ERHIFeatureLevel::SM5);
+	const auto ShaderMap = GetGlobalShaderMap(FeatureLevel);
 	TShaderMapRef<FBuildSolidAnglePdfCS> BuildSolidAnglePdfComputeShader(ShaderMap);
 	RHICmdList.SetComputeShader(BuildSolidAnglePdfComputeShader.GetComputeShader());
 
@@ -314,6 +316,7 @@ void BuildSolidAnglePdf(
 
 void BuildSkyLightMipTreePdf(
 	FRHICommandListImmediate& RHICmdList,
+	ERHIFeatureLevel::Type FeatureLevel,
 	const FRWBuffer& SkyLightMipTreePosX,
 	const FRWBuffer& SkyLightMipTreeNegX,
 	const FRWBuffer& SkyLightMipTreePosY,
@@ -347,7 +350,7 @@ void BuildSkyLightMipTreePdf(
 		&SkyLightMipTreePdfNegZ
 	};
 
-	const auto ShaderMap = GetGlobalShaderMap(ERHIFeatureLevel::SM5);
+	const auto ShaderMap = GetGlobalShaderMap(FeatureLevel);
 	TShaderMapRef<FBuildMipTreePdfCS> BuildSkyLightMipTreePdfComputeShader(ShaderMap);
 	RHICmdList.SetComputeShader(BuildSkyLightMipTreePdfComputeShader.GetComputeShader());
 
@@ -391,7 +394,8 @@ void FSkyLightImportanceSamplingData::BuildCDFs(FTexture* ProcessedTexture)
 
 #if RHI_RAYTRACING
 	BuildSkyLightMipTree(
-		RHICmdList, 
+		RHICmdList,
+		GetFeatureLevel(),
 		ProcessedTexture->TextureRHI, 
 		MipTreePosX, 
 		MipTreeNegX, 
@@ -404,6 +408,7 @@ void FSkyLightImportanceSamplingData::BuildCDFs(FTexture* ProcessedTexture)
 
 	BuildSkyLightMipTreePdf(
 		RHICmdList,
+		GetFeatureLevel(),
 		MipTreePosX,
 		MipTreeNegX, 
 		MipTreePosY, 
@@ -418,7 +423,7 @@ void FSkyLightImportanceSamplingData::BuildCDFs(FTexture* ProcessedTexture)
 		MipTreePdfPosZ, 
 		MipTreePdfNegZ);
 
-	BuildSolidAnglePdf(RHICmdList, MipDimensions, SolidAnglePdf);
+	BuildSolidAnglePdf(RHICmdList, GetFeatureLevel(), MipDimensions, SolidAnglePdf);
 #endif
 
 	bIsValid = true;
