WorldPartition: Serialize actor's folder as part of the actordesc

diff --git a/Engine/Source/Runtime/Core/Public/UObject/UE5ReleaseStreamObjectVersion.h b/Engine/Source/Runtime/Core/Public/UObject/UE5ReleaseStreamObjectVersion.h
--- a/Engine/Source/Runtime/Core/Public/UObject/UE5ReleaseStreamObjectVersion.h
+++ b/Engine/Source/Runtime/Core/Public/UObject/UE5ReleaseStreamObjectVersion.h
@@ -25,7 +25,10 @@ struct CORE_API FUE5ReleaseStreamObjectVersion
 		LevelInstanceSerializeRuntimeBehavior,
 
 		// Refactoring Pose Asset runtime data structures
-        PoseAssetRuntimeRefactor,
+		PoseAssetRuntimeRefactor,
+
+		// Serialize the folder path of actor descs
+		WorldPartitionActorDescSerializeActorFolderPath,
 
 		// -----<new versions can be added above this line>-------------------------------------------------
 		VersionPlusOne,
diff --git a/Engine/Source/Runtime/Engine/Private/WorldPartition/WorldPartitionActorDesc.cpp b/Engine/Source/Runtime/Engine/Private/WorldPartition/WorldPartitionActorDesc.cpp
--- a/Engine/Source/Runtime/Engine/Private/WorldPartition/WorldPartitionActorDesc.cpp
+++ b/Engine/Source/Runtime/Engine/Private/WorldPartition/WorldPartitionActorDesc.cpp
@@ -68,6 +68,7 @@ void FWorldPartitionActorDesc::Init(const AActor* InActor)
 	DataLayers = InActor->GetDataLayerNames();
 	ActorPackage = InActor->GetPackage()->GetFName();
 	ActorPath = *InActor->GetPathName();
+	FolderPath = InActor->GetFolderPath();
 	
 	TArray<AActor*> ActorReferences = ActorsReferencesUtils::GetExternalActorReferences((AActor*)InActor);
 
@@ -198,6 +199,11 @@ void FWorldPartitionActorDesc::Serialize(FArchive& Ar)
 		bActorIsHLODRelevant = true;
 		HLODLayer = FName();
 	}
+
+	if (Ar.CustomVer(FUE5ReleaseStreamObjectVersion::GUID) >= FUE5ReleaseStreamObjectVersion::WorldPartitionActorDescSerializeActorFolderPath)
+	{
+		Ar << FolderPath;
+	}
 }
 
 FBox FWorldPartitionActorDesc::GetBounds() const
diff --git a/Engine/Source/Runtime/Engine/Public/WorldPartition/WorldPartitionActorDesc.h b/Engine/Source/Runtime/Engine/Public/WorldPartition/WorldPartitionActorDesc.h
--- a/Engine/Source/Runtime/Engine/Public/WorldPartition/WorldPartitionActorDesc.h
+++ b/Engine/Source/Runtime/Engine/Public/WorldPartition/WorldPartitionActorDesc.h
@@ -70,6 +70,7 @@ public:
 	inline FName GetActorPackage() const { return ActorPackage; }
 	inline FName GetActorPath() const { return ActorPath; }
 	inline FName GetActorLabel() const { return ActorLabel; }
+	inline FName GetFolderPath() const { return FolderPath; }
 	FBox GetBounds() const;
 
 	virtual bool GetContainerInstance(const UActorDescContainer*& OutLevelContainer, FTransform& OutLevelTransform, EContainerClusterMode& OutClusterMode) const { return false; }
@@ -160,6 +161,7 @@ protected:
 	FName							HLODLayer;
 	TArray<FName>					DataLayers;
 	TArray<FGuid>					References;
+	FName							FolderPath;
 	
 	// Transient
 	mutable uint32					SoftRefCount;
