Fix for Audio Meter in MetaSound Editor

diff --git a/Engine/Source/Runtime/Engine/Private/SoundWave.cpp b/Engine/Source/Runtime/Engine/Private/SoundWave.cpp
--- a/Engine/Source/Runtime/Engine/Private/SoundWave.cpp
+++ b/Engine/Source/Runtime/Engine/Private/SoundWave.cpp
@@ -2248,9 +2248,10 @@ void USoundWave::Parse(FAudioDevice* AudioDevice, const UPTRINT NodeWaveInstance
 			WaveInstance->bEnableSubmixSends = ActiveSound.bEnableSubmixSendRoutingOverride;
 		}
 	}
-	else //if this is a preview sound, ignore sends and only play the base submix
+	else //if this is a preview sound, ignore Submix sends
 	{
 		WaveInstance->bEnableBaseSubmix = true;
+		WaveInstance->bEnableBusSends = ParseParams.bEnableBusSends;
 	}
 
 	for (int32 BusSendType = 0; BusSendType < (int32)EBusSendType::Count; ++BusSendType)
