Initial support for material functions, using new HLSL translator

diff --git a/Engine/Source/Runtime/Engine/Classes/Materials/MaterialExpressionFunctionInput.h b/Engine/Source/Runtime/Engine/Classes/Materials/MaterialExpressionFunctionInput.h
--- a/Engine/Source/Runtime/Engine/Classes/Materials/MaterialExpressionFunctionInput.h
+++ b/Engine/Source/Runtime/Engine/Classes/Materials/MaterialExpressionFunctionInput.h
@@ -110,6 +110,7 @@ class UMaterialExpressionFunctionInput : public UMaterialExpression
 	virtual bool IsResultMaterialAttributes(int32 OutputIndex) override;
 	virtual uint32 GetInputType(int32 InputIndex) override;
 	virtual uint32 GetOutputType(int32 OutputIndex) override;
+	virtual EMaterialGenerateHLSLStatus GenerateHLSLExpression(FMaterialHLSLGenerator& Generator, UE::HLSLTree::FScope& Scope, int32 OutputIndex, UE::HLSLTree::FExpression*& OutExpression) override;
 #endif // WITH_EDITOR
 	//~ End UMaterialExpression Interface
 
diff --git a/Engine/Source/Runtime/Engine/Classes/Materials/MaterialExpressionMaterialFunctionCall.h b/Engine/Source/Runtime/Engine/Classes/Materials/MaterialExpressionMaterialFunctionCall.h
--- a/Engine/Source/Runtime/Engine/Classes/Materials/MaterialExpressionMaterialFunctionCall.h
+++ b/Engine/Source/Runtime/Engine/Classes/Materials/MaterialExpressionMaterialFunctionCall.h
@@ -162,6 +162,8 @@ class UMaterialExpressionMaterialFunctionCall : public UMaterialExpression
 		SharedCompileState = SharedState;
 	}
 
+	virtual EMaterialGenerateHLSLStatus GenerateHLSLExpression(FMaterialHLSLGenerator& Generator, UE::HLSLTree::FScope& Scope, int32 OutputIndex, UE::HLSLTree::FExpression*& OutExpression) override;
+
 private:	
 	/** Helper that fixes up expression links where possible. */
 	void FixupReferencingExpressions(
diff --git a/Engine/Source/Runtime/Engine/Classes/Materials/MaterialFunctionInterface.h b/Engine/Source/Runtime/Engine/Classes/Materials/MaterialFunctionInterface.h
--- a/Engine/Source/Runtime/Engine/Classes/Materials/MaterialFunctionInterface.h
+++ b/Engine/Source/Runtime/Engine/Classes/Materials/MaterialFunctionInterface.h
@@ -17,6 +17,8 @@
 class UMaterial;
 class UTexture;
 struct FPropertyChangedEvent;
+class FMaterialHLSLGenerator;
+class FMaterialHLSLTree;
 
 /** Usage set on a material function determines feature compatibility and validation. */
 UENUM()
@@ -44,6 +46,15 @@ class UMaterialFunctionInterface : public UObject
 	UPROPERTY(duplicatetransient)
 	FGuid StateId;
 
+#if WITH_EDITOR
+public:
+	FMaterialHLSLTree& AcquireHLSLTree(FMaterialHLSLGenerator& Generator);
+#endif
+#if WITH_EDITORONLY_DATA
+private:
+	FMaterialHLSLTree* CachedHLSLTree = nullptr;
+#endif
+
 protected:
 	/** The intended usage of this function, required for material layers. */
 	UPROPERTY(AssetRegistrySearchable)
diff --git a/Engine/Source/Runtime/Engine/Private/HLSLTree/HLSLTree.cpp b/Engine/Source/Runtime/Engine/Private/HLSLTree/HLSLTree.cpp
--- a/Engine/Source/Runtime/Engine/Private/HLSLTree/HLSLTree.cpp
+++ b/Engine/Source/Runtime/Engine/Private/HLSLTree/HLSLTree.cpp
@@ -309,6 +309,46 @@ const TCHAR* UE::HLSLTree::FEmitContext::AcquireHLSLReference(UE::HLSLTree::FExp
 	return Entry.Definition;
 }
 
+const TCHAR* UE::HLSLTree::FEmitContext::AcquireHLSLReference(FFunctionCall* FunctionCall, int32 OutputIndex)
+{
+	check(FunctionCall);
+
+	FFunctionCallEntry& Entry = FunctionCallMap.FindOrAdd(FunctionCall);
+	if (!Entry.OutputDefinition)
+	{
+		FScopeEntry* ScopeEntry = FindScope(FunctionCall->ParentScope);
+		check(ScopeEntry);
+
+		TCHAR const** OutputDefinition = New<const TCHAR*>(*Allocator, FunctionCall->NumOutputs);
+		for (int32 i = 0; i < FunctionCall->NumOutputs; ++i)
+		{
+			OutputDefinition[i] = AllocateStringf(*Allocator, TEXT("Local%d"), NumExpressionLocals++);
+			const FExpressionTypeDescription& TypeDesc = GetExpressionTypeDescription(FunctionCall->OutputTypes[i]);
+			ScopeEntry->ExpressionCodeWriter->WriteLinef(TEXT("%s %s = (%s)0.0f;"),
+				TypeDesc.Name,
+				OutputDefinition[i],
+				TypeDesc.Name);
+		}
+
+		Entry.OutputDefinition = OutputDefinition;
+		Entry.NumOutputs = FunctionCall->NumOutputs;
+
+		FFunctionStackEntry& StackEntry = FunctionStack.AddDefaulted_GetRef();
+		StackEntry.FunctionCall = FunctionCall;
+
+		FunctionCall->FunctionScope->EmitHLSL(*this, *ScopeEntry->ExpressionCodeWriter);
+
+		{
+			const FFunctionStackEntry PoppedStackEntry = FunctionStack.Pop();
+			check(PoppedStackEntry.FunctionCall == FunctionCall);
+		}
+	}
+
+	check(Entry.NumOutputs == FunctionCall->NumOutputs);
+	check(OutputIndex >= 0 && OutputIndex < Entry.NumOutputs);
+	return Entry.OutputDefinition[OutputIndex];
+}
+
 void UE::HLSLTree::FNodeVisitor::VisitNode(UE::HLSLTree::FNode* Node)
 {
 	if (Node)
@@ -352,6 +392,20 @@ UE::HLSLTree::ENodeVisitResult UE::HLSLTree::FTextureParameterDeclaration::Visit
 	return Visitor.OnTextureParameterDeclaration(*this);
 }
 
+UE::HLSLTree::ENodeVisitResult UE::HLSLTree::FFunctionCall::Visit(FNodeVisitor& Visitor)
+{
+	const ENodeVisitResult Result = Visitor.OnFunctionCall(*this);
+	if (ShouldVisitDependentNodes(Result))
+	{
+		// Don't visit the function scope, as that is from a different tree
+		for (int32 i = 0; i < NumInputs; ++i)
+		{
+			Visitor.VisitNode(Inputs[i]);
+		}
+	}
+	return Result;
+}
+
 UE::HLSLTree::ENodeVisitResult UE::HLSLTree::FScope::Visit(UE::HLSLTree::FNodeVisitor& Visitor)
 {
 	const ENodeVisitResult Result = Visitor.OnScope(*this);
@@ -473,6 +527,12 @@ public:
 		return ENodeVisitResult::VisitDependentNodes;
 	}
 
+	virtual ENodeVisitResult OnFunctionCall(FFunctionCall& InFunctionCall) override
+	{
+		Scope->UseNode(&InFunctionCall);
+		return ENodeVisitResult::VisitDependentNodes;
+	}
+
 	FScope* Scope;
 };
 } // namespace HLSLTree
@@ -485,6 +545,12 @@ void UE::HLSLTree::FScope::UseExpression(UE::HLSLTree::FExpression* Expression)
 	Visitor.VisitNode(Expression);
 }
 
+void UE::HLSLTree::FScope::UseFunctionCall(FFunctionCall* FunctionCall)
+{
+	FNodeVisitor_MoveToScope Visitor(this);
+	Visitor.VisitNode(FunctionCall);
+}
+
 void UE::HLSLTree::FScope::AddStatement(UE::HLSLTree::FStatement* Statement)
 {
 	check(!Statement->ParentScope);
@@ -601,9 +667,32 @@ UE::HLSLTree::FParameterDeclaration* UE::HLSLTree::FTree::NewParameterDeclaratio
 	Declaration->ParentScope = &Scope;
 	return Declaration;
 }
+
 UE::HLSLTree::FTextureParameterDeclaration* UE::HLSLTree::FTree::NewTextureParameterDeclaration(UE::HLSLTree::FScope& Scope, const FName& Name, const UE::HLSLTree::FTextureDescription& DefaultValue)
 {
 	FTextureParameterDeclaration* Declaration = NewNode<FTextureParameterDeclaration>(Name, DefaultValue);
 	Declaration->ParentScope = &Scope;
 	return Declaration;
 }
+
+UE::HLSLTree::FFunctionCall* UE::HLSLTree::FTree::NewFunctionCall(FScope& Scope,
+	const FScope& InFunctionScope,
+	FExpression* const* InInputs,
+	EExpressionType const* InOutputTypes,
+	int32 InNumInputs,
+	int32 InNumOutputs)
+{
+	FExpression** Inputs = New<FExpression*>(*Allocator, InNumInputs);
+	EExpressionType* OutputTypes = New<EExpressionType>(*Allocator, InNumOutputs);
+	FMemory::Memcpy(Inputs, InInputs, InNumInputs * sizeof(FExpression*));
+	FMemory::Memcpy(OutputTypes, InOutputTypes, InNumOutputs * sizeof(FExpression*));
+
+	FFunctionCall* FunctionCall = NewNode<FFunctionCall>();
+	FunctionCall->ParentScope = &Scope;
+	FunctionCall->FunctionScope = &InFunctionScope;
+	FunctionCall->Inputs = Inputs;
+	FunctionCall->OutputTypes = OutputTypes;
+	FunctionCall->NumInputs = InNumInputs;
+	FunctionCall->NumOutputs = InNumOutputs;
+	return FunctionCall;
+}
diff --git a/Engine/Source/Runtime/Engine/Private/HLSLTree/HLSLTreeCommon.cpp b/Engine/Source/Runtime/Engine/Private/HLSLTree/HLSLTreeCommon.cpp
--- a/Engine/Source/Runtime/Engine/Private/HLSLTree/HLSLTreeCommon.cpp
+++ b/Engine/Source/Runtime/Engine/Private/HLSLTree/HLSLTreeCommon.cpp
@@ -360,6 +360,28 @@ void UE::HLSLTree::FExpressionCast::EmitHLSL(FEmitContext& Context, FCodeWriter&
 	}
 }
 
+void UE::HLSLTree::FExpressionFunctionInput::EmitHLSL(FEmitContext& Context, FCodeWriter& Writer) const
+{
+	const FEmitContext::FFunctionStackEntry& StackEntry = Context.FunctionStack.Last();
+	FExpression* InputExpression = StackEntry.FunctionCall->Inputs[InputIndex];
+	const TCHAR* Ref = Context.AcquireHLSLReference(InputExpression);
+	Writer.Writef(TEXT("%s"), Ref);
+}
+
+void UE::HLSLTree::FExpressionFunctionOutput::EmitHLSL(FEmitContext& Context, FCodeWriter& Writer) const
+{
+	const TCHAR* Ref = Context.AcquireHLSLReference(FunctionCall, OutputIndex);
+	Writer.Writef(TEXT("%s"), Ref);
+}
+
+void UE::HLSLTree::FStatementSetFunctionOutput::EmitHLSL(FEmitContext& Context, FCodeWriter& Writer) const
+{
+	const FEmitContext::FFunctionStackEntry& StackEntry = Context.FunctionStack.Last();
+	const TCHAR* OutputRef = Context.AcquireHLSLReference(StackEntry.FunctionCall, OutputIndex);
+	const TCHAR* ExpressionRef = Context.AcquireHLSLReference(Expression);
+	Writer.WriteLinef(TEXT("%s = %s;"), OutputRef, ExpressionRef);
+}
+
 void UE::HLSLTree::FStatementReturn::EmitHLSL(UE::HLSLTree::FEmitContext& Context, UE::HLSLTree::FCodeWriter& Writer) const
 {
 	const TCHAR* Ref = Context.AcquireHLSLReference(Expression);
diff --git a/Engine/Source/Runtime/Engine/Private/Materials/MaterialExpressionHLSL.cpp b/Engine/Source/Runtime/Engine/Private/Materials/MaterialExpressionHLSL.cpp
--- a/Engine/Source/Runtime/Engine/Private/Materials/MaterialExpressionHLSL.cpp
+++ b/Engine/Source/Runtime/Engine/Private/Materials/MaterialExpressionHLSL.cpp
@@ -3,11 +3,12 @@
 
 #if WITH_EDITOR
 
-#include "Materials/MaterialExpression.h"
 #include "MaterialHLSLGenerator.h"
+#include "MaterialHLSLTree.h"
 #include "HLSLTree/HLSLTree.h"
 #include "HLSLTree/HLSLTreeCommon.h"
 
+#include "Materials/MaterialExpression.h"
 #include "Materials/MaterialExpressionExecBegin.h"
 #include "Materials/MaterialExpressionReturnMaterialAttributes.h"
 #include "Materials/MaterialExpressionConstant.h"
@@ -20,12 +21,15 @@
 #include "Materials/MaterialExpressionTextureSampleParameter.h"
 #include "Materials/MaterialExpressionTextureObject.h"
 #include "Materials/MaterialExpressionTextureObjectParameter.h"
+#include "Materials/MaterialExpressionFunctionInput.h"
+#include "Materials/MaterialExpressionMaterialFunctionCall.h"
 #include "Materials/MaterialExpressionGetLocal.h"
 #include "Materials/MaterialExpressionAdd.h"
 #include "Materials/MaterialExpressionSetMaterialAttributes.h"
 #include "Materials/MaterialExpressionSetLocal.h"
 #include "Materials/MaterialExpressionIfThenElse.h"
 #include "Materials/MaterialExpressionForLoop.h"
+#include "Materials/MaterialFunctionInterface.h"
 
 EMaterialGenerateHLSLStatus UMaterialExpression::GenerateHLSLExpression(FMaterialHLSLGenerator& Generator, UE::HLSLTree::FScope& Scope, int32 OutputIndex, UE::HLSLTree::FExpression*& OutExpression)
 {
@@ -42,6 +46,16 @@ EMaterialGenerateHLSLStatus UMaterialExpression::GenerateHLSLTexture(FMaterialHL
 	return Generator.Error(TEXT("Node does not support textures"));
 }
 
+FMaterialHLSLTree& UMaterialFunctionInterface::AcquireHLSLTree(FMaterialHLSLGenerator& Generator)
+{
+	if (!CachedHLSLTree)
+	{
+		CachedHLSLTree = new FMaterialHLSLTree();
+		CachedHLSLTree->InitializeForFunction(Generator.GetCompileTarget(), this);
+	}
+	return *CachedHLSLTree;
+}
+
 EMaterialGenerateHLSLStatus UMaterialExpressionConstant::GenerateHLSLExpression(FMaterialHLSLGenerator& Generator, UE::HLSLTree::FScope& Scope, int32 OutputIndex, UE::HLSLTree::FExpression*& OutExpression)
 {
 	OutExpression = Generator.NewConstant(Scope, R);
@@ -207,6 +221,56 @@ EMaterialGenerateHLSLStatus UMaterialExpressionSetMaterialAttributes::GenerateHL
 	return EMaterialGenerateHLSLStatus::Success;
 }
 
+EMaterialGenerateHLSLStatus UMaterialExpressionFunctionInput::GenerateHLSLExpression(FMaterialHLSLGenerator& Generator, UE::HLSLTree::FScope& Scope, int32 OutputIndex, UE::HLSLTree::FExpression*& OutExpression)
+{
+	UE::HLSLTree::EExpressionType ExpressionType = UE::HLSLTree::EExpressionType::Void;
+	switch (InputType)
+	{
+	case FunctionInput_Scalar: ExpressionType = UE::HLSLTree::EExpressionType::Float1; break;
+	case FunctionInput_Vector2: ExpressionType = UE::HLSLTree::EExpressionType::Float2; break;
+	case FunctionInput_Vector3: ExpressionType = UE::HLSLTree::EExpressionType::Float3; break;
+	case FunctionInput_Vector4: ExpressionType = UE::HLSLTree::EExpressionType::Float4; break;
+	case FunctionInput_Texture2D:
+	case FunctionInput_TextureCube:
+	case FunctionInput_Texture2DArray:
+	case FunctionInput_VolumeTexture:
+	case FunctionInput_StaticBool:
+	case FunctionInput_MaterialAttributes:
+	case FunctionInput_TextureExternal:
+		break;
+	}
+	if (ExpressionType == UE::HLSLTree::EExpressionType::Void)
+	{
+		return Generator.Error(TEXT("Invalid input connection"));
+	}
+
+	UE::HLSLTree::FLocalDeclaration* Declaration = Generator.AcquireLocalDeclaration(Scope, ExpressionType, InputName);
+	OutExpression = Generator.GetTree().NewExpression<UE::HLSLTree::FExpressionLocalVariable>(Scope, Declaration);
+	return EMaterialGenerateHLSLStatus::Success;
+}
+
+EMaterialGenerateHLSLStatus UMaterialExpressionMaterialFunctionCall::GenerateHLSLExpression(FMaterialHLSLGenerator& Generator, UE::HLSLTree::FScope& Scope, int32 OutputIndex, UE::HLSLTree::FExpression*& OutExpression)
+{
+	if (!MaterialFunction)
+	{
+		return Generator.Error(TEXT("Missing function"));
+	}
+
+	TArray<UE::HLSLTree::FExpression*> InputExpressions;
+	InputExpressions.Empty(FunctionInputs.Num());
+
+	for (int32 InputIndex = 0; InputIndex < FunctionInputs.Num(); ++InputIndex)
+	{
+		const FFunctionExpressionInput& Input = FunctionInputs[InputIndex];
+		UE::HLSLTree::FExpression* InputExpression = Input.Input.AcquireHLSLExpression(Generator, Scope);
+		InputExpressions.Add(InputExpression);
+	}
+
+	UE::HLSLTree::FFunctionCall* FunctionCall = Generator.AcquireFunctionCall(Scope, MaterialFunction, InputExpressions);
+	OutExpression = Generator.GetTree().NewExpression<UE::HLSLTree::FExpressionFunctionOutput>(Scope, FunctionCall, OutputIndex);
+	return EMaterialGenerateHLSLStatus::Success;
+}
+
 EMaterialGenerateHLSLStatus UMaterialExpressionExecBegin::GenerateHLSLStatement(FMaterialHLSLGenerator& Generator, UE::HLSLTree::FScope& Scope, UE::HLSLTree::FStatement*& OutStatement)
 {
 	OutStatement = Exec.AcquireHLSLStatement(Generator, Scope);
diff --git a/Engine/Source/Runtime/Engine/Private/Materials/MaterialHLSLEmitter.cpp b/Engine/Source/Runtime/Engine/Private/Materials/MaterialHLSLEmitter.cpp
new file mode 100644
--- /dev/null
+++ b/Engine/Source/Runtime/Engine/Private/Materials/MaterialHLSLEmitter.cpp
@@ -0,0 +1,691 @@
+#include "Materials/MaterialHLSLEmitter.h"
+#include "Materials/MaterialExpression.h"
+#include "Materials/MaterialExpressionTextureBase.h"
+#include "Materials/MaterialExpressionExecBegin.h"
+#include "Materials/MaterialExpressionMaterialFunctionCall.h"
+#include "Materials/MaterialExpressionFunctionInput.h"
+#include "Materials/MaterialExpressionFunctionOutput.h"
+#include "Materials/MaterialExpressionVolumetricAdvancedMaterialOutput.h"
+#include "Materials/Material.h"
+#include "ShaderCore.h"
+#include "HLSLTree/HLSLTree.h"
+#include "Containers/LazyPrintf.h"
+
+static bool SharedPixelProperties[CompiledMP_MAX];
+
+static const TCHAR* HLSLTypeString(EMaterialValueType Type)
+{
+	switch (Type)
+	{
+	case MCT_Float1:				return TEXT("MaterialFloat");
+	case MCT_Float2:				return TEXT("MaterialFloat2");
+	case MCT_Float3:				return TEXT("MaterialFloat3");
+	case MCT_Float4:				return TEXT("MaterialFloat4");
+	case MCT_Float:					return TEXT("MaterialFloat");
+	case MCT_Texture2D:				return TEXT("texture2D");
+	case MCT_TextureCube:			return TEXT("textureCube");
+	case MCT_Texture2DArray:		return TEXT("texture2DArray");
+	case MCT_VolumeTexture:			return TEXT("volumeTexture");
+	case MCT_StaticBool:			return TEXT("static bool");
+	case MCT_MaterialAttributes:	return TEXT("FMaterialAttributes");
+	case MCT_TextureExternal:		return TEXT("TextureExternal");
+	case MCT_TextureVirtual:		return TEXT("TextureVirtual");
+	case MCT_VTPageTableResult:		return TEXT("VTPageTableResult");
+	case MCT_ShadingModel:			return TEXT("uint");
+	case MCT_Strata:				return TEXT("FStrataData");
+	default:						return TEXT("unknown");
+	};
+}
+
+static FString GenerateMaterialTemplateHLSL(EShaderPlatform ShaderPlatform,
+	const FMaterial& Material,
+	const UE::HLSLTree::FEmitContext& EmitContext,
+	const TCHAR* PixelShaderCode)
+{
+	FString MaterialTemplateSource;
+	LoadShaderSourceFileChecked(TEXT("/Engine/Private/MaterialTemplate.ush"), ShaderPlatform, MaterialTemplateSource);
+
+	// Find the string index of the '#line' statement in MaterialTemplate.usf
+	const int32 LineIndex = MaterialTemplateSource.Find(TEXT("#line"), ESearchCase::CaseSensitive);
+	check(LineIndex != INDEX_NONE);
+
+	// Count line endings before the '#line' statement
+	int32 MaterialTemplateLineNumber = INDEX_NONE;
+	int32 StartPosition = LineIndex + 1;
+	do
+	{
+		MaterialTemplateLineNumber++;
+		// Using \n instead of LINE_TERMINATOR as not all of the lines are terminated consistently
+		// Subtract one from the last found line ending index to make sure we skip over it
+		StartPosition = MaterialTemplateSource.Find(TEXT("\n"), ESearchCase::CaseSensitive, ESearchDir::FromEnd, StartPosition - 1);
+	} while (StartPosition != INDEX_NONE);
+	check(MaterialTemplateLineNumber != INDEX_NONE);
+	// At this point MaterialTemplateLineNumber is one less than the line number of the '#line' statement
+	// For some reason we have to add 2 more to the #line value to get correct error line numbers from D3DXCompileShader
+	MaterialTemplateLineNumber += 3;
+
+	FLazyPrintf LazyPrintf(*MaterialTemplateSource);
+
+	const uint32 NumUserVertexTexCoords = EmitContext.NumTexCoords;
+	const uint32 NumUserTexCoords = EmitContext.NumTexCoords;
+	const uint32 NumCustomVectors = 0u;
+	const uint32 NumTexCoordVectors = EmitContext.NumTexCoords;
+
+	LazyPrintf.PushParam(*FString::Printf(TEXT("%u"), NumUserVertexTexCoords));
+	LazyPrintf.PushParam(*FString::Printf(TEXT("%u"), NumUserTexCoords));
+	LazyPrintf.PushParam(*FString::Printf(TEXT("%u"), NumCustomVectors));
+	LazyPrintf.PushParam(*FString::Printf(TEXT("%u"), NumTexCoordVectors));
+
+	FString VertexInterpolatorsOffsetsDefinition;
+	LazyPrintf.PushParam(*VertexInterpolatorsOffsetsDefinition);
+
+	FString MaterialAttributesDeclaration;
+	FString MaterialAttributesUtilities;
+	FString MaterialAttributesDefault;
+
+	const EMaterialShadingModel DefaultShadingModel = Material.GetShadingModels().GetFirstShadingModel();
+
+	const TArray<FGuid>& OrderedVisibleAttributes = FMaterialAttributeDefinitionMap::GetOrderedVisibleAttributeList();
+	for (const FGuid& AttributeID : OrderedVisibleAttributes)
+	{
+		const FString PropertyName = FMaterialAttributeDefinitionMap::GetAttributeName(AttributeID);
+		const EMaterialValueType PropertyType = FMaterialAttributeDefinitionMap::GetValueType(AttributeID);
+		const TCHAR* HLSLType = nullptr;
+
+		switch (PropertyType)
+		{
+		case MCT_Float1: case MCT_Float: HLSLType = TEXT("float"); break;
+		case MCT_Float2: HLSLType = TEXT("float2"); break;
+		case MCT_Float3: HLSLType = TEXT("float3"); break;
+		case MCT_Float4: HLSLType = TEXT("float4"); break;
+		case MCT_ShadingModel: HLSLType = TEXT("uint"); break;
+		case MCT_Strata: HLSLType = TEXT("FStrataData"); break;
+		default: break;
+		}
+
+		if (HLSLType)
+		{
+			const FVector4& DefaultValue = FMaterialAttributeDefinitionMap::GetDefaultValue(AttributeID);
+
+			MaterialAttributesDeclaration += FString::Printf(TEXT("\t%s %s;") LINE_TERMINATOR, HLSLType, *PropertyName);
+
+			// Chainable method to set the attribute
+			MaterialAttributesUtilities += FString::Printf(TEXT("FMaterialAttributes FMaterialAttributes_Set%s(FMaterialAttributes InAttributes, %s InValue) { InAttributes.%s = InValue; return InAttributes; }") LINE_TERMINATOR,
+				*PropertyName, HLSLType, *PropertyName);
+
+			// Standard value type
+			switch (PropertyType)
+			{
+			case MCT_Float: case MCT_Float1: MaterialAttributesDefault += FString::Printf(TEXT("\tResult.%s = %0.8f;") LINE_TERMINATOR, *PropertyName, DefaultValue.X); break;
+			case MCT_Float2: MaterialAttributesDefault += FString::Printf(TEXT("\tResult.%s = MaterialFloat2(%0.8f,%0.8f);") LINE_TERMINATOR, *PropertyName, DefaultValue.X, DefaultValue.Y); break;
+			case MCT_Float3: MaterialAttributesDefault += FString::Printf(TEXT("\tResult.%s = MaterialFloat3(%0.8f,%0.8f,%0.8f);") LINE_TERMINATOR, *PropertyName, DefaultValue.X, DefaultValue.Y, DefaultValue.Z); break;
+			case MCT_Float4: MaterialAttributesDefault += FString::Printf(TEXT("\tResult.%s = MaterialFloat4(%0.8f,%0.8f,%0.8f,%0.8f);") LINE_TERMINATOR, *PropertyName, DefaultValue.X, DefaultValue.Y, DefaultValue.Z, DefaultValue.W); break;
+			case MCT_ShadingModel: MaterialAttributesDefault += FString::Printf(TEXT("\tResult.%s = %d;") LINE_TERMINATOR, *PropertyName, (int32)DefaultShadingModel); break;
+			case MCT_Strata: MaterialAttributesDefault += FString::Printf(TEXT("\tResult.%s = GetInitialisedStrataData();") LINE_TERMINATOR, *PropertyName); break; // TODO
+			default: checkNoEntry(); break;
+			}
+		}
+	}
+
+	LazyPrintf.PushParam(*MaterialAttributesDeclaration);
+	LazyPrintf.PushParam(*MaterialAttributesUtilities);
+
+	// Stores the shared shader results member declarations
+	// PixelMembersDeclaration should be the same for all variations, but might change in the future. There are cases where work is shared
+	// between the pixel and vertex shader, but with Nanite all work has to be moved into the pixel shader, which means we will want
+	// different inputs. But for now, we are keeping them the same.
+	FString PixelMembersDeclaration;
+
+	for (int32 PropertyIndex = 0; PropertyIndex < MP_MAX; ++PropertyIndex)
+	{
+		// Skip non-shared properties
+		if (!SharedPixelProperties[PropertyIndex])
+		{
+			continue;
+		}
+
+		const EMaterialProperty Property = (EMaterialProperty)PropertyIndex;
+		check(FMaterialAttributeDefinitionMap::GetShaderFrequency(Property) == SF_Pixel);
+		// Special case MP_SubsurfaceColor as the actual property is a combination of the color and the profile but we don't want to expose the profile
+		const FString PropertyName = Property == MP_SubsurfaceColor ? "Subsurface" : FMaterialAttributeDefinitionMap::GetAttributeName(Property);
+		check(PropertyName.Len() > 0);
+		const EMaterialValueType Type = Property == MP_SubsurfaceColor ? MCT_Float4 : FMaterialAttributeDefinitionMap::GetValueType(Property);
+
+		PixelMembersDeclaration += FString::Printf(TEXT("\t%s %s;\n"), HLSLTypeString(Type), *PropertyName);
+	}
+
+	//for (int32 Index = 0; Index < CompiledPDV_MAX; Index++)
+	//{
+	//	GetSharedInputsMaterialCode(PixelMembersDeclaration[Index], NormalAssignment[Index], PixelMembersSetupAndAssignments[Index], (ECompiledPartialDerivativeVariation)Index);
+	//}
+
+	LazyPrintf.PushParam(*PixelMembersDeclaration);
+
+	{
+		FString DerivativeHelpers;// = DerivativeAutogen.GenerateUsedFunctions(*this);
+		FString DerivativeHelpersAndResources;// = DerivativeHelpers + ResourcesString;
+		//LazyPrintf.PushParam(*ResourcesString);
+		LazyPrintf.PushParam(*DerivativeHelpersAndResources);
+	}
+
+	// Anything used bye the GenerationFunctionCode() like WorldPositionOffset shouldn't be using texures, right?
+	// Let those use the standard finite differences textures, since they should be the same. If we actually want
+	// those to handle texture reads properly, we'll have to make extra versions.
+	ECompiledPartialDerivativeVariation BaseDerivativeVariation = CompiledPDV_FiniteDifferences;
+
+	//if (bCompileForComputeShader)
+	//{
+	//	LazyPrintf.PushParam(*GenerateFunctionCode(CompiledMP_EmissiveColorCS, BaseDerivativeVariation));
+	//}
+	//else
+	{
+		LazyPrintf.PushParam(TEXT("return 0"));
+	}
+
+	LazyPrintf.PushParam(*FString::Printf(TEXT("return %.5f"), Material.GetTranslucencyDirectionalLightingIntensity()));
+
+	LazyPrintf.PushParam(*FString::Printf(TEXT("return %.5f"), Material.GetTranslucentShadowDensityScale()));
+	LazyPrintf.PushParam(*FString::Printf(TEXT("return %.5f"), Material.GetTranslucentSelfShadowDensityScale()));
+	LazyPrintf.PushParam(*FString::Printf(TEXT("return %.5f"), Material.GetTranslucentSelfShadowSecondDensityScale()));
+	LazyPrintf.PushParam(*FString::Printf(TEXT("return %.5f"), Material.GetTranslucentSelfShadowSecondOpacity()));
+	LazyPrintf.PushParam(*FString::Printf(TEXT("return %.5f"), Material.GetTranslucentBackscatteringExponent()));
+
+	{
+		FLinearColor Extinction = Material.GetTranslucentMultipleScatteringExtinction();
+
+		LazyPrintf.PushParam(*FString::Printf(TEXT("return MaterialFloat3(%.5f, %.5f, %.5f)"), Extinction.R, Extinction.G, Extinction.B));
+	}
+
+	LazyPrintf.PushParam(*FString::Printf(TEXT("return %.5f"), Material.GetOpacityMaskClipValue()));
+
+	LazyPrintf.PushParam(TEXT("return Parameters.MaterialVertexAttributes.WorldPositionOffset"));
+	LazyPrintf.PushParam(TEXT("return 0.0f"));
+	LazyPrintf.PushParam(TEXT("return 0.0f"));
+	LazyPrintf.PushParam(TEXT("return 0.0f"));
+	LazyPrintf.PushParam(TEXT("return 0.0f"));
+	LazyPrintf.PushParam(TEXT("return 0.0f"));
+
+	// Print custom texture coordinate assignments, should be fine with regular derivatives
+	FString CustomUVAssignments;
+
+	int32 LastProperty = -1;
+	for (uint32 CustomUVIndex = 0; CustomUVIndex < NumUserTexCoords; CustomUVIndex++)
+	{
+		//if (bEnableExecutionFlow)
+		{
+			const FString AttributeName = FMaterialAttributeDefinitionMap::GetAttributeName((EMaterialProperty)(MP_CustomizedUVs0 + CustomUVIndex));
+			CustomUVAssignments += FString::Printf(TEXT("\tOutTexCoords[%u] = Parameters.MaterialVertexAttributes.%s;") LINE_TERMINATOR, CustomUVIndex, *AttributeName);
+		}
+		/*else
+		{
+			if (CustomUVIndex == 0)
+			{
+				CustomUVAssignments += DerivativeVariations[BaseDerivativeVariation].TranslatedCodeChunkDefinitions[MP_CustomizedUVs0 + CustomUVIndex];
+			}
+
+			if (DerivativeVariations[BaseDerivativeVariation].TranslatedCodeChunkDefinitions[MP_CustomizedUVs0 + CustomUVIndex].Len() > 0)
+			{
+				if (LastProperty >= 0)
+				{
+					check(DerivativeVariations[BaseDerivativeVariation].TranslatedCodeChunkDefinitions[LastProperty].Len() == DerivativeVariations[BaseDerivativeVariation].TranslatedCodeChunkDefinitions[MP_CustomizedUVs0 + CustomUVIndex].Len());
+				}
+				LastProperty = MP_CustomizedUVs0 + CustomUVIndex;
+			}
+			CustomUVAssignments += FString::Printf(TEXT("\tOutTexCoords[%u] = %s;") LINE_TERMINATOR, CustomUVIndex, *DerivativeVariations[BaseDerivativeVariation].TranslatedCodeChunks[MP_CustomizedUVs0 + CustomUVIndex]);
+		}*/
+	}
+
+	LazyPrintf.PushParam(*CustomUVAssignments);
+
+	// Print custom vertex shader interpolator assignments
+	FString CustomInterpolatorAssignments;
+
+	/*for (UMaterialExpressionVertexInterpolator* Interpolator : CustomVertexInterpolators)
+	{
+		if (Interpolator->InterpolatorOffset != INDEX_NONE)
+		{
+			check(Interpolator->InterpolatorIndex != INDEX_NONE);
+			check(Interpolator->InterpolatedType & MCT_Float);
+
+			const EMaterialValueType Type = Interpolator->InterpolatedType == MCT_Float ? MCT_Float1 : Interpolator->InterpolatedType;
+			const TCHAR* Swizzle[2] = { TEXT("x"), TEXT("y") };
+			const int32 Offset = Interpolator->InterpolatorOffset;
+			const int32 Index = Interpolator->InterpolatorIndex;
+
+			// Note: We reference the UV define directly to avoid having to pre-accumulate UV counts before property translation
+			CustomInterpolatorAssignments += FString::Printf(TEXT("\tOutTexCoords[VERTEX_INTERPOLATOR_%i_TEXCOORDS_X].%s = VertexInterpolator%i(Parameters).x;") LINE_TERMINATOR, Index, Swizzle[Offset % 2], Index);
+
+			if (Type >= MCT_Float2)
+			{
+				CustomInterpolatorAssignments += FString::Printf(TEXT("\tOutTexCoords[VERTEX_INTERPOLATOR_%i_TEXCOORDS_Y].%s = VertexInterpolator%i(Parameters).y;") LINE_TERMINATOR, Index, Swizzle[(Offset + 1) % 2], Index);
+
+				if (Type >= MCT_Float3)
+				{
+					CustomInterpolatorAssignments += FString::Printf(TEXT("\tOutTexCoords[VERTEX_INTERPOLATOR_%i_TEXCOORDS_Z].%s = VertexInterpolator%i(Parameters).z;") LINE_TERMINATOR, Index, Swizzle[(Offset + 2) % 2], Index);
+
+					if (Type == MCT_Float4)
+					{
+						CustomInterpolatorAssignments += FString::Printf(TEXT("\tOutTexCoords[VERTEX_INTERPOLATOR_%i_TEXCOORDS_W].%s = VertexInterpolator%i(Parameters).w;") LINE_TERMINATOR, Index, Swizzle[(Offset + 3) % 2], Index);
+					}
+				}
+			}
+		}
+	}*/
+
+	LazyPrintf.PushParam(*CustomInterpolatorAssignments);
+
+	LazyPrintf.PushParam(*MaterialAttributesDefault);
+
+	//if (bEnableExecutionFlow)
+	{
+		FString EvaluateVertexCode;
+
+		// Set default texcoords in the VS
+		for (uint32 TexCoordIndex = 0; TexCoordIndex < NumUserVertexTexCoords; ++TexCoordIndex)
+		{
+			const FString AttributeName = FMaterialAttributeDefinitionMap::GetAttributeName((EMaterialProperty)(MP_CustomizedUVs0 + TexCoordIndex));
+
+			EvaluateVertexCode += FString::Printf(TEXT("\tDefaultMaterialAttributes.%s = Parameters.TexCoords[%d];") LINE_TERMINATOR, *AttributeName, TexCoordIndex);
+		}
+
+		//EvaluateVertexCode += TranslatedAttributesCodeChunks[SF_Vertex];
+
+		LazyPrintf.PushParam(*EvaluateVertexCode);
+		LazyPrintf.PushParam(PixelShaderCode);
+
+		FString EvaluateMaterialAttributesCode = TEXT("    FMaterialAttributes MaterialAttributes = EvaluatePixelMaterialAttributes(Parameters);" LINE_TERMINATOR);
+
+		for (int32 PropertyIndex = 0; PropertyIndex < MP_MAX; ++PropertyIndex)
+		{
+			// Skip non-shared properties
+			if (!SharedPixelProperties[PropertyIndex])
+			{
+				continue;
+			}
+
+			const EMaterialProperty Property = (EMaterialProperty)PropertyIndex;
+			check(FMaterialAttributeDefinitionMap::GetShaderFrequency(Property) == SF_Pixel);
+			// Special case MP_SubsurfaceColor as the actual property is a combination of the color and the profile but we don't want to expose the profile
+			const FString PropertyName = FMaterialAttributeDefinitionMap::GetAttributeName(Property);
+
+			if (PropertyIndex == MP_SubsurfaceColor)
+			{
+				// TODO - properly handle subsurface profile
+				EvaluateMaterialAttributesCode += FString::Printf("    PixelMaterialInputs.Subsurface = float4(MaterialAttributes.%s, 0.0f);" LINE_TERMINATOR, *PropertyName);
+			}
+			else
+			{
+				EvaluateMaterialAttributesCode += FString::Printf("    PixelMaterialInputs.%s = MaterialAttributes.%s;" LINE_TERMINATOR, *PropertyName, *PropertyName);
+			}
+		}
+
+		// TODO - deriv
+		for (int32 Iter = 0; Iter < CompiledPDV_MAX; Iter++)
+		{
+			LazyPrintf.PushParam(*EvaluateMaterialAttributesCode);
+			LazyPrintf.PushParam(TEXT(""));
+			LazyPrintf.PushParam(TEXT(""));
+		}
+	}
+
+	LazyPrintf.PushParam(*FString::Printf(TEXT("%u"), MaterialTemplateLineNumber));
+
+	return LazyPrintf.GetResultString();
+}
+
+static void GetMaterialEnvironment(EShaderPlatform InPlatform,
+	const FMaterial& InMaterial,
+	const FMaterialCompilationOutput& MaterialCompilationOutput,
+	FShaderCompilerEnvironment& OutEnvironment)
+{
+	bool bMaterialRequestsDualSourceBlending = false;
+
+	if (false)//bNeedsParticlePosition || Material->ShouldGenerateSphericalParticleNormals() || bUsesSphericalParticleOpacity)
+	{
+		OutEnvironment.SetDefine(TEXT("NEEDS_PARTICLE_POSITION"), 1);
+	}
+
+	if (false)//bNeedsParticleVelocity || Material->IsUsedWithNiagaraMeshParticles())
+	{
+		OutEnvironment.SetDefine(TEXT("NEEDS_PARTICLE_VELOCITY"), 1);
+	}
+
+	uint32 DynamicParticleParameterMask = 0u;
+	if (DynamicParticleParameterMask)
+	{
+		OutEnvironment.SetDefine(TEXT("USE_DYNAMIC_PARAMETERS"), 1);
+		OutEnvironment.SetDefine(TEXT("DYNAMIC_PARAMETERS_MASK"), DynamicParticleParameterMask);
+	}
+
+	if (false)//bNeedsParticleTime)
+	{
+		OutEnvironment.SetDefine(TEXT("NEEDS_PARTICLE_TIME"), 1);
+	}
+
+	if (false)//bUsesParticleMotionBlur)
+	{
+		OutEnvironment.SetDefine(TEXT("USES_PARTICLE_MOTION_BLUR"), 1);
+	}
+
+	if (false)//bNeedsParticleRandom)
+	{
+		OutEnvironment.SetDefine(TEXT("NEEDS_PARTICLE_RANDOM"), 1);
+	}
+
+	if (false)//bUsesSphericalParticleOpacity)
+	{
+		OutEnvironment.SetDefine(TEXT("SPHERICAL_PARTICLE_OPACITY"), TEXT("1"));
+	}
+
+	if (false)//bUsesParticleSubUVs)
+	{
+		OutEnvironment.SetDefine(TEXT("USE_PARTICLE_SUBUVS"), TEXT("1"));
+	}
+
+	if (false)//bUsesLightmapUVs)
+	{
+		OutEnvironment.SetDefine(TEXT("LIGHTMAP_UV_ACCESS"), TEXT("1"));
+	}
+
+	if (false)//bUsesAOMaterialMask)
+	{
+		OutEnvironment.SetDefine(TEXT("USES_AO_MATERIAL_MASK"), TEXT("1"));
+	}
+
+	if (false)//bUsesSpeedTree)
+	{
+		OutEnvironment.SetDefine(TEXT("USES_SPEEDTREE"), TEXT("1"));
+	}
+
+	if (false)//bNeedsWorldPositionExcludingShaderOffsets)
+	{
+		OutEnvironment.SetDefine(TEXT("NEEDS_WORLD_POSITION_EXCLUDING_SHADER_OFFSETS"), TEXT("1"));
+	}
+
+	if (false)//bNeedsParticleSize)
+	{
+		OutEnvironment.SetDefine(TEXT("NEEDS_PARTICLE_SIZE"), TEXT("1"));
+	}
+
+	if (MaterialCompilationOutput.bNeedsSceneTextures)
+	{
+		OutEnvironment.SetDefine(TEXT("NEEDS_SCENE_TEXTURES"), TEXT("1"));
+	}
+	if (MaterialCompilationOutput.bUsesEyeAdaptation)
+	{
+		OutEnvironment.SetDefine(TEXT("USES_EYE_ADAPTATION"), TEXT("1"));
+	}
+
+	if (MaterialCompilationOutput.bHasRuntimeVirtualTextureOutputNode)
+	{
+		OutEnvironment.SetDefine(TEXT("VIRTUAL_TEXTURE_OUTPUT"), 1);
+	}
+
+	OutEnvironment.SetDefine(TEXT("USES_PER_INSTANCE_CUSTOM_DATA"), false);// bUsesPerInstanceCustomData&& Material->IsUsedWithInstancedStaticMeshes());
+
+	// @todo MetalMRT: Remove this hack and implement proper atmospheric-fog solution for Metal MRT...
+	OutEnvironment.SetDefine(TEXT("MATERIAL_ATMOSPHERIC_FOG"), false);// !IsMetalMRTPlatform(InPlatform) ? bUsesAtmosphericFog : 0);
+	OutEnvironment.SetDefine(TEXT("MATERIAL_SKY_ATMOSPHERE"), false);// bUsesSkyAtmosphere);
+	OutEnvironment.SetDefine(TEXT("INTERPOLATE_VERTEX_COLOR"), false);// bUsesVertexColor);
+	OutEnvironment.SetDefine(TEXT("NEEDS_PARTICLE_COLOR"), false);// bUsesParticleColor);
+	OutEnvironment.SetDefine(TEXT("NEEDS_PARTICLE_LOCAL_TO_WORLD"), false);// bUsesParticleLocalToWorld);
+	OutEnvironment.SetDefine(TEXT("NEEDS_PARTICLE_WORLD_TO_LOCAL"), false);// bUsesParticleWorldToLocal);
+	OutEnvironment.SetDefine(TEXT("USES_TRANSFORM_VECTOR"), false);// bUsesTransformVector);
+	OutEnvironment.SetDefine(TEXT("WANT_PIXEL_DEPTH_OFFSET"), false);// bUsesPixelDepthOffset);
+	if (IsMetalPlatform(InPlatform))
+	{
+		OutEnvironment.SetDefine(TEXT("USES_WORLD_POSITION_OFFSET"), false);// bUsesWorldPositionOffset);
+	}
+	OutEnvironment.SetDefine(TEXT("USES_EMISSIVE_COLOR"), false);// bUsesEmissiveColor);
+	// Distortion uses tangent space transform 
+	OutEnvironment.SetDefine(TEXT("USES_DISTORTION"), InMaterial.IsDistorted());
+
+	OutEnvironment.SetDefine(TEXT("MATERIAL_ENABLE_TRANSLUCENCY_FOGGING"), InMaterial.ShouldApplyFogging());
+	OutEnvironment.SetDefine(TEXT("MATERIAL_ENABLE_TRANSLUCENCY_CLOUD_FOGGING"), InMaterial.ShouldApplyCloudFogging());
+	OutEnvironment.SetDefine(TEXT("MATERIAL_IS_SKY"), InMaterial.IsSky());
+	OutEnvironment.SetDefine(TEXT("MATERIAL_COMPUTE_FOG_PER_PIXEL"), InMaterial.ComputeFogPerPixel());
+	OutEnvironment.SetDefine(TEXT("MATERIAL_FULLY_ROUGH"), false);// bIsFullyRough || InMaterial.IsFullyRough());
+	OutEnvironment.SetDefine(TEXT("MATERIAL_USES_ANISOTROPY"), false);// bUsesAnisotropy);
+
+	// Count the number of VTStacks (each stack will allocate a feedback slot)
+	OutEnvironment.SetDefine(TEXT("NUM_VIRTUALTEXTURE_SAMPLES"), 0);// VTStacks.Num());
+
+	// Setup defines to map each VT stack to either 1 or 2 page table textures, depending on how many layers it uses
+	/*for (int i = 0; i < VTStacks.Num(); ++i)
+	{
+		const FMaterialVirtualTextureStack& Stack = MaterialCompilationOutput.UniformExpressionSet.VTStacks[i];
+		FString PageTableValue = FString::Printf(TEXT("Material.VirtualTexturePageTable0_%d"), i);
+		if (Stack.GetNumLayers() > 4u)
+		{
+			PageTableValue += FString::Printf(TEXT(", Material.VirtualTexturePageTable1_%d"), i);
+		}
+		if (VTStacks[i].bAdaptive)
+		{
+			PageTableValue += FString::Printf(TEXT(", Material.VirtualTexturePageTableIndirection_%d"), i);
+		}
+		OutEnvironment.SetDefine(*FString::Printf(TEXT("VIRTUALTEXTURE_PAGETABLE_%d"), i), *PageTableValue);
+	}*/
+
+	/*for (int32 CollectionIndex = 0; CollectionIndex < ParameterCollections.Num(); CollectionIndex++)
+	{
+		// Add uniform buffer declarations for any parameter collections referenced
+		const FString CollectionName = FString::Printf(TEXT("MaterialCollection%u"), CollectionIndex);
+		// This can potentially become an issue for MaterialCollection Uniform Buffers if they ever get non-numeric resources (eg Textures), as
+		// OutEnvironment.ResourceTableMap has a map by name, and the N ParameterCollection Uniform Buffers ALL are names "MaterialCollection"
+		// (and the hlsl cbuffers are named MaterialCollection0, etc, so the names don't match the layout)
+		FShaderUniformBufferParameter::ModifyCompilationEnvironment(*CollectionName, ParameterCollections[CollectionIndex]->GetUniformBufferStruct(), InPlatform, OutEnvironment);
+	}*/
+	OutEnvironment.SetDefine(TEXT("IS_MATERIAL_SHADER"), true);
+
+	// Set all the shading models for this material here 
+	FMaterialShadingModelField ShadingModels = InMaterial.GetShadingModels();
+
+	// If the material gets its shading model from the material expressions, then we use the result from the compilation (assuming it's valid).
+	// This result will potentially be tighter than what GetShadingModels() returns, because it only picks up the shading models from the expressions that get compiled for a specific feature level and quality level
+	// For example, the material might have shading models behind static switches. GetShadingModels() will return both the true and the false paths from that switch, whereas the shading model field from the compilation will only contain the actual shading model selected 
+	/*if (InMaterial.IsShadingModelFromMaterialExpression() && ShadingModelsFromCompilation.IsValid())
+	{
+		// Shading models fetched from the compilation of the expression graph
+		ShadingModels = ShadingModelsFromCompilation;
+	}*/
+
+	ensure(ShadingModels.IsValid());
+
+	if (ShadingModels.IsLit())
+	{
+		int32 NumSetMaterials = 0;
+		if (ShadingModels.HasShadingModel(MSM_DefaultLit))
+		{
+			OutEnvironment.SetDefine(TEXT("MATERIAL_SHADINGMODEL_DEFAULT_LIT"), TEXT("1"));
+			NumSetMaterials++;
+		}
+		if (ShadingModels.HasShadingModel(MSM_Subsurface))
+		{
+			OutEnvironment.SetDefine(TEXT("MATERIAL_SHADINGMODEL_SUBSURFACE"), TEXT("1"));
+			NumSetMaterials++;
+		}
+		if (ShadingModels.HasShadingModel(MSM_PreintegratedSkin))
+		{
+			OutEnvironment.SetDefine(TEXT("MATERIAL_SHADINGMODEL_PREINTEGRATED_SKIN"), TEXT("1"));
+			NumSetMaterials++;
+		}
+		if (ShadingModels.HasShadingModel(MSM_SubsurfaceProfile))
+		{
+			OutEnvironment.SetDefine(TEXT("MATERIAL_SHADINGMODEL_SUBSURFACE_PROFILE"), TEXT("1"));
+			NumSetMaterials++;
+		}
+		if (ShadingModels.HasShadingModel(MSM_ClearCoat))
+		{
+			OutEnvironment.SetDefine(TEXT("MATERIAL_SHADINGMODEL_CLEAR_COAT"), TEXT("1"));
+			NumSetMaterials++;
+		}
+		if (ShadingModels.HasShadingModel(MSM_TwoSidedFoliage))
+		{
+			OutEnvironment.SetDefine(TEXT("MATERIAL_SHADINGMODEL_TWOSIDED_FOLIAGE"), TEXT("1"));
+			NumSetMaterials++;
+		}
+		if (ShadingModels.HasShadingModel(MSM_Hair))
+		{
+			OutEnvironment.SetDefine(TEXT("MATERIAL_SHADINGMODEL_HAIR"), TEXT("1"));
+			NumSetMaterials++;
+		}
+		if (ShadingModels.HasShadingModel(MSM_Cloth))
+		{
+			OutEnvironment.SetDefine(TEXT("MATERIAL_SHADINGMODEL_CLOTH"), TEXT("1"));
+			NumSetMaterials++;
+		}
+		if (ShadingModels.HasShadingModel(MSM_Eye))
+		{
+			OutEnvironment.SetDefine(TEXT("MATERIAL_SHADINGMODEL_EYE"), TEXT("1"));
+			NumSetMaterials++;
+		}
+		if (ShadingModels.HasShadingModel(MSM_SingleLayerWater))
+		{
+			OutEnvironment.SetDefine(TEXT("MATERIAL_SHADINGMODEL_SINGLELAYERWATER"), TEXT("1"));
+			NumSetMaterials++;
+		}
+		if (ShadingModels.HasShadingModel(MSM_ThinTranslucent))
+		{
+			OutEnvironment.SetDefine(TEXT("MATERIAL_SHADINGMODEL_THIN_TRANSLUCENT"), TEXT("1"));
+			NumSetMaterials++;
+
+			bMaterialRequestsDualSourceBlending = true;
+		}
+
+		if (ShadingModels.HasShadingModel(MSM_SingleLayerWater) &&
+			(IsSwitchPlatform(InPlatform) || IsVulkanMobileSM5Platform(InPlatform) || FDataDrivenShaderPlatformInfo::GetRequiresDisableForwardLocalLights(InPlatform)))
+		{
+			OutEnvironment.SetDefine(TEXT("DISABLE_FORWARD_LOCAL_LIGHTS"), TEXT("1"));
+		}
+
+		// This is to have switch use the simple single layer water shading similar to mobile: no dynamic lights, only sun and sky, no distortion, no colored transmittance on background, no custom depth read.
+		const bool bSingleLayerWaterUsesSimpleShading = (IsSwitchPlatform(InPlatform) || IsVulkanMobileSM5Platform(InPlatform)) && IsForwardShadingEnabled(InPlatform);
+		if (ShadingModels.HasShadingModel(MSM_SingleLayerWater) && bSingleLayerWaterUsesSimpleShading)
+		{
+			OutEnvironment.SetDefine(TEXT("SINGLE_LAYER_WATER_SIMPLE_FORWARD"), TEXT("1"));
+		}
+
+		if (NumSetMaterials == 1)
+		{
+			OutEnvironment.SetDefine(TEXT("MATERIAL_SINGLE_SHADINGMODEL"), TEXT("1"));
+		}
+
+		ensure(NumSetMaterials != 0);
+		if (NumSetMaterials == 0)
+		{
+			// Should not really end up here
+			UE_LOG(LogMaterial, Warning, TEXT("Unknown material shading model(s). Setting to MSM_DefaultLit"));
+			OutEnvironment.SetDefine(TEXT("MATERIAL_SHADINGMODEL_DEFAULT_LIT"), TEXT("1"));
+		}
+	}
+	else
+	{
+		// Unlit shading model can only exist by itself
+		OutEnvironment.SetDefine(TEXT("MATERIAL_SINGLE_SHADINGMODEL"), TEXT("1"));
+		OutEnvironment.SetDefine(TEXT("MATERIAL_SHADINGMODEL_UNLIT"), TEXT("1"));
+	}
+
+	if (InMaterial.GetMaterialDomain() == MD_Volume) // && Material->HasN)
+	{
+		TArray<const UMaterialExpressionVolumetricAdvancedMaterialOutput*> VolumetricAdvancedExpressions;
+		InMaterial.GetMaterialInterface()->GetMaterial()->GetAllExpressionsOfType(VolumetricAdvancedExpressions);
+		if (VolumetricAdvancedExpressions.Num() > 0)
+		{
+			if (VolumetricAdvancedExpressions.Num() > 1)
+			{
+				UE_LOG(LogMaterial, Fatal, TEXT("Only a single UMaterialExpressionVolumetricAdvancedMaterialOutput node is supported."));
+			}
+
+			OutEnvironment.SetDefine(TEXT("MATERIAL_VOLUMETRIC_ADVANCED"), TEXT("1"));
+
+			const UMaterialExpressionVolumetricAdvancedMaterialOutput* VolumetricAdvancedNode = VolumetricAdvancedExpressions[0];
+			if (VolumetricAdvancedNode->GetEvaluatePhaseOncePerSample())
+			{
+				OutEnvironment.SetDefine(TEXT("MATERIAL_VOLUMETRIC_ADVANCED_PHASE_PERSAMPLE"), TEXT("1"));
+			}
+			else
+			{
+				OutEnvironment.SetDefine(TEXT("MATERIAL_VOLUMETRIC_ADVANCED_PHASE_PERPIXEL"), TEXT("1"));
+			}
+
+			OutEnvironment.SetDefine(TEXT("MATERIAL_VOLUMETRIC_ADVANCED_GRAYSCALE_MATERIAL"), VolumetricAdvancedNode->bGrayScaleMaterial ? TEXT("1") : TEXT("0"));
+			OutEnvironment.SetDefine(TEXT("MATERIAL_VOLUMETRIC_ADVANCED_RAYMARCH_VOLUME_SHADOW"), VolumetricAdvancedNode->bRayMarchVolumeShadow ? TEXT("1") : TEXT("0"));
+
+			OutEnvironment.SetDefine(TEXT("MATERIAL_VOLUMETRIC_ADVANCED_MULTISCATTERING_OCTAVE_COUNT"), VolumetricAdvancedNode->GetMultiScatteringApproximationOctaveCount());
+
+			OutEnvironment.SetDefine(TEXT("MATERIAL_VOLUMETRIC_ADVANCED_CONSERVATIVE_DENSITY"),
+				VolumetricAdvancedNode->ConservativeDensity.IsConnected() ? TEXT("1") : TEXT("0"));
+
+			OutEnvironment.SetDefine(TEXT("MATERIAL_VOLUMETRIC_ADVANCED_OVERRIDE_AMBIENT_OCCLUSION"),
+				InMaterial.HasAmbientOcclusionConnected() ? TEXT("1") : TEXT("0"));
+
+			OutEnvironment.SetDefine(TEXT("MATERIAL_VOLUMETRIC_ADVANCED_GROUND_CONTRIBUTION"),
+				VolumetricAdvancedNode->bGroundContribution ? TEXT("1") : TEXT("0"));
+		}
+	}
+
+	OutEnvironment.SetDefine(TEXT("MATERIAL_IS_STRATA"), false);// bMaterialIsStrata ? TEXT("1") : TEXT("0"));
+
+	// STRATA_TODO do not request DualSourceBlending if gray scale transmittance is selected.
+	// bMaterialRequestsDualSourceBlending this base on limited set of blend mode: Opaque, Masked, TransmittanceCoverage, TransmittanceColored;
+	//bMaterialRequestsDualSourceBlending |= bMaterialIsStrata;
+
+	// if duals source blending (colored transmittance) is not supported on a platform, it will fall back to standard alpha blending (grey scale transmittance)
+	OutEnvironment.SetDefine(TEXT("DUAL_SOURCE_COLOR_BLENDING_ENABLED"), false);// bMaterialRequestsDualSourceBlending&& Material->IsDualBlendingEnabled(Platform) ? TEXT("1") : TEXT("0"));
+
+	// Translate() is called before getting here so we can create related define
+	/*if (StrataMaterialAnalysis.RequestedBSDFCount > 0)
+	{
+		OutEnvironment.SetDefine(TEXT("STRATA_CLAMPED_LAYER_COUNT"), StrataMaterialAnalysis.ClampedLayerCount);
+	}*/
+
+	OutEnvironment.SetDefine(TEXT("TEXTURE_SAMPLE_DEBUG"), false);// IsDebugTextureSampleEnabled() ? TEXT("1") : TEXT("0"));
+}
+
+bool MaterialEmitHLSL(const FMaterialCompileTargetParameters& InCompilerTarget,
+	const FMaterial& InMaterial,
+	const UE::HLSLTree::FTree& InTree,
+	FMaterialCompilationOutput& OutCompilationOutput,
+	TRefCountPtr<FSharedShaderCompilerEnvironment>& OutMaterialEnvironment)
+{
+	FMemory::Memzero(SharedPixelProperties);
+	SharedPixelProperties[MP_Normal] = true;
+	SharedPixelProperties[MP_Tangent] = true;
+	SharedPixelProperties[MP_EmissiveColor] = true;
+	SharedPixelProperties[MP_Opacity] = true;
+	SharedPixelProperties[MP_OpacityMask] = true;
+	SharedPixelProperties[MP_BaseColor] = true;
+	SharedPixelProperties[MP_Metallic] = true;
+	SharedPixelProperties[MP_Specular] = true;
+	SharedPixelProperties[MP_Roughness] = true;
+	SharedPixelProperties[MP_Anisotropy] = true;
+	SharedPixelProperties[MP_AmbientOcclusion] = true;
+	SharedPixelProperties[MP_Refraction] = true;
+	SharedPixelProperties[MP_PixelDepthOffset] = true;
+	SharedPixelProperties[MP_SubsurfaceColor] = true;
+	SharedPixelProperties[MP_ShadingModel] = true;
+	SharedPixelProperties[MP_FrontMaterial] = true;
+
+	FMemStackBase Allocator;
+	FMemMark MemMark(Allocator);
+
+	UE::HLSLTree::FCodeWriter* CodeWriter = UE::HLSLTree::FCodeWriter::Create(Allocator);
+
+	UE::HLSLTree::FEmitContext EmitContext;
+	EmitContext.Allocator = &Allocator;
+	EmitContext.Material = &InMaterial;
+	EmitContext.MaterialCompilationOutput = &OutCompilationOutput;
+	InTree.EmitHLSL(EmitContext, *CodeWriter);
+
+	//InOutMaterial.CompileErrors = MoveTemp(CompileErrors);
+	//InOutMaterial.ErrorExpressions = MoveTemp(ErrorExpressions);
+
+	const FString MaterialTemplateSource = GenerateMaterialTemplateHLSL(InCompilerTarget.ShaderPlatform, InMaterial, EmitContext, CodeWriter->StringBuilder->ToString());
+
+	OutMaterialEnvironment = new FSharedShaderCompilerEnvironment();
+	OutMaterialEnvironment->TargetPlatform = InCompilerTarget.TargetPlatform;
+	GetMaterialEnvironment(InCompilerTarget.ShaderPlatform, InMaterial, OutCompilationOutput, *OutMaterialEnvironment);
+	OutMaterialEnvironment->IncludeVirtualPathToContentsMap.Add(TEXT("/Engine/Generated/Material.ush"), MaterialTemplateSource);
+
+	return true;
+}
diff --git a/Engine/Source/Runtime/Engine/Private/Materials/MaterialHLSLEmitter.h b/Engine/Source/Runtime/Engine/Private/Materials/MaterialHLSLEmitter.h
new file mode 100644
--- /dev/null
+++ b/Engine/Source/Runtime/Engine/Private/Materials/MaterialHLSLEmitter.h
@@ -0,0 +1,23 @@
+#pragma once
+
+#include "CoreMinimal.h"
+#include "Templates/RefCounting.h"
+
+class FMaterial;
+class FMaterialCompilationOutput;
+struct FSharedShaderCompilerEnvironment;
+struct FMaterialCompileTargetParameters;
+
+namespace UE
+{
+namespace HLSLTree
+{
+class FTree;
+}
+}
+
+bool MaterialEmitHLSL(const FMaterialCompileTargetParameters& InCompilerTarget,
+	const FMaterial& InOutMaterial,
+	const UE::HLSLTree::FTree& InTree,
+	FMaterialCompilationOutput& OutCompilationOutput,
+	TRefCountPtr<FSharedShaderCompilerEnvironment>& OutMaterialEnvironment);
diff --git a/Engine/Source/Runtime/Engine/Private/Materials/MaterialHLSLGenerator.cpp b/Engine/Source/Runtime/Engine/Private/Materials/MaterialHLSLGenerator.cpp
--- a/Engine/Source/Runtime/Engine/Private/Materials/MaterialHLSLGenerator.cpp
+++ b/Engine/Source/Runtime/Engine/Private/Materials/MaterialHLSLGenerator.cpp
@@ -5,706 +5,27 @@
 
 #include "Materials/MaterialExpression.h"
 #include "Materials/MaterialExpressionTextureBase.h"
-#include "Materials/MaterialExpressionExecBegin.h"
 #include "Materials/MaterialExpressionMaterialFunctionCall.h"
 #include "Materials/MaterialExpressionFunctionInput.h"
 #include "Materials/MaterialExpressionFunctionOutput.h"
 #include "Materials/MaterialExpressionVolumetricAdvancedMaterialOutput.h"
 #include "Materials/Material.h"
+#include "MaterialHLSLTree.h"
 #include "ShaderCore.h"
 #include "HLSLTree/HLSLTree.h"
 #include "HLSLTree/HLSLTreeCommon.h"
 #include "Containers/LazyPrintf.h"
 
-static bool SharedPixelProperties[CompiledMP_MAX];
-
-static const TCHAR* HLSLTypeString(EMaterialValueType Type)
-{
-	switch (Type)
-	{
-	case MCT_Float1:				return TEXT("MaterialFloat");
-	case MCT_Float2:				return TEXT("MaterialFloat2");
-	case MCT_Float3:				return TEXT("MaterialFloat3");
-	case MCT_Float4:				return TEXT("MaterialFloat4");
-	case MCT_Float:					return TEXT("MaterialFloat");
-	case MCT_Texture2D:				return TEXT("texture2D");
-	case MCT_TextureCube:			return TEXT("textureCube");
-	case MCT_Texture2DArray:		return TEXT("texture2DArray");
-	case MCT_VolumeTexture:			return TEXT("volumeTexture");
-	case MCT_StaticBool:			return TEXT("static bool");
-	case MCT_MaterialAttributes:	return TEXT("FMaterialAttributes");
-	case MCT_TextureExternal:		return TEXT("TextureExternal");
-	case MCT_TextureVirtual:		return TEXT("TextureVirtual");
-	case MCT_VTPageTableResult:		return TEXT("VTPageTableResult");
-	case MCT_ShadingModel:			return TEXT("uint");
-	case MCT_Strata:				return TEXT("FStrataData");
-	default:						return TEXT("unknown");
-	};
-}
-
-FMaterialHLSLGenerator::FMaterialHLSLGenerator(EShaderPlatform InShaderPlatform,
-	ERHIFeatureLevel::Type InFeatureLevel,
-	const ITargetPlatform* InTargetPlatform)
-	: ShaderPlatform(InShaderPlatform)
-	, FeatureLevel(InFeatureLevel)
-	, TargetPlatform(InTargetPlatform)
-{
-	FMemory::Memzero(SharedPixelProperties);
-	SharedPixelProperties[MP_Normal] = true;
-	SharedPixelProperties[MP_Tangent] = true;
-	SharedPixelProperties[MP_EmissiveColor] = true;
-	SharedPixelProperties[MP_Opacity] = true;
-	SharedPixelProperties[MP_OpacityMask] = true;
-	SharedPixelProperties[MP_BaseColor] = true;
-	SharedPixelProperties[MP_Metallic] = true;
-	SharedPixelProperties[MP_Specular] = true;
-	SharedPixelProperties[MP_Roughness] = true;
-	SharedPixelProperties[MP_Anisotropy] = true;
-	SharedPixelProperties[MP_AmbientOcclusion] = true;
-	SharedPixelProperties[MP_Refraction] = true;
-	SharedPixelProperties[MP_PixelDepthOffset] = true;
-	SharedPixelProperties[MP_SubsurfaceColor] = true;
-	SharedPixelProperties[MP_ShadingModel] = true;
-	SharedPixelProperties[MP_FrontMaterial] = true;
-}
-
-static FString GenerateMaterialTemplateHLSL(EShaderPlatform ShaderPlatform,
-	FMaterial& Material,
-	const UE::HLSLTree::FEmitContext& EmitContext,
-	const TCHAR* PixelShaderCode)
+FMaterialHLSLGenerator::FMaterialHLSLGenerator(const FMaterialCompileTargetParameters& InCompileTarget, UE::HLSLTree::FTree& InOutTree)
+	: CompileTarget(InCompileTarget), HLSLTree(&InOutTree)
 {
-	FString MaterialTemplateSource;
-	LoadShaderSourceFileChecked(TEXT("/Engine/Private/MaterialTemplate.ush"), ShaderPlatform, MaterialTemplateSource);
-
-	// Find the string index of the '#line' statement in MaterialTemplate.usf
-	const int32 LineIndex = MaterialTemplateSource.Find(TEXT("#line"), ESearchCase::CaseSensitive);
-	check(LineIndex != INDEX_NONE);
-
-	// Count line endings before the '#line' statement
-	int32 MaterialTemplateLineNumber = INDEX_NONE;
-	int32 StartPosition = LineIndex + 1;
-	do
-	{
-		MaterialTemplateLineNumber++;
-		// Using \n instead of LINE_TERMINATOR as not all of the lines are terminated consistently
-		// Subtract one from the last found line ending index to make sure we skip over it
-		StartPosition = MaterialTemplateSource.Find(TEXT("\n"), ESearchCase::CaseSensitive, ESearchDir::FromEnd, StartPosition - 1);
-	} while (StartPosition != INDEX_NONE);
-	check(MaterialTemplateLineNumber != INDEX_NONE);
-	// At this point MaterialTemplateLineNumber is one less than the line number of the '#line' statement
-	// For some reason we have to add 2 more to the #line value to get correct error line numbers from D3DXCompileShader
-	MaterialTemplateLineNumber += 3;
-
-	FLazyPrintf LazyPrintf(*MaterialTemplateSource);
-
-	const uint32 NumUserVertexTexCoords = EmitContext.NumTexCoords;
-	const uint32 NumUserTexCoords = EmitContext.NumTexCoords;
-	const uint32 NumCustomVectors = 0u;
-	const uint32 NumTexCoordVectors = EmitContext.NumTexCoords;
-
-	LazyPrintf.PushParam(*FString::Printf(TEXT("%u"), NumUserVertexTexCoords));
-	LazyPrintf.PushParam(*FString::Printf(TEXT("%u"), NumUserTexCoords));
-	LazyPrintf.PushParam(*FString::Printf(TEXT("%u"), NumCustomVectors));
-	LazyPrintf.PushParam(*FString::Printf(TEXT("%u"), NumTexCoordVectors));
-
-	FString VertexInterpolatorsOffsetsDefinition;
-	LazyPrintf.PushParam(*VertexInterpolatorsOffsetsDefinition);
-
-	FString MaterialAttributesDeclaration;
-	FString MaterialAttributesUtilities;
-	FString MaterialAttributesDefault;
-
-	const EMaterialShadingModel DefaultShadingModel = Material.GetShadingModels().GetFirstShadingModel();
-
-	const TArray<FGuid>& OrderedVisibleAttributes = FMaterialAttributeDefinitionMap::GetOrderedVisibleAttributeList();
-	for (const FGuid& AttributeID : OrderedVisibleAttributes)
-	{
-		const FString PropertyName = FMaterialAttributeDefinitionMap::GetAttributeName(AttributeID);
-		const EMaterialValueType PropertyType = FMaterialAttributeDefinitionMap::GetValueType(AttributeID);
-		const TCHAR* HLSLType = nullptr;
-
-		switch (PropertyType)
-		{
-		case MCT_Float1: case MCT_Float: HLSLType = TEXT("float"); break;
-		case MCT_Float2: HLSLType = TEXT("float2"); break;
-		case MCT_Float3: HLSLType = TEXT("float3"); break;
-		case MCT_Float4: HLSLType = TEXT("float4"); break;
-		case MCT_ShadingModel: HLSLType = TEXT("uint"); break;
-		case MCT_Strata: HLSLType = TEXT("FStrataData"); break;
-		default: break;
-		}
-
-		if (HLSLType)
-		{
-			const FVector4& DefaultValue = FMaterialAttributeDefinitionMap::GetDefaultValue(AttributeID);
-
-			MaterialAttributesDeclaration += FString::Printf(TEXT("\t%s %s;") LINE_TERMINATOR, HLSLType, *PropertyName);
-
-			// Chainable method to set the attribute
-			MaterialAttributesUtilities += FString::Printf(TEXT("FMaterialAttributes FMaterialAttributes_Set%s(FMaterialAttributes InAttributes, %s InValue) { InAttributes.%s = InValue; return InAttributes; }") LINE_TERMINATOR,
-				*PropertyName, HLSLType, *PropertyName);
-
-			// Standard value type
-			switch (PropertyType)
-			{
-			case MCT_Float: case MCT_Float1: MaterialAttributesDefault += FString::Printf(TEXT("\tResult.%s = %0.8f;") LINE_TERMINATOR, *PropertyName, DefaultValue.X); break;
-			case MCT_Float2: MaterialAttributesDefault += FString::Printf(TEXT("\tResult.%s = MaterialFloat2(%0.8f,%0.8f);") LINE_TERMINATOR, *PropertyName, DefaultValue.X, DefaultValue.Y); break;
-			case MCT_Float3: MaterialAttributesDefault += FString::Printf(TEXT("\tResult.%s = MaterialFloat3(%0.8f,%0.8f,%0.8f);") LINE_TERMINATOR, *PropertyName, DefaultValue.X, DefaultValue.Y, DefaultValue.Z); break;
-			case MCT_Float4: MaterialAttributesDefault += FString::Printf(TEXT("\tResult.%s = MaterialFloat4(%0.8f,%0.8f,%0.8f,%0.8f);") LINE_TERMINATOR, *PropertyName, DefaultValue.X, DefaultValue.Y, DefaultValue.Z, DefaultValue.W); break;
-			case MCT_ShadingModel: MaterialAttributesDefault += FString::Printf(TEXT("\tResult.%s = %d;") LINE_TERMINATOR, *PropertyName, (int32)DefaultShadingModel); break;
-			case MCT_Strata: MaterialAttributesDefault += FString::Printf(TEXT("\tResult.%s = GetInitialisedStrataData();") LINE_TERMINATOR, *PropertyName); break; // TODO
-			default: checkNoEntry(); break;
-			}
-		}
-	}
-
-	LazyPrintf.PushParam(*MaterialAttributesDeclaration);
-	LazyPrintf.PushParam(*MaterialAttributesUtilities);
-
-	// Stores the shared shader results member declarations
-	// PixelMembersDeclaration should be the same for all variations, but might change in the future. There are cases where work is shared
-	// between the pixel and vertex shader, but with Nanite all work has to be moved into the pixel shader, which means we will want
-	// different inputs. But for now, we are keeping them the same.
-	FString PixelMembersDeclaration;
-
-	for (int32 PropertyIndex = 0; PropertyIndex < MP_MAX; ++PropertyIndex)
-	{
-		// Skip non-shared properties
-		if (!SharedPixelProperties[PropertyIndex])
-		{
-			continue;
-		}
-
-		const EMaterialProperty Property = (EMaterialProperty)PropertyIndex;
-		check(FMaterialAttributeDefinitionMap::GetShaderFrequency(Property) == SF_Pixel);
-		// Special case MP_SubsurfaceColor as the actual property is a combination of the color and the profile but we don't want to expose the profile
-		const FString PropertyName = Property == MP_SubsurfaceColor ? "Subsurface" : FMaterialAttributeDefinitionMap::GetAttributeName(Property);
-		check(PropertyName.Len() > 0);
-		const EMaterialValueType Type = Property == MP_SubsurfaceColor ? MCT_Float4 : FMaterialAttributeDefinitionMap::GetValueType(Property);
-
-		PixelMembersDeclaration += FString::Printf(TEXT("\t%s %s;\n"), HLSLTypeString(Type), *PropertyName);
-	}
-
-	//for (int32 Index = 0; Index < CompiledPDV_MAX; Index++)
-	//{
-	//	GetSharedInputsMaterialCode(PixelMembersDeclaration[Index], NormalAssignment[Index], PixelMembersSetupAndAssignments[Index], (ECompiledPartialDerivativeVariation)Index);
-	//}
-
-	LazyPrintf.PushParam(*PixelMembersDeclaration);
-
-	{
-		FString DerivativeHelpers;// = DerivativeAutogen.GenerateUsedFunctions(*this);
-		FString DerivativeHelpersAndResources;// = DerivativeHelpers + ResourcesString;
-		//LazyPrintf.PushParam(*ResourcesString);
-		LazyPrintf.PushParam(*DerivativeHelpersAndResources);
-	}
-
-	// Anything used bye the GenerationFunctionCode() like WorldPositionOffset shouldn't be using texures, right?
-	// Let those use the standard finite differences textures, since they should be the same. If we actually want
-	// those to handle texture reads properly, we'll have to make extra versions.
-	ECompiledPartialDerivativeVariation BaseDerivativeVariation = CompiledPDV_FiniteDifferences;
-
-	//if (bCompileForComputeShader)
-	//{
-	//	LazyPrintf.PushParam(*GenerateFunctionCode(CompiledMP_EmissiveColorCS, BaseDerivativeVariation));
-	//}
-	//else
-	{
-		LazyPrintf.PushParam(TEXT("return 0"));
-	}
-
-	LazyPrintf.PushParam(*FString::Printf(TEXT("return %.5f"), Material.GetTranslucencyDirectionalLightingIntensity()));
-
-	LazyPrintf.PushParam(*FString::Printf(TEXT("return %.5f"), Material.GetTranslucentShadowDensityScale()));
-	LazyPrintf.PushParam(*FString::Printf(TEXT("return %.5f"), Material.GetTranslucentSelfShadowDensityScale()));
-	LazyPrintf.PushParam(*FString::Printf(TEXT("return %.5f"), Material.GetTranslucentSelfShadowSecondDensityScale()));
-	LazyPrintf.PushParam(*FString::Printf(TEXT("return %.5f"), Material.GetTranslucentSelfShadowSecondOpacity()));
-	LazyPrintf.PushParam(*FString::Printf(TEXT("return %.5f"), Material.GetTranslucentBackscatteringExponent()));
-
-	{
-		FLinearColor Extinction = Material.GetTranslucentMultipleScatteringExtinction();
-
-		LazyPrintf.PushParam(*FString::Printf(TEXT("return MaterialFloat3(%.5f, %.5f, %.5f)"), Extinction.R, Extinction.G, Extinction.B));
-	}
-
-	LazyPrintf.PushParam(*FString::Printf(TEXT("return %.5f"), Material.GetOpacityMaskClipValue()));
-
-	LazyPrintf.PushParam(TEXT("return Parameters.MaterialVertexAttributes.WorldPositionOffset"));
-	LazyPrintf.PushParam(TEXT("return 0.0f"));
-	LazyPrintf.PushParam(TEXT("return 0.0f"));
-	LazyPrintf.PushParam(TEXT("return 0.0f"));
-	LazyPrintf.PushParam(TEXT("return 0.0f"));
-	LazyPrintf.PushParam(TEXT("return 0.0f"));
-
-	// Print custom texture coordinate assignments, should be fine with regular derivatives
-	FString CustomUVAssignments;
-
-	int32 LastProperty = -1;
-	for (uint32 CustomUVIndex = 0; CustomUVIndex < NumUserTexCoords; CustomUVIndex++)
-	{
-		//if (bEnableExecutionFlow)
-		{
-			const FString AttributeName = FMaterialAttributeDefinitionMap::GetAttributeName((EMaterialProperty)(MP_CustomizedUVs0 + CustomUVIndex));
-			CustomUVAssignments += FString::Printf(TEXT("\tOutTexCoords[%u] = Parameters.MaterialVertexAttributes.%s;") LINE_TERMINATOR, CustomUVIndex, *AttributeName);
-		}
-		/*else
-		{
-			if (CustomUVIndex == 0)
-			{
-				CustomUVAssignments += DerivativeVariations[BaseDerivativeVariation].TranslatedCodeChunkDefinitions[MP_CustomizedUVs0 + CustomUVIndex];
-			}
-
-			if (DerivativeVariations[BaseDerivativeVariation].TranslatedCodeChunkDefinitions[MP_CustomizedUVs0 + CustomUVIndex].Len() > 0)
-			{
-				if (LastProperty >= 0)
-				{
-					check(DerivativeVariations[BaseDerivativeVariation].TranslatedCodeChunkDefinitions[LastProperty].Len() == DerivativeVariations[BaseDerivativeVariation].TranslatedCodeChunkDefinitions[MP_CustomizedUVs0 + CustomUVIndex].Len());
-				}
-				LastProperty = MP_CustomizedUVs0 + CustomUVIndex;
-			}
-			CustomUVAssignments += FString::Printf(TEXT("\tOutTexCoords[%u] = %s;") LINE_TERMINATOR, CustomUVIndex, *DerivativeVariations[BaseDerivativeVariation].TranslatedCodeChunks[MP_CustomizedUVs0 + CustomUVIndex]);
-		}*/
-	}
-
-	LazyPrintf.PushParam(*CustomUVAssignments);
-
-	// Print custom vertex shader interpolator assignments
-	FString CustomInterpolatorAssignments;
-
-	/*for (UMaterialExpressionVertexInterpolator* Interpolator : CustomVertexInterpolators)
-	{
-		if (Interpolator->InterpolatorOffset != INDEX_NONE)
-		{
-			check(Interpolator->InterpolatorIndex != INDEX_NONE);
-			check(Interpolator->InterpolatedType & MCT_Float);
-
-			const EMaterialValueType Type = Interpolator->InterpolatedType == MCT_Float ? MCT_Float1 : Interpolator->InterpolatedType;
-			const TCHAR* Swizzle[2] = { TEXT("x"), TEXT("y") };
-			const int32 Offset = Interpolator->InterpolatorOffset;
-			const int32 Index = Interpolator->InterpolatorIndex;
-
-			// Note: We reference the UV define directly to avoid having to pre-accumulate UV counts before property translation
-			CustomInterpolatorAssignments += FString::Printf(TEXT("\tOutTexCoords[VERTEX_INTERPOLATOR_%i_TEXCOORDS_X].%s = VertexInterpolator%i(Parameters).x;") LINE_TERMINATOR, Index, Swizzle[Offset % 2], Index);
-
-			if (Type >= MCT_Float2)
-			{
-				CustomInterpolatorAssignments += FString::Printf(TEXT("\tOutTexCoords[VERTEX_INTERPOLATOR_%i_TEXCOORDS_Y].%s = VertexInterpolator%i(Parameters).y;") LINE_TERMINATOR, Index, Swizzle[(Offset + 1) % 2], Index);
-
-				if (Type >= MCT_Float3)
-				{
-					CustomInterpolatorAssignments += FString::Printf(TEXT("\tOutTexCoords[VERTEX_INTERPOLATOR_%i_TEXCOORDS_Z].%s = VertexInterpolator%i(Parameters).z;") LINE_TERMINATOR, Index, Swizzle[(Offset + 2) % 2], Index);
-
-					if (Type == MCT_Float4)
-					{
-						CustomInterpolatorAssignments += FString::Printf(TEXT("\tOutTexCoords[VERTEX_INTERPOLATOR_%i_TEXCOORDS_W].%s = VertexInterpolator%i(Parameters).w;") LINE_TERMINATOR, Index, Swizzle[(Offset + 3) % 2], Index);
-					}
-				}
-			}
-		}
-	}*/
-
-	LazyPrintf.PushParam(*CustomInterpolatorAssignments);
-
-	LazyPrintf.PushParam(*MaterialAttributesDefault);
-
-	//if (bEnableExecutionFlow)
-	{
-		FString EvaluateVertexCode;
-
-		// Set default texcoords in the VS
-		for (uint32 TexCoordIndex = 0; TexCoordIndex < NumUserVertexTexCoords; ++TexCoordIndex)
-		{
-			const FString AttributeName = FMaterialAttributeDefinitionMap::GetAttributeName((EMaterialProperty)(MP_CustomizedUVs0 + TexCoordIndex));
-
-			EvaluateVertexCode += FString::Printf(TEXT("\tDefaultMaterialAttributes.%s = Parameters.TexCoords[%d];") LINE_TERMINATOR, *AttributeName, TexCoordIndex);
-		}
-
-		//EvaluateVertexCode += TranslatedAttributesCodeChunks[SF_Vertex];
-
-		LazyPrintf.PushParam(*EvaluateVertexCode);
-		LazyPrintf.PushParam(PixelShaderCode);
-
-		FString EvaluateMaterialAttributesCode = TEXT("    FMaterialAttributes MaterialAttributes = EvaluatePixelMaterialAttributes(Parameters);" LINE_TERMINATOR);
-
-		for (int32 PropertyIndex = 0; PropertyIndex < MP_MAX; ++PropertyIndex)
-		{
-			// Skip non-shared properties
-			if (!SharedPixelProperties[PropertyIndex])
-			{
-				continue;
-			}
-
-			const EMaterialProperty Property = (EMaterialProperty)PropertyIndex;
-			check(FMaterialAttributeDefinitionMap::GetShaderFrequency(Property) == SF_Pixel);
-			// Special case MP_SubsurfaceColor as the actual property is a combination of the color and the profile but we don't want to expose the profile
-			const FString PropertyName = FMaterialAttributeDefinitionMap::GetAttributeName(Property);
-
-			if (PropertyIndex == MP_SubsurfaceColor)
-			{
-				// TODO - properly handle subsurface profile
-				EvaluateMaterialAttributesCode += FString::Printf("    PixelMaterialInputs.Subsurface = float4(MaterialAttributes.%s, 0.0f);" LINE_TERMINATOR, *PropertyName);
-			}
-			else
-			{
-				EvaluateMaterialAttributesCode += FString::Printf("    PixelMaterialInputs.%s = MaterialAttributes.%s;" LINE_TERMINATOR, *PropertyName, *PropertyName);
-			}
-		}
-
-		// TODO - deriv
-		for (int32 Iter = 0; Iter < CompiledPDV_MAX; Iter++)
-		{
-			LazyPrintf.PushParam(*EvaluateMaterialAttributesCode);
-			LazyPrintf.PushParam(TEXT(""));
-			LazyPrintf.PushParam(TEXT(""));
-		}
-	}
-
-	LazyPrintf.PushParam(*FString::Printf(TEXT("%u"), MaterialTemplateLineNumber));
-
-	return LazyPrintf.GetResultString();
-}
-
-static void GetMaterialEnvironment(EShaderPlatform InPlatform,
-	const FMaterial& InMaterial,
-	const FMaterialCompilationOutput& MaterialCompilationOutput,
-	FShaderCompilerEnvironment& OutEnvironment)
-{
-	bool bMaterialRequestsDualSourceBlending = false;
-
-	if (false)//bNeedsParticlePosition || Material->ShouldGenerateSphericalParticleNormals() || bUsesSphericalParticleOpacity)
-	{
-		OutEnvironment.SetDefine(TEXT("NEEDS_PARTICLE_POSITION"), 1);
-	}
-
-	if (false)//bNeedsParticleVelocity || Material->IsUsedWithNiagaraMeshParticles())
-	{
-		OutEnvironment.SetDefine(TEXT("NEEDS_PARTICLE_VELOCITY"), 1);
-	}
-
-	uint32 DynamicParticleParameterMask = 0u;
-	if (DynamicParticleParameterMask)
-	{
-		OutEnvironment.SetDefine(TEXT("USE_DYNAMIC_PARAMETERS"), 1);
-		OutEnvironment.SetDefine(TEXT("DYNAMIC_PARAMETERS_MASK"), DynamicParticleParameterMask);
-	}
-
-	if (false)//bNeedsParticleTime)
-	{
-		OutEnvironment.SetDefine(TEXT("NEEDS_PARTICLE_TIME"), 1);
-	}
-
-	if (false)//bUsesParticleMotionBlur)
-	{
-		OutEnvironment.SetDefine(TEXT("USES_PARTICLE_MOTION_BLUR"), 1);
-	}
-
-	if (false)//bNeedsParticleRandom)
-	{
-		OutEnvironment.SetDefine(TEXT("NEEDS_PARTICLE_RANDOM"), 1);
-	}
-
-	if (false)//bUsesSphericalParticleOpacity)
-	{
-		OutEnvironment.SetDefine(TEXT("SPHERICAL_PARTICLE_OPACITY"), TEXT("1"));
-	}
-
-	if (false)//bUsesParticleSubUVs)
-	{
-		OutEnvironment.SetDefine(TEXT("USE_PARTICLE_SUBUVS"), TEXT("1"));
-	}
-
-	if (false)//bUsesLightmapUVs)
-	{
-		OutEnvironment.SetDefine(TEXT("LIGHTMAP_UV_ACCESS"), TEXT("1"));
-	}
-
-	if (false)//bUsesAOMaterialMask)
-	{
-		OutEnvironment.SetDefine(TEXT("USES_AO_MATERIAL_MASK"), TEXT("1"));
-	}
-
-	if (false)//bUsesSpeedTree)
-	{
-		OutEnvironment.SetDefine(TEXT("USES_SPEEDTREE"), TEXT("1"));
-	}
-
-	if (false)//bNeedsWorldPositionExcludingShaderOffsets)
-	{
-		OutEnvironment.SetDefine(TEXT("NEEDS_WORLD_POSITION_EXCLUDING_SHADER_OFFSETS"), TEXT("1"));
-	}
-
-	if (false)//bNeedsParticleSize)
-	{
-		OutEnvironment.SetDefine(TEXT("NEEDS_PARTICLE_SIZE"), TEXT("1"));
-	}
-
-	if (MaterialCompilationOutput.bNeedsSceneTextures)
-	{
-		OutEnvironment.SetDefine(TEXT("NEEDS_SCENE_TEXTURES"), TEXT("1"));
-	}
-	if (MaterialCompilationOutput.bUsesEyeAdaptation)
-	{
-		OutEnvironment.SetDefine(TEXT("USES_EYE_ADAPTATION"), TEXT("1"));
-	}
-
-	if (MaterialCompilationOutput.bHasRuntimeVirtualTextureOutputNode)
-	{
-		OutEnvironment.SetDefine(TEXT("VIRTUAL_TEXTURE_OUTPUT"), 1);
-	}
-
-	OutEnvironment.SetDefine(TEXT("USES_PER_INSTANCE_CUSTOM_DATA"), false);// bUsesPerInstanceCustomData&& Material->IsUsedWithInstancedStaticMeshes());
-
-	// @todo MetalMRT: Remove this hack and implement proper atmospheric-fog solution for Metal MRT...
-	OutEnvironment.SetDefine(TEXT("MATERIAL_ATMOSPHERIC_FOG"), false);// !IsMetalMRTPlatform(InPlatform) ? bUsesAtmosphericFog : 0);
-	OutEnvironment.SetDefine(TEXT("MATERIAL_SKY_ATMOSPHERE"), false);// bUsesSkyAtmosphere);
-	OutEnvironment.SetDefine(TEXT("INTERPOLATE_VERTEX_COLOR"), false);// bUsesVertexColor);
-	OutEnvironment.SetDefine(TEXT("NEEDS_PARTICLE_COLOR"), false);// bUsesParticleColor);
-	OutEnvironment.SetDefine(TEXT("NEEDS_PARTICLE_LOCAL_TO_WORLD"), false);// bUsesParticleLocalToWorld);
-	OutEnvironment.SetDefine(TEXT("NEEDS_PARTICLE_WORLD_TO_LOCAL"), false);// bUsesParticleWorldToLocal);
-	OutEnvironment.SetDefine(TEXT("USES_TRANSFORM_VECTOR"), false);// bUsesTransformVector);
-	OutEnvironment.SetDefine(TEXT("WANT_PIXEL_DEPTH_OFFSET"), false);// bUsesPixelDepthOffset);
-	if (IsMetalPlatform(InPlatform))
-	{
-		OutEnvironment.SetDefine(TEXT("USES_WORLD_POSITION_OFFSET"), false);// bUsesWorldPositionOffset);
-	}
-	OutEnvironment.SetDefine(TEXT("USES_EMISSIVE_COLOR"), false);// bUsesEmissiveColor);
-	// Distortion uses tangent space transform 
-	OutEnvironment.SetDefine(TEXT("USES_DISTORTION"), InMaterial.IsDistorted());
-
-	OutEnvironment.SetDefine(TEXT("MATERIAL_ENABLE_TRANSLUCENCY_FOGGING"), InMaterial.ShouldApplyFogging());
-	OutEnvironment.SetDefine(TEXT("MATERIAL_ENABLE_TRANSLUCENCY_CLOUD_FOGGING"), InMaterial.ShouldApplyCloudFogging());
-	OutEnvironment.SetDefine(TEXT("MATERIAL_IS_SKY"), InMaterial.IsSky());
-	OutEnvironment.SetDefine(TEXT("MATERIAL_COMPUTE_FOG_PER_PIXEL"), InMaterial.ComputeFogPerPixel());
-	OutEnvironment.SetDefine(TEXT("MATERIAL_FULLY_ROUGH"), false);// bIsFullyRough || InMaterial.IsFullyRough());
-	OutEnvironment.SetDefine(TEXT("MATERIAL_USES_ANISOTROPY"), false);// bUsesAnisotropy);
-
-	// Count the number of VTStacks (each stack will allocate a feedback slot)
-	OutEnvironment.SetDefine(TEXT("NUM_VIRTUALTEXTURE_SAMPLES"), 0);// VTStacks.Num());
-
-	// Setup defines to map each VT stack to either 1 or 2 page table textures, depending on how many layers it uses
-	/*for (int i = 0; i < VTStacks.Num(); ++i)
-	{
-		const FMaterialVirtualTextureStack& Stack = MaterialCompilationOutput.UniformExpressionSet.VTStacks[i];
-		FString PageTableValue = FString::Printf(TEXT("Material.VirtualTexturePageTable0_%d"), i);
-		if (Stack.GetNumLayers() > 4u)
-		{
-			PageTableValue += FString::Printf(TEXT(", Material.VirtualTexturePageTable1_%d"), i);
-		}
-		if (VTStacks[i].bAdaptive)
-		{
-			PageTableValue += FString::Printf(TEXT(", Material.VirtualTexturePageTableIndirection_%d"), i);
-		}
-		OutEnvironment.SetDefine(*FString::Printf(TEXT("VIRTUALTEXTURE_PAGETABLE_%d"), i), *PageTableValue);
-	}*/
-
-	/*for (int32 CollectionIndex = 0; CollectionIndex < ParameterCollections.Num(); CollectionIndex++)
-	{
-		// Add uniform buffer declarations for any parameter collections referenced
-		const FString CollectionName = FString::Printf(TEXT("MaterialCollection%u"), CollectionIndex);
-		// This can potentially become an issue for MaterialCollection Uniform Buffers if they ever get non-numeric resources (eg Textures), as
-		// OutEnvironment.ResourceTableMap has a map by name, and the N ParameterCollection Uniform Buffers ALL are names "MaterialCollection"
-		// (and the hlsl cbuffers are named MaterialCollection0, etc, so the names don't match the layout)
-		FShaderUniformBufferParameter::ModifyCompilationEnvironment(*CollectionName, ParameterCollections[CollectionIndex]->GetUniformBufferStruct(), InPlatform, OutEnvironment);
-	}*/
-	OutEnvironment.SetDefine(TEXT("IS_MATERIAL_SHADER"), true);
-
-	// Set all the shading models for this material here 
-	FMaterialShadingModelField ShadingModels = InMaterial.GetShadingModels();
-
-	// If the material gets its shading model from the material expressions, then we use the result from the compilation (assuming it's valid).
-	// This result will potentially be tighter than what GetShadingModels() returns, because it only picks up the shading models from the expressions that get compiled for a specific feature level and quality level
-	// For example, the material might have shading models behind static switches. GetShadingModels() will return both the true and the false paths from that switch, whereas the shading model field from the compilation will only contain the actual shading model selected 
-	/*if (InMaterial.IsShadingModelFromMaterialExpression() && ShadingModelsFromCompilation.IsValid())
-	{
-		// Shading models fetched from the compilation of the expression graph
-		ShadingModels = ShadingModelsFromCompilation;
-	}*/
-
-	ensure(ShadingModels.IsValid());
-
-	if (ShadingModels.IsLit())
-	{
-		int32 NumSetMaterials = 0;
-		if (ShadingModels.HasShadingModel(MSM_DefaultLit))
-		{
-			OutEnvironment.SetDefine(TEXT("MATERIAL_SHADINGMODEL_DEFAULT_LIT"), TEXT("1"));
-			NumSetMaterials++;
-		}
-		if (ShadingModels.HasShadingModel(MSM_Subsurface))
-		{
-			OutEnvironment.SetDefine(TEXT("MATERIAL_SHADINGMODEL_SUBSURFACE"), TEXT("1"));
-			NumSetMaterials++;
-		}
-		if (ShadingModels.HasShadingModel(MSM_PreintegratedSkin))
-		{
-			OutEnvironment.SetDefine(TEXT("MATERIAL_SHADINGMODEL_PREINTEGRATED_SKIN"), TEXT("1"));
-			NumSetMaterials++;
-		}
-		if (ShadingModels.HasShadingModel(MSM_SubsurfaceProfile))
-		{
-			OutEnvironment.SetDefine(TEXT("MATERIAL_SHADINGMODEL_SUBSURFACE_PROFILE"), TEXT("1"));
-			NumSetMaterials++;
-		}
-		if (ShadingModels.HasShadingModel(MSM_ClearCoat))
-		{
-			OutEnvironment.SetDefine(TEXT("MATERIAL_SHADINGMODEL_CLEAR_COAT"), TEXT("1"));
-			NumSetMaterials++;
-		}
-		if (ShadingModels.HasShadingModel(MSM_TwoSidedFoliage))
-		{
-			OutEnvironment.SetDefine(TEXT("MATERIAL_SHADINGMODEL_TWOSIDED_FOLIAGE"), TEXT("1"));
-			NumSetMaterials++;
-		}
-		if (ShadingModels.HasShadingModel(MSM_Hair))
-		{
-			OutEnvironment.SetDefine(TEXT("MATERIAL_SHADINGMODEL_HAIR"), TEXT("1"));
-			NumSetMaterials++;
-		}
-		if (ShadingModels.HasShadingModel(MSM_Cloth))
-		{
-			OutEnvironment.SetDefine(TEXT("MATERIAL_SHADINGMODEL_CLOTH"), TEXT("1"));
-			NumSetMaterials++;
-		}
-		if (ShadingModels.HasShadingModel(MSM_Eye))
-		{
-			OutEnvironment.SetDefine(TEXT("MATERIAL_SHADINGMODEL_EYE"), TEXT("1"));
-			NumSetMaterials++;
-		}
-		if (ShadingModels.HasShadingModel(MSM_SingleLayerWater))
-		{
-			OutEnvironment.SetDefine(TEXT("MATERIAL_SHADINGMODEL_SINGLELAYERWATER"), TEXT("1"));
-			NumSetMaterials++;
-		}
-		if (ShadingModels.HasShadingModel(MSM_ThinTranslucent))
-		{
-			OutEnvironment.SetDefine(TEXT("MATERIAL_SHADINGMODEL_THIN_TRANSLUCENT"), TEXT("1"));
-			NumSetMaterials++;
-
-			bMaterialRequestsDualSourceBlending = true;
-		}
-
-		if (ShadingModels.HasShadingModel(MSM_SingleLayerWater) &&
-			(IsSwitchPlatform(InPlatform) || IsVulkanMobileSM5Platform(InPlatform) || FDataDrivenShaderPlatformInfo::GetRequiresDisableForwardLocalLights(InPlatform)))
-		{
-			OutEnvironment.SetDefine(TEXT("DISABLE_FORWARD_LOCAL_LIGHTS"), TEXT("1"));
-		}
-
-		// This is to have switch use the simple single layer water shading similar to mobile: no dynamic lights, only sun and sky, no distortion, no colored transmittance on background, no custom depth read.
-		const bool bSingleLayerWaterUsesSimpleShading = (IsSwitchPlatform(InPlatform) || IsVulkanMobileSM5Platform(InPlatform)) && IsForwardShadingEnabled(InPlatform);
-		if (ShadingModels.HasShadingModel(MSM_SingleLayerWater) && bSingleLayerWaterUsesSimpleShading)
-		{
-			OutEnvironment.SetDefine(TEXT("SINGLE_LAYER_WATER_SIMPLE_FORWARD"), TEXT("1"));
-		}
-
-		if (NumSetMaterials == 1)
-		{
-			OutEnvironment.SetDefine(TEXT("MATERIAL_SINGLE_SHADINGMODEL"), TEXT("1"));
-		}
-
-		ensure(NumSetMaterials != 0);
-		if (NumSetMaterials == 0)
-		{
-			// Should not really end up here
-			UE_LOG(LogMaterial, Warning, TEXT("Unknown material shading model(s). Setting to MSM_DefaultLit"));
-			OutEnvironment.SetDefine(TEXT("MATERIAL_SHADINGMODEL_DEFAULT_LIT"), TEXT("1"));
-		}
-	}
-	else
-	{
-		// Unlit shading model can only exist by itself
-		OutEnvironment.SetDefine(TEXT("MATERIAL_SINGLE_SHADINGMODEL"), TEXT("1"));
-		OutEnvironment.SetDefine(TEXT("MATERIAL_SHADINGMODEL_UNLIT"), TEXT("1"));
-	}
-
-	if (InMaterial.GetMaterialDomain() == MD_Volume) // && Material->HasN)
-	{
-		TArray<const UMaterialExpressionVolumetricAdvancedMaterialOutput*> VolumetricAdvancedExpressions;
-		InMaterial.GetMaterialInterface()->GetMaterial()->GetAllExpressionsOfType(VolumetricAdvancedExpressions);
-		if (VolumetricAdvancedExpressions.Num() > 0)
-		{
-			if (VolumetricAdvancedExpressions.Num() > 1)
-			{
-				UE_LOG(LogMaterial, Fatal, TEXT("Only a single UMaterialExpressionVolumetricAdvancedMaterialOutput node is supported."));
-			}
-
-			OutEnvironment.SetDefine(TEXT("MATERIAL_VOLUMETRIC_ADVANCED"), TEXT("1"));
-
-			const UMaterialExpressionVolumetricAdvancedMaterialOutput* VolumetricAdvancedNode = VolumetricAdvancedExpressions[0];
-			if (VolumetricAdvancedNode->GetEvaluatePhaseOncePerSample())
-			{
-				OutEnvironment.SetDefine(TEXT("MATERIAL_VOLUMETRIC_ADVANCED_PHASE_PERSAMPLE"), TEXT("1"));
-			}
-			else
-			{
-				OutEnvironment.SetDefine(TEXT("MATERIAL_VOLUMETRIC_ADVANCED_PHASE_PERPIXEL"), TEXT("1"));
-			}
-
-			OutEnvironment.SetDefine(TEXT("MATERIAL_VOLUMETRIC_ADVANCED_GRAYSCALE_MATERIAL"), VolumetricAdvancedNode->bGrayScaleMaterial ? TEXT("1") : TEXT("0"));
-			OutEnvironment.SetDefine(TEXT("MATERIAL_VOLUMETRIC_ADVANCED_RAYMARCH_VOLUME_SHADOW"), VolumetricAdvancedNode->bRayMarchVolumeShadow ? TEXT("1") : TEXT("0"));
-
-			OutEnvironment.SetDefine(TEXT("MATERIAL_VOLUMETRIC_ADVANCED_MULTISCATTERING_OCTAVE_COUNT"), VolumetricAdvancedNode->GetMultiScatteringApproximationOctaveCount());
-
-			OutEnvironment.SetDefine(TEXT("MATERIAL_VOLUMETRIC_ADVANCED_CONSERVATIVE_DENSITY"),
-				VolumetricAdvancedNode->ConservativeDensity.IsConnected() ? TEXT("1") : TEXT("0"));
-
-			OutEnvironment.SetDefine(TEXT("MATERIAL_VOLUMETRIC_ADVANCED_OVERRIDE_AMBIENT_OCCLUSION"),
-				InMaterial.HasAmbientOcclusionConnected() ? TEXT("1") : TEXT("0"));
-
-			OutEnvironment.SetDefine(TEXT("MATERIAL_VOLUMETRIC_ADVANCED_GROUND_CONTRIBUTION"),
-				VolumetricAdvancedNode->bGroundContribution ? TEXT("1") : TEXT("0"));
-		}
-	}
-
-	OutEnvironment.SetDefine(TEXT("MATERIAL_IS_STRATA"), false);// bMaterialIsStrata ? TEXT("1") : TEXT("0"));
-
-	// STRATA_TODO do not request DualSourceBlending if gray scale transmittance is selected.
-	// bMaterialRequestsDualSourceBlending this base on limited set of blend mode: Opaque, Masked, TransmittanceCoverage, TransmittanceColored;
-	//bMaterialRequestsDualSourceBlending |= bMaterialIsStrata;
-
-	// if duals source blending (colored transmittance) is not supported on a platform, it will fall back to standard alpha blending (grey scale transmittance)
-	OutEnvironment.SetDefine(TEXT("DUAL_SOURCE_COLOR_BLENDING_ENABLED"), false);// bMaterialRequestsDualSourceBlending&& Material->IsDualBlendingEnabled(Platform) ? TEXT("1") : TEXT("0"));
-
-	// Translate() is called before getting here so we can create related define
-	/*if (StrataMaterialAnalysis.RequestedBSDFCount > 0)
-	{
-		OutEnvironment.SetDefine(TEXT("STRATA_CLAMPED_LAYER_COUNT"), StrataMaterialAnalysis.ClampedLayerCount);
-	}*/
+	
 }
 
-
-bool FMaterialHLSLGenerator::GenerateHLSL(FMaterial& InMaterial,
-	FMaterialCompilationOutput& OutCompilationOutput,
-	TRefCountPtr<FSharedShaderCompilerEnvironment>& OutMaterialEnvironment)
+void FMaterialHLSLGenerator::AcquireErrors(TArray<FString>& OutCompileErrors, TArray<UMaterialExpression*>& OutErrorExpressions)
 {
-	UMaterialInterface* MaterialInterface = InMaterial.GetMaterialInterface();
-	UMaterial* BaseMaterial = MaterialInterface->GetMaterial();
-
-	FMemMark MemMark(Allocator);
-
-	HLSLTree = UE::HLSLTree::FTree::Create(Allocator);
-
-	UMaterialExpression* BaseExpression = BaseMaterial->ExpressionExecBegin;
-	UE::HLSLTree::FStatement* RootStatement = AcquireStatement(HLSLTree->GetRootScope(), BaseExpression);
-
-	UE::HLSLTree::FCodeWriter* CodeWriter = UE::HLSLTree::FCodeWriter::Create(Allocator);
-	
-	UE::HLSLTree::FEmitContext EmitContext;
-	EmitContext.Allocator = &Allocator;
-	EmitContext.Material = &InMaterial;
-	EmitContext.MaterialCompilationOutput = &OutCompilationOutput;
-	HLSLTree->EmitHLSL(EmitContext, *CodeWriter);
-
-	InMaterial.CompileErrors = MoveTemp(CompileErrors);
-	InMaterial.ErrorExpressions = MoveTemp(ErrorExpressions);
-
-	const FString MaterialTemplateSource = GenerateMaterialTemplateHLSL(ShaderPlatform, InMaterial, EmitContext, CodeWriter->StringBuilder->ToString());
-
-	OutMaterialEnvironment = new FSharedShaderCompilerEnvironment();
-	OutMaterialEnvironment->TargetPlatform = TargetPlatform;
-	GetMaterialEnvironment(ShaderPlatform, InMaterial, OutCompilationOutput, *OutMaterialEnvironment);
-	OutMaterialEnvironment->IncludeVirtualPathToContentsMap.Add(TEXT("/Engine/Generated/Material.ush"), MaterialTemplateSource);
-
-	return true;
+	OutCompileErrors = MoveTemp(CompileErrors);
+	OutErrorExpressions = MoveTemp(ErrorExpressions);
 }
 
 EMaterialGenerateHLSLStatus FMaterialHLSLGenerator::Error(const FString& Message)
@@ -779,6 +100,34 @@ UE::HLSLTree::FExpression* FMaterialHLSLGenerator::NewCast(UE::HLSLTree::FScope&
 	return Input;
 }
 
+UE::HLSLTree::FExpression* FMaterialHLSLGenerator::NewFunctionInput(UE::HLSLTree::FScope& Scope, int32 InputIndex, UMaterialExpressionFunctionInput* MaterialFunctionInput)
+{
+	UE::HLSLTree::EExpressionType ExpressionType = UE::HLSLTree::EExpressionType::Void;
+	switch (MaterialFunctionInput->InputType)
+	{
+	case FunctionInput_Scalar: ExpressionType = UE::HLSLTree::EExpressionType::Float1; break;
+	case FunctionInput_Vector2: ExpressionType = UE::HLSLTree::EExpressionType::Float2; break;
+	case FunctionInput_Vector3: ExpressionType = UE::HLSLTree::EExpressionType::Float3; break;
+	case FunctionInput_Vector4: ExpressionType = UE::HLSLTree::EExpressionType::Float4; break;
+	case FunctionInput_Texture2D:
+	case FunctionInput_TextureCube:
+	case FunctionInput_Texture2DArray:
+	case FunctionInput_VolumeTexture:
+	case FunctionInput_StaticBool:
+	case FunctionInput_MaterialAttributes:
+	case FunctionInput_TextureExternal:
+		// TODO
+		break;
+	}
+	check(ExpressionType != UE::HLSLTree::EExpressionType::Void);
+
+	UE::HLSLTree::FExpressionFunctionInput* Expression = HLSLTree->NewExpression<UE::HLSLTree::FExpressionFunctionInput>(Scope, MaterialFunctionInput->InputName, ExpressionType, InputIndex);
+	const FExpressionKey Key(MaterialFunctionInput, 0);
+	ExpressionMap.Add(Key, Expression);
+
+	return Expression;
+}
+
 UE::HLSLTree::FLocalDeclaration* FMaterialHLSLGenerator::AcquireLocalDeclaration(UE::HLSLTree::FScope& Scope, UE::HLSLTree::EExpressionType Type, const FName& Name)
 {
 	UE::HLSLTree::FLocalDeclaration*& Declaration = LocalDeclarationMap.FindOrAdd(Name);
@@ -794,7 +143,9 @@ UE::HLSLTree::FLocalDeclaration* FMaterialHLSLGenerator::AcquireLocalDeclaration
 		}
 		else
 		{
-			check(false);
+			Errorf(TEXT("Local %s first accessed as type %s, now type %s"), *Name.ToString(),
+				UE::HLSLTree::GetExpressionTypeDescription(Declaration->Type).Name,
+				UE::HLSLTree::GetExpressionTypeDescription(Type).Name);
 			return nullptr;
 		}
 	}
@@ -814,7 +165,7 @@ UE::HLSLTree::FParameterDeclaration* FMaterialHLSLGenerator::AcquireParameterDec
 UE::HLSLTree::FTextureParameterDeclaration* FMaterialHLSLGenerator::AcquireTextureDeclaration(UE::HLSLTree::FScope& Scope, const UE::HLSLTree::FTextureDescription& Value)
 {
 	FString SamplerTypeError;
-	if (!UMaterialExpressionTextureBase::VerifySamplerType(FeatureLevel, TargetPlatform, Value.Texture, Value.SamplerType, SamplerTypeError))
+	if (!UMaterialExpressionTextureBase::VerifySamplerType(CompileTarget.FeatureLevel, CompileTarget.TargetPlatform, Value.Texture, Value.SamplerType, SamplerTypeError))
 	{
 		Errorf(TEXT("%s"), *SamplerTypeError);
 		return nullptr;
@@ -831,7 +182,7 @@ UE::HLSLTree::FTextureParameterDeclaration* FMaterialHLSLGenerator::AcquireTextu
 UE::HLSLTree::FTextureParameterDeclaration* FMaterialHLSLGenerator::AcquireTextureParameterDeclaration(UE::HLSLTree::FScope& Scope, const FName& Name, const UE::HLSLTree::FTextureDescription& DefaultValue)
 {
 	FString SamplerTypeError;
-	if (!UMaterialExpressionTextureBase::VerifySamplerType(FeatureLevel, TargetPlatform, DefaultValue.Texture, DefaultValue.SamplerType, SamplerTypeError))
+	if (!UMaterialExpressionTextureBase::VerifySamplerType(CompileTarget.FeatureLevel, CompileTarget.TargetPlatform, DefaultValue.Texture, DefaultValue.SamplerType, SamplerTypeError))
 	{
 		Errorf(TEXT("%s"), *SamplerTypeError);
 		return nullptr;
@@ -845,6 +196,34 @@ UE::HLSLTree::FTextureParameterDeclaration* FMaterialHLSLGenerator::AcquireTextu
 	return Declaration;
 }
 
+UE::HLSLTree::FFunctionCall* FMaterialHLSLGenerator::AcquireFunctionCall(UE::HLSLTree::FScope& Scope, UMaterialFunctionInterface* Function, TArrayView<UE::HLSLTree::FExpression*> Inputs)
+{
+	FSHA1 Hasher;
+	for (UE::HLSLTree::FExpression* Input : Inputs)
+	{
+		Hasher.Update((uint8*)&Input, sizeof(Input));
+	}
+
+	FFunctionCallKey Key;
+	Key.Function = Function;
+	Hasher.GetHash(Key.InputHash.Hash);
+
+	UE::HLSLTree::FFunctionCall** PrevFunctionCall = FunctionCallMap.Find(Key);
+	UE::HLSLTree::FFunctionCall* FunctionCall = nullptr;
+	if (!PrevFunctionCall)
+	{
+		const FMaterialHLSLTree& FunctionTree = Function->AcquireHLSLTree(*this);
+		FunctionCall = FunctionTree.GenerateFunctionCall(*this, Scope, Inputs);
+		FunctionCallMap.Add(Key, FunctionCall);
+	}
+	else
+	{
+		FunctionCall = *PrevFunctionCall;
+		Scope.UseFunctionCall(FunctionCall);
+	}
+	return FunctionCall;
+}
+
 UE::HLSLTree::FExpression* FMaterialHLSLGenerator::AcquireExpression(UE::HLSLTree::FScope& Scope, UMaterialExpression* MaterialExpression, int32 OutputIndex)
 {
 	const FExpressionKey Key(MaterialExpression, OutputIndex);
diff --git a/Engine/Source/Runtime/Engine/Private/Materials/MaterialHLSLTree.cpp b/Engine/Source/Runtime/Engine/Private/Materials/MaterialHLSLTree.cpp
new file mode 100644
--- /dev/null
+++ b/Engine/Source/Runtime/Engine/Private/Materials/MaterialHLSLTree.cpp
@@ -0,0 +1,77 @@
+#include "MaterialHLSLTree.h"
+#include "MaterialHLSLGenerator.h"
+#include "Materials/Material.h"
+#include "Materials/MaterialFunction.h"
+#include "Materials/MaterialExpressionExecBegin.h"
+#include "Materials/MaterialExpressionFunctionInput.h"
+#include "Materials/MaterialExpressionFunctionOutput.h"
+#include "Materials/MaterialExpressionMaterialFunctionCall.h"
+#include "HLSLTree/HLSLTreeCommon.h"
+
+FMaterialHLSLTree::FMaterialHLSLTree() : Allocator(0)
+{
+}
+
+void FMaterialHLSLTree::InitializeForMaterial(const FMaterialCompileTargetParameters& InCompilerTarget, FMaterial& InOutMaterial)
+{
+	const UMaterialInterface* MaterialInterface = InOutMaterial.GetMaterialInterface();
+	const UMaterial* BaseMaterial = MaterialInterface->GetMaterial();
+	UMaterialExpression* BaseExpression = BaseMaterial->ExpressionExecBegin;
+
+	HLSLTree = UE::HLSLTree::FTree::Create(Allocator);
+
+	FMaterialHLSLGenerator Generator(InCompilerTarget, *HLSLTree);
+	Generator.AcquireStatement(HLSLTree->GetRootScope(), BaseExpression);
+	Generator.AcquireErrors(InOutMaterial.CompileErrors, InOutMaterial.ErrorExpressions);
+}
+
+void FMaterialHLSLTree::InitializeForFunction(const FMaterialCompileTargetParameters& InCompilerTarget, UMaterialFunctionInterface* InOutFunction)
+{
+	TArray<FFunctionExpressionInput> FunctionInputs;
+	TArray<FFunctionExpressionOutput> FunctionOutputs;
+	InOutFunction->GetInputsAndOutputs(FunctionInputs, FunctionOutputs);
+
+	HLSLTree = UE::HLSLTree::FTree::Create(Allocator);
+	UE::HLSLTree::FScope& RootScope = HLSLTree->GetRootScope();
+
+	NumFunctionInputs = FunctionInputs.Num();
+	NumFunctionOutputs = FunctionOutputs.Num();
+	FunctionOutputTypes = New<UE::HLSLTree::EExpressionType>(Allocator, NumFunctionOutputs);
+
+	FMaterialHLSLGenerator Generator(InCompilerTarget, *HLSLTree);
+	for (int32 InputIndex = 0; InputIndex < FunctionInputs.Num(); ++InputIndex)
+	{
+		const FFunctionExpressionInput& Input = FunctionInputs[InputIndex];
+		Generator.NewFunctionInput(RootScope, InputIndex, Input.ExpressionInput);
+	}
+
+	for (int32 OutputIndex = 0; OutputIndex < FunctionOutputs.Num(); ++OutputIndex)
+	{
+		const FFunctionExpressionOutput& Output = FunctionOutputs[OutputIndex];
+		UMaterialExpressionFunctionOutput* ExpressionOutput = Output.ExpressionOutput;
+
+		UE::HLSLTree::FExpression* Expression = ExpressionOutput->A.AcquireHLSLExpression(Generator, RootScope);
+		UE::HLSLTree::EExpressionType OutputType = UE::HLSLTree::EExpressionType::Void;
+		if (Expression)
+		{
+			UE::HLSLTree::FStatementSetFunctionOutput* OutputStatement = Generator.GetTree().NewStatement<UE::HLSLTree::FStatementSetFunctionOutput>(RootScope);
+			OutputStatement->Name = ExpressionOutput->OutputName;
+			OutputStatement->OutputIndex = OutputIndex;
+			OutputStatement->Expression = Expression;
+			OutputType = Expression->Type;
+		}
+		FunctionOutputTypes[OutputIndex] = OutputType;
+	}
+}
+
+UE::HLSLTree::FFunctionCall* FMaterialHLSLTree::GenerateFunctionCall(FMaterialHLSLGenerator& Generator, UE::HLSLTree::FScope& Scope, TArrayView<UE::HLSLTree::FExpression*> Inputs) const
+{
+	check(Inputs.Num() == NumFunctionInputs);
+
+	return HLSLTree->NewFunctionCall(Scope,
+		HLSLTree->GetRootScope(),
+		Inputs.GetData(),
+		FunctionOutputTypes,
+		NumFunctionInputs,
+		NumFunctionOutputs);
+}
diff --git a/Engine/Source/Runtime/Engine/Private/Materials/MaterialShared.cpp b/Engine/Source/Runtime/Engine/Private/Materials/MaterialShared.cpp
--- a/Engine/Source/Runtime/Engine/Private/Materials/MaterialShared.cpp
+++ b/Engine/Source/Runtime/Engine/Private/Materials/MaterialShared.cpp
@@ -49,6 +49,8 @@
 #if WITH_EDITOR
 #include "Rendering/StaticLightingSystemInterface.h"
 #include "MaterialHLSLGenerator.h"
+#include "MaterialHLSLTree.h"
+#include "MaterialHLSLEmitter.h"
 #include "HLSLTree/HLSLTreeCommon.h"
 #endif
 #if WITH_ODSC
@@ -2308,10 +2310,12 @@ bool FMaterial::Translate_New(const FMaterialShaderMapId& ShaderMapId,
 	TRefCountPtr<FSharedShaderCompilerEnvironment>& OutMaterialEnvironment)
 {
 #if WITH_EDITOR
-	FMaterialHLSLGenerator Generator(InPlatform, ShaderMapId.FeatureLevel, InTargetPlatform);
-	return Generator.GenerateHLSL(*this,
-		OutCompilationOutput,
-		OutMaterialEnvironment);
+	const FMaterialCompileTargetParameters TargetParams(InPlatform, ShaderMapId.FeatureLevel, InTargetPlatform);
+
+	FMaterialHLSLTree Tree;
+	Tree.InitializeForMaterial(TargetParams, *this);
+
+	return MaterialEmitHLSL(TargetParams, *this, Tree.GetTree(), OutCompilationOutput, OutMaterialEnvironment);
 #else
 	checkNoEntry();
 	return false;
diff --git a/Engine/Source/Runtime/Engine/Public/HLSLTree/HLSLTree.h b/Engine/Source/Runtime/Engine/Public/HLSLTree/HLSLTree.h
--- a/Engine/Source/Runtime/Engine/Public/HLSLTree/HLSLTree.h
+++ b/Engine/Source/Runtime/Engine/Public/HLSLTree/HLSLTree.h
@@ -81,23 +81,38 @@ public:
 	/** Returns a snippet of HLSL source code that references the given local variablew */
 	const TCHAR* AcquireHLSLReference(FLocalDeclaration* Declaration);
 
+	const TCHAR* AcquireHLSLReference(FFunctionCall* FunctionCall, int32 OutputIndex);
+
 	struct FScopeEntry
 	{
 		const FScope* Scope;
 		FCodeWriter* ExpressionCodeWriter;
 	};
 
+	struct FFunctionStackEntry
+	{
+		FFunctionCall* FunctionCall;
+	};
+
 	struct FDeclarationEntry
 	{
 		const TCHAR* Definition;
 	};
 
+	struct FFunctionCallEntry
+	{
+		TCHAR const* const* OutputDefinition;
+		int32 NumOutputs;
+	};
+
 	FScopeEntry* FindScope(FScope* Scope);
 
 	TArray<FScopeEntry> ScopeStack;
+	TArray<FFunctionStackEntry> FunctionStack;
 	TMap<FNode*, FDeclarationEntry> DeclarationMap;
+	TMap<FFunctionCall*, FFunctionCallEntry> FunctionCallMap;
 	FMemStackBase* Allocator = nullptr;
-	FMaterial* Material = nullptr; // TODO - remove preshader material dependency
+	const FMaterial* Material = nullptr; // TODO - remove preshader material dependency
 	FMaterialCompilationOutput* MaterialCompilationOutput = nullptr;
 	int32 NumExpressionLocals = 0;
 	int32 NumTexCoords = 0;
@@ -123,6 +138,7 @@ public:
 	virtual ENodeVisitResult OnLocalDeclaration(FLocalDeclaration& Declaration) { return ENodeVisitResult::VisitDependentNodes; }
 	virtual ENodeVisitResult OnParameterDeclaration(FParameterDeclaration& Declaration) { return ENodeVisitResult::VisitDependentNodes; }
 	virtual ENodeVisitResult OnTextureParameterDeclaration(FTextureParameterDeclaration& Declaration) { return ENodeVisitResult::VisitDependentNodes; }
+	virtual ENodeVisitResult OnFunctionCall(FFunctionCall& FunctionCall) { return ENodeVisitResult::VisitDependentNodes; }
 };
 
 /** Root class of the HLSL AST */
@@ -221,7 +237,7 @@ public:
 /**
  * Represents an HLSL texture parameter.
  */
-class FTextureParameterDeclaration : public FNode
+class FTextureParameterDeclaration final : public FNode
 {
 public:
 	FTextureParameterDeclaration(const FName& InName, const FTextureDescription& InDescription) : Name(InName), Description(InDescription) {}
@@ -232,6 +248,23 @@ public:
 	FTextureDescription Description;
 };
 
+/**
+ * Represents a call to an HLSL function from a separate tree
+ */
+class FFunctionCall final : public FNode
+{
+public:
+	virtual ENodeVisitResult Visit(FNodeVisitor& Visitor) override;
+
+	/** Root scope of the function to call. Note that this scope will be from a separate (external) tree */
+	const FScope* FunctionScope;
+
+	FExpression* const* Inputs;
+	EExpressionType const* OutputTypes;
+	int32 NumInputs;
+	int32 NumOutputs;
+};
+
 /**
  * Represents an HLSL scope.  This is an ordered list of statements (enclosed by an {} pair in HLSL)
  * All HLSL nodes track which scopes they are accessed from
@@ -247,6 +280,7 @@ public:
 	void AddStatement(FStatement* Statement);
 
 	void UseDeclaration(FLocalDeclaration* Declaration);
+	void UseFunctionCall(FFunctionCall* FunctionCall);
 	void UseExpression(FExpression* Expression);
 	bool TryMoveStatement(FStatement* Statement);
 
@@ -307,6 +341,14 @@ public:
 	FLocalDeclaration* NewLocalDeclaration(FScope& Scope, EExpressionType Type, const FName& Name);
 	FParameterDeclaration* NewParameterDeclaration(FScope& Scope, const FName& Name, const FConstant& DefaultValue);
 	FTextureParameterDeclaration* NewTextureParameterDeclaration(FScope& Scope, const FName& Name, const FTextureDescription& DefaultValue);
+
+	FFunctionCall* NewFunctionCall(FScope& Scope,
+		const FScope& FunctionScope,
+		FExpression* const* Inputs,
+		EExpressionType const* OutputTypes,
+		int32 NumInputs,
+		int32 NumOutputs);
+
 private:
 	template<typename T, typename... ArgTypes>
 	inline T* NewNode(ArgTypes&&... Args)
diff --git a/Engine/Source/Runtime/Engine/Public/HLSLTree/HLSLTreeCommon.h b/Engine/Source/Runtime/Engine/Public/HLSLTree/HLSLTreeCommon.h
--- a/Engine/Source/Runtime/Engine/Public/HLSLTree/HLSLTreeCommon.h
+++ b/Engine/Source/Runtime/Engine/Public/HLSLTree/HLSLTreeCommon.h
@@ -255,6 +255,66 @@ public:
 	virtual void EmitHLSL(FEmitContext& Context, FCodeWriter& Writer) const override;
 };
 
+class FExpressionFunctionInput : public FExpression
+{
+public:
+	FExpressionFunctionInput(const FName& InName, EExpressionType InType, int32 InIndex)
+		: FExpression(InType), Name(InName), InputIndex(InIndex)
+	{}
+
+	virtual void EmitHLSL(FEmitContext& Context, FCodeWriter& Writer) const override;
+
+	FName Name;
+	int32 InputIndex;
+};
+
+class FExpressionFunctionOutput : public FExpression
+{
+public:
+	FExpressionFunctionOutput(FFunctionCall* InFunctionCall, int32 InIndex)
+		: FExpression(InFunctionCall->OutputTypes[InIndex])
+		, FunctionCall(InFunctionCall)
+		, OutputIndex(InIndex)
+	{
+		check(InIndex >= 0 && InIndex < InFunctionCall->NumOutputs);
+	}
+
+	FFunctionCall* FunctionCall;
+	int32 OutputIndex;
+
+	virtual ENodeVisitResult Visit(FNodeVisitor& Visitor) override
+	{
+		const ENodeVisitResult Result = FExpression::Visit(Visitor);
+		if (ShouldVisitDependentNodes(Result))
+		{
+			Visitor.VisitNode(FunctionCall);
+		}
+		return Result;
+	}
+
+	virtual void EmitHLSL(FEmitContext& Context, FCodeWriter& Writer) const override;
+};
+
+class FStatementSetFunctionOutput : public FStatement
+{
+public:
+	FExpression* Expression;
+	FName Name;
+	int32 OutputIndex;
+
+	virtual ENodeVisitResult Visit(FNodeVisitor& Visitor) override
+	{
+		const ENodeVisitResult Result = FStatement::Visit(Visitor);
+		if (ShouldVisitDependentNodes(Result))
+		{
+			Visitor.VisitNode(Expression);
+		}
+		return Result;
+	}
+
+	virtual void EmitHLSL(FEmitContext& Context, FCodeWriter& Writer) const override;
+};
+
 class FStatementReturn : public FStatement
 {
 public:
diff --git a/Engine/Source/Runtime/Engine/Public/HLSLTree/HLSLTreeTypes.h b/Engine/Source/Runtime/Engine/Public/HLSLTree/HLSLTreeTypes.h
--- a/Engine/Source/Runtime/Engine/Public/HLSLTree/HLSLTreeTypes.h
+++ b/Engine/Source/Runtime/Engine/Public/HLSLTree/HLSLTreeTypes.h
@@ -17,6 +17,7 @@ class FParameterDeclaration;
 class FTextureParameterDeclaration;
 class FStatement;
 class FScope;
+class FFunctionCall;
 class FTree;
 class FCodeWriter;
 
diff --git a/Engine/Source/Runtime/Engine/Public/MaterialHLSLGenerator.h b/Engine/Source/Runtime/Engine/Public/MaterialHLSLGenerator.h
--- a/Engine/Source/Runtime/Engine/Public/MaterialHLSLGenerator.h
+++ b/Engine/Source/Runtime/Engine/Public/MaterialHLSLGenerator.h
@@ -14,7 +14,12 @@
 class FMaterial;
 class FMaterialCompilationOutput;
 struct FSharedShaderCompilerEnvironment;
+struct FMaterialCompileTargetParameters;
+class UMaterialFunctionInterface;
 class UMaterialExpression;
+class UMaterialExpressionFunctionInput;
+class UMaterialExpressionFunctionOutput;
+struct FFunctionExpressionInput;
 class ITargetPlatform;
 enum class EMaterialGenerateHLSLStatus : uint8;
 
@@ -27,6 +32,7 @@ class FExpressionConstant;
 class FExpressionExternalInput;
 class FExpressionSwizzle;
 class FExpressionCast;
+class FFunctionCall;
 }
 }
 
@@ -36,13 +42,17 @@ class FExpressionCast;
 class FMaterialHLSLGenerator
 {
 public:
-	FMaterialHLSLGenerator(EShaderPlatform InShaderPlatform,
-		ERHIFeatureLevel::Type InFeatureLevel,
-		const ITargetPlatform* InTargetPlatform);
+	FMaterialHLSLGenerator(const FMaterialCompileTargetParameters& InCompilerTarget,
+		UE::HLSLTree::FTree& InOutTree);
 
-	bool GenerateHLSL(FMaterial& InMaterial,
-		FMaterialCompilationOutput& OutCompilationOutput,
-		TRefCountPtr<FSharedShaderCompilerEnvironment>& OutMaterialEnvironment);
+	const FMaterialCompileTargetParameters& GetCompileTarget() const { return CompileTarget; }
+
+	//bool GenerateHLSL(FMaterial& InMaterial,
+	//	FMaterialCompilationOutput& OutCompilationOutput,
+	//	TRefCountPtr<FSharedShaderCompilerEnvironment>& OutMaterialEnvironment);
+
+	/** Retrieve the compile errors from the generator */
+	void AcquireErrors(TArray<FString>& OutCompileErrors, TArray<UMaterialExpression*>& OutErrorExpressions);
 
 	EMaterialGenerateHLSLStatus Error(const FString& Message);
 
@@ -59,6 +69,8 @@ public:
 	UE::HLSLTree::FExpressionSwizzle* NewSwizzle(UE::HLSLTree::FScope& Scope, const UE::HLSLTree::FSwizzleParameters& Params, UE::HLSLTree::FExpression* Input);
 	UE::HLSLTree::FExpression* NewCast(UE::HLSLTree::FScope& Scope, UE::HLSLTree::EExpressionType Type, UE::HLSLTree::FExpression* Input, UE::HLSLTree::ECastFlags Flags = UE::HLSLTree::ECastFlags::None);
 
+	UE::HLSLTree::FExpression* NewFunctionInput(UE::HLSLTree::FScope& Scope, int32 InputIndex, UMaterialExpressionFunctionInput* MaterialFunctionInput);
+
 	UE::HLSLTree::FLocalDeclaration* AcquireLocalDeclaration(UE::HLSLTree::FScope& Scope, UE::HLSLTree::EExpressionType Type, const FName& Name);
 	UE::HLSLTree::FParameterDeclaration* AcquireParameterDeclaration(UE::HLSLTree::FScope& Scope, const FName& Name, const UE::HLSLTree::FConstant& DefaultValue);
 
@@ -68,6 +80,8 @@ public:
 	/** Returns a declaration to access the given texture parameter */
 	UE::HLSLTree::FTextureParameterDeclaration* AcquireTextureParameterDeclaration(UE::HLSLTree::FScope& Scope, const FName& Name, const UE::HLSLTree::FTextureDescription& DefaultValue);
 
+	UE::HLSLTree::FFunctionCall* AcquireFunctionCall(UE::HLSLTree::FScope& Scope, UMaterialFunctionInterface* Function, TArrayView<UE::HLSLTree::FExpression*> Inputs);
+
 	/**
 	 * Returns the appropriate HLSLNode representing the given UMaterialExpression.
 	 * The node will be created if it doesn't exist. Otherwise, the tree will be updated to ensure the node is visible in the given scope
@@ -96,10 +110,23 @@ private:
 		}
 	};
 
-	FMemStackBase Allocator;
-	EShaderPlatform ShaderPlatform;
-	ERHIFeatureLevel::Type FeatureLevel;
-	const ITargetPlatform* TargetPlatform;
+	struct FFunctionCallKey
+	{
+		UMaterialFunctionInterface* Function;
+		FSHAHash InputHash;
+
+		friend inline uint32 GetTypeHash(const FFunctionCallKey& Value)
+		{
+			return HashCombine(GetTypeHash(Value.Function), GetTypeHash(Value.InputHash));
+		}
+
+		friend inline bool operator==(const FFunctionCallKey& Lhs, const FFunctionCallKey& Rhs)
+		{
+			return Lhs.Function == Rhs.Function && Lhs.InputHash == Rhs.InputHash;
+		}
+	};
+
+	const FMaterialCompileTargetParameters& CompileTarget;
 	UE::HLSLTree::FTree* HLSLTree;
 	TArray<FExpressionKey> ExpressionStack;
 	TArray<FString> CompileErrors;
@@ -108,6 +135,7 @@ private:
 	TMap<FName, UE::HLSLTree::FParameterDeclaration*> ParameterDeclarationMap;
 	TMap<UE::HLSLTree::FTextureDescription, UE::HLSLTree::FTextureParameterDeclaration*> TextureDeclarationMap;
 	TMap<FName, UE::HLSLTree::FTextureParameterDeclaration*> TextureParameterDeclarationMap;
+	TMap<FFunctionCallKey, UE::HLSLTree::FFunctionCall*> FunctionCallMap;
 	TMap<FExpressionKey, UE::HLSLTree::FExpression*> ExpressionMap;
 	TMap<UMaterialExpression*, UE::HLSLTree::FStatement*> StatementMap;
 };
diff --git a/Engine/Source/Runtime/Engine/Public/MaterialHLSLTree.h b/Engine/Source/Runtime/Engine/Public/MaterialHLSLTree.h
new file mode 100644
--- /dev/null
+++ b/Engine/Source/Runtime/Engine/Public/MaterialHLSLTree.h
@@ -0,0 +1,43 @@
+#pragma once
+
+#include "Misc/MemStack.h"
+
+class FMaterial;
+class UMaterialFunctionInterface;
+struct FMaterialCompileTargetParameters;
+class FMaterialHLSLGenerator;
+
+namespace UE
+{
+namespace HLSLTree
+{
+class FTree;
+class FScope;
+class FExpression;
+class FLocalDeclaration;
+class FFunctionCall;
+enum class EExpressionType : uint8;
+}
+}
+
+class FMaterialHLSLTree
+{
+public:
+	FMaterialHLSLTree();
+
+	void InitializeForMaterial(const FMaterialCompileTargetParameters& InCompilerTarget, FMaterial& InOutMaterial);
+
+	void InitializeForFunction(const FMaterialCompileTargetParameters& InCompilerTarget, UMaterialFunctionInterface* InOutFunction);
+
+	UE::HLSLTree::FTree& GetTree() const { return *HLSLTree; }
+	int32 GetNumFunctionInputs() const { return NumFunctionInputs; }
+	int32 GetNumFunctionOutputs() const { return NumFunctionOutputs; }
+
+	UE::HLSLTree::FFunctionCall* GenerateFunctionCall(FMaterialHLSLGenerator& Generator, UE::HLSLTree::FScope& Scope, TArrayView<UE::HLSLTree::FExpression*> Inputs) const;
+private:
+	FMemStackBase Allocator;
+	UE::HLSLTree::FTree* HLSLTree = nullptr;
+	UE::HLSLTree::EExpressionType* FunctionOutputTypes = nullptr;
+	int32 NumFunctionInputs = 0;
+	int32 NumFunctionOutputs = 0;
+};
diff --git a/Engine/Source/Runtime/Engine/Public/MaterialShared.h b/Engine/Source/Runtime/Engine/Public/MaterialShared.h
--- a/Engine/Source/Runtime/Engine/Public/MaterialShared.h
+++ b/Engine/Source/Runtime/Engine/Public/MaterialShared.h
@@ -209,6 +209,18 @@ enum EMaterialDomain
 
 ENGINE_API FString MaterialDomainString(EMaterialDomain MaterialDomain);
 
+/** Fully describes a material compilation target */
+struct FMaterialCompileTargetParameters
+{
+	FMaterialCompileTargetParameters(EShaderPlatform InShaderPlatform, ERHIFeatureLevel::Type InFeatureLevel, const ITargetPlatform* InTargetPlatform)
+		: ShaderPlatform(InShaderPlatform), FeatureLevel(InFeatureLevel), TargetPlatform(InTargetPlatform)
+	{}
+
+	EShaderPlatform ShaderPlatform;
+	ERHIFeatureLevel::Type FeatureLevel;
+	const ITargetPlatform* TargetPlatform;
+};
+
 /**
  * The context of a material being rendered.
  */
@@ -2231,7 +2243,7 @@ private:
 	friend class FMaterialShaderMap;
 	friend class FShaderCompilingManager;
 	friend class FHLSLMaterialTranslator;
-	friend class FMaterialHLSLGenerator;
+	friend class FMaterialHLSLTree;
 };
 
 /**
