Add missing copyright notices, and WITH_EDITOR guards

diff --git a/Engine/Source/Runtime/Engine/Private/Materials/MaterialHLSLEmitter.cpp b/Engine/Source/Runtime/Engine/Private/Materials/MaterialHLSLEmitter.cpp
--- a/Engine/Source/Runtime/Engine/Private/Materials/MaterialHLSLEmitter.cpp
+++ b/Engine/Source/Runtime/Engine/Private/Materials/MaterialHLSLEmitter.cpp
@@ -1,4 +1,8 @@
+// Copyright Epic Games, Inc. All Rights Reserved.
 #include "Materials/MaterialHLSLEmitter.h"
+
+#if WITH_EDITOR
+
 #include "Materials/MaterialExpression.h"
 #include "Materials/MaterialExpressionTextureBase.h"
 #include "Materials/MaterialExpressionExecBegin.h"
@@ -689,3 +693,5 @@ bool MaterialEmitHLSL(const FMaterialCompileTargetParameters& InCompilerTarget,
 
 	return true;
 }
+
+#endif // WITH_EDITOR
diff --git a/Engine/Source/Runtime/Engine/Private/Materials/MaterialHLSLEmitter.h b/Engine/Source/Runtime/Engine/Private/Materials/MaterialHLSLEmitter.h
--- a/Engine/Source/Runtime/Engine/Private/Materials/MaterialHLSLEmitter.h
+++ b/Engine/Source/Runtime/Engine/Private/Materials/MaterialHLSLEmitter.h
@@ -1,5 +1,8 @@
+// Copyright Epic Games, Inc. All Rights Reserved.
 #pragma once
 
+#if WITH_EDITOR
+
 #include "CoreMinimal.h"
 #include "Templates/RefCounting.h"
 
@@ -21,3 +24,5 @@ bool MaterialEmitHLSL(const FMaterialCompileTargetParameters& InCompilerTarget,
 	const UE::HLSLTree::FTree& InTree,
 	FMaterialCompilationOutput& OutCompilationOutput,
 	TRefCountPtr<FSharedShaderCompilerEnvironment>& OutMaterialEnvironment);
+
+#endif // WITH_EDITOR
diff --git a/Engine/Source/Runtime/Engine/Private/Materials/MaterialHLSLTree.cpp b/Engine/Source/Runtime/Engine/Private/Materials/MaterialHLSLTree.cpp
--- a/Engine/Source/Runtime/Engine/Private/Materials/MaterialHLSLTree.cpp
+++ b/Engine/Source/Runtime/Engine/Private/Materials/MaterialHLSLTree.cpp
@@ -1,4 +1,8 @@
+// Copyright Epic Games, Inc. All Rights Reserved.
 #include "MaterialHLSLTree.h"
+
+#if WITH_EDITOR
+
 #include "MaterialHLSLGenerator.h"
 #include "Materials/Material.h"
 #include "Materials/MaterialFunction.h"
@@ -75,3 +79,5 @@ UE::HLSLTree::FFunctionCall* FMaterialHLSLTree::GenerateFunctionCall(FMaterialHL
 		NumFunctionInputs,
 		NumFunctionOutputs);
 }
+
+#endif // WITH_EDITOR
diff --git a/Engine/Source/Runtime/Engine/Public/MaterialHLSLTree.h b/Engine/Source/Runtime/Engine/Public/MaterialHLSLTree.h
--- a/Engine/Source/Runtime/Engine/Public/MaterialHLSLTree.h
+++ b/Engine/Source/Runtime/Engine/Public/MaterialHLSLTree.h
@@ -1,5 +1,8 @@
+// Copyright Epic Games, Inc. All Rights Reserved.
 #pragma once
 
+#if WITH_EDITOR
+
 #include "Misc/MemStack.h"
 
 class FMaterial;
@@ -41,3 +44,5 @@ private:
 	int32 NumFunctionInputs = 0;
 	int32 NumFunctionOutputs = 0;
 };
+
+#endif // WITH_EDITOR
