Autoderiv:
Fixed issue where non-float inputs to if node would assert
Added more typeindex sanity checks

diff --git a/Engine/Shaders/Public/ShaderVersion.ush b/Engine/Shaders/Public/ShaderVersion.ush
--- a/Engine/Shaders/Public/ShaderVersion.ush
+++ b/Engine/Shaders/Public/ShaderVersion.ush
@@ -3,4 +3,4 @@
 // in Platform.ush (which should be included in any shader) it allows to invalidate the shader DDC.
 //
 // If you are merging streams and there is a conflict with this GUID you should make a new GUID rather than taking one or the other.
-// GUID = D4BD6A39-2B29-485A-8DBA-25CA2FCBB666
+// GUID = 1A51199A-E12A-49BE-B98F-CBF652EDE4A0
diff --git a/Engine/Source/Runtime/Engine/Private/Materials/HLSLMaterialDerivativeAutogen.cpp b/Engine/Source/Runtime/Engine/Private/Materials/HLSLMaterialDerivativeAutogen.cpp
--- a/Engine/Source/Runtime/Engine/Private/Materials/HLSLMaterialDerivativeAutogen.cpp
+++ b/Engine/Source/Runtime/Engine/Private/Materials/HLSLMaterialDerivativeAutogen.cpp
@@ -76,7 +76,7 @@ static inline const TCHAR * GetDerivVectorName(int32 TypeIndex)
 // 2: Float3
 // 3: Float4
 // -1: Everything else
-int32 GetDerivTypeIndex(EMaterialValueType ValueType)
+int32 GetDerivTypeIndex(EMaterialValueType ValueType, bool bAllowNonFloat)
 {
 	int32 Ret = INDEX_NONE;
 	if (ValueType == MCT_Float1 || ValueType == MCT_Float)
@@ -97,8 +97,8 @@ int32 GetDerivTypeIndex(EMaterialValueType ValueType)
 	}
 	else
 	{
-		Ret = -1; // redundant, but leaving it here so we can set breakpoints
-		check(false);
+		Ret = INDEX_NONE;
+		check(bAllowNonFloat);
 	}
 	return Ret;
 }
@@ -211,6 +211,9 @@ void FMaterialDerivativeAutogen::EnableGeneratedDepencencies()
 // given a string, convert it from type to type
 FString FMaterialDerivativeAutogen::CoerceValueRaw(const FString& Token, int32 SrcType, EDerivativeStatus SrcStatus, int32 DstType)
 {
+	check(IsDerivTypeIndexValid(SrcType));
+	check(IsDerivTypeIndexValid(DstType));
+
 	FString Ret = Token;
 
 	// If the original value is a derivative, grab the raw value
@@ -253,8 +256,8 @@ FString FMaterialDerivativeAutogen::CoerceValueRaw(const FString& Token, int32 S
 // given a string, convert it from type to type
 FString FMaterialDerivativeAutogen::CoerceValueDeriv(const FString& Token, int32 SrcType, EDerivativeStatus SrcStatus, int32 DstType)
 {
-	check(SrcType < 4);
-	check(DstType < 4);
+	check(IsDerivTypeIndexValid(SrcType));
+	check(IsDerivTypeIndexValid(DstType));
 
 	FString Ret = Token;
 
@@ -273,8 +276,7 @@ FString FMaterialDerivativeAutogen::CoerceValueDeriv(const FString& Token, int32
 
 FString FMaterialDerivativeAutogen::ConstructDeriv(const FString& Value, const FString& Ddx, const FString& Ddy, int32 DstType)
 {
-	check(DstType >= 0);
-	check(DstType < 4);
+	check(IsDerivTypeIndexValid(DstType));
 
 	bConstructDerivEnabled[DstType] = true;
 	FString TypeName = GetDerivVectorName(DstType);
@@ -284,8 +286,7 @@ FString FMaterialDerivativeAutogen::ConstructDeriv(const FString& Value, const F
 
 FString FMaterialDerivativeAutogen::ConstructDerivFinite(const FString& Value, int32 DstType)
 {
-	check(DstType >= 0);
-	check(DstType < 4);
+	check(IsDerivTypeIndexValid(DstType));
 
 	bConstructFiniteDerivEnabled[DstType] = true;
 
@@ -296,8 +297,9 @@ FString FMaterialDerivativeAutogen::ConstructDerivFinite(const FString& Value, i
 
 FString FMaterialDerivativeAutogen::ConvertDeriv(const FString& Value, int32 DstType, int32 SrcType)
 {
-	check(DstType < 4);
-	check(SrcType < 4);
+	check(IsDerivTypeIndexValid(DstType));
+	check(IsDerivTypeIndexValid(SrcType));
+
 	if (DstType == SrcType)
 		return Value;
 	
@@ -309,6 +311,8 @@ FString FMaterialDerivativeAutogen::ConvertDeriv(const FString& Value, int32 Dst
 
 int32 FMaterialDerivativeAutogen::GetFunc1ReturnNumComponents(int32 SrcTypeIndex, EFunc1 Op)
 {
+	check(IsDerivTypeIndexValid(SrcTypeIndex));
+
 	int32 DstTypeIndex = INDEX_NONE;
 
 	switch(Op)
@@ -340,6 +344,7 @@ int32 FMaterialDerivativeAutogen::GetFunc1ReturnNumComponents(int32 SrcTypeIndex
 		DstTypeIndex = 0;
 		break;
 	default:
+		check(0);
 		break;
 	}
 
@@ -348,6 +353,9 @@ int32 FMaterialDerivativeAutogen::GetFunc1ReturnNumComponents(int32 SrcTypeIndex
 
 int32 FMaterialDerivativeAutogen::GetFunc2ReturnNumComponents(int32 LhsTypeIndex, int32 RhsTypeIndex, EFunc2 Op)
 {
+	check(IsDerivTypeIndexValid(LhsTypeIndex));
+	check(IsDerivTypeIndexValid(RhsTypeIndex));
+
 	int32 DstTypeIndex = INDEX_NONE;
 	
 	switch(Op)
@@ -379,6 +387,7 @@ int32 FMaterialDerivativeAutogen::GetFunc2ReturnNumComponents(int32 LhsTypeIndex
 		DstTypeIndex = 2;
 		break;
 	default:
+		check(0);
 		break;
 	}
 
@@ -590,7 +599,6 @@ int32 FMaterialDerivativeAutogen::GenerateExpressionFunc2(FHLSLMaterialTranslato
 		return INDEX_NONE;
 	}
 
-
 	const FDerivInfo LhsDerivInfo = Translator.GetDerivInfo(LhsCode);
 	const FDerivInfo RhsDerivInfo = Translator.GetDerivInfo(RhsCode);
 
@@ -929,8 +937,8 @@ int32 FMaterialDerivativeAutogen::GenerateIfFunc(FHLSLMaterialTranslator& Transl
 			return INDEX_NONE;
 	}
 	
-	const FDerivInfo GreaterDerivInfo	= Translator.GetDerivInfo(Greater);
-	const FDerivInfo LessDerivInfo		= Translator.GetDerivInfo(Less);
+	const FDerivInfo GreaterDerivInfo	= Translator.GetDerivInfo(Greater, true);
+	const FDerivInfo LessDerivInfo		= Translator.GetDerivInfo(Less, true);
 
 	FString CodeFinite;
 	FString ThresholdFinite;
@@ -968,7 +976,7 @@ int32 FMaterialDerivativeAutogen::GenerateIfFunc(FHLSLMaterialTranslator& Transl
 		FString CodeAnalytic;
 		if (bEqual)
 		{
-			const FDerivInfo EqualDerivInfo = Translator.GetDerivInfo(Equal);
+			const FDerivInfo EqualDerivInfo = Translator.GetDerivInfo(Equal, true);
 			FString EqualDeriv = Translator.GetParameterCodeDeriv(Equal, CompiledPDV_Analytic);
 			EqualDeriv = CoerceValueDeriv(EqualDeriv, ResultTypeIndex, EqualDerivInfo.DerivativeStatus, ResultTypeIndex);
 
diff --git a/Engine/Source/Runtime/Engine/Private/Materials/HLSLMaterialDerivativeAutogen.h b/Engine/Source/Runtime/Engine/Private/Materials/HLSLMaterialDerivativeAutogen.h
--- a/Engine/Source/Runtime/Engine/Private/Materials/HLSLMaterialDerivativeAutogen.h
+++ b/Engine/Source/Runtime/Engine/Private/Materials/HLSLMaterialDerivativeAutogen.h
@@ -29,7 +29,13 @@ inline bool IsDerivativeValid(EDerivativeStatus Status)
 	return Status == EDerivativeStatus::Valid || Status == EDerivativeStatus::Zero;
 }
 
-int32 GetDerivTypeIndex(EMaterialValueType ValueType);
+inline bool IsDerivTypeIndexValid(int32 DerivTypeIndex)
+{
+	return (uint32)DerivTypeIndex < 4u;
+}
+
+int32 GetDerivTypeIndex(EMaterialValueType ValueType, bool bAllowNonFloat = false);
+
 
 struct FDerivInfo
 {
diff --git a/Engine/Source/Runtime/Engine/Private/Materials/HLSLMaterialTranslator.cpp b/Engine/Source/Runtime/Engine/Private/Materials/HLSLMaterialTranslator.cpp
--- a/Engine/Source/Runtime/Engine/Private/Materials/HLSLMaterialTranslator.cpp
+++ b/Engine/Source/Runtime/Engine/Private/Materials/HLSLMaterialTranslator.cpp
@@ -2187,10 +2187,10 @@ EDerivativeStatus FHLSLMaterialTranslator::GetDerivativeStatus(int32 Index) cons
 	return AtParameterCodeChunk(Index).DerivativeStatus;
 }
 
-FDerivInfo FHLSLMaterialTranslator::GetDerivInfo(int32 Index) const
+FDerivInfo FHLSLMaterialTranslator::GetDerivInfo(int32 Index, bool bAllowNonFloat) const
 {
 	const FShaderCodeChunk& CodeChunk = AtParameterCodeChunk(Index);
-	return FDerivInfo { CodeChunk.Type, GetDerivTypeIndex(CodeChunk.Type), CodeChunk.DerivativeStatus };
+	return FDerivInfo { CodeChunk.Type, GetDerivTypeIndex(CodeChunk.Type, bAllowNonFloat), CodeChunk.DerivativeStatus };
 }
 
 // Similar to GetParameterCode, but has no default, and is derivative aware. Making it a separate function in case it needs to diverge,
diff --git a/Engine/Source/Runtime/Engine/Private/Materials/HLSLMaterialTranslator.h b/Engine/Source/Runtime/Engine/Private/Materials/HLSLMaterialTranslator.h
--- a/Engine/Source/Runtime/Engine/Private/Materials/HLSLMaterialTranslator.h
+++ b/Engine/Source/Runtime/Engine/Private/Materials/HLSLMaterialTranslator.h
@@ -453,7 +453,7 @@ public:
 
 	const FShaderCodeChunk& AtParameterCodeChunk(int32 Index) const;
 	EDerivativeStatus GetDerivativeStatus(int32 Index) const;
-	FDerivInfo GetDerivInfo(int32 Index) const;
+	FDerivInfo GetDerivInfo(int32 Index, bool bAllowNonFloat = false) const;
 
 protected:
 
