Added support for mirroring curves to AnimNode_Mirror

diff --git a/Engine/Source/Editor/UnrealEd/Private/Factories/MirrorDataTableFactory.cpp b/Engine/Source/Editor/UnrealEd/Private/Factories/MirrorDataTableFactory.cpp
--- a/Engine/Source/Editor/UnrealEd/Private/Factories/MirrorDataTableFactory.cpp
+++ b/Engine/Source/Editor/UnrealEd/Private/Factories/MirrorDataTableFactory.cpp
@@ -342,6 +342,7 @@ UObject* UMirrorDataTableFactory::FactoryCreateNew(UClass* Class, UObject* InPar
 			{
 				DataTable->FindReplaceMirroredNames(); 
 			}
+			Skeleton->OnSmartNamesChangedEvent.AddUObject(DataTable, &UMirrorDataTable::FillMirrorArrays);
 		}
 	}
 	return DataTable;
diff --git a/Engine/Source/Runtime/AnimGraphRuntime/Private/AnimNodes/AnimNode_Mirror.cpp b/Engine/Source/Runtime/AnimGraphRuntime/Private/AnimNodes/AnimNode_Mirror.cpp
--- a/Engine/Source/Runtime/AnimGraphRuntime/Private/AnimNodes/AnimNode_Mirror.cpp
+++ b/Engine/Source/Runtime/AnimGraphRuntime/Private/AnimNodes/AnimNode_Mirror.cpp
@@ -117,6 +117,7 @@ void FAnimNode_Mirror::Evaluate_AnyThread(FPoseContext& Output)
 	if (bMirrorState && MirrorDataTable)
 	{
 		FAnimationRuntime::MirrorPose(Output.Pose, MirrorDataTable->MirrorAxis, CompactPoseMirrorBones, ComponentSpaceRefRotations);
+		FAnimationRuntime::MirrorCurves(Output.Curve, *MirrorDataTable);
 	}
 }
 
diff --git a/Engine/Source/Runtime/Engine/Classes/Animation/MirrorDataTable.h b/Engine/Source/Runtime/Engine/Classes/Animation/MirrorDataTable.h
--- a/Engine/Source/Runtime/Engine/Classes/Animation/MirrorDataTable.h
+++ b/Engine/Source/Runtime/Engine/Classes/Animation/MirrorDataTable.h
@@ -109,6 +109,8 @@ public:
 
 	ENGINE_API virtual void RestoreAfterStructChange() override;
 
+	ENGINE_API virtual void PreEditChange(FProperty* PropertyThatWillChange) override;
+
 	ENGINE_API virtual void PostEditChangeProperty(FPropertyChangedEvent& PropertyChangedEvent) override;
 
 	ENGINE_API virtual void PostEditUndo() override;
@@ -194,7 +196,7 @@ public:
 protected: 
 
 	// Fill BoneToMirrorBoneIndex, CurveMirrorSourceUIDArray, CurveMirrorTargetUIDArray and NotifyToMirrorNotifyIndex based on the Skeleton and Table Contents
-	void FillMirrorArrays(); 
+	ENGINE_API void FillMirrorArrays();
 
 	void HandleDataTableChanged();
 };
diff --git a/Engine/Source/Runtime/Engine/Private/Animation/AnimationRuntime.cpp b/Engine/Source/Runtime/Engine/Private/Animation/AnimationRuntime.cpp
--- a/Engine/Source/Runtime/Engine/Private/Animation/AnimationRuntime.cpp
+++ b/Engine/Source/Runtime/Engine/Private/Animation/AnimationRuntime.cpp
@@ -1113,6 +1113,58 @@ void FAnimationRuntime::AccumulateMeshSpaceRotationAdditiveToLocalPoseInternal(F
 	}
 }
 
+void  FAnimationRuntime::MirrorCurves(FBlendedCurve& Curves, const UMirrorDataTable& MirrorDataTable)
+{
+	if (Curves.UIDToArrayIndexLUT == nullptr)
+	{
+		return;
+	}
+
+	int32 NumMirrorUIDs = MirrorDataTable.CurveMirrorSourceUIDArray.Num();
+
+	for (int32 MirrorIndex = 0; MirrorIndex < NumMirrorUIDs; ++MirrorIndex)
+	{
+		SmartName::UID_Type SourceMirrorUID = MirrorDataTable.CurveMirrorSourceUIDArray[MirrorIndex];
+		SmartName::UID_Type TargetMirrorUID = MirrorDataTable.CurveMirrorTargetUIDArray[MirrorIndex];
+		int32 SourceMirrorIndex = Curves.GetArrayIndexByUID(SourceMirrorUID);
+		int32 TargetMirrorIndex = Curves.GetArrayIndexByUID(TargetMirrorUID);
+		if (SourceMirrorIndex != INDEX_NONE && TargetMirrorIndex != INDEX_NONE)
+		{
+			if (Curves.ValidCurveWeights[SourceMirrorIndex])
+			{
+				if (Curves.ValidCurveWeights[TargetMirrorIndex])
+				{
+					// Determine if we should swap or overwrite values. If the map has a paired entries (left->right and right->left) 
+					// these entries will appear beside each other in the arrays
+					SmartName::UID_Type NextSourceMirrorUID = INDEX_NONE; 
+					if (MirrorIndex + 1 < NumMirrorUIDs)
+					{
+						NextSourceMirrorUID = MirrorDataTable.CurveMirrorSourceUIDArray[MirrorIndex + 1];
+					}
+					if (NextSourceMirrorUID == TargetMirrorUID)
+					{
+						float SwapWeight = Curves.CurveWeights[TargetMirrorIndex];
+						Curves.CurveWeights[TargetMirrorIndex] = Curves.CurveWeights[SourceMirrorIndex];
+						Curves.CurveWeights[SourceMirrorIndex] = SwapWeight;
+						// skip over the next entry since it is swapped 
+						MirrorIndex++;
+					}
+					else
+					{
+						Curves.CurveWeights[TargetMirrorIndex] = Curves.CurveWeights[SourceMirrorIndex];
+					}
+				}
+				else
+				{
+					Curves.CurveWeights[TargetMirrorIndex] = Curves.CurveWeights[SourceMirrorIndex];
+					Curves.ValidCurveWeights[TargetMirrorIndex] = true;
+					Curves.ValidCurveWeights[SourceMirrorIndex] = false;
+				}
+			}
+		}
+	}
+}
+
 FVector FAnimationRuntime::MirrorVector(const FVector& V, EAxis::Type MirrorAxis)
 {
 	FVector MirrorV(V);
diff --git a/Engine/Source/Runtime/Engine/Private/Animation/MirrorDataTable.cpp b/Engine/Source/Runtime/Engine/Private/Animation/MirrorDataTable.cpp
--- a/Engine/Source/Runtime/Engine/Private/Animation/MirrorDataTable.cpp
+++ b/Engine/Source/Runtime/Engine/Private/Animation/MirrorDataTable.cpp
@@ -73,6 +73,11 @@ void UMirrorDataTable::PostLoad()
 {
 	Super::PostLoad();
 	FillMirrorArrays(); 
+	if (Skeleton)
+	{
+		Skeleton->ConditionalPostLoad();
+		Skeleton->OnSmartNamesChangedEvent.AddUObject(this, &UMirrorDataTable::FillMirrorArrays);
+	}
 }
 
 void UMirrorDataTable::EmptyTable()
@@ -90,10 +95,30 @@ void UMirrorDataTable::RestoreAfterStructChange()
 	Super::RestoreAfterStructChange();
 }
 
+void UMirrorDataTable::PreEditChange(FProperty* PropertyThatWillChange)
+{
+	Super::PreEditChange(PropertyThatWillChange);
+
+	if (PropertyThatWillChange && PropertyThatWillChange->GetFName() == GET_MEMBER_NAME_STRING_CHECKED(UMirrorDataTable, Skeleton))
+	{
+		if (Skeleton)
+		{
+			Skeleton->OnSmartNamesChangedEvent.RemoveAll(this);
+		}
+	}
+}
+
 void UMirrorDataTable::PostEditChangeProperty(FPropertyChangedEvent& PropertyChangedEvent)
 {
 	Super::PostEditChangeProperty(PropertyChangedEvent);
 	FillMirrorArrays(); 
+	if (PropertyChangedEvent.Property && PropertyChangedEvent.Property->GetFName() == GET_MEMBER_NAME_CHECKED(UMirrorDataTable, Skeleton))
+	{
+		if (Skeleton)
+		{
+			Skeleton->OnSmartNamesChangedEvent.AddUObject(this, &UMirrorDataTable::FillMirrorArrays);
+		}
+	}
 }
 
 void UMirrorDataTable::PostEditUndo()
@@ -382,7 +407,6 @@ void UMirrorDataTable::FillMirrorArrays()
 
 	TMap<FName, FName> CurveToMirrorCurveMap;
 	
-
 	ForeachRow<FMirrorTableRow>(TEXT("UMirrorDataTable::FillMirrorArrays"), [this, &CurveToMirrorCurveMap](const FName& Key, const FMirrorTableRow& Value) mutable
 		{
 			if (Value.MirrorEntryType == EMirrorRowType::Curve)
@@ -396,6 +420,8 @@ void UMirrorDataTable::FillMirrorArrays()
 		}
 	);
 
+	//ensure that pairs always appear beside each other in the arrays
+	TSet<SmartName::UID_Type> AddedSourceUIDs; 
 	const FSmartNameMapping* CurveSmartNames = Skeleton->GetSmartNameContainer(USkeleton::AnimCurveMappingName);
 	if (CurveSmartNames)
 	{
@@ -405,10 +431,17 @@ void UMirrorDataTable::FillMirrorArrays()
 		{
 			SmartName::UID_Type SourceCurveUID = CurveSmartNames->FindUID(Elem.Key);
 			SmartName::UID_Type TargetCurveUID = CurveSmartNames->FindUID(Elem.Value);
-			if (SourceCurveUID != INDEX_NONE && TargetCurveUID != INDEX_NONE)
+			if (SourceCurveUID != INDEX_NONE && TargetCurveUID != INDEX_NONE && !AddedSourceUIDs.Contains(SourceCurveUID))
 			{
 				CurveMirrorSourceUIDArray.Add(SourceCurveUID);
+				AddedSourceUIDs.Add(SourceCurveUID); 
 				CurveMirrorTargetUIDArray.Add(TargetCurveUID);
+				if (CurveToMirrorCurveMap.Contains(Elem.Value) && CurveSmartNames->FindUID(CurveToMirrorCurveMap[Elem.Value]) == SourceCurveUID)
+				{
+					CurveMirrorSourceUIDArray.Add(TargetCurveUID);
+					AddedSourceUIDs.Add(TargetCurveUID);
+					CurveMirrorTargetUIDArray.Add(SourceCurveUID);
+				}
 			}
 		}
 		CurveMirrorSourceUIDArray.Shrink(); 
diff --git a/Engine/Source/Runtime/Engine/Public/AnimationRuntime.h b/Engine/Source/Runtime/Engine/Public/AnimationRuntime.h
--- a/Engine/Source/Runtime/Engine/Public/AnimationRuntime.h
+++ b/Engine/Source/Runtime/Engine/Public/AnimationRuntime.h
@@ -444,6 +444,16 @@ public:
 	UE_DEPRECATED(4.26.0, "Please use BlendTransform() for weighted blending")
 	static void BlendTransformsByWeight(FTransform& OutTransform, const TArray<FTransform>& Transforms, const TArray<float>& Weights);
 
+	/**
+	 * Mirror (swap) curves with the specified MirrorDataTable.
+	 * Partial mirroring is supported,  and curves with two entries (Right->Left and Left->Right) will be swapped while 
+	 * curves with a single entry will be overwritten
+	 *
+	 * @param	Curves			The Curves which are swapped
+	 * @param	MirrorDataTable	A UMirrorDataTable specifying which curves to swap
+	 */
+	static void MirrorCurves(FBlendedCurve& Curves, const UMirrorDataTable& MirrorDataTable);
+
 	/**
 	 * Mirror a vector across the specified mirror axis 
 	 * @param	V			The vector to mirror
