Fix ignored bounds scale for Static mesh when GPU-Scene instance culling is enabled
- removed instance setup for static mesh proxy, instead letting defult kick in.

diff --git a/Engine/Source/Runtime/Engine/Classes/Engine/InstancedStaticMesh.h b/Engine/Source/Runtime/Engine/Classes/Engine/InstancedStaticMesh.h
--- a/Engine/Source/Runtime/Engine/Classes/Engine/InstancedStaticMesh.h
+++ b/Engine/Source/Runtime/Engine/Classes/Engine/InstancedStaticMesh.h
@@ -555,6 +555,18 @@ public:
 		return true;
 	}
 
+#if GPUCULL_TODO
+	virtual const TArray<FPrimitiveInstance>* GetPrimitiveInstances() const
+	{
+		return &Instances;
+	}
+
+	virtual TArray<FPrimitiveInstance>* GetPrimitiveInstances()
+	{
+		return &Instances;
+	}
+#endif
+
 protected:
 	/** Cache of the StaticMesh asset, needed to release SpeedTree resources*/
 	UStaticMesh* StaticMesh;
@@ -562,6 +574,10 @@ protected:
 	/** Per component render data */
 	FInstancedStaticMeshRenderData InstancedRenderData;
 
+#if GPUCULL_TODO
+	TArray<FPrimitiveInstance> Instances;
+#endif
+
 #if WITH_EDITOR
 	/* If we we have any selected instances */
 	bool bHasSelectedInstances;
diff --git a/Engine/Source/Runtime/Engine/Private/StaticMeshRender.cpp b/Engine/Source/Runtime/Engine/Private/StaticMeshRender.cpp
--- a/Engine/Source/Runtime/Engine/Private/StaticMeshRender.cpp
+++ b/Engine/Source/Runtime/Engine/Private/StaticMeshRender.cpp
@@ -351,20 +351,6 @@ FStaticMeshSceneProxy::FStaticMeshSceneProxy(UStaticMeshComponent* InComponent,
 #endif
 
 	AddSpeedTreeWind();
-#if GPUCULL_TODO
-	if (UseGPUScene(GMaxRHIShaderPlatform, FeatureLevel))
-	{
-		Instances.SetNum(1);
-		FPrimitiveInstance& Instance = Instances[0];
-		Instance.PrimitiveId = ~uint32(0);
-		Instance.InstanceToLocal.SetIdentity();
-		Instance.LocalToInstance.SetIdentity();
-		Instance.LocalToWorld.SetIdentity();
-		Instance.RenderBounds = InComponent->GetStaticMesh()->GetBounds();
-		Instance.LocalBounds = Instance.RenderBounds;
-		bSupportsInstanceDataBuffer = true;
-	}
-#endif // GPUCULL_TODO
 }
 
 void FStaticMeshSceneProxy::SetEvaluateWorldPositionOffsetInRayTracing(bool NewValue)
diff --git a/Engine/Source/Runtime/Engine/Public/StaticMeshResources.h b/Engine/Source/Runtime/Engine/Public/StaticMeshResources.h
--- a/Engine/Source/Runtime/Engine/Public/StaticMeshResources.h
+++ b/Engine/Source/Runtime/Engine/Public/StaticMeshResources.h
@@ -930,18 +930,6 @@ public:
 
 	virtual int32 GetLightMapCoordinateIndex() const override;
 
-#if GPUCULL_TODO
-	virtual const TArray<FPrimitiveInstance>* GetPrimitiveInstances() const
-	{
-		return &Instances;
-	}
-
-	virtual TArray<FPrimitiveInstance>* GetPrimitiveInstances()
-	{
-		return &Instances;
-	}
-#endif
-
 protected:
 	/** Configures mesh batch vertex / index state. Returns the number of primitives used in the element. */
 	uint32 SetMeshElementGeometrySource(
@@ -1153,11 +1141,6 @@ protected:
 	uint8 HierarchicalLODIndex;
 #endif
 
-#if GPUCULL_TODO
-protected:
-	TArray<FPrimitiveInstance> Instances;
-#endif
-
 public:
 
 	/**
