Don't get level instance bounds from umap package when the level instance doesn't support loading (which is the case for packed level instances).

diff --git a/Engine/Source/Runtime/Engine/Private/LevelInstance/LevelInstanceActor.cpp b/Engine/Source/Runtime/Engine/Private/LevelInstance/LevelInstanceActor.cpp
--- a/Engine/Source/Runtime/Engine/Private/LevelInstance/LevelInstanceActor.cpp
+++ b/Engine/Source/Runtime/Engine/Private/LevelInstance/LevelInstanceActor.cpp
@@ -592,12 +592,15 @@ void ALevelInstance::GetActorLocationBounds(bool bOnlyCollidingComponents, FVect
 	Super::GetActorLocationBounds(bOnlyCollidingComponents, Origin, BoxExtent, bIncludeFromChildActors);
 
 	// Add Level Bounds
-	if (ULevelInstanceSubsystem* LevelInstanceSubsystem = GetLevelInstanceSubsystem())
+	if (SupportsLoading())
 	{
-		FBox LevelInstanceBounds;
-		if (LevelInstanceSubsystem->GetLevelInstanceBounds(this, LevelInstanceBounds))
+		if (ULevelInstanceSubsystem* LevelInstanceSubsystem = GetLevelInstanceSubsystem())
 		{
-			LevelInstanceBounds.GetCenterAndExtents(Origin, BoxExtent);
+			FBox LevelInstanceBounds;
+			if (LevelInstanceSubsystem->GetLevelInstanceBounds(this, LevelInstanceBounds))
+			{
+				LevelInstanceBounds.GetCenterAndExtents(Origin, BoxExtent);
+			}
 		}
 	}
 }
