Added support for unsaved/deleted/newly created actors in PIE:
- Before duplicating the world for PIE, gather the list of dirty actors and store it in the persistent level.
- During PIE duplication, dirty actors will get duplicated and will serve as templates instead of streaming them.
- When streaming a cell, duplicate dirty actors from the list instead of loading them from disk.
- Create temporary actor descriptor views for unsaved, newly added actors.
- Don't create actor descriptor views for unsaved, deleted actors.
- Removed the dialog asking to save dirty actors when starting PIE.

diff --git a/Engine/Source/Editor/UnrealEd/Private/PlayLevel.cpp b/Engine/Source/Editor/UnrealEd/Private/PlayLevel.cpp
--- a/Engine/Source/Editor/UnrealEd/Private/PlayLevel.cpp
+++ b/Engine/Source/Editor/UnrealEd/Private/PlayLevel.cpp
@@ -115,7 +115,6 @@
 #include "IAssetViewport.h"
 #include "IPIEAuthorizer.h"
 #include "Features/IModularFeatures.h"
-#include "WorldPartition/WorldPartitionSubsystem.h"
 
 DEFINE_LOG_CATEGORY(LogPlayLevel);
 
@@ -1083,10 +1082,7 @@ void UEditorEngine::StartQueuedPlaySessionRequestImpl()
 		bIsSeparateProcess |= NetMode == EPlayNetMode::PIE_Client;
 	}
 
-	// For now, WorldPartition also requires the user to save their content for PIE to match what we have in-editor.
-	bool bSaveMapsForPlaySession = bIsSeparateProcess || GWorld->HasSubsystem<UWorldPartitionSubsystem>();
-
-	if (bSaveMapsForPlaySession && !SaveMapsForPlaySession())
+	if (bIsSeparateProcess && !SaveMapsForPlaySession())
 	{
 		// Maps did not save, print a warning
 		FText ErrorMsg = LOCTEXT("PIEWorldSaveFail", "PIE failed because map save was canceled");
diff --git a/Engine/Source/Runtime/Engine/Classes/Engine/Level.h b/Engine/Source/Runtime/Engine/Classes/Engine/Level.h
--- a/Engine/Source/Runtime/Engine/Classes/Engine/Level.h
+++ b/Engine/Source/Runtime/Engine/Classes/Engine/Level.h
@@ -440,6 +440,10 @@ public:
 	TArray<AActor*> ActorsForGC;
 
 #if WITH_EDITORONLY_DATA
+	/** List of modified, unsaved actors that needs to be duplicated for PIE */
+	UPROPERTY(NonPIEDuplicateTransient)
+	TMap<FName, AActor*> ActorsModifiedForPIE;
+
 	/** Use external actors, new actor spawned in this level will be external and existing external actors will be loaded on load. */
 	UPROPERTY(EditInstanceOnly, Category=World)
 	bool bUseExternalActors;
diff --git a/Engine/Source/Runtime/Engine/Classes/GameFramework/Actor.h b/Engine/Source/Runtime/Engine/Classes/GameFramework/Actor.h
--- a/Engine/Source/Runtime/Engine/Classes/GameFramework/Actor.h
+++ b/Engine/Source/Runtime/Engine/Classes/GameFramework/Actor.h
@@ -840,6 +840,7 @@ public:
 	void SetLockLocation(bool bInLockLocation) { bLockLocation = bInLockLocation; }
 
 private:
+	friend class FActorDescList;
 	friend class UActorDescContainer;
 
 	/**
diff --git a/Engine/Source/Runtime/Engine/Private/WorldPartition/ActorDescContainer.cpp b/Engine/Source/Runtime/Engine/Private/WorldPartition/ActorDescContainer.cpp
--- a/Engine/Source/Runtime/Engine/Private/WorldPartition/ActorDescContainer.cpp
+++ b/Engine/Source/Runtime/Engine/Private/WorldPartition/ActorDescContainer.cpp
@@ -16,8 +16,7 @@ UActorDescContainer::UActorDescContainer(const FObjectInitializer& ObjectInitial
 	, bContainerInitialized(false)
 	, bIgnoreAssetRegistryEvents(false)
 #endif
-{
-}
+{}
 
 void UActorDescContainer::Initialize(UWorld* InWorld, FName InPackageName, bool bRegisterDelegates)
 {
@@ -239,44 +238,6 @@ void UActorDescContainer::OnAssetUpdated(const FAssetData& InAssetData)
 	}
 }
 
-const FWorldPartitionActorDesc* UActorDescContainer::GetActorDesc(const FGuid& Guid) const
-{
-	const TUniquePtr<FWorldPartitionActorDesc>* const * ActorDesc = Actors.Find(Guid);
-	return ActorDesc ? (*ActorDesc)->Get() : nullptr;
-}
-
-FWorldPartitionActorDesc* UActorDescContainer::GetActorDesc(const FGuid& Guid)
-{
-	TUniquePtr<FWorldPartitionActorDesc>** ActorDesc = Actors.Find(Guid);
-	return ActorDesc ? (*ActorDesc)->Get() : nullptr;
-}
-
-const FWorldPartitionActorDesc& UActorDescContainer::GetActorDescChecked(const FGuid& Guid) const
-{
-	const TUniquePtr<FWorldPartitionActorDesc>* const ActorDesc = Actors.FindChecked(Guid);
-	return *ActorDesc->Get();
-}
-
-FWorldPartitionActorDesc& UActorDescContainer::GetActorDescChecked(const FGuid& Guid)
-{
-	TUniquePtr<FWorldPartitionActorDesc>* ActorDesc = Actors.FindChecked(Guid);
-	return *ActorDesc->Get();
-}
-
-const FWorldPartitionActorDesc* UActorDescContainer::GetActorDesc(const FString& PackageName) const
-{
-	const FName PackageFName(*PackageName);
-	for (const TUniquePtr<FWorldPartitionActorDesc>& ActorDescPtr : ActorDescList)
-	{
-		if (ActorDescPtr->ActorPackage == PackageFName)
-		{
-			return ActorDescPtr.Get();
-		}
-	}
-
-	return nullptr;
-}
-
 void UActorDescContainer::RegisterDelegates()
 {
 	if (GEditor && !IsTemplate())
diff --git a/Engine/Source/Runtime/Engine/Private/WorldPartition/ActorDescList.cpp b/Engine/Source/Runtime/Engine/Private/WorldPartition/ActorDescList.cpp
new file mode 100644
--- /dev/null
+++ b/Engine/Source/Runtime/Engine/Private/WorldPartition/ActorDescList.cpp
@@ -0,0 +1,60 @@
+// Copyright Epic Games, Inc. All Rights Reserved.
+
+#include "WorldPartition/ActorDescList.h"
+
+#if WITH_EDITOR
+FWorldPartitionActorDesc* FActorDescList::AddActor(const AActor* InActor)
+{
+	TUniquePtr<FWorldPartitionActorDesc>* NewActorDesc = new(ActorDescList) TUniquePtr<FWorldPartitionActorDesc>(InActor->CreateActorDesc());
+	check(NewActorDesc->IsValid());
+
+	check(!Actors.Contains((*NewActorDesc)->GetGuid()));
+	Actors.Add((*NewActorDesc)->GetGuid(), NewActorDesc);
+
+	return NewActorDesc->Get();
+}
+
+const FWorldPartitionActorDesc* FActorDescList::GetActorDesc(const FGuid& Guid) const
+{
+	const TUniquePtr<FWorldPartitionActorDesc>* const * ActorDesc = Actors.Find(Guid);
+	return ActorDesc ? (*ActorDesc)->Get() : nullptr;
+}
+
+FWorldPartitionActorDesc* FActorDescList::GetActorDesc(const FGuid& Guid)
+{
+	TUniquePtr<FWorldPartitionActorDesc>** ActorDesc = Actors.Find(Guid);
+	return ActorDesc ? (*ActorDesc)->Get() : nullptr;
+}
+
+const FWorldPartitionActorDesc& FActorDescList::GetActorDescChecked(const FGuid& Guid) const
+{
+	const TUniquePtr<FWorldPartitionActorDesc>* const ActorDesc = Actors.FindChecked(Guid);
+	return *ActorDesc->Get();
+}
+
+FWorldPartitionActorDesc& FActorDescList::GetActorDescChecked(const FGuid& Guid)
+{
+	TUniquePtr<FWorldPartitionActorDesc>* ActorDesc = Actors.FindChecked(Guid);
+	return *ActorDesc->Get();
+}
+
+const FWorldPartitionActorDesc* FActorDescList::GetActorDesc(const FString& PackageName) const
+{
+	const FName PackageFName(*PackageName);
+	for (const TUniquePtr<FWorldPartitionActorDesc>& ActorDescPtr : ActorDescList)
+	{
+		if (ActorDescPtr->GetActorPackage() == PackageFName)
+		{
+			return ActorDescPtr.Get();
+		}
+	}
+
+	return nullptr;
+}
+
+void FActorDescList::Empty()
+{
+	Actors.Empty();
+	ActorDescList.Empty();
+}
+#endif
\ No newline at end of file
diff --git a/Engine/Source/Runtime/Engine/Private/WorldPartition/WorldPartition.cpp b/Engine/Source/Runtime/Engine/Private/WorldPartition/WorldPartition.cpp
--- a/Engine/Source/Runtime/Engine/Private/WorldPartition/WorldPartition.cpp
+++ b/Engine/Source/Runtime/Engine/Private/WorldPartition/WorldPartition.cpp
@@ -142,6 +142,15 @@ UWorldPartition::UWorldPartition(const FObjectInitializer& ObjectInitializer)
 #if WITH_EDITOR
 void UWorldPartition::OnPreBeginPIE(bool bStartSimulate)
 {
+	// Gather all modified/newly created actors for PIE so they can be duplicated when streaming their cell
+	for (AActor* Actor : World->PersistentLevel->Actors)
+	{
+		if (Actor && !Actor->IsPendingKill() && Actor->GetPackage()->IsDirty())
+		{
+			World->PersistentLevel->ActorsModifiedForPIE.Add(Actor->GetFName(), Actor);
+		}
+	}
+
 	check(IsMainWorldPartition());
 	GenerateStreaming(EWorldPartitionStreamingMode::PIE);
 	RuntimeHash->OnPreBeginPIE();
@@ -152,6 +161,8 @@ void UWorldPartition::OnEndPIE(bool bStartSimulate)
 	check(IsMainWorldPartition());
 	FlushStreaming();
 	RuntimeHash->OnEndPIE();
+
+	World->PersistentLevel->ActorsModifiedForPIE.Empty();
 }
 
 FName UWorldPartition::GetWorldPartitionEditorName()
diff --git a/Engine/Source/Runtime/Engine/Private/WorldPartition/WorldPartitionActorCluster.cpp b/Engine/Source/Runtime/Engine/Private/WorldPartition/WorldPartitionActorCluster.cpp
--- a/Engine/Source/Runtime/Engine/Private/WorldPartition/WorldPartitionActorCluster.cpp
+++ b/Engine/Source/Runtime/Engine/Private/WorldPartition/WorldPartitionActorCluster.cpp
@@ -174,7 +174,7 @@ bool FActorInstance::ShouldStripFromStreaming() const
 		return true;
 	}
 
-	const FWorldPartitionActorDescView ActorDescView = GetActorDescView();
+	const FWorldPartitionActorDescView& ActorDescView = GetActorDescView();
 	return ActorDescView.GetActorIsEditorOnly();
 }
 
diff --git a/Engine/Source/Runtime/Engine/Private/WorldPartition/WorldPartitionActorDescView.cpp b/Engine/Source/Runtime/Engine/Private/WorldPartition/WorldPartitionActorDescView.cpp
--- a/Engine/Source/Runtime/Engine/Private/WorldPartition/WorldPartitionActorDescView.cpp
+++ b/Engine/Source/Runtime/Engine/Private/WorldPartition/WorldPartitionActorDescView.cpp
@@ -95,6 +95,11 @@ const TArray<FGuid>& FWorldPartitionActorDescView::GetReferences() const
 	return ActorDesc->GetReferences();
 }
 
+uint32 FWorldPartitionActorDescView::GetTag() const
+{
+	return ActorDesc->Tag;
+}
+
 bool FWorldPartitionActorDescView::GetContainerInstance(const UActorDescContainer*& OutLevelContainer, FTransform& OutLevelTransform, EContainerClusterMode& OutClusterMode) const
 {
 	return ActorDesc->GetContainerInstance(OutLevelContainer, OutLevelTransform, OutClusterMode);
diff --git a/Engine/Source/Runtime/Engine/Private/WorldPartition/WorldPartitionLevelHelper.cpp b/Engine/Source/Runtime/Engine/Private/WorldPartition/WorldPartitionLevelHelper.cpp
--- a/Engine/Source/Runtime/Engine/Private/WorldPartition/WorldPartitionLevelHelper.cpp
+++ b/Engine/Source/Runtime/Engine/Private/WorldPartition/WorldPartitionLevelHelper.cpp
@@ -112,7 +112,7 @@ ULevel* FWorldPartitionLevelHelper::CreateEmptyLevelForRuntimeCell(const UWorld*
 	return NewLevel;
 }
 
-bool FWorldPartitionLevelHelper::LoadActors(ULevel* InDestLevel, TArrayView<FWorldPartitionRuntimeCellObjectMapping> InActorPackages, FWorldPartitionPackageCache& InPackageCache, FWorldPartitionLevelHelper::FOnLoadActorsCompleted InCompletionCallback, bool bLoadForPIE, bool bLoadAsync /*=true*/, FLinkerInstancingContext* InOutInstancingContext /*=nullptr*/)
+bool FWorldPartitionLevelHelper::LoadActors(ULevel* InDestLevel, TArrayView<FWorldPartitionRuntimeCellObjectMapping> InActorPackages, FWorldPartitionPackageCache& InPackageCache, TFunction<void(bool)> InCompletionCallback, bool bLoadForPIE, bool bLoadAsync /*=true*/, FLinkerInstancingContext* InOutInstancingContext /*=nullptr*/)
 {
 	UPackage* DestPackage = InDestLevel ? InDestLevel->GetPackage() : nullptr;
 	FString ShortLevelPackageName = DestPackage? FPackageName::GetShortName(DestPackage->GetFName()) : FString();
@@ -191,7 +191,7 @@ bool FWorldPartitionLevelHelper::LoadActors(ULevel* InDestLevel, TArrayView<FWor
 
 				if (!LoadProgress->NumPendingLoadRequests)
 				{
-					InCompletionCallback.Execute(!LoadProgress->NumFailedLoadedRequests);
+					InCompletionCallback(!LoadProgress->NumFailedLoadedRequests);
 				}
 			});
 
@@ -274,7 +274,7 @@ bool FWorldPartitionLevelHelper::LoadActors(ULevel* InDestLevel, TArrayView<FWor
 
 			if (!LoadProgress->NumPendingLoadRequests)
 			{
-				InCompletionCallback.Execute(!LoadProgress->NumFailedLoadedRequests);
+				InCompletionCallback(!LoadProgress->NumFailedLoadedRequests);
 			}
 		});
 
diff --git a/Engine/Source/Runtime/Engine/Private/WorldPartition/WorldPartitionLevelStreamingDynamic.cpp b/Engine/Source/Runtime/Engine/Private/WorldPartition/WorldPartitionLevelStreamingDynamic.cpp
--- a/Engine/Source/Runtime/Engine/Private/WorldPartition/WorldPartitionLevelStreamingDynamic.cpp
+++ b/Engine/Source/Runtime/Engine/Private/WorldPartition/WorldPartitionLevelStreamingDynamic.cpp
@@ -222,28 +222,94 @@ bool UWorldPartitionLevelStreamingDynamic::IssueLoadRequests()
 	check(RuntimeLevel);
 	check(!bLoadRequestInProgress);
 
-	TArray<FString> ActorPackageInstanceNames;
-	ActorPackageInstanceNames.Reserve(ChildPackages.Num());
-	UPackage* RuntimePackage = RuntimeLevel->GetPackage();
-
 	bLoadRequestInProgress = true;
 	FLinkerInstancingContext InstancingContext;
+	UPackage* RuntimePackage = RuntimeLevel->GetPackage();
 	InstancingContext.AddMapping(OriginalLevelPackageName, RuntimePackage->GetFName());
-	FWorldPartitionLevelHelper::FOnLoadActorsCompleted CompletionCallback = FWorldPartitionLevelHelper::FOnLoadActorsCompleted::CreateLambda([this](bool bSucceeded)
+
+	TArray<AActor*> ActorsToDuplicate;
+	ActorsToDuplicate.Reserve(ChildPackages.Num());
+
+	TArray<FWorldPartitionRuntimeCellObjectMapping> ChildPackagesToLoad;
+	ChildPackagesToLoad.Reserve(ChildPackages.Num());
+
+	UWorld* World = GetWorld();
+	for (FWorldPartitionRuntimeCellObjectMapping& ChildPackage : ChildPackages)
+	{
+		if (ChildPackage.ContainerID == 0)
 		{
-			if (bSucceeded)
+			bool bNeedDup = false;
+
+			FString SubObjectName;
+			FString SubObjectContext;	
+			if (ChildPackage.LoadedPath.ToString().Split(TEXT("."), &SubObjectContext, &SubObjectName, ESearchCase::IgnoreCase, ESearchDir::FromEnd))
 			{
-				check(bLoadRequestInProgress);
-				bLoadRequestInProgress = false;
-				FinalizeRuntimeLevel();
+				if (AActor* ActorModifiedForPIE = World->PersistentLevel->ActorsModifiedForPIE.FindRef(*SubObjectName))
+				{
+					ActorsToDuplicate.Add(ActorModifiedForPIE);
+					bNeedDup = true;
+				}
 			}
-			else
+
+			if (!bNeedDup)
 			{
-				UE_LOG(LogLevelStreaming, Fatal, TEXT("UWorldPartitionLevelStreamingDynamic::IssueLoadRequests failed %s"), *GetWorldAssetPackageName());
+				ChildPackagesToLoad.Add(ChildPackage);
 			}
-		});
+		}
+	}
+
+	// Duplicate unsaved actors
+	if (ActorsToDuplicate.Num())
+	{
+		// Create an actor container to make sure duplicated actors will share an outer to properly remap inter-actors references
+		UActorContainer* ActorContainer = NewObject<UActorContainer>(World->PersistentLevel);
+
+		for (AActor* ActorToDuplicate : ActorsToDuplicate)
+		{
+			ActorToDuplicate->UObject::Rename(nullptr, ActorContainer);
+			ActorContainer->Actors.Add(ActorToDuplicate);
+		}
+
+		FObjectDuplicationParameters Parameters(ActorContainer, RuntimeLevel);
+		Parameters.DestClass = ActorContainer->GetClass();
+		Parameters.FlagMask = RF_AllFlags & ~(RF_MarkAsRootSet | RF_MarkAsNative | RF_HasExternalPackage);
+		Parameters.InternalFlagMask = EInternalObjectFlags::AllFlags;
+		Parameters.DuplicateMode = EDuplicateMode::PIE;
+		Parameters.PortFlags = PPF_DuplicateForPIE;
+		Parameters.DuplicationSeed.Add(World->PersistentLevel, RuntimeLevel);
 
-	FWorldPartitionLevelHelper::LoadActors(RuntimeLevel, ChildPackages, PackageCache, CompletionCallback, /*bLoadForPie=*/true, /*bLoadAsync=*/true, &InstancingContext);
+		UActorContainer* ActorContainerDup = (UActorContainer*)StaticDuplicateObjectEx(Parameters);
+
+		// Add the duplicated actors to the corresponding cell level
+		for (AActor* Actor : ActorContainerDup->Actors)
+		{
+			Actor->Rename(nullptr, RuntimeLevel);
+		}
+
+		// Bring back actors to the persistent level
+		for (AActor* ActorToDuplicate : ActorsToDuplicate)
+		{
+			ActorToDuplicate->UObject::Rename(nullptr, World->PersistentLevel);
+		}
+
+		ActorContainer->MarkPendingKill();
+		ActorContainerDup->MarkPendingKill();
+	}
+
+	// Load saved actors
+	FWorldPartitionLevelHelper::LoadActors(RuntimeLevel, ChildPackagesToLoad, PackageCache, [this](bool bSucceeded)
+	{
+		if (bSucceeded)
+		{
+			check(bLoadRequestInProgress);
+			bLoadRequestInProgress = false;
+			FinalizeRuntimeLevel();
+		}
+		else
+		{
+			UE_LOG(LogLevelStreaming, Fatal, TEXT("UWorldPartitionLevelStreamingDynamic::IssueLoadRequests failed %s"), *GetWorldAssetPackageName());
+		}
+	}, /*bLoadForPie=*/true, /*bLoadAsync=*/true, &InstancingContext);
 
 	return bLoadRequestInProgress;
 }
diff --git a/Engine/Source/Runtime/Engine/Private/WorldPartition/WorldPartitionRuntimeHash.cpp b/Engine/Source/Runtime/Engine/Private/WorldPartition/WorldPartitionRuntimeHash.cpp
--- a/Engine/Source/Runtime/Engine/Private/WorldPartition/WorldPartitionRuntimeHash.cpp
+++ b/Engine/Source/Runtime/Engine/Private/WorldPartition/WorldPartitionRuntimeHash.cpp
@@ -8,7 +8,7 @@
 #include "ActorReferencesUtils.h"
 
 UWorldPartitionRuntimeHash::UWorldPartitionRuntimeHash(const FObjectInitializer& ObjectInitializer)
-: Super(ObjectInitializer)
+	: Super(ObjectInitializer)
 {}
 
 #if WITH_EDITOR
@@ -26,6 +26,8 @@ void UWorldPartitionRuntimeHash::OnEndPIE()
 
 	// Release references (will unload actors that were not already loaded in the Editor)
 	AlwaysLoadedActorsForPIE.Empty();
+
+	ModifiedActorDescListForPIE.Empty();
 }
 
 void UWorldPartitionRuntimeHash::ForceExternalActorLevelReference(bool bForceExternalActorLevelReferenceForPIE)
@@ -45,9 +47,40 @@ void UWorldPartitionRuntimeHash::CreateActorDescViewMap(const UActorDescContaine
 	OutActorDescViewMap.Empty();
 	for (UActorDescContainer::TConstIterator<> ActorDescIt(Container); ActorDescIt; ++ActorDescIt)
 	{
+		if (AActor* Actor = ActorDescIt->GetActor())
+		{
+			// Don't include deleted, unsaved actors
+			if (Actor->IsPendingKill())
+			{
+				continue;
+			}
+
+			// Handle dirty, unsaved actors
+			if (Actor->GetPackage()->IsDirty())
+			{
+				FWorldPartitionActorDesc* ActorDesc = ModifiedActorDescListForPIE.AddActor(Actor);
+				OutActorDescViewMap.Emplace(ActorDesc->GetGuid(), ActorDesc);
+				continue;
+			}
+		}
+
+		// Normal, non-dirty actors
 		OutActorDescViewMap.Emplace(ActorDescIt->GetGuid(), *ActorDescIt);
 	}
 
+	// Append new unsaved actors for the persistent level
+	if (Container->GetContainerPackage() == GetWorld()->PersistentLevel->GetPackage()->GetLoadedPath().GetPackageFName())
+	{
+		for (auto& ActorPair : GetWorld()->PersistentLevel->ActorsModifiedForPIE)
+		{
+			if (ActorPair.Value->GetPackage()->HasAnyPackageFlags(PKG_NewlyCreated))
+			{
+				FWorldPartitionActorDesc* ActorDesc = ModifiedActorDescListForPIE.AddActor(ActorPair.Value);
+				OutActorDescViewMap.Emplace(ActorDesc->GetGuid(), ActorDesc);
+			}
+		}
+	}
+
 	// Set HLOD parents into actor desc views
 	for (UActorDescContainer::TConstIterator<AWorldPartitionHLOD> HLODIterator(Container); HLODIterator; ++HLODIterator)
 	{
diff --git a/Engine/Source/Runtime/Engine/Private/WorldPartition/WorldPartitionRuntimeLevelStreamingCell.cpp b/Engine/Source/Runtime/Engine/Private/WorldPartition/WorldPartitionRuntimeLevelStreamingCell.cpp
--- a/Engine/Source/Runtime/Engine/Private/WorldPartition/WorldPartitionRuntimeLevelStreamingCell.cpp
+++ b/Engine/Source/Runtime/Engine/Private/WorldPartition/WorldPartitionRuntimeLevelStreamingCell.cpp
@@ -50,11 +50,10 @@ void UWorldPartitionRuntimeLevelStreamingCell::SetIsAlwaysLoaded(bool bInIsAlway
 
 #if WITH_EDITOR
 
-void UWorldPartitionRuntimeLevelStreamingCell::AddActorToCell(const FGuid& InActorGuid, uint32 InContainerID, const FTransform& InContainerTransform, const UActorDescContainer* InContainer)
+void UWorldPartitionRuntimeLevelStreamingCell::AddActorToCell(const FWorldPartitionActorDescView& ActorDescView, uint32 InContainerID, const FTransform& InContainerTransform, const UActorDescContainer* InContainer)
 {
-	const FWorldPartitionActorDesc* ActorDesc = InContainer->GetActorDesc(InActorGuid);
-	check(!ActorDesc->GetActorIsEditorOnly());
-	Packages.Emplace(ActorDesc->GetActorPackage(), ActorDesc->GetActorPath(), InContainerID, InContainerTransform, InContainer->GetContainerPackage());
+	check(!ActorDescView.GetActorIsEditorOnly());
+	Packages.Emplace(ActorDescView.GetActorPackage(), ActorDescView.GetActorPath(), InContainerID, InContainerTransform, InContainer->GetContainerPackage());
 }
 
 ULevelStreaming* UWorldPartitionRuntimeLevelStreamingCell::CreateLevelStreaming(const FString& InPackageName) const
@@ -98,9 +97,7 @@ ULevelStreaming* UWorldPartitionRuntimeLevelStreamingCell::CreateLevelStreaming(
 void UWorldPartitionRuntimeLevelStreamingCell::LoadActorsForCook()
 {
 	FWorldPartitionPackageCache PackageCache;
-	FWorldPartitionLevelHelper::FOnLoadActorsCompleted CompletionCallback = FWorldPartitionLevelHelper::FOnLoadActorsCompleted::CreateLambda([](bool) {});
-	bool bResult = FWorldPartitionLevelHelper::LoadActors(nullptr, Packages, PackageCache, CompletionCallback, /*bLoadForPIE=*/false, /*bLoadAsync=*/false);
-	check(bResult);
+	verify(FWorldPartitionLevelHelper::LoadActors(nullptr, Packages, PackageCache, [](bool){}, /*bLoadForPIE=*/false, /*bLoadAsync=*/false));
 }
 
 void UWorldPartitionRuntimeLevelStreamingCell::MoveAlwaysLoadedContentToPersistentLevel()
diff --git a/Engine/Source/Runtime/Engine/Private/WorldPartition/WorldPartitionRuntimeSpatialHash.cpp b/Engine/Source/Runtime/Engine/Private/WorldPartition/WorldPartitionRuntimeSpatialHash.cpp
--- a/Engine/Source/Runtime/Engine/Private/WorldPartition/WorldPartitionRuntimeSpatialHash.cpp
+++ b/Engine/Source/Runtime/Engine/Private/WorldPartition/WorldPartitionRuntimeSpatialHash.cpp
@@ -698,7 +698,7 @@ bool UWorldPartitionRuntimeSpatialHash::CreateStreamingGrid(const FSpatialHashRu
 					{
 						ReferencedHLODActors.Add(ActorDescView.GetHLODParent());
 					}
-					StreamingCell->AddActorToCell(ActorInstance.Actor, ActorInstance.ContainerInstance->ID, ActorInstance.ContainerInstance->Transform, ActorInstance.ContainerInstance->Container);
+					StreamingCell->AddActorToCell(ActorDescView, ActorInstance.ContainerInstance->ID, ActorInstance.ContainerInstance->Transform, ActorInstance.ContainerInstance->Container);
 					UE_LOG(LogWorldPartitionRuntimeSpatialHash, Verbose, TEXT("  Actor : %s (%s) (Container %08x) Origin(%s)"), *(ActorDescView.GetActorPath().ToString()), *ActorDescView.GetGuid().ToString(EGuidFormats::UniqueObjectGuid), ActorInstance.ContainerInstance->ID, *FVector2D(ActorInstance.GetOrigin()).ToString());
 				}
 
diff --git a/Engine/Source/Runtime/Engine/Public/WorldPartition/ActorDescContainer.h b/Engine/Source/Runtime/Engine/Public/WorldPartition/ActorDescContainer.h
--- a/Engine/Source/Runtime/Engine/Public/WorldPartition/ActorDescContainer.h
+++ b/Engine/Source/Runtime/Engine/Public/WorldPartition/ActorDescContainer.h
@@ -5,12 +5,11 @@
 #include "UObject/ObjectMacros.h"
 #include "UObject/Object.h"
 #include "GameFramework/Actor.h"
-#include "WorldPartition/WorldPartitionActorDesc.h"
-#include "WorldPartition/WorldPartitionActorDescType.h"
+#include "WorldPartition/ActorDescList.h"
 #include "ActorDescContainer.generated.h"
 
 UCLASS()
-class ENGINE_API UActorDescContainer : public UObject
+class ENGINE_API UActorDescContainer : public UObject, public FActorDescList
 {
 	GENERATED_UCLASS_BODY()
 
@@ -29,20 +28,12 @@ public:
 	virtual void OnAssetRemoved(const FAssetData& InAssetData);
 	virtual void OnAssetUpdated(const FAssetData& InAssetData);
 
-	FWorldPartitionActorDesc* GetActorDesc(const FGuid& Guid);
-	const FWorldPartitionActorDesc* GetActorDesc(const FGuid& Guid) const;
-
-	FWorldPartitionActorDesc& GetActorDescChecked(const FGuid& Guid);
-	const FWorldPartitionActorDesc& GetActorDescChecked(const FGuid& Guid) const;
-
-	const FWorldPartitionActorDesc* GetActorDesc(const FString& PackageName) const;
-
-	int32 GetActorDescCount() const { return Actors.Num(); }
 	FName GetContainerPackage() const { return ContainerPackageName; }
 #endif
 
 	UPROPERTY(Transient)
 	TObjectPtr<UWorld> World;
+
 protected:
 	//~ Begin UObject Interface
 	virtual void BeginDestroy() override;
@@ -67,119 +58,6 @@ protected:
 	bool bContainerInitialized;
 	bool bIgnoreAssetRegistryEvents;
 
-	TChunkedArray<TUniquePtr<FWorldPartitionActorDesc>> ActorDescList;
-
 	FName ContainerPackageName;
-
-private:
-	TMap<FGuid, TUniquePtr<FWorldPartitionActorDesc>*> Actors;
-
-public:
-	template<bool bConst, class ActorType>
-	class TBaseIterator
-	{
-		static_assert(TIsDerivedFrom<ActorType, AActor>::IsDerived, "Type is not derived from AActor.");
-
-	protected:
-		typedef TMap<FGuid, TUniquePtr<FWorldPartitionActorDesc>*> MapType;
-		typedef typename FWorldPartitionActorDescType<ActorType>::Type ValueType;
-		typedef typename TChooseClass<bConst, MapType::TConstIterator, MapType::TIterator>::Result IteratorType;
-		typedef typename TChooseClass<bConst, const UActorDescContainer*, UActorDescContainer*>::Result ContainerType;		
-		typedef typename TChooseClass<bConst, const ValueType*, ValueType*>::Result ReturnType;
-
-	public:
-		TBaseIterator(ContainerType InActorDescContainer, UClass* InActorClass)
-			: ActorsIterator(InActorDescContainer->Actors)
-			, ActorClass(InActorClass)
-		{
-			check(ActorClass->IsChildOf(ActorType::StaticClass()));
-
-			if (ShouldSkip())
-			{
-				operator++();
-			}
-		}
-
-		/**
-		 * Iterates to next suitable actor desc.
-		 */
-		void operator++()
-		{
-			do
-			{
-				++ActorsIterator;
-			} while (ShouldSkip());
-		}
-
-		/**
-		 * Returns the current suitable actor desc pointed at by the Iterator
-		 *
-		 * @return	Current suitable actor desc
-		 */
-		FORCEINLINE ReturnType operator*() const
-		{
-			return StaticCast<ReturnType>(ActorsIterator->Value->Get());
-		}
-
-		/**
-		 * Returns the current suitable actor desc pointed at by the Iterator
-		 *
-		 * @return	Current suitable actor desc
-		 */
-		FORCEINLINE ReturnType operator->() const
-		{
-			return StaticCast<ReturnType>(ActorsIterator->Value->Get());
-		}
-		/**
-		 * Returns whether the iterator has reached the end and no longer points
-		 * to a suitable actor desc.
-		 *
-		 * @return true if iterator points to a suitable actor desc, false if it has reached the end
-		 */
-		FORCEINLINE explicit operator bool() const
-		{
-			return (bool)ActorsIterator;
-		}
-
-	protected:
-		/**
-		 * Determines whether the iterator currently points to a valid actor desc or not.
-		 * @return true if we should skip the actor desc
-		 */
-		FORCEINLINE bool ShouldSkip() const
-		{
-			if (!ActorsIterator)
-			{
-				return false;
-			}
-
-			return !ActorsIterator->Value->Get()->GetActorClass()->IsChildOf(ActorClass);
-		}
-
-		IteratorType ActorsIterator;
-		UClass* ActorClass;
-	};
-
-	template <class ActorType = AActor>
-	class TIterator : public TBaseIterator<false, ActorType>
-	{
-		typedef TBaseIterator<false, ActorType> BaseType;
-
-	public:
-		TIterator(typename BaseType::ContainerType InActorDescContainer, UClass* InActorClass = nullptr)
-			: BaseType(InActorDescContainer, InActorClass ? InActorClass : ActorType::StaticClass())
-		{}
-	};
-
-	template <class ActorType = AActor>
-	class TConstIterator : public TBaseIterator<true, ActorType>
-	{
-		typedef TBaseIterator<true, ActorType> BaseType;
-
-	public:
-		TConstIterator(typename BaseType::ContainerType InActorDescContainer, UClass* InActorClass = nullptr)
-			: BaseType(InActorDescContainer, InActorClass ? InActorClass : ActorType::StaticClass())
-		{}
-	};
 #endif
 };
\ No newline at end of file
diff --git a/Engine/Source/Runtime/Engine/Public/WorldPartition/ActorDescList.h b/Engine/Source/Runtime/Engine/Public/WorldPartition/ActorDescList.h
new file mode 100644
--- /dev/null
+++ b/Engine/Source/Runtime/Engine/Public/WorldPartition/ActorDescList.h
@@ -0,0 +1,144 @@
+// Copyright Epic Games, Inc. All Rights Reserved.
+#pragma once
+
+#include "CoreMinimal.h"
+#include "GameFramework/Actor.h"
+#include "WorldPartition/WorldPartitionActorDesc.h"
+#include "WorldPartition/WorldPartitionActorDescType.h"
+
+class ENGINE_API FActorDescList
+{
+#if WITH_EDITOR
+public:
+	FActorDescList() {}
+
+	// Non-copyable
+	FActorDescList(const FActorDescList&) = delete;
+	FActorDescList& operator=(const FActorDescList&) = delete;
+
+	FWorldPartitionActorDesc* AddActor(const AActor* InActor);
+
+	FWorldPartitionActorDesc* GetActorDesc(const FGuid& Guid);
+	const FWorldPartitionActorDesc* GetActorDesc(const FGuid& Guid) const;
+
+	FWorldPartitionActorDesc& GetActorDescChecked(const FGuid& Guid);
+	const FWorldPartitionActorDesc& GetActorDescChecked(const FGuid& Guid) const;
+
+	const FWorldPartitionActorDesc* GetActorDesc(const FString& PackageName) const;
+
+	int32 GetActorDescCount() const { return Actors.Num(); }
+
+	void Empty();
+
+	template<bool bConst, class ActorType>
+	class TBaseIterator
+	{
+		static_assert(TIsDerivedFrom<ActorType, AActor>::IsDerived, "Type is not derived from AActor.");
+
+	protected:
+		typedef TMap<FGuid, TUniquePtr<FWorldPartitionActorDesc>*> MapType;
+		typedef typename FWorldPartitionActorDescType<ActorType>::Type ValueType;
+		typedef typename TChooseClass<bConst, MapType::TConstIterator, MapType::TIterator>::Result IteratorType;
+		typedef typename TChooseClass<bConst, const FActorDescList*, FActorDescList*>::Result ListType;		
+		typedef typename TChooseClass<bConst, const ValueType*, ValueType*>::Result ReturnType;
+
+	public:
+		TBaseIterator(ListType InActorDescList, UClass* InActorClass)
+			: ActorsIterator(InActorDescList->Actors)
+			, ActorClass(InActorClass)
+		{
+			check(ActorClass->IsChildOf(ActorType::StaticClass()));
+
+			if (ShouldSkip())
+			{
+				operator++();
+			}
+		}
+
+		/**
+		 * Iterates to next suitable actor desc.
+		 */
+		void operator++()
+		{
+			do
+			{
+				++ActorsIterator;
+			} while (ShouldSkip());
+		}
+
+		/**
+		 * Returns the current suitable actor desc pointed at by the Iterator
+		 *
+		 * @return	Current suitable actor desc
+		 */
+		FORCEINLINE ReturnType operator*() const
+		{
+			return StaticCast<ReturnType>(ActorsIterator->Value->Get());
+		}
+
+		/**
+		 * Returns the current suitable actor desc pointed at by the Iterator
+		 *
+		 * @return	Current suitable actor desc
+		 */
+		FORCEINLINE ReturnType operator->() const
+		{
+			return StaticCast<ReturnType>(ActorsIterator->Value->Get());
+		}
+		/**
+		 * Returns whether the iterator has reached the end and no longer points
+		 * to a suitable actor desc.
+		 *
+		 * @return true if iterator points to a suitable actor desc, false if it has reached the end
+		 */
+		FORCEINLINE explicit operator bool() const
+		{
+			return (bool)ActorsIterator;
+		}
+
+	protected:
+		/**
+		 * Determines whether the iterator currently points to a valid actor desc or not.
+		 * @return true if we should skip the actor desc
+		 */
+		FORCEINLINE bool ShouldSkip() const
+		{
+			if (!ActorsIterator)
+			{
+				return false;
+			}
+
+			return !ActorsIterator->Value->Get()->GetActorClass()->IsChildOf(ActorClass);
+		}
+
+		IteratorType ActorsIterator;
+		UClass* ActorClass;
+	};
+
+	template <class ActorType = AActor>
+	class TIterator : public TBaseIterator<false, ActorType>
+	{
+		typedef TBaseIterator<false, ActorType> BaseType;
+
+	public:
+		TIterator(typename BaseType::ListType InActorDescList, UClass* InActorClass = nullptr)
+			: BaseType(InActorDescList, InActorClass ? InActorClass : ActorType::StaticClass())
+		{}
+	};
+
+	template <class ActorType = AActor>
+	class TConstIterator : public TBaseIterator<true, ActorType>
+	{
+		typedef TBaseIterator<true, ActorType> BaseType;
+
+	public:
+		TConstIterator(typename BaseType::ListType InActorDescList, UClass* InActorClass = nullptr)
+			: BaseType(InActorDescList, InActorClass ? InActorClass : ActorType::StaticClass())
+		{}
+	};
+
+protected:
+	TChunkedArray<TUniquePtr<FWorldPartitionActorDesc>> ActorDescList;
+	TMap<FGuid, TUniquePtr<FWorldPartitionActorDesc>*> Actors;
+#endif
+};
\ No newline at end of file
diff --git a/Engine/Source/Runtime/Engine/Public/WorldPartition/WorldPartitionActorDesc.h b/Engine/Source/Runtime/Engine/Public/WorldPartition/WorldPartitionActorDesc.h
--- a/Engine/Source/Runtime/Engine/Public/WorldPartition/WorldPartitionActorDesc.h
+++ b/Engine/Source/Runtime/Engine/Public/WorldPartition/WorldPartitionActorDesc.h
@@ -75,6 +75,16 @@ public:
 
 	virtual bool GetContainerInstance(const UActorDescContainer*& OutLevelContainer, FTransform& OutLevelTransform, EContainerClusterMode& OutClusterMode) const { return false; }
 
+	bool operator==(const FWorldPartitionActorDesc& Other) const
+	{
+		return Guid == Other.Guid;
+	}
+
+	friend uint32 GetTypeHash(const FWorldPartitionActorDesc& Key)
+	{
+		return GetTypeHash(Key.Guid);
+	}
+
 protected:
 	inline uint32 IncSoftRefCount() const
 	{
diff --git a/Engine/Source/Runtime/Engine/Public/WorldPartition/WorldPartitionActorDescView.h b/Engine/Source/Runtime/Engine/Public/WorldPartition/WorldPartitionActorDescView.h
--- a/Engine/Source/Runtime/Engine/Public/WorldPartition/WorldPartitionActorDescView.h
+++ b/Engine/Source/Runtime/Engine/Public/WorldPartition/WorldPartitionActorDescView.h
@@ -38,12 +38,23 @@ public:
 	FName GetActorLabel() const;
 	FBox GetBounds() const;
 	const TArray<FGuid>& GetReferences() const;
+	uint32 GetTag() const;
 
 	bool GetContainerInstance(const UActorDescContainer*& OutLevelContainer, FTransform& OutLevelTransform, EContainerClusterMode& OutClusterMode) const;
 
 	FGuid GetHLODParent() const;
 	void SetHLODParent(const FGuid& InHLODParent);
 
+	bool operator==(const FWorldPartitionActorDescView& Other) const
+	{
+		return GetGuid() == Other.GetGuid();
+	}
+
+	friend uint32 GetTypeHash(const FWorldPartitionActorDescView& Key)
+	{
+		return GetTypeHash(Key.GetGuid());
+	}
+
 protected:
 	const FWorldPartitionActorDesc* ActorDesc;
 	EActorGridPlacement EffectiveGridPlacement;
diff --git a/Engine/Source/Runtime/Engine/Public/WorldPartition/WorldPartitionLevelHelper.h b/Engine/Source/Runtime/Engine/Public/WorldPartition/WorldPartitionLevelHelper.h
--- a/Engine/Source/Runtime/Engine/Public/WorldPartition/WorldPartitionLevelHelper.h
+++ b/Engine/Source/Runtime/Engine/Public/WorldPartition/WorldPartitionLevelHelper.h
@@ -22,9 +22,7 @@ public:
 	static ULevel* CreateEmptyLevelForRuntimeCell(const UWorld* InWorld, const FString& InWorldAssetName, UPackage* DestPackage = nullptr);
 	static void MoveExternalActorsToLevel(const TArray<FWorldPartitionRuntimeCellObjectMapping>& InChildPackages, ULevel* InLevel);
 	
-	DECLARE_DELEGATE_OneParam(FOnLoadActorsCompleted, bool);
-
-	static bool LoadActors(ULevel* InDestLevel, TArrayView<FWorldPartitionRuntimeCellObjectMapping> InActorPackages, FWorldPartitionPackageCache& InPackageCache, FOnLoadActorsCompleted InCompletionCallback, bool bLoadForPie, bool bLoadAsync = true, FLinkerInstancingContext* InOutInstancingContext = nullptr);
+	static bool LoadActors(ULevel* InDestLevel, TArrayView<FWorldPartitionRuntimeCellObjectMapping> InActorPackages, FWorldPartitionPackageCache& InPackageCache, TFunction<void(bool)> InCompletionCallback, bool bLoadForPie, bool bLoadAsync = true, FLinkerInstancingContext* InOutInstancingContext = nullptr);
 private:
 	static UWorld::InitializationValues GetWorldInitializationValues();
 };
diff --git a/Engine/Source/Runtime/Engine/Public/WorldPartition/WorldPartitionLevelStreamingDynamic.h b/Engine/Source/Runtime/Engine/Public/WorldPartition/WorldPartitionLevelStreamingDynamic.h
--- a/Engine/Source/Runtime/Engine/Public/WorldPartition/WorldPartitionLevelStreamingDynamic.h
+++ b/Engine/Source/Runtime/Engine/Public/WorldPartition/WorldPartitionLevelStreamingDynamic.h
@@ -64,4 +64,15 @@ private:
 
 	UPROPERTY()
 	TWeakObjectPtr<UWorldPartition> OuterWorldPartition;
-};
\ No newline at end of file
+};
+
+// Actor container class used to duplicate actors during cells streaming
+UCLASS(Transient)
+class UActorContainer : public UObject
+{
+	GENERATED_BODY()
+
+public:
+	UPROPERTY()
+	TArray<AActor*> Actors;
+};
diff --git a/Engine/Source/Runtime/Engine/Public/WorldPartition/WorldPartitionRuntimeCell.h b/Engine/Source/Runtime/Engine/Public/WorldPartition/WorldPartitionRuntimeCell.h
--- a/Engine/Source/Runtime/Engine/Public/WorldPartition/WorldPartitionRuntimeCell.h
+++ b/Engine/Source/Runtime/Engine/Public/WorldPartition/WorldPartitionRuntimeCell.h
@@ -4,6 +4,7 @@
 #include "CoreMinimal.h"
 #include "Templates/SubclassOf.h"
 #include "WorldPartition/DataLayer/DataLayer.h"
+#include "WorldPartition/WorldPartitionActorDescView.h"
 #include "WorldPartitionRuntimeCell.generated.h"
 
 USTRUCT()
@@ -105,7 +106,7 @@ class UWorldPartitionRuntimeCell : public UObject
 #if WITH_EDITOR
 	void SetDataLayers(const TArray<const UDataLayer*> InDataLayers);
 	void AddCellData(const UWorldPartitionRuntimeCellData* InCellData);
-	virtual void AddActorToCell(const FGuid& InActorGuid, uint32 InContainerID, const FTransform& InContainerTransform, const UActorDescContainer* InContainer) PURE_VIRTUAL(UWorldPartitionRuntimeCell::AddActorToCell,);
+	virtual void AddActorToCell(const FWorldPartitionActorDescView& ActorDescView, uint32 InContainerID, const FTransform& InContainerTransform, const UActorDescContainer* InContainer) PURE_VIRTUAL(UWorldPartitionRuntimeCell::AddActorToCell,);
 	virtual int32 GetActorCount() const PURE_VIRTUAL(UWorldPartitionRuntimeCell::GetActorCount, return 0;);
 
 	// Cook methods
diff --git a/Engine/Source/Runtime/Engine/Public/WorldPartition/WorldPartitionRuntimeHash.h b/Engine/Source/Runtime/Engine/Public/WorldPartition/WorldPartitionRuntimeHash.h
--- a/Engine/Source/Runtime/Engine/Public/WorldPartition/WorldPartitionRuntimeHash.h
+++ b/Engine/Source/Runtime/Engine/Public/WorldPartition/WorldPartitionRuntimeHash.h
@@ -7,6 +7,7 @@
 #include "WorldPartition.h"
 #include "WorldPartitionStreamingPolicy.h"
 #include "WorldPartitionActorDescView.h"
+#include "WorldPartition/ActorDescList.h"
 #include "WorldPartition/WorldPartitionHandle.h"
 #include "WorldPartitionRuntimeHash.generated.h"
 
@@ -65,5 +66,7 @@ private:
 protected:
 #if WITH_EDITORONLY_DATA
 	TArray<FWorldPartitionReference> AlwaysLoadedActorsForPIE;
+
+	mutable FActorDescList ModifiedActorDescListForPIE;
 #endif
 };
\ No newline at end of file
diff --git a/Engine/Source/Runtime/Engine/Public/WorldPartition/WorldPartitionRuntimeLevelStreamingCell.h b/Engine/Source/Runtime/Engine/Public/WorldPartition/WorldPartitionRuntimeLevelStreamingCell.h
--- a/Engine/Source/Runtime/Engine/Public/WorldPartition/WorldPartitionRuntimeLevelStreamingCell.h
+++ b/Engine/Source/Runtime/Engine/Public/WorldPartition/WorldPartitionRuntimeLevelStreamingCell.h
@@ -17,7 +17,7 @@ class UWorldPartitionRuntimeLevelStreamingCell : public UWorldPartitionRuntimeSp
 	virtual FLinearColor GetDebugColor() const override;
 
 #if WITH_EDITOR
-	virtual void AddActorToCell(const FGuid& InActorGuid, uint32 InContainerID, const FTransform& InContainerTransform, const UActorDescContainer* InContainer) override;
+	virtual void AddActorToCell(const FWorldPartitionActorDescView& ActorDescView, uint32 InContainerID, const FTransform& InContainerTransform, const UActorDescContainer* InContainer) override;
 	virtual int32 GetActorCount() const override;
 	const TArray<FWorldPartitionRuntimeCellObjectMapping>& GetPackages() const { return Packages; }
 
