Add a slow task when we are blocking on shader compilation jobs.

- Calculate the number of pending jobs ahead of time.
- Keep track of a previous pending jobs so we can calculate a delta of jobs completed to update the progress
- Display the current number of pending jobs in the scoped slow task.

diff --git a/Engine/Source/Runtime/Engine/Private/ShaderCompiler/ShaderCompiler.cpp b/Engine/Source/Runtime/Engine/Private/ShaderCompiler/ShaderCompiler.cpp
--- a/Engine/Source/Runtime/Engine/Private/ShaderCompiler/ShaderCompiler.cpp
+++ b/Engine/Source/Runtime/Engine/Private/ShaderCompiler/ShaderCompiler.cpp
@@ -3164,7 +3164,26 @@ void FShaderCompilingManager::BlockOnShaderMapCompletion(const TArray<int32>& Sh
 	COOK_STAT(FScopedDurationTimer BlockingTimer(ShaderCompilerCookStats::BlockingTimeSec));
 	if (bAllowAsynchronousShaderCompiling)
 	{
+		// Calculate how many shader jobs there are total to provide the slow task with the correct amount of work.
+		int NumJobs = 0;
+		{
+			FScopeLock Lock(&CompileQueueSection);
+			for (int32 ShaderMapIndex = 0; ShaderMapIndex < ShaderMapIdsToFinishCompiling.Num(); ShaderMapIndex++)
+			{
+				FPendingShaderMapCompileResultsPtr* ResultsPtr = ShaderMapJobs.Find(ShaderMapIdsToFinishCompiling[ShaderMapIndex]);
+				if (ResultsPtr)
+				{
+					FShaderMapCompileResults* Results = *ResultsPtr;
+					NumJobs += Results->NumPendingJobs.GetValue();
+				}
+			}
+		}
+
+		FScopedSlowTask SlowTask(NumJobs, FText::Format(LOCTEXT("BlockOnShaderMapCompletion", "Blocking on {0} shader jobs..."), NumJobs), GIsEditor && !IsRunningCommandlet());
+
 		int32 NumPendingJobs = 0;
+		// Keep track of previous number of pending jobs so we can update the slow task with the amount of work done.
+		int32 NumPreviousPendingJobs = NumJobs;
 		int32 LogCounter = 0;
 		do 
 		{
@@ -3224,6 +3243,12 @@ void FShaderCompilingManager::BlockOnShaderMapCompletion(const TArray<int32>& Sh
 				{
 					BuildDistributionController->Tick(SleepTime);
 				}
+
+				// Progress the slow task with how many jobs we've completed since last tick.  Update the slow task message with the current number of pending jobs
+				// we are waiting on.
+				const int32 CompletedJobsSinceLastTick = NumPreviousPendingJobs - NumPendingJobs;
+				SlowTask.EnterProgressFrame(CompletedJobsSinceLastTick, FText::Format(LOCTEXT("BlockOnShaderMapCompletion", "Blocking on {0} shader jobs..."), NumPendingJobs));
+				NumPreviousPendingJobs = NumPendingJobs;
 				
 				// Yield CPU time while waiting
 				FPlatformProcess::Sleep(SleepTime);
@@ -3276,7 +3301,22 @@ void FShaderCompilingManager::BlockOnAllShaderMapCompletion(TMap<int32, FShaderM
 	COOK_STAT(FScopedDurationTimer BlockingTimer(ShaderCompilerCookStats::BlockingTimeSec));
 	if (bAllowAsynchronousShaderCompiling)
 	{
+		// Calculate how many shader jobs there are total to provide the slow task with the correct amount of work.
+		int NumJobs = 0;
+		{
+			FScopeLock Lock(&CompileQueueSection);
+			for (TMap<int32, FPendingShaderMapCompileResultsPtr>::TIterator It(ShaderMapJobs); It; ++It)
+			{
+				FShaderMapCompileResults* Results = It.Value();
+				NumJobs += Results->NumPendingJobs.GetValue();
+			}
+		}
+
+		FScopedSlowTask SlowTask(NumJobs, FText::Format(LOCTEXT("BlockOnAllShaderMapCompletion", "Blocking on {0} shader jobs..."), NumJobs), GIsEditor && !IsRunningCommandlet());
+
 		int32 NumPendingJobs = 0;
+		// Keep track of previous number of pending jobs so we can update the slow task with the amount of work done.
+		int32 NumPreviousPendingJobs = NumJobs;
 
 		do 
 		{
@@ -3319,6 +3359,12 @@ void FShaderCompilingManager::BlockOnAllShaderMapCompletion(TMap<int32, FShaderM
 				{
 					BuildDistributionController->Tick(SleepTime);
 				}
+
+				// Progress the slow task with how many jobs we've completed since last tick.  Update the slow task message with the current number of pending jobs
+				// we are waiting on.
+				const int32 CompletedJobsSinceLastTick = NumPreviousPendingJobs - NumPendingJobs;
+				SlowTask.EnterProgressFrame(CompletedJobsSinceLastTick, FText::Format(LOCTEXT("BlockOnAllShaderMapCompletion", "Blocking on {0} shader jobs..."), NumPendingJobs));
+				NumPreviousPendingJobs = NumPendingJobs;
 				
 				// Yield CPU time while waiting
 				FPlatformProcess::Sleep(SleepTime);
