Changing splash screen text for compiling shaders to: "Compiling Shaders ({0})".

- This matches the in-editor status message for compiling shaders.

diff --git a/Engine/Source/Runtime/Engine/Private/ShaderCompiler/ShaderCompiler.cpp b/Engine/Source/Runtime/Engine/Private/ShaderCompiler/ShaderCompiler.cpp
--- a/Engine/Source/Runtime/Engine/Private/ShaderCompiler/ShaderCompiler.cpp
+++ b/Engine/Source/Runtime/Engine/Private/ShaderCompiler/ShaderCompiler.cpp
@@ -3179,7 +3179,7 @@ void FShaderCompilingManager::BlockOnShaderMapCompletion(const TArray<int32>& Sh
 			}
 		}
 
-		FScopedSlowTask SlowTask(NumJobs, FText::Format(LOCTEXT("BlockOnShaderMapCompletion", "Blocking on {0} shader jobs..."), NumJobs), GIsEditor && !IsRunningCommandlet());
+		FScopedSlowTask SlowTask(NumJobs, FText::Format(LOCTEXT("BlockOnShaderMapCompletion", "Compiling Shaders ({0})"), NumJobs), GIsEditor && !IsRunningCommandlet());
 
 		int32 NumPendingJobs = 0;
 		// Keep track of previous number of pending jobs so we can update the slow task with the amount of work done.
@@ -3247,7 +3247,7 @@ void FShaderCompilingManager::BlockOnShaderMapCompletion(const TArray<int32>& Sh
 				// Progress the slow task with how many jobs we've completed since last tick.  Update the slow task message with the current number of pending jobs
 				// we are waiting on.
 				const int32 CompletedJobsSinceLastTick = NumPreviousPendingJobs - NumPendingJobs;
-				SlowTask.EnterProgressFrame(CompletedJobsSinceLastTick, FText::Format(LOCTEXT("BlockOnShaderMapCompletion", "Blocking on {0} shader jobs..."), NumPendingJobs));
+				SlowTask.EnterProgressFrame(CompletedJobsSinceLastTick, FText::Format(LOCTEXT("BlockOnShaderMapCompletion", "Compiling Shaders ({0})"), NumPendingJobs));
 				NumPreviousPendingJobs = NumPendingJobs;
 				
 				// Yield CPU time while waiting
@@ -3312,7 +3312,7 @@ void FShaderCompilingManager::BlockOnAllShaderMapCompletion(TMap<int32, FShaderM
 			}
 		}
 
-		FScopedSlowTask SlowTask(NumJobs, FText::Format(LOCTEXT("BlockOnAllShaderMapCompletion", "Blocking on {0} shader jobs..."), NumJobs), GIsEditor && !IsRunningCommandlet());
+		FScopedSlowTask SlowTask(NumJobs, FText::Format(LOCTEXT("BlockOnAllShaderMapCompletion", "Compiling Shaders ({0})"), NumJobs), GIsEditor && !IsRunningCommandlet());
 
 		int32 NumPendingJobs = 0;
 		// Keep track of previous number of pending jobs so we can update the slow task with the amount of work done.
@@ -3363,7 +3363,7 @@ void FShaderCompilingManager::BlockOnAllShaderMapCompletion(TMap<int32, FShaderM
 				// Progress the slow task with how many jobs we've completed since last tick.  Update the slow task message with the current number of pending jobs
 				// we are waiting on.
 				const int32 CompletedJobsSinceLastTick = NumPreviousPendingJobs - NumPendingJobs;
-				SlowTask.EnterProgressFrame(CompletedJobsSinceLastTick, FText::Format(LOCTEXT("BlockOnAllShaderMapCompletion", "Blocking on {0} shader jobs..."), NumPendingJobs));
+				SlowTask.EnterProgressFrame(CompletedJobsSinceLastTick, FText::Format(LOCTEXT("BlockOnAllShaderMapCompletion", "Compiling Shaders ({0})"), NumPendingJobs));
 				NumPreviousPendingJobs = NumPendingJobs;
 				
 				// Yield CPU time while waiting
