Fix crash when attempting to process a un-initialized mesh batch in ray tracing code.

- FStaticMeshSceneProxy::GetMeshElement can early out if the num triangles are zero.  This results in the FMeshBatch not being initialized and the material render proxy being null.
- Ray tracing code was assuming the FMeshBatch was valid and accessing the material proxy resulting in a crash.
- Fix this by soft failing around code that was accessing the material.

The first attempt to fix this issue was checking the return value of FStaticMeshSceneProxy::GetMeshElement in FInstancedStaticMeshSceneProxy::GetDynamicRayTracingInstances and not add the FMeshBatch to the RayTracingInstanceTemplate.Materials array.  This resulted in another assert in FDeferredShadingSceneRenderer::GatherRayTracingWorldInstances that assumes the number of materials must match the number of segments (which isn't true when you have zero triangle mesh batches).

ensureMsgf(Instance.GetMaterials().Num() == Instance.Geometry->Initializer.Segments.Num()

diff --git a/Engine/Source/Runtime/Engine/Private/RayTracingInstance.cpp b/Engine/Source/Runtime/Engine/Private/RayTracingInstance.cpp
--- a/Engine/Source/Runtime/Engine/Private/RayTracingInstance.cpp
+++ b/Engine/Source/Runtime/Engine/Private/RayTracingInstance.cpp
@@ -37,14 +37,18 @@ FRayTracingMaskAndFlags BuildRayTracingInstanceMaskAndFlags(TArrayView<const FMe
 	{
 		const FMeshBatch& MeshBatch = MeshBatches[SegmentIndex];
 
-		const FMaterialRenderProxy* FallbackMaterialRenderProxyPtr = nullptr;
-		const FMaterial& Material = MeshBatch.MaterialRenderProxy->GetMaterialWithFallback(ERHIFeatureLevel::SM5, FallbackMaterialRenderProxyPtr);
-		const EBlendMode BlendMode = Material.GetBlendMode();
-		Result.Mask |= ComputeBlendModeMask(BlendMode);
-		bAllSegmentsOpaque &= BlendMode == BLEND_Opaque;
-		bAnySegmentsCastShadow |= MeshBatch.CastRayTracedShadow && Material.CastsRayTracedShadows();
-		bAllSegmentsCastShadow &= MeshBatch.CastRayTracedShadow && Material.CastsRayTracedShadows();
-		bDoubleSided |= MeshBatch.bDisableBackfaceCulling;
+		// Mesh Batches can "null" when they have zero triangles.  Check the MaterialRenderProxy before accessing.
+		if (MeshBatch.bUseForMaterial && MeshBatch.MaterialRenderProxy)
+		{
+			const FMaterialRenderProxy* FallbackMaterialRenderProxyPtr = nullptr;
+			const FMaterial& Material = MeshBatch.MaterialRenderProxy->GetMaterialWithFallback(ERHIFeatureLevel::SM5, FallbackMaterialRenderProxyPtr);
+			const EBlendMode BlendMode = Material.GetBlendMode();
+			Result.Mask |= ComputeBlendModeMask(BlendMode);
+			bAllSegmentsOpaque &= BlendMode == BLEND_Opaque;
+			bAnySegmentsCastShadow |= MeshBatch.CastRayTracedShadow && Material.CastsRayTracedShadows();
+			bAllSegmentsCastShadow &= MeshBatch.CastRayTracedShadow && Material.CastsRayTracedShadows();
+			bDoubleSided |= MeshBatch.bDisableBackfaceCulling;
+		}
 	}
 
 	Result.bForceOpaque = bAllSegmentsOpaque && bAllSegmentsCastShadow;
diff --git a/Engine/Source/Runtime/Renderer/Private/RayTracing/RayTracingMaterialHitShaders.cpp b/Engine/Source/Runtime/Renderer/Private/RayTracing/RayTracingMaterialHitShaders.cpp
--- a/Engine/Source/Runtime/Renderer/Private/RayTracing/RayTracingMaterialHitShaders.cpp
+++ b/Engine/Source/Runtime/Renderer/Private/RayTracing/RayTracingMaterialHitShaders.cpp
@@ -317,7 +317,8 @@ void FRayTracingMeshProcessor::Process(
 void FRayTracingMeshProcessor::AddMeshBatch(const FMeshBatch& RESTRICT MeshBatch, uint64 BatchElementMask, const FPrimitiveSceneProxy* RESTRICT PrimitiveSceneProxy)
 {
 	// Caveat: there are also branches not emitting any MDC
-	if (MeshBatch.bUseForMaterial && IsSupportedVertexFactoryType(MeshBatch.VertexFactory->GetType()))
+	// Mesh Batches can "null" when they have zero triangles.  Check the MaterialRenderProxy before accessing.
+	if (MeshBatch.bUseForMaterial && MeshBatch.MaterialRenderProxy && IsSupportedVertexFactoryType(MeshBatch.VertexFactory->GetType()))
 	{
 		// Determine the mesh's material and blend mode.
 		const FMaterialRenderProxy* FallbackMaterialRenderProxyPtr = nullptr;
