Added new installation paths of IncrediBuild to search paths in ShaderCompilerXGE. Merged from CL 15645419.

#rnx

diff --git a/Engine/Source/Runtime/Engine/Private/ShaderCompiler/ShaderCompilerXGE.cpp b/Engine/Source/Runtime/Engine/Private/ShaderCompiler/ShaderCompilerXGE.cpp
--- a/Engine/Source/Runtime/Engine/Private/ShaderCompiler/ShaderCompilerXGE.cpp
+++ b/Engine/Source/Runtime/Engine/Private/ShaderCompiler/ShaderCompilerXGE.cpp
@@ -107,13 +107,15 @@ bool FShaderCompileXGEThreadRunnable_XmlInterface::IsSupported()
 	// List of possible paths to xgconsole.exe
 	static const TCHAR* Paths[] =
 	{
+		TEXT("C:\\Program Files\\IncrediBuild\\xgConsole.exe"),
+		TEXT("C:\\Program Files (x86)\\IncrediBuild\\xgConsole.exe")
 		TEXT("C:\\Program Files\\Xoreax\\IncrediBuild\\xgConsole.exe"),
 		TEXT("C:\\Program Files (x86)\\Xoreax\\IncrediBuild\\xgConsole.exe")
 	};
 
 	XGEShaderCompilerVariables::Init();
 
-	// Check for a valid installation of Incredibuild by seeing if xgconsole.exe exists.
+	// Check for a valid installation of IncrediBuild by seeing if xgConsole.exe exists.
 	bool bXgeFound = false;
 	if (XGEShaderCompilerVariables::Enabled == 1)
 	{
