Fix for SetSelection_GameThread to use IsInParallelGameThread

diff --git a/Engine/Source/Runtime/Engine/Private/PrimitiveSceneProxy.cpp b/Engine/Source/Runtime/Engine/Private/PrimitiveSceneProxy.cpp
--- a/Engine/Source/Runtime/Engine/Private/PrimitiveSceneProxy.cpp
+++ b/Engine/Source/Runtime/Engine/Private/PrimitiveSceneProxy.cpp
@@ -528,7 +528,7 @@ void FPrimitiveSceneProxy::SetSelection_RenderThread(const bool bInParentSelecte
  */
 void FPrimitiveSceneProxy::SetSelection_GameThread(const bool bInParentSelected, const bool bInIndividuallySelected)
 {
-	check(IsInGameThread());
+	check(IsInParallelGameThread() || IsInGameThread());
 
 	// Enqueue a message to the rendering thread containing the interaction to add.
 	FPrimitiveSceneProxy* PrimitiveSceneProxy = this;
