Fix HitResults not returning the accurate bone indices and names involved in the collision for SkeletalMeshComponents ( was returning root only )

diff --git a/Engine/Source/Runtime/Engine/Private/PhysicsEngine/Experimental/PhysScene_Chaos.cpp b/Engine/Source/Runtime/Engine/Private/PhysicsEngine/Experimental/PhysScene_Chaos.cpp
--- a/Engine/Source/Runtime/Engine/Private/PhysicsEngine/Experimental/PhysScene_Chaos.cpp
+++ b/Engine/Source/Runtime/Engine/Private/PhysicsEngine/Experimental/PhysScene_Chaos.cpp
@@ -635,24 +635,23 @@ void FPhysScene_Chaos::AddReferencedObjects(FReferenceCollector& Collector)
 #endif
 }
 
-static void SetCollisionInfoFromComp(FRigidBodyCollisionInfo& Info, UPrimitiveComponent* Comp)
+FBodyInstance* FPhysScene_Chaos::GetBodyInstanceFromProxy(const IPhysicsProxyBase* PhysicsProxy) const
 {
-	if (Comp)
+	FBodyInstance* BodyInstance = nullptr;
+	if (PhysicsProxy && PhysicsProxy->GetType() == EPhysicsProxyType::SingleParticleProxy)
 	{
-		Info.Component = Comp;
-		Info.Actor = Comp->GetOwner();
-
-		const FBodyInstance* const BodyInst = Comp->GetBodyInstance();
-		Info.BodyIndex = BodyInst ? BodyInst->InstanceBodyIndex : INDEX_NONE;
-		Info.BoneName = BodyInst && BodyInst->BodySetup.IsValid() ? BodyInst->BodySetup->BoneName : NAME_None;
+		const Chaos::FRigidBodyHandle_External& RigidBodyHandle = static_cast<const FSingleParticlePhysicsProxy*>(PhysicsProxy)->GetGameThreadAPI();
+		BodyInstance = FPhysicsUserData::Get<FBodyInstance>(RigidBodyHandle.UserData());
 	}
-	else
+	// found none, let's see if there's an owning component in the scene
+	if (BodyInstance == nullptr)
 	{
-		Info.Component = nullptr;
-		Info.Actor = nullptr;
-		Info.BodyIndex = INDEX_NONE;
-		Info.BoneName = NAME_None;
+		if (UPrimitiveComponent* const OwningComponent = GetOwningComponent<UPrimitiveComponent>(PhysicsProxy))
+		{
+			BodyInstance = OwningComponent->GetBodyInstance();
+		}
 	}
+	return BodyInstance;
 }
 
 FCollisionNotifyInfo& FPhysScene_Chaos::GetPendingCollisionForContactPair(const void* P0, const void* P1, bool& bNewEntry)
@@ -717,14 +716,12 @@ void FPhysScene_Chaos::HandleCollisionEvents(const Chaos::FCollisionEventData& E
 
 							if (bNewEntry)
 							{
-								UPrimitiveComponent* const Comp1 = GetOwningComponent<UPrimitiveComponent>(PhysicsProxy1);
-
 								// fill in legacy contact data
 								NotifyInfo.bCallEvent0 = true;
 								// if Comp1 wants this event too, it will get its own pending collision entry, so we leave it false
 
-								SetCollisionInfoFromComp(NotifyInfo.Info0, Comp0);
-								SetCollisionInfoFromComp(NotifyInfo.Info1, Comp1);
+								NotifyInfo.Info0.SetFrom(GetBodyInstanceFromProxy(PhysicsProxy0));
+								NotifyInfo.Info1.SetFrom(GetBodyInstanceFromProxy(PhysicsProxy1));
 
 								FRigidBodyContactInfo& NewContact = NotifyInfo.RigidCollisionData.ContactInfos.AddZeroed_GetRef();
 								NewContact.ContactNormal = CollisionDataItem.Normal;
diff --git a/Engine/Source/Runtime/Engine/Public/Physics/Experimental/PhysScene_Chaos.h b/Engine/Source/Runtime/Engine/Public/Physics/Experimental/PhysScene_Chaos.h
--- a/Engine/Source/Runtime/Engine/Public/Physics/Experimental/PhysScene_Chaos.h
+++ b/Engine/Source/Runtime/Engine/Public/Physics/Experimental/PhysScene_Chaos.h
@@ -123,7 +123,7 @@ public:
 	virtual void AddReferencedObjects(FReferenceCollector& Collector) override;
 	/** Given a solver object, returns its associated component. */
 	template<class OwnerType>
-	OwnerType* GetOwningComponent(IPhysicsProxyBase* PhysicsProxy) const
+	OwnerType* GetOwningComponent(const IPhysicsProxyBase* PhysicsProxy) const
 	{ 
 		UPrimitiveComponent* const* CompPtr = PhysicsProxyToComponentMap.Find(PhysicsProxy);
 		return CompPtr ? Cast<OwnerType>(*CompPtr) : nullptr;
@@ -135,6 +135,9 @@ public:
 		return ComponentToPhysicsProxyMap.Find(Comp);
 	}
 
+	/** Given a physics proxy, returns its associated body instance if any */
+	FBodyInstance* GetBodyInstanceFromProxy(const IPhysicsProxyBase* PhysicsProxy) const;
+
 	/**
 	 * Callback when a world ends, to mark updated packages dirty. This can't be done in final
 	 * sync as the editor will ignore packages being dirtied in PIE. Also used to clean up any other references
diff --git a/Engine/Source/Runtime/Experimental/Chaos/Public/Chaos/Framework/PhysicsProxyBase.h b/Engine/Source/Runtime/Experimental/Chaos/Public/Chaos/Framework/PhysicsProxyBase.h
--- a/Engine/Source/Runtime/Experimental/Chaos/Public/Chaos/Framework/PhysicsProxyBase.h
+++ b/Engine/Source/Runtime/Experimental/Chaos/Public/Chaos/Framework/PhysicsProxyBase.h
@@ -58,7 +58,7 @@ public:
 	template< class SOLVER_TYPE = Chaos::FPhysicsSolver>
 	void SetSolver(SOLVER_TYPE* InSolver) { Solver = InSolver; }
 
-	EPhysicsProxyType GetType() { return Type; }
+	EPhysicsProxyType GetType() const { return Type; }
 
 	//todo: remove this
 	virtual void* GetHandleUnsafe() const { check(false); return nullptr; }
diff --git a/Engine/Source/Runtime/Experimental/ChaosSolverEngine/Private/Chaos/ChaosGameplayEventDispatcher.cpp b/Engine/Source/Runtime/Experimental/ChaosSolverEngine/Private/Chaos/ChaosGameplayEventDispatcher.cpp
--- a/Engine/Source/Runtime/Experimental/ChaosSolverEngine/Private/Chaos/ChaosGameplayEventDispatcher.cpp
+++ b/Engine/Source/Runtime/Experimental/ChaosSolverEngine/Private/Chaos/ChaosGameplayEventDispatcher.cpp
@@ -2,8 +2,10 @@
 
 #include "Chaos/ChaosGameplayEventDispatcher.h"
 #include "Components/PrimitiveComponent.h"
+#include "Components/SkeletalMeshComponent.h"
 #include "PhysicsEngine/BodyInstance.h"
 #include "Chaos/Framework/PhysicsProxy.h"
+#include "Chaos/Public/PhysicsProxy/SingleParticlePhysicsProxy.h"
 #include "PhysicsSolver.h"
 #include "Physics/Experimental/PhysScene_Chaos.h"
 #include "Engine/World.h"
@@ -53,27 +55,6 @@ static void DispatchPendingBreakEvents(TArray<FChaosBreakEvent> const& Events, T
 	}
 }
 
-static void SetCollisionInfoFromComp(FRigidBodyCollisionInfo& Info, UPrimitiveComponent* Comp)
-{
-	if (Comp)
-	{
-		Info.Component = Comp;
-		Info.Actor = Comp->GetOwner();
-		
-		const FBodyInstance* const BodyInst = Comp->GetBodyInstance();
-		Info.BodyIndex = BodyInst ? BodyInst->InstanceBodyIndex : INDEX_NONE;
-		Info.BoneName = BodyInst && BodyInst->BodySetup.IsValid() ? BodyInst->BodySetup->BoneName : NAME_None;
-	}
-	else
-	{
-		Info.Component = nullptr;
-		Info.Actor = nullptr;
-		Info.BodyIndex = INDEX_NONE;
-		Info.BoneName = NAME_None;
-	}
-}
-
-
 FCollisionNotifyInfo& UChaosGameplayEventDispatcher::GetPendingCollisionForContactPair(const void* P0, const void* P1, bool& bNewEntry)
 {
 	const FUniqueContactPairKey Key = { P0, P1 };
@@ -308,14 +289,12 @@ void UChaosGameplayEventDispatcher::HandleCollisionEvents(const Chaos::FCollisio
 
 									if (bNewEntry)
 									{
-										UPrimitiveComponent* const Comp1 = Scene.GetOwningComponent<UPrimitiveComponent>(PhysicsProxy1);
-
 										// fill in legacy contact data
 										NotifyInfo.bCallEvent0 = true;
 										// if Comp1 wants this event too, it will get its own pending collision entry, so we leave it false
 
-										SetCollisionInfoFromComp(NotifyInfo.Info0, Comp0);
-										SetCollisionInfoFromComp(NotifyInfo.Info1, Comp1);
+										NotifyInfo.Info0.SetFrom(Scene.GetBodyInstanceFromProxy(PhysicsProxy0));
+										NotifyInfo.Info1.SetFrom(Scene.GetBodyInstanceFromProxy(PhysicsProxy1));
 
 										FRigidBodyContactInfo& NewContact = NotifyInfo.RigidCollisionData.ContactInfos.AddZeroed_GetRef();
 										NewContact.ContactNormal = CollisionDataItem.Normal;
