Unify EmitHLSL/EmitPreshader methods on HLSLTree expression, now a single method that can emit one or the other depending on inputs

diff --git a/Engine/Source/Runtime/Engine/Private/HLSLTree/HLSLTree.cpp b/Engine/Source/Runtime/Engine/Private/HLSLTree/HLSLTree.cpp
--- a/Engine/Source/Runtime/Engine/Private/HLSLTree/HLSLTree.cpp
+++ b/Engine/Source/Runtime/Engine/Private/HLSLTree/HLSLTree.cpp
@@ -147,7 +147,21 @@ void UE::HLSLTree::FCodeWriter::WriteIndent()
 {
 	for (int32 i = 0; i < IndentLevel; ++i)
 	{
-		StringBuilder->Append(TCHAR(' '));
+		StringBuilder->Append(TCHAR('\t'));
+	}
+}
+
+void UE::HLSLTree::FCodeWriter::Reset()
+{
+	StringBuilder->Reset();
+}
+
+void UE::HLSLTree::FCodeWriter::Append(const FCodeWriter& InWriter)
+{
+	const int32 Length = InWriter.StringBuilder->Len();
+	if (Length > 0)
+	{
+		StringBuilder->Append(InWriter.StringBuilder->GetData(), Length);
 	}
 }
 
@@ -185,6 +199,22 @@ void UE::HLSLTree::FConstant::EmitHLSL(UE::HLSLTree::FCodeWriter& Writer) const
 	}
 }
 
+UE::HLSLTree::FEmitContext::FEmitContext()
+{
+	FunctionStack.AddDefaulted();
+}
+
+UE::HLSLTree::FEmitContext::~FEmitContext()
+{
+	// Make sure we have only the root stack entry
+	check(FunctionStack.Num() == 1);
+
+	for (FMaterialPreshaderData* Preshader : TempPreshaders)
+	{
+		delete Preshader;
+	}
+}
+
 UE::HLSLTree::FEmitContext::FScopeEntry* UE::HLSLTree::FEmitContext::FindScope(UE::HLSLTree::FScope* Scope)
 {
 	for (int32 Index = ScopeStack.Num() - 1; Index >= 0; --Index)
@@ -198,155 +228,207 @@ UE::HLSLTree::FEmitContext::FScopeEntry* UE::HLSLTree::FEmitContext::FindScope(U
 	return nullptr;
 }
 
-const TCHAR* UE::HLSLTree::FEmitContext::AcquireHLSLReference(UE::HLSLTree::FLocalDeclaration* Declaration)
+const UE::HLSLTree::FEmitValue& UE::HLSLTree::FEmitContext::AcquireValue(UE::HLSLTree::FLocalDeclaration* Declaration)
 {
 	check(Declaration);
 	check(Declaration->Type != EExpressionType::Void);
 
-	FDeclarationEntry& Entry = DeclarationMap.FindOrAdd(Declaration);
-	if (!Entry.Definition)
+	FDeclarationEntry& Entry = FunctionStack.Last().DeclarationMap.FindOrAdd(Declaration);
+	if (Entry.Value.EvaluationType == EExpressionEvaluationType::None)
 	{
 		FScopeEntry* ScopeEntry = FindScope(Declaration->ParentScope);
 		check(ScopeEntry);
 
-		Entry.Definition = AllocateStringf(*Allocator, TEXT("Local%d"), NumExpressionLocals++);
+		Entry.Value.EvaluationType = EExpressionEvaluationType::Shader;
+		Entry.Value.Code = AllocateStringf(*Allocator, TEXT("Local%d"), NumExpressionLocals++);
 		const FExpressionTypeDescription& TypeDesc = GetExpressionTypeDescription(Declaration->Type);
 		ScopeEntry->ExpressionCodeWriter->WriteLinef(TEXT("%s %s = (%s)0.0f;"),
 			TypeDesc.Name,
-			Entry.Definition,
+			Entry.Value.Code,
 			TypeDesc.Name);
 	}
 
-	return Entry.Definition;
+	return Entry.Value;
 }
 
-const TCHAR* UE::HLSLTree::FEmitContext::AcquireHLSLReference(UE::HLSLTree::FExpression* Expression)
+const UE::HLSLTree::FEmitValue& UE::HLSLTree::FEmitContext::AcquireValue(UE::HLSLTree::FExpression* Expression)
 {
 	check(Expression);
 	check(Expression->Type != EExpressionType::Void);
 
-	FDeclarationEntry& Entry = DeclarationMap.FindOrAdd(Expression);
-	if (!Entry.Definition)
+	FDeclarationEntry& Entry = FunctionStack.Last().DeclarationMap.FindOrAdd(Expression);
+	if (Entry.Value.EvaluationType == EExpressionEvaluationType::None)
 	{
-		bool bInline = false;
 		FLocalHLSLCodeWriter LocalWriter;
-		if (Expression->EvaluationType == EExpressionEvaluationType::Shader)
+		FMaterialPreshaderData LocalPreshader;
+		FExpressionEmitResult EmitResult(LocalWriter, LocalPreshader);
+		Expression->EmitHLSL(*this, EmitResult);
+
+		Entry.Value.ExpressionType = Expression->Type;
+		Entry.Value.EvaluationType = EmitResult.EvaluationType;
+		if (EmitResult.EvaluationType == EExpressionEvaluationType::Constant)
+		{
+			// Evaluate the constant preshader and store its value
+			FMaterialRenderContext RenderContext(nullptr, *Material, nullptr);
+			FLinearColor ConstantValue;
+			LocalPreshader.Evaluate(nullptr, RenderContext, ConstantValue);
+
+			const FConstant Constant(Expression->Type, ConstantValue);
+			Entry.Value.ConstantValue = Constant;
+		}
+		else if (EmitResult.EvaluationType == EExpressionEvaluationType::Preshader)
 		{
-			bInline = Expression->IsInline();
-			Expression->EmitHLSL(*this, LocalWriter);
+			// Non-constant preshader, store it
+			FMaterialPreshaderData* Preshader = new FMaterialPreshaderData(MoveTemp(LocalPreshader));
+			TempPreshaders.Add(Preshader); // TODO - dedupe? store more efficiently?
+			Entry.Value.Preshader = Preshader;
 		}
 		else
 		{
-			if (Expression->EvaluationType == EExpressionEvaluationType::Constant)
+			check(EmitResult.EvaluationType == EExpressionEvaluationType::Shader);
+			if (EmitResult.bInline)
 			{
-				// Evaluate the constant preshader, and write out its value
-				FMaterialPreshaderData Preshader;
-				Expression->EmitPreshader(*this, Preshader);
-
-				FMaterialRenderContext RenderContext(nullptr, *Material, nullptr);
-				FLinearColor ConstantValue;
-				Preshader.Evaluate(nullptr, RenderContext, ConstantValue);
-
-				const FConstant Constant(Expression->Type, ConstantValue);
-				Constant.EmitHLSL(LocalWriter);
+				Entry.Value.Code = AllocateString(*Allocator, LocalWriter.GetStringBuilder());
 			}
 			else
 			{
-				// Non-constant preshader, store it on the material
-				check(Expression->EvaluationType == EExpressionEvaluationType::Preshader);
-				FUniformExpressionSet& UniformExpressionSet = MaterialCompilationOutput->UniformExpressionSet;
-
-				FMaterialUniformPreshaderHeader* PreshaderHeader = nullptr;
-				if (Expression->Type == EExpressionType::Float1)
-				{
-					const static TCHAR IndexToMask[] = { 'x', 'y', 'z', 'w' };
-					const int32 ScalarInputIndex = UniformExpressionSet.UniformScalarPreshaders.Num();
-					LocalWriter.Writef(TEXT("Material.ScalarExpressions[%u].%c"), ScalarInputIndex / 4, IndexToMask[ScalarInputIndex % 4]);
-
-					PreshaderHeader = &UniformExpressionSet.UniformScalarPreshaders.AddDefaulted_GetRef();
-				}
-				else
-				{
-					const TCHAR* Mask = TEXT("");
-					switch (Expression->Type)
-					{
-					case EExpressionType::Float1: Mask = TEXT(".r"); break;
-					case EExpressionType::Float2: Mask = TEXT(".rg"); break;
-					case EExpressionType::Float3: Mask = TEXT(".rgb"); break;
-					case EExpressionType::Float4: break;
-					default: checkNoEntry(); break;
-					};
-					const int32 VectorInputIndex = UniformExpressionSet.UniformVectorPreshaders.Num();
-					LocalWriter.Writef(TEXT("Material.VectorExpressions[%u]%s"), VectorInputIndex, Mask);
-
-					PreshaderHeader = &UniformExpressionSet.UniformVectorPreshaders.AddDefaulted_GetRef();
-				}
-				PreshaderHeader->OpcodeOffset = UniformExpressionSet.UniformPreshaderData.Num();
-				Expression->EmitPreshader(*this, UniformExpressionSet.UniformPreshaderData);
-				PreshaderHeader->OpcodeSize = UniformExpressionSet.UniformPreshaderData.Num() - PreshaderHeader->OpcodeOffset;
+				FScopeEntry* ScopeEntry = FindScope(Expression->ParentScope);
+				check(ScopeEntry);
+
+				const FExpressionTypeDescription& TypeDesc = GetExpressionTypeDescription(Expression->Type);
+				Entry.Value.Code = AllocateStringf(*Allocator, TEXT("Local%d"), NumExpressionLocals++);
+				ScopeEntry->ExpressionCodeWriter->WriteLinef(TEXT("%s %s = %s;"),
+					TypeDesc.Name,
+					Entry.Value.Code,
+					LocalWriter.GetStringBuilder().ToString());
 			}
-			bInline = true;
-		}
-
-		if (bInline)
-		{
-			Entry.Definition = AllocateString(*Allocator, LocalWriter.GetStringBuilder());
-		}
-		else
-		{
-			FScopeEntry* ScopeEntry = FindScope(Expression->ParentScope);
-			check(ScopeEntry);
-
-			const FExpressionTypeDescription& TypeDesc = GetExpressionTypeDescription(Expression->Type);
-			Entry.Definition = AllocateStringf(*Allocator, TEXT("Local%d"), NumExpressionLocals++);
-			ScopeEntry->ExpressionCodeWriter->WriteLinef(TEXT("%s %s = %s;"),
-				TypeDesc.Name,
-				Entry.Definition,
-				LocalWriter.GetStringBuilder().ToString());
 		}
 	}
 
-	return Entry.Definition;
+	return Entry.Value;
 }
 
-const TCHAR* UE::HLSLTree::FEmitContext::AcquireHLSLReference(FFunctionCall* FunctionCall, int32 OutputIndex)
+const UE::HLSLTree::FEmitValue& UE::HLSLTree::FEmitContext::AcquireValue(FFunctionCall* FunctionCall, int32 OutputIndex)
 {
 	check(FunctionCall);
 
-	FFunctionCallEntry& Entry = FunctionCallMap.FindOrAdd(FunctionCall);
-	if (!Entry.OutputDefinition)
+	FFunctionCallEntry& Entry = FunctionStack.Last().FunctionCallMap.FindOrAdd(FunctionCall);
+	if (!Entry.OutputRef)
 	{
 		FScopeEntry* ScopeEntry = FindScope(FunctionCall->ParentScope);
 		check(ScopeEntry);
 
-		TCHAR const** OutputDefinition = New<const TCHAR*>(*Allocator, FunctionCall->NumOutputs);
+		FEmitValue* OutputRef = new(*Allocator) FEmitValue[FunctionCall->NumOutputs];
 		for (int32 i = 0; i < FunctionCall->NumOutputs; ++i)
 		{
-			OutputDefinition[i] = AllocateStringf(*Allocator, TEXT("Local%d"), NumExpressionLocals++);
-			const FExpressionTypeDescription& TypeDesc = GetExpressionTypeDescription(FunctionCall->OutputTypes[i]);
-			ScopeEntry->ExpressionCodeWriter->WriteLinef(TEXT("%s %s = (%s)0.0f;"),
-				TypeDesc.Name,
-				OutputDefinition[i],
-				TypeDesc.Name);
+			OutputRef[i].Code = AllocateStringf(*Allocator, TEXT("Local%d"), NumExpressionLocals++);
 		}
 
-		Entry.OutputDefinition = OutputDefinition;
+		Entry.OutputRef = OutputRef;
 		Entry.NumOutputs = FunctionCall->NumOutputs;
 
-		FFunctionStackEntry& StackEntry = FunctionStack.AddDefaulted_GetRef();
-		StackEntry.FunctionCall = FunctionCall;
-
-		FunctionCall->FunctionScope->EmitHLSL(*this, *ScopeEntry->ExpressionCodeWriter);
-
+		{
+			FFunctionStackEntry& StackEntry = FunctionStack.AddDefaulted_GetRef();
+			StackEntry.FunctionCall = FunctionCall;
+			StackEntry.OutputRef = OutputRef;
+		}
+		FLocalHLSLCodeWriter FunctionWriter;
+		FunctionWriter.IndentLevel = ScopeEntry->ExpressionCodeWriter->IndentLevel;
+		FunctionCall->FunctionScope->EmitHLSL(*this, FunctionWriter);
 		{
 			const FFunctionStackEntry PoppedStackEntry = FunctionStack.Pop();
 			check(PoppedStackEntry.FunctionCall == FunctionCall);
 		}
+
+		// Allocate locals to hold any function outputs evaluated in shader
+		for (int32 i = 0; i < FunctionCall->NumOutputs; ++i)
+		{
+			const FEmitValue& Output = OutputRef[i];
+			if (Output.EvaluationType == EExpressionEvaluationType::Shader)
+			{
+				const EExpressionType OutputType = FunctionCall->OutputTypes[i];
+				const FExpressionTypeDescription& TypeDesc = GetExpressionTypeDescription(OutputType);
+				ScopeEntry->ExpressionCodeWriter->WriteLinef(TEXT("%s %s = (%s)0.0f;"),
+					TypeDesc.Name,
+					OutputRef[i].Code,
+					TypeDesc.Name);
+			}
+		}
+
+		// Code to evaluate actual function (TODO - cull empty scopes)
+		ScopeEntry->ExpressionCodeWriter->Append(FunctionWriter);
 	}
 
 	check(Entry.NumOutputs == FunctionCall->NumOutputs);
 	check(OutputIndex >= 0 && OutputIndex < Entry.NumOutputs);
-	return Entry.OutputDefinition[OutputIndex];
+	return Entry.OutputRef[OutputIndex];
+}
+
+const TCHAR* UE::HLSLTree::FEmitContext::GetCode(const FEmitValue& Value) const
+{
+	if (!Value.Code)
+	{
+		FLocalHLSLCodeWriter LocalWriter;
+		if (Value.EvaluationType == EExpressionEvaluationType::Constant)
+		{
+			Value.ConstantValue.EmitHLSL(LocalWriter);
+		}
+		else
+		{
+			check(Value.EvaluationType == EExpressionEvaluationType::Preshader);
+			check(Value.Preshader);
+
+			FUniformExpressionSet& UniformExpressionSet = MaterialCompilationOutput->UniformExpressionSet;
+			FMaterialUniformPreshaderHeader* PreshaderHeader = nullptr;
+			if (Value.ExpressionType == EExpressionType::Float1)
+			{
+				const static TCHAR IndexToMask[] = { 'x', 'y', 'z', 'w' };
+				const int32 ScalarInputIndex = UniformExpressionSet.UniformScalarPreshaders.Num();
+				LocalWriter.Writef(TEXT("Material.ScalarExpressions[%u].%c"), ScalarInputIndex / 4, IndexToMask[ScalarInputIndex % 4]);
+
+				PreshaderHeader = &UniformExpressionSet.UniformScalarPreshaders.AddDefaulted_GetRef();
+			}
+			else
+			{
+				const TCHAR* Mask = TEXT("");
+				switch (Value.ExpressionType)
+				{
+				case EExpressionType::Float1: Mask = TEXT(".r"); break;
+				case EExpressionType::Float2: Mask = TEXT(".rg"); break;
+				case EExpressionType::Float3: Mask = TEXT(".rgb"); break;
+				case EExpressionType::Float4: break;
+				default: checkNoEntry(); break;
+				};
+				const int32 VectorInputIndex = UniformExpressionSet.UniformVectorPreshaders.Num();
+				LocalWriter.Writef(TEXT("Material.VectorExpressions[%u]%s"), VectorInputIndex, Mask);
+
+				PreshaderHeader = &UniformExpressionSet.UniformVectorPreshaders.AddDefaulted_GetRef();
+			}
+			PreshaderHeader->OpcodeOffset = UniformExpressionSet.UniformPreshaderData.Num();
+			UniformExpressionSet.UniformPreshaderData.Append(*Value.Preshader);
+			PreshaderHeader->OpcodeSize = Value.Preshader->Num();
+		}
+		Value.Code = AllocateString(*Allocator, LocalWriter.GetStringBuilder());
+	}
+
+	return Value.Code;
+}
+
+void UE::HLSLTree::FEmitContext::AppendPreshader(const FEmitValue& Value, FMaterialPreshaderData& InOutPreshader) const
+{
+	if (Value.EvaluationType == EExpressionEvaluationType::Constant)
+	{
+		// Push the constant value
+		const FLinearColor Constant = Value.ConstantValue.ToLinearColor();
+		InOutPreshader.WriteOpcode(EMaterialPreshaderOpcode::Constant);
+		InOutPreshader.Write(Constant);
+	}
+	else
+	{
+		check(Value.EvaluationType == EExpressionEvaluationType::Preshader);
+		check(Value.Preshader);
+		InOutPreshader.Append(*Value.Preshader);
+	}
 }
 
 void UE::HLSLTree::FNodeVisitor::VisitNode(UE::HLSLTree::FNode* Node)
@@ -367,16 +449,6 @@ UE::HLSLTree::ENodeVisitResult UE::HLSLTree::FExpression::Visit(UE::HLSLTree::FN
 	return Visitor.OnExpression(*this);
 }
 
-void UE::HLSLTree::FExpression::EmitHLSL(UE::HLSLTree::FEmitContext& Context, UE::HLSLTree::FCodeWriter& Writer) const
-{
-	check(false);
-}
-
-void UE::HLSLTree::FExpression::EmitPreshader(UE::HLSLTree::FEmitContext& Context, FMaterialPreshaderData& OutPreshader) const
-{
-	check(false);
-}
-
 UE::HLSLTree::ENodeVisitResult UE::HLSLTree::FLocalDeclaration::Visit(UE::HLSLTree::FNodeVisitor& Visitor)
 {
 	return Visitor.OnLocalDeclaration(*this);
@@ -442,11 +514,8 @@ void UE::HLSLTree::FScope::EmitHLSL(UE::HLSLTree::FEmitContext& Context, UE::HLS
 		Statement->EmitHLSL(Context, StatementCodeWriter);
 
 		// First write any expressions needed by the statement, then the statement itself
-		if (DeclarationCodeWriter.StringBuilder->Len() > 0)
-		{
-			Writer.StringBuilder->Append(DeclarationCodeWriter.StringBuilder->GetData(), DeclarationCodeWriter.StringBuilder->Len());
-		}
-		Writer.StringBuilder->Append(StatementCodeWriter.StringBuilder->GetData(), StatementCodeWriter.StringBuilder->Len());
+		Writer.Append(DeclarationCodeWriter);
+		Writer.Append(StatementCodeWriter);
 
 		DeclarationCodeWriter.StringBuilder->Reset();
 		StatementCodeWriter.StringBuilder->Reset();
diff --git a/Engine/Source/Runtime/Engine/Private/HLSLTree/HLSLTreeCommon.cpp b/Engine/Source/Runtime/Engine/Private/HLSLTree/HLSLTreeCommon.cpp
--- a/Engine/Source/Runtime/Engine/Private/HLSLTree/HLSLTreeCommon.cpp
+++ b/Engine/Source/Runtime/Engine/Private/HLSLTree/HLSLTreeCommon.cpp
@@ -58,25 +58,26 @@ UE::HLSLTree::FSwizzleParameters UE::HLSLTree::MakeSwizzleMask(bool bInR, bool b
 	return FSwizzleParameters(ComponentIndex[0], ComponentIndex[1], ComponentIndex[2], ComponentIndex[3]);
 }
 
-void UE::HLSLTree::FExpressionConstant::EmitPreshader(UE::HLSLTree::FEmitContext& Context, FMaterialPreshaderData& OutPreshader) const
-{
-	OutPreshader.WriteOpcode(EMaterialPreshaderOpcode::Constant);
-	OutPreshader.Write(Value.ToLinearColor());
-}
 
-void UE::HLSLTree::FExpressionConstant::EmitHLSL(UE::HLSLTree::FEmitContext& Context, UE::HLSLTree::FCodeWriter& Writer) const
+void UE::HLSLTree::FExpressionConstant::EmitHLSL(UE::HLSLTree::FEmitContext& Context, FExpressionEmitResult& OutResult) const
 {
-	Value.EmitHLSL(Writer);
+	OutResult.EvaluationType = EExpressionEvaluationType::Constant;
+	OutResult.Preshader.WriteOpcode(EMaterialPreshaderOpcode::Constant);
+	OutResult.Preshader.Write(Value.ToLinearColor());
 }
 
-void UE::HLSLTree::FExpressionLocalVariable::EmitHLSL(UE::HLSLTree::FEmitContext& Context, UE::HLSLTree::FCodeWriter& Writer) const
+void UE::HLSLTree::FExpressionLocalVariable::EmitHLSL(UE::HLSLTree::FEmitContext& Context, FExpressionEmitResult& OutResult) const
 {
-	const TCHAR* DeclarationRef = Context.AcquireHLSLReference(Declaration);
-	Writer.Writef(TEXT("%s"), DeclarationRef);
+	const FEmitValue&DeclarationRef = Context.AcquireValue(Declaration);
+	
+	OutResult.bInline = true;
+	OutResult.Writer.Writef(TEXT("%s"), Context.GetCode(DeclarationRef));
 }
 
-void UE::HLSLTree::FExpressionParameter::EmitPreshader(UE::HLSLTree::FEmitContext& Context, FMaterialPreshaderData& OutPreshader) const
+void UE::HLSLTree::FExpressionParameter::EmitHLSL(UE::HLSLTree::FEmitContext& Context, FExpressionEmitResult& OutResult) const
 {
+	OutResult.EvaluationType = EExpressionEvaluationType::Preshader;
+
 	EMaterialPreshaderOpcode Opcode = EMaterialPreshaderOpcode::Nop;
 	int32 ParameterIndex = INDEX_NONE;
 	if (Declaration->DefaultValue.Type == UE::HLSLTree::EExpressionType::Float1)
@@ -91,18 +92,20 @@ void UE::HLSLTree::FExpressionParameter::EmitPreshader(UE::HLSLTree::FEmitContex
 	}
 
 	check(ParameterIndex >= 0 && ParameterIndex <= 0xffff);
-	OutPreshader.WriteOpcode(Opcode);
-	OutPreshader.Write((uint16)ParameterIndex);
+	OutResult.Preshader.WriteOpcode(Opcode);
+	OutResult.Preshader.Write((uint16)ParameterIndex);
 }
 
-void UE::HLSLTree::FExpressionExternalInput::EmitHLSL(UE::HLSLTree::FEmitContext& Context, UE::HLSLTree::FCodeWriter& Writer) const
+void UE::HLSLTree::FExpressionExternalInput::EmitHLSL(UE::HLSLTree::FEmitContext& Context, FExpressionEmitResult& OutResult) const
 {
 	const int32 TypeIndex = (int32)InputType;
 	const int32 TexCoordIndex = TypeIndex - (int32)UE::HLSLTree::EExternalInputType::TexCoord0;
+
+	OutResult.bInline = true;
 	if (TexCoordIndex >= 0 && TexCoordIndex < 8)
 	{
 		Context.NumTexCoords = FMath::Max(Context.NumTexCoords, TexCoordIndex + 1);
-		Writer.Writef(TEXT("Parameters.TexCoords[%u].xy"), TexCoordIndex);
+		OutResult.Writer.Writef(TEXT("Parameters.TexCoords[%u].xy"), TexCoordIndex);
 	}
 	else
 	{
@@ -110,7 +113,7 @@ void UE::HLSLTree::FExpressionExternalInput::EmitHLSL(UE::HLSLTree::FEmitContext
 	}
 }
 
-void UE::HLSLTree::FExpressionTextureSample::EmitHLSL(UE::HLSLTree::FEmitContext& Context, UE::HLSLTree::FCodeWriter& Writer) const
+void UE::HLSLTree::FExpressionTextureSample::EmitHLSL(UE::HLSLTree::FEmitContext& Context, FExpressionEmitResult& OutResult) const
 {
 	const UE::HLSLTree::FTextureDescription& Desc = Declaration->Description;
 	const EMaterialValueType MaterialValueType = Desc.Texture->GetMaterialType();
@@ -236,61 +239,55 @@ void UE::HLSLTree::FExpressionTextureSample::EmitHLSL(UE::HLSLTree::FEmitContext
 		}
 	}
 
-	const TCHAR* TexCoordRef = Context.AcquireHLSLReference(TexCoordExpression);
-	Writer.Writef(TEXT("%s(%s(%s, %s, %s))"), SamplerTypeFunction, SampleFunctionName, *TextureName, *SamplerStateCode, TexCoordRef);
+	const FEmitValue& TexCoordRef = Context.AcquireValue(TexCoordExpression);
+	OutResult.Writer.Writef(TEXT("%s(%s(%s, %s, %s))"), SamplerTypeFunction, SampleFunctionName, *TextureName, *SamplerStateCode, Context.GetCode(TexCoordRef));
 }
 
-void UE::HLSLTree::FExpressionDefaultMaterialAttributes::EmitHLSL(UE::HLSLTree::FEmitContext& Context, UE::HLSLTree::FCodeWriter& Writer) const
+void UE::HLSLTree::FExpressionDefaultMaterialAttributes::EmitHLSL(UE::HLSLTree::FEmitContext& Context, FExpressionEmitResult& OutResult) const
 {
-	Writer.Write(TEXT("DefaultMaterialAttributes"));
+	OutResult.bInline = true;
+	OutResult.Writer.Write(TEXT("DefaultMaterialAttributes"));
 }
 
-void UE::HLSLTree::FExpressionSetMaterialAttribute::EmitHLSL(UE::HLSLTree::FEmitContext& Context, UE::HLSLTree::FCodeWriter& Writer) const
+void UE::HLSLTree::FExpressionSetMaterialAttribute::EmitHLSL(UE::HLSLTree::FEmitContext& Context, FExpressionEmitResult& OutResult) const
 {
 	const EMaterialProperty Property = FMaterialAttributeDefinitionMap::GetProperty(AttributeID);
 	const FString PropertyName = FMaterialAttributeDefinitionMap::GetAttributeName(AttributeID);
 	const EMaterialValueType PropertyType = FMaterialAttributeDefinitionMap::GetValueType(AttributeID);
 	const EShaderFrequency Frequency = FMaterialAttributeDefinitionMap::GetShaderFrequency(AttributeID);
 
-	const TCHAR* AttributresRef = Context.AcquireHLSLReference(AttributesExpression);
-	const TCHAR* ValueRef = Context.AcquireHLSLReference(ValueExpression);
+	const FEmitValue&AttributresRef = Context.AcquireValue(AttributesExpression);
+	const FEmitValue&ValueRef = Context.AcquireValue(ValueExpression);
 
-	Writer.Writef(TEXT("FMaterialAttributes_Set%s(%s, %s)"), *PropertyName, AttributresRef, ValueRef);
+	OutResult.Writer.Writef(TEXT("FMaterialAttributes_Set%s(%s, %s)"), *PropertyName, Context.GetCode(AttributresRef), Context.GetCode(ValueRef));
 }
 
-void UE::HLSLTree::FExpressionBinaryOp::EmitPreshader(UE::HLSLTree::FEmitContext& Context, FMaterialPreshaderData& OutPreshader) const
+void UE::HLSLTree::FExpressionBinaryOp::EmitHLSL(UE::HLSLTree::FEmitContext& Context, FExpressionEmitResult& OutResult) const
 {
-	EMaterialPreshaderOpcode Opcode = EMaterialPreshaderOpcode::Nop;
-	switch (Op)
+	const FEmitValue& LhsRef = Context.AcquireValue(Lhs);
+	const FEmitValue& RhsRef = Context.AcquireValue(Rhs);
+
+	OutResult.EvaluationType = CombineEvaluationTypes(LhsRef.GetEvaluationType(), RhsRef.GetEvaluationType());
+	if (OutResult.EvaluationType == EExpressionEvaluationType::Shader)
 	{
-	case UE::HLSLTree::EBinaryOp::Add: Opcode = EMaterialPreshaderOpcode::Add; break;
-	default: checkNoEntry(); break;
+		OutResult.Writer.Writef(TEXT("(%s + %s)"), Context.GetCode(LhsRef), Context.GetCode(RhsRef));
 	}
+	else
+	{
+		EMaterialPreshaderOpcode Opcode = EMaterialPreshaderOpcode::Nop;
+		switch (Op)
+		{
+		case UE::HLSLTree::EBinaryOp::Add: Opcode = EMaterialPreshaderOpcode::Add; break;
+		default: checkNoEntry(); break;
+		}
 
-	Lhs->EmitPreshader(Context, OutPreshader);
-	Rhs->EmitPreshader(Context, OutPreshader);
-	OutPreshader.WriteOpcode(Opcode);
-}
-
-void UE::HLSLTree::FExpressionBinaryOp::EmitHLSL(UE::HLSLTree::FEmitContext& Context, UE::HLSLTree::FCodeWriter& Writer) const
-{
-	const TCHAR* LhsRef = Context.AcquireHLSLReference(Lhs);
-	const TCHAR* RhsRef = Context.AcquireHLSLReference(Rhs);
-	Writer.Writef(TEXT("(%s + %s)"), LhsRef, RhsRef);
-}
-
-void UE::HLSLTree::FExpressionSwizzle::EmitPreshader(FEmitContext& Context, FMaterialPreshaderData& OutPreshader) const
-{
-	Input->EmitPreshader(Context, OutPreshader);
-	OutPreshader.WriteOpcode(EMaterialPreshaderOpcode::ComponentSwizzle)
-		.Write((uint8)Parameters.NumComponents)
-		.Write((uint8)Parameters.ComponentIndex[0])
-		.Write((uint8)Parameters.ComponentIndex[1])
-		.Write((uint8)Parameters.ComponentIndex[2])
-		.Write((uint8)Parameters.ComponentIndex[3]);
+		Context.AppendPreshader(LhsRef, OutResult.Preshader);
+		Context.AppendPreshader(RhsRef, OutResult.Preshader);
+		OutResult.Preshader.WriteOpcode(Opcode);
+	}
 }
 
-void UE::HLSLTree::FExpressionSwizzle::EmitHLSL(FEmitContext& Context, FCodeWriter& Writer) const
+void UE::HLSLTree::FExpressionSwizzle::EmitHLSL(FEmitContext& Context, FExpressionEmitResult& OutResult) const
 {
 	static const TCHAR ComponentName[] = { 'x', 'y', 'z', 'w' };
 	TCHAR Swizzle[5] = TEXT("");
@@ -301,18 +298,39 @@ void UE::HLSLTree::FExpressionSwizzle::EmitHLSL(FEmitContext& Context, FCodeWrit
 		Swizzle[i] = ComponentName[ComponentIndex];
 	}
 
-	const TCHAR* InputRef = Context.AcquireHLSLReference(Input);
-	Writer.Writef(TEXT("%s.%s"), InputRef, Swizzle);
+	const FEmitValue& InputRef = Context.AcquireValue(Input);
+
+	OutResult.bInline = true;
+	OutResult.EvaluationType = InputRef.GetEvaluationType();
+	if (OutResult.EvaluationType == EExpressionEvaluationType::Shader)
+	{
+		OutResult.Writer.Writef(TEXT("%s.%s"), Context.GetCode(InputRef), Swizzle);
+	}
+	else
+	{
+		Context.AppendPreshader(InputRef, OutResult.Preshader);
+		OutResult.Preshader.WriteOpcode(EMaterialPreshaderOpcode::ComponentSwizzle)
+			.Write((uint8)Parameters.NumComponents)
+			.Write((uint8)Parameters.ComponentIndex[0])
+			.Write((uint8)Parameters.ComponentIndex[1])
+			.Write((uint8)Parameters.ComponentIndex[2])
+			.Write((uint8)Parameters.ComponentIndex[3]);
+	}
 }
 
-void UE::HLSLTree::FExpressionCast::EmitHLSL(FEmitContext& Context, FCodeWriter& Writer) const
+void UE::HLSLTree::FExpressionCast::EmitHLSL(FEmitContext& Context, FExpressionEmitResult& OutResult) const
 {
 	const EExpressionType InputType = Input->Type;
-	const TCHAR* InputRef = Context.AcquireHLSLReference(Input);
+	const FEmitValue& InputRef = Context.AcquireValue(Input);
+	const TCHAR* InputCode = Context.GetCode(InputRef);
+
+	OutResult.bInline = true;
+	// TODO - preshader
+
 	if (InputType == Type)
 	{
 		// No cast required (nop)
-		Writer.Writef(TEXT("%s"), InputRef);
+		OutResult.Writer.Writef(TEXT("%s"), InputCode);
 	}
 
 	const FExpressionTypeDescription OutputTypeDesc = GetExpressionTypeDescription(Type);
@@ -320,7 +338,7 @@ void UE::HLSLTree::FExpressionCast::EmitHLSL(FEmitContext& Context, FCodeWriter&
 	if (InputTypeDesc.NumComponents == OutputTypeDesc.NumComponents)
 	{
 		// Cast between different underlying types with the same number of components, use a C-style cast
-		Writer.Writef(TEXT("((%s)%s)"), OutputTypeDesc.Name, InputRef);
+		OutResult.Writer.Writef(TEXT("((%s)%s)"), OutputTypeDesc.Name, InputCode);
 	}
 	else if (OutputTypeDesc.NumComponents < InputTypeDesc.NumComponents)
 	{
@@ -329,12 +347,12 @@ void UE::HLSLTree::FExpressionCast::EmitHLSL(FEmitContext& Context, FCodeWriter&
 		if (OutputTypeDesc.ComponentType == InputTypeDesc.ComponentType)
 		{
 			// underlying types are the same, no C-style cast needed
-			Writer.Writef(TEXT("%s.%s"), InputRef, SwizzleName[OutputTypeDesc.NumComponents]);
+			OutResult.Writer.Writef(TEXT("%s.%s"), InputCode, SwizzleName[OutputTypeDesc.NumComponents]);
 		}
 		else
 		{
 			// Different underlying types, add a C-style cast as well
-			Writer.Writef(TEXT("((%s)%s.%s)"), OutputTypeDesc.Name, InputRef, SwizzleName[OutputTypeDesc.NumComponents]);
+			OutResult.Writer.Writef(TEXT("((%s)%s.%s)"), OutputTypeDesc.Name, InputCode, SwizzleName[OutputTypeDesc.NumComponents]);
 		}
 	}
 	else
@@ -344,70 +362,98 @@ void UE::HLSLTree::FExpressionCast::EmitHLSL(FEmitContext& Context, FCodeWriter&
 		{
 			// Replicating a scalar value, use a swizzle
 			static const TCHAR* ScalarSwizzleName[] = { TEXT("x"), TEXT("xx"), TEXT("xxx"), TEXT("xxxx") };
-			Writer.Writef(TEXT("%s.%s"), InputRef, ScalarSwizzleName[OutputTypeDesc.NumComponents]);
+			OutResult.Writer.Writef(TEXT("%s.%s"), InputCode, ScalarSwizzleName[OutputTypeDesc.NumComponents]);
 		}
 		else
 		{
 			// Padding with 0s, use a constructor
-			Writer.Writef(TEXT("%s(%s"), OutputTypeDesc.Name, InputRef);
+			OutResult.Writer.Writef(TEXT("%s(%s"), OutputTypeDesc.Name, InputCode);
 			for (int32 i = InputTypeDesc.NumComponents; i < OutputTypeDesc.NumComponents; ++i)
 			{
 				// TODO - worth making these 0s the appropriate type? 0u, 0.0f, etc
-				Writer.Write(TEXT(", 0"));
+				OutResult.Writer.Write(TEXT(", 0"));
 			}
-			Writer.Write(TEXT(")"));
+			OutResult.Writer.Write(TEXT(")"));
 		}
 	}
 }
 
-void UE::HLSLTree::FExpressionFunctionInput::EmitHLSL(FEmitContext& Context, FCodeWriter& Writer) const
+void UE::HLSLTree::FExpressionFunctionInput::EmitHLSL(FEmitContext& Context, FExpressionEmitResult& OutResult) const
 {
 	const FEmitContext::FFunctionStackEntry& StackEntry = Context.FunctionStack.Last();
+
+	check(InputIndex >= 0 && InputIndex < StackEntry.FunctionCall->NumInputs);
 	FExpression* InputExpression = StackEntry.FunctionCall->Inputs[InputIndex];
-	const TCHAR* Ref = Context.AcquireHLSLReference(InputExpression);
-	Writer.Writef(TEXT("%s"), Ref);
+	const FEmitValue& Ref = Context.AcquireValue(InputExpression);
+
+	OutResult.bInline = true;
+	OutResult.EvaluationType = Ref.GetEvaluationType();
+	if (OutResult.EvaluationType == EExpressionEvaluationType::Shader)
+	{
+		OutResult.Writer.Writef(TEXT("%s"), Context.GetCode(Ref));
+	}
+	else
+	{
+		Context.AppendPreshader(Ref, OutResult.Preshader);
+	}
 }
 
-void UE::HLSLTree::FExpressionFunctionOutput::EmitHLSL(FEmitContext& Context, FCodeWriter& Writer) const
+void UE::HLSLTree::FExpressionFunctionOutput::EmitHLSL(FEmitContext& Context, FExpressionEmitResult& OutResult) const
 {
-	const TCHAR* Ref = Context.AcquireHLSLReference(FunctionCall, OutputIndex);
-	Writer.Writef(TEXT("%s"), Ref);
+	const FEmitValue& Ref = Context.AcquireValue(FunctionCall, OutputIndex);
+	
+	OutResult.bInline = true;
+	OutResult.EvaluationType = Ref.GetEvaluationType();
+	if (OutResult.EvaluationType == EExpressionEvaluationType::Shader)
+	{
+		OutResult.Writer.Writef(TEXT("%s"), Context.GetCode(Ref));
+	}
+	else
+	{
+		Context.AppendPreshader(Ref, OutResult.Preshader);
+	}
 }
 
-void UE::HLSLTree::FStatementSetFunctionOutput::EmitHLSL(FEmitContext& Context, FCodeWriter& Writer) const
+void UE::HLSLTree::FStatementSetFunctionOutput::EmitHLSL(FEmitContext& Context, FCodeWriter& OutWriter) const
 {
 	const FEmitContext::FFunctionStackEntry& StackEntry = Context.FunctionStack.Last();
-	const TCHAR* OutputRef = Context.AcquireHLSLReference(StackEntry.FunctionCall, OutputIndex);
-	const TCHAR* ExpressionRef = Context.AcquireHLSLReference(Expression);
-	Writer.WriteLinef(TEXT("%s = %s;"), OutputRef, ExpressionRef);
+	const FEmitValue& ExpressionRef = Context.AcquireValue(Expression);
+
+	check(OutputIndex >= 0 && OutputIndex < StackEntry.FunctionCall->NumOutputs);
+	FEmitValue& OutputRef = StackEntry.OutputRef[OutputIndex];
+	OutputRef.Assign(ExpressionRef);
+	if (OutputRef.GetEvaluationType() == EExpressionEvaluationType::Shader)
+	{
+		OutWriter.WriteLinef(TEXT("%s = %s;"), Context.GetCode(OutputRef), Context.GetCode(ExpressionRef));
+	}
 }
 
-void UE::HLSLTree::FStatementReturn::EmitHLSL(UE::HLSLTree::FEmitContext& Context, UE::HLSLTree::FCodeWriter& Writer) const
+void UE::HLSLTree::FStatementReturn::EmitHLSL(UE::HLSLTree::FEmitContext& Context, UE::HLSLTree::FCodeWriter& OutWriter) const
 {
-	const TCHAR* Ref = Context.AcquireHLSLReference(Expression);
-	Writer.WriteLinef(TEXT("return %s;"), Ref);
+	const FEmitValue& Ref = Context.AcquireValue(Expression);
+	OutWriter.WriteLinef(TEXT("return %s;"), Context.GetCode(Ref));
 }
 
-void UE::HLSLTree::FStatementSetLocalVariable::EmitHLSL(UE::HLSLTree::FEmitContext& Context, UE::HLSLTree::FCodeWriter& Writer) const
+void UE::HLSLTree::FStatementSetLocalVariable::EmitHLSL(UE::HLSLTree::FEmitContext& Context, UE::HLSLTree::FCodeWriter& OutWriter) const
 {
-	const TCHAR* DeclarationRef = Context.AcquireHLSLReference(Declaration);
-	const TCHAR* ExpressionRef = Context.AcquireHLSLReference(Expression);
-	Writer.WriteLinef(TEXT("%s = %s;"), DeclarationRef, ExpressionRef);
+	const FEmitValue& DeclarationRef = Context.AcquireValue(Declaration);
+	const FEmitValue& ExpressionRef = Context.AcquireValue(Expression);
+	OutWriter.WriteLinef(TEXT("%s = %s;"), Context.GetCode(DeclarationRef), Context.GetCode(ExpressionRef));
 }
 
-void UE::HLSLTree::FStatementIf::EmitHLSL(UE::HLSLTree::FEmitContext& Context, UE::HLSLTree::FCodeWriter& Writer) const
+void UE::HLSLTree::FStatementIf::EmitHLSL(UE::HLSLTree::FEmitContext& Context, UE::HLSLTree::FCodeWriter& OutWriter) const
 {
-	const TCHAR* ConditionRef = Context.AcquireHLSLReference(ConditionExpression);
-	Writer.WriteLinef(TEXT("if (%s)"), ConditionRef);
-	ThenScope->EmitHLSL(Context, Writer);
+	const FEmitValue& ConditionRef = Context.AcquireValue(ConditionExpression);
+	OutWriter.WriteLinef(TEXT("if (%s)"), Context.GetCode(ConditionRef));
+	ThenScope->EmitHLSL(Context, OutWriter);
 	if (ElseScope)
 	{
-		Writer.WriteLine(TEXT("else"));
-		ElseScope->EmitHLSL(Context, Writer);
+		OutWriter.WriteLine(TEXT("else"));
+		ElseScope->EmitHLSL(Context, OutWriter);
 	}
 }
 
-void UE::HLSLTree::FStatementFor::EmitHLSL(UE::HLSLTree::FEmitContext& Context, UE::HLSLTree::FCodeWriter& Writer) const
+void UE::HLSLTree::FStatementFor::EmitHLSL(UE::HLSLTree::FEmitContext& Context, UE::HLSLTree::FCodeWriter& OutWriter) const
 {
 	
 }
diff --git a/Engine/Source/Runtime/Engine/Private/Materials/MaterialHLSLGenerator.cpp b/Engine/Source/Runtime/Engine/Private/Materials/MaterialHLSLGenerator.cpp
--- a/Engine/Source/Runtime/Engine/Private/Materials/MaterialHLSLGenerator.cpp
+++ b/Engine/Source/Runtime/Engine/Private/Materials/MaterialHLSLGenerator.cpp
@@ -70,21 +70,18 @@ EMaterialGenerateHLSLStatus FMaterialHLSLGenerator::Error(const FString& Message
 UE::HLSLTree::FExpressionConstant* FMaterialHLSLGenerator::NewConstant(UE::HLSLTree::FScope& Scope, const UE::HLSLTree::FConstant& Value)
 {
 	UE::HLSLTree::FExpressionConstant* Expression = HLSLTree->NewExpression<UE::HLSLTree::FExpressionConstant>(Scope, Value);
-	Expression->bInline = true;
 	return Expression;
 }
 
 UE::HLSLTree::FExpressionExternalInput* FMaterialHLSLGenerator::NewTexCoord(UE::HLSLTree::FScope& Scope, int32 Index)
 {
 	UE::HLSLTree::FExpressionExternalInput* Expression = HLSLTree->NewExpression<UE::HLSLTree::FExpressionExternalInput>(Scope, UE::HLSLTree::MakeInputTexCoord(Index));
-	Expression->bInline = true;
 	return Expression;
 }
 
 UE::HLSLTree::FExpressionSwizzle* FMaterialHLSLGenerator::NewSwizzle(UE::HLSLTree::FScope& Scope, const UE::HLSLTree::FSwizzleParameters& Params, UE::HLSLTree::FExpression* Input)
 {
 	UE::HLSLTree::FExpressionSwizzle* Expression = HLSLTree->NewExpression<UE::HLSLTree::FExpressionSwizzle>(Scope, Params, Input);
-	Expression->bInline = true;
 	return Expression;
 }
 
@@ -94,7 +91,6 @@ UE::HLSLTree::FExpression* FMaterialHLSLGenerator::NewCast(UE::HLSLTree::FScope&
 	if (Input && Type != Input->Type)
 	{
 		UE::HLSLTree::FExpressionCast* Expression = HLSLTree->NewExpression<UE::HLSLTree::FExpressionCast>(Scope, Type, Input, Flags);
-		Expression->bInline = true;
 		return Expression;
 	}
 	return Input;
@@ -203,6 +199,7 @@ UE::HLSLTree::FFunctionCall* FMaterialHLSLGenerator::AcquireFunctionCall(UE::HLS
 	{
 		Hasher.Update((uint8*)&Input, sizeof(Input));
 	}
+	Hasher.Final();
 
 	FFunctionCallKey Key;
 	Key.Function = Function;
diff --git a/Engine/Source/Runtime/Engine/Private/Materials/MaterialUniformExpressions.cpp b/Engine/Source/Runtime/Engine/Private/Materials/MaterialUniformExpressions.cpp
--- a/Engine/Source/Runtime/Engine/Private/Materials/MaterialUniformExpressions.cpp
+++ b/Engine/Source/Runtime/Engine/Private/Materials/MaterialUniformExpressions.cpp
@@ -575,9 +575,35 @@ void FMaterialPreshaderData::WriteName(const FScriptName& Name)
 		Index = Names.Add(Name);
 	}
 	check(Index >= 0 && Index <= 0xffff);
+	const uint32 Offset = Data.Num();
+	NameOffsets.Add(Offset);
 	Write((uint16)Index);
 }
 
+void FMaterialPreshaderData::Append(const FMaterialPreshaderData& InPreshader)
+{
+	const uint32 BaseOffset = Data.Num();
+	Data.Append(InPreshader.Data);
+
+	TArray<uint16> NameIndexRemap;
+	NameIndexRemap.Empty(InPreshader.Names.Num());
+	for (const FScriptName& Name : InPreshader.Names)
+	{
+		const int32 Index = Names.AddUnique(Name);
+		check(Index >= 0 && Index <= 0xffff);
+		NameIndexRemap.Add(Index);
+	}
+
+	for (uint32 PrevOffset : InPreshader.NameOffsets)
+	{
+		const uint32 Offset = BaseOffset + PrevOffset;
+		uint16* PrevNameIndex = (uint16*)(Data.GetData() + Offset);
+		const uint16 RemapIndex = NameIndexRemap[*PrevNameIndex];
+		*PrevNameIndex = RemapIndex;
+		NameOffsets.Add(Offset);
+	}
+}
+
 namespace
 {
 	struct FPreshaderDataContext
diff --git a/Engine/Source/Runtime/Engine/Public/HLSLTree/HLSLTree.h b/Engine/Source/Runtime/Engine/Public/HLSLTree/HLSLTree.h
--- a/Engine/Source/Runtime/Engine/Public/HLSLTree/HLSLTree.h
+++ b/Engine/Source/Runtime/Engine/Public/HLSLTree/HLSLTree.h
@@ -67,21 +67,58 @@ public:
 
 	void WriteIndent();
 
+	void Reset();
+
+	void Append(const FCodeWriter& InWriter);
+
 	FStringBuilderBase* StringBuilder;
 	int32 IndentLevel;
 };
 
+class FEmitValue
+{
+public:
+	EExpressionEvaluationType GetEvaluationType() const { return EvaluationType; }
+	const FConstant& GetConstantValue() const { return ConstantValue; }
+
+	inline void Assign(const FEmitValue& InRef)
+	{
+		Preshader = InRef.Preshader;
+		EvaluationType = InRef.EvaluationType;
+		ConstantValue = InRef.ConstantValue;
+	}
+
+private:
+	mutable const TCHAR* Code = nullptr;
+	const FMaterialPreshaderData* Preshader = nullptr;
+	EExpressionEvaluationType EvaluationType = EExpressionEvaluationType::None;
+	EExpressionType ExpressionType = EExpressionType::Void;
+	FConstant ConstantValue;
+
+	friend class FEmitContext;
+};
+
 /** Tracks shared state while emitting HLSL code */
 class FEmitContext
 {
 public:
-	/** Returns a snippet of HLSL source code that references the given expression */
-	const TCHAR* AcquireHLSLReference(FExpression* Expression);
+	FEmitContext();
+	~FEmitContext();
+
+	/** Returns a value that references the given expression */
+	const FEmitValue& AcquireValue(FExpression* Expression);
+
+	/** Returns value that references the given local variable */
+	const FEmitValue& AcquireValue(FLocalDeclaration* Declaration);
+
+	/** Returns value that references the given local function output */
+	const FEmitValue& AcquireValue(FFunctionCall* FunctionCall, int32 OutputIndex);
 
-	/** Returns a snippet of HLSL source code that references the given local variablew */
-	const TCHAR* AcquireHLSLReference(FLocalDeclaration* Declaration);
+	/** Gets HLSL code that references the given value */
+	const TCHAR* GetCode(const FEmitValue& Value) const;
 
-	const TCHAR* AcquireHLSLReference(FFunctionCall* FunctionCall, int32 OutputIndex);
+	/** Append preshader bytecode that represents the given value */
+	void AppendPreshader(const FEmitValue& Value, FMaterialPreshaderData& InOutPreshader) const;
 
 	struct FScopeEntry
 	{
@@ -89,28 +126,31 @@ public:
 		FCodeWriter* ExpressionCodeWriter;
 	};
 
-	struct FFunctionStackEntry
-	{
-		FFunctionCall* FunctionCall;
-	};
-
 	struct FDeclarationEntry
 	{
-		const TCHAR* Definition;
+		FEmitValue Value;
 	};
 
 	struct FFunctionCallEntry
 	{
-		TCHAR const* const* OutputDefinition;
+		const FEmitValue* OutputRef;
 		int32 NumOutputs;
 	};
 
+	struct FFunctionStackEntry
+	{
+		FFunctionCall* FunctionCall = nullptr;
+		FEmitValue* OutputRef = nullptr;
+
+		TMap<FNode*, FDeclarationEntry> DeclarationMap;
+		TMap<FFunctionCall*, FFunctionCallEntry> FunctionCallMap;
+	};
+
 	FScopeEntry* FindScope(FScope* Scope);
 
 	TArray<FScopeEntry> ScopeStack;
 	TArray<FFunctionStackEntry> FunctionStack;
-	TMap<FNode*, FDeclarationEntry> DeclarationMap;
-	TMap<FFunctionCall*, FFunctionCallEntry> FunctionCallMap;
+	TArray<FMaterialPreshaderData*> TempPreshaders;
 	FMemStackBase* Allocator = nullptr;
 	const FMaterial* Material = nullptr; // TODO - remove preshader material dependency
 	FMaterialCompilationOutput* MaterialCompilationOutput = nullptr;
@@ -118,6 +158,21 @@ public:
 	int32 NumTexCoords = 0;
 };
 
+struct FExpressionEmitResult
+{
+	FExpressionEmitResult(FCodeWriter& InWriter, FMaterialPreshaderData& InPreshader)
+		: Writer(InWriter)
+		, Preshader(InPreshader)
+		, EvaluationType(EExpressionEvaluationType::Shader)
+		, bInline(false)
+	{}
+
+	FCodeWriter& Writer;
+	FMaterialPreshaderData& Preshader;
+	EExpressionEvaluationType EvaluationType;
+	bool bInline;
+};
+
 enum class ENodeVisitResult
 {
 	VisitDependentNodes,
@@ -181,28 +236,17 @@ public:
 class FExpression : public FNode
 {
 public:
-	inline bool IsInline() const { return bInline; }
-
 	virtual ENodeVisitResult Visit(FNodeVisitor& Visitor) override;
 
 	/* Emits HLSL code for the expression. The code should NOT include any newlines or semi-colons. The returned string may be assigned to a temporary variable, or embedded in another HLSL string */
-	virtual void EmitHLSL(FEmitContext& Context, FCodeWriter& Writer) const;
-
-	/* Emits bytecode for the preshader virtual machine */
-	virtual void EmitPreshader(FEmitContext& Context, FMaterialPreshaderData& OutPreshader) const;
+	virtual void EmitHLSL(FEmitContext& Context, FExpressionEmitResult& OutResult) const = 0;
 
 	/** The HLSL type of the expression (float, int, etc) */
 	EExpressionType Type;
 
-	EExpressionEvaluationType EvaluationType;
-	bool bInline;
-
 protected:
-	explicit FExpression(EExpressionType InType,
-		EExpressionEvaluationType InEvaluationType = EExpressionEvaluationType::Shader)
+	explicit FExpression(EExpressionType InType)
 		: Type(InType)
-		, EvaluationType(InEvaluationType)
-		, bInline(false)
 	{}
 };
 
diff --git a/Engine/Source/Runtime/Engine/Public/HLSLTree/HLSLTreeCommon.h b/Engine/Source/Runtime/Engine/Public/HLSLTree/HLSLTreeCommon.h
--- a/Engine/Source/Runtime/Engine/Public/HLSLTree/HLSLTreeCommon.h
+++ b/Engine/Source/Runtime/Engine/Public/HLSLTree/HLSLTreeCommon.h
@@ -21,13 +21,12 @@ class FExpressionConstant : public FExpression
 {
 public:
 	explicit FExpressionConstant(const FConstant& InValue)
-		: FExpression(InValue.Type, EExpressionEvaluationType::Constant), Value(InValue)
+		: FExpression(InValue.Type), Value(InValue)
 	{}
 
 	FConstant Value;
 
-	virtual void EmitPreshader(FEmitContext& Context, FMaterialPreshaderData& OutPreshader) const override;
-	virtual void EmitHLSL(FEmitContext& Context, FCodeWriter& Writer) const override;
+	virtual void EmitHLSL(FEmitContext& Context, FExpressionEmitResult& OutResult) const override;
 };
 
 class FExpressionLocalVariable : public FExpression
@@ -47,14 +46,14 @@ public:
 		return Result;
 	}
 
-	virtual void EmitHLSL(FEmitContext& Context, FCodeWriter& Writer) const override;
+	virtual void EmitHLSL(FEmitContext& Context, FExpressionEmitResult& OutResult) const override;
 };
 
 class FExpressionParameter : public FExpression
 {
 public:
 	FExpressionParameter(FParameterDeclaration* InDeclaration)
-		: FExpression(InDeclaration->DefaultValue.Type, EExpressionEvaluationType::Preshader)
+		: FExpression(InDeclaration->DefaultValue.Type)
 		, Declaration(InDeclaration)
 	{}
 
@@ -70,7 +69,7 @@ public:
 		return Result;
 	}
 
-	virtual void EmitPreshader(FEmitContext& Context, FMaterialPreshaderData& OutPreshader) const override;
+	virtual void EmitHLSL(FEmitContext& Context, FExpressionEmitResult& OutResult) const override;
 };
 
 enum class EExternalInputType
@@ -101,7 +100,7 @@ public:
 
 	EExternalInputType InputType;
 
-	virtual void EmitHLSL(FEmitContext& Context, FCodeWriter& Writer) const override;
+	virtual void EmitHLSL(FEmitContext& Context, FExpressionEmitResult& OutResult) const override;
 };
 
 class FExpressionTextureSample : public FExpression
@@ -131,7 +130,7 @@ public:
 		return Result;
 	}
 
-	virtual void EmitHLSL(FEmitContext& Context, FCodeWriter& Writer) const override;
+	virtual void EmitHLSL(FEmitContext& Context, FExpressionEmitResult& OutResult) const override;
 };
 
 class FExpressionDefaultMaterialAttributes : public FExpression
@@ -139,7 +138,7 @@ class FExpressionDefaultMaterialAttributes : public FExpression
 public:
 	FExpressionDefaultMaterialAttributes() : FExpression(EExpressionType::MaterialAttributes) {}
 
-	virtual void EmitHLSL(FEmitContext& Context, FCodeWriter& Writer) const override;
+	virtual void EmitHLSL(FEmitContext& Context, FExpressionEmitResult& OutResult) const override;
 };
 
 class FExpressionSetMaterialAttribute : public FExpression
@@ -162,14 +161,14 @@ public:
 		return Result;
 	}
 
-	virtual void EmitHLSL(FEmitContext& Context, FCodeWriter& Writer) const override;
+	virtual void EmitHLSL(FEmitContext& Context, FExpressionEmitResult& OutResult) const override;
 };
 
 class FExpressionBinaryOp : public FExpression
 {
 public:
 	FExpressionBinaryOp(EExpressionType ResultType, EBinaryOp InOp, FExpression* InLhs, FExpression* InRhs)
-		: FExpression(ResultType, CombineEvaluationTypes(InLhs->EvaluationType, InRhs->EvaluationType))
+		: FExpression(ResultType)
 		, Op(InOp)
 		, Lhs(InLhs)
 		, Rhs(InRhs)
@@ -190,8 +189,7 @@ public:
 		return Result;
 	}
 
-	virtual void EmitPreshader(FEmitContext& Context, FMaterialPreshaderData& OutPreshader) const override;
-	virtual void EmitHLSL(FEmitContext& Context, FCodeWriter& Writer) const override;
+	virtual void EmitHLSL(FEmitContext& Context, FExpressionEmitResult& OutResult) const override;
 };
 
 struct FSwizzleParameters
@@ -208,7 +206,7 @@ class FExpressionSwizzle : public FExpression
 {
 public:
 	FExpressionSwizzle(const FSwizzleParameters& InParams, FExpression* InInput)
-		: FExpression(MakeExpressionType(InInput->Type, InParams.NumComponents), InInput->EvaluationType)
+		: FExpression(MakeExpressionType(InInput->Type, InParams.NumComponents))
 		, Parameters(InParams)
 		, Input(InInput)
 	{}
@@ -226,15 +224,14 @@ public:
 		return Result;
 	}
 
-	virtual void EmitPreshader(FEmitContext& Context, FMaterialPreshaderData& OutPreshader) const override;
-	virtual void EmitHLSL(FEmitContext& Context, FCodeWriter& Writer) const override;
+	virtual void EmitHLSL(FEmitContext& Context, FExpressionEmitResult& OutResult) const override;
 };
 
 class FExpressionCast : public FExpression
 {
 public:
 	FExpressionCast(EExpressionType InType, FExpression* InInput, ECastFlags InFlags = ECastFlags::None)
-		: FExpression(InType) // TODO preshader
+		: FExpression(InType)
 		, Input(InInput)
 		, Flags(InFlags)
 	{}
@@ -252,7 +249,7 @@ public:
 		return Result;
 	}
 
-	virtual void EmitHLSL(FEmitContext& Context, FCodeWriter& Writer) const override;
+	virtual void EmitHLSL(FEmitContext& Context, FExpressionEmitResult& OutResult) const override;
 };
 
 class FExpressionFunctionInput : public FExpression
@@ -262,7 +259,7 @@ public:
 		: FExpression(InType), Name(InName), InputIndex(InIndex)
 	{}
 
-	virtual void EmitHLSL(FEmitContext& Context, FCodeWriter& Writer) const override;
+	virtual void EmitHLSL(FEmitContext& Context, FExpressionEmitResult& OutResult) const override;
 
 	FName Name;
 	int32 InputIndex;
@@ -292,7 +289,7 @@ public:
 		return Result;
 	}
 
-	virtual void EmitHLSL(FEmitContext& Context, FCodeWriter& Writer) const override;
+	virtual void EmitHLSL(FEmitContext& Context, FExpressionEmitResult& OutResult) const override;
 };
 
 class FStatementSetFunctionOutput : public FStatement
diff --git a/Engine/Source/Runtime/Engine/Public/HLSLTree/HLSLTreeTypes.h b/Engine/Source/Runtime/Engine/Public/HLSLTree/HLSLTreeTypes.h
--- a/Engine/Source/Runtime/Engine/Public/HLSLTree/HLSLTreeTypes.h
+++ b/Engine/Source/Runtime/Engine/Public/HLSLTree/HLSLTreeTypes.h
@@ -59,6 +59,9 @@ EExpressionType MakeArithmeticResultType(EExpressionType Lhs, EExpressionType Rh
  * Describes how a given expression needs to be evaluated */
 enum class EExpressionEvaluationType
 {
+	/** Invalid/uninitialized */
+	None,
+
 	/** The expression needs to generate HLSL code (EmitHLSL will be called) */
 	Shader,
 
diff --git a/Engine/Source/Runtime/Engine/Public/MaterialShared.h b/Engine/Source/Runtime/Engine/Public/MaterialShared.h
--- a/Engine/Source/Runtime/Engine/Public/MaterialShared.h
+++ b/Engine/Source/Runtime/Engine/Public/MaterialShared.h
@@ -305,6 +305,8 @@ public:
 
 	void Evaluate(FUniformExpressionSet* UniformExpressionSet, const struct FMaterialRenderContext& Context, FLinearColor& OutValue);
 
+	void Append(const FMaterialPreshaderData& InPreshader);
+
 	const int32 Num() const { return Data.Num(); }
 
 	void WriteData(const void* Value, uint32 Size);
@@ -326,6 +328,7 @@ public:
 	inline FMaterialPreshaderData& WriteOpcode(EMaterialPreshaderOpcode Op) { return Write<uint8>((uint8)Op); }
 
 	LAYOUT_FIELD(TMemoryImageArray<FScriptName>, Names);
+	LAYOUT_FIELD(TMemoryImageArray<uint32>, NameOffsets);
 	LAYOUT_FIELD(TMemoryImageArray<uint8>, Data);
 };
 
