Minor syntax tweak to hopefully make static analyzer happy

diff --git a/Engine/Source/Runtime/Engine/Private/Materials/MaterialHLSLGenerator.cpp b/Engine/Source/Runtime/Engine/Private/Materials/MaterialHLSLGenerator.cpp
--- a/Engine/Source/Runtime/Engine/Private/Materials/MaterialHLSLGenerator.cpp
+++ b/Engine/Source/Runtime/Engine/Private/Materials/MaterialHLSLGenerator.cpp
@@ -197,7 +197,7 @@ UE::HLSLTree::FFunctionCall* FMaterialHLSLGenerator::AcquireFunctionCall(UE::HLS
 	FSHA1 Hasher;
 	for (UE::HLSLTree::FExpression* Input : Inputs)
 	{
-		Hasher.Update((uint8*)&Input, sizeof(Input));
+		Hasher.Update((uint8*)&Input, sizeof(UE::HLSLTree::FExpression*));
 	}
 	Hasher.Final();
 
