Produce an extra error message to be included in test output (UE-109465).

#rnx

diff --git a/Engine/Source/Runtime/Engine/Private/ShaderCompiler/ShaderCompiler.cpp b/Engine/Source/Runtime/Engine/Private/ShaderCompiler/ShaderCompiler.cpp
--- a/Engine/Source/Runtime/Engine/Private/ShaderCompiler/ShaderCompiler.cpp
+++ b/Engine/Source/Runtime/Engine/Private/ShaderCompiler/ShaderCompiler.cpp
@@ -3975,6 +3975,13 @@ void FShaderMapCompileResults::CheckIfHung()
 		if (DurationSoFar >= static_cast<double>(GShaderMapCompilationTimeout))
 		{
 			bIsHung = true;
+			// always produce an error message first, even if going to crash, as the automation controller does not seem to be picking up Fatal messages
+			UE_LOG(LogShaderCompilers, Error, TEXT("Hung shadermap detected, time spent compiling: %f seconds, NumPendingJobs: %d, FinishedJobs: %d"),
+				DurationSoFar,
+				NumPendingJobs.GetValue(),
+				FinishedJobs.Num()
+			);
+
 			if (GCrashOnHungShaderMaps)
 			{
 				UE_LOG(LogShaderCompilers, Fatal, TEXT("Crashing on a hung shadermap, time spent compiling: %f seconds, NumPendingJobs: %d, FinishedJobs: %d"),
@@ -3983,14 +3990,6 @@ void FShaderMapCompileResults::CheckIfHung()
 					FinishedJobs.Num()
 				);
 			}
-			else
-			{
-				UE_LOG(LogShaderCompilers, Error, TEXT("Hung shadermap detected, time spent compiling: %f seconds, NumPendingJobs: %d, FinishedJobs: %d"),
-					DurationSoFar,
-					NumPendingJobs.GetValue(),
-					FinishedJobs.Num()
-				);
-			}
 		}
 	}
 }
