Properly merge finished jobs (UE-109210).

- A similar issue lead in the past to unfinished shadermaps being cooked or the compiler getting stuck

diff --git a/Engine/Source/Runtime/Engine/Private/ShaderCompiler/ShaderCompiler.cpp b/Engine/Source/Runtime/Engine/Private/ShaderCompiler/ShaderCompiler.cpp
--- a/Engine/Source/Runtime/Engine/Private/ShaderCompiler/ShaderCompiler.cpp
+++ b/Engine/Source/Runtime/Engine/Private/ShaderCompiler/ShaderCompiler.cpp
@@ -3156,6 +3156,27 @@ FProcHandle FShaderCompilingManager::LaunchWorker(const FString& WorkingDirector
 	}
 }
 
+void FShaderCompilingManager::AddCompiledResults(TMap<int32, FShaderMapFinalizeResults>& CompiledShaderMaps, int32 ShaderMapIdx, const FShaderMapFinalizeResults& Results)
+{
+	// merge with the previous unprocessed jobs, if any
+	if (FShaderMapCompileResults const* PrevResults = CompiledShaderMaps.Find(ShaderMapIdx))
+	{
+		FShaderMapFinalizeResults NewResults(Results);
+
+		NewResults.bAllJobsSucceeded = NewResults.bAllJobsSucceeded && PrevResults->bAllJobsSucceeded;
+		NewResults.bSkipResultProcessing = NewResults.bSkipResultProcessing || PrevResults->bSkipResultProcessing;
+		NewResults.TimeStarted = FMath::Min(NewResults.TimeStarted, PrevResults->TimeStarted);
+		NewResults.bIsHung = NewResults.bIsHung || PrevResults->bIsHung;
+		NewResults.FinishedJobs.Append(PrevResults->FinishedJobs);
+
+		CompiledShaderMaps.Add(ShaderMapIdx, NewResults);
+	}
+	else
+	{
+		CompiledShaderMaps.Add(ShaderMapIdx, Results);
+	}
+}
+
 /** Flushes all pending jobs for the given shader maps. */
 void FShaderCompilingManager::BlockOnShaderMapCompletion(const TArray<int32>& ShaderMapIdsToFinishCompiling, TMap<int32, FShaderMapFinalizeResults>& CompiledShaderMaps)
 {
@@ -3207,19 +3228,7 @@ void FShaderCompilingManager::BlockOnShaderMapCompletion(const TArray<int32>& Sh
 						{
 							if (Results->FinishedJobs.Num() > 0)
 							{
-								FShaderMapFinalizeResults NewResults(*Results);
-
-								// merge with the previous unprocessed jobs, if any
-								if (FShaderMapCompileResults const* PrevResults = CompiledShaderMaps.Find(ShaderMapIdsToFinishCompiling[ShaderMapIndex]))
-								{
-									NewResults.bAllJobsSucceeded = NewResults.bAllJobsSucceeded && PrevResults->bAllJobsSucceeded;
-									NewResults.bSkipResultProcessing = NewResults.bSkipResultProcessing || PrevResults->bSkipResultProcessing;
-									NewResults.TimeStarted = FMath::Min(NewResults.TimeStarted, PrevResults->TimeStarted);
-									NewResults.bIsHung = NewResults.bIsHung || PrevResults->bIsHung;
-									NewResults.FinishedJobs.Append(PrevResults->FinishedJobs);
-								}
-
-								CompiledShaderMaps.Add(ShaderMapIdsToFinishCompiling[ShaderMapIndex], NewResults);
+								AddCompiledResults(CompiledShaderMaps, ShaderMapIdsToFinishCompiling[ShaderMapIndex], *Results);
 							}
 							ShaderMapJobs.Remove(ShaderMapIdsToFinishCompiling[ShaderMapIndex]);
 						}
@@ -3286,8 +3295,8 @@ void FShaderCompilingManager::BlockOnShaderMapCompletion(const TArray<int32>& Sh
 				const FShaderMapCompileResults* Results = *ResultsPtr;
 				check(Results->NumPendingJobs.GetValue() == 0);
 				check(Results->FinishedJobs.Num() > 0);
-				ensureMsgf(CompiledShaderMaps.Find(ShaderMapIdsToFinishCompiling[ShaderMapIndex]) == nullptr, TEXT("We're likely losing existing shadermap compilation results"));
-				CompiledShaderMaps.Add(ShaderMapIdsToFinishCompiling[ShaderMapIndex], FShaderMapFinalizeResults(*Results));
+
+				AddCompiledResults(CompiledShaderMaps, ShaderMapIdsToFinishCompiling[ShaderMapIndex], *Results);
 				ShaderMapJobs.Remove(ShaderMapIdsToFinishCompiling[ShaderMapIndex]);
 			}
 		}
@@ -3336,8 +3345,7 @@ void FShaderCompilingManager::BlockOnAllShaderMapCompletion(TMap<int32, FShaderM
 
 					if (Results->NumPendingJobs.GetValue() == 0)
 					{
-						ensureMsgf(CompiledShaderMaps.Find(It.Key()) == nullptr, TEXT("We're likely losing existing shadermap compilation results"));
-						CompiledShaderMaps.Add(It.Key(), FShaderMapFinalizeResults(*Results));
+						AddCompiledResults(CompiledShaderMaps, It.Key(), *Results);
 						It.RemoveCurrent();
 					}
 					else
@@ -3396,8 +3404,8 @@ void FShaderCompilingManager::BlockOnAllShaderMapCompletion(TMap<int32, FShaderM
 		{
 			const FShaderMapCompileResults* Results = It.Value();
 			check(Results->NumPendingJobs.GetValue()== 0);
-			ensureMsgf(CompiledShaderMaps.Find(It.Key()) == nullptr, TEXT("We're likely losing existing shadermap compilation results"));
-			CompiledShaderMaps.Add(It.Key(), FShaderMapFinalizeResults(*Results));
+
+			AddCompiledResults(CompiledShaderMaps, It.Key(), *Results);
 			It.RemoveCurrent();
 		}
 	}
diff --git a/Engine/Source/Runtime/Engine/Public/ShaderCompiler.h b/Engine/Source/Runtime/Engine/Public/ShaderCompiler.h
--- a/Engine/Source/Runtime/Engine/Public/ShaderCompiler.h
+++ b/Engine/Source/Runtime/Engine/Public/ShaderCompiler.h
@@ -1024,6 +1024,9 @@ private:
 	/** Blocks on completion of all shader maps. */
 	void BlockOnAllShaderMapCompletion(TMap<int32, FShaderMapFinalizeResults>& CompiledShaderMaps);
 
+	/** Adds compiled results to the CompiledShaderMaps, merging with the existing ones as necessary. */
+	void AddCompiledResults(TMap<int32, FShaderMapFinalizeResults>& CompiledShaderMaps, int32 ShaderMapIdx, const FShaderMapFinalizeResults& Results);
+
 	/** Finalizes the given shader map results and optionally assigns the affected shader maps to materials, while attempting to stay within an execution time budget. */
 	void ProcessCompiledShaderMaps(TMap<int32, FShaderMapFinalizeResults>& CompiledShaderMaps, float TimeBudget);
 
