Allow World Partition streaming source providers to be used in SIE

diff --git a/Engine/Source/Runtime/Engine/Private/WorldPartition/WorldPartitionStreamingPolicy.cpp b/Engine/Source/Runtime/Engine/Private/WorldPartition/WorldPartitionStreamingPolicy.cpp
--- a/Engine/Source/Runtime/Engine/Private/WorldPartition/WorldPartitionStreamingPolicy.cpp
+++ b/Engine/Source/Runtime/Engine/Private/WorldPartition/WorldPartitionStreamingPolicy.cpp
@@ -56,31 +56,32 @@ void UWorldPartitionStreamingPolicy::UpdateStreamingSources()
 		const FRotator ViewRotationLocal = WorldToLocal.TransformRotation(ViewRotation.Quaternion()).Rotator();
 		static const FName NAME_SIE(TEXT("SIE"));
 		StreamingSources.Add(FWorldPartitionStreamingSource(NAME_SIE, ViewLocationLocal, ViewRotationLocal));
-		return;
 	}
+	else
 #endif
-
-	UWorld* World = WorldPartition->GetWorld();
-	if (World->GetNetMode() != NM_DedicatedServer)
 	{
-		const int32 NumPlayers = GEngine->GetNumGamePlayers(World);
-		for (int32 PlayerIndex = 0; PlayerIndex < NumPlayers; ++PlayerIndex)
+		UWorld* World = WorldPartition->GetWorld();
+		if (World->GetNetMode() != NM_DedicatedServer)
 		{
-			ULocalPlayer* Player = GEngine->GetGamePlayer(World, PlayerIndex);
-			if (Player && Player->PlayerController)
+			const int32 NumPlayers = GEngine->GetNumGamePlayers(World);
+			for (int32 PlayerIndex = 0; PlayerIndex < NumPlayers; ++PlayerIndex)
 			{
-				FVector ViewLocation;
-				FRotator ViewRotation;
-				Player->PlayerController->GetPlayerViewPoint(ViewLocation, ViewRotation);
-
-				//@todo_ow: this test is to cover cases where GetPlayerViewPoint returns (0,0,0) when invalid. It should probably return a bool
-				//          to indicate that the returned position is invalid.
-				if (!ViewLocation.IsZero())
+				ULocalPlayer* Player = GEngine->GetGamePlayer(World, PlayerIndex);
+				if (Player && Player->PlayerController)
 				{
-					// Transform to Local
-					ViewLocation = WorldToLocal.TransformPosition(ViewLocation);
-					ViewRotation = WorldToLocal.TransformRotation(ViewRotation.Quaternion()).Rotator();
-					StreamingSources.Add(FWorldPartitionStreamingSource(Player->GetFName(), ViewLocation, ViewRotation));
+					FVector ViewLocation;
+					FRotator ViewRotation;
+					Player->PlayerController->GetPlayerViewPoint(ViewLocation, ViewRotation);
+
+					//@todo_ow: this test is to cover cases where GetPlayerViewPoint returns (0,0,0) when invalid. It should probably return a bool
+					//          to indicate that the returned position is invalid.
+					if (!ViewLocation.IsZero())
+					{
+						// Transform to Local
+						ViewLocation = WorldToLocal.TransformPosition(ViewLocation);
+						ViewRotation = WorldToLocal.TransformRotation(ViewRotation.Quaternion()).Rotator();
+						StreamingSources.Add(FWorldPartitionStreamingSource(Player->GetFName(), ViewLocation, ViewRotation));
+					}
 				}
 			}
 		}
