Added support of running with -game for partitioned worlds.
Removed support of partitioned subworlds referenced by LevelInstances.

diff --git a/Engine/Source/Editor/WorldPartitionEditor/Private/WorldPartitionEditorModule.cpp b/Engine/Source/Editor/WorldPartitionEditor/Private/WorldPartitionEditorModule.cpp
--- a/Engine/Source/Editor/WorldPartitionEditor/Private/WorldPartitionEditorModule.cpp
+++ b/Engine/Source/Editor/WorldPartitionEditor/Private/WorldPartitionEditorModule.cpp
@@ -84,12 +84,15 @@ void FWorldPartitionEditorModule::StartupModule()
 	SWorldPartitionEditorGrid::RegisterPartitionEditorGridCreateInstanceFunc(NAME_None, &SWorldPartitionEditorGrid::CreateInstance);
 	SWorldPartitionEditorGrid::RegisterPartitionEditorGridCreateInstanceFunc(TEXT("SpatialHash"), &SWorldPartitionEditorGridSpatialHash::CreateInstance);
 	
-	FLevelEditorModule& LevelEditorModule = FModuleManager::Get().LoadModuleChecked<FLevelEditorModule>("LevelEditor");
-	TArray<FLevelEditorModule::FLevelViewportMenuExtender_SelectedActors>& MenuExtenderDelegates = LevelEditorModule.GetAllLevelViewportContextMenuExtenders();
+	if (!IsRunningGame())
+	{
+		FLevelEditorModule& LevelEditorModule = FModuleManager::Get().LoadModuleChecked<FLevelEditorModule>("LevelEditor");
+		TArray<FLevelEditorModule::FLevelViewportMenuExtender_SelectedActors>& MenuExtenderDelegates = LevelEditorModule.GetAllLevelViewportContextMenuExtenders();
+
+		MenuExtenderDelegates.Add(FLevelEditorModule::FLevelViewportMenuExtender_SelectedActors::CreateStatic(&OnExtendLevelEditorMenu));
+		LevelEditorExtenderDelegateHandle = MenuExtenderDelegates.Last().GetHandle();
+	}
 
-	MenuExtenderDelegates.Add(FLevelEditorModule::FLevelViewportMenuExtender_SelectedActors::CreateStatic(&OnExtendLevelEditorMenu));
-	LevelEditorExtenderDelegateHandle = MenuExtenderDelegates.Last().GetHandle();
-	
 	IAssetTools& AssetTools = FModuleManager::LoadModuleChecked<FAssetToolsModule>("AssetTools").Get();
 	HLODLayerAssetTypeActions = MakeShareable(new FHLODLayerAssetTypeActions);
 	AssetTools.RegisterAssetTypeActions(HLODLayerAssetTypeActions.ToSharedRef());
@@ -99,9 +102,12 @@ void FWorldPartitionEditorModule::ShutdownModule()
 {
 	UWorldPartition::WorldPartitionChangedEvent.RemoveAll(this);
 	
-	if (FLevelEditorModule* LevelEditorModule = FModuleManager::Get().GetModulePtr<FLevelEditorModule>("LevelEditor"))
+	if (!IsRunningGame())
 	{
-		LevelEditorModule->GetAllLevelViewportContextMenuExtenders().RemoveAll([=](const FLevelEditorModule::FLevelViewportMenuExtender_SelectedActors& In) { return In.GetHandle() == LevelEditorExtenderDelegateHandle; });
+		if (FLevelEditorModule* LevelEditorModule = FModuleManager::Get().GetModulePtr<FLevelEditorModule>("LevelEditor"))
+		{
+			LevelEditorModule->GetAllLevelViewportContextMenuExtenders().RemoveAll([=](const FLevelEditorModule::FLevelViewportMenuExtender_SelectedActors& In) { return In.GetHandle() == LevelEditorExtenderDelegateHandle; });
+		}
 	}
 
 	IAssetTools* AssetToolsModule = nullptr;
diff --git a/Engine/Source/Runtime/Engine/Private/Components/WorldPartitionStreamingSourceComponent.cpp b/Engine/Source/Runtime/Engine/Private/Components/WorldPartitionStreamingSourceComponent.cpp
--- a/Engine/Source/Runtime/Engine/Private/Components/WorldPartitionStreamingSourceComponent.cpp
+++ b/Engine/Source/Runtime/Engine/Private/Components/WorldPartitionStreamingSourceComponent.cpp
@@ -18,7 +18,7 @@ void UWorldPartitionStreamingSourceComponent::OnRegister()
 	UWorld* World = GetWorld();
 
 #if WITH_EDITOR
-	if (!World->IsPlayInEditor())
+	if (!World->IsGameWorld())
 	{
 		return;
 	}
@@ -37,7 +37,7 @@ void UWorldPartitionStreamingSourceComponent::OnUnregister()
 	UWorld* World = GetWorld();
 
 #if WITH_EDITOR
-	if (!World->IsPlayInEditor())
+	if (!World->IsGameWorld())
 	{
 		return;
 	}
diff --git a/Engine/Source/Runtime/Engine/Private/Level.cpp b/Engine/Source/Runtime/Engine/Private/Level.cpp
--- a/Engine/Source/Runtime/Engine/Private/Level.cpp
+++ b/Engine/Source/Runtime/Engine/Private/Level.cpp
@@ -1788,7 +1788,7 @@ namespace LevelAssetRegistryHelper
 
 bool ULevel::GetIsLevelPartitionedFromPackage(FName LevelPackage)
 {
-	return LevelAssetRegistryHelper::GetLevelInfoFromAssetRegistry(LevelPackage, [](const FAssetData& Asset)
+	return LevelPackage.IsNone() ? false : LevelAssetRegistryHelper::GetLevelInfoFromAssetRegistry(LevelPackage, [](const FAssetData& Asset)
 	{
 		FString LevelIsPartitionedStr;
 		static const FName NAME_LevelIsPartitioned(TEXT("LevelIsPartitioned"));
diff --git a/Engine/Source/Runtime/Engine/Private/LevelInstance/LevelInstanceActor.cpp b/Engine/Source/Runtime/Engine/Private/LevelInstance/LevelInstanceActor.cpp
--- a/Engine/Source/Runtime/Engine/Private/LevelInstance/LevelInstanceActor.cpp
+++ b/Engine/Source/Runtime/Engine/Private/LevelInstance/LevelInstanceActor.cpp
@@ -225,6 +225,14 @@ ALevelInstance::FOnLevelInstanceActorPostLoad ALevelInstance::OnLevelInstanceAct
 void ALevelInstance::PostLoad()
 {
 	Super::PostLoad();
+
+#if WITH_EDITOR
+	if (ULevel::GetIsLevelPartitionedFromPackage(*WorldAsset.GetLongPackageName()))
+	{
+		UE_LOG(LogLevelInstance, Warning, TEXT("LevelInstance doesn't support partitioned world %s"), *WorldAsset.GetLongPackageName());
+		WorldAsset.Reset();
+	}
+#endif
 		
 	OnLevelInstanceActorPostLoad.Broadcast(this);
 }
@@ -388,6 +396,16 @@ bool ALevelInstance::CanSetValue(TSoftObjectPtr<UWorld> InLevelInstance, FString
 	TArray<TPair<FText, TSoftObjectPtr<UWorld>>> LoopInfo;
 	const ALevelInstance* LoopStart = nullptr;
 
+	if (ULevel::GetIsLevelPartitionedFromPackage(*InLevelInstance.GetLongPackageName()))
+	{
+		if (Reason)
+		{
+			*Reason = FString::Format(TEXT("LevelInstance doesn't support partitioned world {0}\n"), { InLevelInstance.GetLongPackageName() });
+		}
+
+		return false;
+	}
+
 	if (!CheckForLoop(InLevelInstance, Reason ? &LoopInfo : nullptr, Reason ? &LoopStart : nullptr))
 	{
 		if (Reason)
@@ -569,7 +587,7 @@ bool ALevelInstance::IsLoaded() const
 
 void ALevelInstance::OnLevelInstanceLoaded()
 {
-	if (!GetWorld()->IsPlayInEditor())
+	if (!GetWorld()->IsGameWorld())
 	{
 		// Propagate bounds dirtyness up and check if we need to hide our LevelInstance because self or ancestor is hidden
 		bool bHiddenInEditor = false;
diff --git a/Engine/Source/Runtime/Engine/Private/LevelInstance/LevelInstanceEditorLevelStreaming.cpp b/Engine/Source/Runtime/Engine/Private/LevelInstance/LevelInstanceEditorLevelStreaming.cpp
--- a/Engine/Source/Runtime/Engine/Private/LevelInstance/LevelInstanceEditorLevelStreaming.cpp
+++ b/Engine/Source/Runtime/Engine/Private/LevelInstance/LevelInstanceEditorLevelStreaming.cpp
@@ -25,7 +25,7 @@ ULevelStreamingLevelInstanceEditor::ULevelStreamingLevelInstanceEditor(const FOb
 #if WITH_EDITOR
 	SetShouldBeVisibleInEditor(true);
 
-	if (!IsTemplate() && !GetWorld()->IsPlayInEditor())
+	if (!IsTemplate() && !GetWorld()->IsGameWorld())
 	{
 		GEngine->OnLevelActorAdded().AddUObject(this, &ULevelStreamingLevelInstanceEditor::OnLevelActorAdded);
 	}
diff --git a/Engine/Source/Runtime/Engine/Private/LevelInstance/LevelInstanceLevelStreaming.cpp b/Engine/Source/Runtime/Engine/Private/LevelInstance/LevelInstanceLevelStreaming.cpp
--- a/Engine/Source/Runtime/Engine/Private/LevelInstance/LevelInstanceLevelStreaming.cpp
+++ b/Engine/Source/Runtime/Engine/Private/LevelInstance/LevelInstanceLevelStreaming.cpp
@@ -72,7 +72,7 @@ ULevelStreamingLevelInstance* ULevelStreamingLevelInstance::LoadInstance(ALevelI
 		LevelStreaming->LevelInstanceID = LevelInstanceActor->GetLevelInstanceID();
 		
 #if WITH_EDITOR
-		if (!LevelInstanceActor->GetWorld()->IsPlayInEditor())
+		if (!LevelInstanceActor->GetWorld()->IsGameWorld())
 		{
 			GEngine->BlockTillLevelStreamingCompleted(LevelInstanceActor->GetWorld());
 
diff --git a/Engine/Source/Runtime/Engine/Private/World.cpp b/Engine/Source/Runtime/Engine/Private/World.cpp
--- a/Engine/Source/Runtime/Engine/Private/World.cpp
+++ b/Engine/Source/Runtime/Engine/Private/World.cpp
@@ -2761,11 +2761,6 @@ void UWorld::AddToWorld( ULevel* Level, const FTransform& LevelTransform, bool b
 			// send a callback that a level was added to the world
 			FWorldDelegates::LevelAddedToWorld.Broadcast(Level, this);
 
-			if (UWorldPartition* WorldPartition = Level->GetWorldPartition())
-			{
-				WorldPartition->Initialize(Level->GetWorld(), LevelTransform);
-			}
-
 			BroadcastLevelsChanged();
 
 			ULevelStreaming::BroadcastLevelVisibleStatus(this, Level->GetOutermost()->GetFName(), true);
diff --git a/Engine/Source/Runtime/Engine/Private/WorldPartition/LevelInstance/LevelInstanceActorDesc.cpp b/Engine/Source/Runtime/Engine/Private/WorldPartition/LevelInstance/LevelInstanceActorDesc.cpp
--- a/Engine/Source/Runtime/Engine/Private/WorldPartition/LevelInstance/LevelInstanceActorDesc.cpp
+++ b/Engine/Source/Runtime/Engine/Private/WorldPartition/LevelInstance/LevelInstanceActorDesc.cpp
@@ -43,7 +43,7 @@ void FLevelInstanceActorDesc::RegisterContainer()
 	check(!LevelInstanceContainer);
 	if (Container && (DesiredRuntimeBehavior == ELevelInstanceRuntimeBehavior::Embedded || DesiredRuntimeBehavior == ELevelInstanceRuntimeBehavior::Partitioned))
 	{
-		if (!LevelPackage.IsNone() && ULevel::GetIsLevelUsingExternalActorsFromPackage(LevelPackage))
+		if (!LevelPackage.IsNone() && ULevel::GetIsLevelUsingExternalActorsFromPackage(LevelPackage) && !ULevel::GetIsLevelPartitionedFromPackage(LevelPackage))
 		{
 			UWorldPartition* WorldPartition = Container->GetWorld()->GetWorldPartition();
 			LevelInstanceContainer = WorldPartition->RegisterActorDescContainer(LevelPackage);
diff --git a/Engine/Source/Runtime/Engine/Private/WorldPartition/WorldPartition.cpp b/Engine/Source/Runtime/Engine/Private/WorldPartition/WorldPartition.cpp
--- a/Engine/Source/Runtime/Engine/Private/WorldPartition/WorldPartition.cpp
+++ b/Engine/Source/Runtime/Engine/Private/WorldPartition/WorldPartition.cpp
@@ -36,6 +36,7 @@
 #include "WorldPartition/WorldPartitionRuntimeSpatialHash.h"
 #include "WorldPartition/WorldPartitionEditorPerProjectUserSettings.h"
 #include "Modules/ModuleManager.h"
+#include "GameDelegates.h"
 #endif //WITH_EDITOR
 
 DEFINE_LOG_CATEGORY(LogWorldPartition);
@@ -52,7 +53,7 @@ static FAutoConsoleCommand DumpActorDescs(
 		{
 			if (UWorld* World = GEditor->GetEditorWorldContext().World())
 			{
-				if (!World->IsPlayInEditor())
+				if (!World->IsGameWorld())
 				{
 					if (UWorldPartition* WorldPartition = World->GetWorldPartition())
 					{
@@ -152,15 +153,24 @@ void UWorldPartition::OnPreBeginPIE(bool bStartSimulate)
 	}
 
 	check(IsMainWorldPartition());
-	GenerateStreaming(EWorldPartitionStreamingMode::PIE);
-	RuntimeHash->OnPreBeginPIE();
+
+	OnBeginPlay(EWorldPartitionStreamingMode::PIE);
+}
+
+void UWorldPartition::OnBeginPlay(EWorldPartitionStreamingMode Mode)
+{
+	GenerateStreaming(Mode);
+	RuntimeHash->OnBeginPlay(Mode);
+	GetStreamingPolicy()->OnBeginPlay();
 }
 
-void UWorldPartition::OnEndPIE(bool bStartSimulate)
+void UWorldPartition::OnEndPlay()
 {
 	check(IsMainWorldPartition());
+
 	FlushStreaming();
-	RuntimeHash->OnEndPIE();
+	RuntimeHash->OnEndPlay();
+	GetStreamingPolicy()->OnEndPlay();
 
 	World->PersistentLevel->ActorsModifiedForPIE.Empty();
 }
@@ -197,8 +207,10 @@ void UWorldPartition::Initialize(UWorld* InWorld, const FTransform& InTransform)
 	UWorld* OuterWorld = GetTypedOuter<UWorld>();
 	check(OuterWorld);
 
+	check(IsMainWorldPartition());
+
 #if WITH_EDITOR
-	bool bEditorOnly = !World->IsPlayInEditor();
+	bool bEditorOnly = !World->IsGameWorld();
 	if (bEditorOnly)
 	{
 		if (!EditorHash)
@@ -218,7 +230,7 @@ void UWorldPartition::Initialize(UWorld* InWorld, const FTransform& InTransform)
 		RuntimeHash = RuntimeSpatialHash;
 	}
 
-	if (bEditorOnly || !IsMainWorldPartition())
+	if (bEditorOnly || IsRunningGame())
 	{
 		TArray<FAssetData> Assets;
 		UPackage* LevelPackage = OuterWorld->PersistentLevel->GetOutermost();
@@ -240,7 +252,7 @@ void UWorldPartition::Initialize(UWorld* InWorld, const FTransform& InTransform)
 			ReplaceTo = DestWorldName + TEXT(".") + DestWorldName;
 		}
 
-		const bool bRegisterDelegates = bEditorOnly && IsMainWorldPartition();
+		const bool bRegisterDelegates = bEditorOnly;
 		UActorDescContainer::Initialize(World, PackageName, bRegisterDelegates);
 		check(bContainerInitialized);
 
@@ -264,52 +276,50 @@ void UWorldPartition::Initialize(UWorld* InWorld, const FTransform& InTransform)
 				HashActorDesc(*ActorDescIterator);
 			}
 		}
+	}
 
-		if (bEditorOnly)
-		{
-			// Make sure to preload any AInfo actors first (mainly for ShouldActorBeLoaded as it needs to have access to the WorldDataLayers actor)
-			TArray<UWorldPartitionEditorCell::FActorHandle> AlwaysLoadedActors = EditorHash->GetAlwaysLoadedCell()->Actors.Array();
-			AlwaysLoadedActors.Sort([](const UWorldPartitionEditorCell::FActorHandle& A, const UWorldPartitionEditorCell::FActorHandle& B) { return A->GetActorClass()->IsChildOf<AInfo>(); });
+	if (bEditorOnly)
+	{
+		// Make sure to preload any AInfo actors first (mainly for ShouldActorBeLoaded as it needs to have access to the WorldDataLayers actor)
+		TArray<UWorldPartitionEditorCell::FActorHandle> AlwaysLoadedActors = EditorHash->GetAlwaysLoadedCell()->Actors.Array();
+		AlwaysLoadedActors.Sort([](const UWorldPartitionEditorCell::FActorHandle& A, const UWorldPartitionEditorCell::FActorHandle& B) { return A->GetActorClass()->IsChildOf<AInfo>(); });
 			
-			TArray<FWorldPartitionReference> LoadedActors;
-			LoadedActors.Reserve(AlwaysLoadedActors.Num());
-			Algo::Transform(AlwaysLoadedActors, LoadedActors, [](const UWorldPartitionEditorCell::FActorHandle& ActorHandle) { return ActorHandle.Handle; });
+		TArray<FWorldPartitionReference> LoadedActors;
+		LoadedActors.Reserve(AlwaysLoadedActors.Num());
+		Algo::Transform(AlwaysLoadedActors, LoadedActors, [](const UWorldPartitionEditorCell::FActorHandle& ActorHandle) { return ActorHandle.Handle; });
 
-			// Load the always loaded cell, don't call LoadCells to avoid creating a transaction
-			UpdateLoadingEditorCell(EditorHash->GetAlwaysLoadedCell(), true);
+		// Load the always loaded cell, don't call LoadCells to avoid creating a transaction
+		UpdateLoadingEditorCell(EditorHash->GetAlwaysLoadedCell(), true);
 
-			// Load more cells depending on the user's settings
-			// Skipped when running from a commandlet
-			if (!IsRunningCommandlet())
-			{
-				// Autoload all cells if the world is smaller than the project setting's value
-				IWorldPartitionEditorModule& WorldPartitionEditorModule = FModuleManager::LoadModuleChecked<IWorldPartitionEditorModule>("WorldPartitionEditor");
-				const float AutoCellLoadingMaxWorldSize = WorldPartitionEditorModule.GetAutoCellLoadingMaxWorldSize();
-				FVector WorldSize = GetEditorWorldBounds().GetSize();
-				const bool bItsASmallWorld = WorldSize.X <= AutoCellLoadingMaxWorldSize && WorldSize.Y <= AutoCellLoadingMaxWorldSize && WorldSize.Z <= AutoCellLoadingMaxWorldSize;
+		// Load more cells depending on the user's settings
+		// Skipped when running from a commandlet
+		if (!IsRunningCommandlet())
+		{
+			// Autoload all cells if the world is smaller than the project setting's value
+			IWorldPartitionEditorModule& WorldPartitionEditorModule = FModuleManager::LoadModuleChecked<IWorldPartitionEditorModule>("WorldPartitionEditor");
+			const float AutoCellLoadingMaxWorldSize = WorldPartitionEditorModule.GetAutoCellLoadingMaxWorldSize();
+			FVector WorldSize = GetEditorWorldBounds().GetSize();
+			const bool bItsASmallWorld = WorldSize.X <= AutoCellLoadingMaxWorldSize && WorldSize.Y <= AutoCellLoadingMaxWorldSize && WorldSize.Z <= AutoCellLoadingMaxWorldSize;
 			
-				// When loading a subworld, load all actors
-				const bool bWorldIsSubPartition = !IsMainWorldPartition();
-
-				if (bWorldIsSubPartition || bItsASmallWorld)
+			// When considered as a small world, load all actors
+			if (bItsASmallWorld)
+			{
+				EditorHash->ForEachCell([this](UWorldPartitionEditorCell* Cell)
 				{
-					EditorHash->ForEachCell([this](UWorldPartitionEditorCell* Cell)
-					{
-						UpdateLoadingEditorCell(Cell, true);
-					});
-				}
+					UpdateLoadingEditorCell(Cell, true);
+				});
+			}
 
-				// Load last loaded cells
-				if (WorldPartitionEditorModule.GetEnableLoadingOfLastLoadedCells())
-				{
-					const TArray<FName>& EditorGridLastLoadedCells = GetMutableDefault<UWorldPartitionEditorPerProjectUserSettings>()->GetEditorGridLoadedCells(InWorld);
+			// Load last loaded cells
+			if (WorldPartitionEditorModule.GetEnableLoadingOfLastLoadedCells())
+			{
+				const TArray<FName>& EditorGridLastLoadedCells = GetMutableDefault<UWorldPartitionEditorPerProjectUserSettings>()->GetEditorGridLoadedCells(InWorld);
 
-					for (FName EditorGridLastLoadedCell : EditorGridLastLoadedCells)
+				for (FName EditorGridLastLoadedCell : EditorGridLastLoadedCells)
+				{
+					if (UWorldPartitionEditorCell* Cell = FindObject<UWorldPartitionEditorCell>(EditorHash, *EditorGridLastLoadedCell.ToString()))
 					{
-						if (UWorldPartitionEditorCell* Cell = FindObject<UWorldPartitionEditorCell>(EditorHash, *EditorGridLastLoadedCell.ToString()))
-						{
-							UpdateLoadingEditorCell(Cell, true);
-						}
+						UpdateLoadingEditorCell(Cell, true);
 					}
 				}
 			}
@@ -326,9 +336,9 @@ void UWorldPartition::Initialize(UWorld* InWorld, const FTransform& InTransform)
 	}
 
 #if WITH_EDITOR
-	if (World->IsPlayInEditor())
+	if (IsRunningGame())
 	{
-		PrepareForPIE();
+		OnBeginPlay(EWorldPartitionStreamingMode::EditorStandalone);
 	}
 #endif
 }
@@ -376,9 +386,9 @@ void UWorldPartition::Uninitialize()
 		}
 		ActorDescContainers.Empty();
 
-		if (World->IsPlayInEditor())
+		if (World->IsGameWorld())
 		{
-			CleanupForPIE();
+			OnEndPlay();
 		}
 		else
 		{
@@ -433,7 +443,7 @@ void UWorldPartition::RegisterDelegates()
 	if (GEditor && !IsTemplate())
 	{
 		FEditorDelegates::PreBeginPIE.AddUObject(this, &UWorldPartition::OnPreBeginPIE);
-		FEditorDelegates::EndPIE.AddUObject(this, &UWorldPartition::OnEndPIE);
+		FGameDelegates::Get().GetEndPlayMapDelegate().AddUObject(this, &UWorldPartition::OnEndPlay);
 	}
 }
 
@@ -443,7 +453,7 @@ void UWorldPartition::UnregisterDelegates()
 	if (GEditor && !IsTemplate())
 	{
 		FEditorDelegates::PreBeginPIE.RemoveAll(this);
-		FEditorDelegates::EndPIE.RemoveAll(this);
+		FGameDelegates::Get().GetEndPlayMapDelegate().RemoveAll(this);
 	}
 }
 
@@ -931,29 +941,10 @@ void UWorldPartition::DrawRuntimeHashPreview()
 	RuntimeHash->DrawPreview();
 }
 
-void UWorldPartition::PrepareForPIE()
-{
-	check(World->IsPlayInEditor());
-
-	if (!IsMainWorldPartition())
-	{
-		GenerateStreaming(EWorldPartitionStreamingMode::PIE);
-	}
-
-	GetStreamingPolicy()->PrepareForPIE();
-}
-
-void UWorldPartition::CleanupForPIE()
-{
-	check(World->IsPlayInEditor());
-
-	FlushStreaming();
-}
-
 bool UWorldPartition::GenerateStreaming(EWorldPartitionStreamingMode Mode, TArray<FString>* OutPackagesToGenerate)
 {
 	check(RuntimeHash);
-	return RuntimeHash->GenerateStreaming(Mode, GetStreamingPolicy(), OutPackagesToGenerate);
+	return RuntimeHash->GenerateRuntimeStreaming(Mode, GetStreamingPolicy(), OutPackagesToGenerate);
 }
 
 bool UWorldPartition::PopulateGeneratedPackageForCook(UPackage* InPackage, const FString& InPackageRelativePath, const FString& InPackageCookName)
@@ -1040,11 +1031,11 @@ void UWorldPartition::DumpActorDescs(const FString& Path)
 	}
 }
 
-void UWorldPartition::OnPreFixupForPIE(int32 InPIEInstanceID, FSoftObjectPath& ObjectPath)
+void UWorldPartition::RemapSoftObjectPath(FSoftObjectPath& ObjectPath)
 {
-	if (GetWorld()->IsPlayInEditor())
+	if (GetWorld()->IsGameWorld())
 	{
-		GetStreamingPolicy()->OnPreFixupForPIE(InPIEInstanceID, ObjectPath);
+		GetStreamingPolicy()->RemapSoftObjectPath(ObjectPath);
 	}
 }
 
diff --git a/Engine/Source/Runtime/Engine/Private/WorldPartition/WorldPartitionCookPackageSplitter.cpp b/Engine/Source/Runtime/Engine/Private/WorldPartition/WorldPartitionCookPackageSplitter.cpp
--- a/Engine/Source/Runtime/Engine/Private/WorldPartition/WorldPartitionCookPackageSplitter.cpp
+++ b/Engine/Source/Runtime/Engine/Private/WorldPartition/WorldPartitionCookPackageSplitter.cpp
@@ -38,7 +38,7 @@ TArray<ICookPackageSplitter::FGeneratedPackage> FWorldPartitionCookPackageSplitt
 	WorldPartition->Initialize(PartitionedWorld, FTransform::Identity);
 
 	TArray<FString> WorldPartitionGeneratedPackages;
-	WorldPartition->GenerateStreaming(EWorldPartitionStreamingMode::RuntimeStreamingCells, &WorldPartitionGeneratedPackages);
+	WorldPartition->GenerateStreaming(EWorldPartitionStreamingMode::Cook, &WorldPartitionGeneratedPackages);
 
 	TArray<ICookPackageSplitter::FGeneratedPackage> PackagesToGenerate;
 	PackagesToGenerate.Reserve(WorldPartitionGeneratedPackages.Num());
diff --git a/Engine/Source/Runtime/Engine/Private/WorldPartition/WorldPartitionLevelHelper.cpp b/Engine/Source/Runtime/Engine/Private/WorldPartition/WorldPartitionLevelHelper.cpp
--- a/Engine/Source/Runtime/Engine/Private/WorldPartition/WorldPartitionLevelHelper.cpp
+++ b/Engine/Source/Runtime/Engine/Private/WorldPartition/WorldPartitionLevelHelper.cpp
@@ -112,7 +112,7 @@ ULevel* FWorldPartitionLevelHelper::CreateEmptyLevelForRuntimeCell(const UWorld*
 	return NewLevel;
 }
 
-bool FWorldPartitionLevelHelper::LoadActors(ULevel* InDestLevel, TArrayView<FWorldPartitionRuntimeCellObjectMapping> InActorPackages, FWorldPartitionPackageCache& InPackageCache, TFunction<void(bool)> InCompletionCallback, bool bLoadForPIE, bool bLoadAsync /*=true*/, FLinkerInstancingContext* InOutInstancingContext /*=nullptr*/)
+bool FWorldPartitionLevelHelper::LoadActors(ULevel* InDestLevel, TArrayView<FWorldPartitionRuntimeCellObjectMapping> InActorPackages, FWorldPartitionPackageCache& InPackageCache, TFunction<void(bool)> InCompletionCallback, bool bLoadForPlay, bool bLoadAsync /*=true*/, FLinkerInstancingContext* InOutInstancingContext /*=nullptr*/)
 {
 	UPackage* DestPackage = InDestLevel ? InDestLevel->GetPackage() : nullptr;
 	FString ShortLevelPackageName = DestPackage? FPackageName::GetShortName(DestPackage->GetFName()) : FString();
@@ -132,13 +132,13 @@ bool FWorldPartitionLevelHelper::LoadActors(ULevel* InDestLevel, TArrayView<FWor
 
 	// Levels to Load with actors to duplicate
 	TMap<uint32, TArray<FWorldPartitionRuntimeCellObjectMapping*>> PackagesToDuplicate;
-	check(!bLoadForPIE || InOutInstancingContext);
+	check(!bLoadForPlay || InOutInstancingContext);
 
 	for (FWorldPartitionRuntimeCellObjectMapping& PackageObjectMapping : InActorPackages)
 	{
 		if (PackageObjectMapping.ContainerID == 0)
 		{
-			if (bLoadForPIE)
+			if (bLoadForPlay)
 			{
 				FString ObjectPath = PackageObjectMapping.Package.ToString();
 				FString ActorPackageName = FPackageName::ObjectPathToPackageName(ObjectPath);
@@ -164,7 +164,7 @@ bool FWorldPartitionLevelHelper::LoadActors(ULevel* InDestLevel, TArrayView<FWor
 		const FString ActorPackageName = FPackageName::ObjectPathToPackageName(PackageObjectMapping.Package.ToString());
 		FName ActorName = *FPaths::GetExtension(PackageObjectMapping.Path.ToString());
 
-		FLoadPackageAsyncDelegate CompletionCallback = FLoadPackageAsyncDelegate::CreateLambda([LoadProgress, ActorName, bLoadForPIE, InDestLevel, InCompletionCallback](const FName& LoadedPackageName, UPackage* LoadedPackage, EAsyncLoadingResult::Type Result)
+		FLoadPackageAsyncDelegate CompletionCallback = FLoadPackageAsyncDelegate::CreateLambda([LoadProgress, ActorName, bLoadForPlay, InDestLevel, InCompletionCallback](const FName& LoadedPackageName, UPackage* LoadedPackage, EAsyncLoadingResult::Type Result)
 			{
 				check(LoadProgress->NumPendingLoadRequests);
 				LoadProgress->NumPendingLoadRequests--;
@@ -173,7 +173,7 @@ bool FWorldPartitionLevelHelper::LoadActors(ULevel* InDestLevel, TArrayView<FWor
 					AActor* Actor = FindObject<AActor>(LoadedPackage, *ActorName.ToString());
 
 					check(Actor);
-					if (bLoadForPIE)
+					if (bLoadForPlay)
 					{
 						check(Actor->IsPackageExternal());
 						InDestLevel->Actors.Add(Actor);
@@ -197,10 +197,11 @@ bool FWorldPartitionLevelHelper::LoadActors(ULevel* InDestLevel, TArrayView<FWor
 
 
 		FPackagePath PackagePath = FPackagePath::FromPackageNameChecked(*ActorPackageName);
-		if (bLoadForPIE)
+		if (bLoadForPlay)
 		{
 			check(DestPackage);
-			::LoadPackageAsync(PackagePath, FName(*ActorPackageInstanceNames[ChildIndex]), CompletionCallback, nullptr, PKG_PlayInEditor, DestPackage->PIEInstanceID, 0, InOutInstancingContext);
+			const EPackageFlags PackageFlags = DestPackage->HasAnyPackageFlags(PKG_PlayInEditor) ? PKG_PlayInEditor : PKG_None;
+			::LoadPackageAsync(PackagePath, FName(*ActorPackageInstanceNames[ChildIndex]), CompletionCallback, nullptr, PackageFlags, DestPackage->PIEInstanceID, 0, InOutInstancingContext);
 		}
 		else
 		{
diff --git a/Engine/Source/Runtime/Engine/Private/WorldPartition/WorldPartitionLevelStreamingDynamic.cpp b/Engine/Source/Runtime/Engine/Private/WorldPartition/WorldPartitionLevelStreamingDynamic.cpp
--- a/Engine/Source/Runtime/Engine/Private/WorldPartition/WorldPartitionLevelStreamingDynamic.cpp
+++ b/Engine/Source/Runtime/Engine/Private/WorldPartition/WorldPartitionLevelStreamingDynamic.cpp
@@ -160,7 +160,7 @@ bool UWorldPartitionLevelStreamingDynamic::RequestLevel(UWorld* InPersistentWorl
 		check(CellWorld->PersistentLevel != LoadedLevel);
 
 		// Level already exists but may have the wrong type due to being inactive before, so copy data over
-		check(InPersistentWorld->IsPlayInEditor());
+		check(InPersistentWorld->IsGameWorld());
 		CellWorld->WorldType = InPersistentWorld->WorldType;
 		CellWorld->PersistentLevel->OwningWorld = InPersistentWorld;
 
@@ -185,7 +185,7 @@ bool UWorldPartitionLevelStreamingDynamic::RequestLevel(UWorld* InPersistentWorl
 			check(CellLevelPackage);
 			check(UWorld::FindWorldInPackage(CellLevelPackage));
 			check(RuntimeLevel->OwningWorld);
-			check(RuntimeLevel->OwningWorld->WorldType == EWorldType::PIE);
+			check(RuntimeLevel->OwningWorld->WorldType == EWorldType::PIE || (IsRunningGame() && RuntimeLevel->OwningWorld->WorldType == EWorldType::Game));
 	
 			if (IssueLoadRequests())
 			{
@@ -309,7 +309,7 @@ bool UWorldPartitionLevelStreamingDynamic::IssueLoadRequests()
 		{
 			UE_LOG(LogLevelStreaming, Fatal, TEXT("UWorldPartitionLevelStreamingDynamic::IssueLoadRequests failed %s"), *GetWorldAssetPackageName());
 		}
-	}, /*bLoadForPie=*/true, /*bLoadAsync=*/true, &InstancingContext);
+	}, /*bLoadForPlay=*/true, /*bLoadAsync=*/true, &InstancingContext);
 
 	return bLoadRequestInProgress;
 }
@@ -321,24 +321,62 @@ void UWorldPartitionLevelStreamingDynamic::FinalizeRuntimeLevel()
 	check(!bLoadRequestInProgress);
 
 	// For RuntimeLevel's world NetGUID to be valid, make sure to flag bIsNameStableForNetworking so that IsNameStableForNetworking() returns true. (see FNetGUIDCache::SupportsObject)
-	RuntimeLevel->GetTypedOuter<UWorld>()->bIsNameStableForNetworking = true;
+	UWorld* OuterWorld = RuntimeLevel->GetTypedOuter<UWorld>();
+	OuterWorld->bIsNameStableForNetworking = true;
 
 	UPackage* RuntimePackage = RuntimeLevel->GetPackage();
 	RuntimePackage->MarkAsFullyLoaded();
 
-	int32 PIEInstanceID = GetPackage()->PIEInstanceID;
-	check(PIEInstanceID != INDEX_NONE);
+	if (OuterWorld->IsPlayInEditor())
+	{
+		int32 PIEInstanceID = GetPackage()->PIEInstanceID;
+		check(PIEInstanceID != INDEX_NONE);
 
-	// PIE Fixup LazyObjectPtrs
-	FTemporaryPlayInEditorIDOverride SetPlayInEditorID(PIEInstanceID);
-	FFixupLazyObjectPtrForPIEArchive FixupLazyPointersAr;
-	FixupLazyPointersAr << RuntimeLevel;
+		// PIE Fixup LazyObjectPtrs
+		FTemporaryPlayInEditorIDOverride SetPlayInEditorID(PIEInstanceID);
+		FFixupLazyObjectPtrForPIEArchive FixupLazyPointersAr;
+		FixupLazyPointersAr << RuntimeLevel;
 
-	// PIE Fixup SoftObjectPaths
-	RuntimeLevel->FixupForPIE(PIEInstanceID, [&](int32 InPIEInstanceID, FSoftObjectPath& ObjectPath)
+		// PIE Fixup SoftObjectPaths
+		RuntimeLevel->FixupForPIE(PIEInstanceID, [&](int32 InPIEInstanceID, FSoftObjectPath& ObjectPath)
+		{
+			OuterWorldPartition->RemapSoftObjectPath(ObjectPath);
+		});
+	}
+	else if (IsRunningGame())
 	{
-		OuterWorldPartition->OnPreFixupForPIE(InPIEInstanceID, ObjectPath);
-	});
+		check(OuterWorld->IsGameWorld());
+
+		struct FSoftPathFixupSerializer : public FArchiveUObject
+		{
+			FSoftPathFixupSerializer(TFunctionRef<void(FSoftObjectPath&)> InCustomFixupFunction)
+			: CustomFixupFunction(InCustomFixupFunction)
+			{
+				this->SetIsSaving(true);
+			}
+
+			FArchive& operator<<(FSoftObjectPath& Value)
+			{
+				if (!Value.IsNull())
+				{
+					CustomFixupFunction(Value);
+				}
+				return *this;
+			}
+
+			TFunctionRef<void(FSoftObjectPath&)> CustomFixupFunction;
+		};
+
+		FSoftPathFixupSerializer FixupSerializer([&](FSoftObjectPath& ObjectPath) { OuterWorldPartition->RemapSoftObjectPath(ObjectPath); });
+		
+		TArray<UObject*> SubObjects;
+		GetObjectsWithOuter(RuntimeLevel, SubObjects);
+		
+		for (UObject* Object : SubObjects)
+		{
+			Object->Serialize(FixupSerializer);
+		}
+	}
 
 	SetLoadedLevel(RuntimeLevel);
 
diff --git a/Engine/Source/Runtime/Engine/Private/WorldPartition/WorldPartitionLevelStreamingPolicy.cpp b/Engine/Source/Runtime/Engine/Private/WorldPartition/WorldPartitionLevelStreamingPolicy.cpp
--- a/Engine/Source/Runtime/Engine/Private/WorldPartition/WorldPartitionLevelStreamingPolicy.cpp
+++ b/Engine/Source/Runtime/Engine/Private/WorldPartition/WorldPartitionLevelStreamingPolicy.cpp
@@ -126,7 +126,7 @@ void UWorldPartitionLevelStreamingPolicy::UnloadCell(const UWorldPartitionRuntim
 
 FString UWorldPartitionLevelStreamingPolicy::GetCellPackagePath(const FName& InCellName, const UWorld* InWorld)
 {
-	if (InWorld->IsPlayInEditor())
+	if (InWorld->IsGameWorld())
 	{
 		// Set as memory package to avoid wasting time in FPackageName::DoesPackageExist
 		return FString::Printf(TEXT("/Memory/%s"), *InCellName.ToString());
@@ -142,7 +142,7 @@ TSubclassOf<UWorldPartitionRuntimeCell> UWorldPartitionLevelStreamingPolicy::Get
 	return UWorldPartitionRuntimeLevelStreamingCell::StaticClass();
 }
 
-void UWorldPartitionLevelStreamingPolicy::PrepareForPIE()
+void UWorldPartitionLevelStreamingPolicy::OnBeginPlay()
 {
 	TSet<const UWorldPartitionRuntimeCell*> StreamingCells;
 	WorldPartition->RuntimeHash->GetAllStreamingCells(StreamingCells, /*bIncludeDataLayers*/ true);
@@ -159,7 +159,12 @@ void UWorldPartitionLevelStreamingPolicy::PrepareForPIE()
 	}
 }
 
-void UWorldPartitionLevelStreamingPolicy::OnPreFixupForPIE(int32 InPIEInstanceID, FSoftObjectPath& ObjectPath)
+void UWorldPartitionLevelStreamingPolicy::OnEndPlay()
+{
+	ActorToCellRemapping.Empty();
+}
+
+void UWorldPartitionLevelStreamingPolicy::RemapSoftObjectPath(FSoftObjectPath& ObjectPath)
 {
 	// Once we fix up the path to the proper streaming Level, normal flow of FixupForPIE will adapt the path for PIE.
 	FString SrcPath = ObjectPath.ToString();
diff --git a/Engine/Source/Runtime/Engine/Private/WorldPartition/WorldPartitionRuntimeHash.cpp b/Engine/Source/Runtime/Engine/Private/WorldPartition/WorldPartitionRuntimeHash.cpp
--- a/Engine/Source/Runtime/Engine/Private/WorldPartition/WorldPartitionRuntimeHash.cpp
+++ b/Engine/Source/Runtime/Engine/Private/WorldPartition/WorldPartitionRuntimeHash.cpp
@@ -6,20 +6,48 @@
 #include "Engine/World.h"
 #include "Engine/LevelScriptBlueprint.h"
 #include "ActorReferencesUtils.h"
+#if WITH_EDITOR
+#include "WorldPartition/WorldPartitionHandle.h"
+#endif
 
 UWorldPartitionRuntimeHash::UWorldPartitionRuntimeHash(const FObjectInitializer& ObjectInitializer)
 	: Super(ObjectInitializer)
 {}
 
 #if WITH_EDITOR
-void UWorldPartitionRuntimeHash::OnPreBeginPIE()
+bool UWorldPartitionRuntimeHash::GenerateRuntimeStreaming(EWorldPartitionStreamingMode Mode, UWorldPartitionStreamingPolicy* Policy, TArray<FString>* OutPackagesToGenerate)
+{
+	FWorldPartitionReference LoadedWorldDataLayers;
+	if (Mode == EWorldPartitionStreamingMode::EditorStandalone)
+	{
+		// Pre-load AWorldDataLayer as it is necessary for GenerateStreaming to world properly
+		UWorldPartition* WorldPartition = GetOuterUWorldPartition();
+		for (UActorDescContainer::TConstIterator<> ActorDescIterator(WorldPartition); ActorDescIterator; ++ActorDescIterator)
+		{
+			if (ActorDescIterator->GetActorClass()->IsChildOf<AWorldDataLayers>())
+			{
+				if (ensure(!LoadedWorldDataLayers.IsLoaded()))
+				{
+					LoadedWorldDataLayers = FWorldPartitionReference(WorldPartition, ActorDescIterator->GetGuid());
+				}
+			}
+		}
+	}
+
+	return GenerateStreaming(Mode, Policy, OutPackagesToGenerate);
+}
+
+void UWorldPartitionRuntimeHash::OnBeginPlay(EWorldPartitionStreamingMode Mode)
 {
-	// Mark always loaded actors so that the Level will force reference to these actors for PIE.
-	// These actor will then be duplicated for PIE during the PIE world duplication process
-	ForceExternalActorLevelReference(/*bForceExternalActorLevelReferenceForPIE*/true);
+	if (Mode == EWorldPartitionStreamingMode::PIE)
+	{
+		// Mark always loaded actors so that the Level will force reference to these actors for PIE.
+		// These actor will then be duplicated for PIE during the PIE world duplication process
+		ForceExternalActorLevelReference(/*bForceExternalActorLevelReferenceForPIE*/true);
+	}
 }
 
-void UWorldPartitionRuntimeHash::OnEndPIE()
+void UWorldPartitionRuntimeHash::OnEndPlay()
 {
 	// Unmark always loaded actors
 	ForceExternalActorLevelReference(/*bForceExternalActorLevelReferenceForPIE*/false);
diff --git a/Engine/Source/Runtime/Engine/Private/WorldPartition/WorldPartitionRuntimeLevelStreamingCell.cpp b/Engine/Source/Runtime/Engine/Private/WorldPartition/WorldPartitionRuntimeLevelStreamingCell.cpp
--- a/Engine/Source/Runtime/Engine/Private/WorldPartition/WorldPartitionRuntimeLevelStreamingCell.cpp
+++ b/Engine/Source/Runtime/Engine/Private/WorldPartition/WorldPartitionRuntimeLevelStreamingCell.cpp
@@ -97,7 +97,7 @@ ULevelStreaming* UWorldPartitionRuntimeLevelStreamingCell::CreateLevelStreaming(
 void UWorldPartitionRuntimeLevelStreamingCell::LoadActorsForCook()
 {
 	FWorldPartitionPackageCache PackageCache;
-	verify(FWorldPartitionLevelHelper::LoadActors(nullptr, Packages, PackageCache, [](bool){}, /*bLoadForPIE=*/false, /*bLoadAsync=*/false));
+	verify(FWorldPartitionLevelHelper::LoadActors(nullptr, Packages, PackageCache, [](bool){}, /*bLoadForPlay=*/false, /*bLoadAsync=*/false));
 }
 
 void UWorldPartitionRuntimeLevelStreamingCell::MoveAlwaysLoadedContentToPersistentLevel()
diff --git a/Engine/Source/Runtime/Engine/Private/WorldPartition/WorldPartitionRuntimeSpatialHash.cpp b/Engine/Source/Runtime/Engine/Private/WorldPartition/WorldPartitionRuntimeSpatialHash.cpp
--- a/Engine/Source/Runtime/Engine/Private/WorldPartition/WorldPartitionRuntimeSpatialHash.cpp
+++ b/Engine/Source/Runtime/Engine/Private/WorldPartition/WorldPartitionRuntimeSpatialHash.cpp
@@ -644,7 +644,7 @@ bool UWorldPartitionRuntimeSpatialHash::CreateStreamingGrid(const FSpatialHashRu
 						// Then, duplication of world for PIE will duplicate only these actors. 
 						// When stopping PIE, WorldPartition will release these FWorldPartitionReferences which 
 						// will unload actors that were not already loaded in the non PIE world.
-						if (bIsCellAlwaysLoaded && bIsMainWorldPartition && (Mode == EWorldPartitionStreamingMode::PIE))
+						if (bIsCellAlwaysLoaded && bIsMainWorldPartition && (Mode == EWorldPartitionStreamingMode::PIE || Mode == EWorldPartitionStreamingMode::EditorStandalone))
 						{
 							if (ActorInstance.ContainerInstance->Container == WorldPartition)
 							{
@@ -710,7 +710,7 @@ bool UWorldPartitionRuntimeSpatialHash::CreateStreamingGrid(const FSpatialHashRu
 					StreamingCell->AddCellData(HLODCellData);
 				}
 
-				if (Mode == EWorldPartitionStreamingMode::RuntimeStreamingCells)
+				if (Mode == EWorldPartitionStreamingMode::Cook)
 				{
 					UE_LOG(LogWorldPartitionRuntimeSpatialHash, Log, TEXT("Creating runtime streaming cells %s."), *StreamingCell->GetName());
 
diff --git a/Engine/Source/Runtime/Engine/Public/WorldPartition/WorldPartition.h b/Engine/Source/Runtime/Engine/Public/WorldPartition/WorldPartition.h
--- a/Engine/Source/Runtime/Engine/Public/WorldPartition/WorldPartition.h
+++ b/Engine/Source/Runtime/Engine/Public/WorldPartition/WorldPartition.h
@@ -31,8 +31,9 @@ DECLARE_LOG_CATEGORY_EXTERN(LogWorldPartition, Log, All);
 
 enum class EWorldPartitionStreamingMode
 {
-	PIE,
-	RuntimeStreamingCells
+	PIE,				// Used for PIE
+	EditorStandalone,	// Used when running with -game
+	Cook				// Used by cook commandlet
 };
 
 enum class EWorldPartitionInitState
@@ -85,9 +86,10 @@ private:
 
 	void FlushStreaming();
 
-	// PIE Methods
+	// PIE/Game Methods
 	void OnPreBeginPIE(bool bStartSimulate);
-	void OnEndPIE(bool bStartSimulate);
+	void OnBeginPlay(EWorldPartitionStreamingMode Mode);
+	void OnEndPlay();
 
 	// UActorDescContainer events
 	virtual void OnActorDescAdded(const TUniquePtr<FWorldPartitionActorDesc>& NewActorDesc) override;
@@ -129,12 +131,10 @@ public:
 	bool AreEditorCellsLoaded(const FBox& Box);
 	bool RefreshLoadedEditorCells();
 
-	// PIE Methods
-	void PrepareForPIE();
-	void CleanupForPIE();
-	void OnPreFixupForPIE(int32 InPIEInstanceID, FSoftObjectPath& ObjectPath);
+	// PIE/Game Methods
+	void RemapSoftObjectPath(FSoftObjectPath& ObjectPath);
 
-	// PIE/Cook Methods
+	// PIE/Game/Cook Methods
 	bool GenerateStreaming(EWorldPartitionStreamingMode Mode, TArray<FString>* OutPackagesToGenerate = nullptr);
 
 	// Cook Methods
diff --git a/Engine/Source/Runtime/Engine/Public/WorldPartition/WorldPartitionLevelHelper.h b/Engine/Source/Runtime/Engine/Public/WorldPartition/WorldPartitionLevelHelper.h
--- a/Engine/Source/Runtime/Engine/Public/WorldPartition/WorldPartitionLevelHelper.h
+++ b/Engine/Source/Runtime/Engine/Public/WorldPartition/WorldPartitionLevelHelper.h
@@ -22,7 +22,7 @@ public:
 	static ULevel* CreateEmptyLevelForRuntimeCell(const UWorld* InWorld, const FString& InWorldAssetName, UPackage* DestPackage = nullptr);
 	static void MoveExternalActorsToLevel(const TArray<FWorldPartitionRuntimeCellObjectMapping>& InChildPackages, ULevel* InLevel);
 	
-	static bool LoadActors(ULevel* InDestLevel, TArrayView<FWorldPartitionRuntimeCellObjectMapping> InActorPackages, FWorldPartitionPackageCache& InPackageCache, TFunction<void(bool)> InCompletionCallback, bool bLoadForPie, bool bLoadAsync = true, FLinkerInstancingContext* InOutInstancingContext = nullptr);
+	static bool LoadActors(ULevel* InDestLevel, TArrayView<FWorldPartitionRuntimeCellObjectMapping> InActorPackages, FWorldPartitionPackageCache& InPackageCache, TFunction<void(bool)> InCompletionCallback, bool bLoadForPlay, bool bLoadAsync = true, FLinkerInstancingContext* InOutInstancingContext = nullptr);
 private:
 	static UWorld::InitializationValues GetWorldInitializationValues();
 };
diff --git a/Engine/Source/Runtime/Engine/Public/WorldPartition/WorldPartitionLevelStreamingPolicy.h b/Engine/Source/Runtime/Engine/Public/WorldPartition/WorldPartitionLevelStreamingPolicy.h
--- a/Engine/Source/Runtime/Engine/Public/WorldPartition/WorldPartitionLevelStreamingPolicy.h
+++ b/Engine/Source/Runtime/Engine/Public/WorldPartition/WorldPartitionLevelStreamingPolicy.h
@@ -27,8 +27,9 @@ public:
 
 #if WITH_EDITOR
 	virtual TSubclassOf<class UWorldPartitionRuntimeCell> GetRuntimeCellClass() const override;
-	virtual void PrepareForPIE() override;
-	virtual void OnPreFixupForPIE(int32 InPIEInstanceID, FSoftObjectPath& ObjectPath) override;
+	virtual void OnBeginPlay() override;
+	virtual void OnEndPlay() override;
+	virtual void RemapSoftObjectPath(FSoftObjectPath& ObjectPath) override;
 	static FString GetCellPackagePath(const FName& InCellName, const UWorld* InWorld);
 #endif
 
diff --git a/Engine/Source/Runtime/Engine/Public/WorldPartition/WorldPartitionRuntimeHash.h b/Engine/Source/Runtime/Engine/Public/WorldPartition/WorldPartitionRuntimeHash.h
--- a/Engine/Source/Runtime/Engine/Public/WorldPartition/WorldPartitionRuntimeHash.h
+++ b/Engine/Source/Runtime/Engine/Public/WorldPartition/WorldPartitionRuntimeHash.h
@@ -21,9 +21,10 @@ class ENGINE_API UWorldPartitionRuntimeHash : public UObject
 	friend class FActorClusterContext;
 
 #if WITH_EDITOR
+	bool GenerateRuntimeStreaming(EWorldPartitionStreamingMode Mode, class UWorldPartitionStreamingPolicy* Policy, TArray<FString>* OutPackagesToGenerate = nullptr);
+
 	virtual void SetDefaultValues() {}
 	virtual void ImportFromWorldComposition(class UWorldComposition* WorldComposition) {}
-	virtual bool GenerateStreaming(EWorldPartitionStreamingMode Mode, class UWorldPartitionStreamingPolicy* Policy, TArray<FString>* OutPackagesToGenerate = nullptr) { return false; }
 	virtual bool PopulateGeneratedPackageForCook(UPackage* InPackage, const FString& InPackageRelativePath, const FString& InPackageCookName) { return false; }
 	virtual void FinalizeGeneratedPackageForCook() {}
 	virtual void FlushStreaming() {}
@@ -32,9 +33,9 @@ class ENGINE_API UWorldPartitionRuntimeHash : public UObject
 	virtual FName GetActorRuntimeGrid(const AActor* Actor) const { return NAME_None; }
 	virtual void DrawPreview() const {}
 
-	// PIE methods
-	void OnPreBeginPIE();
-	void OnEndPIE();
+	// PIE/Game methods
+	void OnBeginPlay(EWorldPartitionStreamingMode Mode);
+	void OnEndPlay();
 #endif
 
 	// Streaming interface
@@ -54,6 +55,7 @@ class ENGINE_API UWorldPartitionRuntimeHash : public UObject
 
 protected:
 #if WITH_EDITOR
+	virtual bool GenerateStreaming(EWorldPartitionStreamingMode Mode, class UWorldPartitionStreamingPolicy* Policy, TArray<FString>* OutPackagesToGenerate) { return false; }
 	virtual void CreateActorDescViewMap(const UActorDescContainer* Container, TMap<FGuid, FWorldPartitionActorDescView>& OutActorDescViewMap) const;
 	void ChangeActorDescViewGridPlacement(FWorldPartitionActorDescView& ActorDescView, EActorGridPlacement GridPlacement) const;
 #endif
diff --git a/Engine/Source/Runtime/Engine/Public/WorldPartition/WorldPartitionRuntimeSpatialHash.h b/Engine/Source/Runtime/Engine/Public/WorldPartition/WorldPartitionRuntimeSpatialHash.h
--- a/Engine/Source/Runtime/Engine/Public/WorldPartition/WorldPartitionRuntimeSpatialHash.h
+++ b/Engine/Source/Runtime/Engine/Public/WorldPartition/WorldPartitionRuntimeSpatialHash.h
@@ -197,7 +197,6 @@ public:
 #if WITH_EDITOR
 	virtual void SetDefaultValues() override;
 	virtual void ImportFromWorldComposition(class UWorldComposition* WorldComposition) override;
-	virtual bool GenerateStreaming(EWorldPartitionStreamingMode Mode, class UWorldPartitionStreamingPolicy* StreamingPolicy, TArray<FString>* OutPackagesToGenerate = nullptr) override;
 	virtual bool PopulateGeneratedPackageForCook(UPackage* InPackage, const FString& InPackageRelativePath, const FString& InPackageCookName) override;
 	virtual void FinalizeGeneratedPackageForCook() override;
 	virtual void FlushStreaming() override;
@@ -218,6 +217,11 @@ public:
 	virtual void SortStreamingCellsByImportance(const TSet<const UWorldPartitionRuntimeCell*>& InCells, const TArray<FWorldPartitionStreamingSource>& InSources, TArray<const UWorldPartitionRuntimeCell*, TInlineAllocator<256>>& OutSortedCells) const override;
 
 protected:
+
+#if WITH_EDITOR
+	virtual bool GenerateStreaming(EWorldPartitionStreamingMode Mode, class UWorldPartitionStreamingPolicy* StreamingPolicy, TArray<FString>* OutPackagesToGenerate = nullptr) override;
+#endif
+
 	// Used to convert from a world using World Composition
 	TMap<class ULevelStreaming*, FName> WorldCompositionStreamingLevelToRuntimeGrid;
 
diff --git a/Engine/Source/Runtime/Engine/Public/WorldPartition/WorldPartitionStreamingPolicy.h b/Engine/Source/Runtime/Engine/Public/WorldPartition/WorldPartitionStreamingPolicy.h
--- a/Engine/Source/Runtime/Engine/Public/WorldPartition/WorldPartitionStreamingPolicy.h
+++ b/Engine/Source/Runtime/Engine/Public/WorldPartition/WorldPartitionStreamingPolicy.h
@@ -34,8 +34,11 @@ public:
 
 #if WITH_EDITOR
 	virtual TSubclassOf<class UWorldPartitionRuntimeCell> GetRuntimeCellClass() const PURE_VIRTUAL(UWorldPartitionStreamingPolicy::GetRuntimeCellClass, return UWorldPartitionRuntimeCell::StaticClass(); );
-	virtual void PrepareForPIE() {}
-	virtual void OnPreFixupForPIE(int32 InPIEInstanceID, FSoftObjectPath& ObjectPath) {}
+
+	// PIE/Game methods
+	virtual void OnBeginPlay() {}
+	virtual void OnEndPlay() {}
+	virtual void RemapSoftObjectPath(FSoftObjectPath& ObjectPath) {}
 #endif
 
 	const TArray<FWorldPartitionStreamingSource>& GetStreamingSources() const { return StreamingSources; }
