Fix potential infinite wait when DistanceField is waiting on a StaticMesh to finish which contains a dependency to another building static mesh through its GenerateSource

diff --git a/Engine/Source/Runtime/Engine/Private/DistanceFieldAtlas.cpp b/Engine/Source/Runtime/Engine/Private/DistanceFieldAtlas.cpp
--- a/Engine/Source/Runtime/Engine/Private/DistanceFieldAtlas.cpp
+++ b/Engine/Source/Runtime/Engine/Private/DistanceFieldAtlas.cpp
@@ -1197,14 +1197,13 @@ void FDistanceFieldAsyncQueue::BlockUntilBuildComplete(UStaticMesh* StaticMesh,
 	bool bHadToBlock = false;
 	double StartTime = 0;
 
-#if WITH_EDITOR
-	FStaticMeshCompilingManager::Get().FinishCompilation({ StaticMesh });
-#endif
+	TSet<UStaticMesh*> RequiredFinishCompilation;
 	do 
 	{
 		ProcessAsyncTasks();
 
 		bReferenced = false;
+		RequiredFinishCompilation.Reset();
 
 		// Reschedule the tasks we're waiting on as highest prio
 		{
@@ -1215,11 +1214,38 @@ void FDistanceFieldAsyncQueue::BlockUntilBuildComplete(UStaticMesh* StaticMesh,
 					ReferencedTasks[TaskIndex]->GenerateSource == StaticMesh)
 				{
 					bReferenced = true;
+
+					// If the task we are waiting on depends on other static meshes
+					// we need to force finish them too.
+#if WITH_EDITOR
+					
+					if (ReferencedTasks[TaskIndex]->GenerateSource != nullptr &&
+						ReferencedTasks[TaskIndex]->GenerateSource->IsCompiling())
+					{
+						RequiredFinishCompilation.Add(ReferencedTasks[TaskIndex]->GenerateSource);
+					}
+
+					if (ReferencedTasks[TaskIndex]->StaticMesh != nullptr &&
+						ReferencedTasks[TaskIndex]->StaticMesh->IsCompiling())
+					{
+						RequiredFinishCompilation.Add(ReferencedTasks[TaskIndex]->StaticMesh);
+					}
+#endif
+
 					RescheduleBackgroundTask(ReferencedTasks[TaskIndex], EQueuedWorkPriority::Highest);
 				}
 			}
 		}
 
+#if WITH_EDITOR
+		// Call the finish compilation outside of the critical section since those compilations
+		// might need to register new distance field tasks which also uses the critical section.
+		if (RequiredFinishCompilation.Num())
+		{
+			FStaticMeshCompilingManager::Get().FinishCompilation(RequiredFinishCompilation.Array());
+		}
+#endif
+
 		if (bReferenced)
 		{
 			if (!bHadToBlock)
