In World Partition mode, for now, force GLevelStreamingContinuouslyIncrementalGCWhileLevelsPendingPurge and GLevelStreamingForceGCAfterLevelStreamedOut to be disabled by default.

diff --git a/Engine/Source/Runtime/Engine/Private/WorldPartition/WorldPartitionSubsystem.cpp b/Engine/Source/Runtime/Engine/Private/WorldPartition/WorldPartitionSubsystem.cpp
--- a/Engine/Source/Runtime/Engine/Private/WorldPartition/WorldPartitionSubsystem.cpp
+++ b/Engine/Source/Runtime/Engine/Private/WorldPartition/WorldPartitionSubsystem.cpp
@@ -3,6 +3,7 @@
 #include "WorldPartition/WorldPartitionSubsystem.h"
 #include "WorldPartition/WorldPartition.h"
 #include "WorldPartition/WorldPartitionStreamingPolicy.h"
+#include "HAL/IConsoleManager.h"
 #include "Engine/Canvas.h"
 #include "Engine/Console.h"
 #include "ConsoleSettings.h"
@@ -78,11 +79,29 @@ void UWorldPartitionSubsystem::PostInitialize()
 		{
 			DrawHandle = UDebugDrawService::Register(TEXT("Game"), FDebugDrawDelegate::CreateUObject(this, &UWorldPartitionSubsystem::Draw));
 		}
+
+		// For now, enforce some GC settings when using World Partition
+		if (GetWorld()->IsGameWorld())
+		{
+			PreviousCVarValues.ReadFromCVars();
+
+			FWorldPartitionCVars OverrideCVars;
+			OverrideCVars.ContinuouslyIncremental = 0;
+			OverrideCVars.ForceGCAfterLevelStreamedOut = 0;
+			OverrideCVars.TimeBetweenPurgingPendingKillObjects = 120.f;
+			OverrideCVars.WriteToCVars();
+		}
 	}
 }
 
 void UWorldPartitionSubsystem::Deinitialize()
 {
+	UWorldPartition* MainPartition = GetMainWorldPartition();
+	if (MainPartition && GetWorld()->IsGameWorld())
+	{
+		PreviousCVarValues.WriteToCVars();
+	}
+
 	if (DrawHandle.IsValid())
 	{
 		UDebugDrawService::Unregister(DrawHandle);
@@ -252,4 +271,57 @@ void UWorldPartitionSubsystem::ForEachActorDesc(TSubclassOf<AActor> ActorClass,
 		MainPartition->ForEachActorDesc(ActorClass, Predicate);
 	}
 }
-#endif // WITH_EDITOR
\ No newline at end of file
+#endif // WITH_EDITOR
+
+//
+// UWorldPartitionSubsystem::FWorldPartitionCVars
+//
+
+const TCHAR* UWorldPartitionSubsystem::FWorldPartitionCVars::ContinuouslyIncrementalText = TEXT("s.ContinuouslyIncrementalGCWhileLevelsPendingPurge");
+const TCHAR* UWorldPartitionSubsystem::FWorldPartitionCVars::ForceGCAfterLevelStreamedOutText = TEXT("s.ForceGCAfterLevelStreamedOut");
+const TCHAR* UWorldPartitionSubsystem::FWorldPartitionCVars::TimeBetweenPurgingPendingKillObjectsText = TEXT("gc.TimeBetweenPurgingPendingKillObjects");
+
+void UWorldPartitionSubsystem::FWorldPartitionCVars::ReadFromCVars()
+{
+	if (IConsoleVariable* ContinuouslyIncrementalGCCVar = IConsoleManager::Get().FindConsoleVariable(ContinuouslyIncrementalText))
+	{
+		ContinuouslyIncremental = ContinuouslyIncrementalGCCVar->GetInt();
+	}
+
+	if (IConsoleVariable* ForceGCAfterLevelStreamedOutCVar = IConsoleManager::Get().FindConsoleVariable(ForceGCAfterLevelStreamedOutText))
+	{
+		ForceGCAfterLevelStreamedOut = ForceGCAfterLevelStreamedOutCVar->GetInt();
+	}
+
+	if (IConsoleVariable* TimeBetweenPurgingPendingKillObjectsCVar = IConsoleManager::Get().FindConsoleVariable(TimeBetweenPurgingPendingKillObjectsText))
+	{
+		TimeBetweenPurgingPendingKillObjects = TimeBetweenPurgingPendingKillObjectsCVar->GetFloat();
+	}
+}
+
+void UWorldPartitionSubsystem::FWorldPartitionCVars::WriteToCVars() const
+{
+	if (ContinuouslyIncremental.IsSet())
+	{
+		if (IConsoleVariable* ContinuouslyIncrementalGCCVar = IConsoleManager::Get().FindConsoleVariable(ContinuouslyIncrementalText))
+		{
+			ContinuouslyIncrementalGCCVar->Set(ContinuouslyIncremental.GetValue(), ECVF_SetByCode);
+		}
+	}
+
+	if (ForceGCAfterLevelStreamedOut.IsSet())
+	{
+		if (IConsoleVariable* ForceGCAfterLevelStreamedOutCVar = IConsoleManager::Get().FindConsoleVariable(ForceGCAfterLevelStreamedOutText))
+		{
+			ForceGCAfterLevelStreamedOutCVar->Set(ForceGCAfterLevelStreamedOut.GetValue(), ECVF_SetByCode);
+		}
+	}
+
+	if (TimeBetweenPurgingPendingKillObjects.IsSet())
+	{
+		if (IConsoleVariable* TimeBetweenPurgingPendingKillObjectsCVar = IConsoleManager::Get().FindConsoleVariable(TimeBetweenPurgingPendingKillObjectsText))
+		{
+			TimeBetweenPurgingPendingKillObjectsCVar->Set(TimeBetweenPurgingPendingKillObjects.GetValue(), ECVF_SetByCode);
+		}
+	}
+}
diff --git a/Engine/Source/Runtime/Engine/Public/WorldPartition/WorldPartitionSubsystem.h b/Engine/Source/Runtime/Engine/Public/WorldPartition/WorldPartitionSubsystem.h
--- a/Engine/Source/Runtime/Engine/Public/WorldPartition/WorldPartitionSubsystem.h
+++ b/Engine/Source/Runtime/Engine/Public/WorldPartition/WorldPartitionSubsystem.h
@@ -4,6 +4,7 @@
 
 #include "Subsystems/WorldSubsystem.h"
 #include "Templates/SubclassOf.h"
+#include "Misc/Optional.h"
 #include "WorldPartitionSubsystem.generated.h"
 
 class UWorldPartition;
@@ -57,4 +58,22 @@ private:
 	TArray<TObjectPtr<UWorldPartition>> RegisteredWorldPartitions;
 
 	FDelegateHandle	DrawHandle;
+
+	struct FWorldPartitionCVars
+	{
+		FWorldPartitionCVars() {}
+
+		void ReadFromCVars();
+		void WriteToCVars() const;
+
+		static const TCHAR* ContinuouslyIncrementalText;
+		static const TCHAR* ForceGCAfterLevelStreamedOutText;
+		static const TCHAR* TimeBetweenPurgingPendingKillObjectsText;
+
+		TOptional<int32> ContinuouslyIncremental;
+		TOptional<int32> ForceGCAfterLevelStreamedOut;
+		TOptional<float> TimeBetweenPurgingPendingKillObjects;
+	};
+
+	FWorldPartitionCVars PreviousCVarValues;
 };
