Mostly cosmetic changes to explain a check that is done for valid cached animation data in BlendSpace. Removes this check when not in Editor since it should never be relevant.

diff --git a/Engine/Source/Runtime/Engine/Private/Animation/BlendSpace.cpp b/Engine/Source/Runtime/Engine/Private/Animation/BlendSpace.cpp
--- a/Engine/Source/Runtime/Engine/Private/Animation/BlendSpace.cpp
+++ b/Engine/Source/Runtime/Engine/Private/Animation/BlendSpace.cpp
@@ -170,16 +170,28 @@ bool UBlendSpace::UpdateBlendSamples_Internal(
 
 	InOutOldSampleDataList.Append(InOutSampleDataCache);
 
-	// @fixme: temporary code to clear the invalid sample data 
-	// related jira: UE-71107
+#if WITH_EDITOR
+	// If we are in Editor then samples may be added/removed, and when this happens it is not
+	// practical to find any existing caches that reference affected animations. Note that all the
+	// indices may become invalid. We can just check here for invalid sample indices - if we find one
+	// then remove the cache and start again. There will be some situations where the sample indices
+	// get shuffled, in which case we will not detect the change, but our cache will now point to
+	// incorrect animations. Any glitch that occurs as a result should be removed over the smoothing window. 
+	//
+	// Note that if this is not sufficient, we could store a GUID in UBlendSpace that gets updated
+	// when there is a change, and a GUID in the cache. Then when the GUIDs don't match the cache
+	// could be wiped and restarted with the current GUID.
+	//
+	// See UE-71107
 	for (int32 Index = 0; Index < InOutOldSampleDataList.Num(); ++Index)
 	{
 		if (!SampleData.IsValidIndex(InOutOldSampleDataList[Index].SampleDataIndex))
 		{
-			InOutOldSampleDataList.RemoveAt(Index);
-			--Index;
+			InOutOldSampleDataList.Empty();
+			break;
 		}
 	}
+#endif
 
 	// get sample data based on new input
 	// consolidate all samples and sort them, so that we can handle from biggest weight to smallest
