Fix TextRenderComponent local bounds calculation.

diff --git a/Engine/Source/Runtime/Engine/Classes/Components/PrimitiveComponent.h b/Engine/Source/Runtime/Engine/Classes/Components/PrimitiveComponent.h
--- a/Engine/Source/Runtime/Engine/Classes/Components/PrimitiveComponent.h
+++ b/Engine/Source/Runtime/Engine/Classes/Components/PrimitiveComponent.h
@@ -1734,8 +1734,10 @@ public:
 	void SyncComponentToRBPhysics();
 	
 	/** 
-	 *	Returns the matrix that should be used to render this component. 
-	 *	Allows component class to perform graphical distortion to the component not supported by an FTransform 
+	 * Returns the matrix that should be used to render this component. 
+	 * Allows component class to perform graphical distortion to the component not supported by an FTransform 
+	 * NOTE: When overriding this method to alter the transform used for rendering it is typically neccessary to also implement USceneComponent::UpdateBounds.
+	 *       Otherwise the Local bounds will not match the world space bounds, causing incorrect culling.
 	 */
 	virtual FMatrix GetRenderMatrix() const;
 
diff --git a/Engine/Source/Runtime/Engine/Classes/Components/TextRenderComponent.h b/Engine/Source/Runtime/Engine/Classes/Components/TextRenderComponent.h
--- a/Engine/Source/Runtime/Engine/Classes/Components/TextRenderComponent.h
+++ b/Engine/Source/Runtime/Engine/Classes/Components/TextRenderComponent.h
@@ -168,6 +168,8 @@ class ENGINE_API UTextRenderComponent : public UPrimitiveComponent
 
 	//~ Begin USceneComponent Interface.
 	virtual FBoxSphereBounds CalcBounds(const FTransform& LocalToWorld) const override;
+	// Note: this is required because GetRenderMatrix is overridden, and therefore the updated world-space bounds must be based on that (not the component transform)
+	virtual void UpdateBounds() override;
 	//~ End USceneComponent Interface.
 
 	//~ Begin UActorComponent Interface.
diff --git a/Engine/Source/Runtime/Engine/Private/Components/TextRenderComponent.cpp b/Engine/Source/Runtime/Engine/Private/Components/TextRenderComponent.cpp
--- a/Engine/Source/Runtime/Engine/Private/Components/TextRenderComponent.cpp
+++ b/Engine/Source/Runtime/Engine/Private/Components/TextRenderComponent.cpp
@@ -1157,9 +1157,7 @@ FBoxSphereBounds UTextRenderComponent::CalcBounds(const FTransform& LocalToWorld
 			while (It.NextCharacterInLine(Ch)) {}
 		}
 
-		LeftTop.Y = ComputeVerticalAlignmentOffset(Size.Y, VerticalAlignment, FirstLineHeight);
-		const FBox LocalBox(FVector(0.f, -LeftTop.X, -LeftTop.Y), FVector(0.f, -(LeftTop.X + Size.X), -(LeftTop.Y + Size.Y)));
-
+		const FBox LocalBox(FVector(0.f, -Size.X - LeftTop.X, -Size.Y), FVector(0.f, -LeftTop.X, 0.0f));
 		FBoxSphereBounds Ret(LocalBox.TransformBy(LocalToWorld));
 
 		Ret.BoxExtent *= BoundsScale;
@@ -1173,6 +1171,11 @@ FBoxSphereBounds UTextRenderComponent::CalcBounds(const FTransform& LocalToWorld
 	}
 }
 
+void UTextRenderComponent::UpdateBounds()
+{
+	Bounds = CalcBounds(FTransform(GetRenderMatrix()));
+}
+
 bool UTextRenderComponent::RequiresGameThreadEndOfFrameUpdates() const
 {
 	return true;
