Exclude shader types that are included in FShaderPipelines from the standalone compilation.

- Not the cleanest way to do this, but arguably minimal change. Jason's idea.

diff --git a/Engine/Source/Runtime/Engine/Private/Materials/MaterialShader.cpp b/Engine/Source/Runtime/Engine/Private/Materials/MaterialShader.cpp
--- a/Engine/Source/Runtime/Engine/Private/Materials/MaterialShader.cpp
+++ b/Engine/Source/Runtime/Engine/Private/Materials/MaterialShader.cpp
@@ -19,6 +19,14 @@
 #include "ProfilingDebugging/LoadTimeTracker.h"
 #include "Misc/ScopeRWLock.h"
 
+int32 GMaterialExcludeNonPipelinedShaders = 1;
+static FAutoConsoleVariableRef CVarShaderCompilerJobCache(
+	TEXT("r.Material.ExcludeNonPipelinedShaders"),
+	GMaterialExcludeNonPipelinedShaders,
+	TEXT("if != 0, standalone shaders that are also part of FShaderPipeline will not be compiled (default)."),
+	ECVF_ReadOnly
+);
+
 #if ENABLE_COOK_STATS
 namespace MaterialShaderCookStats
 {
@@ -78,6 +86,36 @@ FString GetShadingModelString(EMaterialShadingModel ShadingModel)
 	return ShadingModelName;
 }
 
+/** Helpers class to identify and remove shader types that are going to be used by the pipelines. 
+    Standalone shaders of that type should no longer be required, but removing them from the shadermap layout is a bigger endeavour */
+class FPipelinedShaderFilter
+{
+	TSet<const FShaderType*>	PipelinedShaderTypes;
+	bool						bAnyTypesExcluded;
+
+public:
+
+	FPipelinedShaderFilter(EShaderPlatform ShaderPlatform, TArray<FShaderPipelineType*> Pipelines)
+	{
+		if (GMaterialExcludeNonPipelinedShaders && RHISupportsShaderPipelines(ShaderPlatform))
+		{
+			for (FShaderPipelineType* Pipeline : Pipelines)
+			{
+				if (Pipeline->ShouldOptimizeUnusedOutputs(ShaderPlatform))
+				{
+					PipelinedShaderTypes.Append(Pipeline->GetStages());
+				}
+			}
+
+			bAnyTypesExcluded = PipelinedShaderTypes.Num() != 0;
+		}
+	}
+
+	inline bool IsPipelinedType(const FShaderType* Type) const
+	{
+		return bAnyTypesExcluded ? PipelinedShaderTypes.Contains(Type) : false;
+	}
+};
 
 /** Converts an FMaterialShadingModelField to a string description containing all the shading models present, delimited by "|" */
 FString GetShadingModelFieldString(FMaterialShadingModelField ShadingModels, const FShadingModelToStringDelegate& Delegate, const FString& Delimiter)
@@ -1472,6 +1510,10 @@ void FMaterialShaderMap::SubmitCompileJobs(uint32 CompilingShaderMapId,
 		uint32 NumShadersPerVF = 0;
 		TSet<FString> ShaderTypeNames;
 
+		// Do not submit jobs for the shader types that are included in some pipeline stages if that pipeline is optimzing unused outputs.
+		// In such a case, these shaders should not be used runtime anymore
+		FPipelinedShaderFilter PipelinedShaderFilter(ShaderPlatform, MeshLayout.ShaderPipelines);
+
 		// Iterate over all mesh material shader types.
 		TMap<TShaderTypePermutation<const FShaderType>, FShaderCompileJob*> SharedShaderJobs;
 		for (const FShaderLayoutEntry& Shader : MeshLayout.Shaders)
@@ -1489,8 +1531,8 @@ void FMaterialShaderMap::SubmitCompileJobs(uint32 CompilingShaderMapId,
 #endif
 
 			NumShadersPerVF++;
-			// only compile the shader if we don't already have it
-			if (!MeshShaderMap || !MeshShaderMap->HasShader(ShaderType, Shader.PermutationId))
+			// only compile the shader if we don't already have it and it is not a pipelined one
+			if (!PipelinedShaderFilter.IsPipelinedType(ShaderType) && (!MeshShaderMap || !MeshShaderMap->HasShader(ShaderType, Shader.PermutationId)))
 			{
 				// Compile this mesh material shader for this material and vertex factory type.
 				ShaderType->BeginCompileShader(InPriority,
@@ -1574,6 +1616,10 @@ void FMaterialShaderMap::SubmitCompileJobs(uint32 CompilingShaderMapId,
 		}
 	}
 
+	// Do not submit jobs for the shader types that are included in some pipeline stages if that pipeline is optimzing unused outputs.
+	// In such a case, these shaders should not be used runtime anymore
+	FPipelinedShaderFilter PipelinedShaderFilter(ShaderPlatform, Layout.ShaderPipelines);
+
 	// Iterate over all material shader types.
 	TMap<TShaderTypePermutation<const FShaderType>, FShaderCompileJob*> SharedShaderJobs;
 	for (const FShaderLayoutEntry& Shader : Layout.Shaders)
@@ -1592,7 +1638,7 @@ void FMaterialShaderMap::SubmitCompileJobs(uint32 CompilingShaderMapId,
 		TArray<FString> ShaderErrors;
 
 		// Only compile the shader if we don't already have it
-		if (!GetContent()->HasShader(ShaderType, Shader.PermutationId))
+		if (!PipelinedShaderFilter.IsPipelinedType(ShaderType) && !GetContent()->HasShader(ShaderType, Shader.PermutationId))
 		{
 			ShaderType->BeginCompileShader(InPriority,
 				CompilingShaderMapId,
@@ -2155,17 +2201,22 @@ bool FMaterialShaderMap::IsComplete(const FMaterial* Material, bool bSilent)
 
 	// If our hash doesn't match the cached layout hash, shader map may still be complete
 	// This can happen if FMaterial::ShouldCache is set to return false for any shaders that are included in the cached layout
-	for (const FShaderLayoutEntry& Shader : Layout.Shaders)
+
 	{
-		if (!LocalContent->HasShader(Shader.ShaderType, Shader.PermutationId))
+		// exclude shaders that are going to be uniquely used by the pipelines
+		FPipelinedShaderFilter PipelinedShaderFilter(Platform, Layout.ShaderPipelines);
+		for (const FShaderLayoutEntry& Shader : Layout.Shaders)
 		{
-			if (Material->ShouldCache(Platform, Shader.ShaderType, nullptr))
+			if (!LocalContent->HasShader(Shader.ShaderType, Shader.PermutationId))
 			{
-				if (!bSilent)
+				if (!PipelinedShaderFilter.IsPipelinedType(Shader.ShaderType) && Material->ShouldCache(Platform, Shader.ShaderType, nullptr))
 				{
-					UE_LOG(LogMaterial, Warning, TEXT("Incomplete material %s, missing FMaterialShader (%s, %d)."), *Material->GetFriendlyName(), Shader.ShaderType->GetName(), Shader.PermutationId);
+					if (!bSilent)
+					{
+						UE_LOG(LogMaterial, Warning, TEXT("Incomplete material %s, missing FMaterialShader (%s, %d)."), *Material->GetFriendlyName(), Shader.ShaderType->GetName(), Shader.PermutationId);
+					}
+					return false;
 				}
-				return false;
 			}
 		}
 	}
@@ -2187,13 +2238,14 @@ bool FMaterialShaderMap::IsComplete(const FMaterial* Material, bool bSilent)
 
 	for (const FMeshMaterialShaderMapLayout& MeshLayout : Layout.MeshShaderMaps)
 	{
+		FPipelinedShaderFilter PipelinedShaderFilter(Platform, MeshLayout.ShaderPipelines);
 		const FMeshMaterialShaderMap* MeshShaderMap = LocalContent->GetMeshShaderMap(MeshLayout.VertexFactoryType->GetHashedName());
 
 		for (const FShaderLayoutEntry& Shader : MeshLayout.Shaders)
 		{
 			if (Material->ShouldCache(Platform, Shader.ShaderType, MeshLayout.VertexFactoryType))
 			{
-				if (!MeshShaderMap || !MeshShaderMap->HasShader(Shader.ShaderType, Shader.PermutationId))
+				if (!PipelinedShaderFilter.IsPipelinedType(Shader.ShaderType) && (!MeshShaderMap || !MeshShaderMap->HasShader(Shader.ShaderType, Shader.PermutationId)))
 				{
 					if (!bSilent)
 					{
