Fix an unitialized variable access (thanks Ben Ingram!).

#rnx

diff --git a/Engine/Source/Runtime/Engine/Private/Materials/MaterialShader.cpp b/Engine/Source/Runtime/Engine/Private/Materials/MaterialShader.cpp
--- a/Engine/Source/Runtime/Engine/Private/Materials/MaterialShader.cpp
+++ b/Engine/Source/Runtime/Engine/Private/Materials/MaterialShader.cpp
@@ -91,7 +91,7 @@ FString GetShadingModelString(EMaterialShadingModel ShadingModel)
 class FPipelinedShaderFilter
 {
 	TSet<const FShaderType*>	PipelinedShaderTypes;
-	bool						bAnyTypesExcluded;
+	bool						bAnyTypesExcluded = false;
 
 public:
 
