Allow editing of world partition actor properties if actor has no world.

diff --git a/Engine/Source/Runtime/Engine/Private/ActorEditor.cpp b/Engine/Source/Runtime/Engine/Private/ActorEditor.cpp
--- a/Engine/Source/Runtime/Engine/Private/ActorEditor.cpp
+++ b/Engine/Source/Runtime/Engine/Private/ActorEditor.cpp
@@ -65,10 +65,13 @@ bool AActor::CanEditChange(const FProperty* PropertyThatWillChange) const
 	{
 		if (!IsTemplate())
 		{
-			const bool bIsPartitionedWorld = UWorld::HasSubsystem<UWorldPartitionSubsystem>(GetWorld());
-			if (!bIsPartitionedWorld)
+			if (UWorld* World = GetWorld())
 			{
-				return false;
+				const bool bIsPartitionedWorld = UWorld::HasSubsystem<UWorldPartitionSubsystem>(World);
+				if (!bIsPartitionedWorld)
+				{
+					return false;
+				}
 			}
 		}
 	}
