In WorldPartition, streamed levels now require full visibility to render.

diff --git a/Engine/Source/Runtime/Engine/Classes/Engine/LevelStreaming.h b/Engine/Source/Runtime/Engine/Classes/Engine/LevelStreaming.h
--- a/Engine/Source/Runtime/Engine/Classes/Engine/LevelStreaming.h
+++ b/Engine/Source/Runtime/Engine/Classes/Engine/LevelStreaming.h
@@ -299,6 +299,9 @@ public:
 	UFUNCTION(BlueprintSetter)
 	void SetShouldBeVisible(bool bInShouldBeVisible);
 
+	/** Returns whether level should start to render only when it will be fully added to the world or not. */
+	virtual bool ShouldRequireFullVisibilityToRender() const { return LODPackageNames.Num() > 0; }
+
 	/** 
 	 * Virtual that can be overriden to change whether a streaming level should be loaded.
 	 * Doesn't do anything at the base level as should be loaded defaults to true 
diff --git a/Engine/Source/Runtime/Engine/Private/LevelStreaming.cpp b/Engine/Source/Runtime/Engine/Private/LevelStreaming.cpp
--- a/Engine/Source/Runtime/Engine/Private/LevelStreaming.cpp
+++ b/Engine/Source/Runtime/Engine/Private/LevelStreaming.cpp
@@ -1215,11 +1215,14 @@ void ULevelStreaming::AsyncLevelLoadComplete(const FName& InPackageName, UPackag
 				IStreamingManager::Get().AddLevel(Level);
 
 				// Make sure this level will start to render only when it will be fully added to the world
-				if (LODPackageNames.Num() > 0)
+				if (ShouldRequireFullVisibilityToRender())
 				{
 					Level->bRequireFullVisibilityToRender = true;
 					// LOD levels should not be visible on server
-					Level->bClientOnlyVisible = LODPackageNames.Contains(InLoadedPackage->GetFName());
+					if (LODPackageNames.Num() > 0)
+					{
+						Level->bClientOnlyVisible = LODPackageNames.Contains(InLoadedPackage->GetFName());
+					}
 				}
 			
 				// In the editor levels must be in the levels array regardless of whether they are visible or not
diff --git a/Engine/Source/Runtime/Engine/Private/WorldPartition/WorldPartitionLevelHelper.cpp b/Engine/Source/Runtime/Engine/Private/WorldPartition/WorldPartitionLevelHelper.cpp
--- a/Engine/Source/Runtime/Engine/Private/WorldPartition/WorldPartitionLevelHelper.cpp
+++ b/Engine/Source/Runtime/Engine/Private/WorldPartition/WorldPartitionLevelHelper.cpp
@@ -108,7 +108,6 @@ ULevel* FWorldPartitionLevelHelper::CreateEmptyLevelForRuntimeCell(const UWorld*
 	check(NewLevel->OwningWorld == NewWorld);
 	check(NewLevel->Model);
 	check(!NewLevel->bIsVisible);
-	NewLevel->bRequireFullVisibilityToRender = true;
 	return NewLevel;
 }
 
diff --git a/Engine/Source/Runtime/Engine/Private/WorldPartition/WorldPartitionLevelStreamingDynamic.cpp b/Engine/Source/Runtime/Engine/Private/WorldPartition/WorldPartitionLevelStreamingDynamic.cpp
--- a/Engine/Source/Runtime/Engine/Private/WorldPartition/WorldPartitionLevelStreamingDynamic.cpp
+++ b/Engine/Source/Runtime/Engine/Private/WorldPartition/WorldPartitionLevelStreamingDynamic.cpp
@@ -387,6 +387,10 @@ void UWorldPartitionLevelStreamingDynamic::FinalizeRuntimeLevel()
 
 	// Notify the streamer to start building incrementally the level streaming data.
 	IStreamingManager::Get().AddLevel(RuntimeLevel);
+
+	// Make sure this level will start to render only when it will be fully added to the world
+	check(ShouldRequireFullVisibilityToRender());
+	RuntimeLevel->bRequireFullVisibilityToRender = true;
 }
 
 /**
diff --git a/Engine/Source/Runtime/Engine/Public/WorldPartition/WorldPartitionLevelStreamingDynamic.h b/Engine/Source/Runtime/Engine/Public/WorldPartition/WorldPartitionLevelStreamingDynamic.h
--- a/Engine/Source/Runtime/Engine/Public/WorldPartition/WorldPartitionLevelStreamingDynamic.h
+++ b/Engine/Source/Runtime/Engine/Public/WorldPartition/WorldPartitionLevelStreamingDynamic.h
@@ -28,6 +28,7 @@ class ENGINE_API UWorldPartitionLevelStreamingDynamic : public ULevelStreamingDy
 	void Activate();
 	void Deactivate();
 	virtual bool ShouldBeAlwaysLoaded() const { return bShouldBeAlwaysLoaded; }
+	virtual bool ShouldRequireFullVisibilityToRender() const override { return true; }
 	void SetShouldBeAlwaysLoaded(bool bInShouldBeAlwaysLoaded) { bShouldBeAlwaysLoaded = bInShouldBeAlwaysLoaded; }
 	UWorld* GetOuterWorld() const;
 
