add ENGINE_API to GetAnimNotifiesFromTrackPosition (change from coalition)

diff --git a/Engine/Source/Runtime/Engine/Classes/Animation/AnimCompositeBase.h b/Engine/Source/Runtime/Engine/Classes/Animation/AnimCompositeBase.h
--- a/Engine/Source/Runtime/Engine/Classes/Animation/AnimCompositeBase.h
+++ b/Engine/Source/Runtime/Engine/Classes/Animation/AnimCompositeBase.h
@@ -284,7 +284,7 @@ struct FAnimTrack
 	* Supports playing backwards (CurrentTrackPosition<PreviousTrackPosition).
 	* Only supports contiguous range, does NOT support looping and wrapping over.
 	*/
-	void GetAnimNotifiesFromTrackPositions(const float& PreviousTrackPosition, const float& CurrentTrackPosition, TArray<FAnimNotifyEventReference> & OutActiveNotifies) const;
+	ENGINE_API void GetAnimNotifiesFromTrackPositions(const float& PreviousTrackPosition, const float& CurrentTrackPosition, TArray<FAnimNotifyEventReference>& OutActiveNotifies) const;
 
 	/** return true if anim notify is available */
 	bool IsNotifyAvailable() const;
