Added UAnimInstance::AddPoseSnapshot and GetPoseSnapshot (change from coalition)

diff --git a/Engine/Source/Runtime/Engine/Classes/Animation/AnimInstance.h b/Engine/Source/Runtime/Engine/Classes/Animation/AnimInstance.h
--- a/Engine/Source/Runtime/Engine/Classes/Animation/AnimInstance.h
+++ b/Engine/Source/Runtime/Engine/Classes/Animation/AnimInstance.h
@@ -500,6 +500,16 @@ public:
 	UFUNCTION(BlueprintCallable, Category = "Pose")
 	virtual void SavePoseSnapshot(FName SnapshotName);
 
+	/** Add an empty pose snapshot to the internal snapshot cache (or recycle an existing pose snapshot if the name is already in use) */
+	FPoseSnapshot& AddPoseSnapshot(FName SnapshotName);
+
+	/** Remove a previously saved pose snapshot from the internal snapshot cache */
+	UFUNCTION(BlueprintCallable, Category = "Pose")
+	void RemovePoseSnapshot(FName SnapshotName);
+
+	/** Get a cached pose snapshot by name */
+	const FPoseSnapshot* GetPoseSnapshot(FName SnapshotName) const;
+
 	/**
 	 * Takes a snapshot of the current skeletal mesh component pose and saves it to the specified snapshot.
 	 * The snapshot is taken at the current LOD, so if for example you took the snapshot at LOD1 
diff --git a/Engine/Source/Runtime/Engine/Private/Animation/AnimInstance.cpp b/Engine/Source/Runtime/Engine/Private/Animation/AnimInstance.cpp
--- a/Engine/Source/Runtime/Engine/Private/Animation/AnimInstance.cpp
+++ b/Engine/Source/Runtime/Engine/Private/Animation/AnimInstance.cpp
@@ -153,6 +153,24 @@ void UAnimInstance::SavePoseSnapshot(FName SnapshotName)
 	}
 }
 
+FPoseSnapshot& UAnimInstance::AddPoseSnapshot(FName SnapshotName)
+{
+	FAnimInstanceProxy& Proxy = GetProxyOnGameThread<FAnimInstanceProxy>();
+	return Proxy.AddPoseSnapshot(SnapshotName);
+}
+
+void UAnimInstance::RemovePoseSnapshot(FName SnapshotName)
+{
+	FAnimInstanceProxy& Proxy = GetProxyOnGameThread<FAnimInstanceProxy>();
+	Proxy.RemovePoseSnapshot(SnapshotName);
+}
+
+const FPoseSnapshot* UAnimInstance::GetPoseSnapshot(FName SnapshotName) const
+{
+	const FAnimInstanceProxy& Proxy = GetProxyOnGameThread<FAnimInstanceProxy>();
+	return Proxy.GetPoseSnapshot(SnapshotName);
+}
+
 void UAnimInstance::SnapshotPose(FPoseSnapshot& Snapshot)
 {
 	if (USkeletalMeshComponent* SkeletalMeshComponent = GetSkelMeshComponent())
diff --git a/Engine/Source/Runtime/Engine/Private/Animation/AnimInstanceProxy.cpp b/Engine/Source/Runtime/Engine/Private/Animation/AnimInstanceProxy.cpp
--- a/Engine/Source/Runtime/Engine/Private/Animation/AnimInstanceProxy.cpp
+++ b/Engine/Source/Runtime/Engine/Private/Animation/AnimInstanceProxy.cpp
@@ -475,14 +475,8 @@ void FAnimInstanceProxy::SavePoseSnapshot(USkeletalMeshComponent* InSkeletalMesh
 {
 	DECLARE_SCOPE_HIERARCHICAL_COUNTER_FUNC()
 
-	FPoseSnapshot* PoseSnapshot = PoseSnapshots.FindByPredicate([SnapshotName](const FPoseSnapshot& PoseData) { return PoseData.SnapshotName == SnapshotName; });
-	if (PoseSnapshot == nullptr)
-	{
-		PoseSnapshot = &PoseSnapshots[PoseSnapshots.AddDefaulted()];
-		PoseSnapshot->SnapshotName = SnapshotName;
-	}
-
-	InSkeletalMeshComponent->SnapshotPose(*PoseSnapshot);
+	FPoseSnapshot& PoseSnapshot = AddPoseSnapshot(SnapshotName);
+	InSkeletalMeshComponent->SnapshotPose(PoseSnapshot);
 }
 
 void FAnimInstanceProxy::PostUpdate(UAnimInstance* InAnimInstance) const
@@ -2745,6 +2739,32 @@ void FAnimInstanceProxy::RegisterWatchedPose(const FCSPose<FCompactPose>& Pose,
 }
 #endif
 
+FPoseSnapshot& FAnimInstanceProxy::AddPoseSnapshot(FName SnapshotName)
+{
+	FPoseSnapshot* PoseSnapshot = PoseSnapshots.FindByPredicate([SnapshotName](const FPoseSnapshot& PoseData) { return PoseData.SnapshotName == SnapshotName; });
+	if (PoseSnapshot)
+	{
+		// Recycle an existing snapshot
+		PoseSnapshot->Reset();
+	}
+	else
+	{
+		// Add a new empty snapshot
+		PoseSnapshot = &PoseSnapshots[PoseSnapshots.AddDefaulted()];
+	}
+	PoseSnapshot->SnapshotName = SnapshotName;
+	return *PoseSnapshot;
+}
+
+void FAnimInstanceProxy::RemovePoseSnapshot(FName SnapshotName)
+{
+	int32 Index = PoseSnapshots.IndexOfByPredicate([SnapshotName](const FPoseSnapshot& PoseData) { return PoseData.SnapshotName == SnapshotName; });
+	if (Index != INDEX_NONE)
+	{
+		PoseSnapshots.RemoveAtSwap(Index);
+	}
+}
+
 const FPoseSnapshot* FAnimInstanceProxy::GetPoseSnapshot(FName SnapshotName) const
 {
 	return PoseSnapshots.FindByPredicate([SnapshotName](const FPoseSnapshot& PoseData) { return PoseData.SnapshotName == SnapshotName; });
diff --git a/Engine/Source/Runtime/Engine/Public/Animation/AnimInstanceProxy.h b/Engine/Source/Runtime/Engine/Public/Animation/AnimInstanceProxy.h
--- a/Engine/Source/Runtime/Engine/Public/Animation/AnimInstanceProxy.h
+++ b/Engine/Source/Runtime/Engine/Public/Animation/AnimInstanceProxy.h
@@ -286,6 +286,12 @@ public:
 	/** Save a pose snapshot to the internal snapshot cache */
 	void SavePoseSnapshot(USkeletalMeshComponent* InSkeletalMeshComponent, FName SnapshotName);
 
+	/** Add an empty pose snapshot to the internal snapshot cache (or recycle an existing pose snapshot if the name is already in use) */
+	FPoseSnapshot& AddPoseSnapshot(FName SnapshotName);
+
+	/** Remove a previously saved pose snapshot from the internal snapshot cache */
+	void RemovePoseSnapshot(FName SnapshotName);
+
 	/** Get a cached pose snapshot by name */
 	const FPoseSnapshot* GetPoseSnapshot(FName SnapshotName) const;
 
diff --git a/Engine/Source/Runtime/Engine/Public/Animation/PoseSnapshot.h b/Engine/Source/Runtime/Engine/Public/Animation/PoseSnapshot.h
--- a/Engine/Source/Runtime/Engine/Public/Animation/PoseSnapshot.h
+++ b/Engine/Source/Runtime/Engine/Public/Animation/PoseSnapshot.h
@@ -16,6 +16,15 @@ public:
 		: bIsValid(false)
 	{}
 
+	void Reset()
+	{
+		LocalTransforms.Reset();
+		BoneNames.Reset();
+		SkeletalMeshName = NAME_None;
+		SnapshotName = NAME_None;
+		bIsValid = false;
+	}
+
 	/** Array of transforms per-bone */
 	UPROPERTY(BlueprintReadWrite, EditAnywhere, Category="Snapshot")
 	TArray<FTransform> LocalTransforms;
