new additive animation base pose mode to use a frame from the current sequence

diff --git a/Engine/Source/Editor/DetailCustomizations/Private/AnimSequenceDetails.cpp b/Engine/Source/Editor/DetailCustomizations/Private/AnimSequenceDetails.cpp
--- a/Engine/Source/Editor/DetailCustomizations/Private/AnimSequenceDetails.cpp
+++ b/Engine/Source/Editor/DetailCustomizations/Private/AnimSequenceDetails.cpp
@@ -230,7 +230,7 @@ EVisibility FAnimSequenceDetails::ShouldShowRefFrameIndex() const
 	uint8 RefPoseType = ABPT_None; 
 	AdditiveAnimTypeHandle->GetValue(AdditiveAnimType);
 	RefPoseTypeHandle->GetValue(RefPoseType);
-	return TargetSkeleton.IsValid() && AdditiveAnimType != AAT_None && RefPoseType == ABPT_AnimFrame? EVisibility::Visible : EVisibility::Collapsed;
+	return TargetSkeleton.IsValid() && AdditiveAnimType != AAT_None && (RefPoseType == ABPT_AnimFrame || RefPoseType == ABPT_LocalAnimFrame)? EVisibility::Visible : EVisibility::Collapsed;
 }
 
 
diff --git a/Engine/Source/Runtime/Engine/Classes/Animation/AnimEnums.h b/Engine/Source/Runtime/Engine/Classes/Animation/AnimEnums.h
--- a/Engine/Source/Runtime/Engine/Classes/Animation/AnimEnums.h
+++ b/Engine/Source/Runtime/Engine/Classes/Animation/AnimEnums.h
@@ -56,6 +56,8 @@ enum EAdditiveBasePoseType
 	ABPT_AnimScaled UMETA(DisplayName = "Selected animation scaled"),
 	/** Use one frame of an animation as a base pose. BasePoseSeq and RefFrameIndex must be set (RefFrameIndex will be clamped). */
 	ABPT_AnimFrame UMETA(DisplayName = "Selected animation frame"),
+	/** Use one frame of this animation. RefFrameIndex must be set (RefFrameIndex will be clamped). */
+	ABPT_LocalAnimFrame UMETA(DisplayName = "Frame from this animation"),
 	ABPT_MAX,
 };
 
diff --git a/Engine/Source/Runtime/Engine/Private/Animation/AnimSequence.cpp b/Engine/Source/Runtime/Engine/Private/Animation/AnimSequence.cpp
--- a/Engine/Source/Runtime/Engine/Private/Animation/AnimSequence.cpp
+++ b/Engine/Source/Runtime/Engine/Private/Animation/AnimSequence.cpp
@@ -1068,6 +1068,12 @@ void UAnimSequence::PostEditChangeProperty(FPropertyChangedEvent& PropertyChange
 			PropertyChangedEvent.Property->GetFName() == GET_MEMBER_NAME_CHECKED(UAnimSequence, RefPoseType))
 		{
 			bAdditiveSettingsChanged = true;
+
+			if (RefPoseType == ABPT_RefPose || RefPoseType == ABPT_LocalAnimFrame || AdditiveAnimType == AAT_None)
+			{
+				// clear RefPoseSeq when selecting settings which will hide the field
+				RefPoseSeq = nullptr;
+			}
 		}
 		
 		bCompressionAffectingSettingsChanged =   PropertyChangedEvent.Property->GetFName() == GET_MEMBER_NAME_CHECKED(UAnimSequence, bAllowFrameStripping)
@@ -1666,6 +1672,15 @@ void UAnimSequence::GetAdditiveBasePose(FCompactPose& OutPose, FBlendedCurve& Ou
 	GetAdditiveBasePose(OutAnimationPoseData, ExtractionContext);
 }
 
+static void GetSequencePose(FAnimationPoseData& OutAnimationPoseData, const FAnimExtractContext& ExtractionContext, const UAnimSequence &Seq, int FrameIndex)
+{
+	const float Fraction = (Seq.GetNumberOfSampledKeys() > 0) ? FMath::Clamp<float>((float)FrameIndex / (float)Seq.GetNumberOfSampledKeys(), 0.f, 1.f) : 0.f;
+	const float BasePoseTime = Seq.GetPlayLength() * Fraction;
+	FAnimExtractContext BasePoseExtractionContext(ExtractionContext);
+	BasePoseExtractionContext.CurrentTime = BasePoseTime;
+	Seq.GetBonePose(OutAnimationPoseData, BasePoseExtractionContext, true);
+}
+
 void UAnimSequence::GetAdditiveBasePose(FAnimationPoseData& OutAnimationPoseData, const FAnimExtractContext& ExtractionContext) const
 {
 	switch (RefPoseType)
@@ -1687,11 +1702,13 @@ void UAnimSequence::GetAdditiveBasePose(FAnimationPoseData& OutAnimationPoseData
 		// use animation as a base pose. Need BasePoseSeq and RefFrameIndex (will clamp if outside).
 		case ABPT_AnimFrame:
 		{
-			const float Fraction = (RefPoseSeq->GetNumberOfSampledKeys() > 0) ? FMath::Clamp<float>((float)RefFrameIndex / (float)RefPoseSeq->GetNumberOfSampledKeys(), 0.f, 1.f) : 0.f;
-			const float BasePoseTime = RefPoseSeq->GetPlayLength() * Fraction;
-			FAnimExtractContext BasePoseExtractionContext(ExtractionContext);
-			BasePoseExtractionContext.CurrentTime = BasePoseTime;
-			RefPoseSeq->GetBonePose(OutAnimationPoseData, BasePoseExtractionContext, true);
+			GetSequencePose(OutAnimationPoseData, ExtractionContext, *RefPoseSeq, RefFrameIndex);
+			break;
+		}
+		// use this animation as a base pose. Need RefFrameIndex (will clamp if outside).
+		case ABPT_LocalAnimFrame:
+		{
+			GetSequencePose(OutAnimationPoseData, ExtractionContext, *this, RefFrameIndex);
 			break;
 		}
 		// use ref pose of Skeleton as base
@@ -2694,6 +2711,8 @@ bool UAnimSequence::IsValidAdditive() const
 			return (RefPoseSeq != NULL);
 		case ABPT_AnimFrame:
 			return (RefPoseSeq != NULL) && (RefFrameIndex >= 0);
+		case ABPT_LocalAnimFrame:
+			return (RefFrameIndex >= 0);
 		default:
 			return false;
 		}
