Fix another bug building virtual textures with a max resolution.

diff --git a/Engine/Source/Runtime/Engine/Private/VT/VirtualTextureDataBuilder.cpp b/Engine/Source/Runtime/Engine/Private/VT/VirtualTextureDataBuilder.cpp
--- a/Engine/Source/Runtime/Engine/Private/VT/VirtualTextureDataBuilder.cpp
+++ b/Engine/Source/Runtime/Engine/Private/VT/VirtualTextureDataBuilder.cpp
@@ -315,8 +315,10 @@ void FVirtualTextureDataBuilder::Build(const FTextureSourceData& InSourceData, c
 	const uint32 ClampBlockSize = SettingsPerLayer[0].MaxTextureResolution;
 	if (FMath::Max<uint32>(BlockSizeX, BlockSizeY) > ClampBlockSize)
 	{
-		BlockSizeX = BlockSizeX >= BlockSizeY ? ClampBlockSize : FMath::Max(ClampBlockSize * BlockSizeX / BlockSizeY, 1u);
-		BlockSizeY = BlockSizeY >= BlockSizeX ? ClampBlockSize : FMath::Max(ClampBlockSize * BlockSizeY / BlockSizeX, 1u);
+		const int32 ClampedBlockSizeX = BlockSizeX >= BlockSizeY ? ClampBlockSize : FMath::Max(ClampBlockSize * BlockSizeX / BlockSizeY, 1u);
+		const int32 ClampedBlockSizeY = BlockSizeY >= BlockSizeX ? ClampBlockSize : FMath::Max(ClampBlockSize * BlockSizeY / BlockSizeX, 1u);
+		BlockSizeX = ClampedBlockSizeX;
+		BlockSizeY = ClampedBlockSizeY;
 	}
 
 	// We require VT blocks (UDIM pages) to be PoT, but multi block textures may have full logical dimension that's not PoT
