New Material Translator
- Some cleanup in the way material functions are wired up
- De-duplicate generated expression code by string equality...this works across material functions, and across parent-child scopes.  Duplicate code will remain in sibling scopes (there are cases where it makes sense to handle this better, but solving it in the general case is complicated)

diff --git a/Engine/Source/Runtime/Engine/Private/HLSLTree/HLSLTree.cpp b/Engine/Source/Runtime/Engine/Private/HLSLTree/HLSLTree.cpp
--- a/Engine/Source/Runtime/Engine/Private/HLSLTree/HLSLTree.cpp
+++ b/Engine/Source/Runtime/Engine/Private/HLSLTree/HLSLTree.cpp
@@ -132,6 +132,13 @@ UE::HLSLTree::FCodeWriter* UE::HLSLTree::FCodeWriter::Create(FMemStackBase& Allo
 	return new(Allocator) FCodeWriter(LocalStringBuilder);
 }
 
+FSHAHash UE::HLSLTree::FCodeWriter::GetCodeHash() const
+{
+	FSHAHash Hash;
+	FSHA1::HashBuffer(StringBuilder->GetData(), StringBuilder->Len() * sizeof(TCHAR), Hash.Hash);
+	return Hash;
+}
+
 void UE::HLSLTree::FCodeWriter::IncreaseIndent()
 {
 	++IndentLevel;
@@ -228,27 +235,41 @@ UE::HLSLTree::FEmitContext::FScopeEntry* UE::HLSLTree::FEmitContext::FindScope(U
 	return nullptr;
 }
 
+int32 UE::HLSLTree::FEmitContext::FindScopeIndex(FScope* Scope)
+{
+	for (int32 Index = ScopeStack.Num() - 1; Index >= 0; --Index)
+	{
+		FScopeEntry& Entry = ScopeStack[Index];
+		if (Entry.Scope == Scope)
+		{
+			return Index;
+		}
+	}
+	return INDEX_NONE;
+}
+
 const UE::HLSLTree::FEmitValue& UE::HLSLTree::FEmitContext::AcquireValue(UE::HLSLTree::FLocalDeclaration* Declaration)
 {
 	check(Declaration);
 	check(Declaration->Type != EExpressionType::Void);
 
-	FDeclarationEntry& Entry = FunctionStack.Last().DeclarationMap.FindOrAdd(Declaration);
-	if (Entry.Value.EvaluationType == EExpressionEvaluationType::None)
+	FDeclarationEntry*& Entry = FunctionStack.Last().DeclarationMap.FindOrAdd(Declaration);
+	if (!Entry)
 	{
 		FScopeEntry* ScopeEntry = FindScope(Declaration->ParentScope);
 		check(ScopeEntry);
+		Entry = new(*Allocator) FDeclarationEntry();
+		Entry->Value.EvaluationType = EExpressionEvaluationType::Shader;
 
-		Entry.Value.EvaluationType = EExpressionEvaluationType::Shader;
-		Entry.Value.Code = AllocateStringf(*Allocator, TEXT("Local%d"), NumExpressionLocals++);
+		Entry->Value.Code = AllocateStringf(*Allocator, TEXT("Local%d"), NumExpressionLocals++);
 		const FExpressionTypeDescription& TypeDesc = GetExpressionTypeDescription(Declaration->Type);
 		ScopeEntry->ExpressionCodeWriter->WriteLinef(TEXT("%s %s = (%s)0.0f;"),
 			TypeDesc.Name,
-			Entry.Value.Code,
+			Entry->Value.Code,
 			TypeDesc.Name);
 	}
 
-	return Entry.Value;
+	return Entry->Value;
 }
 
 const UE::HLSLTree::FEmitValue& UE::HLSLTree::FEmitContext::AcquireValue(UE::HLSLTree::FExpression* Expression)
@@ -256,16 +277,17 @@ const UE::HLSLTree::FEmitValue& UE::HLSLTree::FEmitContext::AcquireValue(UE::HLS
 	check(Expression);
 	check(Expression->Type != EExpressionType::Void);
 
-	FDeclarationEntry& Entry = FunctionStack.Last().DeclarationMap.FindOrAdd(Expression);
-	if (Entry.Value.EvaluationType == EExpressionEvaluationType::None)
+	FDeclarationEntry*& Entry = FunctionStack.Last().DeclarationMap.FindOrAdd(Expression);
+	if (!Entry)
 	{
 		FLocalHLSLCodeWriter LocalWriter;
 		FMaterialPreshaderData LocalPreshader;
 		FExpressionEmitResult EmitResult(LocalWriter, LocalPreshader);
 		Expression->EmitHLSL(*this, EmitResult);
 
-		Entry.Value.ExpressionType = Expression->Type;
-		Entry.Value.EvaluationType = EmitResult.EvaluationType;
+		Entry = new(*Allocator) FDeclarationEntry();
+		Entry->Value.ExpressionType = Expression->Type;
+		Entry->Value.EvaluationType = EmitResult.EvaluationType;
 		if (EmitResult.EvaluationType == EExpressionEvaluationType::Constant)
 		{
 			// Evaluate the constant preshader and store its value
@@ -274,94 +296,123 @@ const UE::HLSLTree::FEmitValue& UE::HLSLTree::FEmitContext::AcquireValue(UE::HLS
 			LocalPreshader.Evaluate(nullptr, RenderContext, ConstantValue);
 
 			const FConstant Constant(Expression->Type, ConstantValue);
-			Entry.Value.ConstantValue = Constant;
+			Entry->Value.ConstantValue = Constant;
 		}
 		else if (EmitResult.EvaluationType == EExpressionEvaluationType::Preshader)
 		{
 			// Non-constant preshader, store it
 			FMaterialPreshaderData* Preshader = new FMaterialPreshaderData(MoveTemp(LocalPreshader));
 			TempPreshaders.Add(Preshader); // TODO - dedupe? store more efficiently?
-			Entry.Value.Preshader = Preshader;
+			Entry->Value.Preshader = Preshader;
 		}
 		else
 		{
 			check(EmitResult.EvaluationType == EExpressionEvaluationType::Shader);
 			if (EmitResult.bInline)
 			{
-				Entry.Value.Code = AllocateString(*Allocator, LocalWriter.GetStringBuilder());
+				Entry->Value.Code = AllocateString(*Allocator, LocalWriter.GetStringBuilder());
 			}
 			else
 			{
-				FScopeEntry* ScopeEntry = FindScope(Expression->ParentScope);
-				check(ScopeEntry);
-
-				const FExpressionTypeDescription& TypeDesc = GetExpressionTypeDescription(Expression->Type);
-				Entry.Value.Code = AllocateStringf(*Allocator, TEXT("Local%d"), NumExpressionLocals++);
-				ScopeEntry->ExpressionCodeWriter->WriteLinef(TEXT("%s %s = %s;"),
-					TypeDesc.Name,
-					Entry.Value.Code,
-					LocalWriter.GetStringBuilder().ToString());
+				const int32 ScopeIndex = FindScopeIndex(Expression->ParentScope);
+				check(ScopeIndex != INDEX_NONE);
+		
+				// Check to see if we've already generated code for an equivalent expression in either this scope, or any outer visible scope
+				const FSHAHash Hash = LocalWriter.GetCodeHash();
+				const TCHAR* Declaration = nullptr;
+				for (int32 CheckScopeIndex = ScopeIndex; CheckScopeIndex >= 0; --CheckScopeIndex)
+				{
+					const FScopeEntry& CheckScopeEntry = ScopeStack[CheckScopeIndex];
+					const TCHAR** FoundDeclaration = CheckScopeEntry.ExpressionMap->Find(Hash);
+					if (FoundDeclaration)
+					{
+						// Re-use results from previous matching expression
+						Declaration = *FoundDeclaration;
+						break;
+					}
+				}
+				
+				if (!Declaration)
+				{
+					const FExpressionTypeDescription& TypeDesc = GetExpressionTypeDescription(Expression->Type);
+					Declaration = AllocateStringf(*Allocator, TEXT("Local%d"), NumExpressionLocals++);
+					ScopeStack[ScopeIndex].ExpressionCodeWriter->WriteLinef(TEXT("const %s %s = %s;"),
+						TypeDesc.Name,
+						Declaration,
+						LocalWriter.GetStringBuilder().ToString());
+					ScopeStack[ScopeIndex].ExpressionMap->Add(Hash, Declaration);
+				}
+				Entry->Value.Code = Declaration;
 			}
 		}
 	}
 
-	return Entry.Value;
+	return Entry->Value;
 }
 
 const UE::HLSLTree::FEmitValue& UE::HLSLTree::FEmitContext::AcquireValue(FFunctionCall* FunctionCall, int32 OutputIndex)
 {
 	check(FunctionCall);
 
-	FFunctionCallEntry& Entry = FunctionStack.Last().FunctionCallMap.FindOrAdd(FunctionCall);
-	if (!Entry.OutputRef)
+	FFunctionCallEntry*& Entry = FunctionStack.Last().FunctionCallMap.FindOrAdd(FunctionCall);
+	if (!Entry)
 	{
-		FScopeEntry* ScopeEntry = FindScope(FunctionCall->ParentScope);
-		check(ScopeEntry);
-
-		FEmitValue* OutputRef = new(*Allocator) FEmitValue[FunctionCall->NumOutputs];
-		for (int32 i = 0; i < FunctionCall->NumOutputs; ++i)
-		{
-			OutputRef[i].Code = AllocateStringf(*Allocator, TEXT("Local%d"), NumExpressionLocals++);
-		}
-
-		Entry.OutputRef = OutputRef;
-		Entry.NumOutputs = FunctionCall->NumOutputs;
+		FScopeEntry* ParentScopeEntry = FindScope(FunctionCall->ParentScope);
+		check(ParentScopeEntry);
 
 		{
 			FFunctionStackEntry& StackEntry = FunctionStack.AddDefaulted_GetRef();
 			StackEntry.FunctionCall = FunctionCall;
-			StackEntry.OutputRef = OutputRef;
 		}
-		FLocalHLSLCodeWriter FunctionWriter;
-		FunctionWriter.IndentLevel = ScopeEntry->ExpressionCodeWriter->IndentLevel;
-		FunctionCall->FunctionScope->EmitHLSL(*this, FunctionWriter);
+
 		{
-			const FFunctionStackEntry PoppedStackEntry = FunctionStack.Pop();
-			check(PoppedStackEntry.FunctionCall == FunctionCall);
+			FScopeEntry& ScopeEntry = ScopeStack.AddDefaulted_GetRef();
+			ScopeEntry.Scope = FunctionCall->FunctionScope;
+			ScopeEntry.ExpressionCodeWriter = ParentScopeEntry->ExpressionCodeWriter;
+			ScopeEntry.ExpressionMap = ParentScopeEntry->ExpressionMap;
 		}
 
-		// Allocate locals to hold any function outputs evaluated in shader
+		// Emit the function's HLSL into the current scope (allows sharing expressions across functions called from the same scope)
+		FLocalHLSLCodeWriter LocalWriter;
+		FunctionCall->FunctionScope->EmitUnscopedHLSL(*this, LocalWriter);
+		ParentScopeEntry->ExpressionCodeWriter->Append(LocalWriter);
+
+		// Assign function outputs
+		FEmitValue* OutputValue = new(*Allocator) FEmitValue[FunctionCall->NumOutputs];
 		for (int32 i = 0; i < FunctionCall->NumOutputs; ++i)
 		{
-			const FEmitValue& Output = OutputRef[i];
-			if (Output.EvaluationType == EExpressionEvaluationType::Shader)
+			FExpression* OutputExpression = FunctionCall->Outputs[i];
+			FEmitValue& Output = OutputValue[i];
+			if (OutputExpression)
+			{
+				Output = AcquireValue(OutputExpression);
+			}
+			else
 			{
-				const EExpressionType OutputType = FunctionCall->OutputTypes[i];
-				const FExpressionTypeDescription& TypeDesc = GetExpressionTypeDescription(OutputType);
-				ScopeEntry->ExpressionCodeWriter->WriteLinef(TEXT("%s %s = (%s)0.0f;"),
-					TypeDesc.Name,
-					OutputRef[i].Code,
-					TypeDesc.Name);
+				Output.EvaluationType = EExpressionEvaluationType::Constant;
+				Output.ExpressionType = EExpressionType::Float1;
+				Output.ConstantValue = 0.0f;
 			}
 		}
 
-		// Code to evaluate actual function (TODO - cull empty scopes)
-		ScopeEntry->ExpressionCodeWriter->Append(FunctionWriter);
+		{
+			const FScopeEntry PoppedScopeEntry = ScopeStack.Pop(false);
+			check(PoppedScopeEntry.Scope == FunctionCall->FunctionScope);
+		}
+
+		{
+			const FFunctionStackEntry PoppedStackEntry = FunctionStack.Pop(false);
+			check(PoppedStackEntry.FunctionCall == FunctionCall);
+		}
+
+		Entry = new(*Allocator) FFunctionCallEntry();
+		Entry->OutputRef = OutputValue;
+		Entry->NumOutputs = FunctionCall->NumOutputs;
 	}
 
-	check(Entry.NumOutputs == FunctionCall->NumOutputs);
-	check(OutputIndex >= 0 && OutputIndex < Entry.NumOutputs);
-	return Entry.OutputRef[OutputIndex];
+	check(Entry->NumOutputs == FunctionCall->NumOutputs);
+	check(OutputIndex >= 0 && OutputIndex < Entry->NumOutputs);
+	return Entry->OutputRef[OutputIndex];
 }
 
 const TCHAR* UE::HLSLTree::FEmitContext::GetCode(const FEmitValue& Value) const
@@ -493,18 +544,21 @@ UE::HLSLTree::ENodeVisitResult UE::HLSLTree::FScope::Visit(UE::HLSLTree::FNodeVi
 	return Result;
 }
 
-void UE::HLSLTree::FScope::EmitHLSL(UE::HLSLTree::FEmitContext& Context, UE::HLSLTree::FCodeWriter& Writer) const
+void UE::HLSLTree::FScope::EmitHLSL(FEmitContext& Context, FCodeWriter& OutWriter) const
 {
-	const int32 ScopeIndentLevel = Writer.IndentLevel + 1;
+	const int32 ScopeIndentLevel = OutWriter.IndentLevel + 1;
 
 	FLocalHLSLCodeWriter DeclarationCodeWriter;
 	DeclarationCodeWriter.IndentLevel = ScopeIndentLevel;
 
+	TMap<FSHAHash, const TCHAR*> LocalExpressionMap;
+
 	FEmitContext::FScopeEntry& ScopeEntry = Context.ScopeStack.AddDefaulted_GetRef();
 	ScopeEntry.Scope = this;
 	ScopeEntry.ExpressionCodeWriter = &DeclarationCodeWriter;
+	ScopeEntry.ExpressionMap = &LocalExpressionMap;
 
-	Writer.WriteLine(TEXT("{"));
+	OutWriter.WriteLine(TEXT("{"));
 
 	FLocalHLSLCodeWriter StatementCodeWriter;
 	StatementCodeWriter.IndentLevel = ScopeIndentLevel;
@@ -514,20 +568,41 @@ void UE::HLSLTree::FScope::EmitHLSL(UE::HLSLTree::FEmitContext& Context, UE::HLS
 		Statement->EmitHLSL(Context, StatementCodeWriter);
 
 		// First write any expressions needed by the statement, then the statement itself
-		Writer.Append(DeclarationCodeWriter);
-		Writer.Append(StatementCodeWriter);
+		OutWriter.Append(DeclarationCodeWriter);
+		OutWriter.Append(StatementCodeWriter);
 
-		DeclarationCodeWriter.StringBuilder->Reset();
-		StatementCodeWriter.StringBuilder->Reset();
+		DeclarationCodeWriter.Reset();
+		StatementCodeWriter.Reset();
 
 		Statement = Statement->NextStatement;
 	}
 
-	Writer.WriteLine(TEXT("}"));
-
+	OutWriter.WriteLine(TEXT("}"));
 	Context.ScopeStack.Pop(false);
 }
 
+void UE::HLSLTree::FScope::EmitUnscopedHLSL(FEmitContext& Context, FCodeWriter& OutWriter) const
+{
+	FCodeWriter& DeclarationCodeWriter = *Context.ScopeStack.Last().ExpressionCodeWriter;
+
+	FLocalHLSLCodeWriter StatementCodeWriter;
+	StatementCodeWriter.IndentLevel = OutWriter.IndentLevel;
+	const FStatement* Statement = FirstStatement;
+	while (Statement)
+	{
+		Statement->EmitHLSL(Context, StatementCodeWriter);
+
+		// First write any expressions needed by the statement, then the statement itself
+		OutWriter.Append(DeclarationCodeWriter);
+		OutWriter.Append(StatementCodeWriter);
+
+		DeclarationCodeWriter.Reset();
+		StatementCodeWriter.Reset();
+
+		Statement = Statement->NextStatement;
+	}
+}
+
 void UE::HLSLTree::FScope::AddDeclaration(UE::HLSLTree::FLocalDeclaration* Declaration)
 {
 	check(!Declaration->ParentScope);
@@ -747,20 +822,20 @@ UE::HLSLTree::FTextureParameterDeclaration* UE::HLSLTree::FTree::NewTextureParam
 UE::HLSLTree::FFunctionCall* UE::HLSLTree::FTree::NewFunctionCall(FScope& Scope,
 	const FScope& InFunctionScope,
 	FExpression* const* InInputs,
-	EExpressionType const* InOutputTypes,
+	FExpression* const* InOutputs,
 	int32 InNumInputs,
 	int32 InNumOutputs)
 {
 	FExpression** Inputs = New<FExpression*>(*Allocator, InNumInputs);
-	EExpressionType* OutputTypes = New<EExpressionType>(*Allocator, InNumOutputs);
+	FExpression** Outputs = New<FExpression*>(*Allocator, InNumOutputs);
 	FMemory::Memcpy(Inputs, InInputs, InNumInputs * sizeof(FExpression*));
-	FMemory::Memcpy(OutputTypes, InOutputTypes, InNumOutputs * sizeof(FExpression*));
+	FMemory::Memcpy(Outputs, InOutputs, InNumOutputs * sizeof(FExpression*));
 
 	FFunctionCall* FunctionCall = NewNode<FFunctionCall>();
 	FunctionCall->ParentScope = &Scope;
 	FunctionCall->FunctionScope = &InFunctionScope;
 	FunctionCall->Inputs = Inputs;
-	FunctionCall->OutputTypes = OutputTypes;
+	FunctionCall->Outputs = Outputs;
 	FunctionCall->NumInputs = InNumInputs;
 	FunctionCall->NumOutputs = InNumOutputs;
 	return FunctionCall;
diff --git a/Engine/Source/Runtime/Engine/Private/HLSLTree/HLSLTreeCommon.cpp b/Engine/Source/Runtime/Engine/Private/HLSLTree/HLSLTreeCommon.cpp
--- a/Engine/Source/Runtime/Engine/Private/HLSLTree/HLSLTreeCommon.cpp
+++ b/Engine/Source/Runtime/Engine/Private/HLSLTree/HLSLTreeCommon.cpp
@@ -414,20 +414,6 @@ void UE::HLSLTree::FExpressionFunctionOutput::EmitHLSL(FEmitContext& Context, FE
 	}
 }
 
-void UE::HLSLTree::FStatementSetFunctionOutput::EmitHLSL(FEmitContext& Context, FCodeWriter& OutWriter) const
-{
-	const FEmitContext::FFunctionStackEntry& StackEntry = Context.FunctionStack.Last();
-	const FEmitValue& ExpressionRef = Context.AcquireValue(Expression);
-
-	check(OutputIndex >= 0 && OutputIndex < StackEntry.FunctionCall->NumOutputs);
-	FEmitValue& OutputRef = StackEntry.OutputRef[OutputIndex];
-	OutputRef.Assign(ExpressionRef);
-	if (OutputRef.GetEvaluationType() == EExpressionEvaluationType::Shader)
-	{
-		OutWriter.WriteLinef(TEXT("%s = %s;"), Context.GetCode(OutputRef), Context.GetCode(ExpressionRef));
-	}
-}
-
 void UE::HLSLTree::FStatementReturn::EmitHLSL(UE::HLSLTree::FEmitContext& Context, UE::HLSLTree::FCodeWriter& OutWriter) const
 {
 	const FEmitValue& Ref = Context.AcquireValue(Expression);
diff --git a/Engine/Source/Runtime/Engine/Private/Materials/MaterialExpressionHLSL.cpp b/Engine/Source/Runtime/Engine/Private/Materials/MaterialExpressionHLSL.cpp
--- a/Engine/Source/Runtime/Engine/Private/Materials/MaterialExpressionHLSL.cpp
+++ b/Engine/Source/Runtime/Engine/Private/Materials/MaterialExpressionHLSL.cpp
@@ -263,6 +263,11 @@ EMaterialGenerateHLSLStatus UMaterialExpressionMaterialFunctionCall::GenerateHLS
 	{
 		const FFunctionExpressionInput& Input = FunctionInputs[InputIndex];
 		UE::HLSLTree::FExpression* InputExpression = Input.Input.AcquireHLSLExpression(Generator, Scope);
+		if (!InputExpression)
+		{
+			InputExpression = Generator.AcquireExpression(Scope, Input.ExpressionInput, 0);
+		}
+		check(InputExpression);
 		InputExpressions.Add(InputExpression);
 	}
 
diff --git a/Engine/Source/Runtime/Engine/Private/Materials/MaterialHLSLTree.cpp b/Engine/Source/Runtime/Engine/Private/Materials/MaterialHLSLTree.cpp
--- a/Engine/Source/Runtime/Engine/Private/Materials/MaterialHLSLTree.cpp
+++ b/Engine/Source/Runtime/Engine/Private/Materials/MaterialHLSLTree.cpp
@@ -40,7 +40,7 @@ void FMaterialHLSLTree::InitializeForFunction(const FMaterialCompileTargetParame
 
 	NumFunctionInputs = FunctionInputs.Num();
 	NumFunctionOutputs = FunctionOutputs.Num();
-	FunctionOutputTypes = New<UE::HLSLTree::EExpressionType>(Allocator, NumFunctionOutputs);
+	FunctionOutputExpressions = New<UE::HLSLTree::FExpression*>(Allocator, NumFunctionOutputs);
 
 	FMaterialHLSLGenerator Generator(InCompilerTarget, *HLSLTree);
 	for (int32 InputIndex = 0; InputIndex < FunctionInputs.Num(); ++InputIndex)
@@ -53,18 +53,7 @@ void FMaterialHLSLTree::InitializeForFunction(const FMaterialCompileTargetParame
 	{
 		const FFunctionExpressionOutput& Output = FunctionOutputs[OutputIndex];
 		UMaterialExpressionFunctionOutput* ExpressionOutput = Output.ExpressionOutput;
-
-		UE::HLSLTree::FExpression* Expression = ExpressionOutput->A.AcquireHLSLExpression(Generator, RootScope);
-		UE::HLSLTree::EExpressionType OutputType = UE::HLSLTree::EExpressionType::Void;
-		if (Expression)
-		{
-			UE::HLSLTree::FStatementSetFunctionOutput* OutputStatement = Generator.GetTree().NewStatement<UE::HLSLTree::FStatementSetFunctionOutput>(RootScope);
-			OutputStatement->Name = ExpressionOutput->OutputName;
-			OutputStatement->OutputIndex = OutputIndex;
-			OutputStatement->Expression = Expression;
-			OutputType = Expression->Type;
-		}
-		FunctionOutputTypes[OutputIndex] = OutputType;
+		FunctionOutputExpressions[OutputIndex] = ExpressionOutput->A.AcquireHLSLExpression(Generator, RootScope);
 	}
 }
 
@@ -75,7 +64,7 @@ UE::HLSLTree::FFunctionCall* FMaterialHLSLTree::GenerateFunctionCall(FMaterialHL
 	return HLSLTree->NewFunctionCall(Scope,
 		HLSLTree->GetRootScope(),
 		Inputs.GetData(),
-		FunctionOutputTypes,
+		FunctionOutputExpressions,
 		NumFunctionInputs,
 		NumFunctionOutputs);
 }
diff --git a/Engine/Source/Runtime/Engine/Public/HLSLTree/HLSLTree.h b/Engine/Source/Runtime/Engine/Public/HLSLTree/HLSLTree.h
--- a/Engine/Source/Runtime/Engine/Public/HLSLTree/HLSLTree.h
+++ b/Engine/Source/Runtime/Engine/Public/HLSLTree/HLSLTree.h
@@ -29,6 +29,8 @@ public:
 
 	const FStringBuilderBase& GetStringBuilder() const { return *StringBuilder; }
 
+	FSHAHash GetCodeHash() const;
+
 	void IncreaseIndent();
 	void DecreaseIndent();
 
@@ -81,13 +83,6 @@ public:
 	EExpressionEvaluationType GetEvaluationType() const { return EvaluationType; }
 	const FConstant& GetConstantValue() const { return ConstantValue; }
 
-	inline void Assign(const FEmitValue& InRef)
-	{
-		Preshader = InRef.Preshader;
-		EvaluationType = InRef.EvaluationType;
-		ConstantValue = InRef.ConstantValue;
-	}
-
 private:
 	mutable const TCHAR* Code = nullptr;
 	const FMaterialPreshaderData* Preshader = nullptr;
@@ -124,6 +119,7 @@ public:
 	{
 		const FScope* Scope;
 		FCodeWriter* ExpressionCodeWriter;
+		TMap<FSHAHash, const TCHAR*>* ExpressionMap;
 	};
 
 	struct FDeclarationEntry
@@ -140,13 +136,13 @@ public:
 	struct FFunctionStackEntry
 	{
 		FFunctionCall* FunctionCall = nullptr;
-		FEmitValue* OutputRef = nullptr;
 
-		TMap<FNode*, FDeclarationEntry> DeclarationMap;
-		TMap<FFunctionCall*, FFunctionCallEntry> FunctionCallMap;
+		TMap<FNode*, FDeclarationEntry*> DeclarationMap;
+		TMap<FFunctionCall*, FFunctionCallEntry*> FunctionCallMap;
 	};
 
 	FScopeEntry* FindScope(FScope* Scope);
+	int32 FindScopeIndex(FScope* Scope);
 
 	TArray<FScopeEntry> ScopeStack;
 	TArray<FFunctionStackEntry> FunctionStack;
@@ -300,11 +296,21 @@ class FFunctionCall final : public FNode
 public:
 	virtual ENodeVisitResult Visit(FNodeVisitor& Visitor) override;
 
+	inline EExpressionType GetOutputType(int32 Index) const
+	{
+		check(Index >= 0 && Index < NumOutputs);
+		return Outputs[Index] ? Outputs[Index]->Type : EExpressionType::Float1;
+	}
+
 	/** Root scope of the function to call. Note that this scope will be from a separate (external) tree */
 	const FScope* FunctionScope;
 
+	/** Outputs are expressions from the function's scope */
+	FExpression* const* Outputs;
+
+	/** Inputs are connected to the calling scope */
 	FExpression* const* Inputs;
-	EExpressionType const* OutputTypes;
+	
 	int32 NumInputs;
 	int32 NumOutputs;
 };
@@ -317,7 +323,9 @@ class FScope final : public FNode
 {
 public:
 	virtual ENodeVisitResult Visit(FNodeVisitor& Visitor) override;
-	void EmitHLSL(FEmitContext& Context, FCodeWriter& Writer) const;
+	void EmitHLSL(FEmitContext& Context, FCodeWriter& OutWriter) const;
+
+	void EmitUnscopedHLSL(FEmitContext& Context, FCodeWriter& OutWriter) const;
 
 	void AddDeclaration(FLocalDeclaration* Declaration);
 	void AddExpression(FExpression* Expression);
@@ -389,7 +397,7 @@ public:
 	FFunctionCall* NewFunctionCall(FScope& Scope,
 		const FScope& FunctionScope,
 		FExpression* const* Inputs,
-		EExpressionType const* OutputTypes,
+		FExpression* const* Outputs,
 		int32 NumInputs,
 		int32 NumOutputs);
 
diff --git a/Engine/Source/Runtime/Engine/Public/HLSLTree/HLSLTreeCommon.h b/Engine/Source/Runtime/Engine/Public/HLSLTree/HLSLTreeCommon.h
--- a/Engine/Source/Runtime/Engine/Public/HLSLTree/HLSLTreeCommon.h
+++ b/Engine/Source/Runtime/Engine/Public/HLSLTree/HLSLTreeCommon.h
@@ -269,7 +269,7 @@ class FExpressionFunctionOutput : public FExpression
 {
 public:
 	FExpressionFunctionOutput(FFunctionCall* InFunctionCall, int32 InIndex)
-		: FExpression(InFunctionCall->OutputTypes[InIndex])
+		: FExpression(InFunctionCall->GetOutputType(InIndex))
 		, FunctionCall(InFunctionCall)
 		, OutputIndex(InIndex)
 	{
@@ -292,26 +292,6 @@ public:
 	virtual void EmitHLSL(FEmitContext& Context, FExpressionEmitResult& OutResult) const override;
 };
 
-class FStatementSetFunctionOutput : public FStatement
-{
-public:
-	FExpression* Expression;
-	FName Name;
-	int32 OutputIndex;
-
-	virtual ENodeVisitResult Visit(FNodeVisitor& Visitor) override
-	{
-		const ENodeVisitResult Result = FStatement::Visit(Visitor);
-		if (ShouldVisitDependentNodes(Result))
-		{
-			Visitor.VisitNode(Expression);
-		}
-		return Result;
-	}
-
-	virtual void EmitHLSL(FEmitContext& Context, FCodeWriter& Writer) const override;
-};
-
 class FStatementReturn : public FStatement
 {
 public:
diff --git a/Engine/Source/Runtime/Engine/Public/MaterialHLSLTree.h b/Engine/Source/Runtime/Engine/Public/MaterialHLSLTree.h
--- a/Engine/Source/Runtime/Engine/Public/MaterialHLSLTree.h
+++ b/Engine/Source/Runtime/Engine/Public/MaterialHLSLTree.h
@@ -19,7 +19,6 @@ class FScope;
 class FExpression;
 class FLocalDeclaration;
 class FFunctionCall;
-enum class EExpressionType : uint8;
 }
 }
 
@@ -40,7 +39,7 @@ public:
 private:
 	FMemStackBase Allocator;
 	UE::HLSLTree::FTree* HLSLTree = nullptr;
-	UE::HLSLTree::EExpressionType* FunctionOutputTypes = nullptr;
+	UE::HLSLTree::FExpression** FunctionOutputExpressions = nullptr;
 	int32 NumFunctionInputs = 0;
 	int32 NumFunctionOutputs = 0;
 };
