Add bNormalize option to Read Render Target Pixel editor functions in order to read actual height/depth data

diff --git a/Engine/Source/Runtime/Engine/Classes/Kismet/KismetRenderingLibrary.h b/Engine/Source/Runtime/Engine/Classes/Kismet/KismetRenderingLibrary.h
--- a/Engine/Source/Runtime/Engine/Classes/Kismet/KismetRenderingLibrary.h
+++ b/Engine/Source/Runtime/Engine/Classes/Kismet/KismetRenderingLibrary.h
@@ -117,13 +117,13 @@ class UKismetRenderingLibrary : public UBlueprintFunctionLibrary
 	* Incredibly inefficient and slow operation! Read a value as-is from a render target using integer pixel coordinates.
 	*/
 	UFUNCTION(BlueprintCallable, Category = "Rendering", meta = (Keywords = "ReadRenderTarget", WorldContext = "WorldContextObject"))
-	static ENGINE_API FLinearColor ReadRenderTargetRawPixel(UObject* WorldContextObject, UTextureRenderTarget2D* TextureRenderTarget, int32 X, int32 Y);
+	static ENGINE_API FLinearColor ReadRenderTargetRawPixel(UObject* WorldContextObject, UTextureRenderTarget2D* TextureRenderTarget, int32 X, int32 Y, bool bNormalize = true);
 
 	/**
 	* Incredibly inefficient and slow operation! Read a value as-is color from a render target using UV [0,1]x[0,1] coordinates.
 	*/
 	UFUNCTION(BlueprintCallable, Category = "Rendering", meta = (Keywords = "ReadRenderTarget", WorldContext = "WorldContextObject"))
-	static ENGINE_API FLinearColor ReadRenderTargetRawUV(UObject* WorldContextObject, UTextureRenderTarget2D* TextureRenderTarget, float U, float V);
+	static ENGINE_API FLinearColor ReadRenderTargetRawUV(UObject* WorldContextObject, UTextureRenderTarget2D* TextureRenderTarget, float U, float V, bool bNormalize = true);
 
 	/**
 	 * Exports a Texture2D as a HDR image onto the disk.
diff --git a/Engine/Source/Runtime/Engine/Private/KismetRenderingLibrary.cpp b/Engine/Source/Runtime/Engine/Private/KismetRenderingLibrary.cpp
--- a/Engine/Source/Runtime/Engine/Private/KismetRenderingLibrary.cpp
+++ b/Engine/Source/Runtime/Engine/Private/KismetRenderingLibrary.cpp
@@ -276,7 +276,8 @@ EPixelFormat ReadRenderTargetHelper(
 	int32 X,
 	int32 Y,
 	int32 Width,
-	int32 Height)
+	int32 Height,
+	bool bNormalize = true)
 {
 	EPixelFormat OutFormat = PF_Unknown;
 
@@ -299,7 +300,8 @@ EPixelFormat ReadRenderTargetHelper(
 	Height = Height - FMath::Max(Y + Height - TextureRenderTarget->SizeY, 0);
 
 	FIntRect SampleRect(X, Y, X + Width, Y + Height);
-	FReadSurfaceDataFlags ReadSurfaceDataFlags;
+
+	FReadSurfaceDataFlags ReadSurfaceDataFlags = bNormalize ? FReadSurfaceDataFlags() : FReadSurfaceDataFlags(RCM_MinMax);
 
 	FRenderTarget* RenderTarget = TextureRenderTarget->GameThread_GetRenderTargetResource();
 	OutFormat = TextureRenderTarget->GetFormat();
@@ -364,12 +366,12 @@ FColor UKismetRenderingLibrary::ReadRenderTargetPixel(UObject* WorldContextObjec
 	}
 }
 
-FLinearColor UKismetRenderingLibrary::ReadRenderTargetRawPixel(UObject * WorldContextObject, UTextureRenderTarget2D * TextureRenderTarget, int32 X, int32 Y)
+FLinearColor UKismetRenderingLibrary::ReadRenderTargetRawPixel(UObject * WorldContextObject, UTextureRenderTarget2D * TextureRenderTarget, int32 X, int32 Y, bool bNormalize)
 {
 	TArray<FColor> Samples;
 	TArray<FLinearColor> LinearSamples;
 
-	switch (ReadRenderTargetHelper(Samples, LinearSamples, WorldContextObject, TextureRenderTarget, X, Y, 1, 1))
+	switch (ReadRenderTargetHelper(Samples, LinearSamples, WorldContextObject, TextureRenderTarget, X, Y, 1, 1, bNormalize))
 	{
 	case PF_B8G8R8A8:
 		check(Samples.Num() == 1 && LinearSamples.Num() == 0);
@@ -383,7 +385,7 @@ FLinearColor UKismetRenderingLibrary::ReadRenderTargetRawPixel(UObject * WorldCo
 	}
 }
 
-FLinearColor UKismetRenderingLibrary::ReadRenderTargetRawUV(UObject * WorldContextObject, UTextureRenderTarget2D * TextureRenderTarget, float U, float V)
+FLinearColor UKismetRenderingLibrary::ReadRenderTargetRawUV(UObject * WorldContextObject, UTextureRenderTarget2D * TextureRenderTarget, float U, float V, bool bNormalize)
 {
 	if (!TextureRenderTarget)
 	{
@@ -395,7 +397,7 @@ FLinearColor UKismetRenderingLibrary::ReadRenderTargetRawUV(UObject * WorldConte
 	int32 XPos = U * (float)TextureRenderTarget->SizeX;
 	int32 YPos = V * (float)TextureRenderTarget->SizeY;
 
-	return ReadRenderTargetRawPixel(WorldContextObject, TextureRenderTarget, XPos, YPos);
+	return ReadRenderTargetRawPixel(WorldContextObject, TextureRenderTarget, XPos, YPos, bNormalize);
 }
 
 /*
