Remove dependency on Asset Registry callbacks for actor added/updated/deleted.

diff --git a/Engine/Source/Editor/UnrealEd/Private/Editor.cpp b/Engine/Source/Editor/UnrealEd/Private/Editor.cpp
--- a/Engine/Source/Editor/UnrealEd/Private/Editor.cpp
+++ b/Engine/Source/Editor/UnrealEd/Private/Editor.cpp
@@ -132,6 +132,7 @@ FEditorDelegates::FOnMapOpened							FEditorDelegates::OnMapOpened;
 FEditorDelegates::FOnEditorCameraMoved					FEditorDelegates::OnEditorCameraMoved;
 FEditorDelegates::FOnDollyPerspectiveCamera				FEditorDelegates::OnDollyPerspectiveCamera;
 FSimpleMulticastDelegate								FEditorDelegates::OnShutdownPostPackagesSaved;
+FEditorDelegates::FOnPackageDeleted						FEditorDelegates::OnPackageDeleted;
 FEditorDelegates::FOnAssetsCanDelete					FEditorDelegates::OnAssetsCanDelete;
 FEditorDelegates::FOnAssetsAddExtraObjectsToDelete			FEditorDelegates::OnAssetsAddExtraObjectsToDelete;
 FEditorDelegates::FOnAssetsPreDelete					FEditorDelegates::OnAssetsPreDelete;
diff --git a/Engine/Source/Editor/UnrealEd/Private/ObjectTools.cpp b/Engine/Source/Editor/UnrealEd/Private/ObjectTools.cpp
--- a/Engine/Source/Editor/UnrealEd/Private/ObjectTools.cpp
+++ b/Engine/Source/Editor/UnrealEd/Private/ObjectTools.cpp
@@ -2125,6 +2125,7 @@ namespace ObjectTools
 		for (UPackage* PackageToDelete : PackagesToDelete)
 		{
 			FAssetRegistryModule::PackageDeleted(PackageToDelete);
+			FEditorDelegates::OnPackageDeleted.Broadcast(PackageToDelete);
 		}
 
 		// Unload the packages and collect garbage.
diff --git a/Engine/Source/Editor/UnrealEd/Public/Editor.h b/Engine/Source/Editor/UnrealEd/Public/Editor.h
--- a/Engine/Source/Editor/UnrealEd/Public/Editor.h
+++ b/Engine/Source/Editor/UnrealEd/Public/Editor.h
@@ -102,6 +102,8 @@ struct UNREALED_API FEditorDelegates
 	DECLARE_MULTICAST_DELEGATE_OneParam(FOnEditorModeIDTransitioned, const FEditorModeID& /*Mode*/);
 	/** delegate type to determine if a user requests can delete certain assets. */
 	DECLARE_MULTICAST_DELEGATE_TwoParams(FOnAssetsCanDelete, const TArray<UObject*>& /*InObjectToDelete*/, FCanDeleteAssetResult& /*OutCanDelete*/);
+	/** delegate type for when a user requests to delete certain package */
+	DECLARE_MULTICAST_DELEGATE_OneParam(FOnPackageDeleted, UPackage*);
 	/** delegate type for when a user requests to delete certain assets... It allows the addition of secondary assets that should also be deleted */
 	DECLARE_MULTICAST_DELEGATE_OneParam(FOnAssetsAddExtraObjectsToDelete, TArray<UObject*>&);
 	/** delegate type for when a user requests to delete certain assets... DOES NOT mean the asset(s) will be deleted (the user could cancel) */
@@ -252,6 +254,8 @@ struct UNREALED_API FEditorDelegates
 	static FOnDollyPerspectiveCamera OnDollyPerspectiveCamera;
 	/** Called on editor shutdown after packages have been successfully saved */
 	static FSimpleMulticastDelegate OnShutdownPostPackagesSaved;
+	/** Called when one or more packages have been deleted through save */
+	static FOnPackageDeleted OnPackageDeleted;
 	/** Called when the user requests assets to be deleted to determine if the operation is available.  */
 	static FOnAssetsCanDelete OnAssetsCanDelete;
 	/** Called when the user requests certain assets be deletedand  allows the addition of secondary assets that should also be deleted */
diff --git a/Engine/Source/Runtime/Engine/Private/WorldPartition/ActorDescContainer.cpp b/Engine/Source/Runtime/Engine/Private/WorldPartition/ActorDescContainer.cpp
--- a/Engine/Source/Runtime/Engine/Private/WorldPartition/ActorDescContainer.cpp
+++ b/Engine/Source/Runtime/Engine/Private/WorldPartition/ActorDescContainer.cpp
@@ -184,57 +184,58 @@ bool UActorDescContainer::ShouldHandleAssetEvent(const FAssetData& InAssetData)
 	return (ThisLevelPath == AssetLevelPath);
 }
 
-void UActorDescContainer::OnAssetAdded(const FAssetData& InAssetData)
+void UActorDescContainer::OnObjectPreSave(UObject* Object)
 {
-	if (ShouldHandleAssetEvent(InAssetData))
+	if (const AActor* Actor = Cast<AActor>(Object))
 	{
-		TUniquePtr<FWorldPartitionActorDesc>* NewActorDesc = new(ActorDescList) TUniquePtr<FWorldPartitionActorDesc>(GetActorDescriptor(InAssetData));
-		check(NewActorDesc->IsValid());
+		if (Actor->GetLevel() == World->PersistentLevel)
+		{
+			if (TUniquePtr<FWorldPartitionActorDesc>* ExistingActorDesc = Actors.FindRef(Actor->GetActorGuid()))
+			{
+				// Pin the actor handle on the actor to prevent unloading it when unhashing
+				FWorldPartitionHandle ExistingActorHandle(ExistingActorDesc);
+				FWorldPartitionHandlePinRefScope ExistingActorHandlePin(ExistingActorHandle);
 
-		check(!Actors.Contains((*NewActorDesc)->GetGuid()));
-		Actors.Add((*NewActorDesc)->GetGuid(), NewActorDesc);
+				TUniquePtr<FWorldPartitionActorDesc> NewActorDesc(Actor->CreateActorDesc());
 
-		OnActorDescAdded(*NewActorDesc);
-	}
-}
+				OnActorDescUpdating(*ExistingActorDesc);
 
-void UActorDescContainer::OnAssetRemoved(const FAssetData& InAssetData)
-{
-	if (ShouldHandleAssetEvent(InAssetData))
-	{
-		TUniquePtr<FWorldPartitionActorDesc> NewActorDesc = GetActorDescriptor(InAssetData);
-		check(NewActorDesc.IsValid());
+				// Transfer any reference count from external sources
+				NewActorDesc->TransferRefCounts(ExistingActorDesc->Get());
 
-		TUniquePtr<FWorldPartitionActorDesc>* ExistingActorDesc = Actors.FindChecked(NewActorDesc->GetGuid());
+				*ExistingActorDesc = MoveTemp(NewActorDesc);
 
-		OnActorDescRemoved(*ExistingActorDesc);
+				OnActorDescUpdated(*ExistingActorDesc);
+			}
+			// New actor
+			else
+			{
+				TUniquePtr<FWorldPartitionActorDesc>* NewActorDesc = new(ActorDescList) TUniquePtr<FWorldPartitionActorDesc>(Actor->CreateActorDesc());
+				check(NewActorDesc->IsValid());
+
+				check(!Actors.Contains((*NewActorDesc)->GetGuid()));
+				Actors.Add((*NewActorDesc)->GetGuid(), NewActorDesc);
 
-		Actors.Remove((*ExistingActorDesc)->GetGuid());
-		ExistingActorDesc->Release();
+				OnActorDescAdded(*NewActorDesc);
+			}
+		}
 	}
 }
 
-void UActorDescContainer::OnAssetUpdated(const FAssetData& InAssetData)
+void UActorDescContainer::OnPackageDeleted(UPackage* Package)
 {
-	if (ShouldHandleAssetEvent(InAssetData))
-	{
-		TUniquePtr<FWorldPartitionActorDesc> NewActorDesc = GetActorDescriptor(InAssetData);
-		check(NewActorDesc.IsValid());
-
-		TUniquePtr<FWorldPartitionActorDesc>* ExistingActorDesc = Actors.FindChecked(NewActorDesc->GetGuid());
-
-		// Pin the actor handle on the actor to prevent unloading it when unhashing
-		FWorldPartitionHandle ExistingActorHandle(ExistingActorDesc);
-		FWorldPartitionHandlePinRefScope ExistingActorHandlePin(ExistingActorHandle);
+	AActor* Actor = nullptr;
 
-		OnActorDescUpdating(*ExistingActorDesc);
+	ForEachObjectWithPackage(Package, [&Actor](UObject* Object)	{ Actor = Cast<AActor>(Object);	return !Actor; }, false);
 
-		// Transfer any reference count from external sources
-		NewActorDesc->TransferRefCounts(ExistingActorDesc->Get());
-
-		*ExistingActorDesc = MoveTemp(NewActorDesc);
-
-		OnActorDescUpdated(*ExistingActorDesc);
+	if (Actor && (Actor->GetLevel() == World->PersistentLevel))
+	{
+		if (TUniquePtr<FWorldPartitionActorDesc>* ExistingActorDesc = Actors.FindRef(Actor->GetActorGuid()))
+		{
+			OnActorDescRemoved(*ExistingActorDesc);
+			Actors.Remove((*ExistingActorDesc)->GetGuid());
+			ExistingActorDesc->Release();
+		}
 	}
 }
 
@@ -242,10 +243,8 @@ void UActorDescContainer::RegisterDelegates()
 {
 	if (GEditor && !IsTemplate())
 	{
-		IAssetRegistry& AssetRegistry = FModuleManager::LoadModuleChecked<FAssetRegistryModule>(TEXT("AssetRegistry")).Get();
-		AssetRegistry.OnAssetAdded().AddUObject(this, &UActorDescContainer::OnAssetAdded);
-		AssetRegistry.OnAssetRemoved().AddUObject(this, &UActorDescContainer::OnAssetRemoved);
-		AssetRegistry.OnAssetUpdated().AddUObject(this, &UActorDescContainer::OnAssetUpdated);
+		FCoreUObjectDelegates::OnObjectSaved.AddUObject(this, &UActorDescContainer::OnObjectPreSave);
+		FEditorDelegates::OnPackageDeleted.AddUObject(this, &UActorDescContainer::OnPackageDeleted);
 	}
 }
 
@@ -253,13 +252,8 @@ void UActorDescContainer::UnregisterDelegates()
 {
 	if (GEditor && !IsTemplate())
 	{
-		if (FAssetRegistryModule* AssetRegistryModule = FModuleManager::GetModulePtr<FAssetRegistryModule>(TEXT("AssetRegistry")))
-		{
-			IAssetRegistry& AssetRegistry = AssetRegistryModule->Get();
-			AssetRegistry.OnAssetAdded().RemoveAll(this);
-			AssetRegistry.OnAssetRemoved().RemoveAll(this);
-			AssetRegistry.OnAssetUpdated().RemoveAll(this);
-		}
+		FCoreUObjectDelegates::OnObjectSaved.RemoveAll(this);
+		FEditorDelegates::OnPackageDeleted.RemoveAll(this);
 	}
 }
 #endif
\ No newline at end of file
diff --git a/Engine/Source/Runtime/Engine/Private/WorldPartition/RuntimeSpatialHash/RuntimeSpatialHashHLOD.cpp b/Engine/Source/Runtime/Engine/Private/WorldPartition/RuntimeSpatialHash/RuntimeSpatialHashHLOD.cpp
--- a/Engine/Source/Runtime/Engine/Private/WorldPartition/RuntimeSpatialHash/RuntimeSpatialHashHLOD.cpp
+++ b/Engine/Source/Runtime/Engine/Private/WorldPartition/RuntimeSpatialHash/RuntimeSpatialHashHLOD.cpp
@@ -228,11 +228,7 @@ static TArray<FGuid> GenerateHLODsForGrid(UWorldPartition* WorldPartition, const
 						}
 					}
 
-					// Manually tick the directory watcher and the asset registry to register newly created actors
-					DirectoryWatcherModule.Get()->Tick(-1.0f);
-					AssetRegistryModule.Get().Tick(-1.0f);
-
-					// Make sure the previous ticks properly registered all our actor descriptors
+					// Make sure we registered actor descriptors for newly created actors
 					for (AWorldPartitionHLOD* CellHLODActor : CellHLODActors)
 					{
 						check(WorldPartition->GetActorDesc(CellHLODActor->GetActorGuid()));
diff --git a/Engine/Source/Runtime/Engine/Public/WorldPartition/ActorDescContainer.h b/Engine/Source/Runtime/Engine/Public/WorldPartition/ActorDescContainer.h
--- a/Engine/Source/Runtime/Engine/Public/WorldPartition/ActorDescContainer.h
+++ b/Engine/Source/Runtime/Engine/Public/WorldPartition/ActorDescContainer.h
@@ -23,10 +23,9 @@ public:
 	virtual UWorld* GetWorld() const override;
 
 #if WITH_EDITOR
-	// Asset registry events
-	virtual void OnAssetAdded(const FAssetData& InAssetData);
-	virtual void OnAssetRemoved(const FAssetData& InAssetData);
-	virtual void OnAssetUpdated(const FAssetData& InAssetData);
+	// Events
+	virtual void OnObjectPreSave(UObject* Object);
+	virtual void OnPackageDeleted(UPackage* Package);
 
 	FName GetContainerPackage() const { return ContainerPackageName; }
 #endif
