Add missing changes done in Main branch for 15665076.

diff --git a/Engine/Source/Runtime/Engine/Private/WorldPartition/ActorDescContainer.cpp b/Engine/Source/Runtime/Engine/Private/WorldPartition/ActorDescContainer.cpp
--- a/Engine/Source/Runtime/Engine/Private/WorldPartition/ActorDescContainer.cpp
+++ b/Engine/Source/Runtime/Engine/Private/WorldPartition/ActorDescContainer.cpp
@@ -188,7 +188,7 @@ void UActorDescContainer::OnObjectPreSave(UObject* Object)
 {
 	if (const AActor* Actor = Cast<AActor>(Object))
 	{
-		if (Actor->GetLevel() == World->PersistentLevel)
+		if (Actor->IsPackageExternal() && (Actor->GetLevel() == World->PersistentLevel))
 		{
 			if (TUniquePtr<FWorldPartitionActorDesc>* ExistingActorDesc = Actors.FindRef(Actor->GetActorGuid()))
 			{
@@ -228,7 +228,7 @@ void UActorDescContainer::OnPackageDeleted(UPackage* Package)
 
 	ForEachObjectWithPackage(Package, [&Actor](UObject* Object)	{ Actor = Cast<AActor>(Object);	return !Actor; }, false);
 
-	if (Actor && (Actor->GetLevel() == World->PersistentLevel))
+	if (Actor && Actor->IsPackageExternal() && (Actor->GetLevel() == World->PersistentLevel))
 	{
 		if (TUniquePtr<FWorldPartitionActorDesc>* ExistingActorDesc = Actors.FindRef(Actor->GetActorGuid()))
 		{
