Optimized FRichCurve::RemoveRedundantKeysInternal by adding an explicit FKeyHandleMap::Initialize method instead of a loop over Add(). Reduces complexity from O(n^2) to O(n).

diff --git a/Engine/Source/Runtime/Engine/Classes/Curves/KeyHandle.h b/Engine/Source/Runtime/Engine/Classes/Curves/KeyHandle.h
--- a/Engine/Source/Runtime/Engine/Classes/Curves/KeyHandle.h
+++ b/Engine/Source/Runtime/Engine/Classes/Curves/KeyHandle.h
@@ -3,6 +3,7 @@
 #pragma once
 
 #include "CoreMinimal.h"
+#include "Containers/ArrayView.h"
 #include "UObject/ObjectMacros.h"
 #include "UObject/Class.h"
 #include "KeyHandle.generated.h"
@@ -70,9 +71,12 @@ public:
 	FKeyHandleMap( const FKeyHandleMap& Other ) {}
 	void operator=(const FKeyHandleMap& Other) {}
 
+	// Quickly initializes this map by clearing it and filling with KeyHandles in O(n) time, instead of O(n^2) if Add() was used in a loop.
+	void Initialize(TArrayView<const FKeyHandle> InKeyHandles);
+
 	/** TMap functionality */
 	void Add( const FKeyHandle& InHandle, int32 InIndex );
-	void Empty();
+	void Empty(int32 ExpectedNumElements = 0);
 	void Remove( const FKeyHandle& InHandle );
 	const int32* Find(const FKeyHandle& InHandle) const { return KeyHandlesToIndices.Find(InHandle); }
 	const FKeyHandle* FindKey( int32 KeyIndex ) const;
diff --git a/Engine/Source/Runtime/Engine/Private/Curves/KeyHandle.cpp b/Engine/Source/Runtime/Engine/Private/Curves/KeyHandle.cpp
--- a/Engine/Source/Runtime/Engine/Private/Curves/KeyHandle.cpp
+++ b/Engine/Source/Runtime/Engine/Private/Curves/KeyHandle.cpp
@@ -20,6 +20,18 @@ FKeyHandle FKeyHandle::Invalid()
 	return FKeyHandle(0);
 }
 
+void FKeyHandleMap::Initialize(TArrayView<const FKeyHandle> InKeyHandles)
+{
+	Empty(InKeyHandles.Num());
+
+	for (int32 Index = 0; Index < InKeyHandles.Num(); ++Index)
+	{
+		const FKeyHandle& Handle = InKeyHandles[Index];
+		KeyHandles.Add(Handle);
+		KeyHandlesToIndices.Add(Handle, Index);
+	}
+}
+
 /* FKeyHandleMap interface
  *****************************************************************************/
 
@@ -63,10 +75,10 @@ void FKeyHandleMap::SetKeyHandles(int32 Num)
 }
 
 
-void FKeyHandleMap::Empty()
+void FKeyHandleMap::Empty(int32 ExpectedNumElements)
 {
-	KeyHandlesToIndices.Empty();
-	KeyHandles.Empty();
+	KeyHandlesToIndices.Empty(ExpectedNumElements);
+	KeyHandles.Empty(ExpectedNumElements);
 }
 
 
diff --git a/Engine/Source/Runtime/Engine/Private/Curves/RichCurve.cpp b/Engine/Source/Runtime/Engine/Private/Curves/RichCurve.cpp
--- a/Engine/Source/Runtime/Engine/Private/Curves/RichCurve.cpp
+++ b/Engine/Source/Runtime/Engine/Private/Curves/RichCurve.cpp
@@ -2,6 +2,7 @@
 
 #include "Curves/RichCurve.h"
 #include "Templates/Function.h"
+#include "Containers/ArrayView.h"
 
 // Broken - do not turn on! 
 #define MIXEDKEY_STRIPS_TANGENTS 0
@@ -1242,11 +1243,9 @@ void FRichCurve::RemoveRedundantKeysInternal(float Tolerance, int32 InStartKeepK
 	AutoSetTangents();
 
 	// Rebuild KeyHandlesToIndices
-	KeyHandlesToIndices.Empty();
-	for (int32 KeyIndex = 0; KeyIndex < Keys.Num(); ++KeyIndex)
-	{
-		KeyHandlesToIndices.Add(KeepHandles[KeyIndex], KeyIndex);
-	}
+
+	check(Keys.Num() == KeepHandles.Num());
+	KeyHandlesToIndices.Initialize(KeepHandles);
 }
 
 void FRichCurve::RemapTimeValue(float& InTime, float& CycleValueOffset) const
