Moved TypedElementCommonActions into Engine

This means it can also be used by the InteractiveToolsFramework, and makes sense as the type has Engine dependencies in its API

#rnx

diff --git a/Engine/Source/Editor/EditorFramework/Private/Elements/Framework/TypedElementCommonActions.cpp b/Engine/Source/Editor/EditorFramework/Private/Elements/Framework/TypedElementCommonActions.cpp
deleted file mode 100644
diff --git a/Engine/Source/Editor/EditorFramework/Public/Elements/Framework/TypedElementCommonActions.h b/Engine/Source/Editor/EditorFramework/Public/Elements/Framework/TypedElementCommonActions.h
deleted file mode 100644
diff --git a/Engine/Source/Runtime/Engine/Private/Elements/Framework/TypedElementCommonActions.cpp b/Engine/Source/Runtime/Engine/Private/Elements/Framework/TypedElementCommonActions.cpp
new file mode 100644
--- /dev/null
+++ b/Engine/Source/Runtime/Engine/Private/Elements/Framework/TypedElementCommonActions.cpp
@@ -0,0 +1,141 @@
+// Copyright Epic Games, Inc. All Rights Reserved.
+
+#include "Elements/Framework/TypedElementCommonActions.h"
+#include "Elements/Framework/TypedElementSelectionSet.h"
+#include "Elements/Framework/TypedElementRegistry.h"
+#include "Elements/Framework/TypedElementUtil.h"
+
+void FTypedElementCommonActionsCustomization::GetElementsToDelete(const TTypedElement<UTypedElementWorldInterface>& InElementWorldHandle, const UTypedElementSelectionSet* InSelectionSet, UTypedElementList* OutElementsToDelete)
+{
+	OutElementsToDelete->Add(InElementWorldHandle);
+}
+
+bool FTypedElementCommonActionsCustomization::DeleteElements(UTypedElementWorldInterface* InWorldInterface, TArrayView<const FTypedElementHandle> InElementHandles, UWorld* InWorld, UTypedElementSelectionSet* InSelectionSet, const FTypedElementDeletionOptions& InDeletionOptions)
+{
+	return InWorldInterface->DeleteElements(InElementHandles, InWorld, InSelectionSet, InDeletionOptions);
+}
+
+void FTypedElementCommonActionsCustomization::DuplicateElements(UTypedElementWorldInterface* InWorldInterface, TArrayView<const FTypedElementHandle> InElementHandles, UWorld* InWorld, bool bOffsetLocations, TArray<FTypedElementHandle>& OutNewElements)
+{
+	InWorldInterface->DuplicateElements(InElementHandles, InWorld, bOffsetLocations, OutNewElements);
+}
+
+
+void UTypedElementCommonActions::GetSelectedElementsToDelete(const UTypedElementSelectionSet* InSelectionSet, UTypedElementList* OutElementsToDelete) const
+{
+	OutElementsToDelete->Reset();
+	InSelectionSet->ForEachSelectedElement<UTypedElementWorldInterface>([this, InSelectionSet, OutElementsToDelete](const TTypedElement<UTypedElementWorldInterface>& InElementWorldHandle)
+	{
+		FTypedElementCommonActionsElement CommonActionsElement(InElementWorldHandle, GetInterfaceCustomizationByTypeId(InElementWorldHandle.GetId().GetTypeId()));
+		check(CommonActionsElement.IsSet());
+		CommonActionsElement.GetElementsToDelete(InSelectionSet, OutElementsToDelete);
+		return true;
+	});
+}
+
+bool UTypedElementCommonActions::DeleteElements(const TArray<FTypedElementHandle>& ElementHandles, UWorld* World, UTypedElementSelectionSet* InSelectionSet, const FTypedElementDeletionOptions& InDeletionOptions)
+{
+	return DeleteElements(MakeArrayView(ElementHandles), World, InSelectionSet, InDeletionOptions);
+}
+
+bool UTypedElementCommonActions::DeleteElements(TArrayView<const FTypedElementHandle> ElementHandles, UWorld* World, UTypedElementSelectionSet* InSelectionSet, const FTypedElementDeletionOptions& InDeletionOptions)
+{
+	TMap<FTypedHandleTypeId, TArray<FTypedElementHandle>> ElementsToDuplicateByType;
+	TypedElementUtil::BatchElementsByType(ElementHandles, ElementsToDuplicateByType);
+
+	bool bSuccess = false;
+
+	UTypedElementRegistry* Registry = UTypedElementRegistry::GetInstance();
+	UTypedElementRegistry::FDisableElementDestructionOnGC GCGuard(Registry);
+
+	for (const auto& ElementsByTypePair : ElementsToDuplicateByType)
+	{
+		FTypedElementCommonActionsCustomization* CommonActionsCustomization = GetInterfaceCustomizationByTypeId(ElementsByTypePair.Key);
+		UTypedElementWorldInterface* WorldInterface = Registry->GetElementInterface<UTypedElementWorldInterface>(ElementsByTypePair.Key);
+		if (CommonActionsCustomization && WorldInterface)
+		{
+			bSuccess |= CommonActionsCustomization->DeleteElements(WorldInterface, ElementsByTypePair.Value, World, InSelectionSet, InDeletionOptions);
+		}
+	}
+
+	return bSuccess;
+}
+
+bool UTypedElementCommonActions::DeleteElements(const UTypedElementList* ElementList, UWorld* World, UTypedElementSelectionSet* InSelectionSet, const FTypedElementDeletionOptions& InDeletionOptions)
+{
+	TMap<FTypedHandleTypeId, TArray<FTypedElementHandle>> ElementsToDuplicateByType;
+	TypedElementUtil::BatchElementsByType(ElementList, ElementsToDuplicateByType);
+
+	bool bSuccess = false;
+
+	UTypedElementRegistry* Registry = UTypedElementRegistry::GetInstance();
+	UTypedElementRegistry::FDisableElementDestructionOnGC GCGuard(Registry);
+
+	for (const auto& ElementsByTypePair : ElementsToDuplicateByType)
+	{
+		FTypedElementCommonActionsCustomization* CommonActionsCustomization = GetInterfaceCustomizationByTypeId(ElementsByTypePair.Key);
+		UTypedElementWorldInterface* WorldInterface = Registry->GetElementInterface<UTypedElementWorldInterface>(ElementsByTypePair.Key);
+		if (CommonActionsCustomization && WorldInterface)
+		{
+			bSuccess |= CommonActionsCustomization->DeleteElements(WorldInterface, ElementsByTypePair.Value, World, InSelectionSet, InDeletionOptions);
+		}
+	}
+
+	return bSuccess;
+}
+
+TArray<FTypedElementHandle> UTypedElementCommonActions::DuplicateElements(const TArray<FTypedElementHandle>& ElementHandles, UWorld* World, bool bOffsetLocations)
+{
+	return DuplicateElements(MakeArrayView(ElementHandles), World, bOffsetLocations);
+}
+
+TArray<FTypedElementHandle> UTypedElementCommonActions::DuplicateElements(TArrayView<const FTypedElementHandle> ElementHandles, UWorld* World, bool bOffsetLocations)
+{
+	TMap<FTypedHandleTypeId, TArray<FTypedElementHandle>> ElementsToDuplicateByType;
+	TypedElementUtil::BatchElementsByType(ElementHandles, ElementsToDuplicateByType);
+
+	TArray<FTypedElementHandle> NewElements;
+	NewElements.Reserve(ElementHandles.Num());
+
+	UTypedElementRegistry* Registry = UTypedElementRegistry::GetInstance();
+	for (const auto& ElementsByTypePair : ElementsToDuplicateByType)
+	{
+		FTypedElementCommonActionsCustomization* CommonActionsCustomization = GetInterfaceCustomizationByTypeId(ElementsByTypePair.Key);
+		UTypedElementWorldInterface* WorldInterface = Registry->GetElementInterface<UTypedElementWorldInterface>(ElementsByTypePair.Key);
+		if (CommonActionsCustomization && WorldInterface)
+		{
+			CommonActionsCustomization->DuplicateElements(WorldInterface, ElementsByTypePair.Value, World, bOffsetLocations, NewElements);
+		}
+	}
+
+	return NewElements;
+}
+
+TArray<FTypedElementHandle> UTypedElementCommonActions::DuplicateElements(const UTypedElementList* ElementList, UWorld* World, bool bOffsetLocations)
+{
+	TMap<FTypedHandleTypeId, TArray<FTypedElementHandle>> ElementsToDuplicateByType;
+	TypedElementUtil::BatchElementsByType(ElementList, ElementsToDuplicateByType);
+
+	TArray<FTypedElementHandle> NewElements;
+	NewElements.Reserve(ElementList->Num());
+
+	UTypedElementRegistry* Registry = UTypedElementRegistry::GetInstance();
+	for (const auto& ElementsByTypePair : ElementsToDuplicateByType)
+	{
+		FTypedElementCommonActionsCustomization* CommonActionsCustomization = GetInterfaceCustomizationByTypeId(ElementsByTypePair.Key);
+		UTypedElementWorldInterface* WorldInterface = Registry->GetElementInterface<UTypedElementWorldInterface>(ElementsByTypePair.Key);
+		if (CommonActionsCustomization && WorldInterface)
+		{
+			CommonActionsCustomization->DuplicateElements(WorldInterface, ElementsByTypePair.Value, World, bOffsetLocations, NewElements);
+		}
+	}
+
+	return NewElements;
+}
+
+FTypedElementCommonActionsElement UTypedElementCommonActions::ResolveCommonActionsElement(const FTypedElementHandle& InElementHandle) const
+{
+	return InElementHandle
+		? FTypedElementCommonActionsElement(UTypedElementRegistry::GetInstance()->GetElement<UTypedElementWorldInterface>(InElementHandle), GetInterfaceCustomizationByTypeId(InElementHandle.GetId().GetTypeId()))
+		: FTypedElementCommonActionsElement();
+}
diff --git a/Engine/Source/Runtime/Engine/Public/Elements/Framework/TypedElementCommonActions.h b/Engine/Source/Runtime/Engine/Public/Elements/Framework/TypedElementCommonActions.h
new file mode 100644
--- /dev/null
+++ b/Engine/Source/Runtime/Engine/Public/Elements/Framework/TypedElementCommonActions.h
@@ -0,0 +1,95 @@
+// Copyright Epic Games, Inc. All Rights Reserved.
+
+#pragma once
+
+#include "Elements/Framework/TypedElementInterfaceCustomization.h"
+#include "Elements/Interfaces/TypedElementWorldInterface.h"
+#include "TypedElementCommonActions.generated.h"
+
+class UTypedElementList;
+class UTypedElementSelectionSet;
+
+/**
+ * Customization used to allow asset editors (such as the level editor) to override the base behavior of common actions.
+ */
+class ENGINE_API FTypedElementCommonActionsCustomization
+{
+public:
+	virtual ~FTypedElementCommonActionsCustomization() = default;
+
+	//~ See UTypedElementCommonActions for API docs
+	virtual void GetElementsToDelete(const TTypedElement<UTypedElementWorldInterface>& InElementWorldHandle, const UTypedElementSelectionSet* InSelectionSet, UTypedElementList* OutElementsToDelete);
+	virtual bool DeleteElements(UTypedElementWorldInterface* InWorldInterface, TArrayView<const FTypedElementHandle> InElementHandles, UWorld* InWorld, UTypedElementSelectionSet* InSelectionSet, const FTypedElementDeletionOptions& InDeletionOptions);
+	virtual void DuplicateElements(UTypedElementWorldInterface* InWorldInterface, TArrayView<const FTypedElementHandle> InElementHandles, UWorld* InWorld, bool bOffsetLocations, TArray<FTypedElementHandle>& OutNewElements);
+};
+
+/**
+ * Utility to hold a typed element handle and its associated world interface and common actions customization.
+ */
+struct ENGINE_API FTypedElementCommonActionsElement
+{
+public:
+	FTypedElementCommonActionsElement() = default;
+
+	FTypedElementCommonActionsElement(TTypedElement<UTypedElementWorldInterface> InElementWorldHandle, FTypedElementCommonActionsCustomization* InCommonActionsCustomization)
+		: ElementWorldHandle(MoveTemp(InElementWorldHandle))
+		, CommonActionsCustomization(InCommonActionsCustomization)
+	{
+	}
+
+	FTypedElementCommonActionsElement(const FTypedElementCommonActionsElement&) = default;
+	FTypedElementCommonActionsElement& operator=(const FTypedElementCommonActionsElement&) = default;
+
+	FTypedElementCommonActionsElement(FTypedElementCommonActionsElement&&) = default;
+	FTypedElementCommonActionsElement& operator=(FTypedElementCommonActionsElement&&) = default;
+
+	FORCEINLINE explicit operator bool() const
+	{
+		return IsSet();
+	}
+
+	FORCEINLINE bool IsSet() const
+	{
+		return ElementWorldHandle.IsSet()
+			&& CommonActionsCustomization;
+	}
+
+	//~ See UTypedElementCommonActions for API docs
+	void GetElementsToDelete(const UTypedElementSelectionSet* InSelectionSet, UTypedElementList* OutElementsToDelete) { CommonActionsCustomization->GetElementsToDelete(ElementWorldHandle, InSelectionSet, OutElementsToDelete); }
+
+private:
+	TTypedElement<UTypedElementWorldInterface> ElementWorldHandle;
+	FTypedElementCommonActionsCustomization* CommonActionsCustomization = nullptr;
+};
+
+/**
+ * A utility to handle higher-level common actions, but default via UTypedElementWorldInterface,
+ * but asset editors can customize this behavior via FTypedElementCommonActionsCustomization.
+ */
+UCLASS(Transient)
+class ENGINE_API UTypedElementCommonActions : public UObject, public TTypedElementInterfaceCustomizationRegistry<FTypedElementCommonActionsCustomization>
+{
+	GENERATED_BODY()
+
+public:
+	/**
+	 * Get the elements from the given selection set that should be deleted
+	 */
+	void GetSelectedElementsToDelete(const UTypedElementSelectionSet* InSelectionSet, UTypedElementList* OutElementsToDelete) const;
+
+	UFUNCTION(BlueprintCallable, Category="TypedElementFramework|Common")
+	bool DeleteElements(const TArray<FTypedElementHandle>& ElementHandles, UWorld* World, UTypedElementSelectionSet* InSelectionSet, const FTypedElementDeletionOptions& InDeletionOptions);
+	bool DeleteElements(TArrayView<const FTypedElementHandle> ElementHandles, UWorld* World, UTypedElementSelectionSet* InSelectionSet, const FTypedElementDeletionOptions& InDeletionOptions);
+	bool DeleteElements(const UTypedElementList* ElementList, UWorld* World, UTypedElementSelectionSet* InSelectionSet, const FTypedElementDeletionOptions& InDeletionOptions);
+
+	UFUNCTION(BlueprintCallable, Category="TypedElementFramework|Common")
+	TArray<FTypedElementHandle> DuplicateElements(const TArray<FTypedElementHandle>& ElementHandles, UWorld* World, bool bOffsetLocations);
+	TArray<FTypedElementHandle> DuplicateElements(TArrayView<const FTypedElementHandle> ElementHandles, UWorld* World, bool bOffsetLocations);
+	TArray<FTypedElementHandle> DuplicateElements(const UTypedElementList* ElementList, UWorld* World, bool bOffsetLocations);
+
+private:
+	/**
+	 * Attempt to resolve the selection interface and common actions customization for the given element, if any.
+	 */
+	FTypedElementCommonActionsElement ResolveCommonActionsElement(const FTypedElementHandle& InElementHandle) const;
+};
