WorldPartition: Fix StreamingGrids property being persisted to disk and causing crash on PIE

diff --git a/Engine/Source/Runtime/Engine/Private/WorldPartition/WorldPartitionRuntimeSpatialHash.cpp b/Engine/Source/Runtime/Engine/Private/WorldPartition/WorldPartitionRuntimeSpatialHash.cpp
--- a/Engine/Source/Runtime/Engine/Private/WorldPartition/WorldPartitionRuntimeSpatialHash.cpp
+++ b/Engine/Source/Runtime/Engine/Private/WorldPartition/WorldPartitionRuntimeSpatialHash.cpp
@@ -386,6 +386,15 @@ FName UWorldPartitionRuntimeSpatialHash::GetActorRuntimeGrid(const AActor* Actor
 	return Super::GetActorRuntimeGrid(Actor);
 }
 
+void UWorldPartitionRuntimeSpatialHash::PreSave(const class ITargetPlatform* TargetPlatform)
+{
+	Super::PreSave(TargetPlatform);
+
+	// We don't want this to be persisted but we can't set the property Transient as it is NonPIEDuplicateTransient and those flags aren't compatible
+	// If at some point GenerateStreaming is done after duplication we can remove this code.
+	StreamingGrids.Empty();
+}
+
 void UWorldPartitionRuntimeSpatialHash::SetDefaultValues()
 {
 	FSpatialHashRuntimeGrid& MainGrid = Grids.AddDefaulted_GetRef();
@@ -456,7 +465,8 @@ bool UWorldPartitionRuntimeSpatialHash::GenerateStreaming(EWorldPartitionStreami
 		}
 	}
 
-	check(!StreamingGrids.Num());
+	// Fix case where StreamingGrids might have been persisted.
+	StreamingGrids.Empty();
 
 	TMap<FName, int32> GridsMapping;
 	GridsMapping.Add(NAME_None, 0);
diff --git a/Engine/Source/Runtime/Engine/Public/WorldPartition/WorldPartitionRuntimeSpatialHash.h b/Engine/Source/Runtime/Engine/Public/WorldPartition/WorldPartitionRuntimeSpatialHash.h
--- a/Engine/Source/Runtime/Engine/Public/WorldPartition/WorldPartitionRuntimeSpatialHash.h
+++ b/Engine/Source/Runtime/Engine/Public/WorldPartition/WorldPartitionRuntimeSpatialHash.h
@@ -194,6 +194,8 @@ class ENGINE_API UWorldPartitionRuntimeSpatialHash : public UWorldPartitionRunti
 	GENERATED_UCLASS_BODY()
 
 public:
+	virtual void PreSave(const class ITargetPlatform* TargetPlatform) override;
+
 #if WITH_EDITOR
 	virtual void SetDefaultValues() override;
 	virtual void ImportFromWorldComposition(class UWorldComposition* WorldComposition) override;
