Fix for conduits adding self transitions to active transition list
#rnx

diff --git a/Engine/Source/Runtime/Engine/Private/Animation/AnimNode_StateMachine.cpp b/Engine/Source/Runtime/Engine/Private/Animation/AnimNode_StateMachine.cpp
--- a/Engine/Source/Runtime/Engine/Private/Animation/AnimNode_StateMachine.cpp
+++ b/Engine/Source/Runtime/Engine/Private/Animation/AnimNode_StateMachine.cpp
@@ -648,6 +648,13 @@ bool FAnimNode_StateMachine::FindValidTransition(const FAnimationUpdateContext&
 			continue;
 		}
 
+		const int32 NextState = GetTransitionInfo(TransitionRule.TransitionIndex).NextState;
+		// Skip this transition if we've already visited the destination state
+		if (OutVisitedStateIndices.Contains(NextState))
+		{
+			continue;
+		}
+
 		FAnimNode_TransitionResult* ResultNode = GetNodeFromPropertyIndex<FAnimNode_TransitionResult>(Context.AnimInstanceProxy->GetAnimInstanceObject(), AnimBlueprintClass, TransitionRule.CanTakeDelegateIndex);
 		float CrossfadeTimeAdjustment = 0.f;
 
@@ -685,7 +692,6 @@ bool FAnimNode_StateMachine::FindValidTransition(const FAnimationUpdateContext&
 
 		if (ResultNode->bCanEnterTransition == TransitionRule.bDesiredTransitionReturnValue)
 		{
-			int32 NextState = GetTransitionInfo(TransitionRule.TransitionIndex).NextState;
 			const FBakedAnimationState& NextStateInfo = GetStateInfo(NextState);
 
 			// if next state is a conduit we want to check for transitions using that state as the root
