WorldPartition - Remove remains of asset registry events handling

diff --git a/Engine/Source/Runtime/Engine/Private/WorldPartition/ActorDescContainer.cpp b/Engine/Source/Runtime/Engine/Private/WorldPartition/ActorDescContainer.cpp
--- a/Engine/Source/Runtime/Engine/Private/WorldPartition/ActorDescContainer.cpp
+++ b/Engine/Source/Runtime/Engine/Private/WorldPartition/ActorDescContainer.cpp
@@ -14,7 +14,6 @@ UActorDescContainer::UActorDescContainer(const FObjectInitializer& ObjectInitial
 	, World(nullptr)
 #if WITH_EDITOR
 	, bContainerInitialized(false)
-	, bIgnoreAssetRegistryEvents(false)
 #endif
 {}
 
@@ -29,8 +28,6 @@ void UActorDescContainer::Initialize(UWorld* InWorld, FName InPackageName, bool
 	
 	if (!ContainerPackageName.IsNone())
 	{
-		TGuardValue<bool> IgnoreAssetRegistryEvents(bIgnoreAssetRegistryEvents, true);
-
 		const FString LevelPathStr = ContainerPackageName.ToString();
 		const FString LevelExternalActorsPath = ULevel::GetExternalActorsPath(LevelPathStr);
 
@@ -140,50 +137,6 @@ TUniquePtr<FWorldPartitionActorDesc> UActorDescContainer::GetActorDescriptor(con
 	return nullptr;
 }
 
-bool UActorDescContainer::ShouldHandleAssetEvent(const FAssetData& InAssetData)
-{
-	// Ignore asset event when specifically asking to
-	if (bIgnoreAssetRegistryEvents)
-	{
-		return false;
-	}
-
-	// Ignore in-memory assets until they gets saved
-	if (InAssetData.HasAnyPackageFlags(PKG_NewlyCreated))
-	{
-		return false;
-	}
-
-	// Only handle actors
-	if (!InAssetData.GetClass()->IsChildOf<AActor>())
-	{
-		return false;
-	}
-
-	// Make sure asset contains the required tags
-	static FName NAME_ActorMetaDataClass(TEXT("ActorMetaDataClass"));
-	static FName NAME_ActorMetaData(TEXT("ActorMetaData"));
-	if (!InAssetData.FindTag(NAME_ActorMetaDataClass) || !InAssetData.FindTag(NAME_ActorMetaData))
-	{
-		return false;
-	}
-
-	// Only handle assets that belongs to our level
-	auto RemoveAfterFirstDot = [](const FString& InValue)
-	{
-		int32 DotIndex;
-		if (InValue.FindChar(TEXT('.'), DotIndex))
-		{
-			return InValue.LeftChop(InValue.Len() - DotIndex);
-		}
-		return InValue;
-	};
-
-	const FString ThisLevelPath = ContainerPackageName.ToString();
-	const FString AssetLevelPath = RemoveAfterFirstDot(InAssetData.ObjectPath.ToString());
-	return (ThisLevelPath == AssetLevelPath);
-}
-
 void UActorDescContainer::OnObjectPreSave(UObject* Object)
 {
 	if (const AActor* Actor = Cast<AActor>(Object))
diff --git a/Engine/Source/Runtime/Engine/Public/WorldPartition/ActorDescContainer.h b/Engine/Source/Runtime/Engine/Public/WorldPartition/ActorDescContainer.h
--- a/Engine/Source/Runtime/Engine/Public/WorldPartition/ActorDescContainer.h
+++ b/Engine/Source/Runtime/Engine/Public/WorldPartition/ActorDescContainer.h
@@ -52,10 +52,7 @@ protected:
 	virtual void OnActorDescRegistered(const FWorldPartitionActorDesc&) {}
 	virtual void OnActorDescUnregistered(const FWorldPartitionActorDesc&) {}
 
-	bool ShouldHandleAssetEvent(const FAssetData& InAssetData);
-
 	bool bContainerInitialized;
-	bool bIgnoreAssetRegistryEvents;
 
 	FName ContainerPackageName;
 #endif
diff --git a/Engine/Source/Runtime/Engine/Public/WorldPartition/WorldPartition.h b/Engine/Source/Runtime/Engine/Public/WorldPartition/WorldPartition.h
--- a/Engine/Source/Runtime/Engine/Public/WorldPartition/WorldPartition.h
+++ b/Engine/Source/Runtime/Engine/Public/WorldPartition/WorldPartition.h
@@ -102,8 +102,6 @@ private:
 
 	virtual void OnActorDescRegistered(const FWorldPartitionActorDesc&) override;
 	virtual void OnActorDescUnregistered(const FWorldPartitionActorDesc&) override;
-
-	bool ShouldHandleAssetEvent(const FAssetData& InAssetData);
 #endif
 
 	//~ Begin UObject Interface
