Remove FlushRenderingCommands that is not needed anymore and happened during garbage collection which could cause long stalls after loading a map
Fix a bug that was camouflaged by the FlushRenderingCommands where TRefCountPtr could exist in lambda captures still waiting to be executed off the render thread before being release.
  - This crashed in the Material destructor when verifying that the refcount was 0 and it was not.
  - It is never safe to destroy Materials directly from the game-thread because of SetGameThreadShaderMap(nullptr) which also enqueue a render command with a TRefCountPtr in the capture.

TESTS (Thanks Krzysztof for the repro in CL 5557239)

Slightly modified repro to take into account code changes
* Insert FWindowsPlatformProcess::Sleep(1.0f) into every ENQUEUE_RENDER_COMMAND(SetCompilingShaderMap) in FMaterial
* Force GC to be run every frame with gc.StressTestGC=1 and gc.ForceCollectGarbageEveryFrame=1
* Copy Samples\Showcases\VirtualCamera project from 4.26
* Run VirtualCameraShowcase.uproject -run=DerivedDataCache -TargetPlatform=Windows+WindowsNoEditor -fill -DDC=CreateInstalledProjectPak -StompMalloc -stdout -CrashForUAT -unattended -NoLogTimes
* Added breakpoint in ClearAllCachedCookedPlatformData and ClearCachedCookedPlatformData to make sure they were being hit

Confirmed that the issue doesn't happen anymore because of the FMaterial::DeferredDeleteArray mecanism.

diff --git a/Engine/Source/Runtime/Engine/Private/Materials/Material.cpp b/Engine/Source/Runtime/Engine/Private/Materials/Material.cpp
--- a/Engine/Source/Runtime/Engine/Private/Materials/Material.cpp
+++ b/Engine/Source/Runtime/Engine/Private/Materials/Material.cpp
@@ -4392,11 +4392,7 @@ bool UMaterial::IsCachedCookedPlatformDataLoaded( const ITargetPlatform* TargetP
 
 void UMaterial::ClearCachedCookedPlatformData( const ITargetPlatform *TargetPlatform )
 {
-	// Make sure that all CacheShaders render thead commands are finished before we destroy FMaterialResources.
-	// TODO - is this needed, since we're using DeferredDeleteArray now?
-	FlushRenderingCommands();
-
-	TArray<FMaterialResource*> *CachedMaterialResourcesForPlatform = CachedMaterialResourcesForCooking.Find( TargetPlatform );
+	TArray<FMaterialResource*>* CachedMaterialResourcesForPlatform = CachedMaterialResourcesForCooking.Find( TargetPlatform );
 	if ( CachedMaterialResourcesForPlatform != nullptr)
 	{
 		FMaterial::DeferredDeleteArray(*CachedMaterialResourcesForPlatform);
@@ -4406,13 +4402,9 @@ void UMaterial::ClearCachedCookedPlatformData( const ITargetPlatform *TargetPlat
 
 void UMaterial::ClearAllCachedCookedPlatformData()
 {
-	// Make sure that all CacheShaders render thead commands are finished before we destroy FMaterialResources.
-	// TODO - is this needed, since we're using DeferredDeleteArray now?
-	FlushRenderingCommands();
-
 	for ( auto& It : CachedMaterialResourcesForCooking )
 	{
-		TArray<FMaterialResource*> &CachedMaterialResourcesForPlatform = It.Value;
+		TArray<FMaterialResource*>& CachedMaterialResourcesForPlatform = It.Value;
 		FMaterial::DeferredDeleteArray(CachedMaterialResourcesForPlatform);
 	}
 	CachedMaterialResourcesForCooking.Empty();
diff --git a/Engine/Source/Runtime/Engine/Private/Materials/MaterialInstance.cpp b/Engine/Source/Runtime/Engine/Private/Materials/MaterialInstance.cpp
--- a/Engine/Source/Runtime/Engine/Private/Materials/MaterialInstance.cpp
+++ b/Engine/Source/Runtime/Engine/Private/Materials/MaterialInstance.cpp
@@ -3308,11 +3308,7 @@ bool UMaterialInstance::IsCachedCookedPlatformDataLoaded( const ITargetPlatform*
 
 void UMaterialInstance::ClearCachedCookedPlatformData( const ITargetPlatform *TargetPlatform )
 {
-	// Make sure that all CacheShaders render thead commands are finished before we destroy FMaterialResources.
-	// TODO - is this needed, since we're using DeferredDeleteArray now?
-	FlushRenderingCommands();
-
-	TArray<FMaterialResource*> *CachedMaterialResourcesForPlatform = CachedMaterialResourcesForCooking.Find( TargetPlatform );
+	TArray<FMaterialResource*>* CachedMaterialResourcesForPlatform = CachedMaterialResourcesForCooking.Find( TargetPlatform );
 	if ( CachedMaterialResourcesForPlatform != nullptr )
 	{
 		FMaterial::DeferredDeleteArray(*CachedMaterialResourcesForPlatform);
@@ -3323,13 +3319,9 @@ void UMaterialInstance::ClearCachedCookedPlatformData( const ITargetPlatform *Ta
 
 void UMaterialInstance::ClearAllCachedCookedPlatformData()
 {
-	// Make sure that all CacheShaders render thead commands are finished before we destroy FMaterialResources.
-	// TODO - is this needed, since we're using DeferredDeleteArray now?
-	FlushRenderingCommands();
-
 	for ( auto& It : CachedMaterialResourcesForCooking )
 	{
-		TArray<FMaterialResource*> &CachedMaterialResourcesForPlatform = It.Value;
+		TArray<FMaterialResource*>& CachedMaterialResourcesForPlatform = It.Value;
 		FMaterial::DeferredDeleteArray(CachedMaterialResourcesForPlatform);
 	}
 
diff --git a/Engine/Source/Runtime/Engine/Public/MaterialShared.h b/Engine/Source/Runtime/Engine/Public/MaterialShared.h
--- a/Engine/Source/Runtime/Engine/Public/MaterialShared.h
+++ b/Engine/Source/Runtime/Engine/Public/MaterialShared.h
@@ -1581,15 +1581,8 @@ public:
 			for (TRefCountPtr<TMaterial>& Material : Materials)
 			{
 				TMaterial* MaterialToDestroy = Material.GetReference();
-				if (MaterialToDestroy->PrepareDestroy_GameThread())
-				{
-					MaterialsRenderThread.Add(MoveTemp(Material));
-				}
-				else
-				{
-					Material.SafeRelease();
-					delete MaterialToDestroy;
-				}
+				MaterialToDestroy->PrepareDestroy_GameThread();
+				MaterialsRenderThread.Add(MoveTemp(Material));
 			}
 
 			Materials.Empty();
@@ -1605,14 +1598,8 @@ public:
 			TArray<TRefCountPtr<TMaterial>> MaterialsRenderThread;
 			for (TMaterial* Material : Materials)
 			{
-				if (Material->PrepareDestroy_GameThread())
-				{
-					MaterialsRenderThread.Add(Material);
-				}
-				else
-				{
-					delete Material;
-				}
+				Material->PrepareDestroy_GameThread();
+				MaterialsRenderThread.Add(Material);
 			}
 
 			Materials.Empty();
