Change the CVar name to unique (UE-110402).

#rnx

diff --git a/Engine/Source/Runtime/Engine/Private/Materials/MaterialShader.cpp b/Engine/Source/Runtime/Engine/Private/Materials/MaterialShader.cpp
--- a/Engine/Source/Runtime/Engine/Private/Materials/MaterialShader.cpp
+++ b/Engine/Source/Runtime/Engine/Private/Materials/MaterialShader.cpp
@@ -20,7 +20,7 @@
 #include "Misc/ScopeRWLock.h"
 
 int32 GMaterialExcludeNonPipelinedShaders = 1;
-static FAutoConsoleVariableRef CVarShaderCompilerJobCache(
+static FAutoConsoleVariableRef CVarMaterialExcludeNonPipelinedShaders(
 	TEXT("r.Material.ExcludeNonPipelinedShaders"),
 	GMaterialExcludeNonPipelinedShaders,
 	TEXT("if != 0, standalone shaders that are also part of FShaderPipeline will not be compiled (default)."),
