<description: restricted, no permission to view>

diff --git a/Engine/Source/Runtime/Engine/Private/VisualLogger/VisualLogger.cpp b/Engine/Source/Runtime/Engine/Private/VisualLogger/VisualLogger.cpp
--- a/Engine/Source/Runtime/Engine/Private/VisualLogger/VisualLogger.cpp
+++ b/Engine/Source/Runtime/Engine/Private/VisualLogger/VisualLogger.cpp
@@ -475,7 +475,7 @@ FVisualLogger::FOwnerToChildrenRedirectionMap& FVisualLogger::GetRedirectionMap(
 	return WorldToRedirectionMap.FindOrAdd(World);
 }
 
-void FVisualLogger::Redirect(UObject* FromObject, UObject* ToObject)
+void FVisualLogger::Redirect(const UObject* FromObject, const UObject* ToObject)
 {
 	if (FromObject == ToObject || FromObject == nullptr || ToObject == nullptr)
 	{
diff --git a/Engine/Source/Runtime/Engine/Public/VisualLogger/VisualLogger.h b/Engine/Source/Runtime/Engine/Public/VisualLogger/VisualLogger.h
--- a/Engine/Source/Runtime/Engine/Public/VisualLogger/VisualLogger.h
+++ b/Engine/Source/Runtime/Engine/Public/VisualLogger/VisualLogger.h
@@ -422,7 +422,7 @@ public:
 	void Cleanup(UWorld* OldWorld, bool bReleaseMemory = false);
 
 	/** Set log owner redirection from one object to another, to combine logs */
-	static void Redirect(UObject* FromObject, UObject* ToObject);
+	static void Redirect(const UObject* FromObject, const UObject* ToObject);
 
 	/** find and return redirection object for given object*/
 	static UObject* FindRedirection(const UObject* Object);
