Added FName namespace to runtime attribute identifier and protected members
Added GetRegisteredTypes function to UE::Anim::AttributeTypes API

diff --git a/Engine/Source/Runtime/Engine/Classes/Animation/AttributeBlendData.h b/Engine/Source/Runtime/Engine/Classes/Animation/AttributeBlendData.h
--- a/Engine/Source/Runtime/Engine/Classes/Animation/AttributeBlendData.h
+++ b/Engine/Source/Runtime/Engine/Classes/Animation/AttributeBlendData.h
@@ -191,7 +191,7 @@ namespace UE
 				{
 					check(Data.HasBoneWeights());
 					check(CurrentIndex < Collection.DataPtrs.Num() && Collection.WeightIndices.IsValidIndex(CurrentIndex));
-					return Data.GetBoneWeight(Collection.WeightIndices[CurrentIndex], Collection.Identifier->Index);
+					return Data.GetBoneWeight(Collection.WeightIndices[CurrentIndex], Collection.Identifier->GetIndex());
 				}
 
 				/** Returns highest (container level) weighted value for the attribute set */
@@ -211,7 +211,7 @@ namespace UE
 					float Weight = -1.f;
 					for (const int32 Index : Collection.WeightIndices)
 					{
-						const float BoneWeight = Data.GetBoneWeight(Index, Collection.Identifier->Index);
+						const float BoneWeight = Data.GetBoneWeight(Index, Collection.Identifier->GetIndex());
 						if (BoneWeight > Weight)
 						{
 							Weight = BoneWeight;
@@ -231,7 +231,7 @@ namespace UE
 					float Weight = -1.f;
 					for (const int32 Index : Collection.WeightIndices)
 					{
-						const float BoneWeight = Data.GetBoneWeight(Index, Collection.Identifier->Index);
+						const float BoneWeight = Data.GetBoneWeight(Index, Collection.Identifier->GetIndex());
 						if (BoneWeight > Weight)
 						{
 							Weight = BoneWeight;
@@ -314,7 +314,7 @@ namespace UE
 					check(Data.HasBoneWeights());
 					check(CurrentIndex < AttributesView.Num() && AttributesView.IsValidIndex(CurrentIndex));
 					
-					return Data.GetBoneWeight(AttributesView[CurrentIndex].WeightIndex, AttributesView[CurrentIndex].Identifier->Index);
+					return Data.GetBoneWeight(AttributesView[CurrentIndex].WeightIndex, AttributesView[CurrentIndex].Identifier->GetIndex());
 				}
 
 				/** Returns whether or not the unique attribute its (bone level) weight is the highest across the containers */
@@ -322,9 +322,9 @@ namespace UE
 				{
 					check(Data.HasBoneWeights());
 					check(CurrentIndex < AttributesView.Num() && AttributesView.IsValidIndex(CurrentIndex));
-					check(Data.HighestBoneWeightedIndices.IsValidIndex(AttributesView[CurrentIndex].Identifier->Index));
+					check(Data.HighestBoneWeightedIndices.IsValidIndex(AttributesView[CurrentIndex].Identifier->GetIndex()));
 
-					return Data.HighestBoneWeightedIndices[AttributesView[CurrentIndex].Identifier->Index] == AttributesView[CurrentIndex].WeightIndex;
+					return Data.HighestBoneWeightedIndices[AttributesView[CurrentIndex].Identifier->GetIndex()] == AttributesView[CurrentIndex].WeightIndex;
 				}
 
 				/** Returns the identifier for the current attribute set */
diff --git a/Engine/Source/Runtime/Engine/Classes/Animation/AttributesContainer.h b/Engine/Source/Runtime/Engine/Classes/Animation/AttributesContainer.h
--- a/Engine/Source/Runtime/Engine/Classes/Animation/AttributesContainer.h
+++ b/Engine/Source/Runtime/Engine/Classes/Animation/AttributesContainer.h
@@ -19,18 +19,24 @@ enum class ECustomAttributeBlendType : uint8;
 namespace UE
 {
 	namespace Anim
-	{	
+	{
+		static FName BoneAttributeNamespace("bone");
 		/** Runtime identifier for an attribute */
 		struct ENGINE_API FAttributeId
 		{
-			FAttributeId(const FName& InName, const FCompactPoseBoneIndex& InCompactBoneIndex) : Name(InName), Index(InCompactBoneIndex.GetInt()) {}
-			FAttributeId(const FName& InName, int32 InIndex) : Name(InName), Index(InIndex) {}
+			FAttributeId(const FName& InName, const FCompactPoseBoneIndex& InCompactBoneIndex) : Namespace(BoneAttributeNamespace), Name(InName), Index(InCompactBoneIndex.GetInt()) {}
+			FAttributeId(const FName& InName, int32 InIndex, const FName& InNamespace) : Namespace(InNamespace), Name(InName), Index(InIndex) {}
 			
 			bool operator==(const FAttributeId& Other) const
 			{
-				return Index == Other.Index && Name == Other.Name;
+				return Index == Other.Index && Name == Other.Name && Namespace == Other.Namespace;
 			}
 
+			int32 GetIndex() const { return Index; }
+			const FName& GetName() const { return Name; }
+			const FName& GetNamespace() const { return Namespace; }			
+		protected:
+			FName Namespace;
 			FName Name;
 			int32 Index;
 		};
@@ -234,7 +240,7 @@ namespace UE
 
 					AttributeIds.Add(InAttributeId);
 
-					UniqueTypedBoneIndices[TypeIndex].AddUnique(InAttributeId.Index);
+					UniqueTypedBoneIndices[TypeIndex].AddUnique(InAttributeId.GetIndex());
 
 					TArray<TWrappedAttribute<InAllocator>, InAllocator>& TypedArray = Values[TypeIndex];
 					NewIndex = TypedArray.Add(InScriptStruct);
@@ -286,7 +292,7 @@ namespace UE
 				{
 					AttributeIds.Add(InAttributeId);
 
-					UniqueTypedBoneIndices[TypeIndex].AddUnique(InAttributeId.Index);
+					UniqueTypedBoneIndices[TypeIndex].AddUnique(InAttributeId.GetIndex());
 
 					AttributeIndex = TypedArray.Add(InScriptStruct);
 					TWrappedAttribute<InAllocator>& StructData = TypedArray[AttributeIndex];
@@ -474,7 +480,7 @@ namespace UE
 					const TArray<int32, InAllocator>& UniqueBoneIndices = UniqueTypedBoneIndices[TypeIndex];
 
 					// Early out if for this bone index no attributes are currently contained
-					if (UniqueBoneIndices.Contains(InAttributeId.Index))
+					if (UniqueBoneIndices.Contains(InAttributeId.GetIndex()))
 					{
 						const TArray<FAttributeId, InAllocator>& AttributeIds = AttributeIdentifiers[TypeIndex];
 						const int32 AttributeIndex = AttributeIds.IndexOfByKey(InAttributeId);
@@ -699,6 +705,6 @@ struct ENGINE_API UE_DEPRECATED(5.0, "FCustomAttributeInfo has been deprecated u
 	UE_DEPRECATED(5.0, "Deprecated constructor, see  UE::Anim::FAttributeId")
 	FCustomAttributeInfo(const FName& InName, const FCompactPoseBoneIndex& InCompactBoneIndex, const ECustomAttributeBlendType& InBlendType) : UE::Anim::FAttributeId(InName, InCompactBoneIndex) {}
 	
-	FCustomAttributeInfo(const FName& InName, const FCompactPoseBoneIndex& InCompactBoneIndex) : UE::Anim::FAttributeId(InName, InCompactBoneIndex.GetInt()) {}
-	FCustomAttributeInfo(const FName& InName, int32 InIndex) : UE::Anim::FAttributeId(InName, InIndex) {}
+	FCustomAttributeInfo(const FName& InName, const FCompactPoseBoneIndex& InCompactBoneIndex) : UE::Anim::FAttributeId(InName, InCompactBoneIndex) {}
+	FCustomAttributeInfo(const FName& InName, int32 InIndex) : UE::Anim::FAttributeId(InName, InIndex, UE::Anim::BoneAttributeNamespace) {}
 };
\ No newline at end of file
diff --git a/Engine/Source/Runtime/Engine/Private/Animation/AnimSequence.cpp b/Engine/Source/Runtime/Engine/Private/Animation/AnimSequence.cpp
--- a/Engine/Source/Runtime/Engine/Private/Animation/AnimSequence.cpp
+++ b/Engine/Source/Runtime/Engine/Private/Animation/AnimSequence.cpp
@@ -4745,11 +4745,11 @@ void UAnimSequence::SynchronousAnimatedBoneAttributesCompression()
 					const float CurrentFrameTime = Frame * EvalContext.IntervalTime;
 
 					UE::Anim::FStackAttributeContainer AdditiveAttributes;
-					uint8* AdditivePtr = AdditiveAttributes.FindOrAdd(AdditiveAttribute.Identifier.GetType(), UE::Anim::FAttributeId(NAME_None, 0));
+					uint8* AdditivePtr = AdditiveAttributes.FindOrAdd(AdditiveAttribute.Identifier.GetType(), UE::Anim::FAttributeId(NAME_None, 0,NAME_None));
 					AdditiveAttribute.Curve.EvaluateToPtr(AdditiveAttribute.Identifier.GetType(), CurrentFrameTime, AdditivePtr);
 
 					UE::Anim::FStackAttributeContainer RefAttributes;
-					uint8* RefPtr = RefAttributes.FindOrAdd(RefAttribute.Identifier.GetType(), UE::Anim::FAttributeId(NAME_None, 0));
+					uint8* RefPtr = RefAttributes.FindOrAdd(RefAttribute.Identifier.GetType(), UE::Anim::FAttributeId(NAME_None, 0, NAME_None));
 					RefAttribute.Curve.EvaluateToPtr(RefAttribute.Identifier.GetType(), GetBasePoseTimeToSample(CurrentFrameTime), RefPtr);
 
 					UE::Anim::Attributes::ConvertToAdditive(RefAttributes, AdditiveAttributes);
diff --git a/Engine/Source/Runtime/Engine/Private/Animation/AttributesRuntime.cpp b/Engine/Source/Runtime/Engine/Private/Animation/AttributesRuntime.cpp
--- a/Engine/Source/Runtime/Engine/Private/Animation/AttributesRuntime.cpp
+++ b/Engine/Source/Runtime/Engine/Private/Animation/AttributesRuntime.cpp
@@ -172,7 +172,7 @@ void Attributes::CopyAndRemapAttributes(const FHeapAttributeContainer& SourceAtt
 			{
 				const FAttributeId& AttributeId = AttributeIds[EntryIndex];
 
-				const FCompactPoseBoneIndex PoseBoneIndex = FCompactPoseBoneIndex(AttributeId.Index);
+				const FCompactPoseBoneIndex PoseBoneIndex = FCompactPoseBoneIndex(AttributeId.GetIndex());
 				const int32 SkeletonBoneIndex = RequiredBones.GetSkeletonIndex(FCompactPoseBoneIndex(PoseBoneIndex));
 				const int32 MeshBoneIndex = RequiredBones.GetSkeletonToPoseBoneIndexArray()[SkeletonBoneIndex];
 				const int32* Value = BoneMapToSource.Find(MeshBoneIndex);
@@ -181,7 +181,7 @@ void Attributes::CopyAndRemapAttributes(const FHeapAttributeContainer& SourceAtt
 				if (Value)
 				{
 					const int32 RemappedBoneIndex = *Value;
-					const FAttributeId NewInfo(AttributeId.Name, RemappedBoneIndex);
+					const FAttributeId NewInfo(AttributeId.GetName(), RemappedBoneIndex, AttributeId.GetNamespace());
 					uint8* NewAttribute = TargetAttributes.FindOrAdd(ScriptStruct, NewInfo);
 					ScriptStruct->CopyScriptStruct(NewAttribute, SourceValues[EntryIndex].GetPtr<void>());
 				}
@@ -456,7 +456,7 @@ ECustomAttributeBlendType Attributes::GetAttributeBlendType(const FAttributeId&
 	Settings->AttributeBlendModes.GenerateKeyArray(Names);
 	const FName* NamePtr = Names.FindByPredicate([Info](const FName& Name)
 	{
-		return Name == Info.Name;
+		return Name == Info.GetName();
 	});
 
 	return NamePtr ? Settings->AttributeBlendModes.FindChecked(*NamePtr) : Settings->DefaultAttributeBlendMode;
diff --git a/Engine/Source/Runtime/Engine/Public/Animation/AttributeTypes.h b/Engine/Source/Runtime/Engine/Public/Animation/AttributeTypes.h
--- a/Engine/Source/Runtime/Engine/Public/Animation/AttributeTypes.h
+++ b/Engine/Source/Runtime/Engine/Public/Animation/AttributeTypes.h
@@ -122,12 +122,15 @@ namespace UE
 			template<typename AttributeType>
 			static void UnregisterType()
 			{
-				UScriptStruct* ScriptStruct = AttributeType::StaticStruct();
-				const int32 Index = AttributeTypes::RegisteredTypes.IndexOfByKey(ScriptStruct);
-				if (Index != INDEX_NONE)
+				if (UObjectInitialized())
 				{
-					AttributeTypes::RegisteredTypes.RemoveAtSwap(Index);
-					AttributeTypes::Operators.RemoveAtSwap(Index);
+					UScriptStruct* ScriptStruct = AttributeType::StaticStruct();
+					const int32 Index = AttributeTypes::RegisteredTypes.IndexOfByKey(ScriptStruct);
+					if (Index != INDEX_NONE)
+					{
+						AttributeTypes::RegisteredTypes.RemoveAtSwap(Index);
+						AttributeTypes::Operators.RemoveAtSwap(Index);
+					}
 				}
 			}
 
@@ -153,6 +156,13 @@ namespace UE
 				AttributeTypes::Initialize();
 				return AttributeTypes::RegisteredTypes.Contains(ScriptStruct);
 			}
+
+			/** Returns all registered types */
+			static TArray<TWeakObjectPtr<const UScriptStruct>>& GetRegisteredTypes()
+			{
+				AttributeTypes::Initialize();
+				return AttributeTypes::RegisteredTypes;
+			}
 		};
 	}
 }
