Only include external actors in the unsaved for PIE actor list.

diff --git a/Engine/Source/Runtime/Engine/Private/WorldPartition/WorldPartition.cpp b/Engine/Source/Runtime/Engine/Private/WorldPartition/WorldPartition.cpp
--- a/Engine/Source/Runtime/Engine/Private/WorldPartition/WorldPartition.cpp
+++ b/Engine/Source/Runtime/Engine/Private/WorldPartition/WorldPartition.cpp
@@ -146,7 +146,7 @@ void UWorldPartition::OnPreBeginPIE(bool bStartSimulate)
 	// Gather all modified/newly created actors for PIE so they can be duplicated when streaming their cell
 	for (AActor* Actor : World->PersistentLevel->Actors)
 	{
-		if (Actor && !Actor->IsPendingKill() && Actor->GetPackage()->IsDirty())
+		if (Actor && Actor->IsPackageExternal() && !Actor->IsPendingKill() && Actor->GetPackage()->IsDirty())
 		{
 			World->PersistentLevel->ActorsModifiedForPIE.Add(Actor->GetFName(), Actor);
 		}
