Bump hung shadermap detector threshold to 2 hours (numerous JIRAs).

- Seems that 0.5h is too small a threshold as we can hit it on overloaded machines or some pathological cases.

#rnx

diff --git a/Engine/Source/Runtime/Engine/Private/ShaderCompiler/ShaderCompiler.cpp b/Engine/Source/Runtime/Engine/Private/ShaderCompiler/ShaderCompiler.cpp
--- a/Engine/Source/Runtime/Engine/Private/ShaderCompiler/ShaderCompiler.cpp
+++ b/Engine/Source/Runtime/Engine/Private/ShaderCompiler/ShaderCompiler.cpp
@@ -743,7 +743,7 @@ static FAutoConsoleVariableRef CVarDumpShaderDebugSCWCommandLine(
 	TEXT("When set to 1, it will generate a file that can be used with ShaderCompileWorker's -directcompile.")
 	);
 
-static int32 GShaderMapCompilationTimeout = 30 * 60;
+static int32 GShaderMapCompilationTimeout = 2 * 60 * 60;	// anything below an hour can hit a false positive
 static FAutoConsoleVariableRef CVarShaderMapCompilationTimeout(
 	TEXT("r.ShaderCompiler.ShadermapCompilationTimeout"),
 	GShaderMapCompilationTimeout,
