Fix check() on a shader compilation failure.

diff --git a/Engine/Source/Runtime/Engine/Private/ShaderCompiler/ShaderCompiler.cpp b/Engine/Source/Runtime/Engine/Private/ShaderCompiler/ShaderCompiler.cpp
--- a/Engine/Source/Runtime/Engine/Private/ShaderCompiler/ShaderCompiler.cpp
+++ b/Engine/Source/Runtime/Engine/Private/ShaderCompiler/ShaderCompiler.cpp
@@ -1281,9 +1281,9 @@ static bool ReadSingleJob(FShaderCompileJob* CurrentJob, FArchive& OutputFile)
 	// Deserialize the shader compilation output.
 	OutputFile << CurrentJob->Output;
 
-	// The job should already have a non-zero output hash
-	checkf(CurrentJob->Output.OutputHash != FSHAHash(), TEXT("OutputHash was not set in the shader compile worker!"));
 	CurrentJob->bSucceeded = CurrentJob->Output.bSucceeded;
+	// The job should already have a non-zero output hash
+	checkf(CurrentJob->Output.OutputHash != FSHAHash() || !CurrentJob->bSucceeded, TEXT("OutputHash for a successful job was not set in the shader compile worker!"));
 
 	if (CurrentJob->bSucceeded && CurrentJob->Input.DumpDebugInfoPath.Len() > 0)
 	{
