WorldPartition: Fix Cook, Avoid clear streaming grids when cooking

diff --git a/Engine/Source/Runtime/Engine/Private/WorldPartition/WorldPartitionRuntimeSpatialHash.cpp b/Engine/Source/Runtime/Engine/Private/WorldPartition/WorldPartitionRuntimeSpatialHash.cpp
--- a/Engine/Source/Runtime/Engine/Private/WorldPartition/WorldPartitionRuntimeSpatialHash.cpp
+++ b/Engine/Source/Runtime/Engine/Private/WorldPartition/WorldPartitionRuntimeSpatialHash.cpp
@@ -372,9 +372,12 @@ void UWorldPartitionRuntimeSpatialHash::PreSave(const class ITargetPlatform* Tar
 {
 	Super::PreSave(TargetPlatform);
 
-	// We don't want this to be persisted but we can't set the property Transient as it is NonPIEDuplicateTransient and those flags aren't compatible
-	// If at some point GenerateStreaming is done after duplication we can remove this code.
-	StreamingGrids.Empty();
+	if (!IsRunningCookCommandlet())
+	{
+		// We don't want this to be persisted but we can't set the property Transient as it is NonPIEDuplicateTransient and those flags aren't compatible
+		// If at some point GenerateStreaming is done after duplication we can remove this code.
+		StreamingGrids.Empty();
+	}
 }
 
 #if WITH_EDITOR
