WorldPartition: Fix Commandlet Foliage crash because of EditorHash not being initialized

diff --git a/Engine/Source/Editor/UnrealEd/Private/Commandlets/WorldPartitionConvertCommandlet.cpp b/Engine/Source/Editor/UnrealEd/Private/Commandlets/WorldPartitionConvertCommandlet.cpp
--- a/Engine/Source/Editor/UnrealEd/Private/Commandlets/WorldPartitionConvertCommandlet.cpp
+++ b/Engine/Source/Editor/UnrealEd/Private/Commandlets/WorldPartitionConvertCommandlet.cpp
@@ -225,6 +225,8 @@ UWorldPartition* UWorldPartitionConvertCommandlet::CreateWorldPartition(AWorldSe
 		WorldPartition->RuntimeHash->LoadConfig(*RuntimeHashClass, *LevelConfigFilename);
 		WorldPartition->DefaultHLODLayer = HLODLayers.FindRef(DefaultHLODLayerName);
 	}
+
+	WorldPartition->EditorHash->Initialize();
 	
 	WorldPartition->RuntimeHash->ImportFromWorldComposition(WorldComposition);
 
diff --git a/Engine/Source/Runtime/Engine/Private/WorldPartition/WorldPartitionEditorSpatialHash.cpp b/Engine/Source/Runtime/Engine/Private/WorldPartition/WorldPartitionEditorSpatialHash.cpp
--- a/Engine/Source/Runtime/Engine/Private/WorldPartition/WorldPartitionEditorSpatialHash.cpp
+++ b/Engine/Source/Runtime/Engine/Private/WorldPartition/WorldPartitionEditorSpatialHash.cpp
@@ -29,10 +29,13 @@ void UWorldPartitionEditorSpatialHash::Initialize()
 	AlwaysLoadedCell = NewObject<UWorldPartitionEditorCell>(this, TEXT("AlwaysLoadedCell"), RF_Transient);
 	AlwaysLoadedCell->Bounds.Init();
 
-	FHashBuilder EditorGridConfigHashBuilder;
-	static uint32 Version = 1;
-	EditorGridConfigHashBuilder << Version << CellSize;	
-	GetMutableDefault<UWorldPartitionEditorPerProjectUserSettings>()->SetEditorGridConfigHash(EditorGridConfigHashBuilder.GetHash());
+	if (!IsRunningCommandlet())
+	{
+		FHashBuilder EditorGridConfigHashBuilder;
+		static uint32 Version = 1;
+		EditorGridConfigHashBuilder << Version << CellSize;
+		GetMutableDefault<UWorldPartitionEditorPerProjectUserSettings>()->SetEditorGridConfigHash(EditorGridConfigHashBuilder.GetHash());
+	}
 }
 
 void UWorldPartitionEditorSpatialHash::SetDefaultValues()
