Filter actor descriptor container save/delete events to properly support maps with level instances.

diff --git a/Engine/Source/Runtime/Engine/Private/WorldPartition/ActorDescContainer.cpp b/Engine/Source/Runtime/Engine/Private/WorldPartition/ActorDescContainer.cpp
--- a/Engine/Source/Runtime/Engine/Private/WorldPartition/ActorDescContainer.cpp
+++ b/Engine/Source/Runtime/Engine/Private/WorldPartition/ActorDescContainer.cpp
@@ -137,11 +137,34 @@ TUniquePtr<FWorldPartitionActorDesc> UActorDescContainer::GetActorDescriptor(con
 	return nullptr;
 }
 
+bool UActorDescContainer::ShouldHandleActorEvent(const AActor* Actor)
+{
+	if (Actor && Actor->IsPackageExternal() && (Actor->GetLevel() == World->PersistentLevel))
+	{
+		// Only handle assets that belongs to our level
+		auto RemoveAfterFirstDot = [](const FString& InValue)
+		{
+			int32 DotIndex;
+			if (InValue.FindChar(TEXT('.'), DotIndex))
+			{
+				return InValue.LeftChop(InValue.Len() - DotIndex);
+			}
+			return InValue;
+		};
+
+		const FString ThisLevelPath = ContainerPackageName.ToString();
+		const FString AssetLevelPath = RemoveAfterFirstDot(Actor->GetPathName());
+		return (ThisLevelPath == AssetLevelPath);
+	}
+
+	return false;
+}
+
 void UActorDescContainer::OnObjectPreSave(UObject* Object)
 {
 	if (const AActor* Actor = Cast<AActor>(Object))
 	{
-		if (Actor->IsPackageExternal() && (Actor->GetLevel() == World->PersistentLevel))
+		if (ShouldHandleActorEvent(Actor))
 		{
 			if (TUniquePtr<FWorldPartitionActorDesc>* ExistingActorDesc = Actors.FindRef(Actor->GetActorGuid()))
 			{
@@ -178,10 +201,9 @@ void UActorDescContainer::OnObjectPreSave(UObject* Object)
 void UActorDescContainer::OnPackageDeleted(UPackage* Package)
 {
 	AActor* Actor = nullptr;
-
 	ForEachObjectWithPackage(Package, [&Actor](UObject* Object)	{ Actor = Cast<AActor>(Object);	return !Actor; }, false);
 
-	if (Actor && Actor->IsPackageExternal() && (Actor->GetLevel() == World->PersistentLevel))
+	if (ShouldHandleActorEvent(Actor))
 	{
 		if (TUniquePtr<FWorldPartitionActorDesc>* ExistingActorDesc = Actors.FindRef(Actor->GetActorGuid()))
 		{
diff --git a/Engine/Source/Runtime/Engine/Public/WorldPartition/ActorDescContainer.h b/Engine/Source/Runtime/Engine/Public/WorldPartition/ActorDescContainer.h
--- a/Engine/Source/Runtime/Engine/Public/WorldPartition/ActorDescContainer.h
+++ b/Engine/Source/Runtime/Engine/Public/WorldPartition/ActorDescContainer.h
@@ -52,6 +52,8 @@ protected:
 	virtual void OnActorDescRegistered(const FWorldPartitionActorDesc&) {}
 	virtual void OnActorDescUnregistered(const FWorldPartitionActorDesc&) {}
 
+	bool ShouldHandleActorEvent(const AActor* Actor);
+
 	bool bContainerInitialized;
 
 	FName ContainerPackageName;
