WorldPartition: Fix AlwaysLoadedFor PIE crash on unsaved new actors

diff --git a/Engine/Source/Runtime/Engine/Private/WorldPartition/WorldPartitionRuntimeHash.cpp b/Engine/Source/Runtime/Engine/Private/WorldPartition/WorldPartitionRuntimeHash.cpp
--- a/Engine/Source/Runtime/Engine/Private/WorldPartition/WorldPartitionRuntimeHash.cpp
+++ b/Engine/Source/Runtime/Engine/Private/WorldPartition/WorldPartitionRuntimeHash.cpp
@@ -60,9 +60,9 @@ void UWorldPartitionRuntimeHash::OnEndPlay()
 
 void UWorldPartitionRuntimeHash::ForceExternalActorLevelReference(bool bForceExternalActorLevelReferenceForPIE)
 {
-	for (const FWorldPartitionReference& AlwaysLoadedActor : AlwaysLoadedActorsForPIE)
+	for (const FAlwaysLoadedActorForPIE& AlwaysLoadedActor : AlwaysLoadedActorsForPIE)
 	{
-		if (AActor* Actor = (*AlwaysLoadedActor)->GetActor())
+		if (AActor* Actor = AlwaysLoadedActor.Actor)
 		{
 			Actor->SetForceExternalActorLevelReferenceForPIE(bForceExternalActorLevelReferenceForPIE);
 		}
diff --git a/Engine/Source/Runtime/Engine/Private/WorldPartition/WorldPartitionRuntimeSpatialHash.cpp b/Engine/Source/Runtime/Engine/Private/WorldPartition/WorldPartitionRuntimeSpatialHash.cpp
--- a/Engine/Source/Runtime/Engine/Private/WorldPartition/WorldPartitionRuntimeSpatialHash.cpp
+++ b/Engine/Source/Runtime/Engine/Private/WorldPartition/WorldPartitionRuntimeSpatialHash.cpp
@@ -661,7 +661,11 @@ bool UWorldPartitionRuntimeSpatialHash::CreateStreamingGrid(const FSpatialHashRu
 						{
 							if (ActorInstance.ContainerInstance->Container == WorldPartition)
 							{
-								AlwaysLoadedActorsForPIE.Emplace(WorldPartition, ActorInstance.Actor);
+								// This will load the actor if it isn't already loaded
+								FWorldPartitionReference Reference(WorldPartition, ActorInstance.Actor);
+								const FWorldPartitionActorDescView& ActorDescView = ActorInstance.GetActorDescView();
+								AActor* AlwaysLoadedActor = FindObject<AActor>(nullptr, *ActorDescView.GetActorPath().ToString());
+								AlwaysLoadedActorsForPIE.Emplace(Reference, AlwaysLoadedActor);
 								continue;
 							}
 						}
diff --git a/Engine/Source/Runtime/Engine/Public/WorldPartition/WorldPartitionRuntimeHash.h b/Engine/Source/Runtime/Engine/Public/WorldPartition/WorldPartitionRuntimeHash.h
--- a/Engine/Source/Runtime/Engine/Public/WorldPartition/WorldPartitionRuntimeHash.h
+++ b/Engine/Source/Runtime/Engine/Public/WorldPartition/WorldPartitionRuntimeHash.h
@@ -67,7 +67,17 @@ private:
 
 protected:
 #if WITH_EDITORONLY_DATA
-	TArray<FWorldPartitionReference> AlwaysLoadedActorsForPIE;
+	struct FAlwaysLoadedActorForPIE
+	{
+		FAlwaysLoadedActorForPIE(const FWorldPartitionReference& InReference, AActor* InActor)
+			: Reference(InReference), Actor(InActor)
+		{}
+
+		FWorldPartitionReference Reference;
+		AActor* Actor;
+	};
+
+	TArray<FAlwaysLoadedActorForPIE> AlwaysLoadedActorsForPIE;
 
 	mutable FActorDescList ModifiedActorDescListForPIE;
 #endif
