Save last loaded cells for partitioned worlds when changing worlds, instead of right after loading cells.
This makes sure we don't end up in a situation where loading all cells crashes the editor and we can't go back in.

diff --git a/Engine/Source/Runtime/Engine/Private/WorldPartition/WorldPartition.cpp b/Engine/Source/Runtime/Engine/Private/WorldPartition/WorldPartition.cpp
--- a/Engine/Source/Runtime/Engine/Private/WorldPartition/WorldPartition.cpp
+++ b/Engine/Source/Runtime/Engine/Private/WorldPartition/WorldPartition.cpp
@@ -90,20 +90,8 @@ struct FWorldPartionCellUpdateContext
 					WorldPartition->WorldPartitionEditor->Refresh();
 				}
 
-				// Save last loaded cells settings (ignore while running commandlets)
 				if (!IsRunningCommandlet())
 				{
-					TArray<FName> EditorGridLastLoadedCells;
-					WorldPartition->EditorHash->ForEachCell([this, &EditorGridLastLoadedCells](UWorldPartitionEditorCell* Cell)
-					{
-						if ((Cell != WorldPartition->EditorHash->GetAlwaysLoadedCell()) && Cell->bLoaded)
-						{
-							EditorGridLastLoadedCells.Add(Cell->GetFName());
-						}
-					});
-
-					GetMutableDefault<UWorldPartitionEditorPerProjectUserSettings>()->SetEditorGridLoadedCells(WorldPartition->GetWorld(), EditorGridLastLoadedCells);
-
 					GEditor->ResetTransaction(LOCTEXT("LoadingEditorCellsResetTrans", "Editor Cells Loading State Changed"));
 				}
 			}
@@ -368,6 +356,21 @@ void UWorldPartition::Uninitialize()
 	{
 		check(World);
 
+		if (!IsRunningCommandlet())
+		{
+			// Save last loaded cells settings
+			TArray<FName> EditorGridLastLoadedCells;
+			EditorHash->ForEachCell([this, &EditorGridLastLoadedCells](UWorldPartitionEditorCell* Cell)
+			{
+				if ((Cell != EditorHash->GetAlwaysLoadedCell()) && Cell->bLoaded)
+				{
+					EditorGridLastLoadedCells.Add(Cell->GetFName());
+				}
+			});
+
+			GetMutableDefault<UWorldPartitionEditorPerProjectUserSettings>()->SetEditorGridLoadedCells(GetWorld(), EditorGridLastLoadedCells);
+		}
+
 		InitState = EWorldPartitionInitState::Uninitializing;
 		
 		// Unload all loaded cells
diff --git a/Engine/Source/Runtime/Engine/Private/WorldPartition/WorldPartitionEditorSpatialHash.cpp b/Engine/Source/Runtime/Engine/Private/WorldPartition/WorldPartitionEditorSpatialHash.cpp
--- a/Engine/Source/Runtime/Engine/Private/WorldPartition/WorldPartitionEditorSpatialHash.cpp
+++ b/Engine/Source/Runtime/Engine/Private/WorldPartition/WorldPartitionEditorSpatialHash.cpp
@@ -32,7 +32,7 @@ void UWorldPartitionEditorSpatialHash::Initialize()
 	if (!IsRunningCommandlet())
 	{
 		FHashBuilder EditorGridConfigHashBuilder;
-		static uint32 Version = 1;
+		static uint32 Version = 2;
 		EditorGridConfigHashBuilder << Version << CellSize;
 		GetMutableDefault<UWorldPartitionEditorPerProjectUserSettings>()->SetEditorGridConfigHash(EditorGridConfigHashBuilder.GetHash());
 	}
diff --git a/Engine/Source/Runtime/Engine/Public/WorldPartition/WorldPartitionEditorPerProjectUserSettings.h b/Engine/Source/Runtime/Engine/Public/WorldPartition/WorldPartitionEditorPerProjectUserSettings.h
--- a/Engine/Source/Runtime/Engine/Public/WorldPartition/WorldPartitionEditorPerProjectUserSettings.h
+++ b/Engine/Source/Runtime/Engine/Public/WorldPartition/WorldPartitionEditorPerProjectUserSettings.h
@@ -65,9 +65,12 @@ public:
 
 	void SetEditorGridLoadedCells(UWorld* InWorld, TArray<FName>& InEditorGridLoadedCells)
 	{
-		TArray<FName>& GridLoadedCells = PerWorldEditorSettings.FindOrAdd(TSoftObjectPtr<UWorld>(InWorld)).LoadedEditorGridCells;
-		GridLoadedCells = InEditorGridLoadedCells;
-		SaveConfig();
+		if (ShouldSaveSettings(InWorld))
+		{
+			TArray<FName>& GridLoadedCells = PerWorldEditorSettings.FindOrAdd(TSoftObjectPtr<UWorld>(InWorld)).LoadedEditorGridCells;
+			GridLoadedCells = InEditorGridLoadedCells;
+			SaveConfig();
+		}
 	}
 
 	bool GetShowDataLayerContent() const
@@ -96,7 +99,7 @@ public:
 
 	void SetWorldDataLayersNotLoadedInEditor(UWorld* InWorld, const TArray<FName>& InDataLayersLoadedInEditor)
 	{
-		if (InWorld)
+		if (ShouldSaveSettings(InWorld))
 		{
 			TArray<FName>& DataLayersNotLoadedInEditor = PerWorldEditorSettings.FindOrAdd(TSoftObjectPtr<UWorld>(InWorld)).NotLoadedDataLayers;
 			DataLayersNotLoadedInEditor = InDataLayersLoadedInEditor;
@@ -107,6 +110,11 @@ public:
 
 private:
 #if WITH_EDITORONLY_DATA
+	bool ShouldSaveSettings(const UWorld* InWorld) const
+	{
+		return InWorld && !InWorld->IsGameWorld();
+	}
+
 	UPROPERTY(config)
 	uint32 EditorGridConfigHash;
 
