Support multiple data layers in console command wp.Runtime.ToggleDataLayerActivation

diff --git a/Engine/Source/Runtime/Engine/Private/WorldPartition/DataLayer/DataLayerSubsystem.cpp b/Engine/Source/Runtime/Engine/Private/WorldPartition/DataLayer/DataLayerSubsystem.cpp
--- a/Engine/Source/Runtime/Engine/Private/WorldPartition/DataLayer/DataLayerSubsystem.cpp
+++ b/Engine/Source/Runtime/Engine/Private/WorldPartition/DataLayer/DataLayerSubsystem.cpp
@@ -191,12 +191,12 @@ void UDataLayerSubsystem::Draw(UCanvas* Canvas, class APlayerController* PC)
 
 FAutoConsoleCommand UDataLayerSubsystem::ToggleDataLayerActivation(
 	TEXT("wp.Runtime.ToggleDataLayerActivation"),
-	TEXT("Toggles a DataLayer. Args [DataLayerLabel]"),
+	TEXT("Toggles DataLayers active state. Args [DataLayerLabels]"),
 	FConsoleCommandWithArgsDelegate::CreateLambda([](const TArray<FString>& Args)
 	{
-		if (Args.Num() == 1)
+		for (const FString& Arg : Args)
 		{
-			FName DataLayerLabel = FName(Args[0]);
+			FName DataLayerLabel = FName(Arg);
 			for (const FWorldContext& Context : GEngine->GetWorldContexts())
 			{
 				UWorld* World = Context.World();
