Don't try to save last loaded cells when the engine is collecting garbage.

diff --git a/Engine/Source/Runtime/Engine/Private/WorldPartition/WorldPartition.cpp b/Engine/Source/Runtime/Engine/Private/WorldPartition/WorldPartition.cpp
--- a/Engine/Source/Runtime/Engine/Private/WorldPartition/WorldPartition.cpp
+++ b/Engine/Source/Runtime/Engine/Private/WorldPartition/WorldPartition.cpp
@@ -357,7 +357,7 @@ void UWorldPartition::Uninitialize()
 		check(World);
 
 #if WITH_EDITOR
-		if (!IsRunningCommandlet())
+		if (!IsRunningCommandlet() && !IsEngineExitRequested())
 		{
 			// Save last loaded cells settings
 			TArray<FName> EditorGridLastLoadedCells;
