Save last loaded cells on editor close, not only when changing maps.

diff --git a/Engine/Source/Runtime/Engine/Private/WorldPartition/WorldPartition.cpp b/Engine/Source/Runtime/Engine/Private/WorldPartition/WorldPartition.cpp
--- a/Engine/Source/Runtime/Engine/Private/WorldPartition/WorldPartition.cpp
+++ b/Engine/Source/Runtime/Engine/Private/WorldPartition/WorldPartition.cpp
@@ -312,6 +312,11 @@ void UWorldPartition::Initialize(UWorld* InWorld, const FTransform& InTransform)
 				}
 			}
 		}
+		
+		if (GEditor)
+		{
+			GEditor->OnEditorClose().AddUObject(this, &UWorldPartition::SavePerUserSettings);
+		}
 	}
 #endif //WITH_EDITOR
 
@@ -356,23 +361,6 @@ void UWorldPartition::Uninitialize()
 	{
 		check(World);
 
-#if WITH_EDITOR
-		if (!IsRunningCommandlet() && !IsEngineExitRequested())
-		{
-			// Save last loaded cells settings
-			TArray<FName> EditorGridLastLoadedCells;
-			EditorHash->ForEachCell([this, &EditorGridLastLoadedCells](UWorldPartitionEditorCell* Cell)
-			{
-				if ((Cell != EditorHash->GetAlwaysLoadedCell()) && Cell->bLoaded)
-				{
-					EditorGridLastLoadedCells.Add(Cell->GetFName());
-				}
-			});
-
-			GetMutableDefault<UWorldPartitionEditorPerProjectUserSettings>()->SetEditorGridLoadedCells(GetWorld(), EditorGridLastLoadedCells);
-		}
-#endif
-
 		InitState = EWorldPartitionInitState::Uninitializing;
 		
 		// Unload all loaded cells
@@ -382,6 +370,13 @@ void UWorldPartition::Uninitialize()
 		}
 
 #if WITH_EDITOR
+		SavePerUserSettings();
+
+		if (GEditor)
+		{
+			GEditor->OnEditorClose().RemoveAll(this);
+		}
+
 		for (auto& Pair : ActorDescContainers)
 		{
 			if (UActorDescContainer* Container = Pair.Value.Get())
@@ -998,6 +993,24 @@ void UWorldPartition::FinalizeGeneratedPackageForCook()
 	RuntimeHash->FinalizeGeneratedPackageForCook();
 }
 
+void UWorldPartition::SavePerUserSettings()
+{
+	if (GIsEditor && !IsRunningCommandlet() && !IsEngineExitRequested())
+	{
+		// Save last loaded cells settings
+		TArray<FName> EditorGridLastLoadedCells;
+		EditorHash->ForEachCell([this, &EditorGridLastLoadedCells](UWorldPartitionEditorCell* Cell)
+		{
+			if ((Cell != EditorHash->GetAlwaysLoadedCell()) && Cell->bLoaded)
+			{
+				EditorGridLastLoadedCells.Add(Cell->GetFName());
+			}
+		});
+
+		GetMutableDefault<UWorldPartitionEditorPerProjectUserSettings>()->SetEditorGridLoadedCells(GetWorld(), EditorGridLastLoadedCells);
+	}
+}
+
 void UWorldPartition::FlushStreaming()
 {
 	RuntimeHash->FlushStreaming();
diff --git a/Engine/Source/Runtime/Engine/Public/WorldPartition/WorldPartition.h b/Engine/Source/Runtime/Engine/Public/WorldPartition/WorldPartition.h
--- a/Engine/Source/Runtime/Engine/Public/WorldPartition/WorldPartition.h
+++ b/Engine/Source/Runtime/Engine/Public/WorldPartition/WorldPartition.h
@@ -86,6 +86,8 @@ private:
 
 	DECLARE_MULTICAST_DELEGATE_OneParam(FWorldPartitionChangedEvent, UWorld*);
 	static FWorldPartitionChangedEvent WorldPartitionChangedEvent;
+
+	void SavePerUserSettings();
 		
 	void FlushStreaming();
 
