Add virtual texture info to ListTextures output dump.

diff --git a/Engine/Source/Runtime/Engine/Private/UnrealEngine.cpp b/Engine/Source/Runtime/Engine/Private/UnrealEngine.cpp
--- a/Engine/Source/Runtime/Engine/Private/UnrealEngine.cpp
+++ b/Engine/Source/Runtime/Engine/Private/UnrealEngine.cpp
@@ -3553,10 +3553,11 @@ struct FSortedTexture
 	FString		Name;
 	int32		LODGroup;
 	bool		bIsStreaming;
+	bool		bIsVirtual;
 	int32		UsageCount;
 
 	/** Constructor, initializing every member variable with passed in values. */
-	FSortedTexture(	int32 InMaxAllowedSizeX, int32 InMaxAllowedSizeY, EPixelFormat InFormat, int32 InCurSizeX, int32 InCurSizeY, int32 InLODBias, int32 InMaxAllowedSize, int32 InCurrentSize, const FString& InName, int32 InLODGroup, bool bInIsStreaming, int32 InUsageCount )
+	FSortedTexture(	int32 InMaxAllowedSizeX, int32 InMaxAllowedSizeY, EPixelFormat InFormat, int32 InCurSizeX, int32 InCurSizeY, int32 InLODBias, int32 InMaxAllowedSize, int32 InCurrentSize, const FString& InName, int32 InLODGroup, bool bInIsStreaming, bool bInIsVirtual, int32 InUsageCount )
 		:	MaxAllowedSizeX( InMaxAllowedSizeX )
 		,	MaxAllowedSizeY( InMaxAllowedSizeY )
 		,	Format( InFormat )
@@ -3568,6 +3569,7 @@ struct FSortedTexture
 		,	Name( InName )
 		,	LODGroup( InLODGroup )
 		,	bIsStreaming( bInIsStreaming )
+		,	bIsVirtual( bInIsVirtual )
 		,	UsageCount( InUsageCount )
 	{}
 };
@@ -5302,6 +5304,7 @@ bool UEngine::HandleListTexturesCommand( const TCHAR* Cmd, FOutputDevice& Ar )
 		int32				CurrentSize			= Texture->CalcTextureMemorySizeEnum(TMC_ResidentMips);
 		int32				UsageCount			= 0;
 		bool				bIsForced			= false;
+		bool				bIsVirtual			= false;
 
 		if (Texture2D != nullptr)
 		{
@@ -5316,6 +5319,7 @@ bool UEngine::HandleListTexturesCommand( const TCHAR* Cmd, FOutputDevice& Ar )
 			bIsStreamingTexture = Texture2D->GetStreamingIndex() != INDEX_NONE;
 			UsageCount			= TextureToUsageMap.FindRef(Texture2D);
 			bIsForced			= Texture2D->ShouldMipLevelsBeForcedResident() && bIsStreamingTexture;
+			bIsVirtual			= Texture2D->IsCurrentlyVirtualTextured();
 
 			if ((NumMips >= Texture2D->GetMinTextureResidentMipCount()) && bIsStreamingTexture)
 			{
@@ -5346,6 +5350,7 @@ bool UEngine::HandleListTexturesCommand( const TCHAR* Cmd, FOutputDevice& Ar )
 				Texture->GetPathName(), 
 				LODGroup, 
 				bIsStreamingTexture,
+				bIsVirtual,
 				UsageCount);
 		}
 	}
@@ -5374,15 +5379,15 @@ bool UEngine::HandleListTexturesCommand( const TCHAR* Cmd, FOutputDevice& Ar )
 
 	if (bCSV)
 	{
-		Ar.Logf(TEXT(",Max Width,Max Height,Max Size (KB),Bias Authored,Current Width,Current Height,Current Size (KB),Format,LODGroup,Name,Streaming,Usage Count"));
+		Ar.Logf(TEXT(",Max Width,Max Height,Max Size (KB),Bias Authored,Current Width,Current Height,Current Size (KB),Format,LODGroup,Name,Streaming,VT,Usage Count"));
 	}
 	else if (!FPlatformProperties::RequiresCookedData())
 	{
-		Ar.Logf(TEXT("MaxAllowedSize: Width x Height (Size in KB, Authored Bias), Current/InMem: Width x Height (Size in KB), Format, LODGroup, Name, Streaming, Usage Count"));
+		Ar.Logf(TEXT("MaxAllowedSize: Width x Height (Size in KB, Authored Bias), Current/InMem: Width x Height (Size in KB), Format, LODGroup, Name, Streaming, VT, Usage Count"));
 	}
 	else
 	{
-		Ar.Logf(TEXT("Cooked/OnDisk: Width x Height (Size in KB, Authored Bias), Current/InMem: Width x Height (Size in KB), Format, LODGroup, Name, Streaming, Usage Count"));
+		Ar.Logf(TEXT("Cooked/OnDisk: Width x Height (Size in KB, Authored Bias), Current/InMem: Width x Height (Size in KB), Format, LODGroup, Name, Streaming, VT, Usage Count"));
 	}
 
 	for( int32 TextureIndex=0; TextureIndex<SortedTextures.Num(); TextureIndex++ )
@@ -5399,7 +5404,7 @@ bool UEngine::HandleListTexturesCommand( const TCHAR* Cmd, FOutputDevice& Ar )
 
 		if (bCSV)
 		{
-			Ar.Logf(TEXT(",%i, %i, %i, %s, %i, %i, %i, %s, %s, %s, %s, %i"),
+			Ar.Logf(TEXT(",%i, %i, %i, %s, %i, %i, %i, %s, %s, %s, %s, %s, %i"),
 				SortedTexture.MaxAllowedSizeX, SortedTexture.MaxAllowedSizeY, (SortedTexture.MaxAllowedSize + 512) / 1024,
 				*AuthoredBiasString,
 				SortedTexture.CurSizeX, SortedTexture.CurSizeY, (SortedTexture.CurrentSize + 512) / 1024,
@@ -5407,11 +5412,12 @@ bool UEngine::HandleListTexturesCommand( const TCHAR* Cmd, FOutputDevice& Ar )
 				bValidTextureGroup ? *TextureGroupNames[SortedTexture.LODGroup] : TEXT("INVALID"),
 				*SortedTexture.Name,
 				SortedTexture.bIsStreaming ? TEXT("YES") : TEXT("NO"),
+				SortedTexture.bIsVirtual ? TEXT("YES") : TEXT("NO"),
 				SortedTexture.UsageCount);
 		}
 		else
 		{
-			Ar.Logf(TEXT("%ix%i (%i KB, %s), %ix%i (%i KB), %s, %s, %s, %s, %i"),
+			Ar.Logf(TEXT("%ix%i (%i KB, %s), %ix%i (%i KB), %s, %s, %s, %s, %s, %i"),
 				SortedTexture.MaxAllowedSizeX, SortedTexture.MaxAllowedSizeY, (SortedTexture.MaxAllowedSize + 512) / 1024,
 				*AuthoredBiasString,
 				SortedTexture.CurSizeX, SortedTexture.CurSizeY, (SortedTexture.CurrentSize + 512) / 1024,
@@ -5419,6 +5425,7 @@ bool UEngine::HandleListTexturesCommand( const TCHAR* Cmd, FOutputDevice& Ar )
 				bValidTextureGroup ? *TextureGroupNames[SortedTexture.LODGroup] : TEXT("INVALID"),
 				*SortedTexture.Name,
 				SortedTexture.bIsStreaming ? TEXT("YES") : TEXT("NO"),
+				SortedTexture.bIsVirtual ? TEXT("YES") : TEXT("NO"),
 				SortedTexture.UsageCount);
 		}
 
