CIS fix - bad merge of ChaosInterfaceUtils.cpp

diff --git a/Engine/Source/Runtime/Engine/Private/PhysicsEngine/Experimental/ChaosInterfaceUtils.cpp b/Engine/Source/Runtime/Engine/Private/PhysicsEngine/Experimental/ChaosInterfaceUtils.cpp
--- a/Engine/Source/Runtime/Engine/Private/PhysicsEngine/Experimental/ChaosInterfaceUtils.cpp
+++ b/Engine/Source/Runtime/Engine/Private/PhysicsEngine/Experimental/ChaosInterfaceUtils.cpp
@@ -480,12 +480,10 @@ namespace ChaosInterface
 		Chaos::FRotation3 RotationOfMass = Chaos::FRotation3::Identity;
 		if (MassPropertiesList.Num())
 		{
-			Tensor = Chaos::CombineWorldSpace(MassPropertiesList, InDensityKGPerCM).InertiaTensor;
 			// NOTE: If multiple items in the list, rotation of mass will be zero, but if only 1 item is the list the item is returned directly and we may have a rotation of mass
 			Chaos::FMassProperties CombinedMassProperties = Chaos::CombineWorldSpace<float, 3>(MassPropertiesList);
 			Tensor = CombinedMassProperties.InertiaTensor;
 			RotationOfMass = CombinedMassProperties.RotationOfMass;
-			Tensor = Chaos::CombineWorldSpace(MassPropertiesList, InDensityKGPerCM).InertiaTensor;
 		}
 		else
 		{
