CIS fix - bad fix of a bad merge of ChaosInterfaceUtils.cpp

diff --git a/Engine/Source/Runtime/Engine/Private/PhysicsEngine/Experimental/ChaosInterfaceUtils.cpp b/Engine/Source/Runtime/Engine/Private/PhysicsEngine/Experimental/ChaosInterfaceUtils.cpp
--- a/Engine/Source/Runtime/Engine/Private/PhysicsEngine/Experimental/ChaosInterfaceUtils.cpp
+++ b/Engine/Source/Runtime/Engine/Private/PhysicsEngine/Experimental/ChaosInterfaceUtils.cpp
@@ -481,7 +481,7 @@ namespace ChaosInterface
 		if (MassPropertiesList.Num())
 		{
 			// NOTE: If multiple items in the list, rotation of mass will be zero, but if only 1 item is the list the item is returned directly and we may have a rotation of mass
-			Chaos::FMassProperties CombinedMassProperties = Chaos::CombineWorldSpace<float, 3>(MassPropertiesList);
+			Chaos::FMassProperties CombinedMassProperties = Chaos::CombineWorldSpace(MassPropertiesList);
 			Tensor = CombinedMassProperties.InertiaTensor;
 			RotationOfMass = CombinedMassProperties.RotationOfMass;
 		}
