Reduce ShaderCompilingThread activity while waiting for compile results.

Sleep for 100ms if no results were processed during CompilingLoop iteration.

diff --git a/Engine/Source/Runtime/Engine/Private/ShaderCompiler/ShaderCompiler.cpp b/Engine/Source/Runtime/Engine/Private/ShaderCompiler/ShaderCompiler.cpp
--- a/Engine/Source/Runtime/Engine/Private/ShaderCompiler/ShaderCompiler.cpp
+++ b/Engine/Source/Runtime/Engine/Private/ShaderCompiler/ShaderCompiler.cpp
@@ -2129,8 +2129,10 @@ bool FShaderCompileThreadRunnable::LaunchWorkersIfNeeded()
 	return bAbandonWorkers;
 }
 
-void FShaderCompileThreadRunnable::ReadAvailableResults()
+int32 FShaderCompileThreadRunnable::ReadAvailableResults()
 {
+	int32 NumProcessed = 0;
+
 	for (int32 WorkerIndex = 0; WorkerIndex < WorkerInfos.Num(); WorkerIndex++)
 	{
 		FShaderCompileWorkerInfo& CurrentWorkerInfo = *WorkerInfos[WorkerIndex];
@@ -2174,9 +2176,13 @@ void FShaderCompileThreadRunnable::ReadAvailableResults()
 
 					CurrentWorkerInfo.bComplete = true;
 				}
+
+				++NumProcessed;
 			}
 		}
 	}
+
+	return NumProcessed;
 }
 
 void FShaderCompileThreadRunnable::CompileDirectlyThroughDll()
@@ -2410,7 +2416,12 @@ int32 FShaderCompileThreadRunnable::CompilingLoop()
 		else
 		{
 			// Read files which are outputs from the shader compile workers
-			ReadAvailableResults();
+			int32 NumProcessedResults = ReadAvailableResults();
+			if (NumProcessedResults == 0)
+			{
+				// Reduce filesystem query rate while actively waiting for results.
+				FPlatformProcess::Sleep(0.1f);
+			}
 		}
 	}
 	else
diff --git a/Engine/Source/Runtime/Engine/Public/ShaderCompiler.h b/Engine/Source/Runtime/Engine/Public/ShaderCompiler.h
--- a/Engine/Source/Runtime/Engine/Public/ShaderCompiler.h
+++ b/Engine/Source/Runtime/Engine/Public/ShaderCompiler.h
@@ -616,8 +616,8 @@ private:
 	/** Used when compiling through workers, launches worker processes if needed. */
 	bool LaunchWorkersIfNeeded();
 
-	/** Used when compiling through workers, attempts to open the worker output file if the worker is done and read the results. */
-	void ReadAvailableResults();
+	/** Used when compiling through workers, attempts to open the worker output file if the worker is done and read the results. Returns number of results processed. */
+	int32 ReadAvailableResults();
 
 	/** Used when compiling directly through the console tools dll. */
 	void CompileDirectlyThroughDll();
