Don't unregister the on editor close event for PIE sessions.

diff --git a/Engine/Source/Runtime/Engine/Private/WorldPartition/WorldPartition.cpp b/Engine/Source/Runtime/Engine/Private/WorldPartition/WorldPartition.cpp
--- a/Engine/Source/Runtime/Engine/Private/WorldPartition/WorldPartition.cpp
+++ b/Engine/Source/Runtime/Engine/Private/WorldPartition/WorldPartition.cpp
@@ -372,7 +372,7 @@ void UWorldPartition::Uninitialize()
 #if WITH_EDITOR
 		SavePerUserSettings();
 
-		if (GEditor)
+		if (GEditor && !World->IsGameWorld())
 		{
 			GEditor->OnEditorClose().RemoveAll(this);
 		}
@@ -995,7 +995,7 @@ void UWorldPartition::FinalizeGeneratedPackageForCook()
 
 void UWorldPartition::SavePerUserSettings()
 {
-	if (GIsEditor && !IsRunningCommandlet() && !IsEngineExitRequested())
+	if (GIsEditor && !World->IsGameWorld() && !IsRunningCommandlet() && !IsEngineExitRequested())
 	{
 		// Save last loaded cells settings
 		TArray<FName> EditorGridLastLoadedCells;
