Add BP editor functions to read an area of render target

diff --git a/Engine/Source/Runtime/Engine/Classes/Kismet/KismetRenderingLibrary.h b/Engine/Source/Runtime/Engine/Classes/Kismet/KismetRenderingLibrary.h
--- a/Engine/Source/Runtime/Engine/Classes/Kismet/KismetRenderingLibrary.h
+++ b/Engine/Source/Runtime/Engine/Classes/Kismet/KismetRenderingLibrary.h
@@ -120,11 +120,23 @@ class UKismetRenderingLibrary : public UBlueprintFunctionLibrary
 	static ENGINE_API FLinearColor ReadRenderTargetRawPixel(UObject* WorldContextObject, UTextureRenderTarget2D* TextureRenderTarget, int32 X, int32 Y, bool bNormalize = true);
 
 	/**
-	* Incredibly inefficient and slow operation! Read a value as-is color from a render target using UV [0,1]x[0,1] coordinates.
+    * Incredibly inefficient and slow operation! Read an area of values as-is from a render target using a rectangle defined by integer pixel coordinates.
+    */
+	UFUNCTION(BlueprintCallable, Category = "Rendering", meta = (Keywords = "ReadRenderTarget", WorldContext = "WorldContextObject"))
+	static ENGINE_API TArray<FLinearColor> ReadRenderTargetRawPixelArea(UObject* WorldContextObject, UTextureRenderTarget2D* TextureRenderTarget, int32 MinX, int32 MinY, int32 MaxX, int32 MaxY, bool bNormalize = true);
+
+	/**
+	* Incredibly inefficient and slow operation! Read a value as-is from a render target using UV [0,1]x[0,1] coordinates.
 	*/
 	UFUNCTION(BlueprintCallable, Category = "Rendering", meta = (Keywords = "ReadRenderTarget", WorldContext = "WorldContextObject"))
 	static ENGINE_API FLinearColor ReadRenderTargetRawUV(UObject* WorldContextObject, UTextureRenderTarget2D* TextureRenderTarget, float U, float V, bool bNormalize = true);
 
+	/**
+	* Incredibly inefficient and slow operation! Read an area of values as-is from a render target using a rectangle defined by UV [0,1]x[0,1] coordinates.
+	*/
+	UFUNCTION(BlueprintCallable, Category = "Rendering", meta = (Keywords = "ReadRenderTarget", WorldContext = "WorldContextObject"))
+	static ENGINE_API TArray<FLinearColor> ReadRenderTargetRawUVArea(UObject* WorldContextObject, UTextureRenderTarget2D* TextureRenderTarget, FBox2D Area, bool bNormalize = true);
+
 	/**
 	 * Exports a Texture2D as a HDR image onto the disk.
 	 */
diff --git a/Engine/Source/Runtime/Engine/Private/KismetRenderingLibrary.cpp b/Engine/Source/Runtime/Engine/Private/KismetRenderingLibrary.cpp
--- a/Engine/Source/Runtime/Engine/Private/KismetRenderingLibrary.cpp
+++ b/Engine/Source/Runtime/Engine/Private/KismetRenderingLibrary.cpp
@@ -385,6 +385,33 @@ FLinearColor UKismetRenderingLibrary::ReadRenderTargetRawPixel(UObject * WorldCo
 	}
 }
 
+ENGINE_API TArray<FLinearColor> UKismetRenderingLibrary::ReadRenderTargetRawPixelArea(UObject* WorldContextObject, UTextureRenderTarget2D* TextureRenderTarget, int32 MinX, int32 MinY, int32 MaxX, int32 MaxY, bool bNormalize /*= true*/)
+{
+	TArray<FColor> Samples;
+	TArray<FLinearColor> LinearSamples;
+
+	switch (ReadRenderTargetHelper(Samples, LinearSamples, WorldContextObject, TextureRenderTarget, MinX, MinY, MaxX, MaxY, bNormalize))
+	{
+	case PF_B8G8R8A8:
+		check(Samples.Num() > 0 && LinearSamples.Num() == 0);
+		LinearSamples.SetNum(Samples.Num());
+		for (int Idx = 0; Idx < Samples.Num(); Idx++)
+		{
+			LinearSamples[Idx] = FLinearColor(float(Samples[Idx].R), float(Samples[Idx].G), float(Samples[Idx].B), float(Samples[Idx].A));
+		}
+		return LinearSamples;
+
+	case PF_FloatRGBA:
+		check(Samples.Num() == 0 && LinearSamples.Num() > 0);
+		return LinearSamples;
+
+	case PF_Unknown:
+
+	default:
+		return TArray<FLinearColor>();
+	}
+}
+
 FLinearColor UKismetRenderingLibrary::ReadRenderTargetRawUV(UObject * WorldContextObject, UTextureRenderTarget2D * TextureRenderTarget, float U, float V, bool bNormalize)
 {
 	if (!TextureRenderTarget)
@@ -400,6 +427,22 @@ FLinearColor UKismetRenderingLibrary::ReadRenderTargetRawUV(UObject * WorldConte
 	return ReadRenderTargetRawPixel(WorldContextObject, TextureRenderTarget, XPos, YPos, bNormalize);
 }
 
+ENGINE_API TArray<FLinearColor> UKismetRenderingLibrary::ReadRenderTargetRawUVArea(UObject* WorldContextObject, UTextureRenderTarget2D* TextureRenderTarget, FBox2D Area, bool bNormalize /*= true*/)
+{
+
+	if (!TextureRenderTarget)
+	{
+		return TArray<FLinearColor>();
+	}
+
+	int32 MinX = FMath::Clamp(Area.Min.X, 0.f, 1.f) * (float)TextureRenderTarget->SizeX;
+	int32 MinY = FMath::Clamp(Area.Min.Y, 0.f, 1.f) * (float)TextureRenderTarget->SizeY;
+	int32 MaxX = FMath::Clamp(Area.Max.X, 0.f, 1.f) * (float)TextureRenderTarget->SizeX;
+	int32 MaxY = FMath::Clamp(Area.Max.Y, 0.f, 1.f) * (float)TextureRenderTarget->SizeY;
+
+	return ReadRenderTargetRawPixelArea(WorldContextObject, TextureRenderTarget, MinX, MinY, MaxX, MaxY, bNormalize);
+}
+
 /*
 
 void UKismetRenderingLibrary::CreateTexture2DFromRenderTarget(UObject* WorldContextObject, UTextureRenderTarget2D* RenderTarget, const FString &TextureAssetName)
