Remove the sorting of the per bone BlendSpace smoothing overrides in the UI. Sorts the runtime list since that is needed for the functionality. It also makes the per-bone override operate on the specified bone, not just its descendants

diff --git a/Engine/Source/Runtime/Engine/Classes/Animation/BlendSpace.h b/Engine/Source/Runtime/Engine/Classes/Animation/BlendSpace.h
--- a/Engine/Source/Runtime/Engine/Classes/Animation/BlendSpace.h
+++ b/Engine/Source/Runtime/Engine/Classes/Animation/BlendSpace.h
@@ -425,7 +425,7 @@ public:
 	* If nothing found, return INDEX_NONE
 	*/
 	virtual int32 GetPerBoneInterpolationIndex(int32 BoneIndex, const FBoneContainer& RequiredBones) const override;	
-	// End UBlendSpace Overrides
+	// End IInterpolationIndexProvider Overrides
 
 	/** Returns whether or not the given additive animation type is compatible with the blendspace type */
 	ENGINE_API virtual bool IsValidAdditiveType(EAdditiveAnimationType AdditiveType) const;
@@ -657,7 +657,7 @@ public:
 	 * If set then this eases in/out the sample weight adjustments, using the speed to determine how much smoothing to apply.
 	 */
 	UPROPERTY(EditAnywhere, Category = SampleInterpolation, meta = (DisplayName = "Sample Weight Speed Smoothing "))
-	bool bTargetWeightInterpolationEaseInOut;
+	bool bTargetWeightInterpolationEaseInOut = true;
 
 	/** The current mode used by the BlendSpace to decide which animation notifies to fire. Valid options are:
 	- AllAnimations - All notify events will fire
@@ -677,13 +677,21 @@ public:
 
 protected:
 
+#if WITH_EDITORONLY_DATA
 	/**
-	* Per bone interpolation speed settings. 
-	* These act as overrides to the global interpolation speed. 
-	* This means the global interpolation speed does not impact these bones.
+	* Per bone interpolation speed settings, which affect the specified bone and all its descendants in the skeleton.
+	* These act as overrides to the global interpolation speed, which means the global interpolation speed does not affect these bones.
+	* Note that they also override each other - so a per-bone setting on the chest will not affect the hand if there is a per-bone setting on the arm.
 	*/
 	UPROPERTY(EditAnywhere, Category = SampleInterpolation, meta = (DisplayName="Per Bone Overrides"))
 	TArray<FPerBoneInterpolation> PerBoneBlend;
+#endif
+
+	/**
+	 * PerBoneBlend sorted so that we find the bones in the correct order, working back from leaf bones.
+	 */
+	UPROPERTY()
+	TArray<FPerBoneInterpolation> SortedPerBoneBlend;
 
 	/** Track index to get marker data from. Samples are tested for the suitability of marker based sync
 	    during load and if we can use marker based sync we cache an index to a representative sample here */
diff --git a/Engine/Source/Runtime/Engine/Private/Animation/BlendSpace.cpp b/Engine/Source/Runtime/Engine/Private/Animation/BlendSpace.cpp
--- a/Engine/Source/Runtime/Engine/Private/Animation/BlendSpace.cpp
+++ b/Engine/Source/Runtime/Engine/Private/Animation/BlendSpace.cpp
@@ -83,6 +83,8 @@ void UBlendSpace::Serialize(FArchive& Ar)
 	if (Ar.IsLoading() && (Ar.CustomVer(FUE5MainStreamObjectVersion::GUID) < FUE5MainStreamObjectVersion::BlendSpaceRuntimeTriangulation))
 	{
 		ResampleData();
+		// Preserve the constant blend for old blend spaces, but allow the ease in/out default for new ones.
+		bTargetWeightInterpolationEaseInOut = false;
 	}
 #endif // WITH_EDITOR
 }
@@ -110,7 +112,9 @@ void UBlendSpace::PostEditChangeProperty(struct FPropertyChangedEvent& PropertyC
 	const FName MemberPropertyName = PropertyChangedEvent.MemberProperty ? PropertyChangedEvent.MemberProperty->GetFName() : NAME_None;
 	const FName PropertyName = PropertyChangedEvent.Property ? PropertyChangedEvent.Property->GetFName() : NAME_None;
 
-	if ((MemberPropertyName == GET_MEMBER_NAME_CHECKED(UBlendSpace, PerBoneBlend) && PropertyName == GET_MEMBER_NAME_CHECKED(FBoneReference, BoneName)) || PropertyName == GET_MEMBER_NAME_CHECKED(UBlendSpace, PerBoneBlend))
+	if ((MemberPropertyName == GET_MEMBER_NAME_CHECKED(UBlendSpace, PerBoneBlend)
+		 && PropertyName == GET_MEMBER_NAME_CHECKED(FBoneReference, BoneName)) 
+		|| PropertyName == GET_MEMBER_NAME_CHECKED(UBlendSpace, PerBoneBlend))
 	{
 		InitializePerBoneBlend();
 	}
@@ -202,7 +206,7 @@ bool UBlendSpace::UpdateBlendSamples_Internal(
 	if (GetSamplesFromBlendInput(InBlendSpacePosition, NewSampleDataList, InOutCachedTriangulationIndex, true))
 	{
 		// if target weight interpolation is set
-		if (TargetWeightInterpolationSpeedPerSec > 0.f || PerBoneBlend.Num() > 0)
+		if (TargetWeightInterpolationSpeedPerSec > 0.f || SortedPerBoneBlend.Num() > 0)
 		{
 			// target weight interpolation
 			if (InterpolateWeightOfSampleData(InDeltaTime, InOutOldSampleDataList, NewSampleDataList, InOutSampleDataCache))
@@ -618,19 +622,27 @@ void UBlendSpace::RuntimeValidateMarkerData()
 
 #endif // WITH_EDITOR
 
-// @todo fixme: slow approach. If the perbone gets popular, we should change this to array of weight 
-// we should think about changing this to 
+// Note that this gets called once per bone in the skeleton - so any excessive looping in here is
+// expensive. See UE-110576
 int32 UBlendSpace::GetPerBoneInterpolationIndex(int32 BoneIndex, const FBoneContainer& RequiredBones) const
 {
-	for (int32 Iter = 0; Iter < PerBoneBlend.Num(); ++Iter)
+	for (int32 Iter = 0; Iter < SortedPerBoneBlend.Num(); ++Iter)
 	{
-		// we would like to make sure if 
-		if (PerBoneBlend[Iter].BoneReference.IsValidToEvaluate(RequiredBones) && RequiredBones.BoneIsChildOf(BoneIndex, RequiredBones.GetCompactPoseIndexFromSkeletonIndex(PerBoneBlend[Iter].BoneReference.BoneIndex).GetInt()))
+		const FBoneReference& SmoothedBone = SortedPerBoneBlend[Iter].BoneReference;
+		if (SmoothedBone.IsValidToEvaluate(RequiredBones))
 		{
-			return Iter;
+			int SmoothedBonePoseIndex = RequiredBones.GetCompactPoseIndexFromSkeletonIndex(SmoothedBone.BoneIndex).GetInt();
+			if (SmoothedBonePoseIndex == BoneIndex)
+			{
+				return Iter;
+			}
+			// Returns true if BoneIndex is a child of SmoothedBonePoseIndex. Searches up from BoneIndex
+			if (RequiredBones.BoneIsChildOf(BoneIndex, SmoothedBonePoseIndex))
+			{
+				return Iter;
+			}
 		}
 	}
-
 	return INDEX_NONE;
 }
 
@@ -756,7 +768,7 @@ void UBlendSpace::GetAnimationPose_Internal(TArray<FBlendSampleData>& BlendSampl
 
 	TArrayView<FCompactPose> ChildrenPosesView(ChildrenPoses);
 
-	if (PerBoneBlend.Num() > 0)
+	if (SortedPerBoneBlend.Num() > 0)
 	{
 		if (IsValidAdditive())
 		{
@@ -1304,12 +1316,24 @@ bool UBlendSpace::IsTooCloseToExistingSamplePoint(const FVector& SampleValue, in
 void UBlendSpace::InitializePerBoneBlend()
 {
 	const USkeleton* MySkeleton = GetSkeleton();
+#if WITH_EDITOR
 	for (FPerBoneInterpolation& BoneInterpolationData : PerBoneBlend)
 	{
 		BoneInterpolationData.Initialize(MySkeleton);
 	}
-	// Sort this by bigger to smaller, then we don't have to worry about checking the best parent
-	PerBoneBlend.Sort([](const FPerBoneInterpolation& A, const FPerBoneInterpolation& B) { return A.BoneReference.BoneIndex > B.BoneReference.BoneIndex; });
+	SortedPerBoneBlend = PerBoneBlend;
+#endif
+	for (FPerBoneInterpolation& BoneInterpolationData : SortedPerBoneBlend)
+	{
+		BoneInterpolationData.Initialize(MySkeleton);
+	}
+#if WITH_EDITOR
+	// Sort this by bigger to smaller. This is because the per-bone blends are intended to affect the
+	// children of as well as the specified bones, and we need to find them (in
+	// GetPerBoneInterpolationIndex) in order working back from the leaf bones.
+	SortedPerBoneBlend.Sort([](const FPerBoneInterpolation& A, const FPerBoneInterpolation& B) 
+		{ return A.BoneReference.BoneIndex > B.BoneReference.BoneIndex; });
+#endif
 }
 
 void UBlendSpace::TickFollowerSamples(TArray<FBlendSampleData>& SampleDataList, const int32 HighestWeightIndex, FAnimAssetTickContext& Context, bool bResetMarkerDataOnFollowers) const
@@ -1467,10 +1491,10 @@ bool UBlendSpace::InterpolateWeightOfSampleData(float DeltaTime, const TArray<FB
 		FBlendSampleData OldSample = *OldIt;
 		bool bTargetSampleExists = false;
 
-		if (OldSample.PerBoneBlendData.Num() != PerBoneBlend.Num())
+		if (OldSample.PerBoneBlendData.Num() != SortedPerBoneBlend.Num())
 		{
-			OldSample.PerBoneBlendData.Init(OldSample.TotalWeight, PerBoneBlend.Num());
-			OldSample.PerBoneWeightRate.Init(OldSample.WeightRate, PerBoneBlend.Num());
+			OldSample.PerBoneBlendData.Init(OldSample.TotalWeight, SortedPerBoneBlend.Num());
+			OldSample.PerBoneWeightRate.Init(OldSample.WeightRate, SortedPerBoneBlend.Num());
 		}
 
 		for (auto NewIt = NewSampleDataList.CreateConstIterator(); NewIt; ++NewIt)
@@ -1491,7 +1515,7 @@ bool UBlendSpace::InterpolateWeightOfSampleData(float DeltaTime, const TArray<FB
 					SmoothWeight(
 						InterpData.PerBoneBlendData[Iter], InterpData.PerBoneWeightRate[Iter],
 						OldSample.PerBoneBlendData[Iter], OldSample.PerBoneWeightRate[Iter], NewSample.TotalWeight,
-						DeltaTime, PerBoneBlend[Iter].InterpolationSpeedPerSec, bTargetWeightInterpolationEaseInOut);
+						DeltaTime, SortedPerBoneBlend[Iter].InterpolationSpeedPerSec, bTargetWeightInterpolationEaseInOut);
 					TotalPerBoneWeight += InterpData.PerBoneBlendData[Iter];
 				}
 
@@ -1518,7 +1542,7 @@ bool UBlendSpace::InterpolateWeightOfSampleData(float DeltaTime, const TArray<FB
 				SmoothWeight(
 					InterpData.PerBoneBlendData[Iter], InterpData.PerBoneWeightRate[Iter],
 					OldSample.PerBoneBlendData[Iter], OldSample.PerBoneWeightRate[Iter], 0.0f,
-					DeltaTime, PerBoneBlend[Iter].InterpolationSpeedPerSec, bTargetWeightInterpolationEaseInOut);
+					DeltaTime, SortedPerBoneBlend[Iter].InterpolationSpeedPerSec, bTargetWeightInterpolationEaseInOut);
 				TotalPerBoneWeight += InterpData.PerBoneBlendData[Iter];
 			}
 
@@ -1539,10 +1563,10 @@ bool UBlendSpace::InterpolateWeightOfSampleData(float DeltaTime, const TArray<FB
 		FBlendSampleData OldSample = *OldIt;
 		bool bOldSampleExists = false;
 
-		if (OldSample.PerBoneBlendData.Num() != PerBoneBlend.Num())
+		if (OldSample.PerBoneBlendData.Num() != SortedPerBoneBlend.Num())
 		{
-			OldSample.PerBoneBlendData.Init(OldSample.TotalWeight, PerBoneBlend.Num());
-			OldSample.PerBoneWeightRate.Init(OldSample.WeightRate, PerBoneBlend.Num());
+			OldSample.PerBoneBlendData.Init(OldSample.TotalWeight, SortedPerBoneBlend.Num());
+			OldSample.PerBoneWeightRate.Init(OldSample.WeightRate, SortedPerBoneBlend.Num());
 		}
 
 		for (auto NewIt = FinalSampleDataList.CreateConstIterator(); NewIt; ++NewIt)
@@ -1573,7 +1597,7 @@ bool UBlendSpace::InterpolateWeightOfSampleData(float DeltaTime, const TArray<FB
 				SmoothWeight(
 					InterpData.PerBoneBlendData[Iter], InterpData.PerBoneWeightRate[Iter],
 					OldSample.PerBoneBlendData[Iter], OldSample.PerBoneWeightRate[Iter], Target,
-					DeltaTime, PerBoneBlend[Iter].InterpolationSpeedPerSec, bTargetWeightInterpolationEaseInOut);
+					DeltaTime, SortedPerBoneBlend[Iter].InterpolationSpeedPerSec, bTargetWeightInterpolationEaseInOut);
 				TotalPerBoneWeight += InterpData.PerBoneBlendData[Iter];
 			}
 			if (InterpData.TotalWeight > ZERO_ANIMWEIGHT_THRESH || TotalPerBoneWeight > ZERO_ANIMWEIGHT_THRESH)
