Fix crash where an actor cannot be loaded in a partitioned world in PIE.
Not sure yet why this happens, but for now just log an error instead of crashing.

diff --git a/Engine/Source/Runtime/Engine/Private/WorldPartition/WorldPartitionLevelHelper.cpp b/Engine/Source/Runtime/Engine/Private/WorldPartition/WorldPartitionLevelHelper.cpp
--- a/Engine/Source/Runtime/Engine/Private/WorldPartition/WorldPartitionLevelHelper.cpp
+++ b/Engine/Source/Runtime/Engine/Private/WorldPartition/WorldPartitionLevelHelper.cpp
@@ -164,35 +164,35 @@ bool FWorldPartitionLevelHelper::LoadActors(ULevel* InDestLevel, TArrayView<FWor
 		FName ActorName = *FPaths::GetExtension(PackageObjectMapping.Path.ToString());
 
 		FLoadPackageAsyncDelegate CompletionCallback = FLoadPackageAsyncDelegate::CreateLambda([LoadProgress, ActorName, bLoadForPlay, InDestLevel, InCompletionCallback](const FName& LoadedPackageName, UPackage* LoadedPackage, EAsyncLoadingResult::Type Result)
-			{
-				check(LoadProgress->NumPendingLoadRequests);
-				LoadProgress->NumPendingLoadRequests--;
-				if (LoadedPackage)
-				{
-					AActor* Actor = FindObject<AActor>(LoadedPackage, *ActorName.ToString());
+		{
+			check(LoadProgress->NumPendingLoadRequests);
+			LoadProgress->NumPendingLoadRequests--;
 
-					check(Actor);
-					if (bLoadForPlay)
-					{
-						check(Actor->IsPackageExternal());
-						InDestLevel->Actors.Add(Actor);
-						check(Actor->GetLevel() == InDestLevel);
-					}
-								
-					UE_LOG(LogEngine, Verbose, TEXT(" ==> Loaded %s (remaining: %d)"), *Actor->GetFullName(), LoadProgress->NumPendingLoadRequests);
-				}
-				else
-				{
-					UE_LOG(LogEngine, Warning, TEXT("Failed to load %s"), *LoadedPackageName.ToString());
-					//@todo_ow: cumulate and process when NumPendingActorRequests == 0
-					LoadProgress->NumFailedLoadedRequests++;
-				}
+			AActor* Actor = LoadedPackage ? FindObject<AActor>(LoadedPackage, *ActorName.ToString()) : nullptr;
 
-				if (!LoadProgress->NumPendingLoadRequests)
+			if (Actor)
+			{
+				if (bLoadForPlay)
 				{
-					InCompletionCallback(!LoadProgress->NumFailedLoadedRequests);
+					check(Actor->IsPackageExternal());
+					InDestLevel->Actors.Add(Actor);
+					check(Actor->GetLevel() == InDestLevel);
 				}
-			});
+								
+				UE_LOG(LogEngine, Verbose, TEXT(" ==> Loaded %s (remaining: %d)"), *Actor->GetFullName(), LoadProgress->NumPendingLoadRequests);
+			}
+			else
+			{
+				UE_LOG(LogEngine, Warning, TEXT("Failed to load %s"), *LoadedPackageName.ToString());
+				//@todo_ow: cumulate and process when NumPendingActorRequests == 0
+				LoadProgress->NumFailedLoadedRequests++;
+			}
+
+			if (!LoadProgress->NumPendingLoadRequests)
+			{
+				InCompletionCallback(!LoadProgress->NumFailedLoadedRequests);
+			}
+		});
 
 
 		FPackagePath PackagePath = FPackagePath::FromPackageNameChecked(*ActorPackageName);
