Fix transition error that happens in the editor if the Instance Data buffer is used before any culling passes has written to it.
- also rename the resources to follow new namespacing convention

diff --git a/Engine/Source/Runtime/Engine/Private/SceneManagement.cpp b/Engine/Source/Runtime/Engine/Private/SceneManagement.cpp
--- a/Engine/Source/Runtime/Engine/Private/SceneManagement.cpp
+++ b/Engine/Source/Runtime/Engine/Private/SceneManagement.cpp
@@ -1304,4 +1304,19 @@ void FDefaultMobileReflectionCaptureUniformBuffer::InitDynamicRHI()
 /** Global uniform buffer containing the default reflection data used in mobile renderer. */
 TGlobalResource<FDefaultMobileReflectionCaptureUniformBuffer> GDefaultMobileReflectionCaptureUniformBuffer;
 
+
+
+void FInstanceCullingManagerResources::InitRHI()
+{
+	InstanceIdsBuffer.Initialize(TEXT("InstanceCulling.InstanceIds"), sizeof(uint32), MaxIndirectInstances, EPixelFormat::PF_R32_UINT, ERHIAccess::SRVGraphics, BUF_Static);
+	PageInfoBuffer.Initialize(TEXT("InstanceCulling.PageInfoBuffer"), sizeof(uint32), MaxIndirectInstances, EPixelFormat::PF_R32_UINT, ERHIAccess::SRVGraphics, BUF_Static);
+}
+
+void FInstanceCullingManagerResources::ReleaseRHI()
+{
+	InstanceIdsBuffer.Release();
+	PageInfoBuffer.Release();
+}
+
+
 TGlobalResource<FInstanceCullingManagerResources> GInstanceCullingManagerResources;
diff --git a/Engine/Source/Runtime/Engine/Public/SceneManagement.h b/Engine/Source/Runtime/Engine/Public/SceneManagement.h
--- a/Engine/Source/Runtime/Engine/Public/SceneManagement.h
+++ b/Engine/Source/Runtime/Engine/Public/SceneManagement.h
@@ -2185,17 +2185,8 @@ public:
 	 */
 	static constexpr uint32 MaxIndirectInstances = 1024 * 1024;
 
-	virtual void InitRHI() override
-	{
-		InstanceIdsBuffer.Initialize(TEXT("FInstanceCullingManagerResources::InstanceIds"), sizeof(uint32), MaxIndirectInstances, EPixelFormat::PF_R32_UINT, BUF_Static);
-		PageInfoBuffer.Initialize(TEXT("FInstanceCullingManagerResources::PageInfoBuffer"), sizeof(uint32), MaxIndirectInstances, EPixelFormat::PF_R32_UINT, BUF_Static);
-	}
-
-	virtual void ReleaseRHI() override
-	{
-		InstanceIdsBuffer.Release();
-		PageInfoBuffer.Release();
-	}
+	virtual void InitRHI() override;
+	virtual void ReleaseRHI() override;
 
 	FRHIBuffer* GetInstancesIdBuffer() const { return InstanceIdsBuffer.Buffer; }
 
