If you want to use these expressions in a separate module, there is currently a linker error. This change fixes that.

diff --git a/Engine/Source/Runtime/Engine/Classes/Materials/MaterialExpressionDDX.h b/Engine/Source/Runtime/Engine/Classes/Materials/MaterialExpressionDDX.h
--- a/Engine/Source/Runtime/Engine/Classes/Materials/MaterialExpressionDDX.h
+++ b/Engine/Source/Runtime/Engine/Classes/Materials/MaterialExpressionDDX.h
@@ -9,7 +9,7 @@
 #include "Materials/MaterialExpression.h"
 #include "MaterialExpressionDDX.generated.h"
 
-UCLASS()
+UCLASS(MinimalAPI)
 class UMaterialExpressionDDX : public UMaterialExpression
 {
 	GENERATED_UCLASS_BODY()
@@ -27,5 +27,3 @@ class UMaterialExpressionDDX : public UMaterialExpression
 	//~ End UMaterialExpression Interface
 };
 
-
-
diff --git a/Engine/Source/Runtime/Engine/Classes/Materials/MaterialExpressionDDY.h b/Engine/Source/Runtime/Engine/Classes/Materials/MaterialExpressionDDY.h
--- a/Engine/Source/Runtime/Engine/Classes/Materials/MaterialExpressionDDY.h
+++ b/Engine/Source/Runtime/Engine/Classes/Materials/MaterialExpressionDDY.h
@@ -9,7 +9,7 @@
 #include "Materials/MaterialExpression.h"
 #include "MaterialExpressionDDY.generated.h"
 
-UCLASS()
+UCLASS(MinimalAPI)
 class UMaterialExpressionDDY : public UMaterialExpression
 {
 	GENERATED_UCLASS_BODY()
