Static compiler warning fixes.

diff --git a/Engine/Source/Runtime/Engine/Private/DataChannel.cpp b/Engine/Source/Runtime/Engine/Private/DataChannel.cpp
--- a/Engine/Source/Runtime/Engine/Private/DataChannel.cpp
+++ b/Engine/Source/Runtime/Engine/Private/DataChannel.cpp
@@ -618,22 +618,26 @@ bool UChannel::ReceivedNextBunch( FInBunch & Bunch, bool & bOutSkipAck )
 			}
 
 			InPartialBunch = new FInBunch(Bunch, false);
-			if ( !Bunch.bHasPackageMapExports && Bunch.GetBitsLeft() > 0 )
-			{
-				if ( Bunch.GetBitsLeft() % 8 != 0 )
+			
+			if ( InPartialBunch!=nullptr )
+			{		
+				if ( !Bunch.bHasPackageMapExports && Bunch.GetBitsLeft() > 0 )
 				{
-					UE_LOG(LogNetPartialBunch, Warning, TEXT("Corrupt partial bunch. Initial partial bunches are expected to be byte-aligned. BitsLeft = %u. %s"), Bunch.GetBitsLeft(), *Describe());
-					Bunch.SetError();
-					return false;
-				}
+					if ( Bunch.GetBitsLeft() % 8 != 0 )
+					{
+						UE_LOG(LogNetPartialBunch, Warning, TEXT("Corrupt partial bunch. Initial partial bunches are expected to be byte-aligned. BitsLeft = %u. %s"), Bunch.GetBitsLeft(), *Describe());
+						Bunch.SetError();
+						return false;
+					}
 
-				InPartialBunch->AppendDataFromChecked( Bunch.GetDataPosChecked(), Bunch.GetBitsLeft() );
+					InPartialBunch->AppendDataFromChecked( Bunch.GetDataPosChecked(), Bunch.GetBitsLeft() );
 
-				LogPartialBunch(TEXT("Received new partial bunch."), Bunch, *InPartialBunch);
-			}
-			else
-			{
-				LogPartialBunch(TEXT("Received New partial bunch. It only contained NetGUIDs."), Bunch, *InPartialBunch);
+					LogPartialBunch(TEXT("Received new partial bunch."), Bunch, *InPartialBunch);
+				}
+				else
+				{
+					LogPartialBunch(TEXT("Received New partial bunch. It only contained NetGUIDs."), Bunch, *InPartialBunch);
+				}
 			}
 		}
 		else
@@ -1332,9 +1336,12 @@ void UChannel::ReceivedNak( int32 NakPacketId )
 				// Inject trace event for the resent bunch if tracing is enabled
 				// The reason behind the complexity is that the outgoing sendbuffer migth be flushed during the call to SendRawBunch()
 				FNetTraceCollector* TempCollector = UE_NET_TRACE_CREATE_COLLECTOR(ENetTraceVerbosity::Trace);
-				UE_NET_TRACE(ResendBunch, TempCollector, 0U, Out->GetNumBits(), ENetTraceVerbosity::Trace);
-				Connection->SendRawBunch(*Out, 0, TempCollector);
-				UE_NET_TRACE_DESTROY_COLLECTOR(TempCollector);
+				if (TempCollector!=nullptr)
+				{				
+					UE_NET_TRACE(ResendBunch, TempCollector, 0U, Out->GetNumBits(), ENetTraceVerbosity::Trace);
+					Connection->SendRawBunch(*Out, 0, TempCollector);
+					UE_NET_TRACE_DESTROY_COLLECTOR(TempCollector);
+				}
 			}
 			else
 			{
