Make world partition HLODs work in -game.

diff --git a/Engine/Source/Runtime/Engine/Private/WorldPartition/ActorDescContainer.cpp b/Engine/Source/Runtime/Engine/Private/WorldPartition/ActorDescContainer.cpp
--- a/Engine/Source/Runtime/Engine/Private/WorldPartition/ActorDescContainer.cpp
+++ b/Engine/Source/Runtime/Engine/Private/WorldPartition/ActorDescContainer.cpp
@@ -125,12 +125,7 @@ TUniquePtr<FWorldPartitionActorDesc> UActorDescContainer::GetActorDescriptor(con
 			NewActorDesc->Init(this, ActorDescInitData);
 			
 			check(bIsValidClass || (NewActorDesc->GetActorIsEditorOnly() && IsRunningGame()));
-			if (!NewActorDesc->GetActorIsEditorOnly() || !IsRunningGame())
-			{				
-				return NewActorDesc;
-			}
-
-			return nullptr;
+			return NewActorDesc;
 		}
 	}
 
diff --git a/Engine/Source/Runtime/Engine/Private/WorldPartition/WorldPartitionRuntimeHash.cpp b/Engine/Source/Runtime/Engine/Private/WorldPartition/WorldPartitionRuntimeHash.cpp
--- a/Engine/Source/Runtime/Engine/Private/WorldPartition/WorldPartitionRuntimeHash.cpp
+++ b/Engine/Source/Runtime/Engine/Private/WorldPartition/WorldPartitionRuntimeHash.cpp
@@ -39,12 +39,9 @@ bool UWorldPartitionRuntimeHash::GenerateRuntimeStreaming(EWorldPartitionStreami
 
 void UWorldPartitionRuntimeHash::OnBeginPlay(EWorldPartitionStreamingMode Mode)
 {
-	if (Mode == EWorldPartitionStreamingMode::PIE)
-	{
-		// Mark always loaded actors so that the Level will force reference to these actors for PIE.
-		// These actor will then be duplicated for PIE during the PIE world duplication process
-		ForceExternalActorLevelReference(/*bForceExternalActorLevelReferenceForPIE*/true);
-	}
+	// Mark always loaded actors so that the Level will force reference to these actors for PIE.
+	// These actor will then be duplicated for PIE during the PIE world duplication process
+	ForceExternalActorLevelReference(/*bForceExternalActorLevelReferenceForPIE*/true);
 }
 
 void UWorldPartitionRuntimeHash::OnEndPlay()
@@ -60,11 +57,14 @@ void UWorldPartitionRuntimeHash::OnEndPlay()
 
 void UWorldPartitionRuntimeHash::ForceExternalActorLevelReference(bool bForceExternalActorLevelReferenceForPIE)
 {
-	for (const FAlwaysLoadedActorForPIE& AlwaysLoadedActor : AlwaysLoadedActorsForPIE)
+	if (!IsRunningGame())
 	{
-		if (AActor* Actor = AlwaysLoadedActor.Actor)
+		for (const FAlwaysLoadedActorForPIE& AlwaysLoadedActor : AlwaysLoadedActorsForPIE)
 		{
-			Actor->SetForceExternalActorLevelReferenceForPIE(bForceExternalActorLevelReferenceForPIE);
+			if (AActor* Actor = AlwaysLoadedActor.Actor)
+			{
+				Actor->SetForceExternalActorLevelReferenceForPIE(bForceExternalActorLevelReferenceForPIE);
+			}
 		}
 	}
 }
diff --git a/Engine/Source/Runtime/Engine/Private/WorldPartition/WorldPartitionRuntimeSpatialHash.cpp b/Engine/Source/Runtime/Engine/Private/WorldPartition/WorldPartitionRuntimeSpatialHash.cpp
--- a/Engine/Source/Runtime/Engine/Private/WorldPartition/WorldPartitionRuntimeSpatialHash.cpp
+++ b/Engine/Source/Runtime/Engine/Private/WorldPartition/WorldPartitionRuntimeSpatialHash.cpp
@@ -456,21 +456,31 @@ bool UWorldPartitionRuntimeSpatialHash::GenerateStreaming(EWorldPartitionStreami
 		return false;
 	}
 
+	// Fix case where StreamingGrids might have been persisted.
+	StreamingGrids.Empty();
+
+	// Create actor clusters
+	FActorClusterContext Context(WorldPartition, this);
+
+	// Append grids from ASpatialHashRuntimeGridInfo actors to runtime spatial hash grids
 	TArray<FSpatialHashRuntimeGrid> AllGrids;
 	AllGrids.Append(Grids);
 
-	// Append grids from runtime grid actors
-	for (AActor* Actor : GetWorld()->PersistentLevel->Actors)
+	FActorContainerInstance* ContainerInstance = Context.GetClusterInstance(WorldPartition);
+	check(ContainerInstance);
+	for (auto& ActorDescViewPair : ContainerInstance->ActorDescViewMap)
 	{
-		if (ASpatialHashRuntimeGridInfo* RuntimeGridActor = Cast<ASpatialHashRuntimeGridInfo>(Actor))
+		FWorldPartitionActorDescView& ActorDescView = ActorDescViewPair.Value;
+		if (ActorDescView.GetActorClass()->IsChildOf<ASpatialHashRuntimeGridInfo>())
 		{
-			AllGrids.Add(RuntimeGridActor->GridSettings);
+			FWorldPartitionReference Ref(WorldPartition, ActorDescView.GetGuid());
+			if (ASpatialHashRuntimeGridInfo* RuntimeGridActor = Cast<ASpatialHashRuntimeGridInfo>(Ref->GetActor()))
+			{
+				AllGrids.Add(RuntimeGridActor->GridSettings);
+			}
 		}
 	}
 
-	// Fix case where StreamingGrids might have been persisted.
-	StreamingGrids.Empty();
-
 	TMap<FName, int32> GridsMapping;
 	GridsMapping.Add(NAME_None, 0);
 	for (int32 i = 0; i < AllGrids.Num(); i++)
@@ -480,9 +490,6 @@ bool UWorldPartitionRuntimeSpatialHash::GenerateStreaming(EWorldPartitionStreami
 		GridsMapping.Add(Grid.GridName, i);
 	}
 
-	// Create actor clusters
-	FActorClusterContext Context(WorldPartition, this);
-
 	TArray<TArray<const FActorClusterInstance*>> GridActors;
 	GridActors.InsertDefaulted(0, AllGrids.Num());
 
