Added extra component properties for describing a run-time hierarchy and culling priority. This will be used to drive Lumen/ray tracing LOD and culling for huge and complex scenes. E.g. merge components into buildings or more quickly cull less important ones.

Naming was highly controversial, everyone was proposing a different name, and current naming scheme was picked as a lesser of multiple evils.

diff --git a/Engine/Source/Runtime/Engine/Classes/Components/PrimitiveComponent.h b/Engine/Source/Runtime/Engine/Classes/Components/PrimitiveComponent.h
--- a/Engine/Source/Runtime/Engine/Classes/Components/PrimitiveComponent.h
+++ b/Engine/Source/Runtime/Engine/Classes/Components/PrimitiveComponent.h
@@ -105,6 +105,21 @@ enum class ERendererStencilMask : uint8
 	ERSM_128 UMETA(DisplayName = "Eighth bit (128), ignore depth")
 };
 
+/** How quickly component should be culled. */
+UENUM()
+enum class ERayTracingGroupCullingPriority : uint8
+{
+	CP_0_NEVER_CULL UMETA(DisplayName = "0 - Never cull"),
+	CP_1 UMETA(DisplayName = "1"),
+	CP_2 UMETA(DisplayName = "2"),
+	CP_3 UMETA(DisplayName = "3"),
+	CP_4_DEFAULT UMETA(DisplayName = "4 - Default"),
+	CP_5 UMETA(DisplayName = "5"),
+	CP_6 UMETA(DisplayName = "6"),
+	CP_7 UMETA(DisplayName = "7"),
+	CP_8_QUICKLY_CULL UMETA(DisplayName = "8 - Quickly cull")
+};
+
 /** Converts a stencil mask from the editor's USTRUCT version to the version the renderer uses. */
 struct FRendererStencilMaskEvaluation
 {
@@ -577,6 +592,19 @@ public:
 	UPROPERTY(EditAnywhere, AdvancedDisplay, BlueprintReadOnly, Category=Lighting)
 	FLightingChannels LightingChannels;
 
+	/**
+	 * Defines run-time groups of components. For example allows to assemble multiple parts of a building at runtime.
+	 * -1 means that component doesn't belong to any group.
+	 */
+	UPROPERTY(EditAnywhere, AdvancedDisplay, BlueprintReadOnly, Category = RayTracing)
+	int32 RayTracingGroupId;
+
+	/**
+	 * Defines how quickly it should be culled. For example buildings should have a low priority, but small dressing should have a high priority.
+	 */
+	UPROPERTY(EditAnywhere, AdvancedDisplay, BlueprintReadOnly, Category = RayTracing)
+	ERayTracingGroupCullingPriority RayTracingGroupCullingPriority;
+
 	/** Mask used for stencil buffer writes. */
 	UPROPERTY(EditAnywhere, AdvancedDisplay, BlueprintReadOnly, Category = "Rendering", meta = (editcondition = "bRenderCustomDepth"))
 	ERendererStencilMask CustomDepthStencilWriteMask;
diff --git a/Engine/Source/Runtime/Engine/Private/Components/PrimitiveComponent.cpp b/Engine/Source/Runtime/Engine/Private/Components/PrimitiveComponent.cpp
--- a/Engine/Source/Runtime/Engine/Private/Components/PrimitiveComponent.cpp
+++ b/Engine/Source/Runtime/Engine/Private/Components/PrimitiveComponent.cpp
@@ -350,6 +350,8 @@ UPrimitiveComponent::UPrimitiveComponent(const FObjectInitializer& ObjectInitial
 	ComponentId.PrimIDValue = NextComponentId.Increment();
 	CustomDepthStencilValue = 0;
 	CustomDepthStencilWriteMask = ERendererStencilMask::ERSM_Default;
+	RayTracingGroupId = -1;
+	RayTracingGroupCullingPriority = ERayTracingGroupCullingPriority::CP_4_DEFAULT;
 
 	LDMaxDrawDistance = 0.f;
 	CachedMaxDrawDistance = 0.f;
diff --git a/Engine/Source/Runtime/Engine/Private/ISMPartition/ISMComponentDescriptor.cpp b/Engine/Source/Runtime/Engine/Private/ISMPartition/ISMComponentDescriptor.cpp
--- a/Engine/Source/Runtime/Engine/Private/ISMPartition/ISMComponentDescriptor.cpp
+++ b/Engine/Source/Runtime/Engine/Private/ISMPartition/ISMComponentDescriptor.cpp
@@ -41,6 +41,8 @@ void FISMComponentDescriptor::InitFrom(const UStaticMeshComponent* Template, boo
 	VirtualTextureRenderPassType = Template->VirtualTextureRenderPassType;
 	LightmapType = Template->LightmapType;
 	LightingChannels = Template->LightingChannels;
+	RayTracingGroupId = Template->RayTracingGroupId;
+	RayTracingGroupCullingPriority = Template->RayTracingGroupCullingPriority;
 	bHasCustomNavigableGeometry = Template->bHasCustomNavigableGeometry;
 	CustomDepthStencilWriteMask = Template->CustomDepthStencilWriteMask;
 	VirtualTextureCullMips = Template->VirtualTextureCullMips;
@@ -97,6 +99,8 @@ bool FISMComponentDescriptor::operator==(const FISMComponentDescriptor& Other) c
 	VirtualTextureRenderPassType == Other.VirtualTextureRenderPassType &&
 	LightmapType == Other.LightmapType &&
 	GetLightingChannelMaskForStruct(LightingChannels) == GetLightingChannelMaskForStruct(Other.LightingChannels) &&
+	RayTracingGroupId == Other.RayTracingGroupId &&
+	RayTracingGroupCullingPriority == Other.RayTracingGroupCullingPriority &&
 	bHasCustomNavigableGeometry == Other.bHasCustomNavigableGeometry &&
 	CustomDepthStencilWriteMask == Other.CustomDepthStencilWriteMask &&
 	InstanceStartCullDistance == Other.InstanceStartCullDistance &&
@@ -156,6 +160,8 @@ void FISMComponentDescriptor::InitComponent(UInstancedStaticMeshComponent* ISMCo
 	ISMComponent->VirtualTextureRenderPassType = VirtualTextureRenderPassType;
 	ISMComponent->LightmapType = LightmapType;
 	ISMComponent->LightingChannels = LightingChannels;
+	ISMComponent->RayTracingGroupId = RayTracingGroupId;
+	ISMComponent->RayTracingGroupCullingPriority = RayTracingGroupCullingPriority;
 	ISMComponent->bHasCustomNavigableGeometry = bHasCustomNavigableGeometry;
 	ISMComponent->CustomDepthStencilWriteMask = CustomDepthStencilWriteMask;
 	ISMComponent->BodyInstance.CopyBodyInstancePropertiesFrom(&BodyInstance);
diff --git a/Engine/Source/Runtime/Engine/Private/PrimitiveSceneProxy.cpp b/Engine/Source/Runtime/Engine/Private/PrimitiveSceneProxy.cpp
--- a/Engine/Source/Runtime/Engine/Private/PrimitiveSceneProxy.cpp
+++ b/Engine/Source/Runtime/Engine/Private/PrimitiveSceneProxy.cpp
@@ -179,6 +179,8 @@ FPrimitiveSceneProxy::FPrimitiveSceneProxy(const UPrimitiveComponent* InComponen
 ,	CustomDepthStencilValue(InComponent->CustomDepthStencilValue)
 ,	CustomDepthStencilWriteMask(FRendererStencilMaskEvaluation::ToStencilMask(InComponent->CustomDepthStencilWriteMask))
 ,	LightingChannelMask(GetLightingChannelMaskForStruct(InComponent->LightingChannels))
+,	RayTracingGroupId(InComponent->RayTracingGroupId)
+,	RayTracingGroupCullingPriority((uint8)InComponent->RayTracingGroupCullingPriority)
 ,	IndirectLightingCacheQuality(InComponent->IndirectLightingCacheQuality)
 ,	VirtualTextureLodBias(InComponent->VirtualTextureLodBias)
 ,	VirtualTextureCullMips(InComponent->VirtualTextureCullMips)
diff --git a/Engine/Source/Runtime/Engine/Public/ISMPartition/ISMComponentDescriptor.h b/Engine/Source/Runtime/Engine/Public/ISMPartition/ISMComponentDescriptor.h
--- a/Engine/Source/Runtime/Engine/Public/ISMPartition/ISMComponentDescriptor.h
+++ b/Engine/Source/Runtime/Engine/Public/ISMPartition/ISMComponentDescriptor.h
@@ -76,6 +76,12 @@ public:
 	UPROPERTY()
 	FLightingChannels LightingChannels;
 
+	UPROPERTY()
+	int32 RayTracingGroupId;
+
+	UPROPERTY()
+	ERayTracingGroupCullingPriority RayTracingGroupCullingPriority;
+
 	UPROPERTY()
 	TEnumAsByte<EHasCustomNavigableGeometry::Type> bHasCustomNavigableGeometry;
 
diff --git a/Engine/Source/Runtime/Engine/Public/PrimitiveSceneProxy.h b/Engine/Source/Runtime/Engine/Public/PrimitiveSceneProxy.h
--- a/Engine/Source/Runtime/Engine/Public/PrimitiveSceneProxy.h
+++ b/Engine/Source/Runtime/Engine/Public/PrimitiveSceneProxy.h
@@ -1027,6 +1027,10 @@ private:
 
 	uint8 LightingChannelMask;
 
+	// Run-time groups of proxies
+	int32 RayTracingGroupId;
+	uint8 RayTracingGroupCullingPriority;
+
 protected:
 
 	/** Quality of interpolated indirect lighting for Movable components. */
