Fix misleading comment about ParentIndex.

#trivial
#rnx

diff --git a/Engine/Source/Runtime/Engine/Public/ReferenceSkeleton.h b/Engine/Source/Runtime/Engine/Public/ReferenceSkeleton.h
--- a/Engine/Source/Runtime/Engine/Public/ReferenceSkeleton.h
+++ b/Engine/Source/Runtime/Engine/Public/ReferenceSkeleton.h
@@ -15,7 +15,7 @@ struct FMeshBoneInfo
 	// Bone's name.
 	FName Name;
 
-	// 0/NULL if this is the root bone. 
+	// INDEX_NONE if this is the root bone. 
 	int32 ParentIndex;
 
 #if WITH_EDITORONLY_DATA
