Added FBoneContainer::operator=! for symmetry.

#trivial

diff --git a/Engine/Source/Runtime/Engine/Public/BoneContainer.h b/Engine/Source/Runtime/Engine/Public/BoneContainer.h
--- a/Engine/Source/Runtime/Engine/Public/BoneContainer.h
+++ b/Engine/Source/Runtime/Engine/Public/BoneContainer.h
@@ -557,6 +557,12 @@ struct FBoneReference
 	{
 		return BoneName == Other.BoneName;
 	}
+
+	bool operator!=(const FBoneReference& Other) const
+	{
+		return BoneName != Other.BoneName;
+	}
+
 	/** Initialize Bone Reference, return TRUE if success, otherwise, return false **/
 	ENGINE_API bool Initialize(const FBoneContainer& RequiredBones);
 
