Added blueprint exposed inertialization slot goup request to AnimInstance.
This is meant to give the user a way to explicitly notify the anim  graph that a pose discontinuity has ocurred.

diff --git a/Engine/Source/Runtime/Engine/Classes/Animation/AnimInstance.h b/Engine/Source/Runtime/Engine/Classes/Animation/AnimInstance.h
--- a/Engine/Source/Runtime/Engine/Classes/Animation/AnimInstance.h
+++ b/Engine/Source/Runtime/Engine/Classes/Animation/AnimInstance.h
@@ -881,6 +881,10 @@ public:
 	/**  Builds an inertialization request from the montage's group and provided duration */
 	void RequestMontageInertialization(const UAnimMontage* Montage, float Duration);
 
+	/**  Requests an inertial blend during the next anim graph update. Requires your anim graph to have a slot node belonging to the specified group name */
+	UFUNCTION(BlueprintCallable, Category = "Inertial Blend")
+	void RequestSlotGroupInertialization(FName InSlotGroupName, float Duration);
+
 	/** Queue a Montage BlendingOut Event to be triggered. */
 	void QueueMontageBlendingOutEvent(const FQueuedMontageBlendingOutEvent& MontageBlendingOutEvent);
 
diff --git a/Engine/Source/Runtime/Engine/Private/Animation/AnimInstance.cpp b/Engine/Source/Runtime/Engine/Private/Animation/AnimInstance.cpp
--- a/Engine/Source/Runtime/Engine/Private/Animation/AnimInstance.cpp
+++ b/Engine/Source/Runtime/Engine/Private/Animation/AnimInstance.cpp
@@ -1654,6 +1654,13 @@ void UAnimInstance::Montage_Advance(float DeltaSeconds)
 	}
 }
 
+void UAnimInstance::RequestSlotGroupInertialization(FName InSlotGroupName, float Duration)
+{
+	// Must add this on both the anim instance and proxy's map, as this could called after UAnimInstance::UpdateMontageEvaluationData.
+	SlotGroupInertializationRequestMap.FindOrAdd(InSlotGroupName) = Duration;
+	GetProxyOnAnyThread<FAnimInstanceProxy>().GetSlotGroupInertializationRequestMap().FindOrAdd(InSlotGroupName) = Duration;
+}
+
 void UAnimInstance::RequestMontageInertialization(const UAnimMontage* Montage, float Duration)
 {
 	if (Montage)
