Add support for SSS transmission with Strata material.

diff --git a/Engine/Shaders/Private/BurleyNormalizedSSSCommon.ush b/Engine/Shaders/Private/BurleyNormalizedSSSCommon.ush
--- a/Engine/Shaders/Private/BurleyNormalizedSSSCommon.ush
+++ b/Engine/Shaders/Private/BurleyNormalizedSSSCommon.ush
@@ -166,4 +166,44 @@ float4 GetBurleyTransmissionProfile(FGBufferData GBufferData, float Thickness)
 	Output.xyz= 0.25*SurfaceAlbedo*(exp(-ScalingFactor * r) + 3 * exp(-ScalingFactor * r / 3));
 	return Output;
 }
-#endif
\ No newline at end of file
+#endif
+
+float3 InternalGetBurleyTransmissionProfile(
+	float3 SubsurfaceAlebdo, 
+	float3 MeanFreePathInCm,
+	float ThicknessInCm)
+{
+	// Assuming that the volume albedo is the same to the surface albedo for transmission;
+	const float3 ScalingFactor = GetSearchLightDiffuseScalingFactor3D(SubsurfaceAlebdo);
+	const float3 r = ThicknessInCm / MeanFreePathInCm.xyz;
+	const float3 Output = 0.25 * SubsurfaceAlebdo * (exp(-ScalingFactor * r) + 3 * exp(-ScalingFactor * r / 3));
+
+	return Output;
+}
+
+float3 GetBurleyTransmissionProfile(
+	float3 SubsurfaceAlebdo,
+	float3 MeanFreePathInCm,
+	float ThicknessInCm)
+{
+	const float  TransmissionMFPScaleFactor = 100.f; // This factor compensate a mistake in SSS profile transmission computation. See BurleyNormalizedSSS.cpp for more details
+	const float3 TransmissionThroughput = InternalGetBurleyTransmissionProfile(SubsurfaceAlebdo, MeanFreePathInCm * TransmissionMFPScaleFactor, ThicknessInCm).xyz;
+
+	// Replicate SSS profile distance fade out, by fading the tail of the curve to black.
+	//  * SSS Profile stores transmittance into a discrete set of samples: SSSS_TRANSMISSION_PROFILE_SIZE(=32)
+	//  * The transmission distance is set to a max of SSSS_MAX_TRANSMISSION_PROFILE_DISTANCE(=5) (in scaled cm, depend on the WorldToUnitScale, cm by default)
+	//
+	// This code fades the last value to black.
+	//  FadeDistance = Distance * SampleCount - MaxDistance * (SampleCount-0.5f)
+	//  FadeDistance = Distance * 32.f - 5.f * 31.5f
+	//
+	// Transmission
+	// 1 ___________
+	//				\
+	// 0             \ _______
+	// ----------------------->
+	// .. 29  30  31  32  Sample Index / Distance
+	const float FadeDistance = 1.0f - saturate(ThicknessInCm * SSSS_TRANSMISSION_PROFILE_SIZE - SSSS_MAX_TRANSMISSION_PROFILE_DISTANCE * (SSSS_TRANSMISSION_PROFILE_SIZE - 0.5f));
+
+	return TransmissionThroughput * FadeDistance;
+}
diff --git a/Engine/Shaders/Private/DeferredLightPixelShaders.usf b/Engine/Shaders/Private/DeferredLightPixelShaders.usf
--- a/Engine/Shaders/Private/DeferredLightPixelShaders.usf
+++ b/Engine/Shaders/Private/DeferredLightPixelShaders.usf
@@ -29,6 +29,7 @@
 #if STRATA_ENABLED
 #define STRATA_INLINE_SHADING 0
 #define STRATA_MICROFACET_MULTIPLESCATTERING 1
+#define STRATA_SSS_TRANSMISSION USE_TRANSMISSION
 #include "/Engine/Private/Strata/StrataEvaluation.ush"
 #include "/Engine/Private/Strata/StrataDeferredLighting.ush"
 #endif
diff --git a/Engine/Shaders/Private/ShadowProjectionPixelShader.usf b/Engine/Shaders/Private/ShadowProjectionPixelShader.usf
--- a/Engine/Shaders/Private/ShadowProjectionPixelShader.usf
+++ b/Engine/Shaders/Private/ShadowProjectionPixelShader.usf
@@ -257,22 +257,22 @@ void Main(
 
 	#if FEATURE_LEVEL >= FEATURE_LEVEL_SM4 && !FORWARD_SHADING && !APPLY_TRANSLUCENCY_SHADOWS && !SUBPIXEL_SHADOW
 
-		bool bHasSubsurface			= false;
-		bool bHasSubsurfaceProfile	= false;
-		bool bIsEyeOrHair 			= false;
+		bool bHasSubsurface				= false;
+		bool bHasSubsurfaceTransmission	= false;
+		bool bIsEyeOrHair 				= false;
 	#if STRATA_ENABLED
 		const FStrataClassification StrataData = StrataUnpackClassificationData(Strata.ClassificationTexture.Load(PixelCoord));
 		if (StrataData.bIsStrataMaterial)
 		{
-			bIsEyeOrHair			= StrataHasShadingModel(StrataData.ShadingModels, STRATA_BSDF_TYPE_HAIR);
-			bHasSubsurface			= StrataData.bHasSubsurface || bIsEyeOrHair;
-			bHasSubsurfaceProfile	= StrataData.bHasSubsurfaceProfile;
+			bIsEyeOrHair				= StrataHasShadingModel(StrataData.ShadingModels, STRATA_BSDF_TYPE_HAIR);
+			bHasSubsurface				= StrataData.bHasSubsurface || bIsEyeOrHair;
+			bHasSubsurfaceTransmission	= StrataData.bHasSubsurface;
 		}
 	#else // STRATA_ENABLED
 		FGBufferData GBufferData = GetGBufferData(ScreenUV);
 		{
 			bHasSubsurface = IsSubsurfaceModel(GBufferData.ShadingModelID);
-			bHasSubsurfaceProfile = GBufferData.ShadingModelID == SHADINGMODELID_SUBSURFACE_PROFILE;
+			bHasSubsurfaceTransmission = GBufferData.ShadingModelID == SHADINGMODELID_SUBSURFACE_PROFILE;
 			bIsEyeOrHair = GBufferData.ShadingModelID == SHADINGMODELID_HAIR || GBufferData.ShadingModelID == SHADINGMODELID_EYE;
 		}
 	#endif // STRATA_ENABLED
@@ -285,9 +285,9 @@ void Main(
 
 			float Opacity = 1;
 			float Density = 1;
-			if (!bIsEyeOrHair && SSSData.bIsProfile)
+			if (!bIsEyeOrHair && SSSData.bIsValid)
 			{
-				Opacity = SSSData.ProfileRadiusScale;
+				Opacity = SSSData.bIsProfile ? SSSData.ProfileRadiusScale : 1.0f;
 				Density = SubsurfaceDensityFromOpacity(Opacity);
 			}
 		#else // STRATA_ENABLED
@@ -326,12 +326,18 @@ void Main(
 				Settings.ProjectionDepthBiasParameters = ProjectionDepthBiasParameters.xw;
 
 			  #if USE_TRANSMISSION
-				if (bHasSubsurfaceProfile)
+				if (bHasSubsurfaceTransmission)
 				{
 				#if STRATA_ENABLED
 					const FStrataTopLayerData TopLayerData = StrataUnpackTopLayerData(Strata.TopLayerNormalTexture.Load(uint3(SVPos.xy, 0)));
 					const float3 WorldNormal = TopLayerData.WorldNormal;
-					const FTransmissionProfileParams TransmissionParams = GetTransmissionProfileParams(SSSData.ProfileId);
+
+					// For non-profile version, use default transmission parameters (only NormalScale & ExtinctionScale are used)
+					FTransmissionProfileParams TransmissionParams = InitTransmissionProfileParams();
+					if (StrataData.bHasSubsurfaceProfile)
+					{
+						TransmissionParams = GetTransmissionProfileParams(SSSData.ProfileId);
+					}
 				#else // STRATA_ENABLED
 					float3 WorldNormal = GBufferData.WorldNormal;
 					FTransmissionProfileParams TransmissionParams = GetTransmissionProfileParams(GBufferData);
diff --git a/Engine/Shaders/Private/Strata/StrataDeferredLighting.ush b/Engine/Shaders/Private/Strata/StrataDeferredLighting.ush
--- a/Engine/Shaders/Private/Strata/StrataDeferredLighting.ush
+++ b/Engine/Shaders/Private/Strata/StrataDeferredLighting.ush
@@ -7,6 +7,7 @@
 struct FStrataLightIntegrationContext
 {
 	float3 IncidentRadianceNoShadow;
+	float3 IncidentTransmittedRadianceNoShadow;
 	float3 IncidentIrradianceNoShadow;	// Modified incident irradiance, based on custom are light integrator
 	float3 OutgoingRadianceNoShadow;	// Continue the outgoing radiance if the BSDF evaluation/integration has been done with the custom area light integrator
 	float3 IncidentIrradianceNoShadowHaze;
@@ -21,6 +22,7 @@ FStrataLightIntegrationContext StrataIniLightIntegrationContext()
 	Out.IncidentIrradianceNoShadow = 0;
 	Out.IncidentIrradianceNoShadowHaze = 0;
 	Out.IncidentRadianceNoShadow = 0;
+	Out.IncidentTransmittedRadianceNoShadow = 0;
 	return Out;
 }
 
@@ -117,11 +119,17 @@ void StrataPreprocessAreaLightIntegration(
 
 				// Update incident irradiance to add the energy rescaling to due MRP integration
 				// The a2 override stored the overriden roughness, which will be used during the BSDF evaluation
-				BSDFContext.a2Override = Pow4(SafeRoughnessX);
-				Out.IncidentIrradianceNoShadow =
-					Out.IncidentRadianceNoShadow *
-					EnergyNormalization(BSDFContext.a2Override, BSDFContext.Context.VoH, LightContext.AreaLight) *
-					LightContext.AreaLight.FalloffColor * (LightContext.Falloff * LightContext.NoL);
+				{
+					BSDFContext.a2Override = Pow4(SafeRoughnessX);
+
+					const float3 IncidentRadiance = 
+						Out.IncidentRadianceNoShadow *
+						EnergyNormalization(BSDFContext.a2Override, BSDFContext.Context.VoH, LightContext.AreaLight) *
+						LightContext.AreaLight.FalloffColor * LightContext.Falloff;
+
+					Out.IncidentTransmittedRadianceNoShadow = IncidentRadiance;
+					Out.IncidentIrradianceNoShadow = IncidentRadiance * LightContext.NoL;
+				}
 
 				if (bHaziness)
 				{
@@ -137,6 +145,7 @@ void StrataPreprocessAreaLightIntegration(
 			}
 		}
 	}
+
 }
 
 #ifndef STRATA_FASTPATH 
@@ -185,10 +194,10 @@ float4 StrataDeferredLighting(
 		FStrataBSDFContext BSDFContext = StrataCreateBSDFContext(StrataPixelHeader, BSDF, StrataAddressing, V, L);
 
 		// Evaluate incident lighting & (optionally) precompute area light integration
-		FStrataLightIntegrationContext IntegrationContext = StrataIniLightIntegrationContext();
-		IntegrationContext.IncidentRadianceNoShadow   = LightData.Color * LightMask;
-		IntegrationContext.IncidentIrradianceNoShadow = PreintegratedIndicentIrradiance * IntegrationContext.IncidentRadianceNoShadow * StrataBSDFProjectedSolidAngleFactor(BSDFContext);
-		IntegrationContext.IncidentIrradianceNoShadowHaze = IntegrationContext.IncidentIrradianceNoShadow;
+		FStrataLightIntegrationContext IntegrationContext	= StrataIniLightIntegrationContext();
+		IntegrationContext.IncidentRadianceNoShadow			= LightData.Color * LightMask;
+		IntegrationContext.IncidentIrradianceNoShadow		= PreintegratedIndicentIrradiance * IntegrationContext.IncidentRadianceNoShadow * StrataBSDFProjectedSolidAngleFactor(BSDFContext);
+		IntegrationContext.IncidentIrradianceNoShadowHaze	= IntegrationContext.IncidentIrradianceNoShadow;
 		StrataPreprocessAreaLightIntegration(ToLight, LightData, RectTexture, BSDFContext, IntegrationContext);
 
 		// Evaluate BSDF
@@ -203,9 +212,9 @@ float4 StrataDeferredLighting(
 
 			const bool bNeedsSeparateSubsurfaceLightAccumulation = BSDFEvaluate.bSubsurface;
 
-			LightAccumulator_AddSplit(Out, DiffusePathWeight,		SpecularPathWeight,		DiffusePathWeight,		IntegrationContext.IncidentIrradianceNoShadow * ShadowTerms.SurfaceShadow,		bNeedsSeparateSubsurfaceLightAccumulation);
-			LightAccumulator_AddSplit(Out, 0.0,						SpecularHazePathWeight, 0.0,					IntegrationContext.IncidentIrradianceNoShadowHaze * ShadowTerms.SurfaceShadow,	bNeedsSeparateSubsurfaceLightAccumulation);
-			LightAccumulator_AddSplit(Out, TransmissionPathWeight,	0.0f,					TransmissionPathWeight,	IntegrationContext.IncidentIrradianceNoShadow * ShadowTerms.TransmissionShadow,	bNeedsSeparateSubsurfaceLightAccumulation);
+			LightAccumulator_AddSplit(Out, DiffusePathWeight,		SpecularPathWeight,		DiffusePathWeight,		IntegrationContext.IncidentIrradianceNoShadow * ShadowTerms.SurfaceShadow,					bNeedsSeparateSubsurfaceLightAccumulation);
+			LightAccumulator_AddSplit(Out, 0.0,						SpecularHazePathWeight, 0.0,					IntegrationContext.IncidentIrradianceNoShadowHaze * ShadowTerms.SurfaceShadow,				bNeedsSeparateSubsurfaceLightAccumulation);
+			LightAccumulator_AddSplit(Out, TransmissionPathWeight,	0.0f,					TransmissionPathWeight, IntegrationContext.IncidentTransmittedRadianceNoShadow * ShadowTerms.TransmissionShadow,	bNeedsSeparateSubsurfaceLightAccumulation);
 		}
 		else
 		{
diff --git a/Engine/Shaders/Private/Strata/StrataEvaluation.ush b/Engine/Shaders/Private/Strata/StrataEvaluation.ush
--- a/Engine/Shaders/Private/Strata/StrataEvaluation.ush
+++ b/Engine/Shaders/Private/Strata/StrataEvaluation.ush
@@ -25,6 +25,10 @@
 #define STRATA_MICROFACET_MULTIPLESCATTERING 0
 #endif
 
+#ifndef STRATA_SSS_TRANSMISSION
+#define STRATA_SSS_TRANSMISSION 0
+#endif
+
 // STRATA_TODO put in a common file
 // Point lobe in off-specular peak direction
 float3 StrataGetOffSpecularPeakReflectionDir(float3 Normal, float3 ReflectionVector, float Roughness)
@@ -347,6 +351,42 @@ FStrataEvaluateResult StrataEvaluateBSDF(FStrataBSDFContext BSDFContext, FShadow
 
 			Sample.Throughput = OpaqueBSDFThroughput;
 
+			#if STRATA_FASTPATH==0 && STRATA_SSS_TRANSMISSION
+			const bool bHasProfile = BSDF_GETHASSSSPROFILE(BSDFContext.BSDF);
+			if (BSDF_GETHASSSS(BSDFContext.BSDF))
+			{
+				const float ThicknessInCm = DecodeThickness(ShadowTerms.TransmissionThickness) * SSSS_MAX_TRANSMISSION_PROFILE_DISTANCE;
+
+				float OneOverIOR = 1.0f;
+				float PhaseFunctionAnisotropy = 0.0f;
+				float3 TransmissionThroughput = 1.0f;
+				if (bHasProfile)
+				{
+					const uint ProfileId = StrataSubsurfaceProfileIdTo8bits(SLAB_SSSPROFILEID(BSDFContext.BSDF)); // TODO move this into the PackStrataOut( function, to avoid this decode here
+					const FTransmissionProfileParams TransmissionParams = GetTransmissionProfileParams(ProfileId);
+					TransmissionThroughput	= GetTransmissionProfile(ProfileId, ThicknessInCm).rgb;
+					PhaseFunctionAnisotropy = TransmissionParams.ScatteringDistribution;
+					OneOverIOR				= TransmissionParams.OneOverIOR;
+
+				}
+				else
+				{
+					const float3 SubsurfaceAlebdo = SLAB_BASECOLOR(BSDFContext.BSDF);
+					const float3 MeanFreePathInCm = SLAB_SSSDMFP(BSDFContext.BSDF);
+					TransmissionThroughput	= GetBurleyTransmissionProfile(SubsurfaceAlebdo, MeanFreePathInCm, ThicknessInCm);
+					PhaseFunctionAnisotropy = 0.93f; // STRATA_TODO: expose phase function eccentricity. For now, forward scattering matching default SSS profile value
+					OneOverIOR				= 1.f / DielectricF0ToIor(F0.y);
+				}
+
+				const float3 RefracV = refract(BSDFContext.V, -BSDFContext.N, OneOverIOR);
+				const float PhaseFunction = ApproximateHG(dot(-BSDFContext.L, RefracV), PhaseFunctionAnisotropy);
+
+				Sample.Throughput				= TransmissionThroughput * PhaseFunction;
+				Sample.TransmissionPathValue	= TransmissionThroughput * PhaseFunction;
+				Sample.TransmissionPDF			= 1.0f / (4.0f * PI);		// STRATA_TODO this currently match the uniform sphere sampling from StrataImportanceSampleBSDF
+			}
+			#endif
+
 			#if STRATA_FASTPATH==0
 			BRANCH
 			if (BSDF_GETISSIMPLEVOLUME(BSDFContext.BSDF))
diff --git a/Engine/Shaders/Private/Strata/StrataVisualize.usf b/Engine/Shaders/Private/Strata/StrataVisualize.usf
--- a/Engine/Shaders/Private/Strata/StrataVisualize.usf
+++ b/Engine/Shaders/Private/Strata/StrataVisualize.usf
@@ -5,6 +5,7 @@
 
 #define STRATA_INLINE_SHADING 0
 #define STRATA_SSS_MATERIAL_OVERRIDE 0
+#define STRATA_SSS_TRANSMISSION 0  // Disable for now, as the SSS profile texture need to be bound
 #include "/Engine/Private/Strata/StrataEvaluation.ush"
 
 #include "/Engine/Private/MiniFontCommon.ush"
diff --git a/Engine/Shaders/Private/TransmissionCommon.ush b/Engine/Shaders/Private/TransmissionCommon.ush
--- a/Engine/Shaders/Private/TransmissionCommon.ush
+++ b/Engine/Shaders/Private/TransmissionCommon.ush
@@ -19,6 +19,16 @@ struct FTransmissionProfileParams
 	float OneOverIOR;
 };
 
+FTransmissionProfileParams InitTransmissionProfileParams()
+{
+	FTransmissionProfileParams Out;
+	Out.ExtinctionScale			= 1.0f;
+	Out.NormalScale				= 0.08f;
+	Out.ScatteringDistribution	= 0.0f;
+	Out.OneOverIOR				= 1.0f;
+	return Out;
+}
+
 FTransmissionProfileParams GetTransmissionProfileParams(uint SubsurfaceProfileInt)
 {
 	FTransmissionProfileParams Result;
diff --git a/Engine/Source/Runtime/Engine/Private/Rendering/BurleyNormalizedSSS.cpp b/Engine/Source/Runtime/Engine/Private/Rendering/BurleyNormalizedSSS.cpp
--- a/Engine/Source/Runtime/Engine/Private/Rendering/BurleyNormalizedSSS.cpp
+++ b/Engine/Source/Runtime/Engine/Private/Rendering/BurleyNormalizedSSS.cpp
@@ -27,18 +27,18 @@ inline float Burley_TransmissionProfile(float r, float A, float S, float L)
 }
 
 
-inline FVector Burley_ScatteringProfile(float r, FLinearColor SurfaceAlbedo, FVector ScalingFactor, FLinearColor DiffuseMeanFreePath)
+inline FVector Burley_ScatteringProfile(float RadiusInMm, FLinearColor SurfaceAlbedo, FVector ScalingFactor, FLinearColor DiffuseMeanFreePathInMm)
 {  
-	return FVector(Burley_ScatteringProfile(r, SurfaceAlbedo.R, ScalingFactor.X, DiffuseMeanFreePath.R) ,
-						Burley_ScatteringProfile(r, SurfaceAlbedo.G, ScalingFactor.Y, DiffuseMeanFreePath.G) ,
-						Burley_ScatteringProfile(r, SurfaceAlbedo.B, ScalingFactor.Z, DiffuseMeanFreePath.B));
+	return FVector(	Burley_ScatteringProfile(RadiusInMm, SurfaceAlbedo.R, ScalingFactor.X, DiffuseMeanFreePathInMm.R) ,
+					Burley_ScatteringProfile(RadiusInMm, SurfaceAlbedo.G, ScalingFactor.Y, DiffuseMeanFreePathInMm.G) ,
+					Burley_ScatteringProfile(RadiusInMm, SurfaceAlbedo.B, ScalingFactor.Z, DiffuseMeanFreePathInMm.B));
 }
 
-inline FLinearColor Burley_TransmissionProfile(float r, FLinearColor SurfaceAlbedo, FVector ScalingFactor, FLinearColor DiffuseMeanFreePath)
+inline FLinearColor Burley_TransmissionProfile(float RadiusInMm, FLinearColor SurfaceAlbedo, FVector ScalingFactor, FLinearColor DiffuseMeanFreePathInMm)
 {
-	return FLinearColor(Burley_TransmissionProfile(r, SurfaceAlbedo.R, ScalingFactor.X, DiffuseMeanFreePath.R),
-						Burley_TransmissionProfile(r, SurfaceAlbedo.G, ScalingFactor.Y, DiffuseMeanFreePath.G),
-						Burley_TransmissionProfile(r, SurfaceAlbedo.B, ScalingFactor.Z, DiffuseMeanFreePath.B));
+	return FLinearColor(Burley_TransmissionProfile(RadiusInMm, SurfaceAlbedo.R, ScalingFactor.X, DiffuseMeanFreePathInMm.R),
+						Burley_TransmissionProfile(RadiusInMm, SurfaceAlbedo.G, ScalingFactor.Y, DiffuseMeanFreePathInMm.G),
+						Burley_TransmissionProfile(RadiusInMm, SurfaceAlbedo.B, ScalingFactor.Z, DiffuseMeanFreePathInMm.B));
 }
 
 //--------------------------------------------------------------------------
@@ -225,13 +225,26 @@ void ComputeMirroredBSSSKernel(FLinearColor* TargetBuffer, uint32 TargetBufferSi
 
 void ComputeTransmissionProfileBurley(FLinearColor* TargetBuffer, uint32 TargetBufferSize, FLinearColor SubsurfaceColor, 
 																		FLinearColor FalloffColor, float ExtinctionScale, 
-																		FLinearColor SurfaceAlbedo, FLinearColor DiffuseMeanFreePath,
-																		float WorldUnitScale, FLinearColor TransmissionTintColor)
+																		FLinearColor SurfaceAlbedo, FLinearColor DiffuseMeanFreePathInMm,
+																		float InUnitToMm, FLinearColor TransmissionTintColor)
 {
 	check(TargetBuffer);
 	check(TargetBufferSize > 0);
 
-	static float MaxTransmissionProfileDistance = 5.0f; // See MAX_TRANSMISSION_PROFILE_DISTANCE in TransmissionCommon.ush
+
+	// Example of scaling
+	// ----------------------------------------
+	// DistanceCm * UnitScale * CmToMm = Value (mm)
+	// ----------------------------------------
+	//   1          0.1         10     =   1mm
+	//   1          1.0         10     =  10mm
+	//   1         10.0         10     = 100mm
+
+	// It seems there is an inconsistency here: either InUnitToMm should in fact be MmToUnit or the division need to be converted to a multiply
+	// This should be corrected as this makes the transmission scaling incorrect
+	const float UnitToMm = 1.f / InUnitToMm;
+
+	static float MaxTransmissionProfileDistance = 5.0f; // See SSSS_MAX_TRANSMISSION_PROFILE_DISTANCE in TransmissionCommon.ush
 
 	//assuming that the volume albedo is the same to the surface albedo for transmission.
 	FVector ScalingFactor = GetSearchLightDiffuseScalingFactor(SurfaceAlbedo);
@@ -240,13 +253,14 @@ void ComputeTransmissionProfileBurley(FLinearColor* TargetBuffer, uint32 TargetB
 
 	for (uint32 i = 0; i < TargetBufferSize; ++i)
 	{
+		const float DistanceInMm = i * InvSize * MaxTransmissionProfileDistance * UnitToMm;
+		const float OffsetInMm = ProfileRadiusOffset * UnitToMm;
 
-		float Distance = i * InvSize * MaxTransmissionProfileDistance/WorldUnitScale;
 
-		FLinearColor TransmissionProfile = Burley_TransmissionProfile(Distance + ProfileRadiusOffset / WorldUnitScale, SurfaceAlbedo,ScalingFactor, DiffuseMeanFreePath);
+		FLinearColor TransmissionProfile = Burley_TransmissionProfile(DistanceInMm + OffsetInMm, SurfaceAlbedo, ScalingFactor, DiffuseMeanFreePathInMm);
 		TargetBuffer[i] = TransmissionProfile * TransmissionTintColor; // Apply tint to the profile
 		//Use Luminance of scattering as SSSS shadow.
-		TargetBuffer[i].A = exp(-Distance * ExtinctionScale);
+		TargetBuffer[i].A = exp(-DistanceInMm * ExtinctionScale);
 	}
 
 	// Do this is because 5mm is not enough cool down the scattering to zero, although which is small number but after tone mapping still noticeable
diff --git a/Engine/Source/Runtime/Engine/Private/Rendering/SubsurfaceProfile.cpp b/Engine/Source/Runtime/Engine/Private/Rendering/SubsurfaceProfile.cpp
--- a/Engine/Source/Runtime/Engine/Private/Rendering/SubsurfaceProfile.cpp
+++ b/Engine/Source/Runtime/Engine/Private/Rendering/SubsurfaceProfile.cpp
@@ -250,7 +250,7 @@ void FSubsurfaceProfileTexture::CreateTexture(FRHICommandListImmediate& RHICmdLi
 
 		Data.SubsurfaceColor = Data.SubsurfaceColor.GetClamped();
 		Data.FalloffColor = Data.FalloffColor.GetClamped(Bias);
-		Data.MeanFreePathColor = Data.MeanFreePathColor.GetClamped(Bias);
+		Data.MeanFreePathColor = Data.MeanFreePathColor.GetClamped(Bias); // In Cm
 		Data.TransmissionTintColor = Data.TransmissionTintColor.GetClamped(Bias);
 		Data.SurfaceAlbedo = Data.SurfaceAlbedo.GetClamped(Bias);
 
@@ -258,10 +258,13 @@ void FSubsurfaceProfileTexture::CreateTexture(FRHICommandListImmediate& RHICmdLi
 		TextureRow[SSSS_SUBSURFACE_COLOR_OFFSET] = Data.SubsurfaceColor;
 		TextureRow[SSSS_SUBSURFACE_COLOR_OFFSET].A = EncodeWorldUnitScale(Data.WorldUnitScale);
 		
-		FLinearColor DifffuseMeanFreePath = Data.MeanFreePathColor*Data.MeanFreePathDistance*10.0f; // convert cm to mm.
-		SetupSurfaceAlbedoAndDiffuseMeanFreePath(Data.SurfaceAlbedo, DifffuseMeanFreePath);
+		const float UnitToCm = Data.WorldUnitScale;
+		const float CmToMm = 10.f;
+
+		FLinearColor DifffuseMeanFreePathInMm = (Data.MeanFreePathColor*Data.MeanFreePathDistance) * CmToMm; // convert cm to mm.
+		SetupSurfaceAlbedoAndDiffuseMeanFreePath(Data.SurfaceAlbedo, DifffuseMeanFreePathInMm);
 		TextureRow[BSSS_SURFACEALBEDO_OFFSET] = Data.SurfaceAlbedo;
-		TextureRow[BSSS_DMFP_OFFSET] = EncodeDiffuseMeanFreePath(DifffuseMeanFreePath);
+		TextureRow[BSSS_DMFP_OFFSET] = EncodeDiffuseMeanFreePath(DifffuseMeanFreePathInMm);
 
 		TextureRow[SSSS_BOUNDARY_COLOR_BLEED_OFFSET] = Data.BoundaryColorBleed;
 
@@ -300,11 +303,11 @@ void FSubsurfaceProfileTexture::CreateTexture(FRHICommandListImmediate& RHICmdLi
 			Data.ScatterRadius = FMath::Max(Data.MeanFreePathDistance*Data.MeanFreePathColor.GetMax()/2.229f, 0.1f);
 
 			ComputeMirroredBSSSKernel(&TextureRow[SSSS_KERNEL0_OFFSET], SSSS_KERNEL0_SIZE, Data.SurfaceAlbedo,
-				DifffuseMeanFreePath, Data.ScatterRadius);
+				DifffuseMeanFreePathInMm, Data.ScatterRadius);
 			ComputeMirroredBSSSKernel(&TextureRow[SSSS_KERNEL1_OFFSET], SSSS_KERNEL1_SIZE, Data.SurfaceAlbedo,
-				DifffuseMeanFreePath, Data.ScatterRadius);
+				DifffuseMeanFreePathInMm, Data.ScatterRadius);
 			ComputeMirroredBSSSKernel(&TextureRow[SSSS_KERNEL2_OFFSET], SSSS_KERNEL2_SIZE, Data.SurfaceAlbedo,
-				DifffuseMeanFreePath, Data.ScatterRadius);
+				DifffuseMeanFreePathInMm, Data.ScatterRadius);
 
 			// Then, scale up by world unit scale and the fitting parameters to affect screen space sampling location.
 			// For high irradiance, the lose of energy due to insufficient sampling count needs to be compensated to
@@ -326,9 +329,8 @@ void FSubsurfaceProfileTexture::CreateTexture(FRHICommandListImmediate& RHICmdLi
 		}
 
 		ComputeTransmissionProfile(&TextureRow[SSSS_TRANSMISSION_PROFILE_OFFSET], SSSS_TRANSMISSION_PROFILE_SIZE, Data.SubsurfaceColor, Data.FalloffColor, Data.ExtinctionScale);
-
-		ComputeTransmissionProfileBurley(&TextureRow[BSSS_TRANSMISSION_PROFILE_OFFSET], BSSS_TRANSMISSION_PROFILE_SIZE,
-			Data.SubsurfaceColor, Data.FalloffColor, Data.ExtinctionScale, Data.SurfaceAlbedo, DifffuseMeanFreePath,Data.WorldUnitScale * 10.0f/*cm to mm*/,Data.TransmissionTintColor);
+		ComputeTransmissionProfileBurley(&TextureRow[BSSS_TRANSMISSION_PROFILE_OFFSET], BSSS_TRANSMISSION_PROFILE_SIZE, 
+			Data.SubsurfaceColor, Data.FalloffColor, Data.ExtinctionScale, Data.SurfaceAlbedo, DifffuseMeanFreePathInMm, UnitToCm * CmToMm, Data.TransmissionTintColor);
 
 		// could be lower than 1 (but higher than 0) to range compress for better quality (for 8 bit)
 		const float TableMaxRGB = 1.0f;
