- Incorporate sampling frame-rate into DDC key suffix
- Bumped anim compression base ddc key
- Ensure that if actually resampling bone data, to copy bone name and indices as well

diff --git a/Engine/Source/Runtime/Engine/Private/Animation/AnimCompressionDerivedData.cpp b/Engine/Source/Runtime/Engine/Private/Animation/AnimCompressionDerivedData.cpp
--- a/Engine/Source/Runtime/Engine/Private/Animation/AnimCompressionDerivedData.cpp
+++ b/Engine/Source/Runtime/Engine/Private/Animation/AnimCompressionDerivedData.cpp
@@ -35,7 +35,7 @@ const TCHAR* FDerivedDataAnimationCompression::GetVersionString() const
 	// This is a version string that mimics the old versioning scheme. If you
 	// want to bump this version, generate a new guid using VS->Tools->Create GUID and
 	// return it here. Ex.
-	return TEXT("C006D6DDA62143F0915F60671659A4FD");
+	return TEXT("832C5B24BBCC4F4EA235EB8112450C55");
 }
 
 bool FDerivedDataAnimationCompression::Build( TArray<uint8>& OutDataArray )
diff --git a/Engine/Source/Runtime/Engine/Private/Animation/AnimSequence.cpp b/Engine/Source/Runtime/Engine/Private/Animation/AnimSequence.cpp
--- a/Engine/Source/Runtime/Engine/Private/Animation/AnimSequence.cpp
+++ b/Engine/Source/Runtime/Engine/Private/Animation/AnimSequence.cpp
@@ -250,6 +250,7 @@ FString GetAnimSequenceSpecificCacheKeySuffix(const UAnimSequence& Seq, bool bPe
 	//	* Additive ref pose GUID or hardcoded string if not available
 	//	* Compression Settings
 	//	* Curve compression settings
+	//  * Target frames per second
 
 	bool bIsValidAdditive = Seq.IsValidAdditive();
 	char AdditiveType = bIsValidAdditive ? NibbleToTChar(Seq.AdditiveAnimType) : '0';
@@ -262,7 +263,7 @@ FString GetAnimSequenceSpecificCacheKeySuffix(const UAnimSequence& Seq, bool bPe
 	Seq.BoneCompressionSettings->PopulateDDCKey(ArcToHexString.Ar);
 	Seq.CurveCompressionSettings->PopulateDDCKey(ArcToHexString.Ar);
 
-	FString Ret = FString::Printf(TEXT("%i_%s%s%s_%c%c%i_%s_%s"),
+	FString Ret = FString::Printf(TEXT("%i_%s%s%s_%c%c%i_%s_%s_%i"),
 		Seq.CompressCommandletVersion,
 		PRAGMA_DISABLE_DEPRECATION_WARNINGS
 		*Seq.GetRawDataGuid().ToString(),
@@ -275,7 +276,8 @@ FString GetAnimSequenceSpecificCacheKeySuffix(const UAnimSequence& Seq, bool bPe
 		PRAGMA_DISABLE_DEPRECATION_WARNINGS
 		(bIsValidAdditive && Seq.RefPoseSeq) ? *Seq.RefPoseSeq->GetRawDataGuid().ToString() : TEXT("NoAdditiveGuid"),
 		PRAGMA_ENABLE_DEPRECATION_WARNINGS
-		*ArcToHexString.MakeString()
+		*ArcToHexString.MakeString(),
+		Seq.GetSamplingFrameRate().Numerator / Seq.GetSamplingFrameRate().Denominator
 	);
 
 	return Ret;
@@ -5736,18 +5738,18 @@ void UAnimSequence::ResampleAnimationTrackData()
 {
 	ValidateModel();
 
+	// Make a copy, deals with bone name and index
+	ResampledAnimationTrackData = DataModel->GetBoneAnimationTracks();
+
 	if (TargetFrameRate == DataModel->GetFrameRate())
 	{
 		// Can directly copy
-		ResampledAnimationTrackData = DataModel->GetBoneAnimationTracks();
 		NumberOfSampledKeys = DataModel->GetNumberOfKeys();
 		NumberOfSampledFrames = DataModel->GetNumberOfFrames();
 	}
 	else
 	{
-		// We actually need to resample
-		ResampledAnimationTrackData.SetNum(DataModel->GetBoneAnimationTracks().Num());
-
+		// We actually need to resample bone transforms
 		NumberOfSampledFrames = TargetFrameRate.AsFrameNumber(DataModel->GetPlayLength()).Value;
 		NumberOfSampledKeys = NumberOfSampledFrames + 1;
 				
