Fix crash caused by a copy of the TOptional render context being created during a MoveTemp
Fix potential race condition between serialization and async compilation
Add missing serialization skip condition for skeletal mesh

diff --git a/Engine/Source/Editor/ContentBrowser/Private/FrontendFilters.cpp b/Engine/Source/Editor/ContentBrowser/Private/FrontendFilters.cpp
--- a/Engine/Source/Editor/ContentBrowser/Private/FrontendFilters.cpp
+++ b/Engine/Source/Editor/ContentBrowser/Private/FrontendFilters.cpp
@@ -1223,10 +1223,9 @@ void FFrontendFilter_InUseByLoadedLevels::OnAssetPostCompile(const TArray<FAsset
 {
 	if (bIsCurrentlyActive && !bIsDirty)
 	{
-		// Only StaticMesh requires that this filter be refreshed for now
 		for (const FAssetCompileData& CompileData : CompiledAssets)
 		{
-			if (CompileData.Asset.IsValid() && CompileData.Asset->IsA<UStaticMesh>())
+			if (CompileData.Asset.IsValid())
 			{
 				bIsDirty = true;
 				break;
diff --git a/Engine/Source/Runtime/Engine/Classes/Engine/SkeletalMesh.h b/Engine/Source/Runtime/Engine/Classes/Engine/SkeletalMesh.h
--- a/Engine/Source/Runtime/Engine/Classes/Engine/SkeletalMesh.h
+++ b/Engine/Source/Runtime/Engine/Classes/Engine/SkeletalMesh.h
@@ -131,7 +131,7 @@ public:
 class FSkeletalMeshBuildContext
 {
 public:
-	TOptional<FSkinnedMeshComponentRecreateRenderStateContext> RecreateRenderStateContext;
+	TUniquePtr<FSkinnedMeshComponentRecreateRenderStateContext> RecreateRenderStateContext;
 };
 
 #if WITH_EDITOR
diff --git a/Engine/Source/Runtime/Engine/Private/SkeletalMesh.cpp b/Engine/Source/Runtime/Engine/Private/SkeletalMesh.cpp
--- a/Engine/Source/Runtime/Engine/Private/SkeletalMesh.cpp
+++ b/Engine/Source/Runtime/Engine/Private/SkeletalMesh.cpp
@@ -1605,6 +1605,23 @@ void USkeletalMesh::Serialize( FArchive& Ar )
 	LLM_SCOPE_BYNAME(TEXT("SkeletalMesh/Serialize")); // This is an important test case for SCOPE_BYNAME with a matching LLM_DEFINE_TAG
 	DECLARE_SCOPE_CYCLE_COUNTER( TEXT("USkeletalMesh::Serialize"), STAT_SkeletalMesh_Serialize, STATGROUP_LoadTime );
 
+#if WITH_EDITOR
+	if (IsCompiling())
+	{
+		// Skip serialization during compilation if told to do so.
+		if (Ar.ShouldSkipCompilingAssets())
+		{
+			return;
+		}
+
+		// Since UPROPERTY are accessed directly by offset during serialization instead of using accessors, 
+		// the protection put in place to automatically finish compilation if a locked property is accessed will not work. 
+		// We have no choice but to force finish the compilation here to avoid potential race conditions between 
+		// async compilation and the serialization.
+		FSkeletalMeshCompilingManager::Get().FinishCompilation({ this });
+	}
+#endif
+
 	Super::Serialize(Ar);
 
 	Ar.UsingCustomVersion(FFrameworkObjectVersion::GUID);
@@ -1999,7 +2016,7 @@ void USkeletalMesh::BeginBuildInternal(FSkeletalMeshBuildContext& Context)
 	TRACE_CPUPROFILER_EVENT_SCOPE(USkeletalMesh::BeginBuildInternal);
 
 	// Unregister all instances of this component
-	Context.RecreateRenderStateContext.Emplace(this, false);
+	Context.RecreateRenderStateContext = MakeUnique<FSkinnedMeshComponentRecreateRenderStateContext>(this, false);
 
 	// Release the static mesh's resources.
 	ReleaseResources();
