Fix potential race condition between serialization and async compilation
  - Thanks Matt for reporting and providing the solution this CL is based on.
  - Race conditions were encountered when loading and then quickly saving static meshes.

TESTS
  - Making sure no performance regression was introduced by testing these scenarios
    - Garbage collect with -asyncstaticmeshcompilation=paused
    - PIE on Reverb P_World
    - Applying different filters in the content browser
    - Computing statistics was already causing stalls for compiling static meshes
    - Transacting a static mesh is expected to cause a finish compilation

diff --git a/Engine/Source/Runtime/Engine/Private/StaticMesh.cpp b/Engine/Source/Runtime/Engine/Private/StaticMesh.cpp
--- a/Engine/Source/Runtime/Engine/Private/StaticMesh.cpp
+++ b/Engine/Source/Runtime/Engine/Private/StaticMesh.cpp
@@ -4722,10 +4722,23 @@ void UStaticMesh::Serialize(FArchive& Ar)
 
 	SCOPE_MS_ACCUMULATOR(STAT_StaticMesh_SerializeFull);
 
-	// Skip serialization for static mesh being compiled if told to do so.
-	if (Ar.ShouldSkipCompilingAssets() && IsCompiling())
+	if (IsCompiling())
 	{
-		return;
+		// Skip serialization during compilation if told to do so.
+		if (Ar.ShouldSkipCompilingAssets())
+		{
+			return;
+		}
+#if WITH_EDITOR
+		// Since UPROPERTY are accessed directly by offset during serialization instead of using accessors, 
+		// the protection put in place to automatically finish compilation if a locked property is accessed will not work. 
+		// We have no choice but to force finish the compilation here to avoid potential race conditions between 
+		// async compilation and the serialization.
+		else
+		{
+			FStaticMeshCompilingManager::Get().FinishCompilation({this});
+		}
+#endif
 	}
 
 	{
