Chaos - fix possible zero mass call to Combine inertias

Manual integrated from EarlyAccess (RM is down)

diff --git a/Engine/Source/Runtime/Engine/Private/PhysicsEngine/Experimental/ChaosInterfaceUtils.cpp b/Engine/Source/Runtime/Engine/Private/PhysicsEngine/Experimental/ChaosInterfaceUtils.cpp
--- a/Engine/Source/Runtime/Engine/Private/PhysicsEngine/Experimental/ChaosInterfaceUtils.cpp
+++ b/Engine/Source/Runtime/Engine/Private/PhysicsEngine/Experimental/ChaosInterfaceUtils.cpp
@@ -471,15 +471,14 @@ namespace ChaosInterface
 			}
 		}
 
-		if (TotalMass > 0.f)
+		Chaos::FMatrix33 Tensor;
+		Chaos::FRotation3 RotationOfMass;
+
+		// If no shapes contribute to mass, or they are scaled to zero, we may end up with zero mass here
+		if ((TotalMass > 0.f) && (MassPropertiesList.Num() > 0))
 		{
 			TotalCenterOfMass /= TotalMass;
-		}
 
-		Chaos::PMatrix<float, 3, 3> Tensor;
-		Chaos::FRotation3 RotationOfMass = Chaos::FRotation3::Identity;
-		if (MassPropertiesList.Num())
-		{
 			// NOTE: If multiple items in the list, rotation of mass will be zero, but if only 1 item is the list the item is returned directly and we may have a rotation of mass
 			Chaos::FMassProperties CombinedMassProperties = Chaos::CombineWorldSpace(MassPropertiesList);
 			Tensor = CombinedMassProperties.InertiaTensor;
