Pipedown FShaderParametersMetadata to the shader compile worker

diff --git a/Engine/Shaders/Public/ShaderVersion.ush b/Engine/Shaders/Public/ShaderVersion.ush
--- a/Engine/Shaders/Public/ShaderVersion.ush
+++ b/Engine/Shaders/Public/ShaderVersion.ush
@@ -1,6 +1,5 @@
 // This file is automatically generated by the console command r.InvalidateCachedShaders
 // Each time the console command is executed it generates a new GUID. As this file is included
 // in Platform.ush (which should be included in any shader) it allows to invalidate the shader DDC.
-//
-// If you are merging streams and there is a conflict with this GUID you should make a new GUID rather than taking one or the other.
-// GUID = 163DFF9F-7FD0-7996-9CBF-C17A215DA46B
+// 
+// GUID = 360BAA534983EA395331E2B48CF37A94
\ No newline at end of file
diff --git a/Engine/Source/Developer/ShaderCompilerCommon/Private/ShaderCompilerCommon.cpp b/Engine/Source/Developer/ShaderCompilerCommon/Private/ShaderCompilerCommon.cpp
--- a/Engine/Source/Developer/ShaderCompilerCommon/Private/ShaderCompilerCommon.cpp
+++ b/Engine/Source/Developer/ShaderCompilerCommon/Private/ShaderCompilerCommon.cpp
@@ -474,20 +474,30 @@ bool FShaderParameterParser::ParseAndMoveShaderParametersToRootConstantBuffer(
 	const TCHAR* ConstantBufferType)
 {
 	// The shader doesn't have any parameter binding through shader structure, therefore don't do anything.
-	if (CompilerInput.RootParameterBindings.Num() == 0)
+	if (!CompilerInput.RootParametersStructure)
 	{
 		return true;
 	}
 
 	const bool bMoveToRootConstantBuffer = ConstantBufferType != nullptr;
 	OriginalParsedShader = PreprocessedShaderSource;
-	ParsedParameters.Reserve(CompilerInput.RootParameterBindings.Num());
 
-	// Prepare the set of parameter to look for during parsing.
-	for (const FShaderCompilerInput::FRootParameterBinding& Member : CompilerInput.RootParameterBindings)
+	CompilerInput.RootParametersStructure->IterateShaderParameterMembers(
+		[&](const FShaderParametersMetadata& ParametersMetadata,
+			const FShaderParametersMetadata::FMember& Member,
+			const TCHAR* ShaderBindingName,
+			uint16 ByteOffset)
 	{
-		ParsedParameters.Add(Member.Name, FParsedShaderParameter());
-	}
+		if (
+			Member.GetBaseType() != UBMT_INT32 &&
+			Member.GetBaseType() != UBMT_UINT32 &&
+			Member.GetBaseType() != UBMT_FLOAT32)
+		{
+			return;
+		}
+
+		ParsedParameters.Add(ShaderBindingName, FParsedShaderParameter());
+	});
 
 	bool bSuccess = true;
 
@@ -886,41 +896,52 @@ bool FShaderParameterParser::ParseAndMoveShaderParametersToRootConstantBuffer(
 	if (bMoveToRootConstantBuffer)
 	{
 		FString RootCBufferContent;
-		for (const auto& Member : CompilerInput.RootParameterBindings)
+
+		CompilerInput.RootParametersStructure->IterateShaderParameterMembers(
+			[&](const FShaderParametersMetadata& ParametersMetadata,
+				const FShaderParametersMetadata::FMember& Member,
+				const TCHAR* ShaderBindingName,
+				uint16 ByteOffset)
 		{
-			const FParsedShaderParameter& ParsedParameter = ParsedParameters[Member.Name];
-			if (!ParsedParameter.IsFound())
+			if (
+				Member.GetBaseType() != UBMT_INT32 &&
+				Member.GetBaseType() != UBMT_UINT32 &&
+				Member.GetBaseType() != UBMT_FLOAT32)
 			{
-				continue;
+				return;
 			}
 
-			FString HLSLOffset;
+			const FParsedShaderParameter& ParsedParameter = ParsedParameters[ShaderBindingName];
+
+			if (ParsedParameter.IsFound())
 			{
-				int32 ByteOffset = int32(Member.ByteOffset);
-				HLSLOffset = FString::FromInt(ByteOffset / 16);
-			
-				switch (ByteOffset % 16)
+				FString HLSLOffset;
 				{
-				case 0:
-					break;
-				case 4:
-					HLSLOffset.Append(TEXT(".y"));
-					break;
-				case 8:
-					HLSLOffset.Append(TEXT(".z"));
-					break;
-				case 12:
-					HLSLOffset.Append(TEXT(".w"));
-					break;
+					HLSLOffset = FString::FromInt(ByteOffset / 16);
+
+					switch (ByteOffset % 16)
+					{
+					case 0:
+						break;
+					case 4:
+						HLSLOffset.Append(TEXT(".y"));
+						break;
+					case 8:
+						HLSLOffset.Append(TEXT(".z"));
+						break;
+					case 12:
+						HLSLOffset.Append(TEXT(".w"));
+						break;
+					}
 				}
-			}
 
-			RootCBufferContent.Append(FString::Printf(
-				TEXT("%s %s : packoffset(c%s);\r\n"),
-				*ParsedParameter.Type,
-				*Member.Name,
-				*HLSLOffset));
-		}
+				RootCBufferContent.Append(FString::Printf(
+					TEXT("%s %s : packoffset(c%s);\r\n"),
+					*ParsedParameter.Type,
+					ShaderBindingName,
+					*HLSLOffset));
+			}
+		});
 
 		FString NewShaderCode = FString::Printf(
 			TEXT("%s %s\r\n")
@@ -943,7 +964,7 @@ void FShaderParameterParser::ValidateShaderParameterTypes(
 	FShaderCompilerOutput& CompilerOutput) const
 {
 	// The shader doesn't have any parameter binding through shader structure, therefore don't do anything.
-	if (CompilerInput.RootParameterBindings.Num() == 0)
+	if (!CompilerInput.RootParametersStructure)
 	{
 		return;
 	}
@@ -953,42 +974,59 @@ void FShaderParameterParser::ValidateShaderParameterTypes(
 		return;
 	}
 
+	EShaderPlatform ShaderPlatform = EShaderPlatform(CompilerInput.Target.Platform);
+
 	const TMap<FString, FParameterAllocation>& ParametersFoundByCompiler = CompilerOutput.ParameterMap.GetParameterMap();
 
 	bool bSuccess = true;
-	for (const FShaderCompilerInput::FRootParameterBinding& Member : CompilerInput.RootParameterBindings)
+	CompilerInput.RootParametersStructure->IterateShaderParameterMembers(
+		[&](const FShaderParametersMetadata& ParametersMetadata,
+			const FShaderParametersMetadata::FMember& Member,
+			const TCHAR* ShaderBindingName,
+			uint16 ByteOffset)
 	{
-		const FParsedShaderParameter& ParsedParameter = ParsedParameters[Member.Name];
+		if (
+			Member.GetBaseType() != UBMT_INT32 &&
+			Member.GetBaseType() != UBMT_UINT32 &&
+			Member.GetBaseType() != UBMT_FLOAT32)
+		{
+			return;
+		}
+
+		const FParsedShaderParameter& ParsedParameter = ParsedParameters[ShaderBindingName];
 
 		// Did not find shader parameter in code.
 		if (!ParsedParameter.IsFound())
 		{
 			// Verify the shader compiler also did not find this parameter to make sure there is no bug in the parser.
 			checkf(
-				!ParametersFoundByCompiler.Contains(Member.Name),
+				!ParametersFoundByCompiler.Contains(ShaderBindingName),
 				TEXT("Looks like there is a bug in FShaderParameterParser ParameterName=%s DumpDebugInfoPath=%s"),
-				*Member.Name,
+				ShaderBindingName,
 				*CompilerInput.DumpDebugInfoPath);
-			continue;
+			return;
 		}
 
-		const bool bShouldBeInt = Member.ExpectedShaderType.StartsWith(TEXT("int"));
-		const bool bShouldBeUint = Member.ExpectedShaderType.StartsWith(TEXT("uint"));
+		FString ExpectedShaderType;
+		Member.GenerateShaderParameterType(ExpectedShaderType, ShaderPlatform);
+
+		const bool bShouldBeInt = Member.GetBaseType() == UBMT_INT32;
+		const bool bShouldBeUint = Member.GetBaseType() == UBMT_UINT32;
 
 		// Match parsed type with expected shader type
-		bool bIsTypeCorrect = ParsedParameter.Type == Member.ExpectedShaderType;
+		bool bIsTypeCorrect = ParsedParameter.Type == ExpectedShaderType;
 		
 		if (!bIsTypeCorrect)
 		{
 			// Accept half-precision floats when single-precision was requested
-			if (ParsedParameter.Type.StartsWith(TEXT("half")) && Member.ExpectedShaderType.StartsWith(TEXT("float")))
+			if (ParsedParameter.Type.StartsWith(TEXT("half")) && Member.GetBaseType() == UBMT_FLOAT32)
 			{
-				bIsTypeCorrect = (FCString::Strcmp(*ParsedParameter.Type + 4, *Member.ExpectedShaderType + 5) == 0);
+				bIsTypeCorrect = (FCString::Strcmp(*ParsedParameter.Type + 4, *ExpectedShaderType + 5) == 0);
 			}
 			// Accept single-precision floats when half-precision was expected
-			else if (ParsedParameter.Type.StartsWith(TEXT("float")) && Member.ExpectedShaderType.StartsWith(TEXT("half")))
+			else if (ParsedParameter.Type.StartsWith(TEXT("float")) && ExpectedShaderType.StartsWith(TEXT("half")))
 			{
-				bIsTypeCorrect = (FCString::Strcmp(*ParsedParameter.Type + 5, *Member.ExpectedShaderType + 4) == 0);
+				bIsTypeCorrect = (FCString::Strcmp(*ParsedParameter.Type + 5, *ExpectedShaderType + 4) == 0);
 			}
 		}
 
@@ -999,12 +1037,12 @@ void FShaderParameterParser::ValidateShaderParameterTypes(
 			if (bShouldBeInt)
 			{
 				// tries up with an uint.
-				NewExpectedShaderType = TEXT("u") + Member.ExpectedShaderType;
+				NewExpectedShaderType = TEXT("u") + ExpectedShaderType;
 			}
 			else
 			{
 				// tries up with an int.
-				NewExpectedShaderType = Member.ExpectedShaderType;
+				NewExpectedShaderType = ExpectedShaderType;
 				NewExpectedShaderType.RemoveAt(0);
 			}
 
@@ -1017,14 +1055,14 @@ void FShaderParameterParser::ValidateShaderParameterTypes(
 			Error.StrippedErrorMessage = FString::Printf(
 				TEXT("Type %s of shader parameter %s in shader mismatch the shader parameter structure: it expects a %s"),
 				*ParsedParameter.Type,
-				*Member.Name,
-				*Member.ExpectedShaderType);
+				ShaderBindingName,
+				*ExpectedShaderType);
 			ExtractFileAndLine(ParsedParameter.PragamLineoffset, ParsedParameter.LineOffset, Error.ErrorVirtualFilePath, Error.ErrorLineString);
 
 			CompilerOutput.Errors.Add(Error);
 			bSuccess = false;
 		}
-	} // for (const auto& Member : CompilerInput.RootParameterBindings)
+	});
 
 	CompilerOutput.bSucceeded = bSuccess;
 }
diff --git a/Engine/Source/Programs/ShaderCompileWorker/Private/ShaderCompileWorker.cpp b/Engine/Source/Programs/ShaderCompileWorker/Private/ShaderCompileWorker.cpp
--- a/Engine/Source/Programs/ShaderCompileWorker/Private/ShaderCompileWorker.cpp
+++ b/Engine/Source/Programs/ShaderCompileWorker/Private/ShaderCompileWorker.cpp
@@ -375,10 +375,29 @@ private:
 		// Initialize shader hash cache before reading any includes.
 		InitializeShaderHashCache();
 
-		TMap<FString, FThreadSafeSharedStringPtr> ExternalIncludes;
-		TArray<FShaderCompilerEnvironment> SharedEnvironments;
+		// Array of string used as const TCHAR* during compilation process.
+		TArray<TUniquePtr<FString>> AllocatedStrings;
+		auto DeserializeConstTCHAR = [&AllocatedStrings](FArchive& Archive)
+		{
+			FString Name;
+			Archive << Name;
+
+			const TCHAR* CharName = nullptr;
+			if (Name.Len() != 0)
+			{
+				if (AllocatedStrings.GetSlack() == 0)
+				{
+					AllocatedStrings.Reserve(AllocatedStrings.Num() + 1024);
+				}
+
+				AllocatedStrings.Add(MakeUnique<FString>(Name));
+				CharName = **AllocatedStrings.Last();
+			}
+			return CharName;
+		};
 
 		// Shared inputs
+		TMap<FString, FThreadSafeSharedStringPtr> ExternalIncludes;
 		{
 			int32 NumExternalIncludes = 0;
 			InputFile << NumExternalIncludes;
@@ -392,7 +411,11 @@ private:
 				InputFile << (*NewIncludeContents);
 				ExternalIncludes.Add(NewIncludeName, MakeShareable(NewIncludeContents));
 			}
+		}
 
+		// Shared environments
+		TArray<FShaderCompilerEnvironment> SharedEnvironments;
+		{
 			int32 NumSharedEnvironments = 0;
 			InputFile << NumSharedEnvironments;
 			SharedEnvironments.Empty(NumSharedEnvironments);
@@ -404,6 +427,97 @@ private:
 			}
 		}
 
+		// All the shader parameter structures
+		// Note: this is a bit more complicated, purposefully to avoid switch const TCHAR* to FString in runtime FShaderParametersMetadata.
+		TArray<TUniquePtr<FShaderParametersMetadata>> ParameterStructures;
+		{
+			int32 NumParameterStructures = 0;
+			InputFile << NumParameterStructures;
+			ParameterStructures.Reserve(NumParameterStructures);
+
+			for (int32 StructIndex = 0; StructIndex < NumParameterStructures; StructIndex++)
+			{
+				const TCHAR* LayoutName;
+				const TCHAR* StructTypeName;
+				const TCHAR* ShaderVariableName;
+				FShaderParametersMetadata::EUseCase UseCase;
+				int32 StructFileLine;
+				uint32 Size;
+				int32 MemberCount;
+
+				LayoutName = DeserializeConstTCHAR(InputFile);
+				StructTypeName = DeserializeConstTCHAR(InputFile);
+				ShaderVariableName = DeserializeConstTCHAR(InputFile);
+				InputFile << UseCase;
+				InputFile << StructFileLine;
+				InputFile << Size;
+				InputFile << MemberCount;
+
+				TArray<FShaderParametersMetadata::FMember> Members;
+				Members.Reserve(MemberCount);
+
+				for (int32 MemberIndex = 0; MemberIndex < MemberCount; MemberIndex++)
+				{
+					const TCHAR* Name;
+					const TCHAR* ShaderType;
+					int32 FileLine;
+					uint32 Offset;
+					uint8 BaseType;
+					uint8 PrecisionModifier;
+					uint32 NumRows;
+					uint32 NumColumns;
+					uint32 NumElements;
+					int32 StructMetadataIndex;
+
+					static_assert(sizeof(BaseType) == sizeof(EUniformBufferBaseType), "Cast failure.");
+					static_assert(sizeof(PrecisionModifier) == sizeof(EShaderPrecisionModifier::Type), "Cast failure.");
+
+					Name = DeserializeConstTCHAR(InputFile);
+					ShaderType = DeserializeConstTCHAR(InputFile);
+					InputFile << FileLine;
+					InputFile << Offset;
+					InputFile << BaseType;
+					InputFile << PrecisionModifier;
+					InputFile << NumRows;
+					InputFile << NumColumns;
+					InputFile << NumElements;
+					InputFile << StructMetadataIndex;
+
+					const FShaderParametersMetadata* StructMetadata = nullptr;
+					if (StructMetadataIndex != INDEX_NONE)
+					{
+						StructMetadata = ParameterStructures[StructMetadataIndex].Get();
+					}
+
+					FShaderParametersMetadata::FMember Member(
+						Name,
+						ShaderType,
+						FileLine,
+						Offset,
+						EUniformBufferBaseType(BaseType),
+						EShaderPrecisionModifier::Type(PrecisionModifier),
+						NumRows,
+						NumColumns,
+						NumElements,
+						StructMetadata);
+					Members.Add(Member);
+				}
+
+				ParameterStructures.Add(MakeUnique<FShaderParametersMetadata>(
+					UseCase,
+					EUniformBufferBindingFlags::Shader,
+					/* InLayoutName = */ LayoutName,
+					/* InStructTypeName = */ StructTypeName,
+					/* InShaderVariableName = */ ShaderVariableName,
+					/* InStaticSlotName = */ nullptr,
+					/* InFileName = */ nullptr,
+					StructFileLine,
+					Size,
+					Members,
+					/* bCompleteInitialization = */ true));
+			}
+		}
+
 		GNumProcessedJobs = 0;
 
 		// Individual jobs
@@ -427,7 +541,7 @@ private:
 				// Deserialize the job's inputs.
 				FShaderCompilerInput CompilerInput;
 				InputFile << CompilerInput;
-				CompilerInput.DeserializeSharedInputs(InputFile, ExternalIncludes, SharedEnvironments);
+				CompilerInput.DeserializeSharedInputs(InputFile, ExternalIncludes, SharedEnvironments, ParameterStructures);
 
 				if (IsValidRef(CompilerInput.SharedEnvironment))
 				{
@@ -473,7 +587,7 @@ private:
 				{
 					// Deserialize the job's inputs.
 					InputFile << CompilerInputs[StageIndex];
-					CompilerInputs[StageIndex].DeserializeSharedInputs(InputFile, ExternalIncludes, SharedEnvironments);
+					CompilerInputs[StageIndex].DeserializeSharedInputs(InputFile, ExternalIncludes, SharedEnvironments, ParameterStructures);
 
 					if (IsValidRef(CompilerInputs[StageIndex].SharedEnvironment))
 					{
diff --git a/Engine/Source/Runtime/Engine/Private/ShaderCompiler/ShaderCompiler.cpp b/Engine/Source/Runtime/Engine/Private/ShaderCompiler/ShaderCompiler.cpp
--- a/Engine/Source/Runtime/Engine/Private/ShaderCompiler/ShaderCompiler.cpp
+++ b/Engine/Source/Runtime/Engine/Private/ShaderCompiler/ShaderCompiler.cpp
@@ -1098,6 +1098,7 @@ bool FShaderCompileUtilities::DoWriteTasks(const TArray<FShaderCommonCompileJobP
 	SplitJobsByType(QueuedJobs, QueuedSingleJobs, QueuedPipelineJobs);
 
 	TArray<TRefCountPtr<FSharedShaderCompilerEnvironment>> SharedEnvironments;
+	TArray<const FShaderParametersMetadata*> RequestShaderParameterStructures;
 
 	// Gather External Includes and serialize separately, these are largely shared between jobs
 	{
@@ -1106,7 +1107,7 @@ bool FShaderCompileUtilities::DoWriteTasks(const TArray<FShaderCommonCompileJobP
 
 		for (int32 JobIndex = 0; JobIndex < QueuedSingleJobs.Num(); JobIndex++)
 		{
-			QueuedSingleJobs[JobIndex]->Input.GatherSharedInputs(ExternalIncludes, SharedEnvironments);
+			QueuedSingleJobs[JobIndex]->Input.GatherSharedInputs(ExternalIncludes, SharedEnvironments, RequestShaderParameterStructures);
 		}
 
 		for (int32 JobIndex = 0; JobIndex < QueuedPipelineJobs.Num(); JobIndex++)
@@ -1116,7 +1117,7 @@ bool FShaderCompileUtilities::DoWriteTasks(const TArray<FShaderCommonCompileJobP
 
 			for (int32 Index = 0; Index < NumStageJobs; Index++)
 			{
-				PipelineJob->StageJobs[Index]->Input.GatherSharedInputs(ExternalIncludes, SharedEnvironments);
+				PipelineJob->StageJobs[Index]->Input.GatherSharedInputs(ExternalIncludes, SharedEnvironments, RequestShaderParameterStructures);
 			}
 		}
 
@@ -1138,6 +1139,77 @@ bool FShaderCompileUtilities::DoWriteTasks(const TArray<FShaderCommonCompileJobP
 		}
 	}
 
+	// Write shader parameter structures
+	TArray<const FShaderParametersMetadata*> AllShaderParameterStructures;
+	{
+		// List all dependencies.
+		for (int32 StructId = 0; StructId < RequestShaderParameterStructures.Num(); StructId++)
+		{
+			RequestShaderParameterStructures[StructId]->IterateStructureMetadataDependencies(
+				[&](const FShaderParametersMetadata* Struct)
+			{
+				AllShaderParameterStructures.AddUnique(Struct);
+			});
+		}
+
+		// Write all shader parameter structure.
+		int32 NumParameterStructures = AllShaderParameterStructures.Num();
+		TransferFile << NumParameterStructures;
+		for (const FShaderParametersMetadata* Struct : AllShaderParameterStructures)
+		{
+			FString LayoutName = Struct->GetLayout().GetDebugName();
+			FString StructTypeName = Struct->GetStructTypeName();
+			FString ShaderVariableName = Struct->GetShaderVariableName();
+			uint8 UseCase = uint8(Struct->GetUseCase());
+			int32 StructFileLine = Struct->GetFileLine();
+			uint32 Size = Struct->GetSize();
+			int32 MemberCount = Struct->GetMembers().Num();
+
+			static_assert(sizeof(UseCase) == sizeof(FShaderParametersMetadata::EUseCase), "Cast failure.");
+
+			TransferFile << LayoutName;
+			TransferFile << StructTypeName;
+			TransferFile << ShaderVariableName;
+			TransferFile << UseCase;
+			TransferFile << StructFileLine;
+			TransferFile << Size;
+			TransferFile << MemberCount;
+
+			for (const FShaderParametersMetadata::FMember& Member : Struct->GetMembers())
+			{
+				FString Name = Member.GetName();
+				FString ShaderType = Member.GetShaderType();
+				int32 FileLine = Member.GetFileLine();
+				uint32 Offset = Member.GetOffset();
+				uint8 BaseType = uint8(Member.GetBaseType());
+				uint8 PrecisionModifier = uint8(Member.GetPrecision());
+				uint32 NumRows = Member.GetNumRows();
+				uint32 NumColumns = Member.GetNumColumns();
+				uint32 NumElements = Member.GetNumElements();
+				int32 StructMetadataIndex = INDEX_NONE;
+				if (Member.GetStructMetadata())
+				{
+					StructMetadataIndex = AllShaderParameterStructures.Find(Member.GetStructMetadata());
+					check(StructMetadataIndex != INDEX_NONE);
+				}
+
+				static_assert(sizeof(BaseType) == sizeof(EUniformBufferBaseType), "Cast failure.");
+				static_assert(sizeof(PrecisionModifier) == sizeof(EShaderPrecisionModifier::Type), "Cast failure.");
+
+				TransferFile << Name;
+				TransferFile << ShaderType;
+				TransferFile << FileLine;
+				TransferFile << Offset;
+				TransferFile << BaseType;
+				TransferFile << PrecisionModifier;
+				TransferFile << NumRows;
+				TransferFile << NumColumns;
+				TransferFile << NumElements;
+				TransferFile << StructMetadataIndex;
+			}
+		}
+	}
+
 	// Write individual shader jobs
 	{
 		int32 SingleJobHeader = ShaderCompileWorkerSingleJobHeader;
@@ -1150,7 +1222,7 @@ bool FShaderCompileUtilities::DoWriteTasks(const TArray<FShaderCommonCompileJobP
 		for (int32 JobIndex = 0; JobIndex < QueuedSingleJobs.Num(); JobIndex++)
 		{
 			TransferFile << QueuedSingleJobs[JobIndex]->Input;
-			QueuedSingleJobs[JobIndex]->Input.SerializeSharedInputs(TransferFile, SharedEnvironments);
+			QueuedSingleJobs[JobIndex]->Input.SerializeSharedInputs(TransferFile, SharedEnvironments, AllShaderParameterStructures);
 		}
 	}
 
@@ -1171,7 +1243,7 @@ bool FShaderCompileUtilities::DoWriteTasks(const TArray<FShaderCommonCompileJobP
 			for (int32 Index = 0; Index < NumStageJobs; Index++)
 			{
 				TransferFile << PipelineJob->StageJobs[Index]->GetSingleShaderJob()->Input;
-				PipelineJob->StageJobs[Index]->Input.SerializeSharedInputs(TransferFile, SharedEnvironments);
+				PipelineJob->StageJobs[Index]->Input.SerializeSharedInputs(TransferFile, SharedEnvironments, AllShaderParameterStructures);
 			}
 		}
 	}
@@ -4354,72 +4426,6 @@ void ValidateShaderFilePath(const FString& VirtualShaderFilePath, const FString&
 		*VirtualShaderFilePath);
 }
 
-static void PullRootShaderParametersLayout(FShaderCompilerInput& CompileInput, EShaderPlatform ShaderPlatform, const FShaderParametersMetadata& ParametersMetadata, uint16 ByteOffset, const FString& Prefix)
-{
-	for (const FShaderParametersMetadata::FMember& Member : ParametersMetadata.GetMembers())
-	{
-		EUniformBufferBaseType BaseType = Member.GetBaseType();
-		uint16 MemberOffset = ByteOffset + uint16(Member.GetOffset());
-		uint32 NumElements = Member.GetNumElements();
-
-		if (BaseType == UBMT_INCLUDED_STRUCT)
-		{
-			check(NumElements == 0);
-			PullRootShaderParametersLayout(CompileInput, ShaderPlatform, *Member.GetStructMetadata(), MemberOffset, Prefix);
-		}
-		else if (BaseType == UBMT_NESTED_STRUCT && NumElements == 0)
-		{
-			FString NewPrefix = FString::Printf(TEXT("%s%s_"), *Prefix, Member.GetName());
-			PullRootShaderParametersLayout(CompileInput, ShaderPlatform, *Member.GetStructMetadata(), MemberOffset, NewPrefix);
-		}
-		else if (BaseType == UBMT_NESTED_STRUCT && NumElements > 0)
-		{
-			for (uint32 ArrayElementId = 0; ArrayElementId < NumElements; ArrayElementId++)
-			{
-				FString NewPrefix = FString::Printf(TEXT("%s%s_%u_"), *Prefix, Member.GetName(), ArrayElementId);
-				PullRootShaderParametersLayout(CompileInput, ShaderPlatform, *Member.GetStructMetadata(), MemberOffset, NewPrefix);
-			}
-		}
-		else if (
-			BaseType == UBMT_INT32 ||
-			BaseType == UBMT_UINT32 ||
-			BaseType == UBMT_FLOAT32)
-		{
-			FShaderCompilerInput::FRootParameterBinding RootParameterBinding;
-			RootParameterBinding.Name = FString::Printf(TEXT("%s%s"), *Prefix, Member.GetName());
-			Member.GenerateShaderParameterType(RootParameterBinding.ExpectedShaderType, ShaderPlatform);
-			RootParameterBinding.ByteOffset = MemberOffset;
-			CompileInput.RootParameterBindings.Add(RootParameterBinding);
-		}
-		continue;
-
-		if (BaseType == UBMT_REFERENCED_STRUCT)
-		{
-			// Referenced structured are manually passed to the RHI.
-		}
-		else if (BaseType == UBMT_RENDER_TARGET_BINDING_SLOTS)
-		{
-			// RHI don't need to care about render target bindings slot anyway.
-		}
-		else if (
-			BaseType == UBMT_RDG_BUFFER_ACCESS ||
-			BaseType == UBMT_RDG_TEXTURE_ACCESS)
-		{
-			// Shaders don't care about RDG access parameters.
-		}
-		else if (
-			BaseType == UBMT_RDG_BUFFER_UAV ||
-			BaseType == UBMT_RDG_TEXTURE_UAV)
-		{
-			// UAV are ignored on purpose because not supported in uniform buffers.
-		}
-		else
-		{
-			check(0);
-		}
-	}
-}
-
 FThreadSafeSharedStringPtr GCachedGeneratedInstancedStereoCode = MakeShareable(new FString());
 
 void GlobalBeginCompileShader(
@@ -4491,11 +4497,7 @@ void GlobalBeginCompileShader(
 	Input.DebugGroupName = DebugGroupName;
 	Input.DebugDescription = DebugDescription;
 	Input.DebugExtension = DebugExtension;
-
-	if (ShaderType->GetRootParametersMetadata())
-	{
-		PullRootShaderParametersLayout(Input, ShaderPlatform, *ShaderType->GetRootParametersMetadata(), /* ByteOffset = */ 0, FString());
-	}
+	Input.RootParametersStructure = ShaderType->GetRootParametersMetadata();
 
 	// Verify FShaderCompilerInput's file paths are consistent. 
 	#if DO_CHECK
diff --git a/Engine/Source/Runtime/RenderCore/Private/ShaderCore.cpp b/Engine/Source/Runtime/RenderCore/Private/ShaderCore.cpp
--- a/Engine/Source/Runtime/RenderCore/Private/ShaderCore.cpp
+++ b/Engine/Source/Runtime/RenderCore/Private/ShaderCore.cpp
@@ -1443,7 +1443,6 @@ FArchive& operator<<(FArchive& Ar, FShaderCompilerInput& Input)
 	Ar << Input.DebugDescription;
 	Ar << Input.Environment;
 	Ar << Input.ExtraSettings;
-	Ar << Input.RootParameterBindings;
 
 	// Note: skipping Input.SharedEnvironment, which is handled by FShaderCompileUtilities::DoWriteTasks in order to maintain sharing
 
diff --git a/Engine/Source/Runtime/RenderCore/Private/ShaderParameterMetadata.cpp b/Engine/Source/Runtime/RenderCore/Private/ShaderParameterMetadata.cpp
--- a/Engine/Source/Runtime/RenderCore/Private/ShaderParameterMetadata.cpp
+++ b/Engine/Source/Runtime/RenderCore/Private/ShaderParameterMetadata.cpp
@@ -282,7 +282,8 @@ FShaderParametersMetadata::FShaderParametersMetadata(
 	const ANSICHAR* InFileName,
 	const int32 InFileLine,
 	uint32 InSize,
-	const TArray<FMember>& InMembers)
+	const TArray<FMember>& InMembers,
+	bool bForceCompleteInitialization)
 	: StructTypeName(InStructTypeName)
 	, ShaderVariableName(InShaderVariableName)
 	, StaticSlotName(InStaticSlotName)
@@ -311,7 +312,7 @@ FShaderParametersMetadata::FShaderParametersMetadata(
 		check(ShaderVariableName);
 	}
 
-	if (UseCase == EUseCase::UniformBuffer)
+	if (UseCase == EUseCase::UniformBuffer && !bForceCompleteInitialization)
 	{
 		// Register this uniform buffer struct in global list.
 		GlobalListLink.LinkHead(GetStructList());
@@ -352,8 +353,9 @@ FShaderParametersMetadata::FShaderParametersMetadata(
 
 FShaderParametersMetadata::~FShaderParametersMetadata()
 {
-	if (UseCase == EUseCase::UniformBuffer)
+	if (GlobalListLink.IsLinked())
 	{
+		check(UseCase == EUseCase::UniformBuffer);
 		GlobalListLink.Unlink();
 		GetNameStructMap().Remove(ShaderVariableHashedName);
 
diff --git a/Engine/Source/Runtime/RenderCore/Public/ShaderCompilerCore.h b/Engine/Source/Runtime/RenderCore/Public/ShaderCompilerCore.h
--- a/Engine/Source/Runtime/RenderCore/Public/ShaderCompilerCore.h
+++ b/Engine/Source/Runtime/RenderCore/Public/ShaderCompilerCore.h
@@ -16,7 +16,7 @@
 class Error;
 
 // this is for the protocol, not the data, bump if FShaderCompilerInput or ProcessInputFromArchive changes.
-const int32 ShaderCompileWorkerInputVersion = 13;
+const int32 ShaderCompileWorkerInputVersion = 14;
 // this is for the protocol, not the data, bump if FShaderCompilerOutput or WriteToOutputArchive changes.
 const int32 ShaderCompileWorkerOutputVersion = 6;
 // this is for the protocol, not the data.
@@ -176,28 +176,9 @@ struct FShaderCompilerInput
 	FShaderCompilerEnvironment Environment;
 	TRefCountPtr<FSharedShaderCompilerEnvironment> SharedEnvironment;
 
-
-	struct FRootParameterBinding
-	{
-		/** Name of the constant buffer stored parameter. */
-		FString Name;
-
-		/** Type expected in the shader code to ensure the binding is bug free. */
-		FString ExpectedShaderType;
-
-		/** The offset of the parameter in the root shader parameter struct. */
-		uint16 ByteOffset;
-
-		friend FArchive& operator<<(FArchive& Ar, FRootParameterBinding& RootParameterBinding)
-		{
-			Ar << RootParameterBinding.Name;
-			Ar << RootParameterBinding.ExpectedShaderType;
-			Ar << RootParameterBinding.ByteOffset;
-			return Ar;
-		}
-	};
-
-	TArray<FRootParameterBinding> RootParameterBindings;
+	// The root of the shader parameter structures / uniform buffers bound to this shader to generate shader resource table from.
+	// This only set if a shader class is defining the 
+	const FShaderParametersMetadata* RootParametersStructure = nullptr;
 
 
 	// Additional compilation settings that can be filled by FMaterial::SetupExtaCompilationSettings
@@ -236,7 +217,10 @@ struct FShaderCompilerInput
 		return FPaths::GetCleanFilename(VirtualSourceFilePath);
 	}
 
-	void GatherSharedInputs(TMap<FString,FString>& ExternalIncludes, TArray<TRefCountPtr<FSharedShaderCompilerEnvironment>>& SharedEnvironments)
+	void GatherSharedInputs(
+		TMap<FString, FString>& ExternalIncludes,
+		TArray<TRefCountPtr<FSharedShaderCompilerEnvironment>>& SharedEnvironments,
+		TArray<const FShaderParametersMetadata*>& ParametersStructures)
 	{
 		check(!SharedEnvironment || SharedEnvironment->IncludeVirtualPathToExternalContentsMap.Num() == 0);
 
@@ -254,9 +238,14 @@ struct FShaderCompilerInput
 		{
 			SharedEnvironments.AddUnique(SharedEnvironment);
 		}
+
+		if (RootParametersStructure)
+		{
+			ParametersStructures.AddUnique(RootParametersStructure);
+		}
 	}
 
-	void SerializeSharedInputs(FArchive& Ar, const TArray<TRefCountPtr<FSharedShaderCompilerEnvironment>>& SharedEnvironments)
+	void SerializeSharedInputs(FArchive& Ar, const TArray<TRefCountPtr<FSharedShaderCompilerEnvironment>>& SharedEnvironments, const TArray<const FShaderParametersMetadata*>& ParametersStructures)
 	{
 		check(Ar.IsSaving());
 
@@ -272,9 +261,21 @@ struct FShaderCompilerInput
 
 		int32 SharedEnvironmentIndex = SharedEnvironments.Find(SharedEnvironment);
 		Ar << SharedEnvironmentIndex;
+
+		int32 ShaderParameterStructureIndex = INDEX_NONE;
+		if (RootParametersStructure)
+		{
+			ShaderParameterStructureIndex = ParametersStructures.Find(RootParametersStructure);
+			check(ShaderParameterStructureIndex != INDEX_NONE);
+		}
+		Ar << ShaderParameterStructureIndex;
 	}
 
-	void DeserializeSharedInputs(FArchive& Ar, const TMap<FString, FThreadSafeSharedStringPtr>& ExternalIncludes, const TArray<FShaderCompilerEnvironment>& SharedEnvironments)
+	void DeserializeSharedInputs(
+		FArchive& Ar,
+		const TMap<FString, FThreadSafeSharedStringPtr>& ExternalIncludes,
+		const TArray<FShaderCompilerEnvironment>& SharedEnvironments,
+		const TArray<TUniquePtr<FShaderParametersMetadata>>& ShaderParameterStructures)
 	{
 		check(Ar.IsLoading());
 
@@ -295,6 +296,13 @@ struct FShaderCompilerInput
 		{
 			Environment.Merge(SharedEnvironments[SharedEnvironmentIndex]);
 		}
+
+		int32 ShaderParameterStructureIndex;
+		Ar << ShaderParameterStructureIndex;
+		if (ShaderParameterStructureIndex != INDEX_NONE)
+		{
+			RootParametersStructure = ShaderParameterStructures[ShaderParameterStructureIndex].Get();
+		}
 	}
 
 	friend RENDERCORE_API FArchive& operator<<(FArchive& Ar, FShaderCompilerInput& Input);
diff --git a/Engine/Source/Runtime/RenderCore/Public/ShaderParameterMetadata.h b/Engine/Source/Runtime/RenderCore/Public/ShaderParameterMetadata.h
--- a/Engine/Source/Runtime/RenderCore/Public/ShaderParameterMetadata.h
+++ b/Engine/Source/Runtime/RenderCore/Public/ShaderParameterMetadata.h
@@ -9,12 +9,14 @@
 #include "CoreMinimal.h"
 #include "Containers/List.h"
 #include "Containers/StaticArray.h"
+#include "Containers/StringFwd.h"
+#include "Misc/StringBuilder.h"
 #include "RHI.h"
 #include "Serialization/MemoryLayout.h"
 
 namespace EShaderPrecisionModifier
 {
-	enum Type
+	enum Type : uint8
 	{
 		Float,
 		Half,
@@ -109,7 +111,7 @@ class RENDERCORE_API FShaderParametersMetadata
 {
 public:
 	/** The use case of the uniform buffer structures. */
-	enum class EUseCase
+	enum class EUseCase : uint8
 	{
 		/** Stand alone shader parameter struct used for render passes and shader parameters. */
 		ShaderParameterStruct,
@@ -214,7 +216,13 @@ public:
 		const FShaderParametersMetadata* Struct;
 	};
 
-	/** Initialization constructor. */
+	/** Initialization constructor.
+	 *
+	 * EUseCase::UniformBuffer are listed in the global GetStructList() that will be visited at engine startup to know all the global uniform buffer
+	 * that can generate code in /Engine/Generated/GeneratedUniformBuffers.ush. Their initialization will be finished during the this list
+	 * traversal. bForceCompleteInitialization force to ignore the list for EUseCase::UniformBuffer and instead handle it like a standalone non
+	 * globally listed EUseCase::ShaderParameterStruct. This is required for the ShaderCompileWorker to deserialize them without side global effects.
+	 */
 	FShaderParametersMetadata(
 		EUseCase UseCase,
 		EUniformBufferBindingFlags InBindingFlags,
@@ -225,7 +233,8 @@ public:
 		const ANSICHAR* InFileName,
 		const int32 InFileLine,
 		uint32 InSize,
-		const TArray<FMember>& InMembers);
+		const TArray<FMember>& InMembers,
+		bool bForceCompleteInitialization = false);
 
 	virtual ~FShaderParametersMetadata();
 
@@ -290,6 +299,38 @@ public:
 		return LayoutHash;	
 	}
 
+	/** Iterate recursively over all shader parameter members. */
+	template<typename TParameterFunction>
+	void IterateShaderParameterMembers(TParameterFunction Lambda) const
+	{
+		TStringBuilder<1024> ShaderBindingNameBuilder;
+		if (UseCase == EUseCase::UniformBuffer || UseCase == EUseCase::DataDrivenUniformBuffer)
+		{
+			ShaderBindingNameBuilder.Append(ShaderVariableName);
+			ShaderBindingNameBuilder.Append(TEXT("_"));
+		}
+
+		FShaderParametersMetadata::IterateShaderParameterMembersInternal(
+			*this, /* ByteOffset = */ 0, ShaderBindingNameBuilder, Lambda);
+	}
+
+	/** Iterate recursively over all FShaderParametersMetadata. */
+	template<typename TParameterFunction>
+	void IterateStructureMetadataDependencies(TParameterFunction Lambda) const
+	{
+		for (const FShaderParametersMetadata::FMember& Member : Members)
+		{
+			const FShaderParametersMetadata* NewParametersMetadata = Member.GetStructMetadata();
+
+			if (NewParametersMetadata)
+			{
+				NewParametersMetadata->IterateStructureMetadataDependencies(Lambda);
+			}
+		}
+
+		Lambda(this);
+	}
+
 private:
 	/** Name of the structure type in C++ and shader code. */
 	const TCHAR* const StructTypeName;
@@ -336,4 +377,70 @@ private:
 	void InitializeLayout();
 
 	void AddResourceTableEntriesRecursive(const TCHAR* UniformBufferName, const TCHAR* Prefix, uint16& ResourceIndex, TMap<FString, FResourceTableEntry>& ResourceTableMap) const;
+
+	template<typename TParameterFunction>
+	static inline void IterateShaderParameterMembersInternal(
+		const FShaderParametersMetadata& ParametersMetadata,
+		uint16 ByteOffset,
+		TStringBuilder<1024>& ShaderBindingNameBuilder,
+		TParameterFunction Lambda)
+	{
+		for (const FShaderParametersMetadata::FMember& Member : ParametersMetadata.GetMembers())
+		{
+			EUniformBufferBaseType BaseType = Member.GetBaseType();
+			uint16 MemberOffset = ByteOffset + uint16(Member.GetOffset());
+			uint32 NumElements = Member.GetNumElements();
+
+			int32 MemberNameLength = FCString::Strlen(Member.GetName());
+
+			if (BaseType == UBMT_INCLUDED_STRUCT)
+			{
+				check(NumElements == 0);
+				const FShaderParametersMetadata& NewParametersMetadata = *Member.GetStructMetadata();
+				IterateShaderParameterMembersInternal(NewParametersMetadata, MemberOffset, ShaderBindingNameBuilder, Lambda);
+			}
+			else if (BaseType == UBMT_NESTED_STRUCT && NumElements == 0)
+			{
+				ShaderBindingNameBuilder.Append(Member.GetName());
+				ShaderBindingNameBuilder.Append(TEXT("_"));
+
+				const FShaderParametersMetadata& NewParametersMetadata = *Member.GetStructMetadata();
+				IterateShaderParameterMembersInternal(NewParametersMetadata, MemberOffset, ShaderBindingNameBuilder, Lambda);
+
+				ShaderBindingNameBuilder.RemoveSuffix(MemberNameLength + 1);
+			}
+			else if (BaseType == UBMT_NESTED_STRUCT && NumElements > 0)
+			{
+				ShaderBindingNameBuilder.Append(Member.GetName());
+				ShaderBindingNameBuilder.Append(TEXT("_"));
+
+				const FShaderParametersMetadata& NewParametersMetadata = *Member.GetStructMetadata();
+				for (uint32 ArrayElementId = 0; ArrayElementId < NumElements; ArrayElementId++)
+				{
+					FString ArrayElementIdString;
+					ArrayElementIdString.AppendInt(ArrayElementId);
+					int32 ArrayElementIdLength = ArrayElementIdString.Len();
+
+					ShaderBindingNameBuilder.Append(ArrayElementIdString);
+					ShaderBindingNameBuilder.Append(TEXT("_"));
+
+					uint16 NewStructOffset = MemberOffset + ArrayElementId * NewParametersMetadata.GetSize();
+					IterateShaderParameterMembersInternal(NewParametersMetadata, NewStructOffset, ShaderBindingNameBuilder, Lambda);
+
+					ShaderBindingNameBuilder.RemoveSuffix(ArrayElementIdLength + 1);
+				}
+
+				ShaderBindingNameBuilder.RemoveSuffix(MemberNameLength + 1);
+			}
+			else
+			{
+				ShaderBindingNameBuilder.Append(Member.GetName());
+
+				const TCHAR* ShaderBindingName = ShaderBindingNameBuilder.ToString();
+				Lambda(ParametersMetadata, Member, ShaderBindingName, MemberOffset);
+
+				ShaderBindingNameBuilder.RemoveSuffix(MemberNameLength);
+			}
+		}
+	}
 };
\ No newline at end of file
