Fixing gpu skinning issues preventing previewing of lower shader models when launching with SM6.

diff --git a/Engine/Source/Runtime/Engine/Private/SkeletalRenderGPUSkin.cpp b/Engine/Source/Runtime/Engine/Private/SkeletalRenderGPUSkin.cpp
--- a/Engine/Source/Runtime/Engine/Private/SkeletalRenderGPUSkin.cpp
+++ b/Engine/Source/Runtime/Engine/Private/SkeletalRenderGPUSkin.cpp
@@ -56,9 +56,9 @@ static FAutoConsoleVariableRef CVarForceUpdateMorphTargets(
 	ECVF_Default
 	);
 
-static bool UseGPUMorphTargets(const EShaderPlatform Platform)
+static bool UseGPUMorphTargets(ERHIFeatureLevel::Type FeatureLevel)
 {
-	return GUseGPUMorphTargets != 0 && IsFeatureLevelSupported(Platform, ERHIFeatureLevel::SM5);
+	return GUseGPUMorphTargets != 0 && FeatureLevel >= ERHIFeatureLevel::SM5;
 }
 
 static float GMorphTargetWeightThreshold = SMALL_NUMBER;
@@ -83,7 +83,8 @@ void FMorphVertexBuffer::InitDynamicRHI()
 	uint32 Size = LodData.GetNumVertices() * sizeof(FMorphGPUSkinVertex);
 	FRHIResourceCreateInfo CreateInfo(TEXT("MorphVertexBuffer"));
 
-	const bool bUseGPUMorphTargets = UseGPUMorphTargets(GMaxRHIShaderPlatform);
+	const EShaderPlatform ShaderPlatform = GetFeatureLevelShaderPlatform(FeatureLevel);
+	const bool bUseGPUMorphTargets = UseGPUMorphTargets(FeatureLevel);
 	bUsesComputeShader = bUseGPUMorphTargets;
 
 #if PLATFORM_PS4
@@ -97,7 +98,7 @@ void FMorphVertexBuffer::InitDynamicRHI()
 	Flags = (EBufferUsageFlags)(Flags | BUF_ShaderResource);
 
 	VertexBufferRHI = RHICreateVertexBuffer(Size, Flags, CreateInfo);
-	bool bUsesSkinCache = RHISupportsComputeShaders(GMaxRHIShaderPlatform) && IsGPUSkinCacheAvailable(GMaxRHIShaderPlatform) && GEnableGPUSkinCache;
+	bool bUsesSkinCache = RHISupportsComputeShaders(ShaderPlatform) && IsGPUSkinCacheAvailable(ShaderPlatform) && GEnableGPUSkinCache;
 	if (bUsesSkinCache)
 	{
 		SRVValue = RHICreateShaderResourceView(VertexBufferRHI, 4, PF_R32_FLOAT);
@@ -151,7 +152,7 @@ FSkeletalMeshObjectGPUSkin::FSkeletalMeshObjectGPUSkin(USkinnedMeshComponent* In
 	LODs.Empty(SkeletalMeshRenderData->LODRenderData.Num());
 	for( int32 LODIndex=0;LODIndex < SkeletalMeshRenderData->LODRenderData.Num();LODIndex++ )
 	{
-		new(LODs) FSkeletalMeshObjectLOD(SkeletalMeshRenderData, LODIndex);
+		new(LODs) FSkeletalMeshObjectLOD(SkeletalMeshRenderData, LODIndex, InFeatureLevel);
 	}
 
 	InitResources(InMeshComponent);
@@ -533,7 +534,7 @@ void FSkeletalMeshObjectGPUSkin::ProcessUpdatedDynamicData(FGPUSkinCache* GPUSki
 		if(bMorphNeedsUpdate)
 		{
 			QUICK_SCOPE_CYCLE_COUNTER(STAT_FSkeletalMeshObjectGPUSkin_ProcessUpdatedDynamicData_UpdateMorphBuffer);
-			if (UseGPUMorphTargets(GMaxRHIShaderPlatform))
+			if (UseGPUMorphTargets(FeatureLevel))
 			{
 				// sometimes this goes out of bound, we'll ensure here
 				ensureAlways(DynamicData->MorphTargetWeights.Num() == LODData.MorphTargetVertexInfoBuffers.GetNumMorphs());
@@ -569,7 +570,7 @@ void FSkeletalMeshObjectGPUSkin::ProcessUpdatedDynamicData(FGPUSkinCache* GPUSki
 		{
 			const FSkelMeshRenderSection& Section = Sections[SectionIdx];
 
-			bool bClothFactory = RHISupportsManualVertexFetch(GMaxRHIShaderPlatform) && (FeatureLevel >= ERHIFeatureLevel::SM5) &&
+			bool bClothFactory = RHISupportsManualVertexFetch(GetFeatureLevelShaderPlatform(FeatureLevel)) && (FeatureLevel >= ERHIFeatureLevel::SM5) &&
 								 (DynamicData->ClothingSimData.Num() > 0) && Section.HasClothingData();
 
 			FGPUBaseSkinVertexFactory* VertexFactory;
@@ -823,7 +824,7 @@ void FSkeletalMeshObjectGPUSkin::FSkeletalMeshObjectLOD::UpdateMorphVertexBuffer
 		// LOD of the skel mesh is used to find number of vertices in buffer
 		FSkeletalMeshLODRenderData& LodData = SkelMeshRenderData->LODRenderData[LODIndex];
 
-		const bool bUseGPUMorphTargets = UseGPUMorphTargets(GMaxRHIShaderPlatform);
+		const bool bUseGPUMorphTargets = UseGPUMorphTargets(FeatureLevel);
 		MorphVertexBuffer.RecreateResourcesIfRequired(bUseGPUMorphTargets);
 
 		SCOPED_GPU_STAT(RHICmdList, MorphTargets);
@@ -864,7 +865,7 @@ void FSkeletalMeshObjectGPUSkin::FSkeletalMeshObjectLOD::UpdateMorphVertexBuffer
 				//the first pass scatters all morph targets into the vertexbuffer using atomics
 				//multiple morph targets can be batched by a single shader where the shader will rely on
 				//binary search to find the correct target weight within the batch.
-				TShaderMapRef<FGPUMorphUpdateCS> GPUMorphUpdateCS(GetGlobalShaderMap(GMaxRHIFeatureLevel));
+				TShaderMapRef<FGPUMorphUpdateCS> GPUMorphUpdateCS(GetGlobalShaderMap(FeatureLevel));
 				for (uint32 i = 0; i < MorphTargetVertexInfoBuffers.GetNumMorphs();)
 				{
 					uint32 NumMorphDeltas = 0;
@@ -912,7 +913,7 @@ void FSkeletalMeshObjectGPUSkin::FSkeletalMeshObjectLOD::UpdateMorphVertexBuffer
 				//Every vertex is touched only by a single permutation. 
 				//multiple permutations can be batched by a single shader where the shader will rely on
 				//binary search to find the correct target weight within the batch.
-				TShaderMapRef<FGPUMorphNormalizeCS> GPUMorphNormalizeCS(GetGlobalShaderMap(GMaxRHIFeatureLevel));
+				TShaderMapRef<FGPUMorphNormalizeCS> GPUMorphNormalizeCS(GetGlobalShaderMap(FeatureLevel));
 				RHICmdList.SetComputeShader(GPUMorphNormalizeCS.GetComputeShader());
 				GPUMorphNormalizeCS->SetParameters(RHICmdList, InvMorphScale, InvWeightScale, MorphTargetVertexInfoBuffers, MorphVertexBuffer);
 				GPUMorphNormalizeCS->Dispatch(RHICmdList, MorphVertexBuffer.GetNumVerticies());
@@ -989,7 +990,7 @@ void FSkeletalMeshObjectGPUSkin::FSkeletalMeshObjectLOD::UpdateMorphVertexBuffer
 		// LOD of the skel mesh is used to find number of vertices in buffer
 		FSkeletalMeshLODRenderData& LodData = SkelMeshRenderData->LODRenderData[LODIndex];
 		
-		const bool bUseGPUMorphTargets = UseGPUMorphTargets(GMaxRHIShaderPlatform);
+		const bool bUseGPUMorphTargets = UseGPUMorphTargets(FeatureLevel);
 		MorphVertexBuffer.RecreateResourcesIfRequired(bUseGPUMorphTargets);
 
 		uint32 Size = LodData.GetNumVertices() * sizeof(FMorphGPUSkinVertex);
@@ -1352,7 +1353,7 @@ void FGPUSkinPassthroughVertexFactory::SetData(const FDataType& InData)
 	FLocalVertexFactory::SetData(InData);
 	const int32 DefaultBaseVertexIndex = 0;
 	const int32 DefaultPreSkinBaseVertexIndex = 0;
-	if (RHISupportsManualVertexFetch(GMaxRHIShaderPlatform))
+	if (RHISupportsManualVertexFetch(GetFeatureLevelShaderPlatform(GetFeatureLevel())))
 	{
 		UniformBuffer = CreateLocalVFUniformBuffer(this, Data.LODLightmapDataIndex, nullptr, DefaultBaseVertexIndex, DefaultPreSkinBaseVertexIndex);
 	}
@@ -1625,13 +1626,13 @@ void FSkeletalMeshObjectGPUSkin::FVertexFactoryData::InitAPEXClothVertexFactorie
 {
 
 	bool bUseMultipleInfluences = (VertexBuffers.APEXClothVertexBuffer->GetNumVertices() > VertexBuffers.StaticVertexBuffers->PositionVertexBuffer.GetNumVertices());
+	bool bSupportsManualFetch = InFeatureLevel >= ERHIFeatureLevel::SM5 && RHISupportsManualVertexFetch(GetFeatureLevelShaderPlatform(InFeatureLevel));
 
 	// clear existing factories (resources assumed to have been released already)
 	ClothVertexFactories.Empty(Sections.Num());
 	for( int32 FactoryIdx=0; FactoryIdx < Sections.Num(); FactoryIdx++ )
 	{
-		if (Sections[FactoryIdx].HasClothingData() && InFeatureLevel >= ERHIFeatureLevel::SM5 &&
-									 RHISupportsManualVertexFetch(GMaxRHIShaderPlatform))
+		if (Sections[FactoryIdx].HasClothingData() && bSupportsManualFetch)
 		{
 			GPUSkinBoneInfluenceType BoneInfluenceType = VertexBuffers.SkinWeightVertexBuffer->GetBoneInfluenceType();
 			if (BoneInfluenceType == GPUSkinBoneInfluenceType::DefaultBoneInfluence)
diff --git a/Engine/Source/Runtime/Engine/Private/SkeletalRenderGPUSkin.h b/Engine/Source/Runtime/Engine/Private/SkeletalRenderGPUSkin.h
--- a/Engine/Source/Runtime/Engine/Private/SkeletalRenderGPUSkin.h
+++ b/Engine/Source/Runtime/Engine/Private/SkeletalRenderGPUSkin.h
@@ -154,10 +154,11 @@ public:
 	* @param	InSkelMeshRenderData	- render data containing the data for each LOD
 	* @param	InLODIdx				- index of LOD model to use from the parent mesh
 	*/
-	FMorphVertexBuffer(FSkeletalMeshRenderData* InSkelMeshRenderData, int32 InLODIdx)
+	FMorphVertexBuffer(FSkeletalMeshRenderData* InSkelMeshRenderData, int32 InLODIdx, ERHIFeatureLevel::Type InFeatureLevel)
 		:	bHasBeenUpdated(false)	
 		,	bNeedsInitialClear(true)
 		,	LODIdx(InLODIdx)
+		,	FeatureLevel(InFeatureLevel)
 		,	SkelMeshRenderData(InSkelMeshRenderData)
 	{
 		check(SkelMeshRenderData);
@@ -263,6 +264,9 @@ protected:
 private:
 	/** index to the SkelMeshResource.LODModels */
 	int32	LODIdx;
+
+	ERHIFeatureLevel::Type FeatureLevel;
+
 	// parent mesh containing the source data, never 0
 	FSkeletalMeshRenderData* SkelMeshRenderData;
 
@@ -538,10 +542,11 @@ private:
 	/** vertex data for rendering a single LOD */
 	struct FSkeletalMeshObjectLOD
 	{
-		FSkeletalMeshObjectLOD(FSkeletalMeshRenderData* InSkelMeshRenderData,int32 InLOD)
+		FSkeletalMeshObjectLOD(FSkeletalMeshRenderData* InSkelMeshRenderData,int32 InLOD, ERHIFeatureLevel::Type InFeatureLevel)
 			: SkelMeshRenderData(InSkelMeshRenderData)
 			, LODIndex(InLOD)
-			, MorphVertexBuffer(InSkelMeshRenderData,LODIndex)
+			, FeatureLevel(InFeatureLevel)
+			, MorphVertexBuffer(InSkelMeshRenderData,LODIndex, FeatureLevel)
 			, MeshObjectWeightBuffer(nullptr)
 			, MeshObjectColorBuffer(nullptr)
 		{
@@ -585,6 +590,8 @@ private:
 		// index into FSkeletalMeshRenderData::LODRenderData[]
 		int32 LODIndex;
 
+		ERHIFeatureLevel::Type FeatureLevel;
+
 		/** Vertex buffer that stores the morph target vertex deltas. Updated on the CPU */
 		FMorphVertexBuffer MorphVertexBuffer;
 
