- Rename HLSLTree::FConstant to Shader::FValue (HLSLTree is editor-only, FValue is needed at runtime)
- Refactor preshader work in terms of Shader::FValue, support for operations between types other than float4
- Add support for static switches/bools to new HLSL generator (intended to support compiling switches to dynamic branches as option, in future)

diff --git a/Engine/Shaders/Public/ShaderVersion.ush b/Engine/Shaders/Public/ShaderVersion.ush
--- a/Engine/Shaders/Public/ShaderVersion.ush
+++ b/Engine/Shaders/Public/ShaderVersion.ush
@@ -1,5 +1,6 @@
 // This file is automatically generated by the console command r.InvalidateCachedShaders
 // Each time the console command is executed it generates a new GUID. As this file is included
 // in Platform.ush (which should be included in any shader) it allows to invalidate the shader DDC.
-// 
-// GUID = 360BAA534983EA395331E2B48CF37A94
\ No newline at end of file
+//
+// If you are merging streams and there is a conflict with this GUID you should make a new GUID rather than taking one or the other.
+// GUID = A2DAA4AD-E532-470F-A9C2-6BB2951C7368
diff --git a/Engine/Source/Runtime/Engine/Classes/Materials/MaterialExpressionStaticBool.h b/Engine/Source/Runtime/Engine/Classes/Materials/MaterialExpressionStaticBool.h
--- a/Engine/Source/Runtime/Engine/Classes/Materials/MaterialExpressionStaticBool.h
+++ b/Engine/Source/Runtime/Engine/Classes/Materials/MaterialExpressionStaticBool.h
@@ -23,6 +23,8 @@ class UMaterialExpressionStaticBool : public UMaterialExpression
 	virtual int32 Compile(class FMaterialCompiler* Compiler, int32 OutputIndex) override;
 	virtual void GetCaption(TArray<FString>& OutCaptions) const override;
 	virtual uint32 GetOutputType(int32 OutputIndex) override {return MCT_StaticBool;}
+
+	virtual EMaterialGenerateHLSLStatus GenerateHLSLExpression(FMaterialHLSLGenerator& Generator, UE::HLSLTree::FScope& Scope, int32 OutputIndex, UE::HLSLTree::FExpression*& OutExpression) override;
 #endif
 	//~ End UMaterialExpression Interface
 };
diff --git a/Engine/Source/Runtime/Engine/Classes/Materials/MaterialExpressionStaticBoolParameter.h b/Engine/Source/Runtime/Engine/Classes/Materials/MaterialExpressionStaticBoolParameter.h
--- a/Engine/Source/Runtime/Engine/Classes/Materials/MaterialExpressionStaticBoolParameter.h
+++ b/Engine/Source/Runtime/Engine/Classes/Materials/MaterialExpressionStaticBoolParameter.h
@@ -24,6 +24,7 @@ public:
 	virtual void GetCaption(TArray<FString>& OutCaptions) const override;
 	virtual uint32 GetOutputType(int32 OutputIndex) override {return MCT_StaticBool;}
 	virtual void SetValueToMatchingExpression(UMaterialExpression* OtherExpression) override;
+	virtual EMaterialGenerateHLSLStatus GenerateHLSLExpression(FMaterialHLSLGenerator& Generator, UE::HLSLTree::FScope& Scope, int32 OutputIndex, UE::HLSLTree::FExpression*& OutExpression) override;
 #endif
 	//~ End UMaterialExpression Interface
 
diff --git a/Engine/Source/Runtime/Engine/Classes/Materials/MaterialExpressionStaticSwitch.h b/Engine/Source/Runtime/Engine/Classes/Materials/MaterialExpressionStaticSwitch.h
--- a/Engine/Source/Runtime/Engine/Classes/Materials/MaterialExpressionStaticSwitch.h
+++ b/Engine/Source/Runtime/Engine/Classes/Materials/MaterialExpressionStaticSwitch.h
@@ -35,6 +35,8 @@ class UMaterialExpressionStaticSwitch : public UMaterialExpression
 	virtual bool IsResultMaterialAttributes(int32 OutputIndex) override;
 	virtual uint32 GetInputType(int32 InputIndex) override;
 	virtual uint32 GetOutputType(int32 OutputIndex) override {return MCT_Unknown;}
+
+	virtual EMaterialGenerateHLSLStatus GenerateHLSLExpression(FMaterialHLSLGenerator& Generator, UE::HLSLTree::FScope& Scope, int32 OutputIndex, UE::HLSLTree::FExpression*& OutExpression) override;
 #endif // WITH_EDITOR
 	//~ End UMaterialExpression Interface
 };
diff --git a/Engine/Source/Runtime/Engine/Private/HLSLTree/HLSLTree.cpp b/Engine/Source/Runtime/Engine/Private/HLSLTree/HLSLTree.cpp
--- a/Engine/Source/Runtime/Engine/Private/HLSLTree/HLSLTree.cpp
+++ b/Engine/Source/Runtime/Engine/Private/HLSLTree/HLSLTree.cpp
@@ -12,81 +12,6 @@ public:
 	TStringBuilder<2048> LocalStringBuilder;
 };
 
-UE::HLSLTree::FExpressionTypeDescription UE::HLSLTree::GetExpressionTypeDescription(EExpressionType Type)
-{
-	switch (Type)
-	{
-	case EExpressionType::Void: return FExpressionTypeDescription(TEXT("void"), EExpressionComponentType::Void, 0);
-	case EExpressionType::Float1: return FExpressionTypeDescription(TEXT("float"), EExpressionComponentType::Float, 1);
-	case EExpressionType::Float2: return FExpressionTypeDescription(TEXT("float2"), EExpressionComponentType::Float, 2);
-	case EExpressionType::Float3: return FExpressionTypeDescription(TEXT("float3"), EExpressionComponentType::Float, 3);
-	case EExpressionType::Float4: return FExpressionTypeDescription(TEXT("float4"), EExpressionComponentType::Float, 4);
-	case EExpressionType::MaterialAttributes: return FExpressionTypeDescription(TEXT("FMaterialAttributes"), EExpressionComponentType::MaterialAttributes, 0);
-	default: checkNoEntry(); return FExpressionTypeDescription();
-	}
-}
-
-UE::HLSLTree::EExpressionType UE::HLSLTree::MakeExpressionType(EExpressionComponentType ComponentType, int32 NumComponents)
-{
-	switch (ComponentType)
-	{
-	case EExpressionComponentType::Void:
-		check(NumComponents == 0);
-		return EExpressionType::Void;
-	case EExpressionComponentType::MaterialAttributes:
-		check(NumComponents == 0);
-		return EExpressionType::MaterialAttributes;
-	case EExpressionComponentType::Float:
-		switch (NumComponents)
-		{
-		case 1: return EExpressionType::Float1;
-		case 2: return EExpressionType::Float2;
-		case 3: return EExpressionType::Float3;
-		case 4: return EExpressionType::Float4;
-		default: break;
-		}
-	default:
-		break;
-	}
-
-	checkNoEntry();
-	return EExpressionType::Void;
-}
-
-UE::HLSLTree::EExpressionType UE::HLSLTree::MakeExpressionType(EExpressionType BaseType, int32 NumComponents)
-{
-	return MakeExpressionType(GetExpressionTypeDescription(BaseType).ComponentType, NumComponents);
-}
-
-UE::HLSLTree::EExpressionType UE::HLSLTree::MakeArithmeticResultType(EExpressionType Lhs, EExpressionType Rhs, FString& OutErrorMessage)
-{
-	const FExpressionTypeDescription LhsDesc = GetExpressionTypeDescription(Lhs);
-	const FExpressionTypeDescription RhsDesc = GetExpressionTypeDescription(Rhs);
-	// Types with 0 components are non-arithmetic
-	if (LhsDesc.NumComponents > 0 && RhsDesc.NumComponents > 0)
-	{
-		if (Lhs == Rhs)
-		{
-			return Lhs;
-		}
-		if (LhsDesc.ComponentType == RhsDesc.ComponentType)
-		{
-			if (LhsDesc.NumComponents == 1 || RhsDesc.NumComponents == 1)
-			{
-				// single component type is valid to combine with other type
-				return MakeExpressionType(LhsDesc.ComponentType, FMath::Max(LhsDesc.NumComponents, RhsDesc.NumComponents));
-			}
-		}
-		OutErrorMessage = FString::Printf(TEXT("Arithmetic between types %s and %s are undefined"), LhsDesc.Name, RhsDesc.Name);
-	}
-	else
-	{
-		OutErrorMessage = FString::Printf(TEXT("Attempting to perform arithmetic on non-numeric types: %s %s"), LhsDesc.Name, RhsDesc.Name);
-	}
-
-	return EExpressionType::Void;
-}
-
 UE::HLSLTree::EExpressionEvaluationType UE::HLSLTree::CombineEvaluationTypes(UE::HLSLTree::EExpressionEvaluationType Lhs, UE::HLSLTree::EExpressionEvaluationType Rhs)
 {
 	if (Lhs == EExpressionEvaluationType::Constant && Rhs == EExpressionEvaluationType::Constant)
@@ -172,33 +97,39 @@ void UE::HLSLTree::FCodeWriter::Append(const FCodeWriter& InWriter)
 	}
 }
 
-FLinearColor UE::HLSLTree::FConstant::ToLinearColor() const
+
+void UE::HLSLTree::FCodeWriter::WriteConstant(const Shader::FValue& Value)
 {
-	switch (Type)
+	auto ToString = [](bool v)
 	{
-	case EExpressionType::Float1: return FLinearColor(Float[0], 0.0f, 0.0f, 0.0f); break;
-	case EExpressionType::Float2: return FLinearColor(Float[0], Float[1], 0.0f, 0.0f); break;
-	case EExpressionType::Float3: return FLinearColor(Float[0], Float[1], Float[2], 0.0f); break;
-	case EExpressionType::Float4: return FLinearColor(Float[0], Float[1], Float[2], Float[3]); break;
-	default: checkNoEntry(); return FLinearColor(ForceInitToZero);
-	}
-}
+		return v ? TEXT("true") : TEXT("false");
+	};
 
-void UE::HLSLTree::FConstant::EmitHLSL(UE::HLSLTree::FCodeWriter& Writer) const
-{
-	switch (Type)
+	switch (Value.GetType())
 	{
-	case EExpressionType::Float1:
-		Writer.Writef(TEXT("float(%0.8f)"), Float[0]);
+	case Shader::EValueType::Float1:
+		Writef(TEXT("float(%0.8f)"), Value.Component[0].Float);
+		break;
+	case Shader::EValueType::Float2:
+		Writef(TEXT("float2(%0.8f, %0.8f)"), Value.Component[0].Float, Value.Component[1].Float);
+		break;
+	case Shader::EValueType::Float3:
+		Writef(TEXT("float3(%0.8f, %0.8f, %0.8f)"), Value.Component[0].Float, Value.Component[1].Float, Value.Component[2].Float);
 		break;
-	case EExpressionType::Float2:
-		Writer.Writef(TEXT("float2(%0.8f, %0.8f)"), Float[0], Float[1]);
+	case Shader::EValueType::Float4:
+		Writef(TEXT("float4(%0.8f, %0.8f, %0.8f, %0.8f)"), Value.Component[0].Float, Value.Component[1].Float, Value.Component[2].Float, Value.Component[3].Float);
 		break;
-	case EExpressionType::Float3:
-		Writer.Writef(TEXT("float3(%0.8f, %0.8f, %0.8f)"), Float[0], Float[1], Float[2]);
+	case Shader::EValueType::Bool1:
+		Writef(TEXT("%s"), ToString(Value.Component[0].Bool));
 		break;
-	case EExpressionType::Float4:
-		Writer.Writef(TEXT("float4(%0.8f, %0.8f, %0.8f, %0.8f)"), Float[0], Float[1], Float[2], Float[3]);
+	case Shader::EValueType::Bool2:
+		Writef(TEXT("bool2(%s, %s)"), ToString(Value.Component[0].Bool), ToString(Value.Component[1].Bool));
+		break;
+	case Shader::EValueType::Bool3:
+		Writef(TEXT("bool3(%s, %s, %s)"), ToString(Value.Component[0].Bool), ToString(Value.Component[1].Bool), ToString(Value.Component[2].Bool));
+		break;
+	case Shader::EValueType::Bool4:
+		Writef(TEXT("bool4(%s, %s, %s, %s)"), ToString(Value.Component[0].Bool), ToString(Value.Component[1].Bool), ToString(Value.Component[2].Bool), ToString(Value.Component[3].Bool));
 		break;
 	default:
 		checkNoEntry();
@@ -216,7 +147,7 @@ UE::HLSLTree::FEmitContext::~FEmitContext()
 	// Make sure we have only the root stack entry
 	check(FunctionStack.Num() == 1);
 
-	for (FMaterialPreshaderData* Preshader : TempPreshaders)
+	for (Shader::FPreshaderData* Preshader : TempPreshaders)
 	{
 		delete Preshader;
 	}
@@ -248,10 +179,25 @@ int32 UE::HLSLTree::FEmitContext::FindScopeIndex(FScope* Scope)
 	return INDEX_NONE;
 }
 
+void UE::HLSLTree::FExpressionEmitResult::ForwardValue(FEmitContext& Context, const FEmitValue& InValue)
+{
+	EvaluationType = InValue.GetEvaluationType();
+	Type = InValue.GetExpressionType();
+	if (EvaluationType == EExpressionEvaluationType::Shader)
+	{
+		bInline = true;
+		Writer.Writef(TEXT("%s"), Context.GetCode(InValue));
+	}
+	else
+	{
+		Context.AppendPreshader(InValue, Preshader);
+	}
+}
+
 const UE::HLSLTree::FEmitValue& UE::HLSLTree::FEmitContext::AcquireValue(UE::HLSLTree::FLocalDeclaration* Declaration)
 {
 	check(Declaration);
-	check(Declaration->Type != EExpressionType::Void);
+	check(Declaration->Type != Shader::EValueType::Void);
 
 	FDeclarationEntry*& Entry = FunctionStack.Last().DeclarationMap.FindOrAdd(Declaration);
 	if (!Entry)
@@ -260,9 +206,10 @@ const UE::HLSLTree::FEmitValue& UE::HLSLTree::FEmitContext::AcquireValue(UE::HLS
 		check(ScopeEntry);
 		Entry = new(*Allocator) FDeclarationEntry();
 		Entry->Value.EvaluationType = EExpressionEvaluationType::Shader;
+		Entry->Value.ExpressionType = Declaration->Type;
 
 		Entry->Value.Code = AllocateStringf(*Allocator, TEXT("Local%d"), NumExpressionLocals++);
-		const FExpressionTypeDescription& TypeDesc = GetExpressionTypeDescription(Declaration->Type);
+		const Shader::FValueTypeDescription& TypeDesc = Shader::GetValueTypeDescription(Declaration->Type);
 		ScopeEntry->ExpressionCodeWriter->WriteLinef(TEXT("%s %s = (%s)0.0f;"),
 			TypeDesc.Name,
 			Entry->Value.Code,
@@ -275,33 +222,30 @@ const UE::HLSLTree::FEmitValue& UE::HLSLTree::FEmitContext::AcquireValue(UE::HLS
 const UE::HLSLTree::FEmitValue& UE::HLSLTree::FEmitContext::AcquireValue(UE::HLSLTree::FExpression* Expression)
 {
 	check(Expression);
-	check(Expression->Type != EExpressionType::Void);
 
 	FDeclarationEntry*& Entry = FunctionStack.Last().DeclarationMap.FindOrAdd(Expression);
 	if (!Entry)
 	{
 		FLocalHLSLCodeWriter LocalWriter;
-		FMaterialPreshaderData LocalPreshader;
+		Shader::FPreshaderData LocalPreshader;
 		FExpressionEmitResult EmitResult(LocalWriter, LocalPreshader);
 		Expression->EmitHLSL(*this, EmitResult);
 
+		check(EmitResult.Type != Shader::EValueType::Void);
+
 		Entry = new(*Allocator) FDeclarationEntry();
-		Entry->Value.ExpressionType = Expression->Type;
+		Entry->Value.ExpressionType = EmitResult.Type;
 		Entry->Value.EvaluationType = EmitResult.EvaluationType;
 		if (EmitResult.EvaluationType == EExpressionEvaluationType::Constant)
 		{
 			// Evaluate the constant preshader and store its value
 			FMaterialRenderContext RenderContext(nullptr, *Material, nullptr);
-			FLinearColor ConstantValue;
-			LocalPreshader.Evaluate(nullptr, RenderContext, ConstantValue);
-
-			const FConstant Constant(Expression->Type, ConstantValue);
-			Entry->Value.ConstantValue = Constant;
+			LocalPreshader.Evaluate(nullptr, RenderContext, Entry->Value.ConstantValue);
 		}
 		else if (EmitResult.EvaluationType == EExpressionEvaluationType::Preshader)
 		{
 			// Non-constant preshader, store it
-			FMaterialPreshaderData* Preshader = new FMaterialPreshaderData(MoveTemp(LocalPreshader));
+			Shader::FPreshaderData* Preshader = new Shader::FPreshaderData(MoveTemp(LocalPreshader));
 			TempPreshaders.Add(Preshader); // TODO - dedupe? store more efficiently?
 			Entry->Value.Preshader = Preshader;
 		}
@@ -334,7 +278,7 @@ const UE::HLSLTree::FEmitValue& UE::HLSLTree::FEmitContext::AcquireValue(UE::HLS
 				
 				if (!Declaration)
 				{
-					const FExpressionTypeDescription& TypeDesc = GetExpressionTypeDescription(Expression->Type);
+					const Shader::FValueTypeDescription& TypeDesc = Shader::GetValueTypeDescription(EmitResult.Type);
 					Declaration = AllocateStringf(*Allocator, TEXT("Local%d"), NumExpressionLocals++);
 					ScopeStack[ScopeIndex].ExpressionCodeWriter->WriteLinef(TEXT("const %s %s = %s;"),
 						TypeDesc.Name,
@@ -390,7 +334,7 @@ const UE::HLSLTree::FEmitValue& UE::HLSLTree::FEmitContext::AcquireValue(FFuncti
 			else
 			{
 				Output.EvaluationType = EExpressionEvaluationType::Constant;
-				Output.ExpressionType = EExpressionType::Float1;
+				Output.ExpressionType = Shader::EValueType::Float1;
 				Output.ConstantValue = 0.0f;
 			}
 		}
@@ -422,7 +366,7 @@ const TCHAR* UE::HLSLTree::FEmitContext::GetCode(const FEmitValue& Value) const
 		FLocalHLSLCodeWriter LocalWriter;
 		if (Value.EvaluationType == EExpressionEvaluationType::Constant)
 		{
-			Value.ConstantValue.EmitHLSL(LocalWriter);
+			LocalWriter.WriteConstant(Value.ConstantValue);
 		}
 		else
 		{
@@ -431,7 +375,7 @@ const TCHAR* UE::HLSLTree::FEmitContext::GetCode(const FEmitValue& Value) const
 
 			FUniformExpressionSet& UniformExpressionSet = MaterialCompilationOutput->UniformExpressionSet;
 			FMaterialUniformPreshaderHeader* PreshaderHeader = nullptr;
-			if (Value.ExpressionType == EExpressionType::Float1)
+			if (Value.ExpressionType == Shader::EValueType::Float1)
 			{
 				const static TCHAR IndexToMask[] = { 'x', 'y', 'z', 'w' };
 				const int32 ScalarInputIndex = UniformExpressionSet.UniformScalarPreshaders.Num();
@@ -444,10 +388,10 @@ const TCHAR* UE::HLSLTree::FEmitContext::GetCode(const FEmitValue& Value) const
 				const TCHAR* Mask = TEXT("");
 				switch (Value.ExpressionType)
 				{
-				case EExpressionType::Float1: Mask = TEXT(".r"); break;
-				case EExpressionType::Float2: Mask = TEXT(".rg"); break;
-				case EExpressionType::Float3: Mask = TEXT(".rgb"); break;
-				case EExpressionType::Float4: break;
+				case Shader::EValueType::Float1: Mask = TEXT(".r"); break;
+				case Shader::EValueType::Float2: Mask = TEXT(".rg"); break;
+				case Shader::EValueType::Float3: Mask = TEXT(".rgb"); break;
+				case Shader::EValueType::Float4: break;
 				default: checkNoEntry(); break;
 				};
 				const int32 VectorInputIndex = UniformExpressionSet.UniformVectorPreshaders.Num();
@@ -465,14 +409,13 @@ const TCHAR* UE::HLSLTree::FEmitContext::GetCode(const FEmitValue& Value) const
 	return Value.Code;
 }
 
-void UE::HLSLTree::FEmitContext::AppendPreshader(const FEmitValue& Value, FMaterialPreshaderData& InOutPreshader) const
+void UE::HLSLTree::FEmitContext::AppendPreshader(const FEmitValue& Value, Shader::FPreshaderData& InOutPreshader) const
 {
 	if (Value.EvaluationType == EExpressionEvaluationType::Constant)
 	{
 		// Push the constant value
-		const FLinearColor Constant = Value.ConstantValue.ToLinearColor();
-		InOutPreshader.WriteOpcode(EMaterialPreshaderOpcode::Constant);
-		InOutPreshader.Write(Constant);
+		InOutPreshader.WriteOpcode(Shader::EPreshaderOpcode::Constant);
+		InOutPreshader.Write(Value.ConstantValue);
 	}
 	else
 	{
@@ -798,14 +741,14 @@ UE::HLSLTree::FScope* UE::HLSLTree::FTree::NewLinkedScope(UE::HLSLTree::FScope&
 	return NewScope;
 }
 
-UE::HLSLTree::FLocalDeclaration* UE::HLSLTree::FTree::NewLocalDeclaration(UE::HLSLTree::FScope& Scope, UE::HLSLTree::EExpressionType Type, const FName& Name)
+UE::HLSLTree::FLocalDeclaration* UE::HLSLTree::FTree::NewLocalDeclaration(UE::HLSLTree::FScope& Scope, Shader::EValueType Type, const FName& Name)
 {
 	FLocalDeclaration* Declaration = NewNode<FLocalDeclaration>(Name, Type);
 	Scope.AddDeclaration(Declaration);
 	return Declaration;
 }
 
-UE::HLSLTree::FParameterDeclaration* UE::HLSLTree::FTree::NewParameterDeclaration(UE::HLSLTree::FScope& Scope, const FName& Name, const UE::HLSLTree::FConstant& DefaultValue)
+UE::HLSLTree::FParameterDeclaration* UE::HLSLTree::FTree::NewParameterDeclaration(UE::HLSLTree::FScope& Scope, const FName& Name, const Shader::FValue& DefaultValue)
 {
 	FParameterDeclaration* Declaration = NewNode<FParameterDeclaration>(Name, DefaultValue);
 	Declaration->ParentScope = &Scope;
diff --git a/Engine/Source/Runtime/Engine/Private/HLSLTree/HLSLTreeCommon.cpp b/Engine/Source/Runtime/Engine/Private/HLSLTree/HLSLTreeCommon.cpp
--- a/Engine/Source/Runtime/Engine/Private/HLSLTree/HLSLTreeCommon.cpp
+++ b/Engine/Source/Runtime/Engine/Private/HLSLTree/HLSLTreeCommon.cpp
@@ -62,38 +62,53 @@ UE::HLSLTree::FSwizzleParameters UE::HLSLTree::MakeSwizzleMask(bool bInR, bool b
 void UE::HLSLTree::FExpressionConstant::EmitHLSL(UE::HLSLTree::FEmitContext& Context, FExpressionEmitResult& OutResult) const
 {
 	OutResult.EvaluationType = EExpressionEvaluationType::Constant;
-	OutResult.Preshader.WriteOpcode(EMaterialPreshaderOpcode::Constant);
-	OutResult.Preshader.Write(Value.ToLinearColor());
+	OutResult.Type = Value.GetType();
+	OutResult.Preshader.WriteOpcode(Shader::EPreshaderOpcode::Constant);
+	OutResult.Preshader.Write(Value);
 }
 
 void UE::HLSLTree::FExpressionLocalVariable::EmitHLSL(UE::HLSLTree::FEmitContext& Context, FExpressionEmitResult& OutResult) const
 {
-	const FEmitValue&DeclarationRef = Context.AcquireValue(Declaration);
+	const FEmitValue& DeclarationRef = Context.AcquireValue(Declaration);
 	
+	OutResult.Type = DeclarationRef.GetExpressionType();
 	OutResult.bInline = true;
 	OutResult.Writer.Writef(TEXT("%s"), Context.GetCode(DeclarationRef));
 }
 
 void UE::HLSLTree::FExpressionParameter::EmitHLSL(UE::HLSLTree::FEmitContext& Context, FExpressionEmitResult& OutResult) const
 {
-	OutResult.EvaluationType = EExpressionEvaluationType::Preshader;
+	OutResult.Type = Declaration->DefaultValue.GetType();
 
-	EMaterialPreshaderOpcode Opcode = EMaterialPreshaderOpcode::Nop;
-	int32 ParameterIndex = INDEX_NONE;
-	if (Declaration->DefaultValue.Type == UE::HLSLTree::EExpressionType::Float1)
+	if (OutResult.Type == Shader::EValueType::Bool1)
+	{
+		const FMaterialParameterInfo ParameterInfo(Declaration->Name);
+		bool bValue = Declaration->DefaultValue.Component[0].Bool;
+		for (const FStaticSwitchParameter& Parameter : Context.StaticParameters->StaticSwitchParameters)
+		{
+			if (Parameter.ParameterInfo == ParameterInfo)
+			{
+				bValue = Parameter.Value;
+				break;
+			}
+		}
+		OutResult.EvaluationType = EExpressionEvaluationType::Constant;
+		OutResult.Preshader.WriteOpcode(Shader::EPreshaderOpcode::Constant).Write(Shader::FValue(bValue));
+	}
+	else if (OutResult.Type == Shader::EValueType::Float1)
 	{
-		ParameterIndex = Context.MaterialCompilationOutput->UniformExpressionSet.FindOrAddScalarParameter(Declaration->Name, Declaration->DefaultValue.Float[0]);
-		Opcode = EMaterialPreshaderOpcode::ScalarParameter;
+		const int32 ParameterIndex = Context.MaterialCompilationOutput->UniformExpressionSet.FindOrAddScalarParameter(Declaration->Name, Declaration->DefaultValue.Component[0].Float);
+		check(ParameterIndex >= 0 && ParameterIndex <= 0xffff);
+		OutResult.EvaluationType = EExpressionEvaluationType::Preshader;
+		OutResult.Preshader.WriteOpcode(Shader::EPreshaderOpcode::ScalarParameter).Write((uint16)ParameterIndex);
 	}
 	else
 	{
-		ParameterIndex = Context.MaterialCompilationOutput->UniformExpressionSet.FindOrAddVectorParameter(Declaration->Name, Declaration->DefaultValue.ToLinearColor());
-		Opcode = EMaterialPreshaderOpcode::VectorParameter;
+		const int32 ParameterIndex = Context.MaterialCompilationOutput->UniformExpressionSet.FindOrAddVectorParameter(Declaration->Name, Declaration->DefaultValue.AsLinearColor());
+		check(ParameterIndex >= 0 && ParameterIndex <= 0xffff);
+		OutResult.EvaluationType = EExpressionEvaluationType::Preshader;
+		OutResult.Preshader.WriteOpcode(Shader::EPreshaderOpcode::VectorParameter).Write((uint16)ParameterIndex);
 	}
-
-	check(ParameterIndex >= 0 && ParameterIndex <= 0xffff);
-	OutResult.Preshader.WriteOpcode(Opcode);
-	OutResult.Preshader.Write((uint16)ParameterIndex);
 }
 
 void UE::HLSLTree::FExpressionExternalInput::EmitHLSL(UE::HLSLTree::FEmitContext& Context, FExpressionEmitResult& OutResult) const
@@ -101,6 +116,7 @@ void UE::HLSLTree::FExpressionExternalInput::EmitHLSL(UE::HLSLTree::FEmitContext
 	const int32 TypeIndex = (int32)InputType;
 	const int32 TexCoordIndex = TypeIndex - (int32)UE::HLSLTree::EExternalInputType::TexCoord0;
 
+	OutResult.Type = GetInputExpressionType(InputType);
 	OutResult.bInline = true;
 	if (TexCoordIndex >= 0 && TexCoordIndex < 8)
 	{
@@ -240,11 +256,13 @@ void UE::HLSLTree::FExpressionTextureSample::EmitHLSL(UE::HLSLTree::FEmitContext
 	}
 
 	const FEmitValue& TexCoordRef = Context.AcquireValue(TexCoordExpression);
+	OutResult.Type = Shader::EValueType::Float4;
 	OutResult.Writer.Writef(TEXT("%s(%s(%s, %s, %s))"), SamplerTypeFunction, SampleFunctionName, *TextureName, *SamplerStateCode, Context.GetCode(TexCoordRef));
 }
 
 void UE::HLSLTree::FExpressionDefaultMaterialAttributes::EmitHLSL(UE::HLSLTree::FEmitContext& Context, FExpressionEmitResult& OutResult) const
 {
+	OutResult.Type = Shader::EValueType::MaterialAttributes;
 	OutResult.bInline = true;
 	OutResult.Writer.Write(TEXT("DefaultMaterialAttributes"));
 }
@@ -259,25 +277,59 @@ void UE::HLSLTree::FExpressionSetMaterialAttribute::EmitHLSL(UE::HLSLTree::FEmit
 	const FEmitValue&AttributresRef = Context.AcquireValue(AttributesExpression);
 	const FEmitValue&ValueRef = Context.AcquireValue(ValueExpression);
 
+	OutResult.Type = Shader::EValueType::MaterialAttributes;
 	OutResult.Writer.Writef(TEXT("FMaterialAttributes_Set%s(%s, %s)"), *PropertyName, Context.GetCode(AttributresRef), Context.GetCode(ValueRef));
 }
 
+void UE::HLSLTree::FExpressionSelect::EmitHLSL(FEmitContext& Context, FExpressionEmitResult& OutResult) const
+{
+	const FEmitValue& ConditionValue = Context.AcquireValue(ConditionExpression);
+	if (ConditionValue.GetEvaluationType() == EExpressionEvaluationType::Constant)
+	{
+		const bool bCondition = ConditionValue.GetConstantValue().AsBoolScalar();
+		if (!bCondition)
+		{
+			const FEmitValue& Value = Context.AcquireValue(FalseExpression);
+			OutResult.ForwardValue(Context, Value);
+		}
+		else
+		{
+			const FEmitValue& Value = Context.AcquireValue(TrueExpression);
+			OutResult.ForwardValue(Context, Value);
+		}
+	}
+	else
+	{
+		// TODO - preshader
+		const FEmitValue& TrueValue = Context.AcquireValue(TrueExpression);
+		const FEmitValue& FalseValue = Context.AcquireValue(FalseExpression);
+		OutResult.Type = TrueValue.GetExpressionType(); // TODO - FalseValue?
+		OutResult.Writer.Writef(TEXT("(%s ? %s : %s)"),
+			Context.GetCode(ConditionValue),
+			Context.GetCode(TrueValue),
+			Context.GetCode(FalseValue));
+	}
+}
+
 void UE::HLSLTree::FExpressionBinaryOp::EmitHLSL(UE::HLSLTree::FEmitContext& Context, FExpressionEmitResult& OutResult) const
 {
 	const FEmitValue& LhsRef = Context.AcquireValue(Lhs);
 	const FEmitValue& RhsRef = Context.AcquireValue(Rhs);
 
+	FString ErrorMessage;
+
 	OutResult.EvaluationType = CombineEvaluationTypes(LhsRef.GetEvaluationType(), RhsRef.GetEvaluationType());
+	OutResult.Type = MakeArithmeticResultType(LhsRef.GetExpressionType(), RhsRef.GetExpressionType(), ErrorMessage);
 	if (OutResult.EvaluationType == EExpressionEvaluationType::Shader)
 	{
 		OutResult.Writer.Writef(TEXT("(%s + %s)"), Context.GetCode(LhsRef), Context.GetCode(RhsRef));
 	}
 	else
 	{
-		EMaterialPreshaderOpcode Opcode = EMaterialPreshaderOpcode::Nop;
+		Shader::EPreshaderOpcode Opcode = Shader::EPreshaderOpcode::Nop;
 		switch (Op)
 		{
-		case UE::HLSLTree::EBinaryOp::Add: Opcode = EMaterialPreshaderOpcode::Add; break;
+		case UE::HLSLTree::EBinaryOp::Add: Opcode = Shader::EPreshaderOpcode::Add; break;
 		default: checkNoEntry(); break;
 		}
 
@@ -302,6 +354,7 @@ void UE::HLSLTree::FExpressionSwizzle::EmitHLSL(FEmitContext& Context, FExpressi
 
 	OutResult.bInline = true;
 	OutResult.EvaluationType = InputRef.GetEvaluationType();
+	OutResult.Type = Shader::MakeValueType(InputRef.GetExpressionType(), Parameters.NumComponents);
 	if (OutResult.EvaluationType == EExpressionEvaluationType::Shader)
 	{
 		OutResult.Writer.Writef(TEXT("%s.%s"), Context.GetCode(InputRef), Swizzle);
@@ -309,7 +362,7 @@ void UE::HLSLTree::FExpressionSwizzle::EmitHLSL(FEmitContext& Context, FExpressi
 	else
 	{
 		Context.AppendPreshader(InputRef, OutResult.Preshader);
-		OutResult.Preshader.WriteOpcode(EMaterialPreshaderOpcode::ComponentSwizzle)
+		OutResult.Preshader.WriteOpcode(Shader::EPreshaderOpcode::ComponentSwizzle)
 			.Write((uint8)Parameters.NumComponents)
 			.Write((uint8)Parameters.ComponentIndex[0])
 			.Write((uint8)Parameters.ComponentIndex[1])
@@ -320,21 +373,23 @@ void UE::HLSLTree::FExpressionSwizzle::EmitHLSL(FEmitContext& Context, FExpressi
 
 void UE::HLSLTree::FExpressionCast::EmitHLSL(FEmitContext& Context, FExpressionEmitResult& OutResult) const
 {
-	const EExpressionType InputType = Input->Type;
 	const FEmitValue& InputRef = Context.AcquireValue(Input);
+	const Shader::EValueType InputType = InputRef.GetExpressionType();
 	const TCHAR* InputCode = Context.GetCode(InputRef);
 
 	OutResult.bInline = true;
+	OutResult.Type = Type;
 	// TODO - preshader
 
 	if (InputType == Type)
 	{
 		// No cast required (nop)
 		OutResult.Writer.Writef(TEXT("%s"), InputCode);
+		return;
 	}
 
-	const FExpressionTypeDescription OutputTypeDesc = GetExpressionTypeDescription(Type);
-	const FExpressionTypeDescription InputTypeDesc = GetExpressionTypeDescription(InputType);
+	const Shader::FValueTypeDescription OutputTypeDesc = Shader::GetValueTypeDescription(Type);
+	const Shader::FValueTypeDescription InputTypeDesc = Shader::GetValueTypeDescription(InputType);
 	if (InputTypeDesc.NumComponents == OutputTypeDesc.NumComponents)
 	{
 		// Cast between different underlying types with the same number of components, use a C-style cast
@@ -384,34 +439,14 @@ void UE::HLSLTree::FExpressionFunctionInput::EmitHLSL(FEmitContext& Context, FEx
 
 	check(InputIndex >= 0 && InputIndex < StackEntry.FunctionCall->NumInputs);
 	FExpression* InputExpression = StackEntry.FunctionCall->Inputs[InputIndex];
-	const FEmitValue& Ref = Context.AcquireValue(InputExpression);
-
-	OutResult.bInline = true;
-	OutResult.EvaluationType = Ref.GetEvaluationType();
-	if (OutResult.EvaluationType == EExpressionEvaluationType::Shader)
-	{
-		OutResult.Writer.Writef(TEXT("%s"), Context.GetCode(Ref));
-	}
-	else
-	{
-		Context.AppendPreshader(Ref, OutResult.Preshader);
-	}
+	const FEmitValue& Value = Context.AcquireValue(InputExpression);
+	OutResult.ForwardValue(Context, Value);
 }
 
 void UE::HLSLTree::FExpressionFunctionOutput::EmitHLSL(FEmitContext& Context, FExpressionEmitResult& OutResult) const
 {
-	const FEmitValue& Ref = Context.AcquireValue(FunctionCall, OutputIndex);
-	
-	OutResult.bInline = true;
-	OutResult.EvaluationType = Ref.GetEvaluationType();
-	if (OutResult.EvaluationType == EExpressionEvaluationType::Shader)
-	{
-		OutResult.Writer.Writef(TEXT("%s"), Context.GetCode(Ref));
-	}
-	else
-	{
-		Context.AppendPreshader(Ref, OutResult.Preshader);
-	}
+	const FEmitValue& Value = Context.AcquireValue(FunctionCall, OutputIndex);
+	OutResult.ForwardValue(Context, Value);
 }
 
 void UE::HLSLTree::FStatementReturn::EmitHLSL(UE::HLSLTree::FEmitContext& Context, UE::HLSLTree::FCodeWriter& OutWriter) const
diff --git a/Engine/Source/Runtime/Engine/Private/Materials/MaterialExpressionHLSL.cpp b/Engine/Source/Runtime/Engine/Private/Materials/MaterialExpressionHLSL.cpp
--- a/Engine/Source/Runtime/Engine/Private/Materials/MaterialExpressionHLSL.cpp
+++ b/Engine/Source/Runtime/Engine/Private/Materials/MaterialExpressionHLSL.cpp
@@ -15,14 +15,17 @@
 #include "Materials/MaterialExpressionConstant2Vector.h"
 #include "Materials/MaterialExpressionConstant3Vector.h"
 #include "Materials/MaterialExpressionConstant4Vector.h"
+#include "Materials/MaterialExpressionStaticBool.h"
 #include "Materials/MaterialExpressionVectorParameter.h"
 #include "Materials/MaterialExpressionScalarParameter.h"
+#include "Materials/MaterialExpressionStaticBoolParameter.h"
 #include "Materials/MaterialExpressionTextureSample.h"
 #include "Materials/MaterialExpressionTextureSampleParameter.h"
 #include "Materials/MaterialExpressionTextureObject.h"
 #include "Materials/MaterialExpressionTextureObjectParameter.h"
 #include "Materials/MaterialExpressionFunctionInput.h"
 #include "Materials/MaterialExpressionMaterialFunctionCall.h"
+#include "Materials/MaterialExpressionStaticSwitch.h"
 #include "Materials/MaterialExpressionGetLocal.h"
 #include "Materials/MaterialExpressionAdd.h"
 #include "Materials/MaterialExpressionSetMaterialAttributes.h"
@@ -64,32 +67,48 @@ EMaterialGenerateHLSLStatus UMaterialExpressionConstant::GenerateHLSLExpression(
 
 EMaterialGenerateHLSLStatus UMaterialExpressionConstant2Vector::GenerateHLSLExpression(FMaterialHLSLGenerator& Generator, UE::HLSLTree::FScope& Scope, int32 OutputIndex, UE::HLSLTree::FExpression*& OutExpression)
 {
-	OutExpression = Generator.NewConstant(Scope, UE::HLSLTree::FConstant(R, G));
+	OutExpression = Generator.NewConstant(Scope, UE::Shader::FValue(R, G));
 	return EMaterialGenerateHLSLStatus::Success;
 }
 
 EMaterialGenerateHLSLStatus UMaterialExpressionConstant3Vector::GenerateHLSLExpression(FMaterialHLSLGenerator& Generator, UE::HLSLTree::FScope& Scope, int32 OutputIndex, UE::HLSLTree::FExpression*& OutExpression)
 {
-	OutExpression = Generator.NewConstant(Scope, UE::HLSLTree::FConstant(Constant.R, Constant.G, Constant.B));
+	OutExpression = Generator.NewConstant(Scope, UE::Shader::FValue(Constant.R, Constant.G, Constant.B));
 	return EMaterialGenerateHLSLStatus::Success;
 }
 
 EMaterialGenerateHLSLStatus UMaterialExpressionConstant4Vector::GenerateHLSLExpression(FMaterialHLSLGenerator& Generator, UE::HLSLTree::FScope& Scope, int32 OutputIndex, UE::HLSLTree::FExpression*& OutExpression)
 {
-	OutExpression = Generator.NewConstant(Scope, UE::HLSLTree::FConstant(Constant.R, Constant.G, Constant.B, Constant.A));
+	OutExpression = Generator.NewConstant(Scope, UE::Shader::FValue(Constant.R, Constant.G, Constant.B, Constant.A));
+	return EMaterialGenerateHLSLStatus::Success;
+}
+
+EMaterialGenerateHLSLStatus UMaterialExpressionStaticBool::GenerateHLSLExpression(FMaterialHLSLGenerator& Generator, UE::HLSLTree::FScope& Scope, int32 OutputIndex, UE::HLSLTree::FExpression*& OutExpression)
+{
+	OutExpression = Generator.NewConstant(Scope, (bool)Value);
+	return EMaterialGenerateHLSLStatus::Success;
+}
+
+EMaterialGenerateHLSLStatus UMaterialExpressionStaticSwitch::GenerateHLSLExpression(FMaterialHLSLGenerator& Generator, UE::HLSLTree::FScope& Scope, int32 OutputIndex, UE::HLSLTree::FExpression*& OutExpression)
+{
+	UE::HLSLTree::FExpression* ConditionExpression = Value.GetTracedInput().Expression ? Value.AcquireHLSLExpression(Generator, Scope) : Generator.NewConstant(Scope, (bool)DefaultValue);
+	UE::HLSLTree::FExpression* TrueExpression = A.AcquireHLSLExpression(Generator, Scope);
+	UE::HLSLTree::FExpression* FalseExpression = B.AcquireHLSLExpression(Generator, Scope);
+
+	OutExpression = Generator.GetTree().NewExpression<UE::HLSLTree::FExpressionSelect>(Scope, ConditionExpression, TrueExpression, FalseExpression);
 	return EMaterialGenerateHLSLStatus::Success;
 }
 
 EMaterialGenerateHLSLStatus UMaterialExpressionGetLocal::GenerateHLSLExpression(FMaterialHLSLGenerator& Generator, UE::HLSLTree::FScope& Scope, int32 OutputIndex, UE::HLSLTree::FExpression*& OutExpression)
 {
-	UE::HLSLTree::FLocalDeclaration* Declaration = Generator.AcquireLocalDeclaration(Scope, UE::HLSLTree::EExpressionType::Float3, LocalName);
+	UE::HLSLTree::FLocalDeclaration* Declaration = Generator.AcquireLocalDeclaration(Scope, UE::Shader::EValueType::Float3, LocalName);
 	OutExpression = Generator.GetTree().NewExpression<UE::HLSLTree::FExpressionLocalVariable>(Scope, Declaration);
 	return EMaterialGenerateHLSLStatus::Success;
 }
 
 EMaterialGenerateHLSLStatus UMaterialExpressionVectorParameter::GenerateHLSLExpression(FMaterialHLSLGenerator& Generator, UE::HLSLTree::FScope& Scope, int32 OutputIndex, UE::HLSLTree::FExpression*& OutExpression)
 {
-	UE::HLSLTree::FParameterDeclaration* Declaration = Generator.AcquireParameterDeclaration(Scope, ParameterName, UE::HLSLTree::FConstant(UE::HLSLTree::EExpressionType::Float4, DefaultValue));
+	UE::HLSLTree::FParameterDeclaration* Declaration = Generator.AcquireParameterDeclaration(Scope, ParameterName, DefaultValue);
 	OutExpression = Generator.GetTree().NewExpression<UE::HLSLTree::FExpressionParameter>(Scope, Declaration);
 	return EMaterialGenerateHLSLStatus::Success;
 }
@@ -101,6 +120,13 @@ EMaterialGenerateHLSLStatus UMaterialExpressionScalarParameter::GenerateHLSLExpr
 	return EMaterialGenerateHLSLStatus::Success;
 }
 
+EMaterialGenerateHLSLStatus UMaterialExpressionStaticBoolParameter::GenerateHLSLExpression(FMaterialHLSLGenerator& Generator, UE::HLSLTree::FScope& Scope, int32 OutputIndex, UE::HLSLTree::FExpression*& OutExpression)
+{
+	UE::HLSLTree::FParameterDeclaration* Declaration = Generator.AcquireParameterDeclaration(Scope, ParameterName, (bool)DefaultValue);
+	OutExpression = Generator.GetTree().NewExpression<UE::HLSLTree::FExpressionParameter>(Scope, Declaration);
+	return EMaterialGenerateHLSLStatus::Success;
+}
+
 EMaterialGenerateHLSLStatus UMaterialExpressionTextureObject::GenerateHLSLTexture(FMaterialHLSLGenerator& Generator, UE::HLSLTree::FScope& Scope, int32 OutputIndex, UE::HLSLTree::FTextureParameterDeclaration*& OutTexture)
 {
 	const UE::HLSLTree::FTextureDescription TextureDesc(Texture, SamplerType);
@@ -171,14 +197,14 @@ EMaterialGenerateHLSLStatus UMaterialExpressionAdd::GenerateHLSLExpression(FMate
 		return EMaterialGenerateHLSLStatus::Error;
 	}
 
-	FString ErrorMessage;
-	const UE::HLSLTree::EExpressionType ResultType = UE::HLSLTree::MakeArithmeticResultType(Lhs->Type, Rhs->Type, ErrorMessage);
-	if (ResultType == UE::HLSLTree::EExpressionType::Void)
+	/*FString ErrorMessage;
+	const UE::Shader::EValueType ResultType = UE::HLSLTree::MakeArithmeticResultType(Lhs->Type, Rhs->Type, ErrorMessage);
+	if (ResultType == UE::Shader::EValueType::Void)
 	{
 		return Generator.Errorf(TEXT("%s"), *ErrorMessage);
-	}
+	}*/
 
-	OutExpression = Generator.GetTree().NewExpression<UE::HLSLTree::FExpressionBinaryOp>(Scope, ResultType, UE::HLSLTree::EBinaryOp::Add, Lhs, Rhs);
+	OutExpression = Generator.GetTree().NewExpression<UE::HLSLTree::FExpressionBinaryOp>(Scope, UE::HLSLTree::EBinaryOp::Add, Lhs, Rhs);
 	return EMaterialGenerateHLSLStatus::Success;
 }
 
@@ -223,13 +249,13 @@ EMaterialGenerateHLSLStatus UMaterialExpressionSetMaterialAttributes::GenerateHL
 
 EMaterialGenerateHLSLStatus UMaterialExpressionFunctionInput::GenerateHLSLExpression(FMaterialHLSLGenerator& Generator, UE::HLSLTree::FScope& Scope, int32 OutputIndex, UE::HLSLTree::FExpression*& OutExpression)
 {
-	UE::HLSLTree::EExpressionType ExpressionType = UE::HLSLTree::EExpressionType::Void;
+	UE::Shader::EValueType ExpressionType = UE::Shader::EValueType::Void;
 	switch (InputType)
 	{
-	case FunctionInput_Scalar: ExpressionType = UE::HLSLTree::EExpressionType::Float1; break;
-	case FunctionInput_Vector2: ExpressionType = UE::HLSLTree::EExpressionType::Float2; break;
-	case FunctionInput_Vector3: ExpressionType = UE::HLSLTree::EExpressionType::Float3; break;
-	case FunctionInput_Vector4: ExpressionType = UE::HLSLTree::EExpressionType::Float4; break;
+	case FunctionInput_Scalar: ExpressionType = UE::Shader::EValueType::Float1; break;
+	case FunctionInput_Vector2: ExpressionType = UE::Shader::EValueType::Float2; break;
+	case FunctionInput_Vector3: ExpressionType = UE::Shader::EValueType::Float3; break;
+	case FunctionInput_Vector4: ExpressionType = UE::Shader::EValueType::Float4; break;
 	case FunctionInput_Texture2D:
 	case FunctionInput_TextureCube:
 	case FunctionInput_Texture2DArray:
@@ -239,7 +265,7 @@ EMaterialGenerateHLSLStatus UMaterialExpressionFunctionInput::GenerateHLSLExpres
 	case FunctionInput_TextureExternal:
 		break;
 	}
-	if (ExpressionType == UE::HLSLTree::EExpressionType::Void)
+	if (ExpressionType == UE::Shader::EValueType::Void)
 	{
 		return Generator.Error(TEXT("Invalid input connection"));
 	}
@@ -284,13 +310,13 @@ EMaterialGenerateHLSLStatus UMaterialExpressionExecBegin::GenerateHLSLStatement(
 
 EMaterialGenerateHLSLStatus UMaterialExpressionSetLocal::GenerateHLSLStatement(FMaterialHLSLGenerator& Generator, UE::HLSLTree::FScope& Scope, UE::HLSLTree::FStatement*& OutStatement)
 {
-	UE::HLSLTree::FExpression* ValueExpression = Value.AcquireHLSLExpressionWithCast(Generator, Scope, UE::HLSLTree::EExpressionType::Float3);
+	UE::HLSLTree::FExpression* ValueExpression = Value.AcquireHLSLExpressionWithCast(Generator, Scope, UE::Shader::EValueType::Float3);
 	if (!ValueExpression)
 	{
 		return Generator.Error(TEXT("Missing value connection"));
 	}
 
-	UE::HLSLTree::FLocalDeclaration* Declaration = Generator.AcquireLocalDeclaration(Scope, UE::HLSLTree::EExpressionType::Float3, LocalName);
+	UE::HLSLTree::FLocalDeclaration* Declaration = Generator.AcquireLocalDeclaration(Scope, UE::Shader::EValueType::Float3, LocalName);
 	if (!Declaration)
 	{
 		return EMaterialGenerateHLSLStatus::Error;
diff --git a/Engine/Source/Runtime/Engine/Private/Materials/MaterialHLSLEmitter.cpp b/Engine/Source/Runtime/Engine/Private/Materials/MaterialHLSLEmitter.cpp
--- a/Engine/Source/Runtime/Engine/Private/Materials/MaterialHLSLEmitter.cpp
+++ b/Engine/Source/Runtime/Engine/Private/Materials/MaterialHLSLEmitter.cpp
@@ -648,6 +648,7 @@ static void GetMaterialEnvironment(EShaderPlatform InPlatform,
 
 bool MaterialEmitHLSL(const FMaterialCompileTargetParameters& InCompilerTarget,
 	const FMaterial& InMaterial,
+	const FStaticParameterSet& InStaticParameters,
 	const UE::HLSLTree::FTree& InTree,
 	FMaterialCompilationOutput& OutCompilationOutput,
 	TRefCountPtr<FSharedShaderCompilerEnvironment>& OutMaterialEnvironment)
@@ -678,6 +679,7 @@ bool MaterialEmitHLSL(const FMaterialCompileTargetParameters& InCompilerTarget,
 	UE::HLSLTree::FEmitContext EmitContext;
 	EmitContext.Allocator = &Allocator;
 	EmitContext.Material = &InMaterial;
+	EmitContext.StaticParameters = &InStaticParameters;
 	EmitContext.MaterialCompilationOutput = &OutCompilationOutput;
 	InTree.EmitHLSL(EmitContext, *CodeWriter);
 
diff --git a/Engine/Source/Runtime/Engine/Private/Materials/MaterialHLSLEmitter.h b/Engine/Source/Runtime/Engine/Private/Materials/MaterialHLSLEmitter.h
--- a/Engine/Source/Runtime/Engine/Private/Materials/MaterialHLSLEmitter.h
+++ b/Engine/Source/Runtime/Engine/Private/Materials/MaterialHLSLEmitter.h
@@ -10,6 +10,7 @@ class FMaterial;
 class FMaterialCompilationOutput;
 struct FSharedShaderCompilerEnvironment;
 struct FMaterialCompileTargetParameters;
+struct FStaticParameterSet;
 
 namespace UE
 {
@@ -21,6 +22,7 @@ class FTree;
 
 bool MaterialEmitHLSL(const FMaterialCompileTargetParameters& InCompilerTarget,
 	const FMaterial& InOutMaterial,
+	const FStaticParameterSet& InStaticParameters,
 	const UE::HLSLTree::FTree& InTree,
 	FMaterialCompilationOutput& OutCompilationOutput,
 	TRefCountPtr<FSharedShaderCompilerEnvironment>& OutMaterialEnvironment);
diff --git a/Engine/Source/Runtime/Engine/Private/Materials/MaterialHLSLGenerator.cpp b/Engine/Source/Runtime/Engine/Private/Materials/MaterialHLSLGenerator.cpp
--- a/Engine/Source/Runtime/Engine/Private/Materials/MaterialHLSLGenerator.cpp
+++ b/Engine/Source/Runtime/Engine/Private/Materials/MaterialHLSLGenerator.cpp
@@ -67,7 +67,7 @@ EMaterialGenerateHLSLStatus FMaterialHLSLGenerator::Error(const FString& Message
 	return EMaterialGenerateHLSLStatus::Error;
 }
 
-UE::HLSLTree::FExpressionConstant* FMaterialHLSLGenerator::NewConstant(UE::HLSLTree::FScope& Scope, const UE::HLSLTree::FConstant& Value)
+UE::HLSLTree::FExpressionConstant* FMaterialHLSLGenerator::NewConstant(UE::HLSLTree::FScope& Scope, const UE::Shader::FValue& Value)
 {
 	UE::HLSLTree::FExpressionConstant* Expression = HLSLTree->NewExpression<UE::HLSLTree::FExpressionConstant>(Scope, Value);
 	return Expression;
@@ -85,26 +85,21 @@ UE::HLSLTree::FExpressionSwizzle* FMaterialHLSLGenerator::NewSwizzle(UE::HLSLTre
 	return Expression;
 }
 
-UE::HLSLTree::FExpression* FMaterialHLSLGenerator::NewCast(UE::HLSLTree::FScope& Scope, UE::HLSLTree::EExpressionType Type, UE::HLSLTree::FExpression* Input, UE::HLSLTree::ECastFlags Flags)
+UE::HLSLTree::FExpression* FMaterialHLSLGenerator::NewCast(UE::HLSLTree::FScope& Scope, UE::Shader::EValueType Type, UE::HLSLTree::FExpression* Input, UE::HLSLTree::ECastFlags Flags)
 {
-	// Only generate the cast expression if the types don't match
-	if (Input && Type != Input->Type)
-	{
-		UE::HLSLTree::FExpressionCast* Expression = HLSLTree->NewExpression<UE::HLSLTree::FExpressionCast>(Scope, Type, Input, Flags);
-		return Expression;
-	}
-	return Input;
+	UE::HLSLTree::FExpressionCast* Expression = HLSLTree->NewExpression<UE::HLSLTree::FExpressionCast>(Scope, Type, Input, Flags);
+	return Expression;
 }
 
 UE::HLSLTree::FExpression* FMaterialHLSLGenerator::NewFunctionInput(UE::HLSLTree::FScope& Scope, int32 InputIndex, UMaterialExpressionFunctionInput* MaterialFunctionInput)
 {
-	UE::HLSLTree::EExpressionType ExpressionType = UE::HLSLTree::EExpressionType::Void;
+	UE::Shader::EValueType ExpressionType = UE::Shader::EValueType::Void;
 	switch (MaterialFunctionInput->InputType)
 	{
-	case FunctionInput_Scalar: ExpressionType = UE::HLSLTree::EExpressionType::Float1; break;
-	case FunctionInput_Vector2: ExpressionType = UE::HLSLTree::EExpressionType::Float2; break;
-	case FunctionInput_Vector3: ExpressionType = UE::HLSLTree::EExpressionType::Float3; break;
-	case FunctionInput_Vector4: ExpressionType = UE::HLSLTree::EExpressionType::Float4; break;
+	case FunctionInput_Scalar: ExpressionType = UE::Shader::EValueType::Float1; break;
+	case FunctionInput_Vector2: ExpressionType = UE::Shader::EValueType::Float2; break;
+	case FunctionInput_Vector3: ExpressionType = UE::Shader::EValueType::Float3; break;
+	case FunctionInput_Vector4: ExpressionType = UE::Shader::EValueType::Float4; break;
 	case FunctionInput_Texture2D:
 	case FunctionInput_TextureCube:
 	case FunctionInput_Texture2DArray:
@@ -115,7 +110,7 @@ UE::HLSLTree::FExpression* FMaterialHLSLGenerator::NewFunctionInput(UE::HLSLTree
 		// TODO
 		break;
 	}
-	check(ExpressionType != UE::HLSLTree::EExpressionType::Void);
+	check(ExpressionType != UE::Shader::EValueType::Void);
 
 	UE::HLSLTree::FExpressionFunctionInput* Expression = HLSLTree->NewExpression<UE::HLSLTree::FExpressionFunctionInput>(Scope, MaterialFunctionInput->InputName, ExpressionType, InputIndex);
 	const FExpressionKey Key(MaterialFunctionInput, 0);
@@ -124,7 +119,7 @@ UE::HLSLTree::FExpression* FMaterialHLSLGenerator::NewFunctionInput(UE::HLSLTree
 	return Expression;
 }
 
-UE::HLSLTree::FLocalDeclaration* FMaterialHLSLGenerator::AcquireLocalDeclaration(UE::HLSLTree::FScope& Scope, UE::HLSLTree::EExpressionType Type, const FName& Name)
+UE::HLSLTree::FLocalDeclaration* FMaterialHLSLGenerator::AcquireLocalDeclaration(UE::HLSLTree::FScope& Scope, UE::Shader::EValueType Type, const FName& Name)
 {
 	UE::HLSLTree::FLocalDeclaration*& Declaration = LocalDeclarationMap.FindOrAdd(Name);
 	if (!Declaration)
@@ -140,15 +135,15 @@ UE::HLSLTree::FLocalDeclaration* FMaterialHLSLGenerator::AcquireLocalDeclaration
 		else
 		{
 			Errorf(TEXT("Local %s first accessed as type %s, now type %s"), *Name.ToString(),
-				UE::HLSLTree::GetExpressionTypeDescription(Declaration->Type).Name,
-				UE::HLSLTree::GetExpressionTypeDescription(Type).Name);
+				UE::Shader::GetValueTypeDescription(Declaration->Type).Name,
+				UE::Shader::GetValueTypeDescription(Type).Name);
 			return nullptr;
 		}
 	}
 	return Declaration;
 }
 
-UE::HLSLTree::FParameterDeclaration* FMaterialHLSLGenerator::AcquireParameterDeclaration(UE::HLSLTree::FScope& Scope, const FName& Name, const UE::HLSLTree::FConstant& DefaultValue)
+UE::HLSLTree::FParameterDeclaration* FMaterialHLSLGenerator::AcquireParameterDeclaration(UE::HLSLTree::FScope& Scope, const FName& Name, const UE::Shader::FValue& DefaultValue)
 {
 	UE::HLSLTree::FParameterDeclaration*& Declaration = ParameterDeclarationMap.FindOrAdd(Name);
 	if (!Declaration)
diff --git a/Engine/Source/Runtime/Engine/Private/Materials/MaterialShared.cpp b/Engine/Source/Runtime/Engine/Private/Materials/MaterialShared.cpp
--- a/Engine/Source/Runtime/Engine/Private/Materials/MaterialShared.cpp
+++ b/Engine/Source/Runtime/Engine/Private/Materials/MaterialShared.cpp
@@ -67,7 +67,6 @@ IMPLEMENT_TYPE_LAYOUT(FMaterialCompilationOutput);
 IMPLEMENT_TYPE_LAYOUT(FMeshMaterialShaderMap);
 IMPLEMENT_TYPE_LAYOUT(FMaterialProcessedSource);
 IMPLEMENT_TYPE_LAYOUT(FMaterialShaderMapContent);
-IMPLEMENT_TYPE_LAYOUT(FMaterialPreshaderData);
 IMPLEMENT_TYPE_LAYOUT(FMaterialUniformPreshaderHeader);
 IMPLEMENT_TYPE_LAYOUT(FMaterialScalarParameterInfo);
 IMPLEMENT_TYPE_LAYOUT(FMaterialVectorParameterInfo);
@@ -279,7 +278,7 @@ UE::HLSLTree::FExpression* FExpressionInput::AcquireHLSLExpression(FMaterialHLSL
 	return Result;
 }
 
-UE::HLSLTree::FExpression* FExpressionInput::AcquireHLSLExpressionWithCast(FMaterialHLSLGenerator& Generator, UE::HLSLTree::FScope& Scope, UE::HLSLTree::EExpressionType Type) const
+UE::HLSLTree::FExpression* FExpressionInput::AcquireHLSLExpressionWithCast(FMaterialHLSLGenerator& Generator, UE::HLSLTree::FScope& Scope, UE::Shader::EValueType Type) const
 {
 	UE::HLSLTree::FExpression* Result = AcquireHLSLExpression(Generator, Scope);
 	return Generator.NewCast(Scope, Type, Result);
@@ -2315,7 +2314,7 @@ bool FMaterial::Translate_New(const FMaterialShaderMapId& ShaderMapId,
 	FMaterialHLSLTree Tree;
 	Tree.InitializeForMaterial(TargetParams, *this);
 
-	return MaterialEmitHLSL(TargetParams, *this, Tree.GetTree(), OutCompilationOutput, OutMaterialEnvironment);
+	return MaterialEmitHLSL(TargetParams, *this, InStaticParameters, Tree.GetTree(), OutCompilationOutput, OutMaterialEnvironment);
 #else
 	checkNoEntry();
 	return false;
diff --git a/Engine/Source/Runtime/Engine/Private/Materials/MaterialUniformExpressions.cpp b/Engine/Source/Runtime/Engine/Private/Materials/MaterialUniformExpressions.cpp
--- a/Engine/Source/Runtime/Engine/Private/Materials/MaterialUniformExpressions.cpp
+++ b/Engine/Source/Runtime/Engine/Private/Materials/MaterialUniformExpressions.cpp
@@ -10,6 +10,7 @@
 #include "Materials/MaterialInstance.h"
 #include "Materials/MaterialInstanceSupport.h"
 #include "Materials/MaterialParameterCollection.h"
+#include "Shader/PreshaderEvaluate.h"
 #include "ExternalTexture.h"
 #include "Misc/UObjectToken.h"
 
@@ -44,7 +45,7 @@ TMap<FName,FMaterialUniformExpressionType*>& FMaterialUniformExpressionType::Get
 	return TypeMap;
 }
 
-static FGuid GetExternalTextureGuid(const FMaterialRenderContext& Context, const FGuid& ExternalTextureGuid, const FName& ParameterName, int32 SourceTextureIndex)
+FGuid FMaterialRenderContext::GetExternalTextureGuid(const FGuid& ExternalTextureGuid, const FName& ParameterName, int32 SourceTextureIndex) const
 {
 	FGuid GuidToLookup;
 	if (ExternalTextureGuid.IsValid())
@@ -55,14 +56,14 @@ static FGuid GetExternalTextureGuid(const FMaterialRenderContext& Context, const
 	else
 	{
 		const UTexture* TextureParameterObject = nullptr;
-		if (!ParameterName.IsNone() && Context.MaterialRenderProxy && Context.MaterialRenderProxy->GetTextureValue(ParameterName, &TextureParameterObject, Context) && TextureParameterObject)
+		if (!ParameterName.IsNone() && MaterialRenderProxy && MaterialRenderProxy->GetTextureValue(ParameterName, &TextureParameterObject, *this) && TextureParameterObject)
 		{
 			GuidToLookup = TextureParameterObject->GetExternalTextureGuid();
 		}
 		else
 		{
 			// Otherwise attempt to use the texture index in the material, if it's valid
-			const UTexture* TextureObject = SourceTextureIndex != INDEX_NONE ? GetIndexedTexture<UTexture>(Context.Material, SourceTextureIndex) : nullptr;
+			const UTexture* TextureObject = SourceTextureIndex != INDEX_NONE ? GetIndexedTexture<UTexture>(Material, SourceTextureIndex) : nullptr;
 			if (TextureObject)
 			{
 				GuidToLookup = TextureObject->GetExternalTextureGuid();
@@ -72,54 +73,54 @@ static FGuid GetExternalTextureGuid(const FMaterialRenderContext& Context, const
 	return GuidToLookup;
 }
 
-static void GetTextureParameterValue(const FHashedMaterialParameterInfo& ParameterInfo, int32 TextureIndex, const FMaterialRenderContext& Context, const UTexture*& OutValue)
+void FMaterialRenderContext::GetTextureParameterValue(const FHashedMaterialParameterInfo& ParameterInfo, int32 TextureIndex, const UTexture*& OutValue) const
 {
 	if (ParameterInfo.Name.IsNone())
 	{
-		OutValue = GetIndexedTexture<UTexture>(Context.Material, TextureIndex);
+		OutValue = GetIndexedTexture<UTexture>(Material, TextureIndex);
 	}
-	else if (!Context.MaterialRenderProxy || !Context.MaterialRenderProxy->GetTextureValue(ParameterInfo, &OutValue, Context))
+	else if (!MaterialRenderProxy || !MaterialRenderProxy->GetTextureValue(ParameterInfo, &OutValue, *this))
 	{
 		UTexture* Value = nullptr;
 
-		if (Context.Material.HasMaterialLayers())
+		if (Material.HasMaterialLayers())
 		{
-			UMaterialInterface* Interface = Context.Material.GetMaterialInterface();
+			UMaterialInterface* Interface = Material.GetMaterialInterface();
 			if (!Interface || !Interface->GetTextureParameterDefaultValue(ParameterInfo, Value))
 			{
-				Value = GetIndexedTexture<UTexture>(Context.Material, TextureIndex);
+				Value = GetIndexedTexture<UTexture>(Material, TextureIndex);
 			}
 		}
 		else
 		{
-			Value = GetIndexedTexture<UTexture>(Context.Material, TextureIndex);
+			Value = GetIndexedTexture<UTexture>(Material, TextureIndex);
 		}
 
 		OutValue = Value;
 	}
 }
 
-static void GetTextureParameterValue(const FHashedMaterialParameterInfo& ParameterInfo, int32 TextureIndex, const FMaterialRenderContext& Context, const URuntimeVirtualTexture*& OutValue)
+void FMaterialRenderContext::GetTextureParameterValue(const FHashedMaterialParameterInfo& ParameterInfo, int32 TextureIndex, const URuntimeVirtualTexture*& OutValue) const
 {
 	if (ParameterInfo.Name.IsNone())
 	{
-		OutValue = GetIndexedTexture<URuntimeVirtualTexture>(Context.Material, TextureIndex);
+		OutValue = GetIndexedTexture<URuntimeVirtualTexture>(Material, TextureIndex);
 	}
-	else if (!Context.MaterialRenderProxy || !Context.MaterialRenderProxy->GetTextureValue(ParameterInfo, &OutValue, Context))
+	else if (!MaterialRenderProxy || !MaterialRenderProxy->GetTextureValue(ParameterInfo, &OutValue, *this))
 	{
 		URuntimeVirtualTexture* Value = nullptr;
 
-		if (Context.Material.HasMaterialLayers())
+		if (Material.HasMaterialLayers())
 		{
-			UMaterialInterface* Interface = Context.Material.GetMaterialInterface();
+			UMaterialInterface* Interface = Material.GetMaterialInterface();
 			if (!Interface || !Interface->GetRuntimeVirtualTextureParameterDefaultValue(ParameterInfo, Value))
 			{
-				Value = GetIndexedTexture<URuntimeVirtualTexture>(Context.Material, TextureIndex);
+				Value = GetIndexedTexture<URuntimeVirtualTexture>(Material, TextureIndex);
 			}
 		}
 		else
 		{
-			Value = GetIndexedTexture<URuntimeVirtualTexture>(Context.Material, TextureIndex);
+			Value = GetIndexedTexture<URuntimeVirtualTexture>(Material, TextureIndex);
 		}
 
 		OutValue = Value;
@@ -133,10 +134,10 @@ FMaterialUniformExpressionType::FMaterialUniformExpressionType(const TCHAR* InNa
 	(new TLinkedList<FMaterialUniformExpressionType*>(this))->LinkHead(GetTypeList());
 }
 
-void FMaterialUniformExpression::WriteNumberOpcodes(FMaterialPreshaderData& OutData) const
+void FMaterialUniformExpression::WriteNumberOpcodes(UE::Shader::FPreshaderData& OutData) const
 {
 	UE_LOG(LogMaterial, Warning, TEXT("Missing WriteNumberOpcodes impl for %s"), GetType()->GetName());
-	OutData.WriteOpcode(EMaterialPreshaderOpcode::ConstantZero);
+	OutData.WriteOpcode(UE::Shader::EPreshaderOpcode::ConstantZero);
 }
 
 void FUniformParameterOverrides::SetScalarOverride(const FHashedMaterialParameterInfo& ParameterInfo, float Value, bool bOverride)
@@ -562,506 +563,13 @@ FUniformExpressionSet::FVTPackedStackAndLayerIndex FUniformExpressionSet::GetVTS
 	return FVTPackedStackAndLayerIndex(0xffff, 0xffff);
 }
 
-void FMaterialPreshaderData::WriteData(const void* Value, uint32 Size)
-{
-	Data.Append((uint8*)Value, Size);
-}
-
-void FMaterialPreshaderData::WriteName(const FScriptName& Name)
-{
-	int32 Index = Names.Find(Name);
-	if (Index == INDEX_NONE)
-	{
-		Index = Names.Add(Name);
-	}
-	check(Index >= 0 && Index <= 0xffff);
-	const uint32 Offset = Data.Num();
-	NameOffsets.Add(Offset);
-	Write((uint16)Index);
-}
-
-void FMaterialPreshaderData::Append(const FMaterialPreshaderData& InPreshader)
-{
-	const uint32 BaseOffset = Data.Num();
-	Data.Append(InPreshader.Data);
-
-	TArray<uint16> NameIndexRemap;
-	NameIndexRemap.Empty(InPreshader.Names.Num());
-	for (const FScriptName& Name : InPreshader.Names)
-	{
-		const int32 Index = Names.AddUnique(Name);
-		check(Index >= 0 && Index <= 0xffff);
-		NameIndexRemap.Add(Index);
-	}
-
-	for (uint32 PrevOffset : InPreshader.NameOffsets)
-	{
-		const uint32 Offset = BaseOffset + PrevOffset;
-		uint16* PrevNameIndex = (uint16*)(Data.GetData() + Offset);
-		const uint16 RemapIndex = NameIndexRemap[*PrevNameIndex];
-		*PrevNameIndex = RemapIndex;
-		NameOffsets.Add(Offset);
-	}
-}
-
-namespace
-{
-	struct FPreshaderDataContext
-	{
-		explicit FPreshaderDataContext(const FMaterialPreshaderData& InData)
-			: Ptr(InData.Data.GetData())
-			, EndPtr(Ptr + InData.Data.Num())
-			, Names(InData.Names.GetData())
-			, NumNames(InData.Names.Num())
-		{}
-
-		explicit FPreshaderDataContext(const FPreshaderDataContext& InContext, const FMaterialUniformPreshaderHeader& InHeader)
-			: Ptr(InContext.Ptr + InHeader.OpcodeOffset)
-			, EndPtr(Ptr + InHeader.OpcodeSize)
-			, Names(InContext.Names)
-			, NumNames(InContext.NumNames)
-		{}
-
-		const uint8* RESTRICT Ptr;
-		const uint8* RESTRICT EndPtr;
-		const FScriptName* RESTRICT Names;
-		int32 NumNames;
-	};
-
-	template<typename T>
-	inline T ReadPreshaderValue(FPreshaderDataContext& RESTRICT Data)
-	{
-		T Result;
-		FMemory::Memcpy(&Result, Data.Ptr, sizeof(T));
-		Data.Ptr += sizeof(T);
-		checkSlow(Data.Ptr <= Data.EndPtr);
-		return Result;
-	}
-
-	template<>
-	inline uint8 ReadPreshaderValue<uint8>(FPreshaderDataContext& RESTRICT Data)
-	{
-		checkSlow(Data.Ptr < Data.EndPtr);
-		return *Data.Ptr++;
-	}
-
-	template<>
-	FScriptName ReadPreshaderValue<FScriptName>(FPreshaderDataContext& RESTRICT Data)
-	{
-		const int32 Index = ReadPreshaderValue<uint16>(Data);
-		check(Index >= 0 && Index < Data.NumNames);
-		return Data.Names[Index];
-	}
-
-	template<>
-	FName ReadPreshaderValue<FName>(FPreshaderDataContext& RESTRICT Data) = delete;
-
-	template<>
-	FHashedMaterialParameterInfo ReadPreshaderValue<FHashedMaterialParameterInfo>(FPreshaderDataContext& RESTRICT Data)
-	{
-		const FScriptName Name = ReadPreshaderValue<FScriptName>(Data);
-		const int32 Index = ReadPreshaderValue<int32>(Data);
-		const TEnumAsByte<EMaterialParameterAssociation> Association = ReadPreshaderValue<TEnumAsByte<EMaterialParameterAssociation>>(Data);
-		return FHashedMaterialParameterInfo(Name, Association, Index);
-	}
-}
-
-static void GetVectorParameter(const FUniformExpressionSet& UniformExpressionSet, uint32 ParameterIndex, const FMaterialRenderContext& Context, FLinearColor& OutValue)
-{
-	const FMaterialVectorParameterInfo& Parameter = UniformExpressionSet.GetVectorParameter(ParameterIndex);
-
-	OutValue.R = OutValue.G = OutValue.B = OutValue.A = 0;
-	bool bNeedsDefaultValue = false;
-	if (!Context.MaterialRenderProxy || !Context.MaterialRenderProxy->GetVectorValue(Parameter.ParameterInfo, &OutValue, Context))
-	{
-		const bool bOveriddenParameterOnly = Parameter.ParameterInfo.Association == EMaterialParameterAssociation::GlobalParameter;
-
-		if (Context.Material.HasMaterialLayers())
-		{
-			UMaterialInterface* Interface = Context.Material.GetMaterialInterface();
-			if (!Interface || !Interface->GetVectorParameterDefaultValue(Parameter.ParameterInfo, OutValue, bOveriddenParameterOnly))
-			{
-				bNeedsDefaultValue = true;
-			}
-		}
-		else
-		{
-			bNeedsDefaultValue = true;
-		}
-	}
-
-	if (bNeedsDefaultValue)
-	{
-#if WITH_EDITOR
-		if (!Context.Material.TransientOverrides.GetVectorOverride(Parameter.ParameterInfo, OutValue))
-#endif // WITH_EDITOR
-		{
-			Parameter.GetDefaultValue(OutValue);
-		}
-	}
-}
-
-static void GetScalarParameter(const FUniformExpressionSet& UniformExpressionSet, uint32 ParameterIndex, const FMaterialRenderContext& Context, FLinearColor& OutValue)
-{
-	const FMaterialScalarParameterInfo& Parameter = UniformExpressionSet.GetScalarParameter(ParameterIndex);
-
-	OutValue.A = 0;
-
-	bool bNeedsDefaultValue = false;
-	if (!Context.MaterialRenderProxy || !Context.MaterialRenderProxy->GetScalarValue(Parameter.ParameterInfo, &OutValue.A, Context))
-	{
-		const bool bOveriddenParameterOnly = Parameter.ParameterInfo.Association == EMaterialParameterAssociation::GlobalParameter;
-
-		if (Context.Material.HasMaterialLayers())
-		{
-			UMaterialInterface* Interface = Context.Material.GetMaterialInterface();
-			if (!Interface || !Interface->GetScalarParameterDefaultValue(Parameter.ParameterInfo, OutValue.A, bOveriddenParameterOnly))
-			{
-				bNeedsDefaultValue = true;
-			}
-		}
-		else
-		{
-			bNeedsDefaultValue = true;
-		}
-	}
-
-	if (bNeedsDefaultValue)
-	{
-#if WITH_EDITOR
-		if (!Context.Material.TransientOverrides.GetScalarOverride(Parameter.ParameterInfo, OutValue.A))
-#endif // WITH_EDITOR
-		{
-			Parameter.GetDefaultValue(OutValue.A);
-		}
-	}
-
-	OutValue.R = OutValue.G = OutValue.B = OutValue.A;
-}
-
-using FPreshaderStack = TArray<FLinearColor, TInlineAllocator<64u>>;
-
-template<typename Operation>
-static inline void EvaluateUnaryOp(FPreshaderStack& Stack, const Operation& Op)
-{
-	const FLinearColor Value = Stack.Pop(false);
-	Stack.Add(FLinearColor(Op(Value.R), Op(Value.G), Op(Value.B), Op(Value.A)));
-}
-
-template<typename Operation>
-static inline void EvaluateBinaryOp(FPreshaderStack& Stack, const Operation& Op)
-{
-	const FLinearColor Value1 = Stack.Pop(false);
-	const FLinearColor Value0 = Stack.Pop(false);
-	Stack.Add(FLinearColor(Op(Value0.R, Value1.R), Op(Value0.G, Value1.G), Op(Value0.B, Value1.B), Op(Value0.A, Value1.A)));
-}
-
-template<typename Operation>
-static inline void EvaluateTernaryOp(FPreshaderStack& Stack, const Operation& Op)
-{
-	const FLinearColor Value2 = Stack.Pop(false);
-	const FLinearColor Value1 = Stack.Pop(false);
-	const FLinearColor Value0 = Stack.Pop(false);
-	Stack.Add(FLinearColor(Op(Value0.R, Value1.R, Value2.R), Op(Value0.G, Value1.G, Value2.G), Op(Value0.B, Value1.B, Value2.B), Op(Value0.A, Value1.A, Value2.A)));
-}
-
-static void EvaluateDot(FPreshaderStack& Stack, FPreshaderDataContext& RESTRICT Data)
-{
-	const uint8 ValueType = ReadPreshaderValue<uint8>(Data);
-	const FLinearColor Value1 = Stack.Pop(false);
-	const FLinearColor Value0 = Stack.Pop(false);
-	float Result = Value0.R * Value1.R;
-	Result += (ValueType >= MCT_Float2) ? Value0.G * Value1.G : 0;
-	Result += (ValueType >= MCT_Float3) ? Value0.B * Value1.B : 0;
-	Result += (ValueType >= MCT_Float4) ? Value0.A * Value1.A : 0;
-	Stack.Add(FLinearColor(Result, Result, Result, Result));
-}
-
-static void EvaluateCross(FPreshaderStack& Stack, FPreshaderDataContext& RESTRICT Data)
-{
-	const uint8 ValueType = ReadPreshaderValue<uint8>(Data);
-	FLinearColor ValueB = Stack.Pop(false);
-	FLinearColor ValueA = Stack.Pop(false);
-	
-	// Must be Float3, replicate CoerceParameter behavior
-	switch (ValueType)
-	{
-	case MCT_Float:
-		ValueA.B = ValueA.G = ValueA.R;
-		ValueB.B = ValueB.G = ValueB.R;
-		break;
-	case MCT_Float1:
-		ValueA.B = ValueA.G = 0.f;
-		ValueB.B = ValueB.G = 0.f;
-		break;
-	case MCT_Float2:
-		ValueA.B = 0.f;
-		ValueB.B = 0.f;
-		break;
-	};
-
-	const FVector Cross = FVector::CrossProduct(FVector(ValueA), FVector(ValueB));
-	Stack.Add(FLinearColor(Cross.X, Cross.Y, Cross.Z, 0.0f));
-}
-
-static void EvaluateComponentSwizzle(FPreshaderStack& Stack, FPreshaderDataContext& RESTRICT Data)
-{
-	const uint8 NumElements = ReadPreshaderValue<uint8>(Data);
-	const uint8 IndexR = ReadPreshaderValue<uint8>(Data);
-	const uint8 IndexG = ReadPreshaderValue<uint8>(Data);
-	const uint8 IndexB = ReadPreshaderValue<uint8>(Data);
-	const uint8 IndexA = ReadPreshaderValue<uint8>(Data);
-
-	FLinearColor Value = Stack.Pop(false);
-	FLinearColor Result(0.0f, 0.0f, 0.0f, 0.0f);
-	switch (NumElements)
-	{
-	case 1:
-		// Replicate scalar
-		Result.R = Result.G = Result.B = Result.A = Value.Component(IndexR);
-		break;
-	case 4:
-		Result.A = Value.Component(IndexA);
-		// Fallthrough...
-	case 3:
-		Result.B = Value.Component(IndexB);
-		// Fallthrough...
-	case 2:
-		Result.G = Value.Component(IndexG);
-		Result.R = Value.Component(IndexR);
-		break;
-	default:
-		UE_LOG(LogMaterial, Fatal, TEXT("Invalid number of swizzle elements: %d"), NumElements);
-		break;
-	}
-	Stack.Add(Result);
-}
-
-static void EvaluateAppenedVector(FPreshaderStack& Stack, FPreshaderDataContext& RESTRICT Data)
-{
-	const uint8 NumComponentsA = ReadPreshaderValue<uint8>(Data);
-
-	const FLinearColor ValueB = Stack.Pop(false);
-	const FLinearColor ValueA = Stack.Pop(false);
-
-	FLinearColor Result;
-	Result.R = NumComponentsA >= 1 ? ValueA.R : (&ValueB.R)[0 - NumComponentsA];
-	Result.G = NumComponentsA >= 2 ? ValueA.G : (&ValueB.R)[1 - NumComponentsA];
-	Result.B = NumComponentsA >= 3 ? ValueA.B : (&ValueB.R)[2 - NumComponentsA];
-	Result.A = NumComponentsA >= 4 ? ValueA.A : (&ValueB.R)[3 - NumComponentsA];
-	Stack.Add(Result);
-}
-
-static const UTexture* GetTextureParameter(const FMaterialRenderContext& Context, FPreshaderDataContext& RESTRICT Data)
-{
-	const FHashedMaterialParameterInfo ParameterInfo = ReadPreshaderValue<FHashedMaterialParameterInfo>(Data);
-	const int32 TextureIndex = ReadPreshaderValue<int32>(Data);
-	
-	const UTexture* Texture = nullptr;
-	GetTextureParameterValue(ParameterInfo, TextureIndex, Context, Texture);
-	return Texture;
-}
-
-static void EvaluateTextureSize(const FMaterialRenderContext& Context, FPreshaderStack& Stack, FPreshaderDataContext& RESTRICT Data)
-{
-	const UTexture* Texture = GetTextureParameter(Context, Data);
-	if (Texture && Texture->Resource)
-	{
-		const uint32 SizeX = Texture->Resource->GetSizeX();
-		const uint32 SizeY = Texture->Resource->GetSizeY();
-		const uint32 SizeZ = Texture->Resource->GetSizeZ();
-		Stack.Add(FLinearColor((float)SizeX, (float)SizeY, (float)SizeZ, 0.0f));
-	}
-	else
-	{
-		Stack.Add(FLinearColor(0.0f, 0.0f, 0.0f, 0.0f));
-	}
-}
-
-static void EvaluateTexelSize(const FMaterialRenderContext& Context, FPreshaderStack& Stack, FPreshaderDataContext& RESTRICT Data)
-{
-	const UTexture* Texture = GetTextureParameter(Context, Data);
-	if (Texture && Texture->Resource)
-	{
-		const uint32 SizeX = Texture->Resource->GetSizeX();
-		const uint32 SizeY = Texture->Resource->GetSizeY();
-		const uint32 SizeZ = Texture->Resource->GetSizeZ();
-		Stack.Add(FLinearColor(1.0f / (float)SizeX, 1.0f / (float)SizeY, (SizeZ > 0 ? 1.0f / (float)SizeZ : 0.0f), 0.0f));
-	}
-	else
-	{
-		Stack.Add(FLinearColor(0.0f, 0.0f, 0.0f, 0.0f));
-	}
-}
-
-static FGuid GetExternalTextureGuid(const FMaterialRenderContext& Context, FPreshaderDataContext& RESTRICT Data)
-{
-	const FScriptName ParameterName = ReadPreshaderValue<FScriptName>(Data);
-	const FGuid ExternalTextureGuid = ReadPreshaderValue<FGuid>(Data);
-	const int32 TextureIndex = ReadPreshaderValue<int32>(Data);
-	return GetExternalTextureGuid(Context, ExternalTextureGuid, ScriptNameToName(ParameterName), TextureIndex);
-}
-
-static void EvaluateExternalTextureCoordinateScaleRotation(const FMaterialRenderContext& Context, FPreshaderStack& Stack, FPreshaderDataContext& RESTRICT Data)
-{
-	const FGuid GuidToLookup = GetExternalTextureGuid(Context, Data);
-	FLinearColor Result(1.f, 0.f, 0.f, 1.f);
-	if (GuidToLookup.IsValid())
-	{
-		FExternalTextureRegistry::Get().GetExternalTextureCoordinateScaleRotation(GuidToLookup, Result);
-	}
-	Stack.Add(Result);
-}
-
-static void EvaluateExternalTextureCoordinateOffset(const FMaterialRenderContext& Context, FPreshaderStack& Stack, FPreshaderDataContext& RESTRICT Data)
-{
-	const FGuid GuidToLookup = GetExternalTextureGuid(Context, Data);
-	FLinearColor Result(0.f, 0.f, 0.f, 0.f);
-	if (GuidToLookup.IsValid())
-	{
-		FExternalTextureRegistry::Get().GetExternalTextureCoordinateOffset(GuidToLookup, Result);
-	}
-	Stack.Add(Result);
-}
-
-static void EvaluateRuntimeVirtualTextureUniform(const FMaterialRenderContext& Context, FPreshaderStack& Stack, FPreshaderDataContext& RESTRICT Data)
-{
-	const FHashedMaterialParameterInfo ParameterInfo = ReadPreshaderValue<FHashedMaterialParameterInfo>(Data);
-	const int32 TextureIndex = ReadPreshaderValue<int32>(Data);
-	const int32 VectorIndex = ReadPreshaderValue<int32>(Data);
-
-	const URuntimeVirtualTexture* Texture = nullptr;
-	if (ParameterInfo.Name.IsNone() || !Context.MaterialRenderProxy || !Context.MaterialRenderProxy->GetTextureValue(ParameterInfo, &Texture, Context))
-	{
-		Texture = GetIndexedTexture<URuntimeVirtualTexture>(Context.Material, TextureIndex);
-	}
-	if (Texture != nullptr && VectorIndex != INDEX_NONE)
-	{
-		Stack.Add(FLinearColor(Texture->GetUniformParameter(VectorIndex)));
-	}
-	else
-	{
-		Stack.Add(FLinearColor(0.f, 0.f, 0.f, 0.f));
-	}
-}
-
-/** Converts an arbitrary number into a safe divisor. i.e. FMath::Abs(Number) >= DELTA */
-static float GetSafeDivisor(float Number)
-{
-	if (FMath::Abs(Number) < DELTA)
-	{
-		if (Number < 0.0f)
-		{
-			return -DELTA;
-		}
-		else
-		{
-			return +DELTA;
-		}
-	}
-	else
-	{
-		return Number;
-	}
-}
-
-/**
- * FORCENOINLINE is required to discourage compiler from vectorizing the Div operation, which may tempt it into optimizing divide as A * rcp(B)
- * This will break shaders that are depending on exact divide results (see SubUV material function)
- * Technically this could still happen for a scalar divide, but it doesn't seem to occur in practice
- */
-FORCENOINLINE static float DivideComponent(float A, float B)
-{
-	return A / GetSafeDivisor(B);
-}
-
-static void EvaluatePreshader(const FUniformExpressionSet* UniformExpressionSet, const FMaterialRenderContext& Context, FPreshaderStack& Stack, FPreshaderDataContext& RESTRICT Data, FLinearColor& OutValue)
-{
-	static const float LogToLog10 = 1.0f / FMath::Loge(10.f);
-	uint8 const* const DataEnd = Data.EndPtr;
-
-	Stack.Reset();
-	while (Data.Ptr < DataEnd)
-	{
-		const EMaterialPreshaderOpcode Opcode = (EMaterialPreshaderOpcode)ReadPreshaderValue<uint8>(Data);
-		switch (Opcode)
-		{
-		case EMaterialPreshaderOpcode::ConstantZero:
-			Stack.Add(FLinearColor(0.0f, 0.0f, 0.0f, 0.0f));
-			break;
-		case EMaterialPreshaderOpcode::Constant:
-			Stack.Add(ReadPreshaderValue<FLinearColor>(Data));
-			break;
-		case EMaterialPreshaderOpcode::VectorParameter:
-			check(UniformExpressionSet);
-			GetVectorParameter(*UniformExpressionSet, ReadPreshaderValue<uint16>(Data), Context, Stack.AddDefaulted_GetRef());
-			break;
-		case EMaterialPreshaderOpcode::ScalarParameter:
-			check(UniformExpressionSet);
-			GetScalarParameter(*UniformExpressionSet, ReadPreshaderValue<uint16>(Data), Context, Stack.AddDefaulted_GetRef());
-			break;
-		case EMaterialPreshaderOpcode::Add: EvaluateBinaryOp(Stack, [](float Lhs, float Rhs) { return Lhs + Rhs; }); break;
-		case EMaterialPreshaderOpcode::Sub: EvaluateBinaryOp(Stack, [](float Lhs, float Rhs) { return Lhs - Rhs; }); break;
-		case EMaterialPreshaderOpcode::Mul: EvaluateBinaryOp(Stack, [](float Lhs, float Rhs) { return Lhs * Rhs; }); break;
-		case EMaterialPreshaderOpcode::Div: EvaluateBinaryOp(Stack, [](float Lhs, float Rhs) { return DivideComponent(Lhs, Rhs); }); break;
-		case EMaterialPreshaderOpcode::Fmod: EvaluateBinaryOp(Stack, [](float Lhs, float Rhs) { return FMath::Fmod(Lhs, Rhs); }); break;
-		case EMaterialPreshaderOpcode::Min: EvaluateBinaryOp(Stack, [](float Lhs, float Rhs) { return FMath::Min(Lhs, Rhs); }); break;
-		case EMaterialPreshaderOpcode::Max: EvaluateBinaryOp(Stack, [](float Lhs, float Rhs) { return FMath::Max(Lhs, Rhs); }); break;
-		case EMaterialPreshaderOpcode::Clamp: EvaluateTernaryOp(Stack, [](float A, float B, float C) { return FMath::Clamp(A, B, C); }); break;
-		case EMaterialPreshaderOpcode::Dot: EvaluateDot(Stack, Data); break;
-		case EMaterialPreshaderOpcode::Cross: EvaluateCross(Stack, Data); break;
-		case EMaterialPreshaderOpcode::Sqrt: EvaluateUnaryOp(Stack, [](float V) { return FMath::Sqrt(V); }); break;
-		case EMaterialPreshaderOpcode::Sin: EvaluateUnaryOp(Stack, [](float V) { return FMath::Sin(V); }); break;
-		case EMaterialPreshaderOpcode::Cos: EvaluateUnaryOp(Stack, [](float V) { return FMath::Cos(V); }); break;
-		case EMaterialPreshaderOpcode::Tan: EvaluateUnaryOp(Stack, [](float V) { return FMath::Tan(V); }); break;
-		case EMaterialPreshaderOpcode::Asin: EvaluateUnaryOp(Stack, [](float V) { return FMath::Asin(V); }); break;
-		case EMaterialPreshaderOpcode::Acos: EvaluateUnaryOp(Stack, [](float V) { return FMath::Acos(V); }); break;
-		case EMaterialPreshaderOpcode::Atan: EvaluateUnaryOp(Stack, [](float V) { return FMath::Atan(V); }); break;
-		case EMaterialPreshaderOpcode::Atan2: EvaluateBinaryOp(Stack, [](float A, float B) { return FMath::Atan2(A, B); }); break;
-		case EMaterialPreshaderOpcode::Abs: EvaluateUnaryOp(Stack, [](float V) { return FMath::Abs(V); }); break;
-		case EMaterialPreshaderOpcode::Saturate: EvaluateUnaryOp(Stack, [](float V) { return FMath::Clamp(V, 0.0f, 1.0f); }); break;
-		case EMaterialPreshaderOpcode::Floor: EvaluateUnaryOp(Stack, [](float V) { return FMath::FloorToFloat(V); }); break;
-		case EMaterialPreshaderOpcode::Ceil: EvaluateUnaryOp(Stack, [](float V) { return FMath::CeilToFloat(V); }); break;
-		case EMaterialPreshaderOpcode::Round: EvaluateUnaryOp(Stack, [](float V) { return FMath::RoundToFloat(V); }); break;
-		case EMaterialPreshaderOpcode::Trunc: EvaluateUnaryOp(Stack, [](float V) { return FMath::TruncToFloat(V); }); break;
-		case EMaterialPreshaderOpcode::Sign: EvaluateUnaryOp(Stack, [](float V) { return FMath::Sign(V); }); break;
-		case EMaterialPreshaderOpcode::Frac: EvaluateUnaryOp(Stack, [](float V) { return FMath::Frac(V); }); break;
-		case EMaterialPreshaderOpcode::Fractional: EvaluateUnaryOp(Stack, [](float V) { return FMath::Fractional(V); }); break;
-		case EMaterialPreshaderOpcode::Log2: EvaluateUnaryOp(Stack, [](float V) { return FMath::Log2(V); }); break;
-		case EMaterialPreshaderOpcode::Log10: EvaluateUnaryOp(Stack, [](float V) { return FMath::Loge(V) * LogToLog10; }); break;
-		case EMaterialPreshaderOpcode::ComponentSwizzle: EvaluateComponentSwizzle(Stack, Data); break;
-		case EMaterialPreshaderOpcode::AppendVector: EvaluateAppenedVector(Stack, Data); break;
-		case EMaterialPreshaderOpcode::TextureSize: EvaluateTextureSize(Context, Stack, Data); break;
-		case EMaterialPreshaderOpcode::TexelSize: EvaluateTexelSize(Context, Stack, Data); break;
-		case EMaterialPreshaderOpcode::ExternalTextureCoordinateScaleRotation: EvaluateExternalTextureCoordinateScaleRotation(Context, Stack, Data); break;
-		case EMaterialPreshaderOpcode::ExternalTextureCoordinateOffset: EvaluateExternalTextureCoordinateOffset(Context, Stack, Data); break;
-		case EMaterialPreshaderOpcode::RuntimeVirtualTextureUniform: EvaluateRuntimeVirtualTextureUniform(Context, Stack, Data); break;
-		default:
-			UE_LOG(LogMaterial, Fatal, TEXT("Unknown preshader opcode %d"), (uint8)Opcode);
-			break;
-		}
-	}
-	check(Data.Ptr == DataEnd);
-
-	ensure(Stack.Num() <= 1);
-	if (Stack.Num() > 0)
-	{
-		OutValue = Stack.Last();
-	}
-}
-
-void FMaterialPreshaderData::Evaluate(FUniformExpressionSet* UniformExpressionSet, const struct FMaterialRenderContext& Context, FLinearColor& OutValue)
-{
-	FPreshaderStack Stack;
-	FPreshaderDataContext PreshaderContext(*this);
-	EvaluatePreshader(UniformExpressionSet, Context, Stack, PreshaderContext, OutValue);
-}
-
 void FMaterialUniformExpression::GetNumberValue(const struct FMaterialRenderContext& Context, FLinearColor& OutValue) const
 {
-	FMaterialPreshaderData PreshaderData;
+	UE::Shader::FPreshaderData PreshaderData;
 	WriteNumberOpcodes(PreshaderData);
-	PreshaderData.Evaluate(nullptr, Context, OutValue);
+	UE::Shader::FValue Value;
+	PreshaderData.Evaluate(nullptr, Context, Value);
+	OutValue = Value.AsLinearColor();
 }
 
 const FMaterialVectorParameterInfo* FUniformExpressionSet::FindVectorParameter(const FHashedMaterialParameterInfo& ParameterInfo) const
@@ -1168,7 +676,7 @@ void FUniformExpressionSet::GetTextureValue(EMaterialTextureParameterType Type,
 		}
 	}
 #endif // WITH_EDITOR
-	GetTextureParameterValue(Parameter.ParameterInfo, Parameter.TextureIndex, Context, OutValue);
+	Context.GetTextureParameterValue(Parameter.ParameterInfo, Parameter.TextureIndex , OutValue);
 }
 
 void FUniformExpressionSet::GetTextureValue(int32 Index, const FMaterialRenderContext& Context, const FMaterial& Material, const URuntimeVirtualTexture*& OutValue) const
@@ -1178,7 +686,7 @@ void FUniformExpressionSet::GetTextureValue(int32 Index, const FMaterialRenderCo
 	if (ensure(Index < VirtualTexturesNum))
 	{
 		const FMaterialTextureParameterInfo& Parameter = GetTextureParameter(EMaterialTextureParameterType::Virtual, Index);
-		GetTextureParameterValue(Parameter.ParameterInfo, Parameter.TextureIndex, Context, OutValue);
+		Context.GetTextureParameterValue(Parameter.ParameterInfo, Parameter.TextureIndex, OutValue);
 	}
 	else
 	{
@@ -1260,18 +768,18 @@ void FUniformExpressionSet::FillUniformBuffer(const FMaterialRenderContext& Mate
 		}
 
 		// Dump vector expression into the buffer.
-		FPreshaderStack PreshaderStack;
-		FPreshaderDataContext PreshaderBaseContext(UniformPreshaderData);
+		UE::Shader::FPreshaderStack PreshaderStack;
+		UE::Shader::FPreshaderDataContext PreshaderBaseContext(UniformPreshaderData);
 		for(int32 VectorIndex = 0;VectorIndex < UniformVectorPreshaders.Num();++VectorIndex)
 		{
-			FLinearColor VectorValue(0, 0, 0, 0);
+			UE::Shader::FValue VectorValue(0.0f, 0.0f, 0.0f, 0.0f);
 
 			const FMaterialUniformPreshaderHeader& Preshader = UniformVectorPreshaders[VectorIndex];
-			FPreshaderDataContext PreshaderContext(PreshaderBaseContext, Preshader);
-			EvaluatePreshader(this, MaterialRenderContext, PreshaderStack, PreshaderContext, VectorValue);
+			UE::Shader::FPreshaderDataContext PreshaderContext(PreshaderBaseContext, Preshader.OpcodeOffset, Preshader.OpcodeSize);
+			UE::Shader::EvaluatePreshader(this, MaterialRenderContext, PreshaderStack, PreshaderContext, VectorValue);
 	
 			FLinearColor* DestAddress = (FLinearColor*)BufferCursor;
-			*DestAddress = VectorValue;
+			*DestAddress = VectorValue.AsLinearColor();
 			BufferCursor = DestAddress + 1;
 			check(BufferCursor <= TempBuffer + TempBufferSize);
 		}
@@ -1279,14 +787,14 @@ void FUniformExpressionSet::FillUniformBuffer(const FMaterialRenderContext& Mate
 		// Dump scalar expression into the buffer.
 		for(int32 ScalarIndex = 0;ScalarIndex < UniformScalarPreshaders.Num();++ScalarIndex)
 		{
-			FLinearColor VectorValue(0,0,0,0);
+			UE::Shader::FValue VectorValue(0.0f, 0.0f, 0.0f, 0.0f);
 
 			const FMaterialUniformPreshaderHeader& Preshader = UniformScalarPreshaders[ScalarIndex];
-			FPreshaderDataContext PreshaderContext(PreshaderBaseContext, Preshader);
-			EvaluatePreshader(this, MaterialRenderContext, PreshaderStack, PreshaderContext, VectorValue);
+			UE::Shader::FPreshaderDataContext PreshaderContext(PreshaderBaseContext, Preshader.OpcodeOffset, Preshader.OpcodeSize);
+			UE::Shader::EvaluatePreshader(this, MaterialRenderContext, PreshaderStack, PreshaderContext, VectorValue);
 
 			float* DestAddress = (float*)BufferCursor;
-			*DestAddress = VectorValue.R;
+			*DestAddress = VectorValue.AsFloat().Component[0];
 			BufferCursor = DestAddress + 1;
 			check(BufferCursor <= TempBuffer + TempBufferSize);
 		}
@@ -1796,7 +1304,7 @@ bool FMaterialUniformExpressionExternalTextureBase::IsIdentical(const FMaterialU
 
 FGuid FMaterialUniformExpressionExternalTextureBase::ResolveExternalTextureGUID(const FMaterialRenderContext& Context, TOptional<FName> ParameterName) const
 {
-	return GetExternalTextureGuid(Context, ExternalTextureGuid, ParameterName.IsSet() ? ParameterName.GetValue() : FName(), SourceTextureIndex);
+	return Context.GetExternalTextureGuid(ExternalTextureGuid, ParameterName.IsSet() ? ParameterName.GetValue() : FName(), SourceTextureIndex);
 }
 
 void FMaterialUniformExpressionExternalTexture::GetExternalTextureParameterInfo(FMaterialExternalTextureParameterInfo& OutParameter) const
@@ -1915,7 +1423,7 @@ void FMaterialTextureParameterInfo::GetGameThreadTextureValue(const UMaterialInt
 bool FMaterialExternalTextureParameterInfo::GetExternalTexture(const FMaterialRenderContext& Context, FTextureRHIRef& OutTextureRHI, FSamplerStateRHIRef& OutSamplerStateRHI) const
 {
 	check(IsInParallelRenderingThread());
-	const FGuid GuidToLookup = GetExternalTextureGuid(Context, ExternalTextureGuid, ScriptNameToName(ParameterName), SourceTextureIndex);
+	const FGuid GuidToLookup = Context.GetExternalTextureGuid(ExternalTextureGuid, ScriptNameToName(ParameterName), SourceTextureIndex);
 	return FExternalTextureRegistry::Get().GetExternalTexture(Context.MaterialRenderProxy, GuidToLookup, OutTextureRHI, OutSamplerStateRHI);
 }
 
@@ -1930,10 +1438,10 @@ bool FMaterialUniformExpressionExternalTextureCoordinateScaleRotation::IsIdentic
 	return ParameterName == Other->ParameterName;
 }
 
-void FMaterialUniformExpressionExternalTextureCoordinateScaleRotation::WriteNumberOpcodes(FMaterialPreshaderData& OutData) const
+void FMaterialUniformExpressionExternalTextureCoordinateScaleRotation::WriteNumberOpcodes(UE::Shader::FPreshaderData& OutData) const
 {
 	const FScriptName Name = ParameterName.IsSet() ? NameToScriptName(ParameterName.GetValue()) : FScriptName();
-	OutData.WriteOpcode(EMaterialPreshaderOpcode::ExternalTextureCoordinateScaleRotation).Write(Name).Write(ExternalTextureGuid).Write<int32>(SourceTextureIndex);
+	OutData.WriteOpcode(UE::Shader::EPreshaderOpcode::ExternalTextureCoordinateScaleRotation).Write(Name).Write(ExternalTextureGuid).Write<int32>(SourceTextureIndex);
 }
 
 bool FMaterialUniformExpressionExternalTextureCoordinateOffset::IsIdentical(const FMaterialUniformExpression* OtherExpression) const
@@ -1947,10 +1455,10 @@ bool FMaterialUniformExpressionExternalTextureCoordinateOffset::IsIdentical(cons
 	return ParameterName == Other->ParameterName;
 }
 
-void FMaterialUniformExpressionExternalTextureCoordinateOffset::WriteNumberOpcodes(FMaterialPreshaderData& OutData) const
+void FMaterialUniformExpressionExternalTextureCoordinateOffset::WriteNumberOpcodes(UE::Shader::FPreshaderData& OutData) const
 {
 	const FScriptName Name = ParameterName.IsSet() ? NameToScriptName(ParameterName.GetValue()) : FScriptName();
-	OutData.WriteOpcode(EMaterialPreshaderOpcode::ExternalTextureCoordinateOffset).Write(Name).Write(ExternalTextureGuid).Write<int32>(SourceTextureIndex);
+	OutData.WriteOpcode(UE::Shader::EPreshaderOpcode::ExternalTextureCoordinateOffset).Write(Name).Write(ExternalTextureGuid).Write<int32>(SourceTextureIndex);
 }
 
 FMaterialUniformExpressionRuntimeVirtualTextureUniform::FMaterialUniformExpressionRuntimeVirtualTextureUniform()
@@ -1986,10 +1494,10 @@ bool FMaterialUniformExpressionRuntimeVirtualTextureUniform::IsIdentical(const F
 	return ParameterInfo == Other->ParameterInfo && TextureIndex == Other->TextureIndex && VectorIndex == Other->VectorIndex;
 }
 
-void FMaterialUniformExpressionRuntimeVirtualTextureUniform::WriteNumberOpcodes(FMaterialPreshaderData& OutData) const
+void FMaterialUniformExpressionRuntimeVirtualTextureUniform::WriteNumberOpcodes(UE::Shader::FPreshaderData& OutData) const
 {
 	const FHashedMaterialParameterInfo WriteParameterInfo = bParameter ? ParameterInfo : FHashedMaterialParameterInfo();
-	OutData.WriteOpcode(EMaterialPreshaderOpcode::RuntimeVirtualTextureUniform).Write(WriteParameterInfo).Write((int32)TextureIndex).Write((int32)VectorIndex);
+	OutData.WriteOpcode(UE::Shader::EPreshaderOpcode::RuntimeVirtualTextureUniform).Write(WriteParameterInfo).Write((int32)TextureIndex).Write((int32)VectorIndex);
 }
 
 /**
diff --git a/Engine/Source/Runtime/Engine/Private/Materials/MaterialUniformExpressions.h b/Engine/Source/Runtime/Engine/Private/Materials/MaterialUniformExpressions.h
--- a/Engine/Source/Runtime/Engine/Private/Materials/MaterialUniformExpressions.h
+++ b/Engine/Source/Runtime/Engine/Private/Materials/MaterialUniformExpressions.h
@@ -64,7 +64,7 @@ public:
 	virtual bool IsConstant() const { return false; }
 	virtual bool IsIdentical(const FMaterialUniformExpression* OtherExpression) const { return false; }
 
-	virtual void WriteNumberOpcodes(FMaterialPreshaderData& OutData) const;
+	virtual void WriteNumberOpcodes(UE::Shader::FPreshaderData& OutData) const;
 
 	virtual void GetNumberValue(const struct FMaterialRenderContext& Context, FLinearColor& OutValue) const;
 };
@@ -190,10 +190,20 @@ public:
 		return OtherConstant->ValueType == ValueType && OtherConstant->Value == Value;
 	}
 
-	virtual void WriteNumberOpcodes(FMaterialPreshaderData& OutData) const override
+	virtual void WriteNumberOpcodes(UE::Shader::FPreshaderData& OutData) const override
 	{
-		OutData.WriteOpcode(EMaterialPreshaderOpcode::Constant);
-		OutData.Write(Value);
+		UE::Shader::FValue ConstantValue(Value);
+		switch (ValueType)
+		{
+		case MCT_Float:
+		case MCT_Float1: ConstantValue.NumComponents = 1; break;
+		case MCT_Float2: ConstantValue.NumComponents = 2; break;
+		case MCT_Float3: ConstantValue.NumComponents = 3; break;
+		case MCT_Float4: ConstantValue.NumComponents = 4; break;
+		default: checkNoEntry(); break;
+		}
+		OutData.WriteOpcode(UE::Shader::EPreshaderOpcode::Constant);
+		OutData.Write(ConstantValue);
 	}
 
 	virtual void GetNumberValue(const FMaterialRenderContext& Context, FLinearColor& OutValue) const override
@@ -222,9 +232,9 @@ public:
 	}
 
 	// FMaterialUniformExpression interface.
-	virtual void WriteNumberOpcodes(FMaterialPreshaderData& OutData) const override
+	virtual void WriteNumberOpcodes(UE::Shader::FPreshaderData& OutData) const override
 	{
-		OutData.WriteOpcode(EMaterialPreshaderOpcode::VectorParameter);
+		OutData.WriteOpcode(UE::Shader::EPreshaderOpcode::VectorParameter);
 		OutData.Write((uint16)ParameterIndex);
 	}
 
@@ -274,9 +284,9 @@ public:
 	}
 
 	// FMaterialUniformExpression interface.
-	virtual void WriteNumberOpcodes(FMaterialPreshaderData& OutData) const override
+	virtual void WriteNumberOpcodes(UE::Shader::FPreshaderData& OutData) const override
 	{
-		OutData.WriteOpcode(EMaterialPreshaderOpcode::ScalarParameter);
+		OutData.WriteOpcode(UE::Shader::EPreshaderOpcode::ScalarParameter);
 		OutData.Write((uint16)ParameterIndex);
 	}
 
@@ -452,10 +462,10 @@ public:
 	{}
 
 	// FMaterialUniformExpression interface.
-	virtual void WriteNumberOpcodes(FMaterialPreshaderData& OutData) const override
+	virtual void WriteNumberOpcodes(UE::Shader::FPreshaderData& OutData) const override
 	{
 		X->WriteNumberOpcodes(OutData);
-		OutData.WriteOpcode(bIsCosine ? EMaterialPreshaderOpcode::Cos : EMaterialPreshaderOpcode::Sin);
+		OutData.WriteOpcode(bIsCosine ? UE::Shader::EPreshaderOpcode::Cos : UE::Shader::EPreshaderOpcode::Sin);
 	}
 	virtual bool IsConstant() const
 	{
@@ -510,7 +520,7 @@ public:
 	{}
 
 	// FMaterialUniformExpression interface.
-	virtual void WriteNumberOpcodes(FMaterialPreshaderData& OutData) const override
+	virtual void WriteNumberOpcodes(UE::Shader::FPreshaderData& OutData) const override
 	{
 		X->WriteNumberOpcodes(OutData);
 		if (Op == TMO_Atan2)
@@ -519,13 +529,13 @@ public:
 		}
 		switch (Op)
 		{
-		case TMO_Sin: OutData.WriteOpcode(EMaterialPreshaderOpcode::Sin); break;
-		case TMO_Cos: OutData.WriteOpcode(EMaterialPreshaderOpcode::Cos); break;
-		case TMO_Tan: OutData.WriteOpcode(EMaterialPreshaderOpcode::Tan); break;
-		case TMO_Asin: OutData.WriteOpcode(EMaterialPreshaderOpcode::Asin); break;
-		case TMO_Acos: OutData.WriteOpcode(EMaterialPreshaderOpcode::Acos); break;
-		case TMO_Atan: OutData.WriteOpcode(EMaterialPreshaderOpcode::Atan); break;
-		case TMO_Atan2: OutData.WriteOpcode(EMaterialPreshaderOpcode::Atan2); break;
+		case TMO_Sin: OutData.WriteOpcode(UE::Shader::EPreshaderOpcode::Sin); break;
+		case TMO_Cos: OutData.WriteOpcode(UE::Shader::EPreshaderOpcode::Cos); break;
+		case TMO_Tan: OutData.WriteOpcode(UE::Shader::EPreshaderOpcode::Tan); break;
+		case TMO_Asin: OutData.WriteOpcode(UE::Shader::EPreshaderOpcode::Asin); break;
+		case TMO_Acos: OutData.WriteOpcode(UE::Shader::EPreshaderOpcode::Acos); break;
+		case TMO_Atan: OutData.WriteOpcode(UE::Shader::EPreshaderOpcode::Atan); break;
+		case TMO_Atan2: OutData.WriteOpcode(UE::Shader::EPreshaderOpcode::Atan2); break;
 		default: checkNoEntry(); break;
 		}
 	}
@@ -563,10 +573,10 @@ public:
 	{}
 
 	// FMaterialUniformExpression interface.
-	virtual void WriteNumberOpcodes(FMaterialPreshaderData& OutData) const override
+	virtual void WriteNumberOpcodes(UE::Shader::FPreshaderData& OutData) const override
 	{
 		X->WriteNumberOpcodes(OutData);
-		OutData.WriteOpcode(EMaterialPreshaderOpcode::Sqrt);
+		OutData.WriteOpcode(UE::Shader::EPreshaderOpcode::Sqrt);
 	}
 	virtual bool IsConstant() const
 	{
@@ -600,10 +610,10 @@ public:
 	{}
 
 	// FMaterialUniformExpression interface.
-	virtual void WriteNumberOpcodes(FMaterialPreshaderData& OutData) const override
+	virtual void WriteNumberOpcodes(UE::Shader::FPreshaderData& OutData) const override
 	{
 		X->WriteNumberOpcodes(OutData);
-		OutData.WriteOpcode(EMaterialPreshaderOpcode::Length).Write((uint8)ValueType);
+		OutData.WriteOpcode(UE::Shader::EPreshaderOpcode::Length).Write((uint8)ValueType);
 	}
 	virtual bool IsConstant() const
 	{
@@ -637,10 +647,10 @@ public:
 	{}
 
 	// FMaterialUniformExpression interface.
-	virtual void WriteNumberOpcodes(FMaterialPreshaderData& OutData) const override
+	virtual void WriteNumberOpcodes(UE::Shader::FPreshaderData& OutData) const override
 	{
 		X->WriteNumberOpcodes(OutData);
-		OutData.WriteOpcode(EMaterialPreshaderOpcode::Log2);
+		OutData.WriteOpcode(UE::Shader::EPreshaderOpcode::Log2);
 	}
 	bool IsConstant() const override
 	{
@@ -674,10 +684,10 @@ public:
 	{}
 
 	// FMaterialUniformExpression interface.
-	virtual void WriteNumberOpcodes(FMaterialPreshaderData& OutData) const override
+	virtual void WriteNumberOpcodes(UE::Shader::FPreshaderData& OutData) const override
 	{
 		X->WriteNumberOpcodes(OutData);
-		OutData.WriteOpcode(EMaterialPreshaderOpcode::Log10);
+		OutData.WriteOpcode(UE::Shader::EPreshaderOpcode::Log10);
 	}
 	bool IsConstant() const override
 	{
@@ -724,19 +734,19 @@ public:
 	{}
 
 	// FMaterialUniformExpression interface.
-	virtual void WriteNumberOpcodes(FMaterialPreshaderData& OutData) const override
+	virtual void WriteNumberOpcodes(UE::Shader::FPreshaderData& OutData) const override
 	{
 		A->WriteNumberOpcodes(OutData);
 		B->WriteNumberOpcodes(OutData);
 
 		switch (Op)
 		{
-		case FMO_Add: OutData.WriteOpcode(EMaterialPreshaderOpcode::Add); break;
-		case FMO_Sub: OutData.WriteOpcode(EMaterialPreshaderOpcode::Sub); break;
-		case FMO_Mul: OutData.WriteOpcode(EMaterialPreshaderOpcode::Mul); break;
-		case FMO_Div: OutData.WriteOpcode(EMaterialPreshaderOpcode::Div); break;
-		case FMO_Dot: OutData.WriteOpcode(EMaterialPreshaderOpcode::Dot).Write((uint8)ValueType); break;
-		case FMO_Cross: OutData.WriteOpcode(EMaterialPreshaderOpcode::Cross).Write((uint8)ValueType); break;
+		case FMO_Add: OutData.WriteOpcode(UE::Shader::EPreshaderOpcode::Add); break;
+		case FMO_Sub: OutData.WriteOpcode(UE::Shader::EPreshaderOpcode::Sub); break;
+		case FMO_Mul: OutData.WriteOpcode(UE::Shader::EPreshaderOpcode::Mul); break;
+		case FMO_Div: OutData.WriteOpcode(UE::Shader::EPreshaderOpcode::Div); break;
+		case FMO_Dot: OutData.WriteOpcode(UE::Shader::EPreshaderOpcode::Dot); break;
+		case FMO_Cross: OutData.WriteOpcode(UE::Shader::EPreshaderOpcode::Cross); break;
 		default: checkNoEntry(); break;
 		}
 	}
@@ -776,10 +786,10 @@ public:
 	{}
 
 	// FMaterialUniformExpression interface.
-	virtual void WriteNumberOpcodes(FMaterialPreshaderData& OutData) const override
+	virtual void WriteNumberOpcodes(UE::Shader::FPreshaderData& OutData) const override
 	{
 		X->WriteNumberOpcodes(OutData);
-		OutData.WriteOpcode(EMaterialPreshaderOpcode::Fractional);
+		OutData.WriteOpcode(UE::Shader::EPreshaderOpcode::Fractional);
 	}
 	virtual bool IsConstant() const
 	{
@@ -814,11 +824,11 @@ public:
 	{}
 
 	// FMaterialUniformExpression interface.
-	virtual void WriteNumberOpcodes(FMaterialPreshaderData& OutData) const override
+	virtual void WriteNumberOpcodes(UE::Shader::FPreshaderData& OutData) const override
 	{
 		A->WriteNumberOpcodes(OutData);
 		B->WriteNumberOpcodes(OutData);
-		OutData.WriteOpcode(EMaterialPreshaderOpcode::AppendVector).Write((uint8)NumComponentsA);
+		OutData.WriteOpcode(UE::Shader::EPreshaderOpcode::AppendVector);// .Write((uint8)NumComponentsA);
 	}
 	virtual bool IsConstant() const
 	{
@@ -854,11 +864,11 @@ public:
 	{}
 
 	// FMaterialUniformExpression interface.
-	virtual void WriteNumberOpcodes(FMaterialPreshaderData& OutData) const override
+	virtual void WriteNumberOpcodes(UE::Shader::FPreshaderData& OutData) const override
 	{
 		A->WriteNumberOpcodes(OutData);
 		B->WriteNumberOpcodes(OutData);
-		OutData.WriteOpcode(EMaterialPreshaderOpcode::Min);
+		OutData.WriteOpcode(UE::Shader::EPreshaderOpcode::Min);
 	}
 	virtual bool IsConstant() const
 	{
@@ -893,11 +903,11 @@ public:
 	{}
 
 	// FMaterialUniformExpression interface.
-	virtual void WriteNumberOpcodes(FMaterialPreshaderData& OutData) const override
+	virtual void WriteNumberOpcodes(UE::Shader::FPreshaderData& OutData) const override
 	{
 		A->WriteNumberOpcodes(OutData);
 		B->WriteNumberOpcodes(OutData);
-		OutData.WriteOpcode(EMaterialPreshaderOpcode::Max);
+		OutData.WriteOpcode(UE::Shader::EPreshaderOpcode::Max);
 	}
 	virtual bool IsConstant() const
 	{
@@ -933,12 +943,12 @@ public:
 	{}
 
 	// FMaterialUniformExpression interface.
-	virtual void WriteNumberOpcodes(FMaterialPreshaderData& OutData) const override
+	virtual void WriteNumberOpcodes(UE::Shader::FPreshaderData& OutData) const override
 	{
 		Input->WriteNumberOpcodes(OutData);
 		Min->WriteNumberOpcodes(OutData);
 		Max->WriteNumberOpcodes(OutData);
-		OutData.WriteOpcode(EMaterialPreshaderOpcode::Clamp);
+		OutData.WriteOpcode(UE::Shader::EPreshaderOpcode::Clamp);
 	}
 	virtual bool IsConstant() const
 	{
@@ -973,10 +983,10 @@ public:
 	{}
 
 	// FMaterialUniformExpression interface.
-	virtual void WriteNumberOpcodes(FMaterialPreshaderData& OutData) const override
+	virtual void WriteNumberOpcodes(UE::Shader::FPreshaderData& OutData) const override
 	{
 		Input->WriteNumberOpcodes(OutData);
-		OutData.WriteOpcode(EMaterialPreshaderOpcode::Saturate);
+		OutData.WriteOpcode(UE::Shader::EPreshaderOpcode::Saturate);
 	}
 	virtual bool IsConstant() const
 	{
@@ -1036,10 +1046,10 @@ public:
 	}
 
 	// FMaterialUniformExpression interface.
-	virtual void WriteNumberOpcodes(FMaterialPreshaderData& OutData) const override
+	virtual void WriteNumberOpcodes(UE::Shader::FPreshaderData& OutData) const override
 	{
 		X->WriteNumberOpcodes(OutData);
-		OutData.WriteOpcode(EMaterialPreshaderOpcode::ComponentSwizzle).Write((uint8)NumElements).Write((uint8)IndexR).Write((uint8)IndexG).Write((uint8)IndexB).Write((uint8)IndexA);
+		OutData.WriteOpcode(UE::Shader::EPreshaderOpcode::ComponentSwizzle).Write((uint8)NumElements).Write((uint8)IndexR).Write((uint8)IndexG).Write((uint8)IndexB).Write((uint8)IndexA);
 	}
 	virtual bool IsConstant() const
 	{
@@ -1082,10 +1092,10 @@ public:
 	{}
 
 	// FMaterialUniformExpression interface.
-	virtual void WriteNumberOpcodes(FMaterialPreshaderData& OutData) const override
+	virtual void WriteNumberOpcodes(UE::Shader::FPreshaderData& OutData) const override
 	{
 		X->WriteNumberOpcodes(OutData);
-		OutData.WriteOpcode(EMaterialPreshaderOpcode::Floor);
+		OutData.WriteOpcode(UE::Shader::EPreshaderOpcode::Floor);
 	}
 	virtual bool IsConstant() const
 	{
@@ -1118,10 +1128,10 @@ public:
 	{}
 
 	// FMaterialUniformExpression interface.
-	virtual void WriteNumberOpcodes(FMaterialPreshaderData& OutData) const override
+	virtual void WriteNumberOpcodes(UE::Shader::FPreshaderData& OutData) const override
 	{
 		X->WriteNumberOpcodes(OutData);
-		OutData.WriteOpcode(EMaterialPreshaderOpcode::Ceil);
+		OutData.WriteOpcode(UE::Shader::EPreshaderOpcode::Ceil);
 	}
 	virtual bool IsConstant() const
 	{
@@ -1154,10 +1164,10 @@ public:
 	{}
 
 	// FMaterialUniformExpression interface.
-	virtual void WriteNumberOpcodes(FMaterialPreshaderData& OutData) const override
+	virtual void WriteNumberOpcodes(UE::Shader::FPreshaderData& OutData) const override
 	{
 		X->WriteNumberOpcodes(OutData);
-		OutData.WriteOpcode(EMaterialPreshaderOpcode::Round);
+		OutData.WriteOpcode(UE::Shader::EPreshaderOpcode::Round);
 	}
 	virtual bool IsConstant() const
 	{
@@ -1190,10 +1200,10 @@ public:
 	{}
 
 	// FMaterialUniformExpression interface.
-	virtual void WriteNumberOpcodes(FMaterialPreshaderData& OutData) const override
+	virtual void WriteNumberOpcodes(UE::Shader::FPreshaderData& OutData) const override
 	{
 		X->WriteNumberOpcodes(OutData);
-		OutData.WriteOpcode(EMaterialPreshaderOpcode::Trunc);
+		OutData.WriteOpcode(UE::Shader::EPreshaderOpcode::Trunc);
 	}
 	virtual bool IsConstant() const
 	{
@@ -1226,10 +1236,10 @@ public:
 	{}
 
 	// FMaterialUniformExpression interface.
-	virtual void WriteNumberOpcodes(FMaterialPreshaderData& OutData) const override
+	virtual void WriteNumberOpcodes(UE::Shader::FPreshaderData& OutData) const override
 	{
 		X->WriteNumberOpcodes(OutData);
-		OutData.WriteOpcode(EMaterialPreshaderOpcode::Sign);
+		OutData.WriteOpcode(UE::Shader::EPreshaderOpcode::Sign);
 	}
 	virtual bool IsConstant() const
 	{
@@ -1262,10 +1272,10 @@ public:
 	{}
 
 	// FMaterialUniformExpression interface.
-	virtual void WriteNumberOpcodes(FMaterialPreshaderData& OutData) const override
+	virtual void WriteNumberOpcodes(UE::Shader::FPreshaderData& OutData) const override
 	{
 		X->WriteNumberOpcodes(OutData);
-		OutData.WriteOpcode(EMaterialPreshaderOpcode::Frac);
+		OutData.WriteOpcode(UE::Shader::EPreshaderOpcode::Frac);
 	}
 	virtual bool IsConstant() const
 	{
@@ -1299,11 +1309,11 @@ public:
 	{}
 
 	// FMaterialUniformExpression interface.
-	virtual void WriteNumberOpcodes(FMaterialPreshaderData& OutData) const override
+	virtual void WriteNumberOpcodes(UE::Shader::FPreshaderData& OutData) const override
 	{
 		A->WriteNumberOpcodes(OutData);
 		B->WriteNumberOpcodes(OutData);
-		OutData.WriteOpcode(EMaterialPreshaderOpcode::Fmod);
+		OutData.WriteOpcode(UE::Shader::EPreshaderOpcode::Fmod);
 	}
 	virtual bool IsConstant() const
 	{
@@ -1338,10 +1348,10 @@ public:
 	{}
 
 	// FMaterialUniformExpression interface.
-	virtual void WriteNumberOpcodes(FMaterialPreshaderData& OutData) const override
+	virtual void WriteNumberOpcodes(UE::Shader::FPreshaderData& OutData) const override
 	{
 		X->WriteNumberOpcodes(OutData);
-		OutData.WriteOpcode(EMaterialPreshaderOpcode::Abs);
+		OutData.WriteOpcode(UE::Shader::EPreshaderOpcode::Abs);
 	}
 	virtual bool IsConstant() const
 	{
@@ -1375,16 +1385,16 @@ public:
 	{}
 
 	// FMaterialUniformExpression interface.
-	virtual void WriteNumberOpcodes(FMaterialPreshaderData& OutData) const override
+	virtual void WriteNumberOpcodes(UE::Shader::FPreshaderData& OutData) const override
 	{
 		FMaterialTextureParameterInfo TextureParameter;
 		TextureExpression->GetTextureParameterInfo(TextureParameter);
 
-		EMaterialPreshaderOpcode Op = EMaterialPreshaderOpcode::Nop;
+		UE::Shader::EPreshaderOpcode Op = UE::Shader::EPreshaderOpcode::Nop;
 		switch (TextureProperty)
 		{
-		case TMTM_TextureSize: Op = EMaterialPreshaderOpcode::TextureSize; break;
-		case TMTM_TexelSize: Op = EMaterialPreshaderOpcode::TexelSize; break;
+		case TMTM_TextureSize: Op = UE::Shader::EPreshaderOpcode::TextureSize; break;
+		case TMTM_TexelSize: Op = UE::Shader::EPreshaderOpcode::TexelSize; break;
 		default: checkNoEntry(); break;
 		}
 		OutData.WriteOpcode(Op).Write(TextureParameter.ParameterInfo).Write((int32)TextureParameter.TextureIndex);
@@ -1425,7 +1435,7 @@ public:
 	FMaterialUniformExpressionExternalTextureCoordinateScaleRotation(int32 InSourceTextureIndex, TOptional<FName> InParameterName) : FMaterialUniformExpressionExternalTextureBase(InSourceTextureIndex), ParameterName(InParameterName) {}
 
 	virtual bool IsIdentical(const FMaterialUniformExpression* OtherExpression) const override;
-	virtual void WriteNumberOpcodes(FMaterialPreshaderData& OutData) const override;
+	virtual void WriteNumberOpcodes(UE::Shader::FPreshaderData& OutData) const override;
 
 protected:
 	typedef FMaterialUniformExpressionExternalTextureBase Super;
@@ -1447,7 +1457,7 @@ public:
 	FMaterialUniformExpressionExternalTextureCoordinateOffset(int32 InSourceTextureIndex, TOptional<FName> InParameterName) : FMaterialUniformExpressionExternalTextureBase(InSourceTextureIndex), ParameterName(InParameterName) {}
 
 	virtual bool IsIdentical(const FMaterialUniformExpression* OtherExpression) const override;
-	virtual void WriteNumberOpcodes(FMaterialPreshaderData& OutData) const override;
+	virtual void WriteNumberOpcodes(UE::Shader::FPreshaderData& OutData) const override;
 
 protected:
 	typedef FMaterialUniformExpressionExternalTextureBase Super;
@@ -1473,7 +1483,7 @@ public:
 	//~ Begin FMaterialUniformExpression Interface.
 	virtual bool IsConstant() const override { return false; }
 	virtual bool IsIdentical(const FMaterialUniformExpression* OtherExpression) const override;
-	virtual void WriteNumberOpcodes(FMaterialPreshaderData& OutData) const override;
+	virtual void WriteNumberOpcodes(UE::Shader::FPreshaderData& OutData) const override;
 	//~ End FMaterialUniformExpression Interface.
 
 protected:
diff --git a/Engine/Source/Runtime/Engine/Private/Shader/Preshader.cpp b/Engine/Source/Runtime/Engine/Private/Shader/Preshader.cpp
new file mode 100644
--- /dev/null
+++ b/Engine/Source/Runtime/Engine/Private/Shader/Preshader.cpp
@@ -0,0 +1,409 @@
+// Copyright Epic Games, Inc. All Rights Reserved.
+#include "Shader/Preshader.h"
+#include "Shader/PreshaderEvaluate.h"
+#include "Materials/Material.h"
+#include "Materials/MaterialUniformExpressions.h"
+#include "Engine/Texture.h"
+#include "VT/RuntimeVirtualTexture.h"
+#include "ExternalTexture.h"
+
+IMPLEMENT_TYPE_LAYOUT(UE::Shader::FPreshaderData);
+
+void UE::Shader::FPreshaderData::WriteData(const void* Value, uint32 Size)
+{
+	Data.Append((uint8*)Value, Size);
+}
+
+void UE::Shader::FPreshaderData::WriteName(const FScriptName& Name)
+{
+	int32 Index = Names.Find(Name);
+	if (Index == INDEX_NONE)
+	{
+		Index = Names.Add(Name);
+	}
+	check(Index >= 0 && Index <= 0xffff);
+	const uint32 Offset = Data.Num();
+	NameOffsets.Add(Offset);
+	Write((uint16)Index);
+}
+
+void UE::Shader::FPreshaderData::Append(const FPreshaderData& InPreshader)
+{
+	const uint32 BaseOffset = Data.Num();
+	Data.Append(InPreshader.Data);
+
+	TArray<uint16> NameIndexRemap;
+	NameIndexRemap.Empty(InPreshader.Names.Num());
+	for (const FScriptName& Name : InPreshader.Names)
+	{
+		const int32 Index = Names.AddUnique(Name);
+		check(Index >= 0 && Index <= 0xffff);
+		NameIndexRemap.Add(Index);
+	}
+
+	for (uint32 PrevOffset : InPreshader.NameOffsets)
+	{
+		const uint32 Offset = BaseOffset + PrevOffset;
+		uint16* PrevNameIndex = (uint16*)(Data.GetData() + Offset);
+		const uint16 RemapIndex = NameIndexRemap[*PrevNameIndex];
+		*PrevNameIndex = RemapIndex;
+		NameOffsets.Add(Offset);
+	}
+}
+
+UE::Shader::FPreshaderDataContext::FPreshaderDataContext(const FPreshaderData& InData)
+	: Ptr(InData.Data.GetData())
+	, EndPtr(Ptr + InData.Data.Num())
+	, Names(InData.Names.GetData())
+	, NumNames(InData.Names.Num())
+{}
+
+UE::Shader::FPreshaderDataContext::FPreshaderDataContext(const FPreshaderDataContext& InContext, uint32 InOffset, uint32 InSize)
+	: Ptr(InContext.Ptr + InOffset)
+	, EndPtr(Ptr + InSize)
+	, Names(InContext.Names)
+	, NumNames(InContext.NumNames)
+{}
+
+template<typename T>
+static inline T ReadPreshaderValue(UE::Shader::FPreshaderDataContext& RESTRICT Data)
+{
+	T Result;
+	FMemory::Memcpy(&Result, Data.Ptr, sizeof(T));
+	Data.Ptr += sizeof(T);
+	checkSlow(Data.Ptr <= Data.EndPtr);
+	return Result;
+}
+
+template<>
+inline uint8 ReadPreshaderValue<uint8>(UE::Shader::FPreshaderDataContext& RESTRICT Data)
+{
+	checkSlow(Data.Ptr < Data.EndPtr);
+	return *Data.Ptr++;
+}
+
+template<>
+FScriptName ReadPreshaderValue<FScriptName>(UE::Shader::FPreshaderDataContext& RESTRICT Data)
+{
+	const int32 Index = ReadPreshaderValue<uint16>(Data);
+	check(Index >= 0 && Index < Data.NumNames);
+	return Data.Names[Index];
+}
+
+template<>
+FName ReadPreshaderValue<FName>(UE::Shader::FPreshaderDataContext& RESTRICT Data) = delete;
+
+template<>
+FHashedMaterialParameterInfo ReadPreshaderValue<FHashedMaterialParameterInfo>(UE::Shader::FPreshaderDataContext& RESTRICT Data)
+{
+	const FScriptName Name = ReadPreshaderValue<FScriptName>(Data);
+	const int32 Index = ReadPreshaderValue<int32>(Data);
+	const TEnumAsByte<EMaterialParameterAssociation> Association = ReadPreshaderValue<TEnumAsByte<EMaterialParameterAssociation>>(Data);
+	return FHashedMaterialParameterInfo(Name, Association, Index);
+}
+
+static void GetVectorParameter(const FUniformExpressionSet& UniformExpressionSet, uint32 ParameterIndex, const FMaterialRenderContext& Context, UE::Shader::FValue& OutValue)
+{
+	const FMaterialVectorParameterInfo& Parameter = UniformExpressionSet.GetVectorParameter(ParameterIndex);
+
+	FLinearColor ParameterValue(ForceInitToZero);
+	bool bNeedsDefaultValue = false;
+	if (!Context.MaterialRenderProxy || !Context.MaterialRenderProxy->GetVectorValue(Parameter.ParameterInfo, &ParameterValue, Context))
+	{
+		const bool bOveriddenParameterOnly = Parameter.ParameterInfo.Association == EMaterialParameterAssociation::GlobalParameter;
+
+		if (Context.Material.HasMaterialLayers())
+		{
+			UMaterialInterface* Interface = Context.Material.GetMaterialInterface();
+			if (!Interface || !Interface->GetVectorParameterDefaultValue(Parameter.ParameterInfo, ParameterValue, bOveriddenParameterOnly))
+			{
+				bNeedsDefaultValue = true;
+			}
+		}
+		else
+		{
+			bNeedsDefaultValue = true;
+		}
+	}
+
+	if (bNeedsDefaultValue)
+	{
+#if WITH_EDITOR
+		if (!Context.Material.TransientOverrides.GetVectorOverride(Parameter.ParameterInfo, ParameterValue))
+#endif // WITH_EDITOR
+		{
+			Parameter.GetDefaultValue(ParameterValue);
+		}
+	}
+
+	OutValue = ParameterValue;
+}
+
+static void GetScalarParameter(const FUniformExpressionSet& UniformExpressionSet, uint32 ParameterIndex, const FMaterialRenderContext& Context, UE::Shader::FValue& OutValue)
+{
+	const FMaterialScalarParameterInfo& Parameter = UniformExpressionSet.GetScalarParameter(ParameterIndex);
+
+	float ParameterValue = 0.0f;
+	bool bNeedsDefaultValue = false;
+	if (!Context.MaterialRenderProxy || !Context.MaterialRenderProxy->GetScalarValue(Parameter.ParameterInfo, &ParameterValue, Context))
+	{
+		const bool bOveriddenParameterOnly = Parameter.ParameterInfo.Association == EMaterialParameterAssociation::GlobalParameter;
+
+		if (Context.Material.HasMaterialLayers())
+		{
+			UMaterialInterface* Interface = Context.Material.GetMaterialInterface();
+			if (!Interface || !Interface->GetScalarParameterDefaultValue(Parameter.ParameterInfo, ParameterValue, bOveriddenParameterOnly))
+			{
+				bNeedsDefaultValue = true;
+			}
+		}
+		else
+		{
+			bNeedsDefaultValue = true;
+		}
+	}
+
+	if (bNeedsDefaultValue)
+	{
+#if WITH_EDITOR
+		if (!Context.Material.TransientOverrides.GetScalarOverride(Parameter.ParameterInfo, ParameterValue))
+#endif // WITH_EDITOR
+		{
+			Parameter.GetDefaultValue(ParameterValue);
+		}
+	}
+
+	OutValue = ParameterValue;
+}
+
+template<typename Operation>
+static inline void EvaluateUnaryOp(UE::Shader::FPreshaderStack& Stack, const Operation& Op)
+{
+	const UE::Shader::FValue Value = Stack.Pop(false);
+	Stack.Add(Op(Value));
+}
+
+template<typename Operation>
+static inline void EvaluateBinaryOp(UE::Shader::FPreshaderStack& Stack, const Operation& Op)
+{
+	const UE::Shader::FValue Value1 = Stack.Pop(false);
+	const UE::Shader::FValue Value0 = Stack.Pop(false);
+	Stack.Add(Op(Value0, Value1));
+}
+
+template<typename Operation>
+static inline void EvaluateTernaryOp(UE::Shader::FPreshaderStack& Stack, const Operation& Op)
+{
+	const UE::Shader::FValue Value2 = Stack.Pop(false);
+	const UE::Shader::FValue Value1 = Stack.Pop(false);
+	const UE::Shader::FValue Value0 = Stack.Pop(false);
+	Stack.Add(Op(Value0, Value1, Value2));
+}
+
+static void EvaluateComponentSwizzle(UE::Shader::FPreshaderStack& Stack, UE::Shader::FPreshaderDataContext& RESTRICT Data)
+{
+	const uint8 NumElements = ReadPreshaderValue<uint8>(Data);
+	const uint8 IndexR = ReadPreshaderValue<uint8>(Data);
+	const uint8 IndexG = ReadPreshaderValue<uint8>(Data);
+	const uint8 IndexB = ReadPreshaderValue<uint8>(Data);
+	const uint8 IndexA = ReadPreshaderValue<uint8>(Data);
+
+	UE::Shader::FValue Value = Stack.Pop(false);
+	UE::Shader::FValue Result;
+	Result.ComponentType = Value.ComponentType;
+	Result.NumComponents = NumElements;
+
+	switch (NumElements)
+	{
+	case 4:
+		Result.Component[3] = Value.GetComponent(IndexA);
+		// Fallthrough...
+	case 3:
+		Result.Component[2] = Value.GetComponent(IndexB);
+		// Fallthrough...
+	case 2:
+		Result.Component[1] = Value.GetComponent(IndexG);
+		// Fallthrough...
+	case 1:
+		Result.Component[0] = Value.GetComponent(IndexR);
+		break;
+	default:
+		UE_LOG(LogMaterial, Fatal, TEXT("Invalid number of swizzle elements: %d"), NumElements);
+		break;
+	}
+	Stack.Add(Result);
+}
+
+static const UTexture* GetTextureParameter(const FMaterialRenderContext& Context, UE::Shader::FPreshaderDataContext& RESTRICT Data)
+{
+	const FHashedMaterialParameterInfo ParameterInfo = ReadPreshaderValue<FHashedMaterialParameterInfo>(Data);
+	const int32 TextureIndex = ReadPreshaderValue<int32>(Data);
+
+	const UTexture* Texture = nullptr;
+	Context.GetTextureParameterValue(ParameterInfo, TextureIndex, Texture);
+	return Texture;
+}
+
+static void EvaluateTextureSize(const FMaterialRenderContext& Context, UE::Shader::FPreshaderStack& Stack, UE::Shader::FPreshaderDataContext& RESTRICT Data)
+{
+	const UTexture* Texture = GetTextureParameter(Context, Data);
+	if (Texture && Texture->Resource)
+	{
+		const uint32 SizeX = Texture->Resource->GetSizeX();
+		const uint32 SizeY = Texture->Resource->GetSizeY();
+		const uint32 SizeZ = Texture->Resource->GetSizeZ();
+		Stack.Add(UE::Shader::FValue((float)SizeX, (float)SizeY, (float)SizeZ));
+	}
+	else
+	{
+		Stack.Add(UE::Shader::FValue(0.0f, 0.0f, 0.0f));
+	}
+}
+
+static void EvaluateTexelSize(const FMaterialRenderContext& Context, UE::Shader::FPreshaderStack& Stack, UE::Shader::FPreshaderDataContext& RESTRICT Data)
+{
+	const UTexture* Texture = GetTextureParameter(Context, Data);
+	if (Texture && Texture->Resource)
+	{
+		const uint32 SizeX = Texture->Resource->GetSizeX();
+		const uint32 SizeY = Texture->Resource->GetSizeY();
+		const uint32 SizeZ = Texture->Resource->GetSizeZ();
+		Stack.Add(UE::Shader::FValue(1.0f / (float)SizeX, 1.0f / (float)SizeY, (SizeZ > 0 ? 1.0f / (float)SizeZ : 0.0f)));
+	}
+	else
+	{
+		Stack.Add(UE::Shader::FValue(0.0f, 0.0f, 0.0f));
+	}
+}
+
+static FGuid GetExternalTextureGuid(const FMaterialRenderContext& Context, UE::Shader::FPreshaderDataContext& RESTRICT Data)
+{
+	const FScriptName ParameterName = ReadPreshaderValue<FScriptName>(Data);
+	const FGuid ExternalTextureGuid = ReadPreshaderValue<FGuid>(Data);
+	const int32 TextureIndex = ReadPreshaderValue<int32>(Data);
+	return Context.GetExternalTextureGuid(ExternalTextureGuid, ScriptNameToName(ParameterName), TextureIndex);
+}
+
+static void EvaluateExternalTextureCoordinateScaleRotation(const FMaterialRenderContext& Context, UE::Shader::FPreshaderStack& Stack, UE::Shader::FPreshaderDataContext& RESTRICT Data)
+{
+	const FGuid GuidToLookup = GetExternalTextureGuid(Context, Data);
+	FLinearColor Result(1.f, 0.f, 0.f, 1.f);
+	if (GuidToLookup.IsValid())
+	{
+		FExternalTextureRegistry::Get().GetExternalTextureCoordinateScaleRotation(GuidToLookup, Result);
+	}
+	Stack.Add(Result);
+}
+
+static void EvaluateExternalTextureCoordinateOffset(const FMaterialRenderContext& Context, UE::Shader::FPreshaderStack& Stack, UE::Shader::FPreshaderDataContext& RESTRICT Data)
+{
+	const FGuid GuidToLookup = GetExternalTextureGuid(Context, Data);
+	FLinearColor Result(0.f, 0.f, 0.f, 0.f);
+	if (GuidToLookup.IsValid())
+	{
+		FExternalTextureRegistry::Get().GetExternalTextureCoordinateOffset(GuidToLookup, Result);
+	}
+	Stack.Add(Result);
+}
+
+static void EvaluateRuntimeVirtualTextureUniform(const FMaterialRenderContext& Context, UE::Shader::FPreshaderStack& Stack, UE::Shader::FPreshaderDataContext& RESTRICT Data)
+{
+	const FHashedMaterialParameterInfo ParameterInfo = ReadPreshaderValue<FHashedMaterialParameterInfo>(Data);
+	const int32 TextureIndex = ReadPreshaderValue<int32>(Data);
+	const int32 VectorIndex = ReadPreshaderValue<int32>(Data);
+
+	const URuntimeVirtualTexture* Texture = nullptr;
+	if (ParameterInfo.Name.IsNone() || !Context.MaterialRenderProxy || !Context.MaterialRenderProxy->GetTextureValue(ParameterInfo, &Texture, Context))
+	{
+		Texture = GetIndexedTexture<URuntimeVirtualTexture>(Context.Material, TextureIndex);
+	}
+	if (Texture != nullptr && VectorIndex != INDEX_NONE)
+	{
+		Stack.Add(UE::Shader::FValue(Texture->GetUniformParameter(VectorIndex)));
+	}
+	else
+	{
+		Stack.Add(UE::Shader::FValue(0.f, 0.f, 0.f, 0.f));
+	}
+}
+
+void UE::Shader::EvaluatePreshader(const FUniformExpressionSet* UniformExpressionSet, const FMaterialRenderContext& Context, FPreshaderStack& Stack, FPreshaderDataContext& RESTRICT Data, FValue& OutValue)
+{
+	uint8 const* const DataEnd = Data.EndPtr;
+
+	Stack.Reset();
+	while (Data.Ptr < DataEnd)
+	{
+		const EPreshaderOpcode Opcode = (EPreshaderOpcode)ReadPreshaderValue<uint8>(Data);
+		switch (Opcode)
+		{
+		case EPreshaderOpcode::ConstantZero:
+			Stack.Add(UE::Shader::FValue(0.0f));
+			break;
+		case EPreshaderOpcode::Constant:
+			Stack.Add(ReadPreshaderValue<UE::Shader::FValue>(Data));
+			break;
+		case EPreshaderOpcode::VectorParameter:
+			check(UniformExpressionSet);
+			GetVectorParameter(*UniformExpressionSet, ReadPreshaderValue<uint16>(Data), Context, Stack.AddDefaulted_GetRef());
+			break;
+		case EPreshaderOpcode::ScalarParameter:
+			check(UniformExpressionSet);
+			GetScalarParameter(*UniformExpressionSet, ReadPreshaderValue<uint16>(Data), Context, Stack.AddDefaulted_GetRef());
+			break;
+		case EPreshaderOpcode::Add: EvaluateBinaryOp(Stack, UE::Shader::Add); break;
+		case EPreshaderOpcode::Sub: EvaluateBinaryOp(Stack, UE::Shader::Sub); break;
+		case EPreshaderOpcode::Mul: EvaluateBinaryOp(Stack, UE::Shader::Mul); break;
+		case EPreshaderOpcode::Div: EvaluateBinaryOp(Stack, UE::Shader::Div); break;
+		case EPreshaderOpcode::Fmod: EvaluateBinaryOp(Stack, UE::Shader::Fmod); break;
+		case EPreshaderOpcode::Min: EvaluateBinaryOp(Stack, UE::Shader::Min); break;
+		case EPreshaderOpcode::Max: EvaluateBinaryOp(Stack, UE::Shader::Max); break;
+		case EPreshaderOpcode::Clamp: EvaluateTernaryOp(Stack, UE::Shader::Clamp); break;
+		case EPreshaderOpcode::Dot: EvaluateBinaryOp(Stack, UE::Shader::Dot); break;
+		case EPreshaderOpcode::Cross: EvaluateBinaryOp(Stack, UE::Shader::Dot); break;
+		case EPreshaderOpcode::Sqrt: EvaluateUnaryOp(Stack, UE::Shader::Sqrt); break;
+		case EPreshaderOpcode::Sin: EvaluateUnaryOp(Stack, UE::Shader::Sin); break;
+		case EPreshaderOpcode::Cos: EvaluateUnaryOp(Stack, UE::Shader::Cos); break;
+		case EPreshaderOpcode::Tan: EvaluateUnaryOp(Stack, UE::Shader::Tan); break;
+		case EPreshaderOpcode::Asin: EvaluateUnaryOp(Stack, UE::Shader::Asin); break;;
+		case EPreshaderOpcode::Acos: EvaluateUnaryOp(Stack, UE::Shader::Acos); break;
+		case EPreshaderOpcode::Atan: EvaluateUnaryOp(Stack, UE::Shader::Atan); break;
+		case EPreshaderOpcode::Atan2: EvaluateBinaryOp(Stack, UE::Shader::Atan2); break;
+		case EPreshaderOpcode::Abs: EvaluateUnaryOp(Stack, UE::Shader::Abs); break;
+		case EPreshaderOpcode::Saturate: EvaluateUnaryOp(Stack, UE::Shader::Saturate); break;
+		case EPreshaderOpcode::Floor: EvaluateUnaryOp(Stack, UE::Shader::Floor); break;
+		case EPreshaderOpcode::Ceil: EvaluateUnaryOp(Stack, UE::Shader::Ceil); break;
+		case EPreshaderOpcode::Round: EvaluateUnaryOp(Stack, UE::Shader::Round); break;
+		case EPreshaderOpcode::Trunc: EvaluateUnaryOp(Stack, UE::Shader::Trunc); break;
+		case EPreshaderOpcode::Sign: EvaluateUnaryOp(Stack, UE::Shader::Sign); break;
+		case EPreshaderOpcode::Frac: EvaluateUnaryOp(Stack, UE::Shader::Frac); break;
+		case EPreshaderOpcode::Fractional: EvaluateUnaryOp(Stack, UE::Shader::Fractional); break;
+		case EPreshaderOpcode::Log2: EvaluateUnaryOp(Stack, UE::Shader::Log2); break;
+		case EPreshaderOpcode::Log10: EvaluateUnaryOp(Stack, UE::Shader::Log10); break;
+		case EPreshaderOpcode::ComponentSwizzle: EvaluateComponentSwizzle(Stack, Data); break;
+		case EPreshaderOpcode::AppendVector: EvaluateBinaryOp(Stack, UE::Shader::Append); break;
+		case EPreshaderOpcode::TextureSize: EvaluateTextureSize(Context, Stack, Data); break;
+		case EPreshaderOpcode::TexelSize: EvaluateTexelSize(Context, Stack, Data); break;
+		case EPreshaderOpcode::ExternalTextureCoordinateScaleRotation: EvaluateExternalTextureCoordinateScaleRotation(Context, Stack, Data); break;
+		case EPreshaderOpcode::ExternalTextureCoordinateOffset: EvaluateExternalTextureCoordinateOffset(Context, Stack, Data); break;
+		case EPreshaderOpcode::RuntimeVirtualTextureUniform: EvaluateRuntimeVirtualTextureUniform(Context, Stack, Data); break;
+		default:
+			UE_LOG(LogMaterial, Fatal, TEXT("Unknown preshader opcode %d"), (uint8)Opcode);
+			break;
+		}
+	}
+	check(Data.Ptr == DataEnd);
+
+	ensure(Stack.Num() <= 1);
+	if (Stack.Num() > 0)
+	{
+		OutValue = Stack.Last();
+	}
+}
+
+void UE::Shader::FPreshaderData::Evaluate(FUniformExpressionSet* UniformExpressionSet, const struct FMaterialRenderContext& Context, FValue& OutValue)
+{
+	FPreshaderStack Stack;
+	FPreshaderDataContext PreshaderContext(*this);
+	EvaluatePreshader(UniformExpressionSet, Context, Stack, PreshaderContext, OutValue);
+}
diff --git a/Engine/Source/Runtime/Engine/Private/Shader/PreshaderEvaluate.h b/Engine/Source/Runtime/Engine/Private/Shader/PreshaderEvaluate.h
new file mode 100644
--- /dev/null
+++ b/Engine/Source/Runtime/Engine/Private/Shader/PreshaderEvaluate.h
@@ -0,0 +1,33 @@
+// Copyright Epic Games, Inc. All Rights Reserved.
+#pragma once
+
+#include "CoreMinimal.h"
+#include "Shader/PreshaderTypes.h"
+
+class FUniformExpressionSet;
+struct FMaterialRenderContext;
+
+namespace UE
+{
+namespace Shader
+{
+
+class FPreshaderData;
+
+using FPreshaderStack = TArray<FValue, TInlineAllocator<64u>>;
+
+struct FPreshaderDataContext
+{
+	explicit FPreshaderDataContext(const FPreshaderData& InData);
+	FPreshaderDataContext(const FPreshaderDataContext& InContext, uint32 InOffset, uint32 InSize);
+
+	const uint8* RESTRICT Ptr;
+	const uint8* RESTRICT EndPtr;
+	const FScriptName* RESTRICT Names;
+	int32 NumNames;
+};
+
+ENGINE_API void EvaluatePreshader(const FUniformExpressionSet* UniformExpressionSet, const FMaterialRenderContext& Context, FPreshaderStack& Stack, FPreshaderDataContext& RESTRICT Data, FValue& OutValue);
+
+} // namespace Shader
+} // namespace UE
diff --git a/Engine/Source/Runtime/Engine/Private/Shader/ShaderValue.cpp b/Engine/Source/Runtime/Engine/Private/Shader/ShaderValue.cpp
new file mode 100644
--- /dev/null
+++ b/Engine/Source/Runtime/Engine/Private/Shader/ShaderValue.cpp
@@ -0,0 +1,632 @@
+#include "Shader/ShaderTypes.h"
+
+namespace UE
+{
+namespace Shader
+{
+namespace Private
+{
+
+struct FCastFloat
+{
+	using FComponentType = float;
+	inline float operator()(EValueComponentType Type, const FValueComponent& Component) const
+	{
+		switch (Type)
+		{
+		case EValueComponentType::Float: return Component.Float;
+		case EValueComponentType::Int: return (float)Component.Int;
+		case EValueComponentType::Bool: return (float)Component.Bool;
+		default: return 0.0f;
+		}
+	}
+};
+
+struct FCastInt
+{
+	using FComponentType = int32;
+	inline int32 operator()(EValueComponentType Type, const FValueComponent& Component) const
+	{
+		switch (Type)
+		{
+		case EValueComponentType::Float: return (int32)Component.Float;
+		case EValueComponentType::Int: return Component.Int;
+		case EValueComponentType::Bool: return Component.Bool ? 1 : 0;
+		default: return false;
+		}
+	}
+};
+
+struct FCastBool
+{
+	using FComponentType = bool;
+	inline bool operator()(EValueComponentType Type, const FValueComponent& Component) const
+	{
+		switch (Type)
+		{
+		case EValueComponentType::Float: return Component.Float != 0.0f;
+		case EValueComponentType::Int: return Component.Int != 0;
+		case EValueComponentType::Bool: return Component.Bool;
+		default: return false;
+		}
+	}
+};
+
+template<typename Operation, typename ResultType>
+void CastValue(const Operation& Op, const FValue& Value, ResultType& OutResult)
+{
+	using FComponentType = typename Operation::FComponentType;
+	const int32 NumComponents = Value.NumComponents;
+	const EValueComponentType ComponentType = Value.ComponentType;
+	if (NumComponents == 1)
+	{
+		const FComponentType Component = Op(ComponentType, Value.Component[0]);
+		for (int32 i = 0; i < 4; ++i)
+		{
+			OutResult[i] = Component;
+		}
+	}
+	else
+	{
+		for (int32 i = 0; i < NumComponents; ++i)
+		{
+			OutResult[i] = Op(ComponentType, Value.Component[i]);
+		}
+		for (int32 i = NumComponents; i < 4; ++i)
+		{
+			OutResult[i] = FComponentType();
+		}
+	}
+}
+
+} // namespace Private
+} // namespace Shader
+} // namespace UE
+
+UE::Shader::FFloatValue UE::Shader::FValue::AsFloat() const
+{
+	FFloatValue Result;
+	Private::CastValue(Private::FCastFloat(), *this, Result);
+	return Result;
+}
+
+FLinearColor UE::Shader::FValue::AsLinearColor() const
+{
+	const FFloatValue Result = AsFloat();
+	return FLinearColor(Result.Component[0], Result.Component[1], Result.Component[2], Result.Component[3]);
+}
+
+UE::Shader::FIntValue UE::Shader::FValue::AsInt() const
+{
+	FIntValue Result;
+	Private::CastValue(Private::FCastInt(), *this, Result);
+	return Result;
+}
+
+UE::Shader::FBoolValue UE::Shader::FValue::AsBool() const
+{
+	FBoolValue Result;
+	Private::CastValue(Private::FCastBool(), *this, Result);
+	return Result;
+}
+
+bool UE::Shader::FValue::AsBoolScalar() const
+{
+	const FBoolValue Result = AsBool();
+	for (int32 i = 0; i < NumComponents; ++i)
+	{
+		if (Result.Component[i])
+		{
+			return true;
+		}
+	}
+	return false;
+}
+
+UE::Shader::FValueTypeDescription UE::Shader::GetValueTypeDescription(EValueType Type)
+{
+	switch (Type)
+	{
+	case EValueType::Void: return FValueTypeDescription(TEXT("void"), EValueComponentType::Void, 0);
+	case EValueType::Float1: return FValueTypeDescription(TEXT("float"), EValueComponentType::Float, 1);
+	case EValueType::Float2: return FValueTypeDescription(TEXT("float2"), EValueComponentType::Float, 2);
+	case EValueType::Float3: return FValueTypeDescription(TEXT("float3"), EValueComponentType::Float, 3);
+	case EValueType::Float4: return FValueTypeDescription(TEXT("float4"), EValueComponentType::Float, 4);
+	case EValueType::Int1: return FValueTypeDescription(TEXT("int"), EValueComponentType::Int, 1);
+	case EValueType::Int2: return FValueTypeDescription(TEXT("int2"), EValueComponentType::Int, 2);
+	case EValueType::Int3: return FValueTypeDescription(TEXT("int3"), EValueComponentType::Int, 3);
+	case EValueType::Int4: return FValueTypeDescription(TEXT("int4"), EValueComponentType::Int, 4);
+	case EValueType::Bool1: return FValueTypeDescription(TEXT("bool"), EValueComponentType::Bool, 1);
+	case EValueType::Bool2: return FValueTypeDescription(TEXT("bool2"), EValueComponentType::Bool, 2);
+	case EValueType::Bool3: return FValueTypeDescription(TEXT("bool3"), EValueComponentType::Bool, 3);
+	case EValueType::Bool4: return FValueTypeDescription(TEXT("bool4"), EValueComponentType::Bool, 4);
+	case EValueType::MaterialAttributes: return FValueTypeDescription(TEXT("FMaterialAttributes"), EValueComponentType::MaterialAttributes, 0);
+	default: checkNoEntry(); return FValueTypeDescription();
+	}
+}
+
+UE::Shader::EValueType UE::Shader::MakeValueType(EValueComponentType ComponentType, int32 NumComponents)
+{
+	switch (ComponentType)
+	{
+	case EValueComponentType::Void:
+		check(NumComponents == 0);
+		return EValueType::Void;
+	case EValueComponentType::MaterialAttributes:
+		check(NumComponents == 0);
+		return EValueType::MaterialAttributes;
+	case EValueComponentType::Float:
+		switch (NumComponents)
+		{
+		case 1: return EValueType::Float1;
+		case 2: return EValueType::Float2;
+		case 3: return EValueType::Float3;
+		case 4: return EValueType::Float4;
+		default: break;
+		}
+	case EValueComponentType::Int:
+		switch (NumComponents)
+		{
+		case 1: return EValueType::Int1;
+		case 2: return EValueType::Int2;
+		case 3: return EValueType::Int3;
+		case 4: return EValueType::Int4;
+		default: break;
+		}
+	case EValueComponentType::Bool:
+		switch (NumComponents)
+		{
+		case 1: return EValueType::Bool1;
+		case 2: return EValueType::Bool2;
+		case 3: return EValueType::Bool3;
+		case 4: return EValueType::Bool4;
+		default: break;
+		}
+	default:
+		break;
+	}
+
+	checkNoEntry();
+	return EValueType::Void;
+}
+
+UE::Shader::EValueType UE::Shader::MakeValueType(EValueType BaseType, int32 NumComponents)
+{
+	return MakeValueType(GetValueTypeDescription(BaseType).ComponentType, NumComponents);
+}
+
+UE::Shader::EValueType UE::Shader::MakeArithmeticResultType(EValueType Lhs, EValueType Rhs, FString& OutErrorMessage)
+{
+	const FValueTypeDescription LhsDesc = GetValueTypeDescription(Lhs);
+	const FValueTypeDescription RhsDesc = GetValueTypeDescription(Rhs);
+	// Types with 0 components are non-arithmetic
+	if (LhsDesc.NumComponents > 0 && RhsDesc.NumComponents > 0)
+	{
+		if (Lhs == Rhs)
+		{
+			return Lhs;
+		}
+		if (LhsDesc.ComponentType == RhsDesc.ComponentType)
+		{
+			if (LhsDesc.NumComponents == 1 || RhsDesc.NumComponents == 1)
+			{
+				// single component type is valid to combine with other type
+				return MakeValueType(LhsDesc.ComponentType, FMath::Max(LhsDesc.NumComponents, RhsDesc.NumComponents));
+			}
+		}
+		OutErrorMessage = FString::Printf(TEXT("Arithmetic between types %s and %s are undefined"), LhsDesc.Name, RhsDesc.Name);
+	}
+	else
+	{
+		OutErrorMessage = FString::Printf(TEXT("Attempting to perform arithmetic on non-numeric types: %s %s"), LhsDesc.Name, RhsDesc.Name);
+	}
+
+	return EValueType::Void;
+}
+
+namespace UE
+{
+namespace Shader
+{
+namespace Private
+{
+
+/** Converts an arbitrary number into a safe divisor. i.e. FMath::Abs(Number) >= DELTA */
+template<typename T> T GetSafeDivisor(T Number)
+{
+	return Number != 0 ? Number : 1;
+}
+
+/**
+ * FORCENOINLINE is required to discourage compiler from vectorizing the Div operation, which may tempt it into optimizing divide as A * rcp(B)
+ * This will break shaders that are depending on exact divide results (see SubUV material function)
+ * Technically this could still happen for a scalar divide, but it doesn't seem to occur in practice
+ */
+template<>
+FORCENOINLINE float GetSafeDivisor<float>(float Number)
+{
+	if (FMath::Abs(Number) < DELTA)
+	{
+		if (Number < 0.0f)
+		{
+			return -DELTA;
+		}
+		else
+		{
+			return +DELTA;
+		}
+	}
+	else
+	{
+		return Number;
+	}
+}
+
+struct FOpBase
+{
+	static constexpr bool SupportsInt = true;
+};
+
+struct FOpBaseFloatOnly : public FOpBase
+{
+	static constexpr bool SupportsInt = false;
+};
+
+struct FOpAbs : public FOpBase { template<typename T> T operator()(T Value) const { return FMath::Abs(Value); } };
+struct FOpSign : public FOpBase { template<typename T> T operator()(T Value) const { return FMath::Sign(Value); } };
+struct FOpSaturate : public FOpBaseFloatOnly { float operator()(float Value) const { return FMath::Clamp(Value, 0.0f, 1.0f); } };
+struct FOpFloor : public FOpBaseFloatOnly { float operator()(float Value) const { return FMath::FloorToFloat(Value); } };
+struct FOpCeil : public FOpBaseFloatOnly { float operator()(float Value) const { return FMath::CeilToFloat(Value); } };
+struct FOpRound : public FOpBaseFloatOnly { float operator()(float Value) const { return FMath::RoundToFloat(Value); } };
+struct FOpTrunc : public FOpBaseFloatOnly { float operator()(float Value) const { return FMath::TruncToFloat(Value); } };
+struct FOpFrac : public FOpBaseFloatOnly { float operator()(float Value) const { return FMath::Frac(Value); } };
+struct FOpFractional : public FOpBaseFloatOnly { float operator()(float Value) const { return FMath::Fractional(Value); } };
+struct FOpSqrt : public FOpBaseFloatOnly { float operator()(float Value) const { return FMath::Sqrt(Value); } };
+struct FOpLog2 : public FOpBaseFloatOnly { float operator()(float Value) const { return FMath::Log2(Value); } };
+struct FOpLog10 : public FOpBaseFloatOnly
+{
+	float operator()(float Value) const
+	{
+		static const float LogToLog10 = 1.0f / FMath::Loge(10.f);
+		return FMath::Loge(Value) * LogToLog10;
+	}
+};
+struct FOpSin : public FOpBaseFloatOnly { float operator()(float Value) const { return FMath::Sin(Value); } };
+struct FOpCos : public FOpBaseFloatOnly { float operator()(float Value) const { return FMath::Cos(Value); } };
+struct FOpTan : public FOpBaseFloatOnly { float operator()(float Value) const { return FMath::Tan(Value); } };
+struct FOpAsin : public FOpBaseFloatOnly { float operator()(float Value) const { return FMath::Asin(Value); } };
+struct FOpAcos : public FOpBaseFloatOnly { float operator()(float Value) const { return FMath::Acos(Value); } };
+struct FOpAtan : public FOpBaseFloatOnly { float operator()(float Value) const { return FMath::Atan(Value); } };
+struct FOpAdd : public FOpBase { template<typename T> T operator()(T Lhs, T Rhs) const { return Lhs + Rhs; } };
+struct FOpSub : public FOpBase { template<typename T> T operator()(T Lhs, T Rhs) const { return Lhs - Rhs; } };
+struct FOpMul : public FOpBase { template<typename T> T operator()(T Lhs, T Rhs) const { return Lhs * Rhs; } };
+struct FOpDiv : public FOpBase { template<typename T> T operator()(T Lhs, T Rhs) const { return Lhs / GetSafeDivisor(Rhs); } };
+struct FOpMin : public FOpBase { template<typename T> T operator()(T Lhs, T Rhs) const { return FMath::Min(Lhs, Rhs); } };
+struct FOpMax : public FOpBase { template<typename T> T operator()(T Lhs, T Rhs) const { return FMath::Max(Lhs, Rhs); } };
+struct FOpFmod : public FOpBaseFloatOnly { float operator()(float Lhs, float Rhs) const { return FMath::Fmod(Lhs, Rhs); } };
+struct FOpAtan2 : public FOpBaseFloatOnly { float operator()(float Lhs, float Rhs) const { return FMath::Atan2(Lhs, Rhs); } };
+
+template<typename Operation>
+inline FValue UnaryOp(const Operation& Op, const FValue& Value)
+{
+	const int8 NumComponents = Value.NumComponents;
+	
+	FValue Result;
+	Result.NumComponents = NumComponents;
+
+	if (!Operation::SupportsInt || Value.ComponentType == UE::Shader::EValueComponentType::Float)
+	{
+		Result.ComponentType = UE::Shader::EValueComponentType::Float;
+		const FFloatValue Cast = Value.AsFloat();
+		for (int32 i = 0; i < NumComponents; ++i)
+		{
+			Result.Component[i].Float = Op(Cast.Component[i]);
+		}
+	}
+	else
+	{
+		Result.ComponentType = UE::Shader::EValueComponentType::Int;
+		const FIntValue Cast = Value.AsInt();
+		for (int32 i = 0; i < NumComponents; ++i)
+		{
+			Result.Component[i].Int = Op(Cast.Component[i]);
+		}
+	}
+
+	return Result;
+}
+
+template<typename Operation>
+inline FValue BinaryOp(const Operation& Op, const FValue& Lhs, const FValue& Rhs)
+{
+	int8 NumComponents;
+	if (Lhs.NumComponents == 1 || Rhs.NumComponents == 1)
+	{
+		// operations between scalar and non-scalar will splat the scalar value
+		NumComponents = FMath::Max(Lhs.NumComponents, Rhs.NumComponents);
+	}
+	else
+	{
+		// otherwise, operations should only be between types with same number of components
+		NumComponents = FMath::Min(Lhs.NumComponents, Rhs.NumComponents);
+	}
+
+	FValue Result;
+	Result.NumComponents = NumComponents;
+
+	if (!Operation::SupportsInt ||
+		Lhs.ComponentType == UE::Shader::EValueComponentType::Float ||
+		Rhs.ComponentType == UE::Shader::EValueComponentType::Float)
+	{
+		Result.ComponentType = UE::Shader::EValueComponentType::Float;
+		const FFloatValue LhsCast = Lhs.AsFloat();
+		const FFloatValue RhsCast = Rhs.AsFloat();
+		for (int32 i = 0; i < NumComponents; ++i)
+		{
+			Result.Component[i].Float = Op(LhsCast.Component[i], RhsCast.Component[i]);
+		}
+	}
+	else
+	{
+		Result.ComponentType = UE::Shader::EValueComponentType::Int;
+		const FIntValue LhsCast = Lhs.AsInt();
+		const FIntValue RhsCast = Rhs.AsInt();
+		for (int32 i = 0; i < NumComponents; ++i)
+		{
+			Result.Component[i].Int = Op(LhsCast.Component[i], RhsCast.Component[i]);
+		}
+	}
+
+	return Result;
+}
+
+} // namespace Private
+} // namespace Shader
+} // namespace UE
+
+UE::Shader::FValue UE::Shader::Abs(const FValue& Value)
+{
+	return Private::UnaryOp(Private::FOpAbs(), Value);
+}
+
+UE::Shader::FValue UE::Shader::Saturate(const FValue& Value)
+{
+	return Private::UnaryOp(Private::FOpSaturate(), Value);
+}
+
+UE::Shader::FValue UE::Shader::Floor(const FValue& Value)
+{
+	return Private::UnaryOp(Private::FOpFloor(), Value);
+}
+
+UE::Shader::FValue UE::Shader::Ceil(const FValue& Value)
+{
+	return Private::UnaryOp(Private::FOpCeil(), Value);
+}
+
+UE::Shader::FValue UE::Shader::Round(const FValue& Value)
+{
+	return Private::UnaryOp(Private::FOpRound(), Value);
+}
+
+UE::Shader::FValue UE::Shader::Trunc(const FValue& Value)
+{
+	return Private::UnaryOp(Private::FOpTrunc(), Value);
+}
+
+UE::Shader::FValue UE::Shader::Sign(const FValue& Value)
+{
+	return Private::UnaryOp(Private::FOpSign(), Value);
+}
+
+UE::Shader::FValue UE::Shader::Frac(const FValue& Value)
+{
+	return Private::UnaryOp(Private::FOpFrac(), Value);
+}
+
+UE::Shader::FValue UE::Shader::Fractional(const FValue& Value)
+{
+	return Private::UnaryOp(Private::FOpFractional(), Value);
+}
+
+UE::Shader::FValue UE::Shader::Sqrt(const FValue& Value)
+{
+	return Private::UnaryOp(Private::FOpSqrt(), Value);
+}
+
+UE::Shader::FValue UE::Shader::Log2(const FValue& Value)
+{
+	return Private::UnaryOp(Private::FOpLog2(), Value);
+}
+
+UE::Shader::FValue UE::Shader::Log10(const FValue& Value)
+{
+	return Private::UnaryOp(Private::FOpLog10(), Value);
+}
+
+UE::Shader::FValue UE::Shader::Sin(const FValue& Value)
+{
+	return Private::UnaryOp(Private::FOpSin(), Value);
+}
+
+UE::Shader::FValue UE::Shader::Cos(const FValue& Value)
+{
+	return Private::UnaryOp(Private::FOpCos(), Value);
+}
+
+UE::Shader::FValue UE::Shader::Tan(const FValue& Value)
+{
+	return Private::UnaryOp(Private::FOpTan(), Value);
+}
+
+UE::Shader::FValue UE::Shader::Asin(const FValue& Value)
+{
+	return Private::UnaryOp(Private::FOpAsin(), Value);
+}
+
+UE::Shader::FValue UE::Shader::Acos(const FValue& Value)
+{
+	return Private::UnaryOp(Private::FOpAcos(), Value);
+}
+
+UE::Shader::FValue UE::Shader::Atan(const FValue& Value)
+{
+	return Private::UnaryOp(Private::FOpAtan(), Value);
+}
+
+UE::Shader::FValue UE::Shader::Add(const UE::Shader::FValue& Lhs, const UE::Shader::FValue& Rhs)
+{
+	return Private::BinaryOp(Private::FOpAdd(), Lhs, Rhs);
+}
+
+UE::Shader::FValue UE::Shader::Sub(const UE::Shader::FValue& Lhs, const UE::Shader::FValue& Rhs)
+{
+	return Private::BinaryOp(Private::FOpSub(), Lhs, Rhs);
+}
+
+UE::Shader::FValue UE::Shader::Mul(const UE::Shader::FValue& Lhs, const UE::Shader::FValue& Rhs)
+{
+	return Private::BinaryOp(Private::FOpMul(), Lhs, Rhs);
+}
+
+UE::Shader::FValue UE::Shader::Div(const UE::Shader::FValue& Lhs, const UE::Shader::FValue& Rhs)
+{
+	return Private::BinaryOp(Private::FOpDiv(), Lhs, Rhs);
+}
+
+UE::Shader::FValue UE::Shader::Min(const UE::Shader::FValue& Lhs, const UE::Shader::FValue& Rhs)
+{
+	return Private::BinaryOp(Private::FOpMin(), Lhs, Rhs);
+}
+
+UE::Shader::FValue UE::Shader::Max(const UE::Shader::FValue& Lhs, const UE::Shader::FValue& Rhs)
+{
+	return Private::BinaryOp(Private::FOpMax(), Lhs, Rhs);
+}
+
+UE::Shader::FValue UE::Shader::Fmod(const UE::Shader::FValue& Lhs, const UE::Shader::FValue& Rhs)
+{
+	return Private::BinaryOp(Private::FOpFmod(), Lhs, Rhs);
+}
+
+UE::Shader::FValue UE::Shader::Atan2(const FValue& Lhs, const FValue& Rhs)
+{
+	return Private::BinaryOp(Private::FOpAtan2(), Lhs, Rhs);
+}
+
+UE::Shader::FValue UE::Shader::Clamp(const FValue& Value, const FValue& Low, const FValue& High)
+{
+	return Min(Max(Value, Low), High);
+}
+
+UE::Shader::FValue UE::Shader::Dot(const FValue& Lhs, const FValue& Rhs)
+{
+	const int32 NumComponents = FMath::Min(Lhs.NumComponents, Rhs.NumComponents);
+
+	FValue Result;
+	Result.NumComponents = 1;
+
+	if (Lhs.ComponentType == UE::Shader::EValueComponentType::Float || Rhs.ComponentType == UE::Shader::EValueComponentType::Float)
+	{
+		Result.ComponentType = UE::Shader::EValueComponentType::Float;
+		const FFloatValue LhsValue = Lhs.AsFloat();
+		const FFloatValue RhsValue = Rhs.AsFloat();
+		float ComponentValue = 0.0f;
+		for (int32 i = 0; i < NumComponents; ++i)
+		{
+			ComponentValue += LhsValue.Component[i] * RhsValue.Component[i];
+		}
+		Result.Component[0].Float = ComponentValue;
+	}
+	else
+	{
+		Result.ComponentType = UE::Shader::EValueComponentType::Int;
+		const FIntValue LhsValue = Lhs.AsInt();
+		const FIntValue RhsValue = Rhs.AsInt();
+		int32 ComponentValue = 0;
+		for (int32 i = 0; i < NumComponents; ++i)
+		{
+			ComponentValue += LhsValue.Component[i] * RhsValue.Component[i];
+		}
+		Result.Component[0].Int = ComponentValue;
+	}
+	return Result;
+}
+
+UE::Shader::FValue UE::Shader::Cross(const FValue& Lhs, const FValue& Rhs)
+{
+	FValue Result;
+	Result.NumComponents = 3;
+
+	if (Lhs.ComponentType == UE::Shader::EValueComponentType::Float || Rhs.ComponentType == UE::Shader::EValueComponentType::Float)
+	{
+		Result.ComponentType = UE::Shader::EValueComponentType::Float;
+		const FFloatValue LhsValue = Lhs.AsFloat();
+		const FFloatValue RhsValue = Rhs.AsFloat();
+		
+		Result.Component[0].Float = LhsValue.Component[1] * RhsValue.Component[2] - LhsValue.Component[2] * RhsValue.Component[1];
+		Result.Component[1].Float = LhsValue.Component[2] * RhsValue.Component[0] - LhsValue.Component[0] * RhsValue.Component[2];
+		Result.Component[2].Float = LhsValue.Component[0] * RhsValue.Component[1] - LhsValue.Component[1] * RhsValue.Component[0];
+	}
+	else
+	{
+		Result.ComponentType = UE::Shader::EValueComponentType::Int;
+		const FIntValue LhsValue = Lhs.AsInt();
+		const FIntValue RhsValue = Rhs.AsInt();
+
+		Result.Component[0].Int = LhsValue.Component[1] * RhsValue.Component[2] - LhsValue.Component[2] * RhsValue.Component[1];
+		Result.Component[1].Int = LhsValue.Component[2] * RhsValue.Component[0] - LhsValue.Component[0] * RhsValue.Component[2];
+		Result.Component[2].Int = LhsValue.Component[0] * RhsValue.Component[1] - LhsValue.Component[1] * RhsValue.Component[0];
+	}
+	return Result;
+}
+
+UE::Shader::FValue UE::Shader::Append(const FValue& Lhs, const FValue& Rhs)
+{
+	FValue Result;
+	int32 NumComponents = 0;
+	if (Lhs.ComponentType == Rhs.ComponentType)
+	{
+		// If both values have the same component type, use as-is
+		// (otherwise will need to convert)
+		Result.ComponentType = Lhs.ComponentType;
+		for (int32 i = 0; i < Lhs.NumComponents; ++i)
+		{
+			Result.Component[NumComponents++] = Lhs.Component[i];
+		}
+		for (int32 i = 0; i < Rhs.NumComponents && NumComponents < 4; ++i)
+		{
+			Result.Component[NumComponents++] = Rhs.Component[i];
+		}
+	}
+	else if (Lhs.ComponentType == UE::Shader::EValueComponentType::Float || Rhs.ComponentType == UE::Shader::EValueComponentType::Float)
+	{
+		Result.ComponentType = UE::Shader::EValueComponentType::Float;
+		const FFloatValue LhsValue = Lhs.AsFloat();
+		const FFloatValue RhsValue = Rhs.AsFloat();
+		for (int32 i = 0; i < Lhs.NumComponents; ++i)
+		{
+			Result.Component[NumComponents++].Float = LhsValue.Component[i];
+		}
+		for (int32 i = 0; i < Rhs.NumComponents && NumComponents < 4; ++i)
+		{
+			Result.Component[NumComponents++].Float = RhsValue.Component[i];
+		}
+	}
+	else
+	{
+		Result.ComponentType = UE::Shader::EValueComponentType::Int;
+		const FIntValue LhsValue = Lhs.AsInt();
+		const FIntValue RhsValue = Rhs.AsInt();
+		for (int32 i = 0; i < Lhs.NumComponents; ++i)
+		{
+			Result.Component[NumComponents++].Int = LhsValue.Component[i];
+		}
+		for (int32 i = 0; i < Rhs.NumComponents && NumComponents < 4; ++i)
+		{
+			Result.Component[NumComponents++].Int = RhsValue.Component[i];
+		}
+	}
+	Result.NumComponents = NumComponents;
+	return Result;
+}
diff --git a/Engine/Source/Runtime/Engine/Public/HLSLTree/HLSLTree.h b/Engine/Source/Runtime/Engine/Public/HLSLTree/HLSLTree.h
--- a/Engine/Source/Runtime/Engine/Public/HLSLTree/HLSLTree.h
+++ b/Engine/Source/Runtime/Engine/Public/HLSLTree/HLSLTree.h
@@ -9,15 +9,21 @@
 
 class FMaterial;
 class FMaterialCompilationOutput;
-class FMaterialPreshaderData;
+struct FStaticParameterSet;
+
+namespace UE
+{
+
+namespace Shader
+{
+class FPreshaderData;
+}
 
 /**
  * The HLSLTree module contains classes to build an HLSL AST (abstract syntax tree)
  * This allows C++ to procedurally define an HLSL program.  The structure of the tree is designed to be flexible, to facilitate incremental generation from a material node graph
  * Once the tree is complete, HLSL source code may be generated
  */
-namespace UE
-{
 namespace HLSLTree
 {
 
@@ -62,6 +68,8 @@ public:
 		StringBuilder->Append('\n');
 	}
 
+	void WriteConstant(const Shader::FValue& Value);
+
 	explicit FCodeWriter(FStringBuilderBase* InStringBuilder)
 		: StringBuilder(InStringBuilder)
 		, IndentLevel(0)
@@ -81,14 +89,15 @@ class FEmitValue
 {
 public:
 	EExpressionEvaluationType GetEvaluationType() const { return EvaluationType; }
-	const FConstant& GetConstantValue() const { return ConstantValue; }
+	Shader::EValueType GetExpressionType() const { return ExpressionType; }
+	const Shader::FValue& GetConstantValue() const { return ConstantValue; }
 
 private:
 	mutable const TCHAR* Code = nullptr;
-	const FMaterialPreshaderData* Preshader = nullptr;
+	const Shader::FPreshaderData* Preshader = nullptr;
 	EExpressionEvaluationType EvaluationType = EExpressionEvaluationType::None;
-	EExpressionType ExpressionType = EExpressionType::Void;
-	FConstant ConstantValue;
+	Shader::EValueType ExpressionType = Shader::EValueType::Void;
+	Shader::FValue ConstantValue;
 
 	friend class FEmitContext;
 };
@@ -113,7 +122,7 @@ public:
 	const TCHAR* GetCode(const FEmitValue& Value) const;
 
 	/** Append preshader bytecode that represents the given value */
-	void AppendPreshader(const FEmitValue& Value, FMaterialPreshaderData& InOutPreshader) const;
+	void AppendPreshader(const FEmitValue& Value, Shader::FPreshaderData& InOutPreshader) const;
 
 	struct FScopeEntry
 	{
@@ -146,9 +155,10 @@ public:
 
 	TArray<FScopeEntry> ScopeStack;
 	TArray<FFunctionStackEntry> FunctionStack;
-	TArray<FMaterialPreshaderData*> TempPreshaders;
+	TArray<Shader::FPreshaderData*> TempPreshaders;
 	FMemStackBase* Allocator = nullptr;
 	const FMaterial* Material = nullptr; // TODO - remove preshader material dependency
+	const FStaticParameterSet* StaticParameters = nullptr;
 	FMaterialCompilationOutput* MaterialCompilationOutput = nullptr;
 	int32 NumExpressionLocals = 0;
 	int32 NumTexCoords = 0;
@@ -156,16 +166,20 @@ public:
 
 struct FExpressionEmitResult
 {
-	FExpressionEmitResult(FCodeWriter& InWriter, FMaterialPreshaderData& InPreshader)
+	FExpressionEmitResult(FCodeWriter& InWriter, Shader::FPreshaderData& InPreshader)
 		: Writer(InWriter)
 		, Preshader(InPreshader)
 		, EvaluationType(EExpressionEvaluationType::Shader)
+		, Type(Shader::EValueType::Void)
 		, bInline(false)
 	{}
 
+	void ForwardValue(FEmitContext& Context, const FEmitValue& InValue);
+
 	FCodeWriter& Writer;
-	FMaterialPreshaderData& Preshader;
+	Shader::FPreshaderData& Preshader;
 	EExpressionEvaluationType EvaluationType;
+	Shader::EValueType Type;
 	bool bInline;
 };
 
@@ -236,14 +250,6 @@ public:
 
 	/* Emits HLSL code for the expression. The code should NOT include any newlines or semi-colons. The returned string may be assigned to a temporary variable, or embedded in another HLSL string */
 	virtual void EmitHLSL(FEmitContext& Context, FExpressionEmitResult& OutResult) const = 0;
-
-	/** The HLSL type of the expression (float, int, etc) */
-	EExpressionType Type;
-
-protected:
-	explicit FExpression(EExpressionType InType)
-		: Type(InType)
-	{}
 };
 
 /**
@@ -252,12 +258,12 @@ protected:
 class FLocalDeclaration final : public FNode
 {
 public:
-	FLocalDeclaration(const FName& InName, EExpressionType InType) : Name(InName), Type(InType) {}
+	FLocalDeclaration(const FName& InName, Shader::EValueType InType) : Name(InName), Type(InType) {}
 
 	virtual ENodeVisitResult Visit(FNodeVisitor& Visitor) override;
 
 	FName Name;
-	EExpressionType Type;
+	Shader::EValueType Type;
 };
 
 /**
@@ -266,12 +272,12 @@ public:
 class FParameterDeclaration final : public FNode
 {
 public:
-	FParameterDeclaration(const FName& InName, const FConstant& InDefaultValue) : Name(InName), DefaultValue(InDefaultValue) {}
+	FParameterDeclaration(const FName& InName, const Shader::FValue& InDefaultValue) : Name(InName), DefaultValue(InDefaultValue) {}
 
 	virtual ENodeVisitResult Visit(FNodeVisitor& Visitor) override;
 
 	FName Name;
-	FConstant DefaultValue;
+	Shader::FValue DefaultValue;
 };
 
 /**
@@ -296,12 +302,6 @@ class FFunctionCall final : public FNode
 public:
 	virtual ENodeVisitResult Visit(FNodeVisitor& Visitor) override;
 
-	inline EExpressionType GetOutputType(int32 Index) const
-	{
-		check(Index >= 0 && Index < NumOutputs);
-		return Outputs[Index] ? Outputs[Index]->Type : EExpressionType::Float1;
-	}
-
 	/** Root scope of the function to call. Note that this scope will be from a separate (external) tree */
 	const FScope* FunctionScope;
 
@@ -390,8 +390,8 @@ public:
 	 */
 	FScope* NewLinkedScope(FScope& Scope);
 	
-	FLocalDeclaration* NewLocalDeclaration(FScope& Scope, EExpressionType Type, const FName& Name);
-	FParameterDeclaration* NewParameterDeclaration(FScope& Scope, const FName& Name, const FConstant& DefaultValue);
+	FLocalDeclaration* NewLocalDeclaration(FScope& Scope, Shader::EValueType Type, const FName& Name);
+	FParameterDeclaration* NewParameterDeclaration(FScope& Scope, const FName& Name, const Shader::FValue& DefaultValue);
 	FTextureParameterDeclaration* NewTextureParameterDeclaration(FScope& Scope, const FName& Name, const FTextureDescription& DefaultValue);
 
 	FFunctionCall* NewFunctionCall(FScope& Scope,
diff --git a/Engine/Source/Runtime/Engine/Public/HLSLTree/HLSLTreeCommon.h b/Engine/Source/Runtime/Engine/Public/HLSLTree/HLSLTreeCommon.h
--- a/Engine/Source/Runtime/Engine/Public/HLSLTree/HLSLTreeCommon.h
+++ b/Engine/Source/Runtime/Engine/Public/HLSLTree/HLSLTreeCommon.h
@@ -20,11 +20,11 @@ enum class EBinaryOp
 class FExpressionConstant : public FExpression
 {
 public:
-	explicit FExpressionConstant(const FConstant& InValue)
-		: FExpression(InValue.Type), Value(InValue)
+	explicit FExpressionConstant(const Shader::FValue& InValue)
+		: Value(InValue)
 	{}
 
-	FConstant Value;
+	Shader::FValue Value;
 
 	virtual void EmitHLSL(FEmitContext& Context, FExpressionEmitResult& OutResult) const override;
 };
@@ -32,7 +32,7 @@ public:
 class FExpressionLocalVariable : public FExpression
 {
 public:
-	explicit FExpressionLocalVariable(FLocalDeclaration* InDeclaration) : FExpression(InDeclaration->Type), Declaration(InDeclaration) {}
+	explicit FExpressionLocalVariable(FLocalDeclaration* InDeclaration) : Declaration(InDeclaration) {}
 
 	FLocalDeclaration* Declaration;
 
@@ -53,8 +53,7 @@ class FExpressionParameter : public FExpression
 {
 public:
 	FExpressionParameter(FParameterDeclaration* InDeclaration)
-		: FExpression(InDeclaration->DefaultValue.Type)
-		, Declaration(InDeclaration)
+		: Declaration(InDeclaration)
 	{}
 
 	FParameterDeclaration* Declaration;
@@ -83,9 +82,9 @@ enum class EExternalInputType
 	TexCoord6,
 	TexCoord7,
 };
-inline EExpressionType GetInputExpressionType(EExternalInputType Type)
+inline Shader::EValueType GetInputExpressionType(EExternalInputType Type)
 {
-	return EExpressionType::Float2;
+	return Shader::EValueType::Float2;
 }
 inline EExternalInputType MakeInputTexCoord(int32 Index)
 {
@@ -96,7 +95,7 @@ inline EExternalInputType MakeInputTexCoord(int32 Index)
 class FExpressionExternalInput : public FExpression
 {
 public:
-	FExpressionExternalInput(EExternalInputType InInputType) : FExpression(GetInputExpressionType(InInputType)), InputType(InInputType) {}
+	FExpressionExternalInput(EExternalInputType InInputType) : InputType(InInputType) {}
 
 	EExternalInputType InputType;
 
@@ -107,8 +106,7 @@ class FExpressionTextureSample : public FExpression
 {
 public:
 	FExpressionTextureSample(FTextureParameterDeclaration* InDeclaration, FExpression* InTexCoordExpression)
-		: FExpression(EExpressionType::Float4)
-		, Declaration(InDeclaration)
+		: Declaration(InDeclaration)
 		, TexCoordExpression(InTexCoordExpression)
 		, SamplerSource(SSM_FromTextureAsset)
 		, MipValueMode(TMVM_None)
@@ -136,7 +134,7 @@ public:
 class FExpressionDefaultMaterialAttributes : public FExpression
 {
 public:
-	FExpressionDefaultMaterialAttributes() : FExpression(EExpressionType::MaterialAttributes) {}
+	FExpressionDefaultMaterialAttributes() {}
 
 	virtual void EmitHLSL(FEmitContext& Context, FExpressionEmitResult& OutResult) const override;
 };
@@ -144,7 +142,7 @@ public:
 class FExpressionSetMaterialAttribute : public FExpression
 {
 public:
-	FExpressionSetMaterialAttribute() : FExpression(EExpressionType::MaterialAttributes) {}
+	FExpressionSetMaterialAttribute() {}
 
 	FGuid AttributeID;
 	FExpression* AttributesExpression;
@@ -164,12 +162,39 @@ public:
 	virtual void EmitHLSL(FEmitContext& Context, FExpressionEmitResult& OutResult) const override;
 };
 
+class FExpressionSelect : public FExpression
+{
+public:
+	FExpressionSelect(FExpression* InCondition, FExpression* InTrue, FExpression* InFalse)
+		: ConditionExpression(InCondition)
+		, TrueExpression(InTrue)
+		, FalseExpression(InFalse)
+	{}
+
+	FExpression* ConditionExpression;
+	FExpression* TrueExpression;
+	FExpression* FalseExpression;
+
+	virtual ENodeVisitResult Visit(FNodeVisitor& Visitor) override
+	{
+		const ENodeVisitResult Result = FExpression::Visit(Visitor);
+		if (ShouldVisitDependentNodes(Result))
+		{
+			Visitor.VisitNode(ConditionExpression);
+			Visitor.VisitNode(TrueExpression);
+			Visitor.VisitNode(FalseExpression);
+		}
+		return Result;
+	}
+
+	virtual void EmitHLSL(FEmitContext& Context, FExpressionEmitResult& OutResult) const override;
+};
+
 class FExpressionBinaryOp : public FExpression
 {
 public:
-	FExpressionBinaryOp(EExpressionType ResultType, EBinaryOp InOp, FExpression* InLhs, FExpression* InRhs)
-		: FExpression(ResultType)
-		, Op(InOp)
+	FExpressionBinaryOp(EBinaryOp InOp, FExpression* InLhs, FExpression* InRhs)
+		: Op(InOp)
 		, Lhs(InLhs)
 		, Rhs(InRhs)
 	{}
@@ -206,8 +231,7 @@ class FExpressionSwizzle : public FExpression
 {
 public:
 	FExpressionSwizzle(const FSwizzleParameters& InParams, FExpression* InInput)
-		: FExpression(MakeExpressionType(InInput->Type, InParams.NumComponents))
-		, Parameters(InParams)
+		: Parameters(InParams)
 		, Input(InInput)
 	{}
 
@@ -230,12 +254,13 @@ public:
 class FExpressionCast : public FExpression
 {
 public:
-	FExpressionCast(EExpressionType InType, FExpression* InInput, ECastFlags InFlags = ECastFlags::None)
-		: FExpression(InType)
+	FExpressionCast(Shader::EValueType InType, FExpression* InInput, ECastFlags InFlags = ECastFlags::None)
+		: Type(InType)
 		, Input(InInput)
 		, Flags(InFlags)
 	{}
 
+	Shader::EValueType Type;
 	FExpression* Input;
 	ECastFlags Flags;
 
@@ -255,13 +280,14 @@ public:
 class FExpressionFunctionInput : public FExpression
 {
 public:
-	FExpressionFunctionInput(const FName& InName, EExpressionType InType, int32 InIndex)
-		: FExpression(InType), Name(InName), InputIndex(InIndex)
+	FExpressionFunctionInput(const FName& InName, Shader::EValueType InType, int32 InIndex)
+		: Name(InName), Type(InType), InputIndex(InIndex)
 	{}
 
 	virtual void EmitHLSL(FEmitContext& Context, FExpressionEmitResult& OutResult) const override;
 
 	FName Name;
+	Shader::EValueType Type;
 	int32 InputIndex;
 };
 
@@ -269,8 +295,7 @@ class FExpressionFunctionOutput : public FExpression
 {
 public:
 	FExpressionFunctionOutput(FFunctionCall* InFunctionCall, int32 InIndex)
-		: FExpression(InFunctionCall->GetOutputType(InIndex))
-		, FunctionCall(InFunctionCall)
+		: FunctionCall(InFunctionCall)
 		, OutputIndex(InIndex)
 	{
 		check(InIndex >= 0 && InIndex < InFunctionCall->NumOutputs);
diff --git a/Engine/Source/Runtime/Engine/Public/HLSLTree/HLSLTreeTypes.h b/Engine/Source/Runtime/Engine/Public/HLSLTree/HLSLTreeTypes.h
--- a/Engine/Source/Runtime/Engine/Public/HLSLTree/HLSLTreeTypes.h
+++ b/Engine/Source/Runtime/Engine/Public/HLSLTree/HLSLTreeTypes.h
@@ -2,6 +2,7 @@
 #pragma once
 
 #include "Misc/EnumClassFlags.h"
+#include "Shader/ShaderTypes.h"
 
 class UTexture;
 
@@ -21,40 +22,6 @@ class FFunctionCall;
 class FTree;
 class FCodeWriter;
 
-enum class EExpressionComponentType
-{
-	Void,
-	Float,
-	MaterialAttributes,
-};
-
-struct FExpressionTypeDescription
-{
-	FExpressionTypeDescription() : Name(nullptr), ComponentType(EExpressionComponentType::Void), NumComponents(0) {}
-	FExpressionTypeDescription(const TCHAR* InName, EExpressionComponentType InComponentType, int32 InNumComponents) : Name(InName), ComponentType(InComponentType), NumComponents(InNumComponents) {}
-
-	const TCHAR* Name;
-	EExpressionComponentType ComponentType;
-	int32 NumComponents;
-};
-
-enum class EExpressionType : uint8
-{
-	Void,
-
-	Float1,
-	Float2,
-	Float3,
-	Float4,
-
-	MaterialAttributes,
-};
-
-FExpressionTypeDescription GetExpressionTypeDescription(EExpressionType Type);
-EExpressionType MakeExpressionType(EExpressionComponentType ComponentType, int32 NumComponents);
-EExpressionType MakeExpressionType(EExpressionType BaseType, int32 NumComponents);
-EExpressionType MakeArithmeticResultType(EExpressionType Lhs, EExpressionType Rhs, FString& OutErrorMessage);
-
 /**
  * Describes how a given expression needs to be evaluated */
 enum class EExpressionEvaluationType
@@ -74,27 +41,6 @@ enum class EExpressionEvaluationType
 
 EExpressionEvaluationType CombineEvaluationTypes(EExpressionEvaluationType Lhs, EExpressionEvaluationType Rhs);
 
-struct FConstant
-{
-	inline FConstant() : Type(EExpressionType::Void) { FMemory::Memzero(Raw); }
-	inline FConstant(float v) : Type(EExpressionType::Float1) { FMemory::Memzero(Raw); Float[0] = v; }
-	inline FConstant(float X, float Y) : Type(EExpressionType::Float2) { FMemory::Memzero(Raw); Float[0] = X; Float[1] = Y; }
-	inline FConstant(float X, float Y, float Z) : Type(EExpressionType::Float3) { FMemory::Memzero(Raw); Float[0] = X; Float[1] = Y; Float[2] = Z; }
-	inline FConstant(float X, float Y, float Z, float W) : Type(EExpressionType::Float4) { FMemory::Memzero(Raw); Float[0] = X; Float[1] = Y; Float[2] = Z; Float[3] = W; }
-	inline FConstant(EExpressionType InType, const FLinearColor& Value) : Type(InType) { FMemory::Memzero(Raw); Float[0] = Value.R; Float[1] = Value.G; Float[2] = Value.B; Float[3] = Value.A; }
-
-	FLinearColor ToLinearColor() const;
-
-	void EmitHLSL(FCodeWriter& Writer) const;
-
-	EExpressionType Type;
-	union
-	{
-		uint32 Raw[4];
-		float Float[4];
-	};
-};
-
 struct FTextureDescription
 {
 	FTextureDescription()
diff --git a/Engine/Source/Runtime/Engine/Public/MaterialExpressionIO.h b/Engine/Source/Runtime/Engine/Public/MaterialExpressionIO.h
--- a/Engine/Source/Runtime/Engine/Public/MaterialExpressionIO.h
+++ b/Engine/Source/Runtime/Engine/Public/MaterialExpressionIO.h
@@ -12,11 +12,14 @@ namespace UE
 {
 namespace HLSLTree
 {
-enum class EExpressionType : uint8;
 class FScope;
 class FExpression;
 class FTextureParameterDeclaration;
 }
+namespace Shader
+{
+enum class EValueType : uint8;
+}
 }
 
 //
@@ -75,7 +78,7 @@ struct FExpressionInput
 #if WITH_EDITOR
 	ENGINE_API int32 Compile(class FMaterialCompiler* Compiler);
 	ENGINE_API UE::HLSLTree::FExpression* AcquireHLSLExpression(FMaterialHLSLGenerator& Generator, UE::HLSLTree::FScope& Scope) const;
-	ENGINE_API UE::HLSLTree::FExpression* AcquireHLSLExpressionWithCast(FMaterialHLSLGenerator& Generator, UE::HLSLTree::FScope& Scope, UE::HLSLTree::EExpressionType Type) const;
+	ENGINE_API UE::HLSLTree::FExpression* AcquireHLSLExpressionWithCast(FMaterialHLSLGenerator& Generator, UE::HLSLTree::FScope& Scope, UE::Shader::EValueType Type) const;
 	ENGINE_API UE::HLSLTree::FTextureParameterDeclaration* AcquireHLSLTexture(FMaterialHLSLGenerator& Generator, UE::HLSLTree::FScope& Scope) const;
 #endif // WITH_EDITOR
 
diff --git a/Engine/Source/Runtime/Engine/Public/MaterialHLSLGenerator.h b/Engine/Source/Runtime/Engine/Public/MaterialHLSLGenerator.h
--- a/Engine/Source/Runtime/Engine/Public/MaterialHLSLGenerator.h
+++ b/Engine/Source/Runtime/Engine/Public/MaterialHLSLGenerator.h
@@ -64,15 +64,15 @@ public:
 
 	UE::HLSLTree::FTree& GetTree() const { return *HLSLTree; }
 
-	UE::HLSLTree::FExpressionConstant* NewConstant(UE::HLSLTree::FScope& Scope, const UE::HLSLTree::FConstant& Value);
+	UE::HLSLTree::FExpressionConstant* NewConstant(UE::HLSLTree::FScope& Scope, const UE::Shader::FValue& Value);
 	UE::HLSLTree::FExpressionExternalInput* NewTexCoord(UE::HLSLTree::FScope& Scope, int32 Index);
 	UE::HLSLTree::FExpressionSwizzle* NewSwizzle(UE::HLSLTree::FScope& Scope, const UE::HLSLTree::FSwizzleParameters& Params, UE::HLSLTree::FExpression* Input);
-	UE::HLSLTree::FExpression* NewCast(UE::HLSLTree::FScope& Scope, UE::HLSLTree::EExpressionType Type, UE::HLSLTree::FExpression* Input, UE::HLSLTree::ECastFlags Flags = UE::HLSLTree::ECastFlags::None);
+	UE::HLSLTree::FExpression* NewCast(UE::HLSLTree::FScope& Scope, UE::Shader::EValueType Type, UE::HLSLTree::FExpression* Input, UE::HLSLTree::ECastFlags Flags = UE::HLSLTree::ECastFlags::None);
 
 	UE::HLSLTree::FExpression* NewFunctionInput(UE::HLSLTree::FScope& Scope, int32 InputIndex, UMaterialExpressionFunctionInput* MaterialFunctionInput);
 
-	UE::HLSLTree::FLocalDeclaration* AcquireLocalDeclaration(UE::HLSLTree::FScope& Scope, UE::HLSLTree::EExpressionType Type, const FName& Name);
-	UE::HLSLTree::FParameterDeclaration* AcquireParameterDeclaration(UE::HLSLTree::FScope& Scope, const FName& Name, const UE::HLSLTree::FConstant& DefaultValue);
+	UE::HLSLTree::FLocalDeclaration* AcquireLocalDeclaration(UE::HLSLTree::FScope& Scope, UE::Shader::EValueType Type, const FName& Name);
+	UE::HLSLTree::FParameterDeclaration* AcquireParameterDeclaration(UE::HLSLTree::FScope& Scope, const FName& Name, const UE::Shader::FValue& DefaultValue);
 
 	/** Returns a declaration to access the given texture, with no parameter */
 	UE::HLSLTree::FTextureParameterDeclaration* AcquireTextureDeclaration(UE::HLSLTree::FScope& Scope, const UE::HLSLTree::FTextureDescription& Value);
diff --git a/Engine/Source/Runtime/Engine/Public/MaterialShared.h b/Engine/Source/Runtime/Engine/Public/MaterialShared.h
--- a/Engine/Source/Runtime/Engine/Public/MaterialShared.h
+++ b/Engine/Source/Runtime/Engine/Public/MaterialShared.h
@@ -28,6 +28,8 @@
 #include "MaterialSceneTextureId.h"
 #include "VirtualTexturing.h"
 
+#include "Shader/Preshader.h"
+
 #include <atomic>
 
 struct FExpressionInput;
@@ -241,95 +243,10 @@ struct ENGINE_API FMaterialRenderContext
 		const FMaterialRenderProxy* InMaterialRenderProxy,
 		const FMaterial& InMaterial,
 		const FSceneView* InView);
-};
-
-enum class EMaterialPreshaderOpcode : uint8
-{
-	Nop,
-	ConstantZero,
-	Constant,
-	ScalarParameter,
-	VectorParameter,
-	Add,
-	Sub,
-	Mul,
-	Div,
-	Fmod,
-	Min,
-	Max,
-	Clamp,
-	Sin,
-	Cos,
-	Tan,
-	Asin,
-	Acos,
-	Atan,
-	Atan2,
-	Dot,
-	Cross,
-	Sqrt,
-	Length,
-	Saturate,
-	Abs,
-	Floor,
-	Ceil,
-	Round,
-	Trunc,
-	Sign,
-	Frac,
-	Fractional,
-	Log2,
-	Log10,
-	ComponentSwizzle,
-	AppendVector,
-	TextureSize,
-	TexelSize,
-	ExternalTextureCoordinateScaleRotation,
-	ExternalTextureCoordinateOffset,
-	RuntimeVirtualTextureUniform,
-};
-
-class FMaterialPreshaderData
-{
-	DECLARE_TYPE_LAYOUT(FMaterialPreshaderData, NonVirtual);
-public:
-	friend inline bool operator==(const FMaterialPreshaderData& Lhs, const FMaterialPreshaderData& Rhs)
-	{
-		return Lhs.Data == Rhs.Data;
-	}
-
-	friend inline bool operator!=(const FMaterialPreshaderData& Lhs, const FMaterialPreshaderData& Rhs)
-	{
-		return !operator==(Lhs, Rhs);
-	}
-
-	void Evaluate(FUniformExpressionSet* UniformExpressionSet, const struct FMaterialRenderContext& Context, FLinearColor& OutValue);
-
-	void Append(const FMaterialPreshaderData& InPreshader);
-
-	const int32 Num() const { return Data.Num(); }
-
-	void WriteData(const void* Value, uint32 Size);
-	void WriteName(const FScriptName& Name);
-
-	template<typename T>
-	FMaterialPreshaderData& Write(const T& Value) { WriteData(&Value, sizeof(T)); return *this; }
-
-	template<>
-	FMaterialPreshaderData& Write<FScriptName>(const FScriptName& Value) { WriteName(Value); return *this; }
-
-	/** Can't write FName, use FScriptName instead */
-	template<>
-	FMaterialPreshaderData& Write<FName>(const FName& Value) = delete;
-
-	template<>
-	FMaterialPreshaderData& Write<FHashedMaterialParameterInfo>(const FHashedMaterialParameterInfo& Value) { return Write(Value.Name).Write(Value.Index).Write(Value.Association); }
-
-	inline FMaterialPreshaderData& WriteOpcode(EMaterialPreshaderOpcode Op) { return Write<uint8>((uint8)Op); }
 
-	LAYOUT_FIELD(TMemoryImageArray<FScriptName>, Names);
-	LAYOUT_FIELD(TMemoryImageArray<uint32>, NameOffsets);
-	LAYOUT_FIELD(TMemoryImageArray<uint8>, Data);
+	void GetTextureParameterValue(const FHashedMaterialParameterInfo& ParameterInfo, int32 TextureIndex, const UTexture*& OutValue) const;
+	void GetTextureParameterValue(const FHashedMaterialParameterInfo& ParameterInfo, int32 TextureIndex, const URuntimeVirtualTexture*& OutValue) const;
+	FGuid GetExternalTextureGuid(const FGuid& ExternalTextureGuid, const FName& ParameterName, int32 SourceTextureIndex) const;
 };
 
 class FMaterialVirtualTextureStack
@@ -611,7 +528,7 @@ protected:
 	LAYOUT_ARRAY(TMemoryImageArray<FMaterialTextureParameterInfo>, UniformTextureParameters, NumMaterialTextureParameterTypes);
 	LAYOUT_FIELD(TMemoryImageArray<FMaterialExternalTextureParameterInfo>, UniformExternalTextureParameters);
 
-	LAYOUT_FIELD(FMaterialPreshaderData, UniformPreshaderData);
+	LAYOUT_FIELD(UE::Shader::FPreshaderData, UniformPreshaderData);
 
 	/** Virtual texture stacks found during compilation */
 	LAYOUT_FIELD(TMemoryImageArray<FMaterialVirtualTextureStack>, VTStacks);
diff --git a/Engine/Source/Runtime/Engine/Public/Shader/Preshader.h b/Engine/Source/Runtime/Engine/Public/Shader/Preshader.h
new file mode 100644
--- /dev/null
+++ b/Engine/Source/Runtime/Engine/Public/Shader/Preshader.h
@@ -0,0 +1,107 @@
+// Copyright Epic Games, Inc. All Rights Reserved.
+#pragma once
+
+#include "CoreMinimal.h"
+#include "Shader/PreshaderTypes.h"
+#include "Serialization/MemoryImage.h"
+#include "Materials/MaterialLayersFunctions.h"
+
+class FUniformExpressionSet;
+struct FMaterialRenderContext;
+
+namespace UE
+{
+namespace Shader
+{
+
+enum class EPreshaderOpcode : uint8
+{
+	Nop,
+	ConstantZero,
+	Constant,
+	ScalarParameter,
+	VectorParameter,
+	Add,
+	Sub,
+	Mul,
+	Div,
+	Fmod,
+	Min,
+	Max,
+	Clamp,
+	Sin,
+	Cos,
+	Tan,
+	Asin,
+	Acos,
+	Atan,
+	Atan2,
+	Dot,
+	Cross,
+	Sqrt,
+	Length,
+	Saturate,
+	Abs,
+	Floor,
+	Ceil,
+	Round,
+	Trunc,
+	Sign,
+	Frac,
+	Fractional,
+	Log2,
+	Log10,
+	ComponentSwizzle,
+	AppendVector,
+	TextureSize,
+	TexelSize,
+	ExternalTextureCoordinateScaleRotation,
+	ExternalTextureCoordinateOffset,
+	RuntimeVirtualTextureUniform,
+};
+
+class FPreshaderData
+{
+	DECLARE_TYPE_LAYOUT(FPreshaderData, NonVirtual);
+public:
+	friend inline bool operator==(const FPreshaderData& Lhs, const FPreshaderData& Rhs)
+	{
+		return Lhs.Data == Rhs.Data;
+	}
+
+	friend inline bool operator!=(const FPreshaderData& Lhs, const FPreshaderData& Rhs)
+	{
+		return !operator==(Lhs, Rhs);
+	}
+
+	void Evaluate(FUniformExpressionSet* UniformExpressionSet, const struct FMaterialRenderContext& Context, FValue& OutValue);
+
+	void Append(const FPreshaderData& InPreshader);
+
+	const int32 Num() const { return Data.Num(); }
+
+	void WriteData(const void* Value, uint32 Size);
+	void WriteName(const FScriptName& Name);
+
+	template<typename T>
+	FPreshaderData& Write(const T& Value) { WriteData(&Value, sizeof(T)); return *this; }
+
+	template<>
+	FPreshaderData& Write<FScriptName>(const FScriptName& Value) { WriteName(Value); return *this; }
+
+	/** Can't write FName, use FScriptName instead */
+	template<>
+	FPreshaderData& Write<FName>(const FName& Value) = delete;
+
+	template<>
+	FPreshaderData& Write<FHashedMaterialParameterInfo>(const FHashedMaterialParameterInfo& Value) { return Write(Value.Name).Write(Value.Index).Write(Value.Association); }
+
+	inline FPreshaderData& WriteOpcode(EPreshaderOpcode Op) { return Write<uint8>((uint8)Op); }
+
+	LAYOUT_FIELD(TMemoryImageArray<FScriptName>, Names);
+	LAYOUT_FIELD(TMemoryImageArray<uint32>, NameOffsets);
+	LAYOUT_FIELD(TMemoryImageArray<uint8>, Data);
+};
+
+} // namespace Shader
+} // namespace UE
\ No newline at end of file
diff --git a/Engine/Source/Runtime/Engine/Public/Shader/PreshaderTypes.h b/Engine/Source/Runtime/Engine/Public/Shader/PreshaderTypes.h
new file mode 100644
--- /dev/null
+++ b/Engine/Source/Runtime/Engine/Public/Shader/PreshaderTypes.h
@@ -0,0 +1,3 @@
+#pragma once
+
+#include "Shader/ShaderTypes.h"
diff --git a/Engine/Source/Runtime/Engine/Public/Shader/ShaderTypes.h b/Engine/Source/Runtime/Engine/Public/Shader/ShaderTypes.h
new file mode 100644
--- /dev/null
+++ b/Engine/Source/Runtime/Engine/Public/Shader/ShaderTypes.h
@@ -0,0 +1,178 @@
+#pragma once
+
+#include "CoreMinimal.h"
+
+namespace UE
+{
+namespace Shader
+{
+
+enum class EValueComponentType : uint8
+{
+	Void,
+	Float,
+	Int,
+	Bool,
+	MaterialAttributes,
+};
+
+struct FValueTypeDescription
+{
+	FValueTypeDescription() : Name(nullptr), ComponentType(EValueComponentType::Void), NumComponents(0) {}
+	FValueTypeDescription(const TCHAR* InName, EValueComponentType InComponentType, int8 InNumComponents) : Name(InName), ComponentType(InComponentType), NumComponents(InNumComponents) {}
+
+	const TCHAR* Name;
+	EValueComponentType ComponentType;
+	int8 NumComponents;
+};
+
+enum class EValueType : uint8
+{
+	Void,
+
+	Float1,
+	Float2,
+	Float3,
+	Float4,
+
+	Int1,
+	Int2,
+	Int3,
+	Int4,
+
+	Bool1,
+	Bool2,
+	Bool3,
+	Bool4,
+
+	MaterialAttributes,
+};
+
+FValueTypeDescription GetValueTypeDescription(EValueType Type);
+EValueType MakeValueType(EValueComponentType ComponentType, int32 NumComponents);
+EValueType MakeValueType(EValueType BaseType, int32 NumComponents);
+EValueType MakeArithmeticResultType(EValueType Lhs, EValueType Rhs, FString& OutErrorMessage);
+
+template<typename T>
+struct TValue
+{
+	inline T& operator[](int32 i) { check(i >= 0 && i < 4); return Component[i]; }
+	inline const T& operator[](int32 i) const { check(i >= 0 && i < 4); return Component[i]; }
+
+	T Component[4];
+};
+using FFloatValue = TValue<float>;
+using FIntValue = TValue<int32>;
+using FBoolValue = TValue<bool>;
+
+union FValueComponent
+{
+	FValueComponent() : Packed(0u) {}
+
+	uint32 Packed;
+	float Float;
+	int32 Int;
+	bool Bool;
+};
+static_assert(sizeof(FValueComponent) == sizeof(uint32), "bad packing");
+
+struct FValue
+{
+	inline FValue() : ComponentType(EValueComponentType::Void), NumComponents(0) {}
+
+	inline FValue(float v) : ComponentType(EValueComponentType::Float), NumComponents(1)
+	{
+		Component[0].Float = v;
+	}
+
+	inline FValue(float X, float Y) : ComponentType(EValueComponentType::Float), NumComponents(2)
+	{
+		Component[0].Float = X;
+		Component[1].Float = Y;
+	}
+
+	inline FValue(float X, float Y, float Z) : ComponentType(EValueComponentType::Float), NumComponents(3)
+	{
+		Component[0].Float = X;
+		Component[1].Float = Y;
+		Component[2].Float = Z;
+	}
+
+	inline FValue(float X, float Y, float Z, float W) : ComponentType(EValueComponentType::Float), NumComponents(4)
+	{
+		Component[0].Float = X;
+		Component[1].Float = Y;
+		Component[2].Float = Z;
+		Component[3].Float = W;
+	}
+
+	inline FValue(const FLinearColor& Value) : ComponentType(EValueComponentType::Float), NumComponents(4)
+	{
+		Component[0].Float = Value.R;
+		Component[1].Float = Value.G;
+		Component[2].Float = Value.B;
+		Component[3].Float = Value.A;
+	}
+
+	inline FValue(const FVector4& Value) : ComponentType(EValueComponentType::Float), NumComponents(4)
+	{
+		Component[0].Float = Value.X;
+		Component[1].Float = Value.Y;
+		Component[2].Float = Value.Z;
+		Component[3].Float = Value.W;
+	}
+
+	inline FValue(bool v) : ComponentType(EValueComponentType::Bool), NumComponents(1)
+	{
+		Component[0].Bool = v;
+	}
+
+	inline EValueType GetType() const { return MakeValueType(ComponentType, NumComponents); }
+
+	inline const FValueComponent& GetComponent(int32 i) const { check(i >= 0 && i < NumComponents); return Component[i]; }
+
+	FFloatValue AsFloat() const;
+	FIntValue AsInt() const;
+	FBoolValue AsBool() const;
+
+	FLinearColor AsLinearColor() const;
+	bool AsBoolScalar() const;
+
+	EValueComponentType ComponentType;
+	int8 NumComponents;
+	FValueComponent Component[4];
+};
+
+ENGINE_API FValue Abs(const FValue& Value);
+ENGINE_API FValue Saturate(const FValue& Value);
+ENGINE_API FValue Floor(const FValue& Value);
+ENGINE_API FValue Ceil(const FValue& Value);
+ENGINE_API FValue Round(const FValue& Value);
+ENGINE_API FValue Trunc(const FValue& Value);
+ENGINE_API FValue Sign(const FValue& Value);
+ENGINE_API FValue Frac(const FValue& Value);
+ENGINE_API FValue Fractional(const FValue& Value);
+ENGINE_API FValue Sqrt(const FValue& Value);
+ENGINE_API FValue Log2(const FValue& Value);
+ENGINE_API FValue Log10(const FValue& Value);
+ENGINE_API FValue Sin(const FValue& Value);
+ENGINE_API FValue Cos(const FValue& Value);
+ENGINE_API FValue Tan(const FValue& Value);
+ENGINE_API FValue Asin(const FValue& Value);
+ENGINE_API FValue Acos(const FValue& Value);
+ENGINE_API FValue Atan(const FValue& Value);
+
+ENGINE_API FValue Add(const FValue& Lhs, const FValue& Rhs);
+ENGINE_API FValue Sub(const FValue& Lhs, const FValue& Rhs);
+ENGINE_API FValue Mul(const FValue& Lhs, const FValue& Rhs);
+ENGINE_API FValue Div(const FValue& Lhs, const FValue& Rhs);
+ENGINE_API FValue Min(const FValue& Lhs, const FValue& Rhs);
+ENGINE_API FValue Max(const FValue& Lhs, const FValue& Rhs);
+ENGINE_API FValue Clamp(const FValue& Value, const FValue& Low, const FValue& High);
+ENGINE_API FValue Fmod(const FValue& Lhs, const FValue& Rhs);
+ENGINE_API FValue Atan2(const FValue& Lhs, const FValue& Rhs);
+ENGINE_API FValue Dot(const FValue& Lhs, const FValue& Rhs);
+ENGINE_API FValue Cross(const FValue& Lhs, const FValue& Rhs);
+ENGINE_API FValue Append(const FValue& Lhs, const FValue& Rhs);
+}
+}
