Fixes for RotateScaleOffsetTexCoordsDeriv

diff --git a/Engine/Source/Runtime/Engine/Private/Materials/HLSLMaterialDerivativeAutogen.cpp b/Engine/Source/Runtime/Engine/Private/Materials/HLSLMaterialDerivativeAutogen.cpp
--- a/Engine/Source/Runtime/Engine/Private/Materials/HLSLMaterialDerivativeAutogen.cpp
+++ b/Engine/Source/Runtime/Engine/Private/Materials/HLSLMaterialDerivativeAutogen.cpp
@@ -874,9 +874,9 @@ int32 FMaterialDerivativeAutogen::GenerateRotateScaleOffsetTexCoordsFunc(FHLSLMa
 	const uint32 NumResultComponents = 2;
 
 	const bool bAllZeroDeriv = (TexCoordDerivInfo.DerivativeStatus == EDerivativeStatus::Zero && RotationScaleDerivInfo.DerivativeStatus == EDerivativeStatus::Zero && OffsetDerivInfo.DerivativeStatus == EDerivativeStatus::Zero);
-	FString FiniteString = FString::Printf(TEXT("RotateScaleOffsetTexCoords(%s, %s, %s.xy)"),	*Translator.CoerceParameter(TexCoord,		ResultType), 
-																								*Translator.CoerceParameter(RotationScale,	ResultType),
-																								*Translator.CoerceParameter(Offset,			ResultType));
+	FString FiniteString = FString::Printf(TEXT("RotateScaleOffsetTexCoords(%s, %s, %s.xy)"),	*Translator.GetParameterCode(TexCoord), 
+																								*Translator.GetParameterCode(RotationScale),
+																								*Translator.GetParameterCode(Offset));
 
 	if (!bAllZeroDeriv && IsDerivativeValid(TexCoordDerivInfo.DerivativeStatus) && IsDerivativeValid(RotationScaleDerivInfo.DerivativeStatus) && IsDerivativeValid(OffsetDerivInfo.DerivativeStatus))
 	{
@@ -886,10 +886,8 @@ int32 FMaterialDerivativeAutogen::GenerateRotateScaleOffsetTexCoordsFunc(FHLSLMa
 		FString OffsetDeriv			= Translator.GetParameterCodeDeriv(Offset, CompiledPDV_Analytic);
 
 		TexCoordDeriv		= CoerceValueDeriv(TexCoordDeriv,		TexCoordDerivInfo.TypeIndex,		TexCoordDerivInfo.DerivativeStatus,			ResultTypeIndex);
-		RotationScaleDeriv	= CoerceValueDeriv(RotationScaleDeriv,	RotationScaleDerivInfo.TypeIndex,	RotationScaleDerivInfo.DerivativeStatus,	ResultTypeIndex);
-		OffsetDeriv			= CoerceValueDeriv(OffsetDeriv,			OffsetDerivInfo.TypeIndex,			OffsetDerivInfo.DerivativeStatus,			ResultTypeIndex);
 
-		FString AnalyticString = FString::Printf(TEXT("RotateScaleOffsetTexCoordsDeriv(%s, %s, %s)"), *TexCoordDeriv, *RotationScaleDeriv, *OffsetDeriv);
+		FString AnalyticString = FString::Printf(TEXT("RotateScaleOffsetTexCoordsDeriv(%s, %s, SwizzleDeriv2(%s, xy))"), *TexCoordDeriv, *RotationScaleDeriv, *OffsetDeriv);
 
 		check(NumResultComponents <= 4);
 		bRotateScaleOffsetTexCoords = true;
@@ -1650,11 +1648,11 @@ FString FMaterialDerivativeAutogen::GenerateUsedFunctions(FHLSLMaterialTranslato
 	{
 		// float2(dot(InTexCoords, InRotationScale.xy), dot(InTexCoords, InRotationScale.zw)) + InOffset;
 		// InTexCoords.xy * InRotationScale.xw + InTexCoords.yx * InRotationScale.yz + InOffset;
-		Ret += TEXT("FloatDeriv2 RotateScaleOffsetTexCoordsDeriv(FloatDeriv2 TexCoord, FloatDeriv2 RotationScale, FloatDeriv2 Offset)") LINE_TERMINATOR;
+		Ret += TEXT("FloatDeriv2 RotateScaleOffsetTexCoordsDeriv(FloatDeriv2 TexCoord, FloatDeriv4 RotationScale, FloatDeriv2 Offset)") LINE_TERMINATOR;
 		Ret += TEXT("{") LINE_TERMINATOR;
 		Ret += TEXT("\tFloatDeriv2 Ret = Offset;") LINE_TERMINATOR;
-		Ret += TEXT("\tRet = AddDeriv(Ret, MulDeriv(TexCoord, SwizzleDeriv2(RotationScale, xw));") LINE_TERMINATOR;
-		Ret += TEXT("\tRet = AddDeriv(Ret, MulDeriv(SwizzleDeriv2(TexCoord, yx), SwizzleDeriv2(RotationScale, yz));") LINE_TERMINATOR;
+		Ret += TEXT("\tRet = AddDeriv(Ret, MulDeriv(TexCoord, SwizzleDeriv2(RotationScale, xw)));") LINE_TERMINATOR;
+		Ret += TEXT("\tRet = AddDeriv(Ret, MulDeriv(SwizzleDeriv2(TexCoord, yx), SwizzleDeriv2(RotationScale, yz)));") LINE_TERMINATOR;
 		Ret += TEXT("\treturn Ret;") LINE_TERMINATOR;
 		Ret += TEXT("}") LINE_TERMINATOR;
 		Ret += TEXT("") LINE_TERMINATOR;
