Improve ensure message if it does go off during normal gameplay

diff --git a/Engine/Source/Runtime/Engine/Private/Components/PrimitiveComponent.cpp b/Engine/Source/Runtime/Engine/Private/Components/PrimitiveComponent.cpp
--- a/Engine/Source/Runtime/Engine/Private/Components/PrimitiveComponent.cpp
+++ b/Engine/Source/Runtime/Engine/Private/Components/PrimitiveComponent.cpp
@@ -2199,18 +2199,19 @@ bool UPrimitiveComponent::MoveComponentImpl( const FVector& Delta, const FQuat&
 
 		// Perform movement collision checking if needed for this actor.
 		const bool bCollisionEnabled = IsQueryCollisionEnabled();
-		if( bCollisionEnabled && (DeltaSizeSq > 0.f))
+		UWorld* const MyWorld = GetWorld();
+		if (MyWorld && bCollisionEnabled && (DeltaSizeSq > 0.f))
 		{
 #if !(UE_BUILD_SHIPPING || UE_BUILD_TEST)
-			if( !IsRegistered() )
+			if(!IsRegistered() && !MyWorld->bIsTearingDown)
 			{
 				if (Actor)
 				{
-					ensureMsgf(IsRegistered(), TEXT("%s MovedComponent %s not initialized deleteme %d"),*Actor->GetName(), *GetName(), Actor->IsPendingKill());
+					ensureMsgf(IsRegistered(), TEXT("%s MovedComponent %s not registered during sweep (PendingKill %d)"), *Actor->GetName(), *GetName(), Actor->IsPendingKill());
 				}
 				else
 				{ //-V523
-					ensureMsgf(IsRegistered(), TEXT("MovedComponent %s not initialized"), *GetFullName());
+					ensureMsgf(IsRegistered(), TEXT("Non-actor MovedComponent %s not registered during sweep"), *GetFullName());
 				}
 			}
 #endif
@@ -2218,8 +2219,6 @@ bool UPrimitiveComponent::MoveComponentImpl( const FVector& Delta, const FQuat&
 #if !(UE_BUILD_SHIPPING || UE_BUILD_TEST) && PERF_MOVECOMPONENT_STATS
 			MoveTimer.bDidLineCheck = true;
 #endif 
-			UWorld* const MyWorld = GetWorld();
-
 			static const FName TraceTagName = TEXT("MoveComponent");
 			const bool bForceGatherOverlaps = !ShouldCheckOverlapFlagToQueueOverlaps(*this);
 			FComponentQueryParams Params(SCENE_QUERY_STAT(MoveComponent), Actor);
