World Partition PIE streaming fixes
- Support the case where all actors in a cell are modified (so no actors to load).
- Don't crash when failing to load some actors (just log a warning for now).

diff --git a/Engine/Source/Runtime/Engine/Private/WorldPartition/WorldPartitionLevelStreamingDynamic.cpp b/Engine/Source/Runtime/Engine/Private/WorldPartition/WorldPartitionLevelStreamingDynamic.cpp
--- a/Engine/Source/Runtime/Engine/Private/WorldPartition/WorldPartitionLevelStreamingDynamic.cpp
+++ b/Engine/Source/Runtime/Engine/Private/WorldPartition/WorldPartitionLevelStreamingDynamic.cpp
@@ -296,20 +296,30 @@ bool UWorldPartitionLevelStreamingDynamic::IssueLoadRequests()
 		ActorContainerDup->MarkPendingKill();
 	}
 
+	auto FinalizeLoading = [this]()
+	{
+		check(bLoadRequestInProgress);
+		bLoadRequestInProgress = false;
+		FinalizeRuntimeLevel();
+	};
+
 	// Load saved actors
-	FWorldPartitionLevelHelper::LoadActors(RuntimeLevel, ChildPackagesToLoad, PackageCache, [this](bool bSucceeded)
+	if (ChildPackagesToLoad.Num())
 	{
-		if (bSucceeded)
-		{
-			check(bLoadRequestInProgress);
-			bLoadRequestInProgress = false;
-			FinalizeRuntimeLevel();
-		}
-		else
+		FWorldPartitionLevelHelper::LoadActors(RuntimeLevel, ChildPackagesToLoad, PackageCache, [this, FinalizeLoading](bool bSucceeded)
 		{
-			UE_LOG(LogLevelStreaming, Fatal, TEXT("UWorldPartitionLevelStreamingDynamic::IssueLoadRequests failed %s"), *GetWorldAssetPackageName());
-		}
-	}, /*bLoadForPlay=*/true, /*bLoadAsync=*/true, &InstancingContext);
+			if (!bSucceeded)
+			{
+				UE_LOG(LogLevelStreaming, Warning, TEXT("UWorldPartitionLevelStreamingDynamic::IssueLoadRequests failed %s"), *GetWorldAssetPackageName());
+			}
+
+			FinalizeLoading();
+		},/*bLoadForPlay=*/true, /*bLoadAsync=*/true, &InstancingContext);
+	}
+	else
+	{
+		FinalizeLoading();
+	}
 
 	return bLoadRequestInProgress;
 }
