Get rid of 1 in-memory copy when building a texture asynchronously by releasing the bulkdata earlier in the process.
  - 37GB -> 29GB peak memory when building T_vehCar_vehicle03_Grime

diff --git a/Engine/Source/Runtime/Engine/Private/TextureDerivedDataTask.cpp b/Engine/Source/Runtime/Engine/Private/TextureDerivedDataTask.cpp
--- a/Engine/Source/Runtime/Engine/Private/TextureDerivedDataTask.cpp
+++ b/Engine/Source/Runtime/Engine/Private/TextureDerivedDataTask.cpp
@@ -568,10 +568,13 @@ void FTextureCacheDerivedDataWorker::DoWork()
 		if (bAllowAsyncLoading && !bHasTextureSourceMips)
 		{
 			TextureData.GetAsyncSourceMips(ImageWrapper);
+			TextureData.AsyncSource.RemoveBulkData();
 		}
+
 		if (bAllowAsyncLoading && !bHasCompositeTextureSourceMips)
 		{
 			CompositeTextureData.GetAsyncSourceMips(ImageWrapper);
+			CompositeTextureData.AsyncSource.RemoveBulkData();
 		}
 
 		if (TextureData.Blocks.Num() && TextureData.Blocks[0].MipsPerLayer.Num() && TextureData.Blocks[0].MipsPerLayer[0].Num() && 
