Address some minor points with prior commit, raised in code review

diff --git a/Engine/Source/Runtime/Engine/Private/HLSLTree/HLSLTree.cpp b/Engine/Source/Runtime/Engine/Private/HLSLTree/HLSLTree.cpp
--- a/Engine/Source/Runtime/Engine/Private/HLSLTree/HLSLTree.cpp
+++ b/Engine/Source/Runtime/Engine/Private/HLSLTree/HLSLTree.cpp
@@ -229,8 +229,9 @@ const UE::HLSLTree::FEmitValue& UE::HLSLTree::FEmitContext::AcquireValue(UE::HLS
 		FLocalHLSLCodeWriter LocalWriter;
 		Shader::FPreshaderData LocalPreshader;
 		FExpressionEmitResult EmitResult(LocalWriter, LocalPreshader);
-		Expression->EmitHLSL(*this, EmitResult);
+		Expression->EmitCode(*this, EmitResult);
 
+		check(EmitResult.EvaluationType != EExpressionEvaluationType::None);
 		check(EmitResult.Type != Shader::EValueType::Void);
 
 		Entry = new(*Allocator) FDeclarationEntry();
diff --git a/Engine/Source/Runtime/Engine/Private/HLSLTree/HLSLTreeCommon.cpp b/Engine/Source/Runtime/Engine/Private/HLSLTree/HLSLTreeCommon.cpp
--- a/Engine/Source/Runtime/Engine/Private/HLSLTree/HLSLTreeCommon.cpp
+++ b/Engine/Source/Runtime/Engine/Private/HLSLTree/HLSLTreeCommon.cpp
@@ -59,7 +59,7 @@ UE::HLSLTree::FSwizzleParameters UE::HLSLTree::MakeSwizzleMask(bool bInR, bool b
 }
 
 
-void UE::HLSLTree::FExpressionConstant::EmitHLSL(UE::HLSLTree::FEmitContext& Context, FExpressionEmitResult& OutResult) const
+void UE::HLSLTree::FExpressionConstant::EmitCode(FEmitContext& Context, FExpressionEmitResult& OutResult) const
 {
 	OutResult.EvaluationType = EExpressionEvaluationType::Constant;
 	OutResult.Type = Value.GetType();
@@ -67,16 +67,17 @@ void UE::HLSLTree::FExpressionConstant::EmitHLSL(UE::HLSLTree::FEmitContext& Con
 	OutResult.Preshader.Write(Value);
 }
 
-void UE::HLSLTree::FExpressionLocalVariable::EmitHLSL(UE::HLSLTree::FEmitContext& Context, FExpressionEmitResult& OutResult) const
+void UE::HLSLTree::FExpressionLocalVariable::EmitCode(FEmitContext& Context, FExpressionEmitResult& OutResult) const
 {
 	const FEmitValue& DeclarationRef = Context.AcquireValue(Declaration);
 	
+	OutResult.EvaluationType = EExpressionEvaluationType::Shader;
 	OutResult.Type = DeclarationRef.GetExpressionType();
 	OutResult.bInline = true;
 	OutResult.Writer.Writef(TEXT("%s"), Context.GetCode(DeclarationRef));
 }
 
-void UE::HLSLTree::FExpressionParameter::EmitHLSL(UE::HLSLTree::FEmitContext& Context, FExpressionEmitResult& OutResult) const
+void UE::HLSLTree::FExpressionParameter::EmitCode(FEmitContext& Context, FExpressionEmitResult& OutResult) const
 {
 	OutResult.Type = Declaration->DefaultValue.GetType();
 
@@ -111,11 +112,12 @@ void UE::HLSLTree::FExpressionParameter::EmitHLSL(UE::HLSLTree::FEmitContext& Co
 	}
 }
 
-void UE::HLSLTree::FExpressionExternalInput::EmitHLSL(UE::HLSLTree::FEmitContext& Context, FExpressionEmitResult& OutResult) const
+void UE::HLSLTree::FExpressionExternalInput::EmitCode(FEmitContext& Context, FExpressionEmitResult& OutResult) const
 {
 	const int32 TypeIndex = (int32)InputType;
 	const int32 TexCoordIndex = TypeIndex - (int32)UE::HLSLTree::EExternalInputType::TexCoord0;
 
+	OutResult.EvaluationType = EExpressionEvaluationType::Shader;
 	OutResult.Type = GetInputExpressionType(InputType);
 	OutResult.bInline = true;
 	if (TexCoordIndex >= 0 && TexCoordIndex < 8)
@@ -129,7 +131,7 @@ void UE::HLSLTree::FExpressionExternalInput::EmitHLSL(UE::HLSLTree::FEmitContext
 	}
 }
 
-void UE::HLSLTree::FExpressionTextureSample::EmitHLSL(UE::HLSLTree::FEmitContext& Context, FExpressionEmitResult& OutResult) const
+void UE::HLSLTree::FExpressionTextureSample::EmitCode(FEmitContext& Context, FExpressionEmitResult& OutResult) const
 {
 	const UE::HLSLTree::FTextureDescription& Desc = Declaration->Description;
 	const EMaterialValueType MaterialValueType = Desc.Texture->GetMaterialType();
@@ -256,18 +258,20 @@ void UE::HLSLTree::FExpressionTextureSample::EmitHLSL(UE::HLSLTree::FEmitContext
 	}
 
 	const FEmitValue& TexCoordRef = Context.AcquireValue(TexCoordExpression);
+	OutResult.EvaluationType = EExpressionEvaluationType::Shader;
 	OutResult.Type = Shader::EValueType::Float4;
 	OutResult.Writer.Writef(TEXT("%s(%s(%s, %s, %s))"), SamplerTypeFunction, SampleFunctionName, *TextureName, *SamplerStateCode, Context.GetCode(TexCoordRef));
 }
 
-void UE::HLSLTree::FExpressionDefaultMaterialAttributes::EmitHLSL(UE::HLSLTree::FEmitContext& Context, FExpressionEmitResult& OutResult) const
+void UE::HLSLTree::FExpressionDefaultMaterialAttributes::EmitCode(FEmitContext& Context, FExpressionEmitResult& OutResult) const
 {
+	OutResult.EvaluationType = EExpressionEvaluationType::Shader;
 	OutResult.Type = Shader::EValueType::MaterialAttributes;
 	OutResult.bInline = true;
 	OutResult.Writer.Write(TEXT("DefaultMaterialAttributes"));
 }
 
-void UE::HLSLTree::FExpressionSetMaterialAttribute::EmitHLSL(UE::HLSLTree::FEmitContext& Context, FExpressionEmitResult& OutResult) const
+void UE::HLSLTree::FExpressionSetMaterialAttribute::EmitCode(FEmitContext& Context, FExpressionEmitResult& OutResult) const
 {
 	const EMaterialProperty Property = FMaterialAttributeDefinitionMap::GetProperty(AttributeID);
 	const FString PropertyName = FMaterialAttributeDefinitionMap::GetAttributeName(AttributeID);
@@ -277,11 +281,12 @@ void UE::HLSLTree::FExpressionSetMaterialAttribute::EmitHLSL(UE::HLSLTree::FEmit
 	const FEmitValue&AttributresRef = Context.AcquireValue(AttributesExpression);
 	const FEmitValue&ValueRef = Context.AcquireValue(ValueExpression);
 
+	OutResult.EvaluationType = EExpressionEvaluationType::Shader;
 	OutResult.Type = Shader::EValueType::MaterialAttributes;
 	OutResult.Writer.Writef(TEXT("FMaterialAttributes_Set%s(%s, %s)"), *PropertyName, Context.GetCode(AttributresRef), Context.GetCode(ValueRef));
 }
 
-void UE::HLSLTree::FExpressionSelect::EmitHLSL(FEmitContext& Context, FExpressionEmitResult& OutResult) const
+void UE::HLSLTree::FExpressionSelect::EmitCode(FEmitContext& Context, FExpressionEmitResult& OutResult) const
 {
 	const FEmitValue& ConditionValue = Context.AcquireValue(ConditionExpression);
 	if (ConditionValue.GetEvaluationType() == EExpressionEvaluationType::Constant)
@@ -300,9 +305,9 @@ void UE::HLSLTree::FExpressionSelect::EmitHLSL(FEmitContext& Context, FExpressio
 	}
 	else
 	{
-		// TODO - preshader
 		const FEmitValue& TrueValue = Context.AcquireValue(TrueExpression);
 		const FEmitValue& FalseValue = Context.AcquireValue(FalseExpression);
+		OutResult.EvaluationType = EExpressionEvaluationType::Shader; // TODO - preshader
 		OutResult.Type = TrueValue.GetExpressionType(); // TODO - FalseValue?
 		OutResult.Writer.Writef(TEXT("(%s ? %s : %s)"),
 			Context.GetCode(ConditionValue),
@@ -311,7 +316,7 @@ void UE::HLSLTree::FExpressionSelect::EmitHLSL(FEmitContext& Context, FExpressio
 	}
 }
 
-void UE::HLSLTree::FExpressionBinaryOp::EmitHLSL(UE::HLSLTree::FEmitContext& Context, FExpressionEmitResult& OutResult) const
+void UE::HLSLTree::FExpressionBinaryOp::EmitCode(FEmitContext& Context, FExpressionEmitResult& OutResult) const
 {
 	const FEmitValue& LhsRef = Context.AcquireValue(Lhs);
 	const FEmitValue& RhsRef = Context.AcquireValue(Rhs);
@@ -339,7 +344,7 @@ void UE::HLSLTree::FExpressionBinaryOp::EmitHLSL(UE::HLSLTree::FEmitContext& Con
 	}
 }
 
-void UE::HLSLTree::FExpressionSwizzle::EmitHLSL(FEmitContext& Context, FExpressionEmitResult& OutResult) const
+void UE::HLSLTree::FExpressionSwizzle::EmitCode(FEmitContext& Context, FExpressionEmitResult& OutResult) const
 {
 	static const TCHAR ComponentName[] = { 'x', 'y', 'z', 'w' };
 	TCHAR Swizzle[5] = TEXT("");
@@ -371,7 +376,7 @@ void UE::HLSLTree::FExpressionSwizzle::EmitHLSL(FEmitContext& Context, FExpressi
 	}
 }
 
-void UE::HLSLTree::FExpressionCast::EmitHLSL(FEmitContext& Context, FExpressionEmitResult& OutResult) const
+void UE::HLSLTree::FExpressionCast::EmitCode(FEmitContext& Context, FExpressionEmitResult& OutResult) const
 {
 	const FEmitValue& InputRef = Context.AcquireValue(Input);
 	const Shader::EValueType InputType = InputRef.GetExpressionType();
@@ -379,7 +384,7 @@ void UE::HLSLTree::FExpressionCast::EmitHLSL(FEmitContext& Context, FExpressionE
 
 	OutResult.bInline = true;
 	OutResult.Type = Type;
-	// TODO - preshader
+	OutResult.EvaluationType = EExpressionEvaluationType::Shader; // TODO - preshader
 
 	if (InputType == Type)
 	{
@@ -433,7 +438,7 @@ void UE::HLSLTree::FExpressionCast::EmitHLSL(FEmitContext& Context, FExpressionE
 	}
 }
 
-void UE::HLSLTree::FExpressionFunctionInput::EmitHLSL(FEmitContext& Context, FExpressionEmitResult& OutResult) const
+void UE::HLSLTree::FExpressionFunctionInput::EmitCode(FEmitContext& Context, FExpressionEmitResult& OutResult) const
 {
 	const FEmitContext::FFunctionStackEntry& StackEntry = Context.FunctionStack.Last();
 
@@ -443,26 +448,26 @@ void UE::HLSLTree::FExpressionFunctionInput::EmitHLSL(FEmitContext& Context, FEx
 	OutResult.ForwardValue(Context, Value);
 }
 
-void UE::HLSLTree::FExpressionFunctionOutput::EmitHLSL(FEmitContext& Context, FExpressionEmitResult& OutResult) const
+void UE::HLSLTree::FExpressionFunctionOutput::EmitCode(FEmitContext& Context, FExpressionEmitResult& OutResult) const
 {
 	const FEmitValue& Value = Context.AcquireValue(FunctionCall, OutputIndex);
 	OutResult.ForwardValue(Context, Value);
 }
 
-void UE::HLSLTree::FStatementReturn::EmitHLSL(UE::HLSLTree::FEmitContext& Context, UE::HLSLTree::FCodeWriter& OutWriter) const
+void UE::HLSLTree::FStatementReturn::EmitHLSL(FEmitContext& Context, FCodeWriter& OutWriter) const
 {
 	const FEmitValue& Ref = Context.AcquireValue(Expression);
 	OutWriter.WriteLinef(TEXT("return %s;"), Context.GetCode(Ref));
 }
 
-void UE::HLSLTree::FStatementSetLocalVariable::EmitHLSL(UE::HLSLTree::FEmitContext& Context, UE::HLSLTree::FCodeWriter& OutWriter) const
+void UE::HLSLTree::FStatementSetLocalVariable::EmitHLSL(FEmitContext& Context, FCodeWriter& OutWriter) const
 {
 	const FEmitValue& DeclarationRef = Context.AcquireValue(Declaration);
 	const FEmitValue& ExpressionRef = Context.AcquireValue(Expression);
 	OutWriter.WriteLinef(TEXT("%s = %s;"), Context.GetCode(DeclarationRef), Context.GetCode(ExpressionRef));
 }
 
-void UE::HLSLTree::FStatementIf::EmitHLSL(UE::HLSLTree::FEmitContext& Context, UE::HLSLTree::FCodeWriter& OutWriter) const
+void UE::HLSLTree::FStatementIf::EmitHLSL(FEmitContext& Context, FCodeWriter& OutWriter) const
 {
 	const FEmitValue& ConditionRef = Context.AcquireValue(ConditionExpression);
 	OutWriter.WriteLinef(TEXT("if (%s)"), Context.GetCode(ConditionRef));
@@ -474,7 +479,7 @@ void UE::HLSLTree::FStatementIf::EmitHLSL(UE::HLSLTree::FEmitContext& Context, U
 	}
 }
 
-void UE::HLSLTree::FStatementFor::EmitHLSL(UE::HLSLTree::FEmitContext& Context, UE::HLSLTree::FCodeWriter& OutWriter) const
+void UE::HLSLTree::FStatementFor::EmitHLSL(FEmitContext& Context, FCodeWriter& OutWriter) const
 {
 	
 }
diff --git a/Engine/Source/Runtime/Engine/Public/HLSLTree/HLSLTree.h b/Engine/Source/Runtime/Engine/Public/HLSLTree/HLSLTree.h
--- a/Engine/Source/Runtime/Engine/Public/HLSLTree/HLSLTree.h
+++ b/Engine/Source/Runtime/Engine/Public/HLSLTree/HLSLTree.h
@@ -169,7 +169,7 @@ struct FExpressionEmitResult
 	FExpressionEmitResult(FCodeWriter& InWriter, Shader::FPreshaderData& InPreshader)
 		: Writer(InWriter)
 		, Preshader(InPreshader)
-		, EvaluationType(EExpressionEvaluationType::Shader)
+		, EvaluationType(EExpressionEvaluationType::None)
 		, Type(Shader::EValueType::Void)
 		, bInline(false)
 	{}
@@ -248,8 +248,8 @@ class FExpression : public FNode
 public:
 	virtual ENodeVisitResult Visit(FNodeVisitor& Visitor) override;
 
-	/* Emits HLSL code for the expression. The code should NOT include any newlines or semi-colons. The returned string may be assigned to a temporary variable, or embedded in another HLSL string */
-	virtual void EmitHLSL(FEmitContext& Context, FExpressionEmitResult& OutResult) const = 0;
+	/** Emits code for the given expression, either HLSL code or preshader bytecode */
+	virtual void EmitCode(FEmitContext& Context, FExpressionEmitResult& OutResult) const = 0;
 };
 
 /**
diff --git a/Engine/Source/Runtime/Engine/Public/HLSLTree/HLSLTreeCommon.h b/Engine/Source/Runtime/Engine/Public/HLSLTree/HLSLTreeCommon.h
--- a/Engine/Source/Runtime/Engine/Public/HLSLTree/HLSLTreeCommon.h
+++ b/Engine/Source/Runtime/Engine/Public/HLSLTree/HLSLTreeCommon.h
@@ -26,7 +26,7 @@ public:
 
 	Shader::FValue Value;
 
-	virtual void EmitHLSL(FEmitContext& Context, FExpressionEmitResult& OutResult) const override;
+	virtual void EmitCode(FEmitContext& Context, FExpressionEmitResult& OutResult) const override;
 };
 
 class FExpressionLocalVariable : public FExpression
@@ -46,7 +46,7 @@ public:
 		return Result;
 	}
 
-	virtual void EmitHLSL(FEmitContext& Context, FExpressionEmitResult& OutResult) const override;
+	virtual void EmitCode(FEmitContext& Context, FExpressionEmitResult& OutResult) const override;
 };
 
 class FExpressionParameter : public FExpression
@@ -68,7 +68,7 @@ public:
 		return Result;
 	}
 
-	virtual void EmitHLSL(FEmitContext& Context, FExpressionEmitResult& OutResult) const override;
+	virtual void EmitCode(FEmitContext& Context, FExpressionEmitResult& OutResult) const override;
 };
 
 enum class EExternalInputType
@@ -99,7 +99,7 @@ public:
 
 	EExternalInputType InputType;
 
-	virtual void EmitHLSL(FEmitContext& Context, FExpressionEmitResult& OutResult) const override;
+	virtual void EmitCode(FEmitContext& Context, FExpressionEmitResult& OutResult) const override;
 };
 
 class FExpressionTextureSample : public FExpression
@@ -128,7 +128,7 @@ public:
 		return Result;
 	}
 
-	virtual void EmitHLSL(FEmitContext& Context, FExpressionEmitResult& OutResult) const override;
+	virtual void EmitCode(FEmitContext& Context, FExpressionEmitResult& OutResult) const override;
 };
 
 class FExpressionDefaultMaterialAttributes : public FExpression
@@ -136,7 +136,7 @@ class FExpressionDefaultMaterialAttributes : public FExpression
 public:
 	FExpressionDefaultMaterialAttributes() {}
 
-	virtual void EmitHLSL(FEmitContext& Context, FExpressionEmitResult& OutResult) const override;
+	virtual void EmitCode(FEmitContext& Context, FExpressionEmitResult& OutResult) const override;
 };
 
 class FExpressionSetMaterialAttribute : public FExpression
@@ -159,7 +159,7 @@ public:
 		return Result;
 	}
 
-	virtual void EmitHLSL(FEmitContext& Context, FExpressionEmitResult& OutResult) const override;
+	virtual void EmitCode(FEmitContext& Context, FExpressionEmitResult& OutResult) const override;
 };
 
 class FExpressionSelect : public FExpression
@@ -187,7 +187,7 @@ public:
 		return Result;
 	}
 
-	virtual void EmitHLSL(FEmitContext& Context, FExpressionEmitResult& OutResult) const override;
+	virtual void EmitCode(FEmitContext& Context, FExpressionEmitResult& OutResult) const override;
 };
 
 class FExpressionBinaryOp : public FExpression
@@ -214,7 +214,7 @@ public:
 		return Result;
 	}
 
-	virtual void EmitHLSL(FEmitContext& Context, FExpressionEmitResult& OutResult) const override;
+	virtual void EmitCode(FEmitContext& Context, FExpressionEmitResult& OutResult) const override;
 };
 
 struct FSwizzleParameters
@@ -248,7 +248,7 @@ public:
 		return Result;
 	}
 
-	virtual void EmitHLSL(FEmitContext& Context, FExpressionEmitResult& OutResult) const override;
+	virtual void EmitCode(FEmitContext& Context, FExpressionEmitResult& OutResult) const override;
 };
 
 class FExpressionCast : public FExpression
@@ -274,7 +274,7 @@ public:
 		return Result;
 	}
 
-	virtual void EmitHLSL(FEmitContext& Context, FExpressionEmitResult& OutResult) const override;
+	virtual void EmitCode(FEmitContext& Context, FExpressionEmitResult& OutResult) const override;
 };
 
 class FExpressionFunctionInput : public FExpression
@@ -284,7 +284,7 @@ public:
 		: Name(InName), Type(InType), InputIndex(InIndex)
 	{}
 
-	virtual void EmitHLSL(FEmitContext& Context, FExpressionEmitResult& OutResult) const override;
+	virtual void EmitCode(FEmitContext& Context, FExpressionEmitResult& OutResult) const override;
 
 	FName Name;
 	Shader::EValueType Type;
@@ -314,7 +314,7 @@ public:
 		return Result;
 	}
 
-	virtual void EmitHLSL(FEmitContext& Context, FExpressionEmitResult& OutResult) const override;
+	virtual void EmitCode(FEmitContext& Context, FExpressionEmitResult& OutResult) const override;
 };
 
 class FStatementReturn : public FStatement
diff --git a/Engine/Source/Runtime/Engine/Public/HLSLTree/HLSLTreeTypes.h b/Engine/Source/Runtime/Engine/Public/HLSLTree/HLSLTreeTypes.h
--- a/Engine/Source/Runtime/Engine/Public/HLSLTree/HLSLTreeTypes.h
+++ b/Engine/Source/Runtime/Engine/Public/HLSLTree/HLSLTreeTypes.h
@@ -29,13 +29,13 @@ enum class EExpressionEvaluationType
 	/** Invalid/uninitialized */
 	None,
 
-	/** The expression needs to generate HLSL code (EmitHLSL will be called) */
+	/** The expression outputs HLSL code (via FExpressionEmitResult::Writer) */
 	Shader,
 
-	/** The expression can generate a preshader (EmitPreshader will be called) */
+	/** The expression outputs preshader code evaluated at runtime (via FExpressionEmitResult::Preshader) */
 	Preshader,
 
-	/** The expression is constant (EmitPreshader will be called, generated preshader will be evaluated and cached at compile-time) */
+	/** The expression outputs constant preshader code evaluated at compile time (via FExpressionEmitResult::Preshader) */
 	Constant,
 };
 
