- Pass FRDGBuilder into FXSystem to allow them to make a choice about what runs on the RDG execution timeline and what needs to run on the building timeline

diff --git a/Engine/Plugins/FX/Niagara/Source/Niagara/Classes/NiagaraEmitterInstanceBatcher.h b/Engine/Plugins/FX/Niagara/Source/Niagara/Classes/NiagaraEmitterInstanceBatcher.h
--- a/Engine/Plugins/FX/Niagara/Source/Niagara/Classes/NiagaraEmitterInstanceBatcher.h
+++ b/Engine/Plugins/FX/Niagara/Source/Niagara/Classes/NiagaraEmitterInstanceBatcher.h
@@ -73,18 +73,18 @@ public:
 	virtual void AddVectorField(UVectorFieldComponent* VectorFieldComponent) override {}
 	virtual void RemoveVectorField(UVectorFieldComponent* VectorFieldComponent) override {}
 	virtual void UpdateVectorField(UVectorFieldComponent* VectorFieldComponent) override {}
-	virtual void PreInitViews(FRHICommandListImmediate& RHICmdList, bool bAllowGPUParticleUpdate) override;
-	virtual void PostInitViews(FRHICommandListImmediate& RHICmdList, FRHIUniformBuffer* ViewUniformBuffer, bool bAllowGPUParticleUpdate) override;
+	virtual void PreInitViews(class FRDGBuilder& GraphBuilder, bool bAllowGPUParticleUpdate) override;
+	virtual void PostInitViews(class FRDGBuilder& GraphBuilder, FRHIUniformBuffer* ViewUniformBuffer, bool bAllowGPUParticleUpdate) override;
 	virtual bool UsesGlobalDistanceField() const override;
 	virtual bool UsesDepthBuffer() const override;
 	virtual bool RequiresEarlyViewUniformBuffer() const override;
-	virtual void PreRender(FRHICommandListImmediate& RHICmdList, FRHIUniformBuffer* ViewUniformBuffer, const class FGlobalDistanceFieldParameterData* GlobalDistanceFieldParameterData, bool bAllowGPUParticleUpdate) override;
+	virtual void PreRender(class FRDGBuilder& GraphBuilder, FRHIUniformBuffer* ViewUniformBuffer, const class FGlobalDistanceFieldParameterData* GlobalDistanceFieldParameterData, bool bAllowGPUParticleUpdate) override;
 	virtual void OnDestroy() override; // Called on the gamethread to delete the batcher on the renderthread.
 
 	virtual void Tick(float DeltaTime) override;
 
 	virtual void PostRenderOpaque(
-		FRHICommandListImmediate& RHICmdList,
+		class FRDGBuilder& GraphBuilder,
 		FRHIUniformBuffer* ViewUniformBuffer,
 		const class FShaderParametersMetadata* SceneTexturesUniformBufferStruct,
 		FRHIUniformBuffer* SceneTexturesUniformBuffer,
diff --git a/Engine/Plugins/FX/Niagara/Source/Niagara/Private/NiagaraEmitterInstanceBatcher.cpp b/Engine/Plugins/FX/Niagara/Source/Niagara/Private/NiagaraEmitterInstanceBatcher.cpp
--- a/Engine/Plugins/FX/Niagara/Source/Niagara/Private/NiagaraEmitterInstanceBatcher.cpp
+++ b/Engine/Plugins/FX/Niagara/Source/Niagara/Private/NiagaraEmitterInstanceBatcher.cpp
@@ -480,9 +480,11 @@ void NiagaraEmitterInstanceBatcher::ProcessPendingTicksFlush(FRHICommandListImme
 			View.ViewUniformBuffer = TUniformBufferRef<FViewUniformShaderParameters>::CreateUniformBufferImmediate(*View.CachedViewUniformShaderParameters, UniformBuffer_SingleFrame);
 
 			// Execute all ticks that we can support without invalid simulations
-			PreInitViews(RHICmdList, true);
-			PostInitViews(RHICmdList, View.ViewUniformBuffer, true);
-			PostRenderOpaque(RHICmdList, View.ViewUniformBuffer, nullptr, nullptr, true);
+			FRDGBuilder GraphBuilder(RHICmdList);
+			PreInitViews(GraphBuilder, true);
+			PostInitViews(GraphBuilder, View.ViewUniformBuffer, true);
+			PostRenderOpaque(GraphBuilder, View.ViewUniformBuffer, nullptr, nullptr, true);
+			GraphBuilder.Execute();
 			break;
 		}
 
@@ -1401,7 +1403,7 @@ void NiagaraEmitterInstanceBatcher::ExecuteAll(FRHICommandList& RHICmdList, FRHI
 	CountsToRelease[int(TickStage)].Reset();
 }
 
-void NiagaraEmitterInstanceBatcher::PreInitViews(FRHICommandListImmediate& RHICmdList, bool bAllowGPUParticleUpdate)
+void NiagaraEmitterInstanceBatcher::PreInitViews(FRDGBuilder& GraphBuilder, bool bAllowGPUParticleUpdate)
 {
 	FramesBeforeTickFlush = 0;
 
@@ -1409,7 +1411,7 @@ void NiagaraEmitterInstanceBatcher::PreInitViews(FRHICommandListImmediate& RHICm
 #if NIAGARA_COMPUTEDEBUG_ENABLED
 	if ( FNiagaraGpuComputeDebug* GpuComputeDebug = GetGpuComputeDebug() )
 	{
-		GpuComputeDebug->Tick(RHICmdList);
+		GpuComputeDebug->Tick(GraphBuilder.RHICmdList);
 	}
 #endif
 
@@ -1418,7 +1420,7 @@ void NiagaraEmitterInstanceBatcher::PreInitViews(FRHICommandListImmediate& RHICm
 	if (GPUProfiler.IsProfilingEnabled())
 	{
 		TArray<FNiagaraGPUTimingResult, TInlineAllocator<16>> ProfilingResults;
-		GPUProfiler.QueryTimingResults(RHICmdList, ProfilingResults);
+		GPUProfiler.QueryTimingResults(GraphBuilder.RHICmdList, ProfilingResults);
 		TMap<TWeakObjectPtr<UNiagaraEmitter>, TMap<TStatIdData const*, float>> CapturedStats;
 		for (const FNiagaraGPUTimingResult& Result : ProfilingResults)
 		{
@@ -1461,55 +1463,76 @@ void NiagaraEmitterInstanceBatcher::PreInitViews(FRHICommandListImmediate& RHICm
 	// Update draw indirect buffer to max possible size.
 	if (bAllowGPUParticleUpdate && FNiagaraUtilities::AllowGPUParticles(GetShaderPlatform()))
 	{
-		UpdateInstanceCountManager(RHICmdList);
-		BuildTickStagePasses(RHICmdList, ETickStage::PreInitViews);
+		UpdateInstanceCountManager(GraphBuilder.RHICmdList);
 
 		if (GNiagaraAllowTickBeforeRender)
 		{
-			ExecuteAll(RHICmdList, nullptr, ETickStage::PreInitViews);
+			AddPass(
+				GraphBuilder,
+				RDG_EVENT_NAME("Niagara::PreInitViews"),
+				[this](FRHICommandListImmediate& RHICmdList)
+				{
+					BuildTickStagePasses(RHICmdList, ETickStage::PreInitViews);
+					ExecuteAll(RHICmdList, nullptr, ETickStage::PreInitViews);
+				}
+			);
 		}
 	}
 	else
 	{
-		GPUInstanceCounterManager.ResizeBuffers(RHICmdList, FeatureLevel,  0);
+		GPUInstanceCounterManager.ResizeBuffers(GraphBuilder.RHICmdList, FeatureLevel,  0);
 		FinishDispatches();
 	}
 }
 
-void NiagaraEmitterInstanceBatcher::PostInitViews(FRHICommandListImmediate& RHICmdList, FRHIUniformBuffer* ViewUniformBuffer, bool bAllowGPUParticleUpdate)
+void NiagaraEmitterInstanceBatcher::PostInitViews(FRDGBuilder& GraphBuilder, FRHIUniformBuffer* ViewUniformBuffer, bool bAllowGPUParticleUpdate)
 {
 	LLM_SCOPE(ELLMTag::Niagara);
 
 	if (bAllowGPUParticleUpdate && FNiagaraUtilities::AllowGPUParticles(GetShaderPlatform()))
 	{
-		BuildTickStagePasses(RHICmdList, ETickStage::PostInitViews);
-		ExecuteAll(RHICmdList, ViewUniformBuffer, ETickStage::PostInitViews);
+		AddPass(
+			GraphBuilder,
+			RDG_EVENT_NAME("Niagara::PostInitViews"),
+			[this, ViewUniformBuffer](FRHICommandListImmediate& RHICmdList)
+			{
+				BuildTickStagePasses(RHICmdList, ETickStage::PostInitViews);
+				ExecuteAll(RHICmdList, ViewUniformBuffer, ETickStage::PostInitViews);
+			}
+		);
 	}
 }
 
-void NiagaraEmitterInstanceBatcher::PostRenderOpaque(FRHICommandListImmediate& RHICmdList, FRHIUniformBuffer* ViewUniformBuffer, const class FShaderParametersMetadata* SceneTexturesUniformBufferStruct, FRHIUniformBuffer* SceneTexturesUniformBuffer, bool bAllowGPUParticleUpdate)
+void NiagaraEmitterInstanceBatcher::PostRenderOpaque(FRDGBuilder& GraphBuilder, FRHIUniformBuffer* ViewUniformBuffer, const class FShaderParametersMetadata* SceneTexturesUniformBufferStruct, FRHIUniformBuffer* SceneTexturesUniformBuffer, bool bAllowGPUParticleUpdate)
 {
 	LLM_SCOPE(ELLMTag::Niagara);
 
-	if (bAllowGPUParticleUpdate && FNiagaraUtilities::AllowGPUParticles(GetShaderPlatform()))
-	{
-		BuildTickStagePasses(RHICmdList, ETickStage::PostOpaqueRender);
+	AddPass(
+		GraphBuilder,
+		RDG_EVENT_NAME("Niagara::PostRenderOpaque"),
+		[this, ViewUniformBuffer, SceneTexturesUniformBufferStruct, SceneTexturesUniformBuffer, bAllowGPUParticleUpdate](FRHICommandListImmediate& RHICmdList)
+		{
+			if (bAllowGPUParticleUpdate && FNiagaraUtilities::AllowGPUParticles(GetShaderPlatform()))
+			{
+				BuildTickStagePasses(RHICmdList, ETickStage::PostOpaqueRender);
 
-		// Setup new readback since if there is no pending request, there is no risk of having invalid data read (offset being allocated after the readback was sent).
-		ExecuteAll(RHICmdList, ViewUniformBuffer, ETickStage::PostOpaqueRender);
+				// Setup new readback since if there is no pending request, there is no risk of having invalid data read (offset being allocated after the readback was sent).
+				ExecuteAll(RHICmdList, ViewUniformBuffer, ETickStage::PostOpaqueRender);
 
-		UpdateFreeIDBuffers(RHICmdList, DeferredIDBufferUpdates);
-		DeferredIDBufferUpdates.SetNum(0, false);
+				UpdateFreeIDBuffers(RHICmdList, DeferredIDBufferUpdates);
+				DeferredIDBufferUpdates.SetNum(0, false);
 
-		FinishDispatches();
+				FinishDispatches();
 
-		ProcessDebugReadbacks(RHICmdList, false);
-	}
+				ProcessDebugReadbacks(RHICmdList, false);
+			}
 
-	if (!GPUInstanceCounterManager.HasPendingGPUReadback())
-	{
-		GPUInstanceCounterManager.EnqueueGPUReadback(RHICmdList);
-	}
+			if (!GPUInstanceCounterManager.HasPendingGPUReadback())
+			{
+				GPUInstanceCounterManager.EnqueueGPUReadback(RHICmdList);
+			}
+		}
+	);
 }
 
 void NiagaraEmitterInstanceBatcher::ProcessDebugReadbacks(FRHICommandListImmediate& RHICmdList, bool bWaitCompletion)
@@ -1594,7 +1617,7 @@ bool NiagaraEmitterInstanceBatcher::RequiresEarlyViewUniformBuffer() const
 	return NumTicksThatRequireEarlyViewData > 0;
 }
 
-void NiagaraEmitterInstanceBatcher::PreRender(FRHICommandListImmediate& RHICmdList, FRHIUniformBuffer* ViewUniformBuffer, const class FGlobalDistanceFieldParameterData* GlobalDistanceFieldParameterData, bool bAllowGPUParticleUpdate)
+void NiagaraEmitterInstanceBatcher::PreRender(FRDGBuilder& GraphBuilder, FRHIUniformBuffer* ViewUniformBuffer, const class FGlobalDistanceFieldParameterData* GlobalDistanceFieldParameterData, bool bAllowGPUParticleUpdate)
 {
 	if (!FNiagaraUtilities::AllowGPUParticles(GetShaderPlatform()))
 	{
diff --git a/Engine/Source/Runtime/Engine/Private/GPUSortManager.cpp b/Engine/Source/Runtime/Engine/Private/GPUSortManager.cpp
--- a/Engine/Source/Runtime/Engine/Private/GPUSortManager.cpp
+++ b/Engine/Source/Runtime/Engine/Private/GPUSortManager.cpp
@@ -12,6 +12,9 @@
 #include "ProfilingDebugging/RealtimeGPUProfiler.h"
 #include "FXSystem.h" // FXConsoleVariables::bAllowGPUSorting
 
+#include "RenderGraphBuilder.h"
+#include "RenderGraphUtils.h"
+
 //*****************************************************************************
 
 int32 GGPUSortFrameCountBeforeBufferShrinking = 100;
@@ -722,7 +725,7 @@ void FGPUSortManager::UpdateSortBuffersPool()
 	}
 }
 
-void FGPUSortManager::OnPreRender(FRHICommandListImmediate& RHICmdList)
+void FGPUSortManager::OnPreRender(FRDGBuilder& GraphBuilder)
 {
 	LLM_SCOPE(ELLMTag::GPUSort);
 
@@ -732,79 +735,92 @@ void FGPUSortManager::OnPreRender(FRHICommandListImmediate& RHICmdList)
 	// Sort batches so that the next batch to handle is at the end of the array.
 	SortBatches.Sort([](const FSortBatch& A, const FSortBatch& B) { return (uint32)A.ProcessingOrder > (uint32)B.ProcessingOrder; });
 
-	if (SortBatches.Num())
-	{
-		SCOPED_GPU_STAT(RHICmdList, GPUKeyGenAndSort);
-		while (SortBatches.Num() && SortBatches.Last().ProcessingOrder == ESortBatchProcessingOrder::KeyGenAndSortAfterPreRender)
+	AddPass(
+		GraphBuilder,
+		RDG_EVENT_NAME("FGPUSortManager::OnPreRender"),
+		[this](FRHICommandListImmediate& RHICmdList)
 		{
-			// Remove the SortBatch but don't remove the SortBuffers from the pool since it can be reused immediately.
-			FSortBatch SortBatch = SortBatches.Pop();
-			FParticleSortBuffers* SortBuffers = SortBuffersPool.Last();
-			SortBatch.SortBuffers = SortBuffers;
-			checkSlow(SortBuffers && SortBatch.GetUsedValueCount() <= SortBuffers->GetSize());
+			if (SortBatches.Num())
+			{
+				SCOPED_GPU_STAT(RHICmdList, GPUKeyGenAndSort);
+				while (SortBatches.Num() && SortBatches.Last().ProcessingOrder == ESortBatchProcessingOrder::KeyGenAndSortAfterPreRender)
+				{
+					// Remove the SortBatch but don't remove the SortBuffers from the pool since it can be reused immediately.
+					FSortBatch SortBatch = SortBatches.Pop();
+					FParticleSortBuffers* SortBuffers = SortBuffersPool.Last();
+					SortBatch.SortBuffers = SortBuffers;
+					checkSlow(SortBuffers && SortBatch.GetUsedValueCount() <= SortBuffers->GetSize());
 
-			SCOPED_DRAW_EVENTF(RHICmdList, GPUSortBatch, TEXT("GPUSort_Batch%d(%s)"), SortBatch.Id, GetPrecisionString(SortBatch.Flags));
+					SCOPED_DRAW_EVENTF(RHICmdList, GPUSortBatch, TEXT("GPUSort_Batch%d(%s)"), SortBatch.Id, GetPrecisionString(SortBatch.Flags));
 
-			SortBatch.GenerateKeys(RHICmdList, Callbacks, EGPUSortFlags::KeyGenAfterPreRender);
-			SortBatch.SortAndResolve(RHICmdList, FeatureLevel);
+					SortBatch.GenerateKeys(RHICmdList, Callbacks, EGPUSortFlags::KeyGenAfterPreRender);
+					SortBatch.SortAndResolve(RHICmdList, FeatureLevel);
 
-			// Release the sort batch.
-			checkSlow(SortBatch.DynamicValueBuffer && SortBatch.DynamicValueBuffer->CurrentSortBatchId == SortBatch.Id);
-			SortBatch.DynamicValueBuffer->CurrentSortBatchId = INDEX_NONE;
-			SortBatch.SortBuffers = nullptr;
-		}
+					// Release the sort batch.
+					checkSlow(SortBatch.DynamicValueBuffer && SortBatch.DynamicValueBuffer->CurrentSortBatchId == SortBatch.Id);
+					SortBatch.DynamicValueBuffer->CurrentSortBatchId = INDEX_NONE;
+					SortBatch.SortBuffers = nullptr;
+				}
 
-		for (int32 BatchIndex = SortBatches.Num() - 1; BatchIndex >= 0 && SortBatches[BatchIndex].ProcessingOrder == ESortBatchProcessingOrder::KeyGenAfterPreRenderAndSortAfterPostRenderOpaque; --BatchIndex)
-		{
-			// Those sort batches will be processed again after PostRenderOpaque(). Because of this, the particle sort buffers can not be reused.
-			FSortBatch& SortBatch = SortBatches[BatchIndex];
-			FParticleSortBuffers* SortBuffers = SortBuffersPool.Pop();
-			SortBatch.SortBuffers = SortBuffers;
-			checkSlow(SortBuffers && SortBatch.GetUsedValueCount() <= SortBuffers->GetSize());
+				for (int32 BatchIndex = SortBatches.Num() - 1; BatchIndex >= 0 && SortBatches[BatchIndex].ProcessingOrder == ESortBatchProcessingOrder::KeyGenAfterPreRenderAndSortAfterPostRenderOpaque; --BatchIndex)
+				{
+					// Those sort batches will be processed again after PostRenderOpaque(). Because of this, the particle sort buffers can not be reused.
+					FSortBatch& SortBatch = SortBatches[BatchIndex];
+					FParticleSortBuffers* SortBuffers = SortBuffersPool.Pop();
+					SortBatch.SortBuffers = SortBuffers;
+					checkSlow(SortBuffers && SortBatch.GetUsedValueCount() <= SortBuffers->GetSize());
 
-			SCOPED_DRAW_EVENTF(RHICmdList, GPUSortBatchPreStep, TEXT("GPUSort_Batch%d(PreStep,%s)"), SortBatch.Id, GetPrecisionString(SortBatch.Flags));
+					SCOPED_DRAW_EVENTF(RHICmdList, GPUSortBatchPreStep, TEXT("GPUSort_Batch%d(PreStep,%s)"), SortBatch.Id, GetPrecisionString(SortBatch.Flags));
 
-			SortBatch.GenerateKeys(RHICmdList, Callbacks, EGPUSortFlags::KeyGenAfterPreRender);
+					SortBatch.GenerateKeys(RHICmdList, Callbacks, EGPUSortFlags::KeyGenAfterPreRender);
+				}
+			}
+			PostPreRenderEvent.Broadcast(RHICmdList);
 		}
-	}
-
-	PostPreRenderEvent.Broadcast(RHICmdList);
+	);
 }
 
-void FGPUSortManager::OnPostRenderOpaque(FRHICommandListImmediate& RHICmdList)
+void FGPUSortManager::OnPostRenderOpaque(FRDGBuilder& GraphBuilder)
 {
 	LLM_SCOPE(ELLMTag::GPUSort);
 
-	if (SortBatches.Num())
-	{
-		SCOPED_GPU_STAT(RHICmdList, GPUKeyGenAndSort);
-		while (SortBatches.Num())
+	AddPass(
+		GraphBuilder,
+		RDG_EVENT_NAME("FGPUSortManager::OnPostRenderOpaque"),
+		[this](FRHICommandListImmediate& RHICmdList)
 		{
-			// Remove the SortBatch but don't remove the SortBuffers from the pool since it can be reused immediately.
-			FSortBatch SortBatch = SortBatches.Pop();
-			checkSlow((SortBatch.SortBuffers != nullptr) == (SortBatch.ProcessingOrder == ESortBatchProcessingOrder::KeyGenAfterPreRenderAndSortAfterPostRenderOpaque));
-
-			SCOPED_DRAW_EVENTF(RHICmdList, GPUSortBatch, TEXT("GPUSort_Batch%d(%s)"), SortBatch.Id, GetPrecisionString(SortBatch.Flags));
-
-			if (!SortBatch.SortBuffers)
+			if (SortBatches.Num())
 			{
-				FParticleSortBuffers* SortBuffers = SortBuffersPool.Pop();
-				SortBatch.SortBuffers = SortBuffers;
-				checkSlow(SortBuffers && SortBatch.GetUsedValueCount() <= SortBuffers->GetSize());
+				SCOPED_GPU_STAT(RHICmdList, GPUKeyGenAndSort);
+				while (SortBatches.Num())
+				{
+					// Remove the SortBatch but don't remove the SortBuffers from the pool since it can be reused immediately.
+					FSortBatch SortBatch = SortBatches.Pop();
+					checkSlow((SortBatch.SortBuffers != nullptr) == (SortBatch.ProcessingOrder == ESortBatchProcessingOrder::KeyGenAfterPreRenderAndSortAfterPostRenderOpaque));
+
+					SCOPED_DRAW_EVENTF(RHICmdList, GPUSortBatch, TEXT("GPUSort_Batch%d(%s)"), SortBatch.Id, GetPrecisionString(SortBatch.Flags));
+
+					if (!SortBatch.SortBuffers)
+					{
+						FParticleSortBuffers* SortBuffers = SortBuffersPool.Pop();
+						SortBatch.SortBuffers = SortBuffers;
+						checkSlow(SortBuffers && SortBatch.GetUsedValueCount() <= SortBuffers->GetSize());
+					}
+
+					SortBatch.GenerateKeys(RHICmdList, Callbacks, EGPUSortFlags::KeyGenAfterPostRenderOpaque);
+					SortBatch.SortAndResolve(RHICmdList, FeatureLevel);
+
+					// Release the sort batch.
+					checkSlow(SortBatch.DynamicValueBuffer && SortBatch.DynamicValueBuffer->CurrentSortBatchId == SortBatch.Id);
+					SortBatch.DynamicValueBuffer->CurrentSortBatchId = INDEX_NONE;
+					SortBuffersPool.Push(SortBatch.SortBuffers);
+					SortBatch.SortBuffers = nullptr;
+				}
 			}
 
-			SortBatch.GenerateKeys(RHICmdList, Callbacks, EGPUSortFlags::KeyGenAfterPostRenderOpaque);
-			SortBatch.SortAndResolve(RHICmdList, FeatureLevel);
-
-			// Release the sort batch.
-			checkSlow(SortBatch.DynamicValueBuffer && SortBatch.DynamicValueBuffer->CurrentSortBatchId == SortBatch.Id);
-			SortBatch.DynamicValueBuffer->CurrentSortBatchId = INDEX_NONE;
-			SortBuffersPool.Push(SortBatch.SortBuffers);
-			SortBatch.SortBuffers = nullptr;
+			ResetDynamicValuesBuffers();
 		}
-	}
-
-	ResetDynamicValuesBuffers();
+	);
 }
 
 FGPUSortManager::FDynamicValueBuffer* FGPUSortManager::GetDynamicValueBufferFromPool(EGPUSortFlags TaskFlags, int32 SortBatchId)
diff --git a/Engine/Source/Runtime/Engine/Private/Particles/FXSystem.cpp b/Engine/Source/Runtime/Engine/Private/Particles/FXSystem.cpp
--- a/Engine/Source/Runtime/Engine/Private/Particles/FXSystem.cpp
+++ b/Engine/Source/Runtime/Engine/Private/Particles/FXSystem.cpp
@@ -386,7 +386,7 @@ void FFXSystem::DrawDebug( FCanvas* Canvas )
 	}
 }
 
-void FFXSystem::PreInitViews(FRHICommandListImmediate& RHICmdList, bool bAllowGPUParticleUpdate)
+void FFXSystem::PreInitViews(FRDGBuilder& GraphBuilder, bool bAllowGPUParticleUpdate)
 {
 	if (RHISupportsGPUParticles())
 	{
@@ -394,7 +394,7 @@ void FFXSystem::PreInitViews(FRHICommandListImmediate& RHICmdList, bool bAllowGP
 	}
 }
 
-void FFXSystem::PostInitViews(FRHICommandListImmediate& RHICmdList, FRHIUniformBuffer* ViewUniformBuffer, bool bAllowGPUParticleUpdate)
+void FFXSystem::PostInitViews(FRDGBuilder& GraphBuilder, FRHIUniformBuffer* ViewUniformBuffer, bool bAllowGPUParticleUpdate)
 {
 	// nothing to do here
 }
@@ -437,52 +437,65 @@ DECLARE_CYCLE_STAT(TEXT("FXPreRender_SimulateCDF"), STAT_CLM_FXPreRender_Simulat
 DECLARE_CYCLE_STAT(TEXT("FXPreRender_FinalizeCDF"), STAT_CLM_FXPreRender_FinalizeCDF, STATGROUP_CommandListMarkers);
 
 
-void FFXSystem::PreRender(FRHICommandListImmediate& RHICmdList, FRHIUniformBuffer* ViewUniformBuffer, const FGlobalDistanceFieldParameterData* GlobalDistanceFieldParameterData, bool bAllowGPUParticleSceneUpdate)
+void FFXSystem::PreRender(FRDGBuilder& GraphBuilder, FRHIUniformBuffer* ViewUniformBuffer, const FGlobalDistanceFieldParameterData* GlobalDistanceFieldParameterData, bool bAllowGPUParticleSceneUpdate)
 {
 	if (RHISupportsGPUParticles() && bAllowGPUParticleSceneUpdate)
 	{
-		SCOPED_DRAW_EVENT(RHICmdList, GPUParticles_PreRender);
-		UpdateMultiGPUResources(RHICmdList);
+		AddPass(
+			GraphBuilder,
+			RDG_EVENT_NAME("FFXSystem::PreRender"),
+			[this, ViewUniformBuffer, GlobalDistanceFieldParameterData](FRHICommandListImmediate& RHICmdList)
+			{
+				//SCOPE_CYCLE_COUNTER(STAT_FDeferredShadingSceneRenderer_FXSystem_PostRenderOpaque);
 
-		RHICmdList.SetCurrentStat(GET_STATID(STAT_CLM_FXPreRender_Prepare));
-		PrepareGPUSimulation(RHICmdList);
+				SCOPED_DRAW_EVENT(RHICmdList, GPUParticles_PreRender);
+				UpdateMultiGPUResources(RHICmdList);
 
-		RHICmdList.SetCurrentStat(GET_STATID(STAT_CLM_FXPreRender_Simulate));
-		SimulateGPUParticles(RHICmdList, EParticleSimulatePhase::Main, nullptr, nullptr, nullptr, nullptr);
+				RHICmdList.SetCurrentStat(GET_STATID(STAT_CLM_FXPreRender_Prepare));
+				PrepareGPUSimulation(RHICmdList);
 
-		RHICmdList.SetCurrentStat(GET_STATID(STAT_CLM_FXPreRender_Finalize));
-		FinalizeGPUSimulation(RHICmdList);
+				RHICmdList.SetCurrentStat(GET_STATID(STAT_CLM_FXPreRender_Simulate));
+				SimulateGPUParticles(RHICmdList, EParticleSimulatePhase::Main, nullptr, nullptr, nullptr, nullptr);
 
-		if (IsParticleCollisionModeSupported(GetShaderPlatform(), PCM_DistanceField))
-		{
-			RHICmdList.SetCurrentStat(GET_STATID(STAT_CLM_FXPreRender_PrepareCDF));
-			PrepareGPUSimulation(RHICmdList);
-
-			RHICmdList.SetCurrentStat(GET_STATID(STAT_CLM_FXPreRender_SimulateCDF));
-			SimulateGPUParticles(RHICmdList, EParticleSimulatePhase::CollisionDistanceField, ViewUniformBuffer, GlobalDistanceFieldParameterData, nullptr, nullptr);
-			//particles rendered during basepass may need to read pos/velocity buffers; must finalize unless we know for sure that nothing in base pass will read it.
-			RHICmdList.SetCurrentStat(GET_STATID(STAT_CLM_FXPreRender_FinalizeCDF));
-			FinalizeGPUSimulation(RHICmdList);
-		}
+				RHICmdList.SetCurrentStat(GET_STATID(STAT_CLM_FXPreRender_Finalize));
+				FinalizeGPUSimulation(RHICmdList);
+
+				if (IsParticleCollisionModeSupported(GetShaderPlatform(), PCM_DistanceField))
+				{
+					RHICmdList.SetCurrentStat(GET_STATID(STAT_CLM_FXPreRender_PrepareCDF));
+					PrepareGPUSimulation(RHICmdList);
+
+					RHICmdList.SetCurrentStat(GET_STATID(STAT_CLM_FXPreRender_SimulateCDF));
+					SimulateGPUParticles(RHICmdList, EParticleSimulatePhase::CollisionDistanceField, ViewUniformBuffer, GlobalDistanceFieldParameterData, nullptr, nullptr);
+					//particles rendered during basepass may need to read pos/velocity buffers; must finalize unless we know for sure that nothing in base pass will read it.
+					RHICmdList.SetCurrentStat(GET_STATID(STAT_CLM_FXPreRender_FinalizeCDF));
+					FinalizeGPUSimulation(RHICmdList);
+				}
+			}
+		);
     }
 }
 
 void FFXSystem::PostRenderOpaque(
-	FRHICommandListImmediate& RHICmdList, 
+	FRDGBuilder& GraphBuilder,
 	FRHIUniformBuffer* ViewUniformBuffer,
 	const FShaderParametersMetadata* SceneTexturesUniformBufferStruct,
 	FRHIUniformBuffer* SceneTexturesUniformBuffer,
 	bool bAllowGPUParticleUpdate)
 {
-	if (RHISupportsGPUParticles() && IsParticleCollisionModeSupported(GetShaderPlatform(), PCM_DepthBuffer))
+	if (RHISupportsGPUParticles() && IsParticleCollisionModeSupported(GetShaderPlatform(), PCM_DepthBuffer) && bAllowGPUParticleUpdate)
 	{
-		if (bAllowGPUParticleUpdate)
-		{
-			SCOPED_DRAW_EVENT(RHICmdList, GPUParticles_PostRenderOpaque);
-			PrepareGPUSimulation(RHICmdList);
-			SimulateGPUParticles(RHICmdList, EParticleSimulatePhase::CollisionDepthBuffer, ViewUniformBuffer, NULL, SceneTexturesUniformBufferStruct, SceneTexturesUniformBuffer);
-			FinalizeGPUSimulation(RHICmdList);
-		}
+		AddPass(
+			GraphBuilder,
+			RDG_EVENT_NAME("FFXSystem::PostRenderOpaque"),
+			[this, ViewUniformBuffer, SceneTexturesUniformBufferStruct, SceneTexturesUniformBuffer](FRHICommandListImmediate& RHICmdList)
+			{
+				SCOPED_DRAW_EVENT(RHICmdList, GPUParticles_PostRenderOpaque);
+				PrepareGPUSimulation(RHICmdList);
+				SimulateGPUParticles(RHICmdList, EParticleSimulatePhase::CollisionDepthBuffer, ViewUniformBuffer, NULL, SceneTexturesUniformBufferStruct, SceneTexturesUniformBuffer);
+				FinalizeGPUSimulation(RHICmdList);
+			}
+		);
 	}
 }
 
diff --git a/Engine/Source/Runtime/Engine/Private/Particles/FXSystemPrivate.h b/Engine/Source/Runtime/Engine/Private/Particles/FXSystemPrivate.h
--- a/Engine/Source/Runtime/Engine/Private/Particles/FXSystemPrivate.h
+++ b/Engine/Source/Runtime/Engine/Private/Particles/FXSystemPrivate.h
@@ -150,14 +150,14 @@ public:
 	virtual void RemoveVectorField(UVectorFieldComponent* VectorFieldComponent) override;
 	virtual void UpdateVectorField(UVectorFieldComponent* VectorFieldComponent) override;
 	FParticleEmitterInstance* CreateGPUSpriteEmitterInstance(FGPUSpriteEmitterInfo& EmitterInfo);
-	virtual void PreInitViews(FRHICommandListImmediate& RHICmdList, bool bAllowGPUParticleUpdate) override;
-	virtual void PostInitViews(FRHICommandListImmediate& RHICmdList, FRHIUniformBuffer* ViewUniformBuffer, bool bAllowGPUParticleUpdate) override;
+	virtual void PreInitViews(class FRDGBuilder& GraphBuilder, bool bAllowGPUParticleUpdate) override;
+	virtual void PostInitViews(class FRDGBuilder& GraphBuilder, FRHIUniformBuffer* ViewUniformBuffer, bool bAllowGPUParticleUpdate) override;
 	virtual bool UsesGlobalDistanceField() const override;
 	virtual bool UsesDepthBuffer() const override;
 	virtual bool RequiresEarlyViewUniformBuffer() const override;
-	virtual void PreRender(FRHICommandListImmediate& RHICmdList, FRHIUniformBuffer* ViewUniformBuffer, const FGlobalDistanceFieldParameterData* GlobalDistanceFieldParameterData, bool bAllowGPUParticleUpdate) override;
+	virtual void PreRender(class FRDGBuilder& GraphBuilder, FRHIUniformBuffer* ViewUniformBuffer, const FGlobalDistanceFieldParameterData* GlobalDistanceFieldParameterData, bool bAllowGPUParticleUpdate) override;
 	virtual void PostRenderOpaque(
-		FRHICommandListImmediate& RHICmdList, 
+		class FRDGBuilder& GraphBuilder,
 		FRHIUniformBuffer* ViewUniformBuffer,
 		const FShaderParametersMetadata* SceneTexturesUniformBufferStruct,
 		FRHIUniformBuffer* SceneTexturesUniformBuffer,
diff --git a/Engine/Source/Runtime/Engine/Private/Particles/FXSystemSet.cpp b/Engine/Source/Runtime/Engine/Private/Particles/FXSystemSet.cpp
--- a/Engine/Source/Runtime/Engine/Private/Particles/FXSystemSet.cpp
+++ b/Engine/Source/Runtime/Engine/Private/Particles/FXSystemSet.cpp
@@ -124,21 +124,21 @@ void FFXSystemSet::UpdateVectorField(UVectorFieldComponent* VectorFieldComponent
 	}
 }
 
-void FFXSystemSet::PreInitViews(FRHICommandListImmediate& RHICmdList, bool bAllowGPUParticleUpdate)
+void FFXSystemSet::PreInitViews(class FRDGBuilder& GraphBuilder, bool bAllowGPUParticleUpdate)
 {
 	for (FFXSystemInterface* FXSystem : FXSystems)
 	{
 		check(FXSystem);
-		FXSystem->PreInitViews(RHICmdList, bAllowGPUParticleUpdate);
+		FXSystem->PreInitViews(GraphBuilder, bAllowGPUParticleUpdate);
 	}
 }
 
-void FFXSystemSet::PostInitViews(FRHICommandListImmediate& RHICmdList, FRHIUniformBuffer* ViewUniformBuffer, bool bAllowGPUParticleUpdate)
+void FFXSystemSet::PostInitViews(class FRDGBuilder& GraphBuilder, FRHIUniformBuffer* ViewUniformBuffer, bool bAllowGPUParticleUpdate)
 {
 	for (FFXSystemInterface* FXSystem : FXSystems)
 	{
 		check(FXSystem);
-		FXSystem->PostInitViews(RHICmdList, ViewUniformBuffer, bAllowGPUParticleUpdate);
+		FXSystem->PostInitViews(GraphBuilder, ViewUniformBuffer, bAllowGPUParticleUpdate);
 	}
 }
 
@@ -181,17 +181,17 @@ bool FFXSystemSet::RequiresEarlyViewUniformBuffer() const
 	return false;
 }
 
-void FFXSystemSet::PreRender(FRHICommandListImmediate& RHICmdList, FRHIUniformBuffer* ViewUniformBuffer, const class FGlobalDistanceFieldParameterData* GlobalDistanceFieldParameterData, bool bAllowGPUParticleSceneUpdate)
+void FFXSystemSet::PreRender(class FRDGBuilder& GraphBuilder, FRHIUniformBuffer* ViewUniformBuffer, const class FGlobalDistanceFieldParameterData* GlobalDistanceFieldParameterData, bool bAllowGPUParticleSceneUpdate)
 {
 	for (FFXSystemInterface* FXSystem : FXSystems)
 	{
 		check(FXSystem);
-		FXSystem->PreRender(RHICmdList, ViewUniformBuffer, GlobalDistanceFieldParameterData, bAllowGPUParticleSceneUpdate);
+		FXSystem->PreRender(GraphBuilder, ViewUniformBuffer, GlobalDistanceFieldParameterData, bAllowGPUParticleSceneUpdate);
 	}
 }
 
 void FFXSystemSet::PostRenderOpaque(
-	FRHICommandListImmediate& RHICmdList,
+	class FRDGBuilder& GraphBuilder,
 	FRHIUniformBuffer* ViewUniformBuffer,
 	const class FShaderParametersMetadata* SceneTexturesUniformBufferStruct,
 	FRHIUniformBuffer* SceneTexturesUniformBuffer,
@@ -201,7 +201,7 @@ void FFXSystemSet::PostRenderOpaque(
 	{
 		check(FXSystem);
 		FXSystem->PostRenderOpaque(
-			RHICmdList,
+			GraphBuilder,
 			ViewUniformBuffer,
 			SceneTexturesUniformBufferStruct,
 			SceneTexturesUniformBuffer,
diff --git a/Engine/Source/Runtime/Engine/Private/Particles/FXSystemSet.h b/Engine/Source/Runtime/Engine/Private/Particles/FXSystemSet.h
--- a/Engine/Source/Runtime/Engine/Private/Particles/FXSystemSet.h
+++ b/Engine/Source/Runtime/Engine/Private/Particles/FXSystemSet.h
@@ -39,14 +39,14 @@ public:
 	virtual void AddVectorField(UVectorFieldComponent* VectorFieldComponent) override;
 	virtual void RemoveVectorField(UVectorFieldComponent* VectorFieldComponent) override;
 	virtual void UpdateVectorField(UVectorFieldComponent* VectorFieldComponent) override;
-	virtual void PreInitViews(FRHICommandListImmediate& RHICmdList, bool bAllowGPUParticleUpdate) override;
-	virtual void PostInitViews(FRHICommandListImmediate& RHICmdList, FRHIUniformBuffer* ViewUniformBuffer, bool bAllowGPUParticleUpdate) override;
+	virtual void PreInitViews(class FRDGBuilder& GraphBuilder, bool bAllowGPUParticleUpdate) override;
+	virtual void PostInitViews(class FRDGBuilder& GraphBuilder, FRHIUniformBuffer* ViewUniformBuffer, bool bAllowGPUParticleUpdate) override;
 	virtual bool UsesGlobalDistanceField() const override;
 	virtual bool UsesDepthBuffer() const override;
 	virtual bool RequiresEarlyViewUniformBuffer() const override;
-	virtual void PreRender(FRHICommandListImmediate& RHICmdList, FRHIUniformBuffer* ViewUniformBuffer, const class FGlobalDistanceFieldParameterData* GlobalDistanceFieldParameterData, bool bAllowGPUParticleSceneUpdate) override;
+	virtual void PreRender(class FRDGBuilder& GraphBuilder, FRHIUniformBuffer* ViewUniformBuffer, const class FGlobalDistanceFieldParameterData* GlobalDistanceFieldParameterData, bool bAllowGPUParticleSceneUpdate) override;
 	virtual void PostRenderOpaque(
-		FRHICommandListImmediate& RHICmdList, 
+		class FRDGBuilder& GraphBuilder,
 		FRHIUniformBuffer* ViewUniformBuffer,
 		const class FShaderParametersMetadata* SceneTexturesUniformBufferStruct,
 		FRHIUniformBuffer* SceneTexturesUniformBuffer,
diff --git a/Engine/Source/Runtime/Engine/Public/FXSystem.h b/Engine/Source/Runtime/Engine/Public/FXSystem.h
--- a/Engine/Source/Runtime/Engine/Public/FXSystem.h
+++ b/Engine/Source/Runtime/Engine/Public/FXSystem.h
@@ -232,9 +232,9 @@ public:
 	 * Notification from the renderer that it is about to perform visibility
 	 * checks on FX belonging to this system.
 	 */
-	virtual void PreInitViews(FRHICommandListImmediate& RHICmdList, bool bAllowGPUParticleUpdate) = 0;
+	virtual void PreInitViews(class FRDGBuilder& GraphBuilder, bool bAllowGPUParticleUpdate) = 0;
 
-	virtual void PostInitViews(FRHICommandListImmediate& RHICmdList, FRHIUniformBuffer* ViewUniformBuffer, bool bAllowGPUParticleUpdate) = 0;
+	virtual void PostInitViews(class FRDGBuilder& GraphBuilder, FRHIUniformBuffer* ViewUniformBuffer, bool bAllowGPUParticleUpdate) = 0;
 
 	virtual bool UsesGlobalDistanceField() const = 0;
 
@@ -246,13 +246,13 @@ public:
 	 * Notification from the renderer that it is about to draw FX belonging to
 	 * this system.
 	 */
-	virtual void PreRender(FRHICommandListImmediate& RHICmdList, FRHIUniformBuffer* ViewUniformBuffer, const class FGlobalDistanceFieldParameterData* GlobalDistanceFieldParameterData, bool bAllowGPUParticleSceneUpdate) = 0;
+	virtual void PreRender(class FRDGBuilder& GraphBuilder, FRHIUniformBuffer* ViewUniformBuffer, const class FGlobalDistanceFieldParameterData* GlobalDistanceFieldParameterData, bool bAllowGPUParticleSceneUpdate) = 0;
 
 	/**
 	 * Notification from the renderer that opaque primitives have rendered.
 	 */
 	virtual void PostRenderOpaque(
-		FRHICommandListImmediate& RHICmdList, 
+		class FRDGBuilder& GraphBuilder,
 		FRHIUniformBuffer* ViewUniformBuffer,
 		const class FShaderParametersMetadata* SceneTexturesUniformBufferStruct,
 		FRHIUniformBuffer* SceneTexturesUniformBuffer,
diff --git a/Engine/Source/Runtime/Engine/Public/GPUSortManager.h b/Engine/Source/Runtime/Engine/Public/GPUSortManager.h
--- a/Engine/Source/Runtime/Engine/Public/GPUSortManager.h
+++ b/Engine/Source/Runtime/Engine/Public/GPUSortManager.h
@@ -435,7 +435,7 @@ public:
 	 *
 	 * @param RHICmdList - The command list to be used.
 	 */
-	void OnPreRender(FRHICommandListImmediate& RHICmdList);
+	void OnPreRender(class FRDGBuilder& GraphBuilder);
 
 	/**
 	 * Callback that needs to be called in the rendering loop, after calls to FFXSystemInterface::PostRenderOpaque() are issued.
@@ -444,7 +444,7 @@ public:
 	 *
 	 * @param RHICmdList - The command list to be used.
 	 */
-	void OnPostRenderOpaque(FRHICommandListImmediate& RHICmdList);
+	void OnPostRenderOpaque(class FRDGBuilder& GraphBuilder);
 
 	/**
 	 * Event to register and receive post-prerender notification.
diff --git a/Engine/Source/Runtime/Renderer/Private/DeferredShadingRenderer.cpp b/Engine/Source/Runtime/Renderer/Private/DeferredShadingRenderer.cpp
--- a/Engine/Source/Runtime/Renderer/Private/DeferredShadingRenderer.cpp
+++ b/Engine/Source/Runtime/Renderer/Private/DeferredShadingRenderer.cpp
@@ -427,33 +427,20 @@ static void RenderOpaqueFX(
 		RDG_GPU_STAT_SCOPE(GraphBuilder, PostRenderOpsFX);
 		RDG_CSV_STAT_EXCLUSIVE_SCOPE(GraphBuilder, RenderOpaqueFX);
 
-		auto* PassParameters = GraphBuilder.AllocParameters<FRenderOpaqueFXPassParameters>();
-		PassParameters->SceneTextures = SceneTexturesUniformBuffer;
-
-		// Cascade uses pixel shaders for compute stuff in PostRenderOpaque so ERDGPassFlags::Raster is needed
-		GraphBuilder.AddPass(
-			RDG_EVENT_NAME("OpaqueFX"),
-			PassParameters,
-			ERDGPassFlags::Raster | ERDGPassFlags::SkipRenderPass | ERDGPassFlags::Compute | ERDGPassFlags::NeverCull,
-			[FXSystem, Views](FRHICommandListImmediate& RHICmdList)
-		{
-			SCOPE_CYCLE_COUNTER(STAT_FDeferredShadingSceneRenderer_FXSystem_PostRenderOpaque);
-
-			FXSystem->PostRenderOpaque(
-				RHICmdList,
-				Views[0].ViewUniformBuffer,
-				&FSceneTextureUniformParameters::StaticStructMetadata,
-				nullptr,
-				Views[0].AllowGPUParticleUpdate()
-			);
+		FXSystem->PostRenderOpaque(
+			GraphBuilder,
+			Views[0].ViewUniformBuffer,
+			&FSceneTextureUniformParameters::StaticStructMetadata,
+			nullptr,
+			Views[0].AllowGPUParticleUpdate()
+		);
 
-			if (FGPUSortManager* GPUSortManager = FXSystem->GetGPUSortManager())
-			{
-				GPUSortManager->OnPostRenderOpaque(RHICmdList);
-			}
+		if (FGPUSortManager* GPUSortManager = FXSystem->GetGPUSortManager())
+		{
+			GPUSortManager->OnPostRenderOpaque(GraphBuilder);
+		}
 
-			ServiceLocalQueue();
-		});
+		ServiceLocalQueue();
 	}
 }
 
@@ -1808,14 +1795,11 @@ void FDeferredShadingSceneRenderer::Render(FRDGBuilder& GraphBuilder)
 	{
 		SCOPE_CYCLE_COUNTER(STAT_FDeferredShadingSceneRenderer_FXSystem_PreRender);
 		GraphBuilder.SetCommandListStat(GET_STATID(STAT_CLM_FXPreRender));
-		AddPass(GraphBuilder, [this](FRHICommandListImmediate& InRHICmdList)
+		FXSystem->PreRender(GraphBuilder, Views[0].ViewUniformBuffer, &Views[0].GlobalDistanceFieldInfo.ParameterData, Views[0].AllowGPUParticleUpdate());
+		if (FGPUSortManager* GPUSortManager = FXSystem->GetGPUSortManager())
 		{
-			FXSystem->PreRender(InRHICmdList, Views[0].ViewUniformBuffer, &Views[0].GlobalDistanceFieldInfo.ParameterData, Views[0].AllowGPUParticleUpdate());
-			if (FGPUSortManager* GPUSortManager = FXSystem->GetGPUSortManager())
-			{
-				GPUSortManager->OnPreRender(InRHICmdList);
-			}
-		});
+			GPUSortManager->OnPreRender(GraphBuilder);
+		}
 	}
 
 	AddPass(GraphBuilder, [this](FRHICommandList& InRHICmdList)
diff --git a/Engine/Source/Runtime/Renderer/Private/MobileShadingRenderer.cpp b/Engine/Source/Runtime/Renderer/Private/MobileShadingRenderer.cpp
--- a/Engine/Source/Runtime/Renderer/Private/MobileShadingRenderer.cpp
+++ b/Engine/Source/Runtime/Renderer/Private/MobileShadingRenderer.cpp
@@ -667,14 +667,11 @@ void FMobileSceneRenderer::Render(FRDGBuilder& GraphBuilder)
 	// Notify the FX system that the scene is about to be rendered.
 	if (FXSystem && ViewFamily.EngineShowFlags.Particles)
 	{
-		AddPass(GraphBuilder, [this](FRHICommandListImmediate& InRHICmdList)
+		FXSystem->PreRender(GraphBuilder, Views[0].ViewUniformBuffer, NULL, !Views[0].bIsPlanarReflection);
+		if (FGPUSortManager* GPUSortManager = FXSystem->GetGPUSortManager())
 		{
-			FXSystem->PreRender(InRHICmdList, Views[0].ViewUniformBuffer, NULL, !Views[0].bIsPlanarReflection);
-			if (FGPUSortManager* GPUSortManager = FXSystem->GetGPUSortManager())
-			{
-				GPUSortManager->OnPreRender(InRHICmdList);
-			}
-		});
+			GPUSortManager->OnPreRender(GraphBuilder);
+		}
 	}
 
 	auto PollOcclusionQueriesAndDispatchToRHIThreadPass = [](FRHICommandListImmediate& InRHICmdList)
@@ -731,31 +728,25 @@ void FMobileSceneRenderer::Render(FRDGBuilder& GraphBuilder)
 
 	if (FXSystem && Views.IsValidIndex(0))
 	{
-		auto* PassParameters = GraphBuilder.AllocParameters<FMobileRenderOpaqueFXPassParameters>();
-		PassParameters->SceneTextures = SceneTextures.MobileUniformBuffer;
-
-		// Cascade uses pixel shaders for compute stuff in PostRenderOpaque so ERDGPassFlags::Raster is needed
-		GraphBuilder.AddPass(
-			RDG_EVENT_NAME("OpaqueFX"),
-			PassParameters,
-			ERDGPassFlags::Raster | ERDGPassFlags::SkipRenderPass | ERDGPassFlags::Compute | ERDGPassFlags::NeverCull,
-			[this](FRHICommandListImmediate& InRHICmdList)
+		FXSystem->PostRenderOpaque(
+			GraphBuilder,
+			Views[0].ViewUniformBuffer,
+			nullptr,
+			nullptr,
+			Views[0].AllowGPUParticleUpdate()
+		);
+		if (FGPUSortManager* GPUSortManager = FXSystem->GetGPUSortManager())
 		{
-			check(InRHICmdList.IsOutsideRenderPass());
-
-			FXSystem->PostRenderOpaque(
-				InRHICmdList,
-				Views[0].ViewUniformBuffer,
-				nullptr,
-				nullptr,
-				Views[0].AllowGPUParticleUpdate()
-			);
-			if (FGPUSortManager* GPUSortManager = FXSystem->GetGPUSortManager())
+			GPUSortManager->OnPostRenderOpaque(GraphBuilder);
+		}
+		AddPass(
+			GraphBuilder,
+			RDG_EVENT_NAME("DispatchRHICmdList"),
+			[](FRHICommandListImmediate& RHICmdList)
 			{
-				GPUSortManager->OnPostRenderOpaque(InRHICmdList);
+				RHICmdList.ImmediateFlush(EImmediateFlushType::DispatchToRHIThread);
 			}
-			InRHICmdList.ImmediateFlush(EImmediateFlushType::DispatchToRHIThread);
-		});
+		);
 	}
 
 	// Flush / submit cmdbuffer
diff --git a/Engine/Source/Runtime/Renderer/Private/SceneHitProxyRendering.cpp b/Engine/Source/Runtime/Renderer/Private/SceneHitProxyRendering.cpp
--- a/Engine/Source/Runtime/Renderer/Private/SceneHitProxyRendering.cpp
+++ b/Engine/Source/Runtime/Renderer/Private/SceneHitProxyRendering.cpp
@@ -650,22 +650,19 @@ void FDeferredShadingSceneRenderer::RenderHitProxies(FRDGBuilder& GraphBuilder)
 	// Notify the FX system that the scene is about to be rendered.
 	if (FXSystem && Views.IsValidIndex(0))
 	{
-		AddPass(GraphBuilder, [this](FRHICommandListImmediate& RHICmdList)
+		FGPUSortManager* GPUSortManager = FXSystem->GetGPUSortManager();
+		FXSystem->PreRender(GraphBuilder, Views[0].ViewUniformBuffer, &Views[0].GlobalDistanceFieldInfo.ParameterData, false);
+		if (GPUSortManager)
 		{
-			FGPUSortManager* GPUSortManager = FXSystem->GetGPUSortManager();
-			FXSystem->PreRender(RHICmdList, Views[0].ViewUniformBuffer, &Views[0].GlobalDistanceFieldInfo.ParameterData, false);
-			if (GPUSortManager)
-			{
-				GPUSortManager->OnPreRender(RHICmdList);
-			}
-			// Call PostRenderOpaque now as this is irrelevant for when rendering hit proxies.
-			// because we don't tick the particles in the render loop (see last param being "false").
-			FXSystem->PostRenderOpaque(RHICmdList, Views[0].ViewUniformBuffer, nullptr, nullptr, false);
-			if (GPUSortManager)
-			{
-				GPUSortManager->OnPostRenderOpaque(RHICmdList);
-			}
-		});
+			GPUSortManager->OnPreRender(GraphBuilder);
+		}
+		// Call PostRenderOpaque now as this is irrelevant for when rendering hit proxies.
+		// because we don't tick the particles in the render loop (see last param being "false").
+		FXSystem->PostRenderOpaque(GraphBuilder, Views[0].ViewUniformBuffer, nullptr, nullptr, false);
+		if (GPUSortManager)
+		{
+			GPUSortManager->OnPostRenderOpaque(GraphBuilder);
+		}
 	}
 
 	TArray<Nanite::FRasterResults, TInlineAllocator<2>> NaniteRasterResults;
diff --git a/Engine/Source/Runtime/Renderer/Private/SceneVisibility.cpp b/Engine/Source/Runtime/Renderer/Private/SceneVisibility.cpp
--- a/Engine/Source/Runtime/Renderer/Private/SceneVisibility.cpp
+++ b/Engine/Source/Runtime/Renderer/Private/SceneVisibility.cpp
@@ -3156,10 +3156,7 @@ void FSceneRenderer::PreVisibilityFrameSetup(FRDGBuilder& GraphBuilder, const FS
 
 	if (FXSystem && Views.IsValidIndex(0))
 	{
-		AddPass(GraphBuilder, [this](FRHICommandListImmediate& InRHICmdList)
-		{
-			FXSystem->PreInitViews(InRHICmdList, Views[0].AllowGPUParticleUpdate() && !ViewFamily.EngineShowFlags.HitProxies);
-		});
+		FXSystem->PreInitViews(GraphBuilder, Views[0].AllowGPUParticleUpdate() && !ViewFamily.EngineShowFlags.HitProxies);
 	}
 
 	// Draw lines to lights affecting this mesh if its selected.
@@ -4531,7 +4528,7 @@ bool FDeferredShadingSceneRenderer::InitViews(FRDGBuilder& GraphBuilder, const F
 		if (FXSystem && FXSystem->RequiresEarlyViewUniformBuffer() && Views.IsValidIndex(0))
 		{
 			Views[0].InitRHIResources();
-			FXSystem->PostInitViews(RHICmdList, Views[0].ViewUniformBuffer, Views[0].AllowGPUParticleUpdate() && !ViewFamily.EngineShowFlags.HitProxies);
+			FXSystem->PostInitViews(GraphBuilder, Views[0].ViewUniformBuffer, Views[0].AllowGPUParticleUpdate() && !ViewFamily.EngineShowFlags.HitProxies);
 		}
 	}
 
