Hide the TC_ReflectionCapture setting

diff --git a/Engine/Source/Runtime/Engine/Classes/Engine/Texture.h b/Engine/Source/Runtime/Engine/Classes/Engine/Texture.h
--- a/Engine/Source/Runtime/Engine/Classes/Engine/Texture.h
+++ b/Engine/Source/Runtime/Engine/Classes/Engine/Texture.h
@@ -47,7 +47,7 @@ enum TextureCompressionSettings
 	TC_HDR_Compressed			UMETA(DisplayName="HDRCompressed (RGB, BC6H, DX11)"),
 	TC_BC7						UMETA(DisplayName="BC7 (DX11, optional A)"),
 	TC_HalfFloat				UMETA(DisplayName="Half Float (R16F)"),
-	TC_ReflectionCapture		UMETA(DisplayName="Default (DXT5)"),
+	TC_ReflectionCapture		UMETA(Hidden),
 	TC_MAX,
 };
 
