Move validation code to slow checks in FGrowOnlySpanAllocator::Free to avoid pathological case when loading large maps.

diff --git a/Engine/Source/Runtime/Engine/Private/GrowOnlySpanAllocator.cpp b/Engine/Source/Runtime/Engine/Private/GrowOnlySpanAllocator.cpp
--- a/Engine/Source/Runtime/Engine/Private/GrowOnlySpanAllocator.cpp
+++ b/Engine/Source/Runtime/Engine/Private/GrowOnlySpanAllocator.cpp
@@ -41,12 +41,12 @@ void FGrowOnlySpanAllocator::Free(int32 BaseOffset, int32 Num)
 
 	FLinearAllocation NewFreeSpan(BaseOffset, Num);
 
-#if DO_CHECK
+#if DO_GUARD_SLOW
 	// Detect double delete
 	for (int32 i = 0; i < FreeSpans.Num(); i++)
 	{
 		FLinearAllocation CurrentSpan = FreeSpans[i];
-		check(!(CurrentSpan.Contains(NewFreeSpan)));
+		checkSlow(!(CurrentSpan.Contains(NewFreeSpan)));
 	}
 #endif
 
