Added bRerunConstructionScripts property to PackedLevelInstance actors to control how construction scripts are run for that class.
This greatly reduce the time spent in the "making visible" state for level streaming.

diff --git a/Engine/Source/Runtime/Engine/Private/LevelInstance/Packed/PackedLevelInstanceActor.cpp b/Engine/Source/Runtime/Engine/Private/LevelInstance/Packed/PackedLevelInstanceActor.cpp
--- a/Engine/Source/Runtime/Engine/Private/LevelInstance/Packed/PackedLevelInstanceActor.cpp
+++ b/Engine/Source/Runtime/Engine/Private/LevelInstance/Packed/PackedLevelInstanceActor.cpp
@@ -9,6 +9,7 @@ APackedLevelInstance::APackedLevelInstance(const FObjectInitializer& ObjectIniti
 	: Super(ObjectInitializer)
 #if WITH_EDITORONLY_DATA
 	, ISMComponentClass(UInstancedStaticMeshComponent::StaticClass())
+	, bRerunConstructionScripts(false)
 	, bEditing(false)
 	, ChildEditing(0)
 	, bChildChanged(false)
@@ -171,4 +172,12 @@ void APackedLevelInstance::DestroyPackedComponents()
 		PackedComponent->DestroyComponent();
 	}
 }
+
+void APackedLevelInstance::RerunConstructionScripts()
+{
+	if (bRerunConstructionScripts)
+	{
+		Super::RerunConstructionScripts();
+	}
+}
 #endif
\ No newline at end of file
diff --git a/Engine/Source/Runtime/Engine/Public/LevelInstance/Packed/PackedLevelInstanceActor.h b/Engine/Source/Runtime/Engine/Public/LevelInstance/Packed/PackedLevelInstanceActor.h
--- a/Engine/Source/Runtime/Engine/Public/LevelInstance/Packed/PackedLevelInstanceActor.h
+++ b/Engine/Source/Runtime/Engine/Public/LevelInstance/Packed/PackedLevelInstanceActor.h
@@ -61,6 +61,8 @@ public:
 		NewComponent->ComponentTags.Add(GetPackedComponentTag());
 		return NewComponent;
 	}
+
+	virtual void RerunConstructionScripts() override;
 #endif
 
 #if WITH_EDITORONLY_DATA
@@ -70,6 +72,9 @@ public:
 	UPROPERTY(VisibleAnywhere, Category = Pack)
 	TSoftObjectPtr<UBlueprint> BlueprintAsset;
 
+	UPROPERTY(VisibleAnywhere, Category = Pack)
+	bool bRerunConstructionScripts;
+
 	UPROPERTY()
 	TArray<TSoftObjectPtr<UWorld>> PackDependencies;
 
