WorldPartition: Query runtime streaming state support through the WorldPartitionSubsystem + expose Loaded target state to DataLayers + StreamingSource

- This will allow prefetching of streaming cells without activating them (making visible)
- Ability to query the state of WorldPartition streaming at a specific location/radius/for specific data layers/etc

diff --git a/Engine/Source/Editor/DataLayerEditor/Private/DataLayer/SDataLayerTreeLabel.cpp b/Engine/Source/Editor/DataLayerEditor/Private/DataLayer/SDataLayerTreeLabel.cpp
--- a/Engine/Source/Editor/DataLayerEditor/Private/DataLayer/SDataLayerTreeLabel.cpp
+++ b/Engine/Source/Editor/DataLayerEditor/Private/DataLayer/SDataLayerTreeLabel.cpp
@@ -90,14 +90,15 @@ FText SDataLayerTreeLabel::GetDisplayText() const
 {
 	const UDataLayer* DataLayer = DataLayerPtr.Get();
 	bool bIsDataLayerActive = false;
+	EDataLayerState DataLayerState = EDataLayerState::Unloaded;
 	if (DataLayer && DataLayer->IsDynamicallyLoaded() && DataLayer->GetWorld()->IsPlayInEditor())
 	{
 		const UDataLayerSubsystem* DataLayerSubsystem = DataLayer->GetWorld()->GetSubsystem<UDataLayerSubsystem>();
-		bIsDataLayerActive = DataLayerSubsystem->IsDataLayerActive(DataLayer);
+		DataLayerState = DataLayerSubsystem->GetDataLayerState(DataLayer);
 	}
-	static const FText DataLayerActive = LOCTEXT("DataLayerActive", "(Active)");
+	const FText DataLayerStateText = FText::Format(LOCTEXT("DataLayerActive", "({0})"), StaticEnum<EDataLayerState>()->GetDisplayNameTextByValue((int64)DataLayerState));
 	static const FText DataLayerDeleted = LOCTEXT("DataLayerLabelForMissingDataLayer", "(Deleted Data Layer)");
-	return DataLayer ? FText::Format(LOCTEXT("DataLayerDisplayText", "{0} {1}"), FText::FromName(DataLayer->GetDataLayerLabel()), bIsDataLayerActive ? DataLayerActive : FText::GetEmpty()) : DataLayerDeleted;
+	return DataLayer ? FText::Format(LOCTEXT("DataLayerDisplayText", "{0} {1}"), FText::FromName(DataLayer->GetDataLayerLabel()), DataLayerStateText) : DataLayerDeleted;
 }
 
 FText SDataLayerTreeLabel::GetTooltipText() const
@@ -163,10 +164,15 @@ FSlateColor SDataLayerTreeLabel::GetForegroundColor() const
 		else if (DataLayer->IsDynamicallyLoaded() && DataLayer->GetWorld()->IsPlayInEditor())
 		{
 			const UDataLayerSubsystem* DataLayerSubsystem = DataLayer->GetWorld()->GetSubsystem<UDataLayerSubsystem>();
-			if (DataLayerSubsystem->IsDataLayerActive(DataLayer))
+			EDataLayerState DataLayerState = DataLayerSubsystem->GetDataLayerState(DataLayer);
+			if (DataLayerState == EDataLayerState::Activated)
 			{
 				return FColorList::LimeGreen;
 			}
+			else if (DataLayerState == EDataLayerState::Loaded)
+			{
+				return FColorList::NeonBlue;
+			}
 		}
 	}
 	return (!DataLayer || !DataLayer->GetWorld()) ? FLinearColor(0.2f, 0.2f, 0.25f) : FSlateColor::UseForeground();
diff --git a/Engine/Source/Runtime/Engine/Classes/Components/WorldPartitionStreamingSourceComponent.h b/Engine/Source/Runtime/Engine/Classes/Components/WorldPartitionStreamingSourceComponent.h
--- a/Engine/Source/Runtime/Engine/Classes/Components/WorldPartitionStreamingSourceComponent.h
+++ b/Engine/Source/Runtime/Engine/Classes/Components/WorldPartitionStreamingSourceComponent.h
@@ -35,4 +35,7 @@ private:
 	/** Whether this component is enabled or not */
 	UPROPERTY(EditAnywhere, Category = "Streaming")
 	bool bStreamingSourceEnabled;
+
+	UPROPERTY(EditAnywhere, Category = "Streaming")
+	EStreamingSourceTargetState TargetState;
 };
\ No newline at end of file
diff --git a/Engine/Source/Runtime/Engine/Private/Components/WorldPartitionStreamingSourceComponent.cpp b/Engine/Source/Runtime/Engine/Private/Components/WorldPartitionStreamingSourceComponent.cpp
--- a/Engine/Source/Runtime/Engine/Private/Components/WorldPartitionStreamingSourceComponent.cpp
+++ b/Engine/Source/Runtime/Engine/Private/Components/WorldPartitionStreamingSourceComponent.cpp
@@ -2,11 +2,13 @@
 
 #include "Components/WorldPartitionStreamingSourceComponent.h"
 #include "WorldPartition/WorldPartition.h"
+#include "WorldPartition/WorldPartitionRuntimeCell.h"
 #include "Engine/World.h"
 
 UWorldPartitionStreamingSourceComponent::UWorldPartitionStreamingSourceComponent(const FObjectInitializer& ObjectInitializer)
 	: Super(ObjectInitializer)
 	, bStreamingSourceEnabled(true)
+	, TargetState(EStreamingSourceTargetState::Activated)
 {
 	PrimaryComponentTick.bCanEverTick = false;
 }
@@ -67,6 +69,7 @@ bool UWorldPartitionStreamingSourceComponent::GetStreamingSource(FWorldPartition
 		OutStreamingSource.Name = SourceName;
 		OutStreamingSource.Location = Actor->GetActorLocation();
 		OutStreamingSource.Rotation = Actor->GetActorRotation();
+		OutStreamingSource.TargetState = TargetState;
 		return true;
 	}
 	return false;
diff --git a/Engine/Source/Runtime/Engine/Private/WorldPartition/DataLayer/DataLayer.cpp b/Engine/Source/Runtime/Engine/Private/WorldPartition/DataLayer/DataLayer.cpp
--- a/Engine/Source/Runtime/Engine/Private/WorldPartition/DataLayer/DataLayer.cpp
+++ b/Engine/Source/Runtime/Engine/Private/WorldPartition/DataLayer/DataLayer.cpp
@@ -12,9 +12,10 @@ UDataLayer::UDataLayer(const FObjectInitializer& ObjectInitializer)
 : Super(ObjectInitializer)
 #if WITH_EDITORONLY_DATA
 , DataLayerLabel(GetFName())
+, InitialState(EDataLayerState::Unloaded)
 , bIsVisible(true)
 , bIsDynamicallyLoaded(false)
-, bIsInitiallyActive(false)
+, bIsInitiallyActive_DEPRECATED(false)
 , bIsDynamicallyLoadedInEditor(true)
 , bGeneratesHLODs(false)
 , DefaultHLODLayer()
@@ -22,6 +23,18 @@ UDataLayer::UDataLayer(const FObjectInitializer& ObjectInitializer)
 {
 }
 
+void UDataLayer::PostLoad()
+{
+	Super::PostLoad();
+
+#if WITH_EDITORONLY_DATA
+	if (bIsInitiallyActive_DEPRECATED)
+	{
+		InitialState = EDataLayerState::Activated;
+	}
+#endif
+}
+
 #if WITH_EDITOR
 void UDataLayer::SetDataLayerLabel(FName InDataLayerLabel)
 {
diff --git a/Engine/Source/Runtime/Engine/Private/WorldPartition/DataLayer/DataLayerSubsystem.cpp b/Engine/Source/Runtime/Engine/Private/WorldPartition/DataLayer/DataLayerSubsystem.cpp
--- a/Engine/Source/Runtime/Engine/Private/WorldPartition/DataLayer/DataLayerSubsystem.cpp
+++ b/Engine/Source/Runtime/Engine/Private/WorldPartition/DataLayer/DataLayerSubsystem.cpp
@@ -82,12 +82,31 @@ UDataLayer* UDataLayerSubsystem::GetDataLayerFromName(const FName& InDataLayerNa
 	return WorldDataLayers ? const_cast<UDataLayer*>(WorldDataLayers->GetDataLayerFromName(InDataLayerName)) : nullptr;
 }
 
+// Deprecated start
 void UDataLayerSubsystem::ActivateDataLayer(const FActorDataLayer& InDataLayer, bool bInActivate)
 {
-	ActivateDataLayerByName(InDataLayer.Name, bInActivate);
+	// Keep deprecated behavior (!Activated == Unloaded)
+	SetDataLayerState(GetDataLayerFromName(InDataLayer.Name), bInActivate ? EDataLayerState::Activated : EDataLayerState::Unloaded);
 }
 
-void UDataLayerSubsystem::ActivateDataLayer(const UDataLayer* InDataLayer, bool bInActivate)
+void UDataLayerSubsystem::ActivateDataLayerByLabel(const FName& InDataLayerLabel, bool bInActivate)
+{
+	// Keep deprecated behavior (!Activated == Unloaded)
+	SetDataLayerState(GetDataLayerFromLabel(InDataLayerLabel), bInActivate ? EDataLayerState::Activated : EDataLayerState::Unloaded);
+}
+
+bool UDataLayerSubsystem::IsDataLayerActive(const FActorDataLayer& InDataLayer) const
+{
+	return GetDataLayerStateByName(InDataLayer.Name) == EDataLayerState::Activated;
+}
+
+bool UDataLayerSubsystem::IsDataLayerActiveByLabel(const FName& InDataLayerLabel) const
+{
+	return GetDataLayerStateByLabel(InDataLayerLabel) == EDataLayerState::Activated;
+}
+// Deprecation end
+
+void UDataLayerSubsystem::SetDataLayerState(const UDataLayer* InDataLayer, EDataLayerState InState)
 {
 	if (!InDataLayer || !InDataLayer->IsDynamicallyLoaded())
 	{
@@ -95,59 +114,90 @@ void UDataLayerSubsystem::ActivateDataLayer(const UDataLayer* InDataLayer, bool
 	}
 
 	FName DataLayerName = InDataLayer->GetFName();
-	const bool bIsLayerActive = IsDataLayerActiveByName(DataLayerName);
-	if (bIsLayerActive != bInActivate)
+	EDataLayerState CurrentState = GetDataLayerStateByName(DataLayerName);
+	if (CurrentState != InState)
 	{
-		if (bInActivate)
+		LoadedDataLayerNames.Remove(DataLayerName);
+		ActiveDataLayerNames.Remove(DataLayerName);
+		
+		if (InState == EDataLayerState::Loaded)
 		{
-			ActiveDataLayerNames.Add(DataLayerName);
+			LoadedDataLayerNames.Add(DataLayerName);
 		}
-		else
+		else if(InState == EDataLayerState::Activated)
 		{
-			ActiveDataLayerNames.Remove(DataLayerName);
+			ActiveDataLayerNames.Add(DataLayerName);
 		}
+		
+		// todo_ow: remove
 		if (OnDataLayerActivationStateChanged.IsBound())
 		{
-			OnDataLayerActivationStateChanged.Broadcast(InDataLayer, bInActivate);
+			OnDataLayerActivationStateChanged.Broadcast(InDataLayer, InState == EDataLayerState::Activated);
+		}
+
+		if (OnDataLayerStateChanged.IsBound())
+		{
+			OnDataLayerStateChanged.Broadcast(InDataLayer, InState);
 		}
 	}
 }
 
-void UDataLayerSubsystem::ActivateDataLayerByName(const FName& InDataLayerName, bool bInActivate)
+void UDataLayerSubsystem::SetDataLayerStateByName(const FName& InDataLayerName, EDataLayerState InState)
 {
-	ActivateDataLayer(GetDataLayerFromName(InDataLayerName), bInActivate);
+	SetDataLayerState(GetDataLayerFromName(InDataLayerName), InState);
 }
 
-void UDataLayerSubsystem::ActivateDataLayerByLabel(const FName& InDataLayerLabel, bool bInActivate)
+void UDataLayerSubsystem::SetDataLayerState(const FActorDataLayer& InDataLayer, EDataLayerState InState)
 {
-	ActivateDataLayer(GetDataLayerFromLabel(InDataLayerLabel), bInActivate);
+	SetDataLayerState(GetDataLayerFromName(InDataLayer.Name), InState);
 }
 
-bool UDataLayerSubsystem::IsDataLayerActive(const FActorDataLayer& InDataLayer) const
+void UDataLayerSubsystem::SetDataLayerStateByLabel(const FName& InDataLayerLabel, EDataLayerState InState)
 {
-	return IsDataLayerActiveByName(InDataLayer.Name);
+	SetDataLayerState(GetDataLayerFromLabel(InDataLayerLabel), InState);
 }
 
-bool UDataLayerSubsystem::IsDataLayerActive(const UDataLayer* InDataLayer) const
+EDataLayerState UDataLayerSubsystem::GetDataLayerState(const UDataLayer* InDataLayer) const
 {
-	return InDataLayer && IsDataLayerActiveByName(InDataLayer->GetFName());
+	if (!InDataLayer)
+	{
+		return EDataLayerState::Unloaded;
+	}
+
+	return GetDataLayerStateByName(InDataLayer->GetFName());
 }
 
-bool UDataLayerSubsystem::IsDataLayerActiveByName(const FName& InDataLayerName) const
+EDataLayerState UDataLayerSubsystem::GetDataLayerStateByName(const FName& InDataLayerName) const
 {
-	return ActiveDataLayerNames.Contains(InDataLayerName);
+	if (ActiveDataLayerNames.Contains(InDataLayerName))
+	{
+		check(!LoadedDataLayerNames.Contains(InDataLayerName));
+		return EDataLayerState::Activated;
+	}
+	else if (LoadedDataLayerNames.Contains(InDataLayerName))
+	{
+		check(!ActiveDataLayerNames.Contains(InDataLayerName));
+		return EDataLayerState::Loaded;
+	}
+
+	return EDataLayerState::Unloaded;
 }
 
-bool UDataLayerSubsystem::IsDataLayerActiveByLabel(const FName& InDataLayerLabel) const
+EDataLayerState UDataLayerSubsystem::GetDataLayerState(const FActorDataLayer& InDataLayer) const
 {
-	return IsDataLayerActive(GetDataLayerFromLabel(InDataLayerLabel));
+	return GetDataLayerStateByName(InDataLayer.Name);
 }
 
-bool UDataLayerSubsystem::IsAnyDataLayerActive(const TArray<FName>& InDataLayerNames) const
+EDataLayerState UDataLayerSubsystem::GetDataLayerStateByLabel(const FName& InDataLayerLabel) const
+{
+	return GetDataLayerState(GetDataLayerFromLabel(InDataLayerLabel));
+}
+
+bool UDataLayerSubsystem::IsAnyDataLayerInState(const TArray<FName>& InDataLayerNames, EDataLayerState InState) const
 {
 	for (FName DataLayerName : InDataLayerNames)
 	{
-		if (IsDataLayerActiveByName(DataLayerName))
+		if (GetDataLayerStateByName(DataLayerName) == InState)
 		{
 			return true;
 		}
@@ -171,22 +221,28 @@ void UDataLayerSubsystem::Draw(UCanvas* Canvas, class APlayerController* PC)
 		Pos.Y += TextHeight + 1;
 	};
 
-	if (ActiveDataLayerNames.Num() > 0)
+	const FVector2D CanvasTopLeftPadding(10.f, 10.f);
+	FVector2D Pos = CanvasTopLeftPadding;
+
+	auto DrawLayerNames = [this, &Pos, DrawText, Canvas](const FString& Title, FColor Color, const TSet<FName>& LayerNames)
 	{
-		const FVector2D CanvasTopLeftPadding(10.f, 10.f);
-		FVector2D Pos = CanvasTopLeftPadding;
-		FString Title(TEXT("Active Data Layers"));
-		DrawText(Canvas, Title, GEngine->GetSmallFont(), FColor::Yellow, Pos);
-		UFont* DataLayerFont = GEngine->GetTinyFont();
-		for (const FName& DataLayerName : ActiveDataLayerNames)
+		if (LayerNames.Num() > 0)
 		{
-			if (UDataLayer* DataLayer = GetDataLayerFromName(DataLayerName))
+			DrawText(Canvas, Title, GEngine->GetSmallFont(), FColor::Yellow, Pos);
+			UFont* DataLayerFont = GEngine->GetTinyFont();
+			for (const FName& DataLayerName : LayerNames)
 			{
-				const FString Text = DataLayer->GetDataLayerLabel().ToString();
-				DrawText(Canvas, Text, DataLayerFont, FColor::White, Pos);
+				if (UDataLayer* DataLayer = GetDataLayerFromName(DataLayerName))
+				{
+					const FString Text = DataLayer->GetDataLayerLabel().ToString();
+					DrawText(Canvas, Text, DataLayerFont, FColor::White, Pos);
+				}
 			}
 		}
-	}
+	};
+
+	DrawLayerNames(TEXT("Loaded Data Layers"), FColor::Green, LoadedDataLayerNames);
+	DrawLayerNames(TEXT("Active Data Layers"), FColor::White, ActiveDataLayerNames);
 }
 
 FAutoConsoleCommand UDataLayerSubsystem::ToggleDataLayerActivation(
@@ -205,7 +261,7 @@ FAutoConsoleCommand UDataLayerSubsystem::ToggleDataLayerActivation(
 					UDataLayerSubsystem* DataLayerSubsystem = World->GetSubsystem<UDataLayerSubsystem>();
 					if (const UDataLayer* DataLayer = DataLayerSubsystem ? DataLayerSubsystem->GetDataLayerFromLabel(DataLayerLabel) : nullptr)
 					{
-						DataLayerSubsystem->ActivateDataLayer(DataLayer, !DataLayerSubsystem->IsDataLayerActive(DataLayer));
+						DataLayerSubsystem->SetDataLayerState(DataLayer, DataLayerSubsystem->GetDataLayerState(DataLayer) == EDataLayerState::Activated ? EDataLayerState::Unloaded : EDataLayerState::Activated);
 					}
 				}
 			}
diff --git a/Engine/Source/Runtime/Engine/Private/WorldPartition/DataLayer/WorldDataLayers.cpp b/Engine/Source/Runtime/Engine/Private/WorldPartition/DataLayer/WorldDataLayers.cpp
--- a/Engine/Source/Runtime/Engine/Private/WorldPartition/DataLayer/WorldDataLayers.cpp
+++ b/Engine/Source/Runtime/Engine/Private/WorldPartition/DataLayer/WorldDataLayers.cpp
@@ -259,9 +259,9 @@ void AWorldDataLayers::BeginPlay()
 	{
 		for (const UDataLayer* DataLayer : WorldDataLayers)
 		{
-			if (DataLayer && DataLayer->IsInitiallyActive())
+			if (DataLayer && DataLayer->IsDynamicallyLoaded())
 			{
-				DataLayerSubsystem->ActivateDataLayer(DataLayer, true);
+				DataLayerSubsystem->SetDataLayerState(DataLayer, DataLayer->GetInitialState());
 			}
 		}
 	}
diff --git a/Engine/Source/Runtime/Engine/Private/WorldPartition/WorldPartition.cpp b/Engine/Source/Runtime/Engine/Private/WorldPartition/WorldPartition.cpp
--- a/Engine/Source/Runtime/Engine/Private/WorldPartition/WorldPartition.cpp
+++ b/Engine/Source/Runtime/Engine/Private/WorldPartition/WorldPartition.cpp
@@ -946,6 +946,16 @@ class ULevel* UWorldPartition::GetPreferredLoadedLevelToAddToWorld() const
 	return nullptr;
 }
 
+bool UWorldPartition::IsStreamingCompleted(EWorldPartitionRuntimeCellState QueryState, const TArray<FWorldPartitionStreamingQuerySource>& QuerySources, bool bExactState) const
+{
+	if (GetWorld()->IsGameWorld())
+	{
+		return GetStreamingPolicy()->IsStreamingCompleted(QueryState, QuerySources, bExactState);
+	}
+
+	return false;
+}
+
 bool UWorldPartition::CanDrawRuntimeHash() const
 {
 	return GetWorld()->IsGameWorld() || UWorldPartition::IsSimulating();
diff --git a/Engine/Source/Runtime/Engine/Private/WorldPartition/WorldPartitionLevelStreamingDynamic.cpp b/Engine/Source/Runtime/Engine/Private/WorldPartition/WorldPartitionLevelStreamingDynamic.cpp
--- a/Engine/Source/Runtime/Engine/Private/WorldPartition/WorldPartitionLevelStreamingDynamic.cpp
+++ b/Engine/Source/Runtime/Engine/Private/WorldPartition/WorldPartitionLevelStreamingDynamic.cpp
@@ -34,7 +34,6 @@ UWorldPartitionLevelStreamingDynamic::UWorldPartitionLevelStreamingDynamic(const
 	, RuntimeLevel(nullptr)
 	, bLoadRequestInProgress(false)
 #endif
-	, bIsActivated(false)
 	, bShouldBeAlwaysLoaded(false)
 {
 }
@@ -47,7 +46,6 @@ UWorldPartitionLevelStreamingDynamic::UWorldPartitionLevelStreamingDynamic(const
 void UWorldPartitionLevelStreamingDynamic::Initialize(const UWorldPartitionRuntimeLevelStreamingCell& InCell)
 {
 	UWorld* World = GetWorld();
-	check(!bIsActivated);
 	check(!ShouldBeLoaded());
 	check((World->IsGameWorld() && !ShouldBeVisible()) || (!World->IsGameWorld() && !GetShouldBeVisibleFlag()));
 	check(ChildPackages.Num() == 0);
@@ -217,7 +215,6 @@ bool UWorldPartitionLevelStreamingDynamic::RequestLevel(UWorld* InPersistentWorl
 bool UWorldPartitionLevelStreamingDynamic::IssueLoadRequests()
 {
 	check(ShouldBeLoaded());
-	check(ShouldBeVisible());
 	check(!HasLoadedLevel());
 	check(RuntimeLevel);
 	check(!bLoadRequestInProgress);
@@ -448,6 +445,38 @@ void UWorldPartitionLevelStreamingDynamic::OnCleanupLevel()
 
 #endif
 
+/*
+ * Load StreamingLevel without adding it to world 
+ */
+void UWorldPartitionLevelStreamingDynamic::Load()
+{
+	UE_LOG(LogLevelStreaming, Verbose, TEXT("UWorldPartitionLevelStreamingDynamic::Loading %s"), *GetWorldAssetPackageName());
+
+	check(!ShouldBeLoaded());
+	
+	SetShouldBeLoaded(true);
+	SetShouldBeVisible(false);
+	SetIsRequestingUnloadAndRemoval(false);
+
+	UWorld* PlayWorld = GetWorld();
+	check(PlayWorld && PlayWorld->IsGameWorld());
+	PlayWorld->AddUniqueStreamingLevel(this);
+}
+
+/*
+ * Unload StreamingLevel
+ */
+void UWorldPartitionLevelStreamingDynamic::Unload()
+{
+	UE_LOG(LogLevelStreaming, Verbose, TEXT("UWorldPartitionLevelStreamingDynamic::Unloading %s"), *GetWorldAssetPackageName());
+
+	check(ShouldBeLoaded());
+
+	SetShouldBeLoaded(false);
+	SetShouldBeVisible(false);
+	SetIsRequestingUnloadAndRemoval(true);
+}
+
 /**
   * Activates StreamingLevel by making sure it's in the World's StreamingLevels and that it should be loaded & visible
   */
@@ -455,8 +484,6 @@ void UWorldPartitionLevelStreamingDynamic::Activate()
 {
 	UE_LOG(LogLevelStreaming, Verbose, TEXT("UWorldPartitionLevelStreamingDynamic::Activating %s"), *GetWorldAssetPackageName());
 
-	check(!bIsActivated);
-	check(!ShouldBeLoaded());
 	check(!ShouldBeVisible());
 
 	// Make sure we are in the correct state
@@ -468,26 +495,19 @@ void UWorldPartitionLevelStreamingDynamic::Activate()
 	UWorld* PlayWorld = GetWorld();
 	check(PlayWorld && PlayWorld->IsGameWorld());
 	PlayWorld->AddUniqueStreamingLevel(this);
-
-	bIsActivated = true;
 }
 
 /**
- * Deactivates StreamingLevel
+ * Deactivates StreamingLevel (Remove from world, keep loaded)
  */
 void UWorldPartitionLevelStreamingDynamic::Deactivate()
 {
 	UE_LOG(LogLevelStreaming, Verbose, TEXT("UWorldPartitionLevelStreamingDynamic::Deactivating %s"), *GetWorldAssetPackageName());
 
-	check(bIsActivated);
 	check(ShouldBeLoaded());
 	check(ShouldBeVisible());
 
-	SetShouldBeLoaded(false);
 	SetShouldBeVisible(false);
-	SetIsRequestingUnloadAndRemoval(true);
-
-	bIsActivated = false;
 }
 
 UWorld* UWorldPartitionLevelStreamingDynamic::GetOuterWorld() const
diff --git a/Engine/Source/Runtime/Engine/Private/WorldPartition/WorldPartitionLevelStreamingPolicy.cpp b/Engine/Source/Runtime/Engine/Private/WorldPartition/WorldPartitionLevelStreamingPolicy.cpp
--- a/Engine/Source/Runtime/Engine/Private/WorldPartition/WorldPartitionLevelStreamingPolicy.cpp
+++ b/Engine/Source/Runtime/Engine/Private/WorldPartition/WorldPartitionLevelStreamingPolicy.cpp
@@ -27,9 +27,9 @@ static FAutoConsoleVariableRef CMaxLoadingLevelStreamingCells(
 int32 UWorldPartitionLevelStreamingPolicy::GetCellLoadingCount() const
 {
 	int32 CellLoadingCount = 0;
-	for (const UWorldPartitionRuntimeCell* LoadedCell : LoadedCells)
+	for (const UWorldPartitionRuntimeCell* ActivatedCell : ActivatedCells)
 	{
-		const UWorldPartitionRuntimeLevelStreamingCell* Cell = Cast<const UWorldPartitionRuntimeLevelStreamingCell>(LoadedCell);
+		const UWorldPartitionRuntimeLevelStreamingCell* Cell = Cast<const UWorldPartitionRuntimeLevelStreamingCell>(ActivatedCell);
 		if (!Cell->IsAlwaysLoaded())
 		{
 			if (UWorldPartitionLevelStreamingDynamic* LevelStreaming = Cell->GetLevelStreaming())
@@ -45,7 +45,58 @@ int32 UWorldPartitionLevelStreamingPolicy::GetCellLoadingCount() const
 	return CellLoadingCount;
 }
 
-void UWorldPartitionLevelStreamingPolicy::LoadCells(const TSet<const UWorldPartitionRuntimeCell*>& ToLoadCells)
+void UWorldPartitionLevelStreamingPolicy::SetTargetStateForCells(EWorldPartitionRuntimeCellState TargetState, const TSet<const UWorldPartitionRuntimeCell*>& Cells)
+{
+	switch (TargetState)
+	{
+	case EWorldPartitionRuntimeCellState::Unloaded:
+		SetCellsStateToUnloaded(Cells);
+		break;
+	case EWorldPartitionRuntimeCellState::Loaded:
+		SetCellsStateToLoaded(Cells);
+		break;
+	case EWorldPartitionRuntimeCellState::Activated:
+		SetCellsStateToActivated(Cells);
+		break;
+	default:
+		check(false);
+	}
+}
+
+void UWorldPartitionLevelStreamingPolicy::SetCellsStateToLoaded(const TSet<const UWorldPartitionRuntimeCell*>& ToLoadCells)
+{
+	// This policy limits the number of concurrent loading streaming cells
+	int32 CellLoadingCount = GetCellLoadingCount();
+	int32 MaxCellsToLoad = GMaxLoadingLevelStreamingCells - CellLoadingCount;
+	
+	// Sort cells based on importance
+	TArray<const UWorldPartitionRuntimeCell*, TInlineAllocator<256>> SortedCells;
+	WorldPartition->RuntimeHash->SortStreamingCellsByImportance(ToLoadCells, StreamingSources, SortedCells);
+
+	// Trigger cell loading. Depending on actual state of cell limit loading.
+	for (const UWorldPartitionRuntimeCell* Cell : SortedCells)
+	{
+		UE_LOG(LogWorldPartition, Verbose, TEXT("UWorldPartitionLevelStreamingPolicy::LoadCells %s"), *Cell->GetName());
+		const UWorldPartitionRuntimeLevelStreamingCell* LevelCell = CastChecked<const UWorldPartitionRuntimeLevelStreamingCell>(Cell);
+		if (ActivatedCells.Contains(Cell))
+		{
+			LevelCell->Deactivate();
+			ActivatedCells.Remove(Cell);
+			LoadedCells.Add(Cell);
+		}
+		else if (MaxCellsToLoad > 0)
+		{
+			LevelCell->Load();
+			LoadedCells.Add(Cell);
+			if (!Cell->IsAlwaysLoaded())
+			{
+				--MaxCellsToLoad;
+			}
+		}
+	}
+}
+
+void UWorldPartitionLevelStreamingPolicy::SetCellsStateToActivated(const TSet<const UWorldPartitionRuntimeCell*>& ToActivateCells)
 {
 	// This policy limits the number of concurrent loading streaming cells
 	int32 CellLoadingCount = GetCellLoadingCount();
@@ -57,20 +108,40 @@ void UWorldPartitionLevelStreamingPolicy::LoadCells(const TSet<const UWorldParti
 
 	// Sort cells based on importance
 	TArray<const UWorldPartitionRuntimeCell*, TInlineAllocator<256>> SortedCells;
-	WorldPartition->RuntimeHash->SortStreamingCellsByImportance(ToLoadCells, StreamingSources, SortedCells);
+	WorldPartition->RuntimeHash->SortStreamingCellsByImportance(ToActivateCells, StreamingSources, SortedCells);
 
-	// Trigger cell loading until we hit the maximum
+	// Trigger cell activation. Depending on actual state of cell limit loading.
 	for (const UWorldPartitionRuntimeCell* Cell : SortedCells)
 	{
-		LoadCell(Cell);
-		LoadedCells.Add(Cell);
-		if (!Cell->IsAlwaysLoaded())
+		UE_LOG(LogWorldPartition, Verbose, TEXT("UWorldPartitionLevelStreamingPolicy::ActivateCells %s"), *Cell->GetName());
+		const UWorldPartitionRuntimeLevelStreamingCell* LevelCell = CastChecked<const UWorldPartitionRuntimeLevelStreamingCell>(Cell);
+		if (LoadedCells.Contains(LevelCell))
 		{
-			if (--MaxCellsToLoad <= 0)
+			LoadedCells.Remove(Cell);
+			ActivatedCells.Add(Cell);
+			LevelCell->Activate();
+		}
+		else if (MaxCellsToLoad > 0)
+		{
+			if (!Cell->IsAlwaysLoaded())
 			{
-				break;
+				--MaxCellsToLoad;
 			}
-		}
+			ActivatedCells.Add(Cell);
+			LevelCell->Activate();
+		}		
+	}
+}
+
+void UWorldPartitionLevelStreamingPolicy::SetCellsStateToUnloaded(const TSet<const UWorldPartitionRuntimeCell*>& ToUnloadCells)
+{
+	for (const UWorldPartitionRuntimeCell* Cell : ToUnloadCells)
+	{
+		UE_LOG(LogWorldPartition, Verbose, TEXT("UWorldPartitionLevelStreamingPolicy::UnloadCells %s"), *Cell->GetName());
+		const UWorldPartitionRuntimeLevelStreamingCell* LevelCell = CastChecked<const UWorldPartitionRuntimeLevelStreamingCell>(Cell);
+		LevelCell->Unload();
+		ActivatedCells.Remove(Cell);
+		LoadedCells.Remove(Cell);
 	}
 }
 
@@ -86,8 +157,8 @@ ULevel* UWorldPartitionLevelStreamingPolicy::GetPreferredLoadedLevelToAddToWorld
 	// Sort loaded cells based on importance (only those with a streaming level in MakingVisible state)
 	TArray<const UWorldPartitionRuntimeCell*, TInlineAllocator<256>> SortedMakingVisibileCells;
 	TSet<const UWorldPartitionRuntimeCell*> MakingVisibileCells;
-	MakingVisibileCells.Reserve(LoadedCells.Num());
-	for (const UWorldPartitionRuntimeCell* LoadedCell : LoadedCells)
+	MakingVisibileCells.Reserve(ActivatedCells.Num());
+	for (const UWorldPartitionRuntimeCell* LoadedCell : ActivatedCells)
 	{
 		if (const UWorldPartitionRuntimeLevelStreamingCell* LevelSreamingCell = Cast<const UWorldPartitionRuntimeLevelStreamingCell>(LoadedCell))
 		{
@@ -102,24 +173,10 @@ ULevel* UWorldPartitionLevelStreamingPolicy::GetPreferredLoadedLevelToAddToWorld
 	return SortedMakingVisibileCells.Num() > 0 ? CastChecked<UWorldPartitionRuntimeLevelStreamingCell>(SortedMakingVisibileCells[0])->GetLevelStreaming()->GetLoadedLevel() : nullptr;
 }
 
-void UWorldPartitionLevelStreamingPolicy::LoadCell(const UWorldPartitionRuntimeCell* InCell)
-{
-	UE_LOG(LogWorldPartition, Verbose, TEXT("UWorldPartitionLevelStreamingPolicy::LoadCell %s"), *InCell->GetName());
-	const UWorldPartitionRuntimeLevelStreamingCell* Cell = Cast<const UWorldPartitionRuntimeLevelStreamingCell>(InCell);
-	if (ensure(Cell))
-	{
-		Cell->Activate();
-	}
-}
-
-void UWorldPartitionLevelStreamingPolicy::UnloadCell(const UWorldPartitionRuntimeCell* InCell)
+EWorldPartitionRuntimeCellState UWorldPartitionLevelStreamingPolicy::GetCurrentStateForCell(const UWorldPartitionRuntimeCell* Cell) const
 {
-	UE_LOG(LogWorldPartition, Verbose, TEXT("UWorldPartitionLevelStreamingPolicy::UnloadCell %s"), *InCell->GetName());
-	const UWorldPartitionRuntimeLevelStreamingCell* Cell = Cast<const UWorldPartitionRuntimeLevelStreamingCell>(InCell);
-	if (ensure(Cell))
-	{
-		Cell->Deactivate();
-	}
+	const UWorldPartitionRuntimeLevelStreamingCell* LevelCell = CastChecked<const UWorldPartitionRuntimeLevelStreamingCell>(Cell);
+	return LevelCell->GetCurrentState();
 }
 
 #if WITH_EDITOR
diff --git a/Engine/Source/Runtime/Engine/Private/WorldPartition/WorldPartitionRuntimeLevelStreamingCell.cpp b/Engine/Source/Runtime/Engine/Private/WorldPartition/WorldPartitionRuntimeLevelStreamingCell.cpp
--- a/Engine/Source/Runtime/Engine/Private/WorldPartition/WorldPartitionRuntimeLevelStreamingCell.cpp
+++ b/Engine/Source/Runtime/Engine/Private/WorldPartition/WorldPartitionRuntimeLevelStreamingCell.cpp
@@ -21,6 +21,24 @@ UWorldPartitionRuntimeLevelStreamingCell::UWorldPartitionRuntimeLevelStreamingCe
 {
 }
 
+EWorldPartitionRuntimeCellState UWorldPartitionRuntimeLevelStreamingCell::GetCurrentState() const
+{
+	if (LevelStreaming)
+	{
+		ULevelStreaming::ECurrentState CurrentStreamingState = LevelStreaming->GetCurrentState();
+		if (CurrentStreamingState == ULevelStreaming::ECurrentState::LoadedVisible)
+		{
+			return EWorldPartitionRuntimeCellState::Activated;
+		}
+		else if (CurrentStreamingState >= ULevelStreaming::ECurrentState::LoadedNotVisible)
+		{
+			return EWorldPartitionRuntimeCellState::Loaded;
+		}
+	}
+	
+	return EWorldPartitionRuntimeCellState::Unloaded;
+}
+
 UWorldPartitionLevelStreamingDynamic* UWorldPartitionRuntimeLevelStreamingCell::GetLevelStreaming() const
 {
 	return LevelStreaming;
@@ -152,12 +170,12 @@ FString UWorldPartitionRuntimeLevelStreamingCell::GetPackageNameToCreate() const
 
 #endif
 
-void UWorldPartitionRuntimeLevelStreamingCell::Activate() const
+UWorldPartitionLevelStreamingDynamic* UWorldPartitionRuntimeLevelStreamingCell::GetOrCreateLevelStreaming() const
 {
 #if WITH_EDITOR
 	if (GetActorCount() == 0)
 	{
-		return;
+		return nullptr;
 	}
 
 	if (!LevelStreaming)
@@ -203,8 +221,43 @@ void UWorldPartitionRuntimeLevelStreamingCell::Activate() const
 			LevelStreaming->OnLevelShown.AddUniqueDynamic(this, &UWorldPartitionRuntimeLevelStreamingCell::OnLevelShown);
 			LevelStreaming->OnLevelHidden.AddUniqueDynamic(this, &UWorldPartitionRuntimeLevelStreamingCell::OnLevelHidden);
 		}
+	}
+
+	return LevelStreaming;
+}
+
+void UWorldPartitionRuntimeLevelStreamingCell::Load() const
+{
+	if (UWorldPartitionLevelStreamingDynamic* LocalLevelStreaming = GetOrCreateLevelStreaming())
+	{
+		LocalLevelStreaming->Load();
+	}
+}
 
-		LevelStreaming->Activate();
+void UWorldPartitionRuntimeLevelStreamingCell::Activate() const
+{
+	if (UWorldPartitionLevelStreamingDynamic* LocalLevelStreaming = GetOrCreateLevelStreaming())
+	{
+		LocalLevelStreaming->Activate();
+	}
+}
+
+void UWorldPartitionRuntimeLevelStreamingCell::Unload() const
+{
+#if WITH_EDITOR
+	if (GetActorCount() == 0)
+	{
+		return;
+	}
+	check(LevelStreaming);
+#else
+	// In Runtime, always loaded cell level is handled by World directly
+	check(LevelStreaming || IsAlwaysLoaded());
+#endif
+
+	if (LevelStreaming)
+	{
+		LevelStreaming->Unload();
 	}
 }
 
diff --git a/Engine/Source/Runtime/Engine/Private/WorldPartition/WorldPartitionRuntimeSpatialHash.cpp b/Engine/Source/Runtime/Engine/Private/WorldPartition/WorldPartitionRuntimeSpatialHash.cpp
--- a/Engine/Source/Runtime/Engine/Private/WorldPartition/WorldPartitionRuntimeSpatialHash.cpp
+++ b/Engine/Source/Runtime/Engine/Private/WorldPartition/WorldPartitionRuntimeSpatialHash.cpp
@@ -103,7 +103,64 @@ const FSquare2DGridHelper& FSpatialHashStreamingGrid::GetGridHelper() const
 	return *GridHelper;
 }
 
-void FSpatialHashStreamingGrid::GetCells(const TArray<FWorldPartitionStreamingSource>& Sources, const UDataLayerSubsystem* DataLayerSubsystem, TSet<const UWorldPartitionRuntimeCell*>& Cells) const
+void FSpatialHashStreamingGrid::GetCells(const TArray<FWorldPartitionStreamingQuerySource>& QuerySources, TSet<const UWorldPartitionRuntimeCell*>& OutCells) const
+{
+	auto ShouldAddCell = [](const UWorldPartitionRuntimeCell* Cell, const FWorldPartitionStreamingQuerySource& QuerySource)
+	{
+		if (Cell->HasDataLayers())
+		{
+			if (Cell->GetDataLayers().FindByPredicate([&](const FName& DataLayerName) { return QuerySource.DataLayers.Contains(DataLayerName); }))
+			{
+				return true;
+			}
+		}
+		else if (!QuerySource.bDataLayersOnly)
+		{
+			return true;
+		}
+
+		return false;
+	};
+
+	const FSquare2DGridHelper& Helper = GetGridHelper();
+	for (const FWorldPartitionStreamingQuerySource& QuerySource : QuerySources)
+	{
+		// Spatial Query
+		if (QuerySource.bSpatialQuery)
+		{
+			const FSphere GridSphere(QuerySource.Location, QuerySource.bUseLoadingRangeRadius? GetLoadingRange() : QuerySource.Radius);
+			Helper.ForEachIntersectingCells(GridSphere, [&](const FIntVector& Coords)
+				{
+					if (const int32* LayerCellIndexPtr = GridLevels[Coords.Z].LayerCellsMapping.Find(Coords.Y * Helper.Levels[Coords.Z].GridSize + Coords.X))
+					{
+						const FSpatialHashStreamingGridLayerCell& LayerCell = GridLevels[Coords.Z].LayerCells[*LayerCellIndexPtr];
+						for (const UWorldPartitionRuntimeCell* Cell : LayerCell.GridCells)
+						{
+							if (ShouldAddCell(Cell, QuerySource))
+							{
+								OutCells.Add(Cell);
+							}
+						}
+					}
+				});
+		}
+
+		// Non Spatial (always included)
+		const int32 TopLevel = GridLevels.Num() - 1;
+		for (const FSpatialHashStreamingGridLayerCell& LayerCell : GridLevels[TopLevel].LayerCells)
+		{
+			for (const UWorldPartitionRuntimeCell* Cell : LayerCell.GridCells)
+			{
+				if (ShouldAddCell(Cell, QuerySource))
+				{
+					OutCells.Add(Cell);
+				}
+			}
+		}
+	}
+}
+
+void FSpatialHashStreamingGrid::GetCells(const TArray<FWorldPartitionStreamingSource>& Sources, const UDataLayerSubsystem* DataLayerSubsystem, TSet<const UWorldPartitionRuntimeCell*>& OutActivateCells, TSet<const UWorldPartitionRuntimeCell*>& OutLoadCells) const
 {
 	const FSquare2DGridHelper& Helper = GetGridHelper();
 	for (const FWorldPartitionStreamingSource& Source : Sources)
@@ -116,19 +173,31 @@ void FSpatialHashStreamingGrid::GetCells(const TArray<FWorldPartitionStreamingSo
 				const FSpatialHashStreamingGridLayerCell& LayerCell = GridLevels[Coords.Z].LayerCells[*LayerCellIndexPtr];
 				for (const UWorldPartitionRuntimeCell* Cell : LayerCell.GridCells)
 				{
-					if (!Cell->HasDataLayers() || (DataLayerSubsystem && DataLayerSubsystem->IsAnyDataLayerActive(Cell->GetDataLayers())))
+					if (!Cell->HasDataLayers() || (DataLayerSubsystem && DataLayerSubsystem->IsAnyDataLayerInState(Cell->GetDataLayers(), EDataLayerState::Activated)))
 					{
-						Cells.Add(Cell);
+						if (Source.TargetState == EStreamingSourceTargetState::Loaded)
+						{
+							OutLoadCells.Add(Cell);
+						}
+						else
+						{
+							check(Source.TargetState == EStreamingSourceTargetState::Activated);
+							OutActivateCells.Add(Cell);
+						}
+					}
+					else if (DataLayerSubsystem && DataLayerSubsystem->IsAnyDataLayerInState(Cell->GetDataLayers(), EDataLayerState::Loaded))
+					{
+						OutLoadCells.Add(Cell);
 					}
 				}
 			}
 		});
 	}
 
-	GetAlwaysLoadedCells(DataLayerSubsystem, Cells);
+	GetAlwaysLoadedCells(DataLayerSubsystem, OutActivateCells, OutLoadCells);
 }
 
-void FSpatialHashStreamingGrid::GetAlwaysLoadedCells(const UDataLayerSubsystem* DataLayerSubsystem, TSet<const UWorldPartitionRuntimeCell*>& Cells) const
+void FSpatialHashStreamingGrid::GetAlwaysLoadedCells(const UDataLayerSubsystem* DataLayerSubsystem, TSet<const UWorldPartitionRuntimeCell*>& OutActivateCells, TSet<const UWorldPartitionRuntimeCell*>& OutLoadCells) const
 {
 	if (GridLevels.Num() > 0)
 	{
@@ -137,10 +206,15 @@ void FSpatialHashStreamingGrid::GetAlwaysLoadedCells(const UDataLayerSubsystem*
 		{
 			for (const UWorldPartitionRuntimeCell* Cell : LayerCell.GridCells)
 			{
-				if (!Cell->HasDataLayers() || (DataLayerSubsystem && DataLayerSubsystem->IsAnyDataLayerActive(Cell->GetDataLayers())))
+				if (!Cell->HasDataLayers() || (DataLayerSubsystem && DataLayerSubsystem->IsAnyDataLayerInState(Cell->GetDataLayers(), EDataLayerState::Activated)))
 				{
 					check(Cell->IsAlwaysLoaded() || Cell->HasDataLayers());
-					Cells.Add(Cell);
+					OutActivateCells.Add(Cell);
+				}
+				else if(DataLayerSubsystem && DataLayerSubsystem->IsAnyDataLayerInState(Cell->GetDataLayers(), EDataLayerState::Loaded))
+				{
+					check(Cell->HasDataLayers());
+					OutLoadCells.Add(Cell);
 				}
 			}
 		}
@@ -858,7 +932,17 @@ int32 UWorldPartitionRuntimeSpatialHash::GetAllStreamingCells(TSet<const UWorldP
 	return Cells.Num();
 }
 
-int32 UWorldPartitionRuntimeSpatialHash::GetStreamingCells(const TArray<FWorldPartitionStreamingSource>& Sources, TSet<const UWorldPartitionRuntimeCell*>& Cells) const
+bool UWorldPartitionRuntimeSpatialHash::GetStreamingCells(const TArray<FWorldPartitionStreamingQuerySource>& Sources, TSet<const UWorldPartitionRuntimeCell*>& OutCells) const 
+{
+	for (const FSpatialHashStreamingGrid& StreamingGrid : StreamingGrids)
+	{
+		StreamingGrid.GetCells(Sources, OutCells);
+	}
+
+	return !!OutCells.Num();
+}
+
+bool UWorldPartitionRuntimeSpatialHash::GetStreamingCells(const TArray<FWorldPartitionStreamingSource>& Sources, TSet<const UWorldPartitionRuntimeCell*>& OutActivateCells, TSet<const UWorldPartitionRuntimeCell*>& OutLoadCells) const
 {
 	UDataLayerSubsystem* DataLayerSubsystem = GetWorld()->GetSubsystem<UDataLayerSubsystem>();
 
@@ -867,7 +951,7 @@ int32 UWorldPartitionRuntimeSpatialHash::GetStreamingCells(const TArray<FWorldPa
 		// Get always loaded cells
 		for (const FSpatialHashStreamingGrid& StreamingGrid : StreamingGrids)
 		{
-			StreamingGrid.GetAlwaysLoadedCells(DataLayerSubsystem, Cells);
+			StreamingGrid.GetAlwaysLoadedCells(DataLayerSubsystem, OutActivateCells, OutLoadCells);
 		}
 	}
 	else
@@ -875,11 +959,11 @@ int32 UWorldPartitionRuntimeSpatialHash::GetStreamingCells(const TArray<FWorldPa
 		// Get cells based on streaming sources
 		for (const FSpatialHashStreamingGrid& StreamingGrid : StreamingGrids)
 		{
-			StreamingGrid.GetCells(Sources, DataLayerSubsystem, Cells);
+			StreamingGrid.GetCells(Sources, DataLayerSubsystem, OutActivateCells, OutLoadCells);
 		}
 	}
 
-	return Cells.Num();
+	return !!(OutActivateCells.Num() + OutLoadCells.Num());
 }
 
 void UWorldPartitionRuntimeSpatialHash::SortStreamingCellsByImportance(const TSet<const UWorldPartitionRuntimeCell*>& InCells, const TArray<FWorldPartitionStreamingSource>& InSources, TArray<const UWorldPartitionRuntimeCell*, TInlineAllocator<256>>& OutSortedCells) const
diff --git a/Engine/Source/Runtime/Engine/Private/WorldPartition/WorldPartitionStreamingPolicy.cpp b/Engine/Source/Runtime/Engine/Private/WorldPartition/WorldPartitionStreamingPolicy.cpp
--- a/Engine/Source/Runtime/Engine/Private/WorldPartition/WorldPartitionStreamingPolicy.cpp
+++ b/Engine/Source/Runtime/Engine/Private/WorldPartition/WorldPartitionStreamingPolicy.cpp
@@ -55,7 +55,7 @@ void UWorldPartitionStreamingPolicy::UpdateStreamingSources()
 		const FVector ViewLocationLocal = WorldToLocal.TransformPosition(ViewLocation);
 		const FRotator ViewRotationLocal = WorldToLocal.TransformRotation(ViewRotation.Quaternion()).Rotator();
 		static const FName NAME_SIE(TEXT("SIE"));
-		StreamingSources.Add(FWorldPartitionStreamingSource(NAME_SIE, ViewLocationLocal, ViewRotationLocal));
+		StreamingSources.Add(FWorldPartitionStreamingSource(NAME_SIE, ViewLocationLocal, ViewRotationLocal, EStreamingSourceTargetState::Activated));
 	}
 	else
 #endif
@@ -80,7 +80,7 @@ void UWorldPartitionStreamingPolicy::UpdateStreamingSources()
 						// Transform to Local
 						ViewLocation = WorldToLocal.TransformPosition(ViewLocation);
 						ViewRotation = WorldToLocal.TransformRotation(ViewRotation.Quaternion()).Rotator();
-						StreamingSources.Add(FWorldPartitionStreamingSource(Player->GetFName(), ViewLocation, ViewRotation));
+						StreamingSources.Add(FWorldPartitionStreamingSource(Player->GetFName(), ViewLocation, ViewRotation, EStreamingSourceTargetState::Activated));
 					}
 				}
 			}
@@ -103,7 +103,8 @@ void UWorldPartitionStreamingPolicy::UpdateStreamingSources()
 
 void UWorldPartitionStreamingPolicy::UpdateStreamingState()
 {
-	TSet<const UWorldPartitionRuntimeCell*> StreamingCells;
+	TSet<const UWorldPartitionRuntimeCell*> LoadStreamingCells;
+	TSet<const UWorldPartitionRuntimeCell*> ActivateStreamingCells;
 	UWorld* World = WorldPartition->GetWorld();
 	check(World && World->IsGameWorld());
 	
@@ -115,24 +116,24 @@ void UWorldPartitionStreamingPolicy::UpdateStreamingState()
 		// Load all cells on server (do this once)
 		if (!bIsServerLoadingDone)
 		{
-			WorldPartition->RuntimeHash->GetAllStreamingCells(StreamingCells);
-			TSet<const UWorldPartitionRuntimeCell*> ToLoadCells = StreamingCells.Difference(LoadedCells);
+			WorldPartition->RuntimeHash->GetAllStreamingCells(ActivateStreamingCells);
+			TSet<const UWorldPartitionRuntimeCell*> ToActivateCells = ActivateStreamingCells.Difference(ActivatedCells);
 
 			// Mark Runtime Cells to be Always Loaded
-			for (const UWorldPartitionRuntimeCell* Cell : ToLoadCells)
+			for (const UWorldPartitionRuntimeCell* Cell : ToActivateCells)
 			{
 				UWorldPartitionRuntimeCell* MutableCell = const_cast<UWorldPartitionRuntimeCell*>(Cell);
 				MutableCell->SetIsAlwaysLoaded(true);
 			}
 
-			LoadCells(ToLoadCells);
+			SetTargetStateForCells(EWorldPartitionRuntimeCellState::Activated, ToActivateCells);
 			bIsServerLoadingDone = true;
 		}
 	}
 	else
 	{
 		// Early out if nothing loaded and no streaming source
-		if ((StreamingSources.Num() == 0) && (LoadedCells.Num() == 0))
+		if ((StreamingSources.Num() == 0) && (ActivatedCells.Num() == 0) && (LoadedCells.Num() == 0))
 		{
 			return;
 		}
@@ -140,18 +141,21 @@ void UWorldPartitionStreamingPolicy::UpdateStreamingState()
 		// When uninitializing, UpdateStreamingState is called, but we don't want any cells to be loaded
 		if (WorldPartition->IsInitialized())
 		{
-			WorldPartition->RuntimeHash->GetStreamingCells(StreamingSources, StreamingCells);
+			WorldPartition->RuntimeHash->GetStreamingCells(StreamingSources, ActivateStreamingCells, LoadStreamingCells);
+			// Activation superseeds Loading
+			LoadStreamingCells = LoadStreamingCells.Difference(ActivateStreamingCells);
 		}
 
 		// Determine cells to load/unload
-		TSet<const UWorldPartitionRuntimeCell*> ToLoadCells = StreamingCells.Difference(LoadedCells);
-		TSet<const UWorldPartitionRuntimeCell*> ToUnloadCells = LoadedCells.Difference(StreamingCells);
+		TSet<const UWorldPartitionRuntimeCell*> ToActivateCells = ActivateStreamingCells.Difference(ActivatedCells);
+		TSet<const UWorldPartitionRuntimeCell*> ToLoadCells = LoadStreamingCells.Difference(LoadedCells);
+		TSet<const UWorldPartitionRuntimeCell*> ToUnloadCells = ActivatedCells.Union(LoadedCells).Difference(ActivateStreamingCells.Union(LoadStreamingCells));
 
 		UE_SUPPRESS(LogWorldPartition, Verbose,
 		{
-			if (ToLoadCells.Num() > 0 || ToUnloadCells.Num() > 0)
+			if (ToActivateCells.Num() > 0 || ToUnloadCells.Num() > 0)
 			{
-				UE_LOG(LogWorldPartition, Verbose, TEXT("UWorldPartitionStreamingPolicy: CellsToLoad(%d), CellsToUnload(%d)"), ToLoadCells.Num(), ToUnloadCells.Num());
+				UE_LOG(LogWorldPartition, Verbose, TEXT("UWorldPartitionStreamingPolicy: CellsToLoad(%d), CellsToUnload(%d)"), ToActivateCells.Num(), ToUnloadCells.Num());
 				
 				FTransform LocalToWorld = WorldPartition->GetInstanceTransform();
 				for (int i = 0; i < StreamingSources.Num(); ++i)
@@ -166,14 +170,40 @@ void UWorldPartitionStreamingPolicy::UpdateStreamingState()
 		// Process unload first, so that LoadedCells is up-to-date when calling LoadCells
 		if (ToUnloadCells.Num() > 0)
 		{
-			UnloadCells(ToUnloadCells);
+			SetTargetStateForCells(EWorldPartitionRuntimeCellState::Unloaded, ToUnloadCells);
+		}
+
+		// Do Activation State first as it is higher prio than Load State (if we have a limited number of loading cells per frame)
+		if (ToActivateCells.Num() > 0)
+		{
+			SetTargetStateForCells(EWorldPartitionRuntimeCellState::Activated, ToActivateCells);
 		}
 
 		if (ToLoadCells.Num() > 0)
 		{
-			LoadCells(ToLoadCells);
+			SetTargetStateForCells(EWorldPartitionRuntimeCellState::Loaded, ToLoadCells);
+		}
+	}
+}
+
+bool UWorldPartitionStreamingPolicy::IsStreamingCompleted(EWorldPartitionRuntimeCellState QueryState, const TArray<FWorldPartitionStreamingQuerySource>& QuerySources, bool bExactState) const
+{
+	TSet<const UWorldPartitionRuntimeCell*> Cells;
+	WorldPartition->RuntimeHash->GetStreamingCells(QuerySources, Cells);
+	
+	for (const UWorldPartitionRuntimeCell* Cell : Cells)
+	{
+		EWorldPartitionRuntimeCellState CellState = GetCurrentStateForCell(Cell);
+		if (CellState != QueryState)
+		{
+			if (bExactState || CellState < QueryState)
+			{
+				return false;
+			}
 		}
 	}
+
+	return true;
 }
 
 FVector2D UWorldPartitionStreamingPolicy::GetDrawRuntimeHash2DDesiredFootprint(const FVector2D& CanvasSize)
@@ -197,21 +227,3 @@ void UWorldPartitionStreamingPolicy::DrawRuntimeHash3D()
 		WorldPartition->RuntimeHash->Draw3D(StreamingSources);
 	}
 }
-
-void UWorldPartitionStreamingPolicy::LoadCells(const TSet<const UWorldPartitionRuntimeCell*>& ToLoadCells)
-{
-	for (const UWorldPartitionRuntimeCell* ToLoadCell : ToLoadCells)
-	{
-		LoadCell(ToLoadCell);
-		LoadedCells.Add(ToLoadCell);
-	}
-}
-
-void UWorldPartitionStreamingPolicy::UnloadCells(const TSet<const UWorldPartitionRuntimeCell*>& ToUnloadCells)
-{
-	for (const UWorldPartitionRuntimeCell* ToUnloadCell : ToUnloadCells)
-	{
-		UnloadCell(ToUnloadCell);
-		LoadedCells.Remove(ToUnloadCell);
-	}
-}
diff --git a/Engine/Source/Runtime/Engine/Private/WorldPartition/WorldPartitionSubsystem.cpp b/Engine/Source/Runtime/Engine/Private/WorldPartition/WorldPartitionSubsystem.cpp
--- a/Engine/Source/Runtime/Engine/Private/WorldPartition/WorldPartitionSubsystem.cpp
+++ b/Engine/Source/Runtime/Engine/Private/WorldPartition/WorldPartitionSubsystem.cpp
@@ -170,6 +170,19 @@ TStatId UWorldPartitionSubsystem::GetStatId() const
 	RETURN_QUICK_DECLARE_CYCLE_STAT(UWorldPartitionSubsystem, STATGROUP_Tickables);
 }
 
+bool UWorldPartitionSubsystem::IsStreamingCompleted(EWorldPartitionRuntimeCellState QueryState, const TArray<FWorldPartitionStreamingQuerySource>& QuerySources, bool bExactState) const
+{
+	for (UWorldPartition* Partition : RegisteredWorldPartitions)
+	{
+		if (!Partition->IsStreamingCompleted(QueryState, QuerySources, bExactState))
+		{
+			return false;
+		}
+	}
+
+	return true;
+}
+
 void UWorldPartitionSubsystem::UpdateStreamingState()
 {
 #if WITH_EDITOR
diff --git a/Engine/Source/Runtime/Engine/Public/WorldPartition/DataLayer/DataLayer.h b/Engine/Source/Runtime/Engine/Public/WorldPartition/DataLayer/DataLayer.h
--- a/Engine/Source/Runtime/Engine/Public/WorldPartition/DataLayer/DataLayer.h
+++ b/Engine/Source/Runtime/Engine/Public/WorldPartition/DataLayer/DataLayer.h
@@ -11,6 +11,14 @@
 
 #include "DataLayer.generated.h"
 
+UENUM(BlueprintType)
+enum class EDataLayerState : uint8
+{
+	Unloaded,
+	Loaded,
+	Activated
+};
+
 UCLASS(Config = Engine, PerObjectConfig, Within = WorldDataLayers, BlueprintType)
 class ENGINE_API UDataLayer : public UObject
 {
@@ -32,6 +40,7 @@ public:
 
 	bool IsLocked() const { return bIsLocked; }
 #endif
+	virtual void PostLoad() override;
 
 	UFUNCTION(Category = DataLayer, BlueprintCallable)
 	bool Equals(const FActorDataLayer& ActorDataLayer) const { return ActorDataLayer.Name == GetFName(); }
@@ -45,14 +54,20 @@ public:
 	UFUNCTION(Category = DataLayer, BlueprintCallable)
 	bool IsDynamicallyLoaded() const { return bIsDynamicallyLoaded; }
 
+	UFUNCTION(Category = DataLayer, BlueprintCallable, meta = (DeprecatedFunction, DeprecationMessage = "Use GetInitialState instead"))
+	bool IsInitiallyActive() const { return IsDynamicallyLoaded() && GetInitialState() == EDataLayerState::Activated; }	
+
 	UFUNCTION(Category = DataLayer, BlueprintCallable)
-	bool IsInitiallyActive() const { return IsDynamicallyLoaded() && bIsInitiallyActive; }	
+	EDataLayerState GetInitialState() const { return IsDynamicallyLoaded() ? InitialState : EDataLayerState::Unloaded; }
 
 private:
 	/** The display name of the DataLayer */
 	UPROPERTY()
 	FName DataLayerLabel;
 
+	UPROPERTY(Category = DataLayer, EditAnywhere, meta = (EditConditionHides, EditCondition = "bIsDynamicallyLoaded"))
+	EDataLayerState InitialState;
+
 	/** Whether actors associated with the DataLayer are visible in the viewport */
 	UPROPERTY(Category = DataLayer, EditAnywhere)
 	uint32 bIsVisible : 1;
@@ -61,11 +76,11 @@ private:
 	UPROPERTY(Category = DataLayer, EditAnywhere)
 	uint32 bIsDynamicallyLoaded : 1;
 
-	/** Whether a dynamically loaded DataLayer should be initially active at runtime */
-	UPROPERTY(Category = DataLayer, EditAnywhere, meta = (EditConditionHides, EditCondition = "bIsDynamicallyLoaded"))
-	uint32 bIsInitiallyActive : 1;
-
 #if WITH_EDITORONLY_DATA
+	/** Whether a dynamically loaded DataLayer should be initially active at runtime */
+	UPROPERTY()
+	uint32 bIsInitiallyActive_DEPRECATED : 1;
+	
 	/** Whether the DataLayer affects actor editor loading */
 	UPROPERTY(Transient)
 	uint32 bIsDynamicallyLoadedInEditor : 1;
diff --git a/Engine/Source/Runtime/Engine/Public/WorldPartition/DataLayer/DataLayerSubsystem.h b/Engine/Source/Runtime/Engine/Public/WorldPartition/DataLayer/DataLayerSubsystem.h
--- a/Engine/Source/Runtime/Engine/Public/WorldPartition/DataLayer/DataLayerSubsystem.h
+++ b/Engine/Source/Runtime/Engine/Public/WorldPartition/DataLayer/DataLayerSubsystem.h
@@ -12,8 +12,11 @@
  * UDataLayerSubsystem
  */
 
+// Deprecate
 DECLARE_DYNAMIC_MULTICAST_DELEGATE_TwoParams(FOnDataLayerActivationStateChanged, const UDataLayer*, DataLayer, bool, bIsActive);
 
+DECLARE_DYNAMIC_MULTICAST_DELEGATE_TwoParams(FOnDataLayerStateChanged, const UDataLayer*, DataLayer, EDataLayerState, State);
+
 UCLASS()
 class ENGINE_API UDataLayerSubsystem : public UWorldSubsystem
 {
@@ -33,30 +36,46 @@ public:
 	//~ End UWorldSubsystem Interface.
 
 	//~ Begin Blueprint callable functions
-	UFUNCTION(BlueprintCallable, Category = DataLayers)
+	UFUNCTION(BlueprintCallable, Category = DataLayers, meta=(DeprecatedFunction, DeprecationMessage="Use SetDataLayerState instead"))
 	void ActivateDataLayer(const FActorDataLayer& InDataLayer, bool bInActivate);
 	
-	UFUNCTION(BlueprintCallable, Category = DataLayers)
+	UFUNCTION(BlueprintCallable, Category = DataLayers, meta = (DeprecatedFunction, DeprecationMessage = "Use SetDataLayerStateByLabel instead"))
 	void ActivateDataLayerByLabel(const FName& InDataLayerLabel, bool bInActivate);
 
-	UFUNCTION(BlueprintCallable, Category = DataLayers)
+	UFUNCTION(BlueprintCallable, Category = DataLayers, meta = (DeprecatedFunction, DeprecationMessage = "Use GetDataLayerState instead"))
 	bool IsDataLayerActive(const FActorDataLayer& InDataLayer) const;
 
+	UFUNCTION(BlueprintCallable, Category = DataLayers, meta = (DeprecatedFunction, DeprecationMessage = "Use GetDataLayerStateByLabel instead"))
+	bool IsDataLayerActiveByLabel(const FName& InDataLayerLabel) const;
+
+	UFUNCTION(BlueprintCallable, Category = DataLayers)
+	void SetDataLayerState(const FActorDataLayer& InDataLayer, EDataLayerState InState);
+
+	UFUNCTION(BlueprintCallable, Category = DataLayers)
+	void SetDataLayerStateByLabel(const FName& InDataLayerLabel, EDataLayerState InState);
+		
 	UFUNCTION(BlueprintCallable, Category = DataLayers)
-	bool IsDataLayerActiveByLabel(const FName& InDataLayerName) const;
+	EDataLayerState GetDataLayerState(const FActorDataLayer& InDataLayer) const;
+
+	UFUNCTION(BlueprintCallable, Category = DataLayers)
+	EDataLayerState GetDataLayerStateByLabel(const FName& InDataLayerLabel) const;
 
 	UPROPERTY(BlueprintAssignable)
 	FOnDataLayerActivationStateChanged OnDataLayerActivationStateChanged;
+
+	UPROPERTY(BlueprintAssignable)
+	FOnDataLayerStateChanged OnDataLayerStateChanged;
 	//~ End Blueprint callable functions
 
 	UDataLayer* GetDataLayerFromLabel(const FName& InDataLayerLabel) const;
 	UDataLayer* GetDataLayerFromName(const FName& InDataLayerName) const;
-	void ActivateDataLayer(const UDataLayer* InDataLayer, bool bInActivate);
-	void ActivateDataLayerByName(const FName& InDataLayerName, bool bInActivate);
-	bool IsDataLayerActive(const UDataLayer* InDataLayer) const;
-	bool IsDataLayerActiveByName(const FName& InDataLayerName) const;
-	bool IsAnyDataLayerActive(const TArray<FName>& InDataLayerNames) const;
-
+	void SetDataLayerState(const UDataLayer* InDataLayer, EDataLayerState InState);
+	void SetDataLayerStateByName(const FName& InDataLayerName, EDataLayerState InState);
+	EDataLayerState GetDataLayerState(const UDataLayer* InDataLayer) const;
+	EDataLayerState GetDataLayerStateByName(const FName& InDataLayerName) const;
+	
+	bool IsAnyDataLayerInState(const TArray<FName>& InDataLayerNames, EDataLayerState InState) const;
+	
 	void Draw(class UCanvas* Canvas, class APlayerController* PC);
 
 #if WITH_EDITOR
@@ -74,6 +93,7 @@ private:
 	static class FAutoConsoleCommand ToggleDataLayerActivation;
 
 	TSet<FName> ActiveDataLayerNames;
+	TSet<FName> LoadedDataLayerNames;
 
 	FDelegateHandle	DrawHandle;
 };
\ No newline at end of file
diff --git a/Engine/Source/Runtime/Engine/Public/WorldPartition/WorldPartition.h b/Engine/Source/Runtime/Engine/Public/WorldPartition/WorldPartition.h
--- a/Engine/Source/Runtime/Engine/Public/WorldPartition/WorldPartition.h
+++ b/Engine/Source/Runtime/Engine/Public/WorldPartition/WorldPartition.h
@@ -27,6 +27,8 @@ class FHLODActorDesc;
 
 struct IWorldPartitionStreamingSourceProvider;
 
+enum class EWorldPartitionRuntimeCellState : uint8;
+
 DECLARE_LOG_CATEGORY_EXTERN(LogWorldPartition, Log, All);
 
 enum class EWorldPartitionStreamingMode
@@ -170,6 +172,7 @@ public:
 	void Tick(float DeltaSeconds);
 	void UpdateStreamingState();
 	class ULevel* GetPreferredLoadedLevelToAddToWorld() const;
+	bool IsStreamingCompleted(EWorldPartitionRuntimeCellState QueryState, const TArray<FWorldPartitionStreamingQuerySource>& QuerySources, bool bExactState) const;
 
 	bool CanDrawRuntimeHash() const;
 	FVector2D GetDrawRuntimeHash2DDesiredFootprint(const FVector2D& CanvasSize);
diff --git a/Engine/Source/Runtime/Engine/Public/WorldPartition/WorldPartitionLevelStreamingDynamic.h b/Engine/Source/Runtime/Engine/Public/WorldPartition/WorldPartitionLevelStreamingDynamic.h
--- a/Engine/Source/Runtime/Engine/Public/WorldPartition/WorldPartitionLevelStreamingDynamic.h
+++ b/Engine/Source/Runtime/Engine/Public/WorldPartition/WorldPartitionLevelStreamingDynamic.h
@@ -25,6 +25,8 @@ class ENGINE_API UWorldPartitionLevelStreamingDynamic : public ULevelStreamingDy
 {
 	GENERATED_UCLASS_BODY()
 
+	void Load();
+	void Unload();
 	void Activate();
 	void Deactivate();
 	virtual bool ShouldBeAlwaysLoaded() const { return bShouldBeAlwaysLoaded; }
@@ -58,8 +60,6 @@ private:
 	FWorldPartitionPackageCache PackageCache;
 #endif
 
-	bool bIsActivated;
-
 	UPROPERTY()
 	bool bShouldBeAlwaysLoaded;
 
diff --git a/Engine/Source/Runtime/Engine/Public/WorldPartition/WorldPartitionLevelStreamingPolicy.h b/Engine/Source/Runtime/Engine/Public/WorldPartition/WorldPartitionLevelStreamingPolicy.h
--- a/Engine/Source/Runtime/Engine/Public/WorldPartition/WorldPartitionLevelStreamingPolicy.h
+++ b/Engine/Source/Runtime/Engine/Public/WorldPartition/WorldPartitionLevelStreamingPolicy.h
@@ -14,15 +14,16 @@
 #include "WorldPartition/WorldPartitionStreamingPolicy.h"
 #include "WorldPartitionLevelStreamingPolicy.generated.h"
 
+enum class EWorldPartitionRuntimeCellState : uint8;
+
 UCLASS()
 class UWorldPartitionLevelStreamingPolicy : public UWorldPartitionStreamingPolicy
 {
 	GENERATED_BODY()
 
 public:
-	virtual void LoadCells(const TSet<const UWorldPartitionRuntimeCell*>& ToLoadCells) override;
-	virtual void LoadCell(const UWorldPartitionRuntimeCell* Cell) override;
-	virtual void UnloadCell(const UWorldPartitionRuntimeCell* Cell) override;
+	virtual void SetTargetStateForCells(EWorldPartitionRuntimeCellState TargetState, const TSet<const UWorldPartitionRuntimeCell*>& Cells) override;
+	virtual EWorldPartitionRuntimeCellState GetCurrentStateForCell(const UWorldPartitionRuntimeCell* Cell) const override;
 	virtual class ULevel* GetPreferredLoadedLevelToAddToWorld() const override;
 
 #if WITH_EDITOR
@@ -34,6 +35,9 @@ public:
 #endif
 
 private:
+	void SetCellsStateToLoaded(const TSet<const UWorldPartitionRuntimeCell*>& ToLoadCells);
+	void SetCellsStateToActivated(const TSet<const UWorldPartitionRuntimeCell*>& ToActivateCells);
+	void SetCellsStateToUnloaded(const TSet<const UWorldPartitionRuntimeCell*>& ToUnloadCells);
 	int32 GetCellLoadingCount() const;
 
 #if WITH_EDITOR
diff --git a/Engine/Source/Runtime/Engine/Public/WorldPartition/WorldPartitionRuntimeCell.h b/Engine/Source/Runtime/Engine/Public/WorldPartition/WorldPartitionRuntimeCell.h
--- a/Engine/Source/Runtime/Engine/Public/WorldPartition/WorldPartitionRuntimeCell.h
+++ b/Engine/Source/Runtime/Engine/Public/WorldPartition/WorldPartitionRuntimeCell.h
@@ -89,6 +89,17 @@ class UWorldPartitionRuntimeCellData : public UObject
 
 class UActorDescContainer;
 
+/**
+ * Cell State
+ */
+UENUM(BlueprintType)
+enum class EWorldPartitionRuntimeCellState : uint8
+{
+	Unloaded,
+	Loaded,
+	Activated
+};
+
 /**
  * Represents a PIE/Game streaming cell which points to external actor/data chunk packages
  */
diff --git a/Engine/Source/Runtime/Engine/Public/WorldPartition/WorldPartitionRuntimeHash.h b/Engine/Source/Runtime/Engine/Public/WorldPartition/WorldPartitionRuntimeHash.h
--- a/Engine/Source/Runtime/Engine/Public/WorldPartition/WorldPartitionRuntimeHash.h
+++ b/Engine/Source/Runtime/Engine/Public/WorldPartition/WorldPartitionRuntimeHash.h
@@ -40,7 +40,8 @@ class ENGINE_API UWorldPartitionRuntimeHash : public UObject
 
 	// Streaming interface
 	virtual int32 GetAllStreamingCells(TSet<const UWorldPartitionRuntimeCell*>& Cells, bool bIncludeDataLayers = false) const { return 0; }
-	virtual int32 GetStreamingCells(const TArray<FWorldPartitionStreamingSource>& Sources, TSet<const UWorldPartitionRuntimeCell*>& Cells) const { return 0; };
+	virtual bool GetStreamingCells(const TArray<FWorldPartitionStreamingQuerySource>& Sources, TSet<const UWorldPartitionRuntimeCell*>& OutCells) const { return false; }
+	virtual bool GetStreamingCells(const TArray<FWorldPartitionStreamingSource>& Sources, TSet<const UWorldPartitionRuntimeCell*>& OutActivateCells, TSet<const UWorldPartitionRuntimeCell*>& OutLoadCells) const { return false; };
 	virtual void SortStreamingCellsByImportance(const TSet<const UWorldPartitionRuntimeCell*>& InCells, const TArray<FWorldPartitionStreamingSource>& InSources, TArray<const UWorldPartitionRuntimeCell*, TInlineAllocator<256>>& OutSortedCells) const {}
 
 	/* Returns desired footprint that Draw2D should take relative to given Canvas size (the value can exceed the given size).
diff --git a/Engine/Source/Runtime/Engine/Public/WorldPartition/WorldPartitionRuntimeLevelStreamingCell.h b/Engine/Source/Runtime/Engine/Public/WorldPartition/WorldPartitionRuntimeLevelStreamingCell.h
--- a/Engine/Source/Runtime/Engine/Public/WorldPartition/WorldPartitionRuntimeLevelStreamingCell.h
+++ b/Engine/Source/Runtime/Engine/Public/WorldPartition/WorldPartitionRuntimeLevelStreamingCell.h
@@ -10,11 +10,14 @@ class UWorldPartitionRuntimeLevelStreamingCell : public UWorldPartitionRuntimeSp
 {
 	GENERATED_UCLASS_BODY()
 
+	void Load() const;
+	void Unload() const;
 	void Activate() const;
 	void Deactivate() const;
 	int32 GetStreamingPriority() const { return Level; }
 	class UWorldPartitionLevelStreamingDynamic* GetLevelStreaming() const;
 	virtual FLinearColor GetDebugColor() const override;
+	EWorldPartitionRuntimeCellState GetCurrentState() const;
 
 #if WITH_EDITOR
 	virtual void AddActorToCell(const FWorldPartitionActorDescView& ActorDescView, uint32 InContainerID, const FTransform& InContainerTransform, const UActorDescContainer* InContainer) override;
@@ -37,6 +40,8 @@ private:
 
 	EStreamingStatus GetLevelStreamingStatus() const;
 
+	class UWorldPartitionLevelStreamingDynamic* GetOrCreateLevelStreaming() const;
+
 #if WITH_EDITOR
 	void LoadActorsForCook();
 	void MoveAlwaysLoadedContentToPersistentLevel();
diff --git a/Engine/Source/Runtime/Engine/Public/WorldPartition/WorldPartitionRuntimeSpatialHash.h b/Engine/Source/Runtime/Engine/Public/WorldPartition/WorldPartitionRuntimeSpatialHash.h
--- a/Engine/Source/Runtime/Engine/Public/WorldPartition/WorldPartitionRuntimeSpatialHash.h
+++ b/Engine/Source/Runtime/Engine/Public/WorldPartition/WorldPartitionRuntimeSpatialHash.h
@@ -126,8 +126,9 @@ struct FSpatialHashStreamingGrid
 	}
 
 	// Used by PIE/Game
-	void GetCells(const TArray<FWorldPartitionStreamingSource>& Sources, const class UDataLayerSubsystem* DataLayerSubsystem, TSet<const UWorldPartitionRuntimeCell*>& Cells) const;
-	void GetAlwaysLoadedCells(const UDataLayerSubsystem* DataLayerSubsystem, TSet<const UWorldPartitionRuntimeCell*>& Cells) const;
+	void GetCells(const TArray<FWorldPartitionStreamingQuerySource>& QuerySources, TSet<const UWorldPartitionRuntimeCell*>& OutCells) const;
+	void GetCells(const TArray<FWorldPartitionStreamingSource>& Sources, const class UDataLayerSubsystem* DataLayerSubsystem, TSet<const UWorldPartitionRuntimeCell*>& OutActivateCells, TSet<const UWorldPartitionRuntimeCell*>& OutLoadCells) const;
+	void GetAlwaysLoadedCells(const UDataLayerSubsystem* DataLayerSubsystem, TSet<const UWorldPartitionRuntimeCell*>& OutActivateCells, TSet<const UWorldPartitionRuntimeCell*>& OutLoadCells) const;
 	void Draw2D(UCanvas* Canvas, const TArray<FWorldPartitionStreamingSource>& Sources, const FBox& Region, const FBox2D& GridScreenBounds, TFunctionRef<FVector2D(const FVector2D&)> WorldToScreen) const;
 	void Draw3D(UWorld* World, const TArray<FWorldPartitionStreamingSource>& Sources, const FTransform& Transform) const;
 
@@ -215,7 +216,8 @@ public:
 
 	// streaming interface
 	virtual int32 GetAllStreamingCells(TSet<const UWorldPartitionRuntimeCell*>& Cells, bool bIncludeDataLayers = false) const override;
-	virtual int32 GetStreamingCells(const TArray<FWorldPartitionStreamingSource>& Sources, TSet<const UWorldPartitionRuntimeCell*>& Cells) const override;
+	virtual bool GetStreamingCells(const TArray<FWorldPartitionStreamingQuerySource>& Sources, TSet<const UWorldPartitionRuntimeCell*>& OutCells) const override;
+	virtual bool GetStreamingCells(const TArray<FWorldPartitionStreamingSource>& Sources, TSet<const UWorldPartitionRuntimeCell*>& OutActivateCells, TSet<const UWorldPartitionRuntimeCell*>& OutLoadCells) const override;
 	virtual void SortStreamingCellsByImportance(const TSet<const UWorldPartitionRuntimeCell*>& InCells, const TArray<FWorldPartitionStreamingSource>& InSources, TArray<const UWorldPartitionRuntimeCell*, TInlineAllocator<256>>& OutSortedCells) const override;
 
 protected:
diff --git a/Engine/Source/Runtime/Engine/Public/WorldPartition/WorldPartitionStreamingPolicy.h b/Engine/Source/Runtime/Engine/Public/WorldPartition/WorldPartitionStreamingPolicy.h
--- a/Engine/Source/Runtime/Engine/Public/WorldPartition/WorldPartitionStreamingPolicy.h
+++ b/Engine/Source/Runtime/Engine/Public/WorldPartition/WorldPartitionStreamingPolicy.h
@@ -23,15 +23,15 @@ class UWorldPartitionStreamingPolicy : public UObject
 
 public:
 	virtual void UpdateStreamingState();
-	virtual void LoadCells(const TSet<const UWorldPartitionRuntimeCell*>& ToLoadCells);
-	virtual void UnloadCells(const TSet<const UWorldPartitionRuntimeCell*>& ToUnloadCells);
-	virtual void LoadCell(const UWorldPartitionRuntimeCell* Cell) PURE_VIRTUAL(UWorldPartitionStreamingPolicy::LoadCell, );
-	virtual void UnloadCell(const UWorldPartitionRuntimeCell* Cell) PURE_VIRTUAL(UWorldPartitionStreamingPolicy::UnloadCell, );
+	virtual void SetTargetStateForCells(EWorldPartitionRuntimeCellState TargetState, const TSet<const UWorldPartitionRuntimeCell*>& Cells) PURE_VIRTUAL(UWorldPartitionStreamingPolicy::SetTargetStateForCells, );
+	virtual EWorldPartitionRuntimeCellState GetCurrentStateForCell(const UWorldPartitionRuntimeCell* Cell) const PURE_VIRTUAL(UWorldPartitionStreamingPolicy::GetCurrentStateForCell, return EWorldPartitionRuntimeCellState::Unloaded; );
 	virtual class ULevel* GetPreferredLoadedLevelToAddToWorld() const { return nullptr; }
 	virtual FVector2D GetDrawRuntimeHash2DDesiredFootprint(const FVector2D& CanvasSize);
 	virtual void DrawRuntimeHash2D(class UCanvas* Canvas, const FVector2D& PartitionCanvasOffset, const FVector2D& PartitionCanvasSize);
 	virtual void DrawRuntimeHash3D();
 
+	virtual bool IsStreamingCompleted(EWorldPartitionRuntimeCellState QueryState, const TArray<FWorldPartitionStreamingQuerySource>& QuerySources, bool bExactState = true) const;
+
 #if WITH_EDITOR
 	virtual TSubclassOf<class UWorldPartitionRuntimeCell> GetRuntimeCellClass() const PURE_VIRTUAL(UWorldPartitionStreamingPolicy::GetRuntimeCellClass, return UWorldPartitionRuntimeCell::StaticClass(); );
 
@@ -49,5 +49,6 @@ protected:
 	bool bIsServerLoadingDone;
 	const UWorldPartition* WorldPartition;
 	TSet<const UWorldPartitionRuntimeCell*> LoadedCells;
-	TArray<FWorldPartitionStreamingSource> StreamingSources; // Streaming sources (local to world partition)
+	TSet<const UWorldPartitionRuntimeCell*> ActivatedCells;
+	TArray<FWorldPartitionStreamingSource> StreamingSources;
 };
\ No newline at end of file
diff --git a/Engine/Source/Runtime/Engine/Public/WorldPartition/WorldPartitionStreamingSource.h b/Engine/Source/Runtime/Engine/Public/WorldPartition/WorldPartitionStreamingSource.h
--- a/Engine/Source/Runtime/Engine/Public/WorldPartition/WorldPartitionStreamingSource.h
+++ b/Engine/Source/Runtime/Engine/Public/WorldPartition/WorldPartitionStreamingSource.h
@@ -2,6 +2,66 @@
 #pragma once
 
 #include "CoreMinimal.h"
+#include "WorldPartitionStreamingSource.generated.h"
+
+/**
+ * Streaming Source Target State
+ */
+UENUM()
+enum class EStreamingSourceTargetState : uint8
+{
+	Loaded,
+	Activated
+};
+
+/**
+ * Structure containing all properties required to query a streaming state
+ */
+USTRUCT(BlueprintType)
+struct FWorldPartitionStreamingQuerySource
+{
+	GENERATED_USTRUCT_BODY()
+
+	FWorldPartitionStreamingQuerySource()
+		: Location(FVector::ZeroVector)
+		, Radius(0.f)
+		, bUseLoadingRangeRadius(true)
+		, bDataLayersOnly(false)
+		, bSpatialQuery(true)
+	{}
+
+	FWorldPartitionStreamingQuerySource(const FVector& InLocation)
+		: Location(InLocation)
+		, Radius(0.f)
+		, bUseLoadingRangeRadius(true)
+		, bDataLayersOnly(false)
+		, bSpatialQuery(true)
+	{}
+
+	/* Location to query. (not used if bNonSpatialQuery is true) */
+	UPROPERTY(BlueprintReadWrite, EditAnywhere, Category = "Query")
+	FVector Location;
+
+	/* Radius to query. (not used if bNonSpatialQuery is true) */
+	UPROPERTY(BlueprintReadWrite, EditAnywhere, Category = "Query")
+	float Radius;
+
+	/* If True, Instead of providing a query radius, query can be bound to loading range radius. */
+	UPROPERTY(BlueprintReadWrite, EditAnywhere, Category = "Query")
+	bool bUseLoadingRangeRadius;
+
+	/* Optional list of data layers to specialize the query. If empty only non data layer cells will be returned by the query. */
+	UPROPERTY(BlueprintReadWrite, EditAnywhere, Category = "Query")
+	TArray<FName> DataLayers;
+
+	/* If True, Only cells that are in a data layer found in DataLayers property will be returned by the query. */
+	UPROPERTY(BlueprintReadWrite, EditAnywhere, Category = "Query")
+	bool bDataLayersOnly;
+
+	/* If False, Location/Radius will not be used to find the cells. Only AlwaysLoaded cells will be returned by the query. */
+	UPROPERTY(BlueprintReadWrite, EditAnywhere, Category = "Query")
+	bool bSpatialQuery;
+};
 
 /**
  * Structure containing all properties required to stream from a source
@@ -11,15 +71,17 @@ struct ENGINE_API FWorldPartitionStreamingSource
 	FWorldPartitionStreamingSource()
 	{}
 
-	FWorldPartitionStreamingSource(FName InName, const FVector& InLocation, const FRotator& InRotation)
+	FWorldPartitionStreamingSource(FName InName, const FVector& InLocation, const FRotator& InRotation, EStreamingSourceTargetState InTargetState)
 		: Name(InName)
 		, Location(InLocation)
 		, Rotation(InRotation)
+		, TargetState(InTargetState)
 	{}
 
 	FName Name;
 	FVector Location;
 	FRotator Rotation;
+	EStreamingSourceTargetState TargetState;
 };
 
 /**
diff --git a/Engine/Source/Runtime/Engine/Public/WorldPartition/WorldPartitionSubsystem.h b/Engine/Source/Runtime/Engine/Public/WorldPartition/WorldPartitionSubsystem.h
--- a/Engine/Source/Runtime/Engine/Public/WorldPartition/WorldPartitionSubsystem.h
+++ b/Engine/Source/Runtime/Engine/Public/WorldPartition/WorldPartitionSubsystem.h
@@ -5,12 +5,15 @@
 #include "Subsystems/WorldSubsystem.h"
 #include "Templates/SubclassOf.h"
 #include "Misc/Optional.h"
+#include "WorldPartition/WorldPartitionStreamingSource.h"
 #include "WorldPartitionSubsystem.generated.h"
 
 class UWorldPartition;
 class UWorldPartitionEditorCell;
 class FWorldPartitionActorDesc;
 
+enum class EWorldPartitionRuntimeCellState : uint8;
+
 /**
  * UWorldPartitionSubsystem
  */
@@ -40,6 +43,9 @@ public:
 	virtual TStatId GetStatId() const override;
 	//~End FTickableGameObject
 
+	UFUNCTION(BlueprintCallable, Category = Streaming)
+	bool IsStreamingCompleted(EWorldPartitionRuntimeCellState QueryState, const TArray<FWorldPartitionStreamingQuerySource>& QuerySources, bool bExactState) const;
+
 #if WITH_EDITOR
 	void ForEachIntersectingActorDesc(const FBox& Box, TSubclassOf<AActor> ActorClass, TFunctionRef<bool(const FWorldPartitionActorDesc*)> Predicate) const;
 	void ForEachActorDesc(TSubclassOf<AActor> ActorClass, TFunctionRef<bool(const FWorldPartitionActorDesc*)> Predicate) const;
