Avoid waiting on async static mesh when loading cells

diff --git a/Engine/Source/Runtime/Engine/Private/ActorReferencesUtils.cpp b/Engine/Source/Runtime/Engine/Private/ActorReferencesUtils.cpp
--- a/Engine/Source/Runtime/Engine/Private/ActorReferencesUtils.cpp
+++ b/Engine/Source/Runtime/Engine/Private/ActorReferencesUtils.cpp
@@ -18,6 +18,10 @@ public:
 		ArIsObjectReferenceCollector = true;
 		ArShouldSkipBulkData = true;
 
+		// Avoid having to wait until compilable asset are built just to scan for actors. 
+		// Assets like Textures/Shaders/StaticMesh should not refer to actors anyway.
+		SetShouldSkipCompilingAssets(true);
+
 		SubObjects.Add(Root);
 
 		Root->Serialize(*this);
