Avoid calling GetAttachChildren() in USceneComponent::AttachToComponent(), as it's now quite expensive due to access tracking. Use AttachChildren directly.

diff --git a/Engine/Source/Runtime/Engine/Private/Components/SceneComponent.cpp b/Engine/Source/Runtime/Engine/Private/Components/SceneComponent.cpp
--- a/Engine/Source/Runtime/Engine/Private/Components/SceneComponent.cpp
+++ b/Engine/Source/Runtime/Engine/Private/Components/SceneComponent.cpp
@@ -1832,8 +1832,10 @@ bool USceneComponent::AttachToComponent(USceneComponent* Parent, const FAttachme
 
 	if(Parent != nullptr)
 	{
+		const int32 LastAttachIndex = Parent->AttachChildren.Find(TObjectPtr<USceneComponent>(this));
+
 		const bool bSameAttachParentAndSocket = (Parent == GetAttachParent() && SocketName == GetAttachSocketName());
-		if (bSameAttachParentAndSocket && Parent->GetAttachChildren().Contains(this))
+		if (bSameAttachParentAndSocket && LastAttachIndex != INDEX_NONE)
 		{
 			// already attached!
 			return true;
@@ -1907,10 +1909,6 @@ bool USceneComponent::AttachToComponent(USceneComponent* Parent, const FAttachme
 		const bool bSavedDisableDetachmentUpdateOverlaps = bDisableDetachmentUpdateOverlaps;
 		bDisableDetachmentUpdateOverlaps = true;
 
-		// Find out if we're already attached, and save off our position in the array if we are
-		int32 LastAttachIndex = INDEX_NONE;
-		Parent->GetAttachChildren().Find(this, LastAttachIndex);
-
 		if (!ShouldSkipUpdateOverlaps())	//if we can't skip UpdateOverlaps, make sure the parent doesn't either
 		{
 			Parent->ClearSkipUpdateOverlaps();
