Fix for Drive component ordering in LODSyncComponent
#rnx

diff --git a/Engine/Source/Runtime/Engine/Private/Components/LODSyncComponent.cpp b/Engine/Source/Runtime/Engine/Private/Components/LODSyncComponent.cpp
--- a/Engine/Source/Runtime/Engine/Private/Components/LODSyncComponent.cpp
+++ b/Engine/Source/Runtime/Engine/Private/Components/LODSyncComponent.cpp
@@ -62,8 +62,17 @@ void ULODSyncComponent::InitializeSyncComponents()
 	SubComponents.Reset();
 
 	AActor* Owner = GetOwner();
-	// for now we only support skinnedmeshcomponent
-	TArray<UActorComponent*> AllComponents = Owner->GetComponentsByInterface(ULODSyncInterface::StaticClass());
+	
+	// We only support primitive components that implement the ILODSyncInterface.
+	TInlineComponentArray<UPrimitiveComponent*, 16> PrimitiveComponents(Owner);
+	TMap<FName, int32, TInlineSetAllocator<16>> ComponentNameToIndexMap;
+
+	const int32 NumPrimitiveComponents = PrimitiveComponents.Num();
+	for (int32 Index = 0; Index < NumPrimitiveComponents; ++Index)
+	{
+		FName ComponentName = PrimitiveComponents[Index]->GetFName();
+		ComponentNameToIndexMap.Add(ComponentName, Index);
+	}
 
 	// current num LODs start with NumLODs
 	// but if nothing is set, it will be -1
@@ -71,30 +80,32 @@ void ULODSyncComponent::InitializeSyncComponents()
 	// if NumLODs are -1, we try to find the max number of LODs of all the components
 	const bool bFindTheMaxLOD = (NumLODs == -1);
 	// we find all the components of the child and add this to prerequisite
-	for (UActorComponent* Component : AllComponents)
+	for (const FComponentSync& CompSync : ComponentsToSync)
 	{
-		UPrimitiveComponent* PrimComponent = Cast<UPrimitiveComponent>(Component);
-		if (PrimComponent)
+		if (const int32* ComponentIndex = ComponentNameToIndexMap.Find(CompSync.Name))
 		{
-			FName Name = PrimComponent->GetFName();
-			ILODSyncInterface* LODInterface = Cast<ILODSyncInterface>(PrimComponent);
-			const FComponentSync* CompSync = GetComponentSync(Name);
-			if (LODInterface && CompSync && CompSync->SyncOption != ESyncOption::Disabled)
+			UPrimitiveComponent* PrimComponent = PrimitiveComponents[*ComponentIndex];
+			if (PrimComponent)
 			{
-				PrimComponent->PrimaryComponentTick.AddPrerequisite(this, PrimaryComponentTick);
-				SubComponents.Add(PrimComponent);
-
-				if (CompSync->SyncOption == ESyncOption::Drive)
+				// We don't need to check for ImplementsInterface() since it's tagged as CannotImplementInterfaceInBlueprint
+				ILODSyncInterface* LODInterface = Cast<ILODSyncInterface>(PrimComponent);
+				if (LODInterface && CompSync.SyncOption != ESyncOption::Disabled)
 				{
-					DriveComponents.Add(PrimComponent);
+					PrimComponent->PrimaryComponentTick.AddPrerequisite(this, PrimaryComponentTick);
+					SubComponents.Add(PrimComponent);
 
-					const int LODCount = LODInterface->GetNumSyncLODs();
-					UE_LOG(LogLODSync, Verbose, TEXT("Adding new component (%s - LODCount : %d ) to sync."), *Name.ToString(), LODCount);
-					if (bFindTheMaxLOD)
+					if (CompSync.SyncOption == ESyncOption::Drive)
 					{
-						// get max lod
-						CurrentNumLODs = FMath::Max(CurrentNumLODs, LODCount);
-						UE_LOG(LogLODSync, Verbose, TEXT("MaxLOD now is set to (%d )."), CurrentNumLODs);
+						DriveComponents.Add(PrimComponent);
+
+						const int LODCount = LODInterface->GetNumSyncLODs();
+						UE_LOG(LogLODSync, Verbose, TEXT("Adding new component (%s - LODCount : %d ) to sync."), *CompSync.Name.ToString(), LODCount);
+						if (bFindTheMaxLOD)
+						{
+							// get max lod
+							CurrentNumLODs = FMath::Max(CurrentNumLODs, LODCount);
+							UE_LOG(LogLODSync, Verbose, TEXT("MaxLOD now is set to (%d )."), CurrentNumLODs);
+						}
 					}
 				}
 			}
