Implements r.ScreenPercentage.MaxResolution to control max resolution on desktop to targeted average resolution on consoles

diff --git a/Engine/Source/Runtime/Engine/Private/GameViewportClient.cpp b/Engine/Source/Runtime/Engine/Private/GameViewportClient.cpp
--- a/Engine/Source/Runtime/Engine/Private/GameViewportClient.cpp
+++ b/Engine/Source/Runtime/Engine/Private/GameViewportClient.cpp
@@ -115,6 +115,15 @@ static TAutoConsoleVariable<float> CVarSecondaryScreenPercentage( // TODO: make
 	TEXT(" 1: override secondary screen percentage."),
 	ECVF_Default);
 
+static TAutoConsoleVariable<float> CVarScreenPercentageMaxResolution(
+	TEXT("r.ScreenPercentage.MaxResolution"),
+	0.0f,
+	TEXT("Controls maximum number of rendered pixel before any upscaling such that doesn't go higher than the specified 16:9 resolution ")
+	TEXT("of this variable. For instance set this value to 1440 so that you are not rendering more than 2560x1440 = 3.6M pixels. ")
+	TEXT("This is useful to set this on PC in your project's DefaultEditor.ini so you are not rendering more pixel on PC in PIE that you would ")
+	TEXT("in average on console with your project specific dynamic resolution settings."),
+	ECVF_Default);
+
 #if CSV_PROFILER
 struct FCsvLocalPlayer
 {
@@ -1397,27 +1406,6 @@ void UGameViewportClient::Draw(FViewport* InViewport, FCanvas* SceneCanvas)
 			GCustomStaticScreenPercentage->SetupMainGameViewFamily(ViewFamily);
 		}
 
-		// If a screen percentage interface was not set by dynamic resolution, then create one matching legacy behavior.
-		if (ViewFamily.GetScreenPercentageInterface() == nullptr)
-		{
-			bool AllowPostProcessSettingsScreenPercentage = false;
-			float GlobalResolutionFraction = 1.0f;
-
-			if (ViewFamily.EngineShowFlags.ScreenPercentage)
-			{
-				// Allow FPostProcessSettings::ScreenPercentage.
-				AllowPostProcessSettingsScreenPercentage = true;
-
-				// Get global view fraction set by r.ScreenPercentage.
-				GlobalResolutionFraction = FLegacyScreenPercentageDriver::GetCVarResolutionFraction();
-			}
-
-			ViewFamily.SetScreenPercentageInterface(new FLegacyScreenPercentageDriver(
-				ViewFamily, GlobalResolutionFraction, AllowPostProcessSettingsScreenPercentage));
-		}
-
-		check(ViewFamily.GetScreenPercentageInterface() != nullptr);
-
 		bFinalScreenPercentageShowFlag = ViewFamily.EngineShowFlags.ScreenPercentage;
 	}
 
@@ -1619,6 +1607,53 @@ void UGameViewportClient::Draw(FViewport* InViewport, FCanvas* SceneCanvas)
 	}
 
 	{
+		// If a screen percentage interface was not set by dynamic resolution, then create one matching legacy behavior.
+		if (ViewFamily.GetScreenPercentageInterface() == nullptr)
+		{
+			bool AllowPostProcessSettingsScreenPercentage = false;
+			float GlobalResolutionFraction = 1.0f;
+
+			if (ViewFamily.EngineShowFlags.ScreenPercentage)
+			{
+				// Allow FPostProcessSettings::ScreenPercentage.
+				AllowPostProcessSettingsScreenPercentage = true;
+
+				// Get global view fraction set by r.ScreenPercentage.
+				GlobalResolutionFraction = FLegacyScreenPercentageDriver::GetCVarResolutionFraction();
+
+				// Max the rendering resolution to average target resolution used on platform that have dynamic resolution.
+				float MaxRenderingHeight = CVarScreenPercentageMaxResolution.GetValueOnGameThread();
+				if (MaxRenderingHeight > 0.0f)
+				{
+					int32 MaxRenderingWidth = MaxRenderingHeight * (1920.0f / 1080.0f);
+					int32 MaxRenderingPixelCount = MaxRenderingWidth * MaxRenderingHeight;
+
+					int32 TotalRenderingPixelCount = 0;
+
+					for (const FSceneView* View : ViewFamily.Views)
+					{
+						// Number of pixel drawn in ViewFamily.RenderTarget
+						int32 DisplayPixelCount = View->CameraConstrainedViewRect.Area();
+
+						// Number of pixel drawn before secondary spatial upscale.
+						int32 PrimaryScreenPercentageUpscaledPixelCount = DisplayPixelCount * (ViewFamily.SecondaryViewFraction * ViewFamily.SecondaryViewFraction);
+
+						TotalRenderingPixelCount += PrimaryScreenPercentageUpscaledPixelCount;
+					}
+
+					// Compute max resolution fraction such that the total number of pixel doesn't go over the CVarScreenPercentageMaxResolution.
+					float MaxGlobalResolutionFraction = FMath::Sqrt(float(MaxRenderingPixelCount) / float(TotalRenderingPixelCount));
+				
+					GlobalResolutionFraction = FMath::Min(GlobalResolutionFraction, MaxGlobalResolutionFraction);
+				}
+			}
+
+			ViewFamily.SetScreenPercentageInterface(new FLegacyScreenPercentageDriver(
+				ViewFamily, GlobalResolutionFraction, AllowPostProcessSettingsScreenPercentage));
+		}
+
+		check(ViewFamily.GetScreenPercentageInterface() != nullptr);
+
 		// Make sure the engine show flag for screen percentage is still what it was when setting up the screen percentage interface
 		ViewFamily.EngineShowFlags.ScreenPercentage = bFinalScreenPercentageShowFlag;
 
