Fix the assert to take account for the failed jobs.

diff --git a/Engine/Source/Runtime/Engine/Private/ShaderCompiler/ShaderCompiler.cpp b/Engine/Source/Runtime/Engine/Private/ShaderCompiler/ShaderCompiler.cpp
--- a/Engine/Source/Runtime/Engine/Private/ShaderCompiler/ShaderCompiler.cpp
+++ b/Engine/Source/Runtime/Engine/Private/ShaderCompiler/ShaderCompiler.cpp
@@ -6530,7 +6530,7 @@ void FShaderCompileJob::SerializeOutput(FArchive& Ar)
 
 	Ar << Output;
 	// output hash is now serialized as part of the output, as the shader code is compressed in SCWs
-	checkf(Output.OutputHash != FSHAHash(), TEXT("Compile job does not have an OutputHash generated."));
+	checkf(!Output.bSucceeded || Output.OutputHash != FSHAHash(), TEXT("Successful compile job does not have an OutputHash generated."));
 
 	if (Ar.IsLoading())
 	{
