Disable APackedLevelInstance RerunConstructionScripts optimization until we confirm it wasn't causing an issue.

diff --git a/Engine/Source/Runtime/Engine/Private/LevelInstance/Packed/PackedLevelInstanceActor.cpp b/Engine/Source/Runtime/Engine/Private/LevelInstance/Packed/PackedLevelInstanceActor.cpp
--- a/Engine/Source/Runtime/Engine/Private/LevelInstance/Packed/PackedLevelInstanceActor.cpp
+++ b/Engine/Source/Runtime/Engine/Private/LevelInstance/Packed/PackedLevelInstanceActor.cpp
@@ -175,7 +175,7 @@ void APackedLevelInstance::DestroyPackedComponents()
 
 void APackedLevelInstance::RerunConstructionScripts()
 {
-	if (bRerunConstructionScripts)
+	//if (bRerunConstructionScripts)
 	{
 		Super::RerunConstructionScripts();
 	}
