ShaderCompiler: do not catch crashes in compiling threads.

- This created a few deadlocks on the builder when the compiling thread crashed while holding a lock. The main thread has never gotten to the point to call CheckHealth() because it tried to take the same lock.
- Pass-through the crash instead of trying to handle it in a somewhat peculiar fashion (catch and then crash on another thread). Crash handler should be able to deal with crashes on non-main threads.

diff --git a/Engine/Source/Runtime/Engine/Private/ShaderCompiler/ShaderCompiler.cpp b/Engine/Source/Runtime/Engine/Private/ShaderCompiler/ShaderCompiler.cpp
--- a/Engine/Source/Runtime/Engine/Private/ShaderCompiler/ShaderCompiler.cpp
+++ b/Engine/Source/Runtime/Engine/Private/ShaderCompiler/ShaderCompiler.cpp
@@ -1800,7 +1800,6 @@ FShaderCompileThreadRunnableBase::FShaderCompileThreadRunnableBase(FShaderCompil
 	: Manager(InManager)
 	, MinPriorityIndex(0)
 	, MaxPriorityIndex(NumShaderCompileJobPriorities - 1)
-	, bTerminatedByError(false)
 	, bForceFinish(false)
 {
 }
@@ -1835,67 +1834,16 @@ FShaderCompileThreadRunnable::~FShaderCompileThreadRunnable()
 /** Entry point for the shader compiling thread. */
 uint32 FShaderCompileThreadRunnableBase::Run()
 {
-#ifdef _MSC_VER
-	if(!FPlatformMisc::IsDebuggerPresent())
+	check(Manager->bAllowAsynchronousShaderCompiling);
+	while (!bForceFinish)
 	{
-#if !PLATFORM_SEH_EXCEPTIONS_DISABLED
-		__try
-#endif
-		{
-			check(Manager->bAllowAsynchronousShaderCompiling);
-			// Do the work
-			while (!bForceFinish)
-			{
-				CompilingLoop();
-			}
-		}
-#if !PLATFORM_SEH_EXCEPTIONS_DISABLED
-		__except( 
-#if PLATFORM_WINDOWS
-			ReportCrash( GetExceptionInformation() )
-#else
-			EXCEPTION_EXECUTE_HANDLER
-#endif
-			)
-		{
-#if WITH_EDITORONLY_DATA
-			ErrorMessage = GErrorHist;
-#endif
-			// Use a memory barrier to ensure that the main thread sees the write to ErrorMessage before
-			// the write to bTerminatedByError.
-			FPlatformMisc::MemoryBarrier();
-
-			bTerminatedByError = true;
-		}
-#endif
-	}
-	else
-#endif
-	{
-		check(Manager->bAllowAsynchronousShaderCompiling);
-		while (!bForceFinish)
-		{
-			CompilingLoop();
-		}
+		CompilingLoop();
 	}
 	UE_LOG(LogShaderCompilers, Display, TEXT("Shaders left to compile 0"));
 
 	return 0;
 }
 
-/** Called by the main thread only, reports exceptions in the worker threads */
-void FShaderCompileThreadRunnableBase::CheckHealth() const
-{
-	if (bTerminatedByError)
-	{
-#if WITH_EDITORONLY_DATA
-		GErrorHist[0] = 0;
-#endif
-		GIsCriticalError = false;
-		UE_LOG(LogShaderCompilers, Fatal,TEXT("Shader Compiling thread exception:\r\n%s"), *ErrorMessage);
-	}
-}
-
 int32 FShaderCompileThreadRunnable::PullTasksFromQueue()
 {
 	int32 NumActiveThreads = 0;
@@ -3298,10 +3246,6 @@ void FShaderCompilingManager::BlockOnShaderMapCompletion(const TArray<int32>& Sh
 		int32 LogCounter = 0;
 		do 
 		{
-			for (const auto& Thread : Threads)
-			{
-				Thread->CheckHealth();
-			}
 			NumPendingJobs = 0;
 			{
 				// Lock CompileQueueSection so we can access the input and output queues
@@ -3419,10 +3363,6 @@ void FShaderCompilingManager::BlockOnAllShaderMapCompletion(TMap<int32, FShaderM
 
 		do 
 		{
-			for (const auto& Thread : Threads)
-			{
-				Thread->CheckHealth();
-			}
 			NumPendingJobs = 0;
 			{
 				// Lock CompileQueueSection so we can access the input and output queues
@@ -4216,11 +4156,6 @@ void FShaderCompilingManager::ProcessAsyncResults(bool bLimitExecutionTime, bool
 	check(IsInGameThread());
 	if (bAllowAsynchronousShaderCompiling)
 	{
-		for (const auto& Thread : Threads)
-		{
-			Thread->CheckHealth();
-		}
-
 		{
 			const double StartTime = FPlatformTime::Seconds();
 
diff --git a/Engine/Source/Runtime/Engine/Public/ShaderCompiler.h b/Engine/Source/Runtime/Engine/Public/ShaderCompiler.h
--- a/Engine/Source/Runtime/Engine/Public/ShaderCompiler.h
+++ b/Engine/Source/Runtime/Engine/Public/ShaderCompiler.h
@@ -544,11 +544,6 @@ protected:
 	int32 MinPriorityIndex;
 	int32 MaxPriorityIndex;
 	
-	/** If the thread has been terminated by an unhandled exception, this contains the error message. */
-	FString ErrorMessage;
-	/** true if the thread has been terminated by an unhandled exception. */
-	bool bTerminatedByError;
-
 	TAtomic<bool> bForceFinish;
 
 public:
@@ -576,9 +571,6 @@ public:
 		}
 	}
 
-	/** Checks the thread's health, and passes on any errors that have occured.  Called by the main thread. */
-	void CheckHealth() const;
-
 	/** Main work loop. */
 	virtual int32 CompilingLoop() = 0;
 };
