FCanvasRenderer optims/fixes :
* Batched canvas tiles/triangles that share the same hit proxy id : this saves thousands of render graph passes on large batches
* Added support for more than 32k tiles / 64k triangles

diff --git a/Engine/Source/Runtime/Engine/Private/TileRendering.cpp b/Engine/Source/Runtime/Engine/Private/TileRendering.cpp
--- a/Engine/Source/Runtime/Engine/Private/TileRendering.cpp
+++ b/Engine/Source/Runtime/Engine/Private/TileRendering.cpp
@@ -19,7 +19,7 @@
 #include "EngineModule.h"
 #include "MeshPassProcessor.h"
 
-DECLARE_GPU_STAT_NAMED(CanvasDrawTile, TEXT("CanvasDrawTile"));
+DECLARE_GPU_STAT_NAMED(CanvasDrawTiles, TEXT("CanvasDrawTiles"));
 
 static const uint32 CanvasTileVertexCount = 4;
 static const uint32 CanvasTileIndexCount = 6;
@@ -44,27 +44,26 @@ void FCanvasTileRendererItem::FTileVertexFactory::InitResource()
 	FLocalVertexFactory::InitResource();
 }
 
-FCanvasTileRendererItem::FTileMesh::FTileMesh(
-	const FRawIndexBuffer* InIndexBuffer,
-	const FCanvasTileRendererItem::FTileVertexFactory* InVertexFactory)
-	: IndexBuffer(InIndexBuffer)
-	, VertexFactory(InVertexFactory)
-{}
-
-void FCanvasTileRendererItem::FTileMesh::InitRHI()
+FMeshBatch* FCanvasTileRendererItem::FRenderData::AllocTileMeshBatch(FCanvasRenderContext& InRenderContext, FHitProxyId InHitProxyId)
 {
-	MeshElement.VertexFactory = VertexFactory;
-	MeshElement.ReverseCulling = false;
-	MeshElement.Type = PT_TriangleList;
-	MeshElement.DepthPriorityGroup = SDPG_Foreground;
-
-	FMeshBatchElement& BatchElement = MeshElement.Elements[0];
-	BatchElement.IndexBuffer = IndexBuffer;
-	BatchElement.FirstIndex = 0;
-	BatchElement.NumPrimitives = 2;
-	BatchElement.MinVertexIndex = 0;
-	BatchElement.MaxVertexIndex = CanvasTileVertexCount - 1;
-	BatchElement.PrimitiveUniformBufferResource = &GIdentityPrimitiveUniformBuffer;
+	FMeshBatch* MeshBatch = InRenderContext.Alloc<FMeshBatch>();
+
+	MeshBatch->VertexFactory = &VertexFactory;
+	MeshBatch->MaterialRenderProxy = MaterialRenderProxy;
+	MeshBatch->ReverseCulling = false;
+	MeshBatch->Type = PT_TriangleList;
+	MeshBatch->DepthPriorityGroup = SDPG_Foreground;
+	MeshBatch->BatchHitProxyId = InHitProxyId;
+
+	FMeshBatchElement& MeshBatchElement = MeshBatch->Elements[0];
+	MeshBatchElement.IndexBuffer = &IndexBuffer;
+	MeshBatchElement.FirstIndex = 0;
+	MeshBatchElement.NumPrimitives = 0;
+	MeshBatchElement.MinVertexIndex = 0;
+	MeshBatchElement.MaxVertexIndex = GetNumVertices() - 1;
+	MeshBatchElement.PrimitiveUniformBufferResource = &GIdentityPrimitiveUniformBuffer;
+
+	return MeshBatch;
 }
 
 FCanvasTileRendererItem::FRenderData::FRenderData(
@@ -74,21 +73,33 @@ FCanvasTileRendererItem::FRenderData::FRenderData(
 	: MaterialRenderProxy(InMaterialRenderProxy)
 	, Transform(InTransform)
 	, VertexFactory(&StaticMeshVertexBuffers, InFeatureLevel)
-	, TileMesh(&IndexBuffer, &VertexFactory)
 {}
 
+uint32 FCanvasTileRendererItem::FRenderData::GetNumVertices() const
+{
+	return Tiles.Num() * CanvasTileVertexCount;
+}
+
+uint32 FCanvasTileRendererItem::FRenderData::GetNumIndices() const
+{
+	return Tiles.Num() * CanvasTileIndexCount;
+}
+
 void FCanvasTileRendererItem::FRenderData::InitTileMesh(const FSceneView& View, bool bNeedsToSwitchVerticalAxis)
 {
 	static_assert(CanvasTileVertexCount == 4, "Invalid tile tri-list size.");
 	static_assert(CanvasTileIndexCount == 6, "Invalid tile tri-list size.");
 
-	const uint32 TotalVertexCount = Tiles.Num() * CanvasTileVertexCount;
-	const uint32 TotalIndexCount = Tiles.Num() * CanvasTileIndexCount;
+	const uint32 TotalVertexCount = GetNumVertices();
+	const uint32 TotalIndexCount = GetNumIndices();
 
 	StaticMeshVertexBuffers.PositionVertexBuffer.Init(TotalVertexCount);
 	StaticMeshVertexBuffers.StaticMeshVertexBuffer.Init(TotalVertexCount, 1);
 	StaticMeshVertexBuffers.ColorVertexBuffer.Init(TotalVertexCount);
+	
 	IndexBuffer.Indices.SetNum(TotalIndexCount);
+	// Make sure the index buffer is using the appropriate size :
+	IndexBuffer.ForceUse32Bit(TotalVertexCount > MAX_uint16);
 
 	for (int32 i = 0; i < Tiles.Num(); i++)
 	{
@@ -149,12 +160,10 @@ void FCanvasTileRendererItem::FRenderData::InitTileMesh(const FSceneView& View,
 	StaticMeshVertexBuffers.ColorVertexBuffer.InitResource();
 	IndexBuffer.InitResource();
 	VertexFactory.InitResource();
-	TileMesh.InitResource();
 }
 
 void FCanvasTileRendererItem::FRenderData::ReleaseTileMesh()
 {
-	TileMesh.ReleaseResource();
 	VertexFactory.ReleaseResource();
 	IndexBuffer.ReleaseResource();
 	StaticMeshVertexBuffers.PositionVertexBuffer.ReleaseResource();
@@ -172,26 +181,46 @@ void FCanvasTileRendererItem::FRenderData::RenderTiles(
 {
 	check(IsInRenderingThread());
 
-	RDG_GPU_STAT_SCOPE(RenderContext.GraphBuilder, CanvasDrawTile);
+	if (Tiles.Num() == 0)
+	{
+		return;
+	}
+
+	RDG_GPU_STAT_SCOPE(RenderContext.GraphBuilder, CanvasDrawTiles);
 	RDG_EVENT_SCOPE(RenderContext.GraphBuilder, "%s", *MaterialRenderProxy->GetIncompleteMaterialWithFallback(GMaxRHIFeatureLevel).GetFriendlyName());
-	TRACE_CPUPROFILER_EVENT_SCOPE(CanvasDrawTile);
-	QUICK_SCOPE_CYCLE_COUNTER(STAT_CanvasDrawTile)
+	TRACE_CPUPROFILER_EVENT_SCOPE(CanvasDrawTiles);
+	QUICK_SCOPE_CYCLE_COUNTER(STAT_CanvasDrawTiles)
 
 	IRendererModule& RendererModule = GetRendererModule();
 
 	InitTileMesh(View, bNeedsToSwitchVerticalAxis);
 
+	// We know we have at least 1 tile so prep up a new batch right away : 
+	FMeshBatch* CurrentMeshBatch = AllocTileMeshBatch(RenderContext, Tiles[0].HitProxyId);
+	check(CurrentMeshBatch->Elements[0].FirstIndex == 0); // The first batch should always start at the first index 
+
 	for (int32 TileIdx = 0; TileIdx < Tiles.Num(); TileIdx++)
 	{
-		FRenderData::FTileInst& Tile = Tiles[TileIdx];
+		const FTileInst& Tile = Tiles[TileIdx];
 
-		FMeshBatch& Mesh = TileMesh.MeshElement;
-		Mesh.MaterialRenderProxy = MaterialRenderProxy;
-		Mesh.Elements[0].FirstIndex = CanvasTileIndexCount * TileIdx;
+		// We only need a new batch when the hit proxy id changes : 
+		if (CurrentMeshBatch->BatchHitProxyId != Tile.HitProxyId)
+		{
+			// Flush the current batch before allocating a new one: 
+			GetRendererModule().DrawTileMesh(RenderContext, DrawRenderState, View, *CurrentMeshBatch, bIsHitTesting, CurrentMeshBatch->BatchHitProxyId);
+
+			CurrentMeshBatch = AllocTileMeshBatch(RenderContext, Tile.HitProxyId);
+			CurrentMeshBatch->Elements[0].FirstIndex = CanvasTileIndexCount * TileIdx;
+		}
 
-		RendererModule.DrawTileMesh(RenderContext, DrawRenderState, View, Mesh, bIsHitTesting, Tile.HitProxyId, bUse128bitRT);
+		// Add 2 triangles to the batch per tile : 
+		CurrentMeshBatch->Elements[0].NumPrimitives += 2;
 	}
 
+	// Flush the final batch: 
+	check(CurrentMeshBatch != nullptr);
+	GetRendererModule().DrawTileMesh(RenderContext, DrawRenderState, View, *CurrentMeshBatch, bIsHitTesting, CurrentMeshBatch->BatchHitProxyId);
+
 	AddPass(RenderContext.GraphBuilder, [this](FRHICommandList&)
 	{
 		ReleaseTileMesh();
diff --git a/Engine/Source/Runtime/Engine/Private/TriangleRendering.cpp b/Engine/Source/Runtime/Engine/Private/TriangleRendering.cpp
--- a/Engine/Source/Runtime/Engine/Private/TriangleRendering.cpp
+++ b/Engine/Source/Runtime/Engine/Private/TriangleRendering.cpp
@@ -19,6 +19,8 @@
 #include "EngineModule.h"
 #include "MeshPassProcessor.h"
 
+DECLARE_GPU_STAT_NAMED(CanvasDrawTriangles, TEXT("CanvasDrawTriangles"));
+
 FCanvasTriangleRendererItem::FTriangleVertexFactory::FTriangleVertexFactory(
 	const FStaticMeshVertexBuffers* InVertexBuffers,
 	ERHIFeatureLevel::Type InFeatureLevel)
@@ -26,13 +28,6 @@ FCanvasTriangleRendererItem::FTriangleVertexFactory::FTriangleVertexFactory(
 	, VertexBuffers(InVertexBuffers)
 {}
 
-FCanvasTriangleRendererItem::FTriangleMesh::FTriangleMesh(
-	const FRawIndexBuffer* InIndexBuffer,
-	const FTriangleVertexFactory* InVertexFactory)
-	: IndexBuffer(InIndexBuffer)
-	, VertexFactory(InVertexFactory)
-{}
-
 void FCanvasTriangleRendererItem::FTriangleVertexFactory::InitResource()
 {
 	FLocalVertexFactory::FDataType VertexData;
@@ -46,30 +41,50 @@ void FCanvasTriangleRendererItem::FTriangleVertexFactory::InitResource()
 	FLocalVertexFactory::InitResource();
 }
 
-void FCanvasTriangleRendererItem::FTriangleMesh::InitRHI()
+FMeshBatch* FCanvasTriangleRendererItem::FRenderData::AllocTriangleMeshBatch(FCanvasRenderContext& InRenderContext, FHitProxyId InHitProxyId)
 {
-	MeshBatch.VertexFactory = VertexFactory;
-	MeshBatch.ReverseCulling = false;
-	MeshBatch.bDisableBackfaceCulling = true;
-	MeshBatch.Type = PT_TriangleList;
-	MeshBatch.DepthPriorityGroup = SDPG_Foreground;
-
-	FMeshBatchElement& MeshBatchElement = MeshBatch.Elements[0];
-	MeshBatchElement.IndexBuffer = IndexBuffer;
+	FMeshBatch* MeshBatch = InRenderContext.Alloc<FMeshBatch>();
+
+	MeshBatch->VertexFactory = &VertexFactory;
+	MeshBatch->MaterialRenderProxy = MaterialRenderProxy;
+	MeshBatch->ReverseCulling = false;
+	MeshBatch->bDisableBackfaceCulling = true;
+	MeshBatch->Type = PT_TriangleList;
+	MeshBatch->DepthPriorityGroup = SDPG_Foreground;
+	MeshBatch->BatchHitProxyId = InHitProxyId;
+
+	FMeshBatchElement& MeshBatchElement = MeshBatch->Elements[0];
+	MeshBatchElement.IndexBuffer = &IndexBuffer;
 	MeshBatchElement.FirstIndex = 0;
-	MeshBatchElement.NumPrimitives = 1;
+	MeshBatchElement.NumPrimitives = 0;
 	MeshBatchElement.MinVertexIndex = 0;
-	MeshBatchElement.MaxVertexIndex = 2;
+	MeshBatchElement.MaxVertexIndex = GetNumVertices() - 1;
 	MeshBatchElement.PrimitiveUniformBufferResource = &GIdentityPrimitiveUniformBuffer;
+
+	return MeshBatch;
+}
+
+uint32 FCanvasTriangleRendererItem::FRenderData::GetNumVertices() const
+{
+	return Triangles.Num() * 3;
+}
+
+uint32 FCanvasTriangleRendererItem::FRenderData::GetNumIndices() const
+{
+	return Triangles.Num() * 3;
 }
 
 void FCanvasTriangleRendererItem::FRenderData::InitTriangleMesh(const FSceneView& View, bool bNeedsToSwitchVerticalAxis)
 {
-	StaticMeshVertexBuffers.PositionVertexBuffer.Init(Triangles.Num() * 3);
-	StaticMeshVertexBuffers.StaticMeshVertexBuffer.Init(Triangles.Num() * 3, 1);
-	StaticMeshVertexBuffers.ColorVertexBuffer.Init(Triangles.Num() * 3);
+	const uint32 NumIndices = GetNumIndices();
+	const uint32 NumVertices = GetNumIndices();
+	StaticMeshVertexBuffers.PositionVertexBuffer.Init(NumVertices);
+	StaticMeshVertexBuffers.StaticMeshVertexBuffer.Init(NumVertices, 1);
+	StaticMeshVertexBuffers.ColorVertexBuffer.Init(NumVertices);
 
-	IndexBuffer.Indices.SetNum(Triangles.Num() * 3);
+	IndexBuffer.Indices.SetNum(NumIndices);
+	// Make sure the index buffer is using the appropriate size :
+	IndexBuffer.ForceUse32Bit(NumVertices > MAX_uint16);
 
 	for (int32 i = 0; i < Triangles.Num(); i++)
 	{
@@ -116,12 +131,10 @@ void FCanvasTriangleRendererItem::FRenderData::InitTriangleMesh(const FSceneView
 	StaticMeshVertexBuffers.ColorVertexBuffer.InitResource();
 	IndexBuffer.InitResource();
 	VertexFactory.InitResource();
-	TriMesh.InitResource();
 };
 
 void FCanvasTriangleRendererItem::FRenderData::ReleaseTriangleMesh()
 {
-	TriMesh.ReleaseResource();
 	VertexFactory.ReleaseResource();
 	IndexBuffer.ReleaseResource();
 	StaticMeshVertexBuffers.PositionVertexBuffer.ReleaseResource();
@@ -138,24 +151,47 @@ void FCanvasTriangleRendererItem::FRenderData::RenderTriangles(
 {
 	check(IsInRenderingThread());
 
+	if (Triangles.Num() == 0)
+	{
+		return;
+	}
+
+	RDG_GPU_STAT_SCOPE(RenderContext.GraphBuilder, CanvasDrawTriangles);
+	RDG_EVENT_SCOPE(RenderContext.GraphBuilder, "%s", *MaterialRenderProxy->GetIncompleteMaterialWithFallback(GMaxRHIFeatureLevel).GetFriendlyName());
+	TRACE_CPUPROFILER_EVENT_SCOPE(CanvasDrawTriangles);
+	QUICK_SCOPE_CYCLE_COUNTER(STAT_CanvasDrawTriangles)
+
 	IRendererModule& RendererModule = GetRendererModule();
 
 	InitTriangleMesh(View, bNeedsToSwitchVerticalAxis);
 
-	RDG_EVENT_SCOPE(RenderContext.GraphBuilder, "%s", *MaterialRenderProxy->GetIncompleteMaterialWithFallback(GMaxRHIFeatureLevel).GetFriendlyName());
+
+	// We know we have at least 1 triangle so prep up a new batch right away : 
+	FMeshBatch* CurrentMeshBatch = AllocTriangleMeshBatch(RenderContext, Triangles[0].HitProxyId);
+	check (CurrentMeshBatch->Elements[0].FirstIndex == 0); // The first batch should always start at the first index 
 
 	for (int32 TriIdx = 0; TriIdx < Triangles.Num(); TriIdx++)
 	{
 		const FTriangleInst& Tri = Triangles[TriIdx];
 
-		FMeshBatch& MeshBatch = *RenderContext.Alloc<FMeshBatch>(TriMesh.MeshBatch);
-		MeshBatch.VertexFactory = &VertexFactory;
-		MeshBatch.MaterialRenderProxy = MaterialRenderProxy;
-		MeshBatch.Elements[0].FirstIndex = 3 * TriIdx;
+		// We only need a new batch when the hit proxy id changes : 
+		if (CurrentMeshBatch->BatchHitProxyId != Tri.HitProxyId)
+		{
+			// Flush the current batch before allocating a new one: 
+			GetRendererModule().DrawTileMesh(RenderContext, DrawRenderState, View, *CurrentMeshBatch, bIsHitTesting, CurrentMeshBatch->BatchHitProxyId);
+
+			CurrentMeshBatch = AllocTriangleMeshBatch(RenderContext, Tri.HitProxyId);
+			CurrentMeshBatch->Elements[0].FirstIndex = 3 * TriIdx;
+		}
 
-		GetRendererModule().DrawTileMesh(RenderContext, DrawRenderState, View, MeshBatch, bIsHitTesting, Tri.HitProxyId);
+		// Add 1 triangle to the batch :
+		++CurrentMeshBatch->Elements[0].NumPrimitives;
 	}
 
+	// Flush the final batch: 
+	check(CurrentMeshBatch != nullptr);
+	GetRendererModule().DrawTileMesh(RenderContext, DrawRenderState, View, *CurrentMeshBatch, bIsHitTesting, CurrentMeshBatch->BatchHitProxyId);
+
 	AddPass(RenderContext.GraphBuilder, [this](FRHICommandList&)
 	{
 		ReleaseTriangleMesh();
@@ -276,4 +312,4 @@ bool FCanvasTriangleRendererItem::Render_GameThread(const FCanvas* Canvas, FCanv
 	}
 
 	return true;
-}
\ No newline at end of file
+}
diff --git a/Engine/Source/Runtime/Engine/Public/CanvasTypes.h b/Engine/Source/Runtime/Engine/Public/CanvasTypes.h
--- a/Engine/Source/Runtime/Engine/Public/CanvasTypes.h
+++ b/Engine/Source/Runtime/Engine/Public/CanvasTypes.h
@@ -1054,19 +1054,6 @@ private:
 		const FStaticMeshVertexBuffers* VertexBuffers;
 	};
 
-	class FTileMesh : public FRenderResource
-	{
-	public:
-		FTileMesh(const FRawIndexBuffer* IndexBuffer, const FTileVertexFactory* VertexFactory);
-
-		FMeshBatch MeshElement;
-
-		void InitRHI() override;
-	private:
-		const FRawIndexBuffer* IndexBuffer;
-		const FTileVertexFactory* VertexFactory;
-	};
-
 	class FRenderData
 	{
 	public:
@@ -1092,14 +1079,17 @@ private:
 			return Tiles.Add(NewTile);
 		};
 
+		uint32 GetNumVertices() const;
+		uint32 GetNumIndices() const;
+
 	private:
+		FMeshBatch* AllocTileMeshBatch(FCanvasRenderContext& InRenderContext, FHitProxyId InHitProxyId);
 		void InitTileMesh(const FSceneView& View, bool bNeedsToSwitchVerticalAxis);
 		void ReleaseTileMesh();
 
-		FRawIndexBuffer IndexBuffer;
+		FRawIndexBuffer16or32 IndexBuffer;
 		FStaticMeshVertexBuffers StaticMeshVertexBuffers;
 		FTileVertexFactory VertexFactory;
-		FTileMesh TileMesh;
 
 		struct FTileInst
 		{
@@ -1223,18 +1213,6 @@ private:
 		const FStaticMeshVertexBuffers* VertexBuffers;
 	};
 
-	class FTriangleMesh : public FRenderResource
-	{
-	public:
-		FTriangleMesh(const FRawIndexBuffer* IndexBuffer, const FTriangleVertexFactory* VertexFactory);
-
-		FMeshBatch MeshBatch;
-		virtual void InitRHI() override;
-	private:
-		const FRawIndexBuffer* IndexBuffer;
-		const FTriangleVertexFactory* VertexFactory;
-	};
-
 	class FRenderData
 	{
 	public:
@@ -1244,7 +1222,6 @@ private:
 			: MaterialRenderProxy(InMaterialRenderProxy)
 			, Transform(InTransform)
 			, VertexFactory(&StaticMeshVertexBuffers, InFeatureLevel)
-			, TriMesh(&IndexBuffer, &VertexFactory)
 		{}
 
 		FORCEINLINE int32 AddTriangle(const FCanvasUVTri& Tri, FHitProxyId HitProxyId)
@@ -1273,14 +1250,17 @@ private:
 		const FMaterialRenderProxy* const MaterialRenderProxy;
 		const FCanvas::FTransformEntry Transform;
 
+		uint32 GetNumVertices() const;
+		uint32 GetNumIndices() const; 
+
 	private:
+		FMeshBatch* AllocTriangleMeshBatch(FCanvasRenderContext& InRenderContext, FHitProxyId InHitProxyId);
 		void InitTriangleMesh(const FSceneView& View, bool bNeedsToSwitchVerticalAxis);
 		void ReleaseTriangleMesh();
 
-		FRawIndexBuffer IndexBuffer;
+		FRawIndexBuffer16or32 IndexBuffer;
 		FStaticMeshVertexBuffers StaticMeshVertexBuffers;
 		FTriangleVertexFactory VertexFactory;
-		FTriangleMesh TriMesh;
 
 		struct FTriangleInst
 		{
diff --git a/Engine/Source/Runtime/Engine/Public/RawIndexBuffer.h b/Engine/Source/Runtime/Engine/Public/RawIndexBuffer.h
--- a/Engine/Source/Runtime/Engine/Public/RawIndexBuffer.h
+++ b/Engine/Source/Runtime/Engine/Public/RawIndexBuffer.h
@@ -51,6 +51,12 @@ public:
 	 */
 	void ComputeIndexWidth();
 
+	/**
+	 * Forces (or not) usage of 32 bits indices. No validation is made as to whether Indices can all be stored in 16 bits indices (if bIn32Bit == false) or not : 
+	 *  use only if you know the max value in Indices, otherwise, use ComputeIndexWidth
+	 */
+	void ForceUse32Bit(bool bIn32Bit) { b32Bit = bIn32Bit; }
+
 	// FRenderResource interface.
 	virtual void InitRHI() override;
 
