Added extra logging for invalid state after streaming requests

diff --git a/Engine/Source/Runtime/Engine/Private/StreamableRenderAsset.cpp b/Engine/Source/Runtime/Engine/Private/StreamableRenderAsset.cpp
--- a/Engine/Source/Runtime/Engine/Private/StreamableRenderAsset.cpp
+++ b/Engine/Source/Runtime/Engine/Private/StreamableRenderAsset.cpp
@@ -106,7 +106,7 @@ void UStreamableRenderAsset::TickStreaming(bool bSendCompletionEvents, TArray<US
 			}
 			else
 			{
-				check(PendingUpdate->IsCancelled());
+				checkf(PendingUpdate->IsCancelled(), TEXT("Invalid completion of streaming request for asset %s of type %s."), *GetName(), *GetClass()->GetName());
 				CachedSRRState.NumRequestedLODs = CachedSRRState.NumResidentLODs;
 			}
 
