Take bVisible flag into account in ISMComponentDescriptor (allows packing a collision only sm component and keep it invisible)

diff --git a/Engine/Source/Runtime/Engine/Private/ISMPartition/ISMComponentDescriptor.cpp b/Engine/Source/Runtime/Engine/Private/ISMPartition/ISMComponentDescriptor.cpp
--- a/Engine/Source/Runtime/Engine/Private/ISMPartition/ISMComponentDescriptor.cpp
+++ b/Engine/Source/Runtime/Engine/Private/ISMPartition/ISMComponentDescriptor.cpp
@@ -63,6 +63,7 @@ void FISMComponentDescriptor::InitFrom(const UStaticMeshComponent* Template, boo
 	bIncludeInHLOD = Template->bEnableAutoLODGeneration;
 	bHiddenInGame = Template->bHiddenInGame;
 	bIsEditorOnly = Template->bIsEditorOnly;
+	bVisible = Template->GetVisibleFlag();
 	bVisibleInRayTracing = Template->bVisibleInRayTracing;
 
 	if (const UInstancedStaticMeshComponent* ISMTemplate = Cast<UInstancedStaticMeshComponent>(Template))
@@ -124,6 +125,7 @@ bool FISMComponentDescriptor::operator==(const FISMComponentDescriptor& Other) c
 	bEnableDiscardOnLoad == Other.bEnableDiscardOnLoad &&
 	bHiddenInGame == Other.bHiddenInGame &&
 	bIsEditorOnly == Other.bIsEditorOnly &&
+	bVisible == Other.bVisible &&
 	bVisibleInRayTracing == Other.bVisibleInRayTracing &&
 	BodyInstance.GetCollisionEnabled() == Other.BodyInstance.GetCollisionEnabled() && 
 	BodyInstance.GetCollisionResponse() == Other.BodyInstance.GetCollisionResponse() &&
@@ -185,6 +187,7 @@ void FISMComponentDescriptor::InitComponent(UInstancedStaticMeshComponent* ISMCo
 	ISMComponent->bEnableAutoLODGeneration = bIncludeInHLOD;;
 	ISMComponent->bHiddenInGame = bHiddenInGame;
 	ISMComponent->bIsEditorOnly = bIsEditorOnly;
+	ISMComponent->SetVisibleFlag(bVisible);
 	ISMComponent->bVisibleInRayTracing = bVisibleInRayTracing;
 
 	// HISM Specific
diff --git a/Engine/Source/Runtime/Engine/Public/ISMPartition/ISMComponentDescriptor.h b/Engine/Source/Runtime/Engine/Public/ISMPartition/ISMComponentDescriptor.h
--- a/Engine/Source/Runtime/Engine/Public/ISMPartition/ISMComponentDescriptor.h
+++ b/Engine/Source/Runtime/Engine/Public/ISMPartition/ISMComponentDescriptor.h
@@ -159,5 +159,8 @@ public:
 
 	UPROPERTY()
 	uint8 bIsEditorOnly : 1;
+
+	UPROPERTY()
+	uint8 bVisible : 1;
 #endif
 };
