Added support for mirroring sync markers

diff --git a/Engine/Source/Runtime/AnimGraphRuntime/Private/AnimNodes/AnimNode_Mirror.cpp b/Engine/Source/Runtime/AnimGraphRuntime/Private/AnimNodes/AnimNode_Mirror.cpp
--- a/Engine/Source/Runtime/AnimGraphRuntime/Private/AnimNodes/AnimNode_Mirror.cpp
+++ b/Engine/Source/Runtime/AnimGraphRuntime/Private/AnimNodes/AnimNode_Mirror.cpp
@@ -4,6 +4,7 @@
 #include "Animation/AnimInstanceProxy.h"
 #include "Animation/AnimNode_Inertialization.h"
 #include "AnimationRuntime.h"
+#include "Animation/MirrorSyncScope.h"
 
 #define LOCTEXT_NAMESPACE "AnimNode_Mirror"
 
@@ -104,6 +105,8 @@ void FAnimNode_Mirror::Update_AnyThread(const FAnimationUpdateContext& Context)
 		}
 	}
 
+	UE::Anim::TOptionalScopedGraphMessage<UE::Anim::FMirrorSyncScope> Message(bMirror, Context, Context, MirrorDataTable);
+	
 	bMirrorState = bMirror;
 	bMirrorStateIsValid = true;
 
diff --git a/Engine/Source/Runtime/Engine/Classes/Animation/AnimMontage.h b/Engine/Source/Runtime/Engine/Classes/Animation/AnimMontage.h
--- a/Engine/Source/Runtime/Engine/Classes/Animation/AnimMontage.h
+++ b/Engine/Source/Runtime/Engine/Classes/Animation/AnimMontage.h
@@ -734,7 +734,7 @@ public:
 	virtual EAnimEventTriggerOffsets::Type CalculateOffsetForNotify(float NotifyDisplayTime) const override;
 #endif // WITH_EDITOR
 	virtual void GetMarkerIndicesForTime(float CurrentTime, bool bLooping, const TArray<FName>& ValidMarkerNames, FMarkerPair& OutPrevMarker, FMarkerPair& OutNextMarker) const override;
-	virtual FMarkerSyncAnimPosition GetMarkerSyncPositionfromMarkerIndicies(int32 PrevMarker, int32 NextMarker, float CurrentTime) const override;
+	virtual FMarkerSyncAnimPosition GetMarkerSyncPositionfromMarkerIndicies(int32 PrevMarker, int32 NextMarker, float CurrentTime, const UMirrorDataTable* MirrorTable) const override;
 	virtual void TickAssetPlayer(FAnimTickRecord& Instance, struct FAnimNotifyQueue& NotifyQueue, FAnimAssetTickContext& Context) const override;
 	virtual TArray<FName>* GetUniqueMarkerNames() override { return &MarkerData.UniqueMarkerNames; }
 	virtual void RefreshCacheData() override;
diff --git a/Engine/Source/Runtime/Engine/Classes/Animation/AnimSequence.h b/Engine/Source/Runtime/Engine/Classes/Animation/AnimSequence.h
--- a/Engine/Source/Runtime/Engine/Classes/Animation/AnimSequence.h
+++ b/Engine/Source/Runtime/Engine/Classes/Animation/AnimSequence.h
@@ -865,11 +865,11 @@ public:
 
 	// Advancing based on markers
 	float GetCurrentTimeFromMarkers(FMarkerPair& PrevMarker, FMarkerPair& NextMarker, float PositionBetweenMarkers) const;
-	virtual void AdvanceMarkerPhaseAsLeader(bool bLooping, float MoveDelta, const TArray<FName>& ValidMarkerNames, float& CurrentTime, FMarkerPair& PrevMarker, FMarkerPair& NextMarker, TArray<FPassedMarker>& MarkersPassed) const;
-	virtual void AdvanceMarkerPhaseAsFollower(const FMarkerTickContext& Context, float DeltaRemaining, bool bLooping, float& CurrentTime, FMarkerPair& PreviousMarker, FMarkerPair& NextMarker) const;
+	virtual void AdvanceMarkerPhaseAsLeader(bool bLooping, float MoveDelta, const TArray<FName>& ValidMarkerNames, float& CurrentTime, FMarkerPair& PrevMarker, FMarkerPair& NextMarker, TArray<FPassedMarker>& MarkersPassed, const UMirrorDataTable* MirrorTable) const;
+	virtual void AdvanceMarkerPhaseAsFollower(const FMarkerTickContext& Context, float DeltaRemaining, bool bLooping, float& CurrentTime, FMarkerPair& PreviousMarker, FMarkerPair& NextMarker, const UMirrorDataTable* MirrorTable) const;
 	virtual void GetMarkerIndicesForTime(float CurrentTime, bool bLooping, const TArray<FName>& ValidMarkerNames, FMarkerPair& OutPrevMarker, FMarkerPair& OutNextMarker) const;
-	virtual FMarkerSyncAnimPosition GetMarkerSyncPositionfromMarkerIndicies(int32 PrevMarker, int32 NextMarker, float CurrentTime) const;
-	virtual void GetMarkerIndicesForPosition(const FMarkerSyncAnimPosition& SyncPosition, bool bLooping, FMarkerPair& OutPrevMarker, FMarkerPair& OutNextMarker, float& CurrentTime) const;
+	virtual FMarkerSyncAnimPosition GetMarkerSyncPositionfromMarkerIndicies(int32 PrevMarker, int32 NextMarker, float CurrentTime, const UMirrorDataTable* MirrorTable) const;
+	virtual void GetMarkerIndicesForPosition(const FMarkerSyncAnimPosition& SyncPosition, bool bLooping, FMarkerPair& OutPrevMarker, FMarkerPair& OutNextMarker, float& CurrentTime, const UMirrorDataTable* MirrorTable) const;
 	
 	virtual float GetFirstMatchingPosFromMarkerSyncPos(const FMarkerSyncAnimPosition& InMarkerSyncGroupPosition) const override;
 	virtual float GetNextMatchingPosFromMarkerSyncPos(const FMarkerSyncAnimPosition& InMarkerSyncGroupPosition, const float& StartingPosition) const override;
@@ -991,7 +991,7 @@ public:
 	void RefreshSyncMarkerDataFromAuthored();
 
 	/** Take a set of marker positions and validates them against a requested start position, updating them as desired */
-	void ValidateCurrentPosition(const FMarkerSyncAnimPosition& Position, bool bPlayingForwards, bool bLooping, float&CurrentTime, FMarkerPair& PreviousMarker, FMarkerPair& NextMarker) const;
+	void ValidateCurrentPosition(const FMarkerSyncAnimPosition& Position, bool bPlayingForwards, bool bLooping, float&CurrentTime, FMarkerPair& PreviousMarker, FMarkerPair& NextMarker, const UMirrorDataTable* MirrorTable = nullptr) const;
 	bool UseRawDataForPoseExtraction(const FBoneContainer& RequiredBones) const;
 	// Should we be always using our raw data (i.e is our compressed data stale)
 	bool bUseRawDataOnly;
diff --git a/Engine/Source/Runtime/Engine/Classes/Animation/AnimSequenceBase.h b/Engine/Source/Runtime/Engine/Classes/Animation/AnimSequenceBase.h
--- a/Engine/Source/Runtime/Engine/Classes/Animation/AnimSequenceBase.h
+++ b/Engine/Source/Runtime/Engine/Classes/Animation/AnimSequenceBase.h
@@ -173,8 +173,9 @@ public:
 	virtual void RemapTracksToNewSkeleton(USkeleton* NewSkeleton, bool bConvertSpaces) override;
 #endif
 	virtual void TickAssetPlayer(FAnimTickRecord& Instance, struct FAnimNotifyQueue& NotifyQueue, FAnimAssetTickContext& Context) const override;
-	void TickByMarkerAsFollower(FMarkerTickRecord &Instance, FMarkerTickContext &MarkerContext, float& CurrentTime, float& OutPreviousTime, const float MoveDelta, const bool bLooping) const;
-	void TickByMarkerAsLeader(FMarkerTickRecord& Instance, FMarkerTickContext& MarkerContext, float& CurrentTime, float& OutPreviousTime, const float MoveDelta, const bool bLooping) const;
+
+	void TickByMarkerAsFollower(FMarkerTickRecord &Instance, FMarkerTickContext &MarkerContext, float& CurrentTime, float& OutPreviousTime, const float MoveDelta, const bool bLooping, const UMirrorDataTable* MirrorTable = nullptr) const;
+	void TickByMarkerAsLeader(FMarkerTickRecord& Instance, FMarkerTickContext& MarkerContext, float& CurrentTime, float& OutPreviousTime, const float MoveDelta, const bool bLooping, const UMirrorDataTable* MirrorTable = nullptr) const;
 	//~ End UAnimationAsset Interface
 
 	/**
@@ -197,11 +198,23 @@ public:
 
 	virtual void Serialize(FArchive& Ar) override;
 
-	virtual void AdvanceMarkerPhaseAsLeader(bool bLooping, float MoveDelta, const TArray<FName>& ValidMarkerNames, float& CurrentTime, FMarkerPair& PrevMarker, FMarkerPair& NextMarker, TArray<FPassedMarker>& MarkersPassed) const { check(false); /*Should never call this (either missing override or calling on unsupported asset */ }
-	virtual void AdvanceMarkerPhaseAsFollower(const FMarkerTickContext& Context, float DeltaRemaining, bool bLooping, float& CurrentTime, FMarkerPair& PreviousMarker, FMarkerPair& NextMarker) const { check(false); /*Should never call this (either missing override or calling on unsupported asset */ }
+	UE_DEPRECATED(5.0, "Use other AdvanceMarkerPhaseAsLeader signature")
+	virtual void AdvanceMarkerPhaseAsLeader(bool bLooping, float MoveDelta, const TArray<FName>& ValidMarkerNames, float& CurrentTime, FMarkerPair& PrevMarker, FMarkerPair& NextMarker, TArray<FPassedMarker>& MarkersPassed) const {AdvanceMarkerPhaseAsLeader(bLooping, MoveDelta, ValidMarkerNames, CurrentTime, PrevMarker, NextMarker, MarkersPassed, nullptr); }
+	virtual void AdvanceMarkerPhaseAsLeader(bool bLooping, float MoveDelta, const TArray<FName>& ValidMarkerNames, float& CurrentTime, FMarkerPair& PrevMarker, FMarkerPair& NextMarker, TArray<FPassedMarker>& MarkersPassed, const UMirrorDataTable* MirrorTable) const { check(false); /*Should never call this (either missing override or calling on unsupported asset */ }
+
+	UE_DEPRECATED(5.0, "Use other AdvanceMarkerPhaseAsFollower signature")
+	virtual void AdvanceMarkerPhaseAsFollower(const FMarkerTickContext& Context, float DeltaRemaining, bool bLooping, float& CurrentTime, FMarkerPair& PreviousMarker, FMarkerPair& NextMarker) const { AdvanceMarkerPhaseAsFollower(Context, DeltaRemaining, bLooping, CurrentTime, PreviousMarker, NextMarker, nullptr);}
+	virtual void AdvanceMarkerPhaseAsFollower(const FMarkerTickContext& Context, float DeltaRemaining, bool bLooping, float& CurrentTime, FMarkerPair& PreviousMarker, FMarkerPair& NextMarker, const UMirrorDataTable* MirrorTable) const { check(false); /*Should never call this (either missing override or calling on unsupported asset */ }
+	
 	virtual void GetMarkerIndicesForTime(float CurrentTime, bool bLooping, const TArray<FName>& ValidMarkerNames, FMarkerPair& OutPrevMarker, FMarkerPair& OutNextMarker) const { check(false); /*Should never call this (either missing override or calling on unsupported asset */ }
-	virtual FMarkerSyncAnimPosition GetMarkerSyncPositionfromMarkerIndicies(int32 PrevMarker, int32 NextMarker, float CurrentTime) const { check(false); return FMarkerSyncAnimPosition(); /*Should never call this (either missing override or calling on unsupported asset */ }
-	virtual void GetMarkerIndicesForPosition(const FMarkerSyncAnimPosition& SyncPosition, bool bLooping, FMarkerPair& OutPrevMarker, FMarkerPair& OutNextMarker, float& CurrentTime) const { check(false); /*Should never call this (either missing override or calling on unsupported asset */ }
+
+	UE_DEPRECATED(5.0, "Use other GetMarkerSyncPositionfromMarkerIndicies signature")
+	virtual FMarkerSyncAnimPosition GetMarkerSyncPositionfromMarkerIndicies(int32 PrevMarker, int32 NextMarker, float CurrentTime) const { return GetMarkerSyncPositionfromMarkerIndicies(PrevMarker, NextMarker, CurrentTime, nullptr); }
+	virtual FMarkerSyncAnimPosition GetMarkerSyncPositionfromMarkerIndicies(int32 PrevMarker, int32 NextMarker, float CurrentTime, const UMirrorDataTable* MirrorTable) const { check(false); return FMarkerSyncAnimPosition(); /*Should never call this (either missing override or calling on unsupported asset */ }
+
+	UE_DEPRECATED(5.0, "Use other GetMarkerIndicesForPosition signature")
+	virtual void GetMarkerIndicesForPosition(const FMarkerSyncAnimPosition& SyncPosition, bool bLooping, FMarkerPair& OutPrevMarker, FMarkerPair& OutNextMarker, float& CurrentTime) const { GetMarkerIndicesForPosition(SyncPosition, bLooping, OutPrevMarker, OutNextMarker, CurrentTime, nullptr); }
+	virtual void GetMarkerIndicesForPosition(const FMarkerSyncAnimPosition& SyncPosition, bool bLooping, FMarkerPair& OutPrevMarker, FMarkerPair& OutNextMarker, float& CurrentTime, const UMirrorDataTable* MirrorTable ) const { check(false); /*Should never call this (either missing override or calling on unsupported asset */ }
 	
 	virtual float GetFirstMatchingPosFromMarkerSyncPos(const FMarkerSyncAnimPosition& InMarkerSyncGroupPosition) const { return 0.f; }
 	virtual float GetNextMatchingPosFromMarkerSyncPos(const FMarkerSyncAnimPosition& InMarkerSyncGroupPosition, const float& StartingPosition) const { return 0.f; }
diff --git a/Engine/Source/Runtime/Engine/Classes/Animation/AnimationAsset.h b/Engine/Source/Runtime/Engine/Classes/Animation/AnimationAsset.h
--- a/Engine/Source/Runtime/Engine/Classes/Animation/AnimationAsset.h
+++ b/Engine/Source/Runtime/Engine/Classes/Animation/AnimationAsset.h
@@ -26,6 +26,7 @@ class USkeleton;
 class UAnimSequenceBase;
 class UBlendSpace;
 class UPoseAsset;
+class UMirrorDataTable; 
 
 namespace MarkerIndexSpecialValues
 {
@@ -304,6 +305,7 @@ struct FAnimTickRecord
 	float RootMotionWeightModifier = 1.0f;
 
 	bool bLooping = false;
+	const UMirrorDataTable* MirrorDataTable = nullptr;
 
 	union
 	{
diff --git a/Engine/Source/Runtime/Engine/Classes/Animation/MirrorDataTable.h b/Engine/Source/Runtime/Engine/Classes/Animation/MirrorDataTable.h
--- a/Engine/Source/Runtime/Engine/Classes/Animation/MirrorDataTable.h
+++ b/Engine/Source/Runtime/Engine/Classes/Animation/MirrorDataTable.h
@@ -11,8 +11,9 @@ UENUM()
 enum EMirrorRowType
 {
 	Bone,
-	Notify,
-	Curve
+	AnimationNotify,
+	Curve,
+	SyncMarker
 };
 
 /** Find and Replace Method for FMirrorFindReplaceExpression. */
@@ -190,8 +191,11 @@ public:
 	// Array with the target UIDs of curves that should be mirrored. 
 	TArray<SmartName::UID_Type> CurveMirrorTargetUIDArray;
 
-	// Map from notify to mirror notify
-	TMap<FName, FName> NotifyToMirrorNotifyMap;
+	// Map from animation notify to mirrored animation notify
+	TMap<FName, FName> AnimNotifyToMirrorAnimNotifyMap;
+	
+	// Map from sync marker to mirrored sync marker 
+	TMap<FName, FName> SyncToMirrorSyncMap;
 
 protected: 
 
diff --git a/Engine/Source/Runtime/Engine/Private/Animation/AnimMontage.cpp b/Engine/Source/Runtime/Engine/Private/Animation/AnimMontage.cpp
--- a/Engine/Source/Runtime/Engine/Private/Animation/AnimMontage.cpp
+++ b/Engine/Source/Runtime/Engine/Private/Animation/AnimMontage.cpp
@@ -1165,7 +1165,7 @@ void UAnimMontage::TickAssetPlayer(FAnimTickRecord& Instance, struct FAnimNotify
 
 				if (MarkerTickRecord->IsValid(Instance.bLooping))
 				{
-					MarkerTickContext.SetMarkerSyncStartPosition(GetMarkerSyncPositionfromMarkerIndicies(MarkerTickRecord->PreviousMarker.MarkerIndex, MarkerTickRecord->NextMarker.MarkerIndex, PreviousTime));
+					MarkerTickContext.SetMarkerSyncStartPosition(GetMarkerSyncPositionfromMarkerIndicies(MarkerTickRecord->PreviousMarker.MarkerIndex, MarkerTickRecord->NextMarker.MarkerIndex, PreviousTime, nullptr));
 
 				}
 				else
@@ -1174,14 +1174,14 @@ void UAnimMontage::TickAssetPlayer(FAnimTickRecord& Instance, struct FAnimNotify
 					FMarkerPair PreviousMarker;
 					FMarkerPair NextMarker;
 					GetMarkerIndicesForTime(PreviousTime, false, MarkerTickContext.GetValidMarkerNames(), PreviousMarker, NextMarker);
-					MarkerTickContext.SetMarkerSyncStartPosition(GetMarkerSyncPositionfromMarkerIndicies(PreviousMarker.MarkerIndex, NextMarker.MarkerIndex, PreviousTime));
+					MarkerTickContext.SetMarkerSyncStartPosition(GetMarkerSyncPositionfromMarkerIndicies(PreviousMarker.MarkerIndex, NextMarker.MarkerIndex, PreviousTime, nullptr));
 				}
 
 				// @todo this won't work well once we start jumping
 				// only thing is that passed markers won't work in this frame. To do that, I have to figure out how it jumped from where to where, 
 				GetMarkerIndicesForTime(CurrentTime, false, MarkerTickContext.GetValidMarkerNames(), MarkerTickRecord->PreviousMarker, MarkerTickRecord->NextMarker);
 				bRecordNeedsResetting = false; // we have updated it now, no need to reset
-				MarkerTickContext.SetMarkerSyncEndPosition(GetMarkerSyncPositionfromMarkerIndicies(MarkerTickRecord->PreviousMarker.MarkerIndex, MarkerTickRecord->NextMarker.MarkerIndex, CurrentTime));
+				MarkerTickContext.SetMarkerSyncEndPosition(GetMarkerSyncPositionfromMarkerIndicies(MarkerTickRecord->PreviousMarker.MarkerIndex, MarkerTickRecord->NextMarker.MarkerIndex, CurrentTime, nullptr));
 
 				MarkerTickContext.MarkersPassedThisTick = *Instance.Montage.MarkersPassedThisTick;
 
@@ -1254,9 +1254,9 @@ void UAnimMontage::GetMarkerIndicesForTime(float CurrentTime, bool bLooping, con
 	MarkerData.GetMarkerIndicesForTime(CurrentTime, bLooping, ValidMarkerNames, OutPrevMarker, OutNextMarker, GetPlayLength());
 }
 
-FMarkerSyncAnimPosition UAnimMontage::GetMarkerSyncPositionfromMarkerIndicies(int32 PrevMarker, int32 NextMarker, float CurrentTime) const
+FMarkerSyncAnimPosition UAnimMontage::GetMarkerSyncPositionfromMarkerIndicies(int32 PrevMarker, int32 NextMarker, float CurrentTime, const UMirrorDataTable* MirrorTable) const
 {
-	return MarkerData.GetMarkerSyncPositionfromMarkerIndicies(PrevMarker, NextMarker, CurrentTime, GetPlayLength());
+	return MarkerData.GetMarkerSyncPositionfromMarkerIndicies(PrevMarker, NextMarker, CurrentTime, GetPlayLength(), MirrorTable);
 }
 
 void UAnimMontage::InvalidateRecursiveAsset()
diff --git a/Engine/Source/Runtime/Engine/Private/Animation/AnimSequence.cpp b/Engine/Source/Runtime/Engine/Private/Animation/AnimSequence.cpp
--- a/Engine/Source/Runtime/Engine/Private/Animation/AnimSequence.cpp
+++ b/Engine/Source/Runtime/Engine/Private/Animation/AnimSequence.cpp
@@ -4369,7 +4369,7 @@ bool IsMarkerValid(const FAnimSyncMarker* Marker, bool bLooping, const TArray<FN
 	return (Marker == nullptr && !bLooping) || (Marker && ValidMarkerNames.Contains(Marker->MarkerName));
 }
 
-void UAnimSequence::AdvanceMarkerPhaseAsLeader(bool bLooping, float MoveDelta, const TArray<FName>& ValidMarkerNames, float& CurrentTime, FMarkerPair& PrevMarker, FMarkerPair& NextMarker, TArray<FPassedMarker>& MarkersPassed) const
+void UAnimSequence::AdvanceMarkerPhaseAsLeader(bool bLooping, float MoveDelta, const TArray<FName>& ValidMarkerNames, float& CurrentTime, FMarkerPair& PrevMarker, FMarkerPair& NextMarker, TArray<FPassedMarker>& MarkersPassed, const UMirrorDataTable* MirrorTable) const
 {
 	check(MoveDelta != 0.f);
 	const bool bPlayingForwards = MoveDelta > 0.f;
@@ -4504,10 +4504,24 @@ void UAnimSequence::AdvanceMarkerPhaseAsLeader(bool bLooping, float MoveDelta, c
 	check(CurrentTime >= 0.f && CurrentTime <= GetPlayLength());
 }
 
-void AdvanceMarkerForwards(int32& Marker, FName MarkerToFind, bool bLooping, const TArray<FAnimSyncMarker>& AuthoredSyncMarkers)
+FName MarkerOrMirroredName(FName MarkerToFind, const UMirrorDataTable* MirrorTable)
+{
+	if (MirrorTable)
+	{
+		const FName* MirroredName = MirrorTable->SyncToMirrorSyncMap.Find(MarkerToFind);
+		if (MirroredName)
+		{
+			return *MirroredName;
+		}
+	}
+	return MarkerToFind;
+}
+
+void AdvanceMarkerForwards(int32& Marker, FName MarkerToFind, bool bLooping, const TArray<FAnimSyncMarker>& AuthoredSyncMarkers, const UMirrorDataTable* MirrorTable)
 {
 	int32 MaxIterations = AuthoredSyncMarkers.Num();
-	while ((AuthoredSyncMarkers[Marker].MarkerName != MarkerToFind) && (--MaxIterations >= 0))
+
+	while ((MarkerOrMirroredName(AuthoredSyncMarkers[Marker].MarkerName, MirrorTable) != MarkerToFind) && (--MaxIterations >= 0))
 	{
 		++Marker;
 		if (Marker == AuthoredSyncMarkers.Num() && !bLooping)
@@ -4517,7 +4531,7 @@ void AdvanceMarkerForwards(int32& Marker, FName MarkerToFind, bool bLooping, con
 		Marker %= AuthoredSyncMarkers.Num();
 	}
 
-	if (!AuthoredSyncMarkers.IsValidIndex(Marker) || (AuthoredSyncMarkers[Marker].MarkerName != MarkerToFind))
+	if (!AuthoredSyncMarkers.IsValidIndex(Marker) || (MarkerOrMirroredName(AuthoredSyncMarkers[Marker].MarkerName, MirrorTable) != MarkerToFind))
 	{
 		Marker = MarkerIndexSpecialValues::AnimationBoundary;
 	}
@@ -4528,12 +4542,12 @@ int32 MarkerCounterSpaceTransform(int32 MaxMarker, int32 Source)
 	return MaxMarker - 1 - Source;
 }
 
-void AdvanceMarkerBackwards(int32& Marker, FName MarkerToFind, bool bLooping, const TArray<FAnimSyncMarker>& AuthoredSyncMarkers)
+void AdvanceMarkerBackwards(int32& Marker, FName MarkerToFind, bool bLooping, const TArray<FAnimSyncMarker>& AuthoredSyncMarkers, const UMirrorDataTable* MirrorTable)
 {
 	int32 MaxIterations = AuthoredSyncMarkers.Num();
 	const int32 MarkerMax = AuthoredSyncMarkers.Num();
 	int32 Counter = MarkerCounterSpaceTransform(MarkerMax, Marker);
-	while ((AuthoredSyncMarkers[Marker].MarkerName != MarkerToFind) && (--MaxIterations >= 0))
+	while ((MarkerOrMirroredName(AuthoredSyncMarkers[Marker].MarkerName, MirrorTable) != MarkerToFind) && (--MaxIterations >= 0))
 	{
 		if ((Marker == 0) && !bLooping)
 		{
@@ -4543,25 +4557,25 @@ void AdvanceMarkerBackwards(int32& Marker, FName MarkerToFind, bool bLooping, co
 		Marker = MarkerCounterSpaceTransform(MarkerMax, Counter);
 	}
 
-	if (!AuthoredSyncMarkers.IsValidIndex(Marker) || (AuthoredSyncMarkers[Marker].MarkerName != MarkerToFind))
+	if (!AuthoredSyncMarkers.IsValidIndex(Marker) || (MarkerOrMirroredName(AuthoredSyncMarkers[Marker].MarkerName, MirrorTable) != MarkerToFind))
 	{
 		Marker = MarkerIndexSpecialValues::AnimationBoundary;
 	}
 }
 
-bool MarkerMatchesPosition(const UAnimSequence* Sequence, int32 MarkerIndex, FName CorrectMarker)
+bool MarkerMatchesPosition(const UAnimSequence* Sequence, int32 MarkerIndex, FName CorrectMarker, const UMirrorDataTable* MirrorTable)
 {
 	checkf(MarkerIndex != MarkerIndexSpecialValues::Unitialized, TEXT("Uninitialized marker supplied to MarkerMatchesPosition. Anim: %s Expecting marker %s (Added to help debug Jira OR-9675)"), *Sequence->GetName(), *CorrectMarker.ToString());
-	return MarkerIndex == MarkerIndexSpecialValues::AnimationBoundary || CorrectMarker == Sequence->AuthoredSyncMarkers[MarkerIndex].MarkerName;
+	return MarkerIndex == MarkerIndexSpecialValues::AnimationBoundary || CorrectMarker == MarkerOrMirroredName(Sequence->AuthoredSyncMarkers[MarkerIndex].MarkerName, MirrorTable);
 }
 
-void UAnimSequence::ValidateCurrentPosition(const FMarkerSyncAnimPosition& Position, bool bPlayingForwards, bool bLooping, float&CurrentTime, FMarkerPair& PreviousMarker, FMarkerPair& NextMarker) const
+void UAnimSequence::ValidateCurrentPosition(const FMarkerSyncAnimPosition& Position, bool bPlayingForwards, bool bLooping, float&CurrentTime, FMarkerPair& PreviousMarker, FMarkerPair& NextMarker, const UMirrorDataTable* MirrorTable) const
 {
 	if (bPlayingForwards)
 	{
-		if (!MarkerMatchesPosition(this, PreviousMarker.MarkerIndex, Position.PreviousMarkerName))
+		if (!MarkerMatchesPosition(this, PreviousMarker.MarkerIndex, Position.PreviousMarkerName, MirrorTable))
 		{
-			AdvanceMarkerForwards(PreviousMarker.MarkerIndex, Position.PreviousMarkerName, bLooping, AuthoredSyncMarkers);
+			AdvanceMarkerForwards(PreviousMarker.MarkerIndex, Position.PreviousMarkerName, bLooping, AuthoredSyncMarkers, MirrorTable);
 			NextMarker.MarkerIndex = (PreviousMarker.MarkerIndex + 1);
 			if(NextMarker.MarkerIndex >= AuthoredSyncMarkers.Num())
 			{
@@ -4569,17 +4583,17 @@ void UAnimSequence::ValidateCurrentPosition(const FMarkerSyncAnimPosition& Posit
 			}
 		}
 
-		if (!MarkerMatchesPosition(this, NextMarker.MarkerIndex, Position.NextMarkerName))
+		if (!MarkerMatchesPosition(this, NextMarker.MarkerIndex, Position.NextMarkerName, MirrorTable))
 		{
-			AdvanceMarkerForwards(NextMarker.MarkerIndex, Position.NextMarkerName, bLooping, AuthoredSyncMarkers);
+			AdvanceMarkerForwards(NextMarker.MarkerIndex, Position.NextMarkerName, bLooping, AuthoredSyncMarkers, MirrorTable);
 		}
 	}
 	else
 	{
 		const int32 MarkerRange = AuthoredSyncMarkers.Num();
-		if (!MarkerMatchesPosition(this, NextMarker.MarkerIndex, Position.NextMarkerName))
+		if (!MarkerMatchesPosition(this, NextMarker.MarkerIndex, Position.NextMarkerName, MirrorTable))
 		{
-			AdvanceMarkerBackwards(NextMarker.MarkerIndex, Position.NextMarkerName, bLooping, AuthoredSyncMarkers);
+			AdvanceMarkerBackwards(NextMarker.MarkerIndex, Position.NextMarkerName, bLooping, AuthoredSyncMarkers, MirrorTable);
 			if(NextMarker.MarkerIndex == MarkerIndexSpecialValues::AnimationBoundary || (NextMarker.MarkerIndex == 0 && bLooping))
 			{
 				PreviousMarker.MarkerIndex = AuthoredSyncMarkers.Num() - 1;
@@ -4589,14 +4603,14 @@ void UAnimSequence::ValidateCurrentPosition(const FMarkerSyncAnimPosition& Posit
 				PreviousMarker.MarkerIndex = NextMarker.MarkerIndex - 1;
 			}
 		}
-		if (!MarkerMatchesPosition(this, PreviousMarker.MarkerIndex, Position.PreviousMarkerName))
+		if (!MarkerMatchesPosition(this, PreviousMarker.MarkerIndex, Position.PreviousMarkerName, MirrorTable))
 		{
-			AdvanceMarkerBackwards(PreviousMarker.MarkerIndex, Position.PreviousMarkerName, bLooping, AuthoredSyncMarkers);
+			AdvanceMarkerBackwards(PreviousMarker.MarkerIndex, Position.PreviousMarkerName, bLooping, AuthoredSyncMarkers, MirrorTable);
 		}
 	}
 
-	checkSlow(MarkerMatchesPosition(this, PreviousMarker.MarkerIndex, Position.PreviousMarkerName));
-	checkSlow(MarkerMatchesPosition(this, NextMarker.MarkerIndex, Position.NextMarkerName));
+	checkSlow(MarkerMatchesPosition(this, PreviousMarker.MarkerIndex, Position.PreviousMarkerName, MirrorTable));
+	checkSlow(MarkerMatchesPosition(this, NextMarker.MarkerIndex, Position.NextMarkerName, MirrorTable));
 
 	// Only reset position if we found valid markers. Otherwise stay where we are to not pop.
 	if ((PreviousMarker.MarkerIndex != MarkerIndexSpecialValues::AnimationBoundary) && (NextMarker.MarkerIndex != MarkerIndexSpecialValues::AnimationBoundary))
@@ -4905,11 +4919,11 @@ void UAnimSequence::MoveAttributesToModel()
 }
 #endif // WITH_EDITOR
 
-void UAnimSequence::AdvanceMarkerPhaseAsFollower(const FMarkerTickContext& Context, float DeltaRemaining, bool bLooping, float& CurrentTime, FMarkerPair& PreviousMarker, FMarkerPair& NextMarker) const
+void UAnimSequence::AdvanceMarkerPhaseAsFollower(const FMarkerTickContext& Context, float DeltaRemaining, bool bLooping, float& CurrentTime, FMarkerPair& PreviousMarker, FMarkerPair& NextMarker, const UMirrorDataTable* MirrorTable) const
 {
 	const bool bPlayingForwards = DeltaRemaining > 0.f;
 
-	ValidateCurrentPosition(Context.GetMarkerSyncStartPosition(), bPlayingForwards, bLooping, CurrentTime, PreviousMarker, NextMarker);
+	ValidateCurrentPosition(Context.GetMarkerSyncStartPosition(), bPlayingForwards, bLooping, CurrentTime, PreviousMarker, NextMarker, MirrorTable);
 	if (bPlayingForwards)
 	{
 		int32 PassedMarkersIndex = 0;
@@ -4926,7 +4940,7 @@ void UAnimSequence::AdvanceMarkerPhaseAsFollower(const FMarkerTickContext& Conte
 				PreviousMarker.MarkerIndex = NextMarker.MarkerIndex;
 				checkSlow(NextMarker.MarkerIndex != -1);
 				const FPassedMarker& PassedMarker = Context.MarkersPassedThisTick[PassedMarkersIndex];
-				AdvanceMarkerForwards(NextMarker.MarkerIndex, PassedMarker.PassedMarkerName, bLooping, AuthoredSyncMarkers);
+				AdvanceMarkerForwards(NextMarker.MarkerIndex, PassedMarker.PassedMarkerName, bLooping, AuthoredSyncMarkers, MirrorTable);
 				if (NextMarker.MarkerIndex == -1)
 				{
 					DeltaRemaining = PassedMarker.DeltaTimeWhenPassed;
@@ -4944,13 +4958,13 @@ void UAnimSequence::AdvanceMarkerPhaseAsFollower(const FMarkerTickContext& Conte
 
 		if (NextMarker.MarkerIndex != -1 && Context.MarkersPassedThisTick.Num() > 0)
 		{
-			AdvanceMarkerForwards(NextMarker.MarkerIndex, End.NextMarkerName, bLooping, AuthoredSyncMarkers);
+			AdvanceMarkerForwards(NextMarker.MarkerIndex, End.NextMarkerName, bLooping, AuthoredSyncMarkers, MirrorTable);
 		}
 
 		//Validation
 		if (NextMarker.MarkerIndex != -1)
 		{
-			check(AuthoredSyncMarkers[NextMarker.MarkerIndex].MarkerName == End.NextMarkerName);
+			check(MarkerOrMirroredName(AuthoredSyncMarkers[NextMarker.MarkerIndex].MarkerName, MirrorTable) == End.NextMarkerName);
 		}
 
 		// End Validation
@@ -4976,7 +4990,7 @@ void UAnimSequence::AdvanceMarkerPhaseAsFollower(const FMarkerTickContext& Conte
 				NextMarker.MarkerIndex = PreviousMarker.MarkerIndex;
 				checkSlow(PreviousMarker.MarkerIndex != -1);
 				const FPassedMarker& PassedMarker = Context.MarkersPassedThisTick[PassedMarkersIndex];
-				AdvanceMarkerBackwards(PreviousMarker.MarkerIndex, PassedMarker.PassedMarkerName, bLooping, AuthoredSyncMarkers);
+				AdvanceMarkerBackwards(PreviousMarker.MarkerIndex, PassedMarker.PassedMarkerName, bLooping, AuthoredSyncMarkers, MirrorTable);
 				if (PreviousMarker.MarkerIndex == -1)
 				{
 					DeltaRemaining = PassedMarker.DeltaTimeWhenPassed;
@@ -4989,7 +5003,7 @@ void UAnimSequence::AdvanceMarkerPhaseAsFollower(const FMarkerTickContext& Conte
 
 		if (PreviousMarker.MarkerIndex != -1 && Context.MarkersPassedThisTick.Num() > 0)
 		{
-			AdvanceMarkerBackwards(PreviousMarker.MarkerIndex, End.PreviousMarkerName, bLooping, AuthoredSyncMarkers);
+			AdvanceMarkerBackwards(PreviousMarker.MarkerIndex, End.PreviousMarkerName, bLooping, AuthoredSyncMarkers, MirrorTable);
 		}
 
 		if (End.PreviousMarkerName == NAME_None)
@@ -5051,7 +5065,7 @@ void UAnimSequence::GetMarkerIndicesForTime(float CurrentTime, bool bLooping, co
 	}
 }
 
-FMarkerSyncAnimPosition UAnimSequence::GetMarkerSyncPositionfromMarkerIndicies(int32 PrevMarker, int32 NextMarker, float CurrentTime) const
+FMarkerSyncAnimPosition UAnimSequence::GetMarkerSyncPositionfromMarkerIndicies(int32 PrevMarker, int32 NextMarker, float CurrentTime, const UMirrorDataTable* MirrorTable) const
 {
 	FMarkerSyncAnimPosition SyncPosition;
 	float PrevTime, NextTime;
@@ -5060,7 +5074,7 @@ FMarkerSyncAnimPosition UAnimSequence::GetMarkerSyncPositionfromMarkerIndicies(i
 		TEXT("%s - MarkerCount: %d, PrevMarker : %d, NextMarker: %d, CurrentTime : %0.2f"), *GetFullName(), AuthoredSyncMarkers.Num(), PrevMarker, NextMarker, CurrentTime))
 	{
 		PrevTime = AuthoredSyncMarkers[PrevMarker].Time;
-		SyncPosition.PreviousMarkerName = AuthoredSyncMarkers[PrevMarker].MarkerName;
+		SyncPosition.PreviousMarkerName = MarkerOrMirroredName(AuthoredSyncMarkers[PrevMarker].MarkerName, MirrorTable); 
 	}
 	else
 	{
@@ -5071,7 +5085,7 @@ FMarkerSyncAnimPosition UAnimSequence::GetMarkerSyncPositionfromMarkerIndicies(i
 		TEXT("%s - MarkerCount: %d, PrevMarker : %d, NextMarker: %d, CurrentTime : %0.2f"), *GetFullName(), AuthoredSyncMarkers.Num(), PrevMarker, NextMarker, CurrentTime))
 	{
 		NextTime = AuthoredSyncMarkers[NextMarker].Time;
-		SyncPosition.NextMarkerName = AuthoredSyncMarkers[NextMarker].MarkerName;
+		SyncPosition.NextMarkerName = MarkerOrMirroredName(AuthoredSyncMarkers[NextMarker].MarkerName, MirrorTable);
 	}
 	else
 	{
@@ -5123,8 +5137,13 @@ float UAnimSequence::GetCurrentTimeFromMarkers(FMarkerPair& PrevMarker, FMarkerP
 	return CurrentTime;
 }
 
-void UAnimSequence::GetMarkerIndicesForPosition(const FMarkerSyncAnimPosition& SyncPosition, bool bLooping, FMarkerPair& OutPrevMarker, FMarkerPair& OutNextMarker, float& OutCurrentTime) const
+void UAnimSequence::GetMarkerIndicesForPosition(const FMarkerSyncAnimPosition& SyncPosition, bool bLooping, FMarkerPair& OutPrevMarker, FMarkerPair& OutNextMarker, float& OutCurrentTime, const UMirrorDataTable* MirrorTable) const
 {
+	auto GetMarkerName = [MirrorTable](const FAnimSyncMarker& SyncMarker)->FName
+	{
+		return MarkerOrMirroredName(SyncMarker.MarkerName, MirrorTable);
+	};
+
 	// If we're not looping, assume we're playing a transition and we need to stay where we are.
 	if (!bLooping)
 	{
@@ -5136,12 +5155,12 @@ void UAnimSequence::GetMarkerIndicesForPosition(const FMarkerSyncAnimPosition& S
 			const FAnimSyncMarker& SyncMarker = AuthoredSyncMarkers[Idx];
 			const float MarkerTime = SyncMarker.Time;
 
-			if (OutCurrentTime > MarkerTime && SyncMarker.MarkerName == SyncPosition.PreviousMarkerName)
+			if (OutCurrentTime > MarkerTime && GetMarkerName(SyncMarker) == SyncPosition.PreviousMarkerName)
 			{
 				OutPrevMarker.MarkerIndex = Idx;
 				OutPrevMarker.TimeToMarker = MarkerTime - OutCurrentTime;
 			}
-			else if (OutCurrentTime < MarkerTime && SyncMarker.MarkerName == SyncPosition.NextMarkerName)
+			else if (OutCurrentTime < MarkerTime && GetMarkerName(SyncMarker) == SyncPosition.NextMarkerName)
 			{
 				OutNextMarker.MarkerIndex = Idx;
 				OutNextMarker.TimeToMarker = MarkerTime - OutCurrentTime;
@@ -5161,7 +5180,7 @@ void UAnimSequence::GetMarkerIndicesForPosition(const FMarkerSyncAnimPosition& S
 		for (int32 Idx = 0; Idx < AuthoredSyncMarkers.Num(); ++Idx)
 		{
 			const FAnimSyncMarker& Marker = AuthoredSyncMarkers[Idx];
-			if (Marker.MarkerName == SyncPosition.NextMarkerName)
+			if (GetMarkerName(Marker) == SyncPosition.NextMarkerName)
 			{
 				OutNextMarker.MarkerIndex = Idx;
 				OutCurrentTime = GetCurrentTimeFromMarkers(OutPrevMarker, OutNextMarker, SyncPosition.PositionBetweenMarkers);
@@ -5180,7 +5199,7 @@ void UAnimSequence::GetMarkerIndicesForPosition(const FMarkerSyncAnimPosition& S
 		for (int32 Idx = AuthoredSyncMarkers.Num() - 1; Idx >= 0; --Idx)
 		{
 			const FAnimSyncMarker& Marker = AuthoredSyncMarkers[Idx];
-			if (Marker.MarkerName == SyncPosition.PreviousMarkerName)
+			if (GetMarkerName(Marker) == SyncPosition.PreviousMarkerName)
 			{
 				OutPrevMarker.MarkerIndex = Idx;
 				OutCurrentTime = GetCurrentTimeFromMarkers(OutPrevMarker, OutNextMarker, SyncPosition.PositionBetweenMarkers);
@@ -5197,7 +5216,7 @@ void UAnimSequence::GetMarkerIndicesForPosition(const FMarkerSyncAnimPosition& S
 	for (int32 PrevMarkerIdx = 0; PrevMarkerIdx < AuthoredSyncMarkers.Num(); ++PrevMarkerIdx)
 	{
 		const FAnimSyncMarker& PrevMarker = AuthoredSyncMarkers[PrevMarkerIdx];
-		if (PrevMarker.MarkerName == SyncPosition.PreviousMarkerName)
+		if (GetMarkerName(PrevMarker) == SyncPosition.PreviousMarkerName)
 		{
 			const int32 EndMarkerSearchStart = PrevMarkerIdx + 1;
 
@@ -5206,7 +5225,7 @@ void UAnimSequence::GetMarkerIndicesForPosition(const FMarkerSyncAnimPosition& S
 			{
 				const int32 NextMarkerIdx = NextMarkerCount % AuthoredSyncMarkers.Num();
 
-				if (AuthoredSyncMarkers[NextMarkerIdx].MarkerName == SyncPosition.NextMarkerName)
+				if (GetMarkerName(AuthoredSyncMarkers[NextMarkerIdx]) == SyncPosition.NextMarkerName)
 				{
 					float NextMarkerTime = AuthoredSyncMarkers[NextMarkerIdx].Time;
 					if (NextMarkerTime < PrevMarker.Time)
diff --git a/Engine/Source/Runtime/Engine/Private/Animation/AnimSequenceBase.cpp b/Engine/Source/Runtime/Engine/Private/Animation/AnimSequenceBase.cpp
--- a/Engine/Source/Runtime/Engine/Private/Animation/AnimSequenceBase.cpp
+++ b/Engine/Source/Runtime/Engine/Private/Animation/AnimSequenceBase.cpp
@@ -399,7 +399,7 @@ void UAnimSequenceBase::TickAssetPlayer(FAnimTickRecord& Instance, struct FAnimN
 		{
 			if (Instance.bCanUseMarkerSync && Context.CanUseMarkerPosition())
 			{
-				TickByMarkerAsLeader(*Instance.MarkerTickRecord, Context.MarkerTickContext, CurrentTime, PreviousTime, MoveDelta, Instance.bLooping);
+				TickByMarkerAsLeader(*Instance.MarkerTickRecord, Context.MarkerTickContext, CurrentTime, PreviousTime, MoveDelta, Instance.bLooping, Instance.MirrorDataTable);
 			}
 			else
 			{
@@ -418,7 +418,7 @@ void UAnimSequenceBase::TickAssetPlayer(FAnimTickRecord& Instance, struct FAnimN
 		{
 			if (Context.CanUseMarkerPosition() && Context.MarkerTickContext.IsMarkerSyncStartValid())
 			{
-				TickByMarkerAsFollower(*Instance.MarkerTickRecord, Context.MarkerTickContext, CurrentTime, PreviousTime, Context.GetLeaderDelta(), Instance.bLooping);
+				TickByMarkerAsFollower(*Instance.MarkerTickRecord, Context.MarkerTickContext, CurrentTime, PreviousTime, Context.GetLeaderDelta(), Instance.bLooping, Instance.MirrorDataTable);
 			}
 			else
 			{
@@ -453,27 +453,27 @@ void UAnimSequenceBase::TickAssetPlayer(FAnimTickRecord& Instance, struct FAnimN
 	HandleAssetPlayerTickedInternal(Context, PreviousTime, MoveDelta, Instance, NotifyQueue);
 }
 
-void UAnimSequenceBase::TickByMarkerAsFollower(FMarkerTickRecord &Instance, FMarkerTickContext &MarkerContext, float& CurrentTime, float& OutPreviousTime, const float MoveDelta, const bool bLooping) const
+void UAnimSequenceBase::TickByMarkerAsFollower(FMarkerTickRecord &Instance, FMarkerTickContext &MarkerContext, float& CurrentTime, float& OutPreviousTime, const float MoveDelta, const bool bLooping, const UMirrorDataTable* MirrorTable) const
 {
 	if (!Instance.IsValid(bLooping))
 	{
-		GetMarkerIndicesForPosition(MarkerContext.GetMarkerSyncStartPosition(), bLooping, Instance.PreviousMarker, Instance.NextMarker, CurrentTime);
+		GetMarkerIndicesForPosition(MarkerContext.GetMarkerSyncStartPosition(), bLooping, Instance.PreviousMarker, Instance.NextMarker, CurrentTime, MirrorTable);
 	}
 
 	OutPreviousTime = CurrentTime;
 
-	AdvanceMarkerPhaseAsFollower(MarkerContext, MoveDelta, bLooping, CurrentTime, Instance.PreviousMarker, Instance.NextMarker);
+	AdvanceMarkerPhaseAsFollower(MarkerContext, MoveDelta, bLooping, CurrentTime, Instance.PreviousMarker, Instance.NextMarker, MirrorTable);
 
 	UE_LOG(LogAnimMarkerSync, Log, TEXT("Follower (%s) (TickByMarker) PreviousTime(%0.2f) CurrentTime(%0.2f) MoveDelta(%0.2f) Looping(%d) %s"), *GetName(), OutPreviousTime, CurrentTime, MoveDelta, bLooping ? 1 : 0, *MarkerContext.ToString());
 }
 
-void UAnimSequenceBase::TickByMarkerAsLeader(FMarkerTickRecord& Instance, FMarkerTickContext& MarkerContext, float& CurrentTime, float& OutPreviousTime, const float MoveDelta, const bool bLooping) const
+void UAnimSequenceBase::TickByMarkerAsLeader(FMarkerTickRecord& Instance, FMarkerTickContext& MarkerContext, float& CurrentTime, float& OutPreviousTime, const float MoveDelta, const bool bLooping, const UMirrorDataTable* MirrorTable) const
 {
 	if (!Instance.IsValid(bLooping))
 	{
 		if (MarkerContext.IsMarkerSyncStartValid())
 		{
-			GetMarkerIndicesForPosition(MarkerContext.GetMarkerSyncStartPosition(), bLooping, Instance.PreviousMarker, Instance.NextMarker, CurrentTime);
+			GetMarkerIndicesForPosition(MarkerContext.GetMarkerSyncStartPosition(), bLooping, Instance.PreviousMarker, Instance.NextMarker, CurrentTime, MirrorTable);
 		}
 		else
 		{
@@ -482,13 +482,13 @@ void UAnimSequenceBase::TickByMarkerAsLeader(FMarkerTickRecord& Instance, FMarke
 		
 	}
 
-	MarkerContext.SetMarkerSyncStartPosition(GetMarkerSyncPositionfromMarkerIndicies(Instance.PreviousMarker.MarkerIndex, Instance.NextMarker.MarkerIndex, CurrentTime));
+	MarkerContext.SetMarkerSyncStartPosition(GetMarkerSyncPositionfromMarkerIndicies(Instance.PreviousMarker.MarkerIndex, Instance.NextMarker.MarkerIndex, CurrentTime, MirrorTable));
 
 	OutPreviousTime = CurrentTime;
 
-	AdvanceMarkerPhaseAsLeader(bLooping, MoveDelta, MarkerContext.GetValidMarkerNames(), CurrentTime, Instance.PreviousMarker, Instance.NextMarker, MarkerContext.MarkersPassedThisTick);
+	AdvanceMarkerPhaseAsLeader(bLooping, MoveDelta, MarkerContext.GetValidMarkerNames(), CurrentTime, Instance.PreviousMarker, Instance.NextMarker, MarkerContext.MarkersPassedThisTick, MirrorTable);
 
-	MarkerContext.SetMarkerSyncEndPosition(GetMarkerSyncPositionfromMarkerIndicies(Instance.PreviousMarker.MarkerIndex, Instance.NextMarker.MarkerIndex, CurrentTime));
+	MarkerContext.SetMarkerSyncEndPosition(GetMarkerSyncPositionfromMarkerIndicies(Instance.PreviousMarker.MarkerIndex, Instance.NextMarker.MarkerIndex, CurrentTime, MirrorTable));
 
 	UE_LOG(LogAnimMarkerSync, Log, TEXT("Leader (%s) (TickByMarker) PreviousTime(%0.2f) CurrentTime(%0.2f) MoveDelta(%0.2f) Looping(%d) %s"), *GetName(), OutPreviousTime, CurrentTime, MoveDelta, bLooping ? 1 : 0, *MarkerContext.ToString());
 }
diff --git a/Engine/Source/Runtime/Engine/Private/Animation/AnimSync.cpp b/Engine/Source/Runtime/Engine/Private/Animation/AnimSync.cpp
--- a/Engine/Source/Runtime/Engine/Private/Animation/AnimSync.cpp
+++ b/Engine/Source/Runtime/Engine/Private/Animation/AnimSync.cpp
@@ -28,6 +28,7 @@ void FAnimSync::AddTickRecord(const FAnimTickRecord& InTickRecord, const FAnimSy
 		FSyncGroupMap& SyncGroupMap = SyncGroupMaps[GetSyncGroupWriteIndex()];
 		FAnimGroupInstance& SyncGroupInstance = SyncGroupMap.FindOrAdd(InSyncParams.GroupName);
 		SyncGroupInstance.ActivePlayers.Add(InTickRecord);
+		SyncGroupInstance.ActivePlayers.Top().MirrorDataTable = MirrorDataTable;
 		SyncGroupInstance.TestTickRecordForLeadership(InSyncParams.Role);
 	}
 	else
@@ -36,6 +37,11 @@ void FAnimSync::AddTickRecord(const FAnimTickRecord& InTickRecord, const FAnimSy
 	}
 }
 
+void FAnimSync::SetMirror(const UMirrorDataTable* MirrorTable)
+{
+	MirrorDataTable = MirrorTable;
+}
+
 void FAnimSync::TickAssetPlayerInstances(FAnimInstanceProxy& InProxy, float InDeltaSeconds)
 {
 	DECLARE_SCOPE_HIERARCHICAL_COUNTER_FUNC()
diff --git a/Engine/Source/Runtime/Engine/Private/Animation/AnimSyncScope.cpp b/Engine/Source/Runtime/Engine/Private/Animation/AnimSyncScope.cpp
--- a/Engine/Source/Runtime/Engine/Private/Animation/AnimSyncScope.cpp
+++ b/Engine/Source/Runtime/Engine/Private/Animation/AnimSyncScope.cpp
@@ -71,4 +71,19 @@ void FAnimSyncGroupScope::AddTickRecord(const FAnimTickRecord& InTickRecord, con
 	}
 }
 
+void FAnimSyncGroupScope::SetMirror(const UMirrorDataTable* MirrorDataTable)
+{
+	// Forward to outer instance if we have one
+	if (OuterProxy)
+	{
+		OuterProxy->SetSyncMirror(MirrorDataTable);
+	}
+	else
+	{
+		Proxy.SetSyncMirror(MirrorDataTable);
+	}
+}
+
+
+
 }}	// namespace UE::Anim
\ No newline at end of file
diff --git a/Engine/Source/Runtime/Engine/Private/Animation/AnimTypes.cpp b/Engine/Source/Runtime/Engine/Private/Animation/AnimTypes.cpp
--- a/Engine/Source/Runtime/Engine/Private/Animation/AnimTypes.cpp
+++ b/Engine/Source/Runtime/Engine/Private/Animation/AnimTypes.cpp
@@ -5,6 +5,7 @@
 #include "AnimationUtils.h"
 #include "Animation/AnimNotifies/AnimNotify.h"
 #include "Animation/AnimNotifies/AnimNotifyState.h"
+#include "Animation/MirrorDataTable.h"
 #include "UObject/AnimObjectVersion.h"
 
 #define NOTIFY_TRIGGER_OFFSET KINDA_SMALL_NUMBER;
@@ -209,7 +210,7 @@ void FMarkerSyncData::GetMarkerIndicesForTime(float CurrentTime, bool bLooping,
 	}
 }
 
-FMarkerSyncAnimPosition FMarkerSyncData::GetMarkerSyncPositionfromMarkerIndicies(int32 PrevMarker, int32 NextMarker, float CurrentTime, float SequenceLength) const
+FMarkerSyncAnimPosition FMarkerSyncData::GetMarkerSyncPositionfromMarkerIndicies(int32 PrevMarker, int32 NextMarker, float CurrentTime, float SequenceLength, const UMirrorDataTable* MirrorTable) const
 {
 	FMarkerSyncAnimPosition SyncPosition;
 	float PrevTime, NextTime;
@@ -218,6 +219,14 @@ FMarkerSyncAnimPosition FMarkerSyncData::GetMarkerSyncPositionfromMarkerIndicies
 	{
 		PrevTime = AuthoredSyncMarkers[PrevMarker].Time;
 		SyncPosition.PreviousMarkerName = AuthoredSyncMarkers[PrevMarker].MarkerName;
+		if (MirrorTable)
+		{
+			const FName* MirroredName = MirrorTable->SyncToMirrorSyncMap.Find(SyncPosition.PreviousMarkerName);
+			if (MirroredName)
+			{
+				SyncPosition.PreviousMarkerName = *MirroredName;
+			}
+		}
 	}
 	else
 	{
@@ -228,6 +237,14 @@ FMarkerSyncAnimPosition FMarkerSyncData::GetMarkerSyncPositionfromMarkerIndicies
 	{
 		NextTime = AuthoredSyncMarkers[NextMarker].Time;
 		SyncPosition.NextMarkerName = AuthoredSyncMarkers[NextMarker].MarkerName;
+		if (MirrorTable)
+		{
+			const FName* MirroredName = MirrorTable->SyncToMirrorSyncMap.Find(SyncPosition.NextMarkerName);
+			if (MirroredName)
+			{
+				SyncPosition.NextMarkerName = *MirroredName;
+			}
+		}
 	}
 	else
 	{
diff --git a/Engine/Source/Runtime/Engine/Private/Animation/BlendSpace.cpp b/Engine/Source/Runtime/Engine/Private/Animation/BlendSpace.cpp
--- a/Engine/Source/Runtime/Engine/Private/Animation/BlendSpace.cpp
+++ b/Engine/Source/Runtime/Engine/Private/Animation/BlendSpace.cpp
@@ -401,14 +401,14 @@ void UBlendSpace::TickAssetPlayer(FAnimTickRecord& Instance, struct FAnimNotifyQ
 					}
 					else if (!SampleDataItem.MarkerTickRecord.IsValid(Instance.bLooping) && Context.MarkerTickContext.GetMarkerSyncStartPosition().IsValid())
 					{
-						Sample.Animation->GetMarkerIndicesForPosition(Context.MarkerTickContext.GetMarkerSyncStartPosition(), true, SampleDataItem.MarkerTickRecord.PreviousMarker, SampleDataItem.MarkerTickRecord.NextMarker, SampleDataItem.Time);
+						Sample.Animation->GetMarkerIndicesForPosition(Context.MarkerTickContext.GetMarkerSyncStartPosition(), true, SampleDataItem.MarkerTickRecord.PreviousMarker, SampleDataItem.MarkerTickRecord.NextMarker, SampleDataItem.Time, nullptr);
 					}
 
 					const float NewDeltaTime = Context.GetDeltaTime() * Instance.PlayRateMultiplier * Sample.RateScale * Sample.Animation->RateScale;
 					if (!FMath::IsNearlyZero(NewDeltaTime))
 					{
 						Context.SetLeaderDelta(NewDeltaTime);
-						Sample.Animation->TickByMarkerAsLeader(SampleDataItem.MarkerTickRecord, Context.MarkerTickContext, SampleDataItem.Time, SampleDataItem.PreviousTime, NewDeltaTime, Instance.bLooping);
+						Sample.Animation->TickByMarkerAsLeader(SampleDataItem.MarkerTickRecord, Context.MarkerTickContext, SampleDataItem.Time, SampleDataItem.PreviousTime, NewDeltaTime, Instance.bLooping, nullptr);
 						check(!Instance.bLooping || Context.MarkerTickContext.IsMarkerSyncStartValid());
 						TickFollowerSamples(SampleDataList, HighestMarkerSyncWeightIndex, Context, bResetMarkerDataOnFollowers);
 					}
@@ -1351,7 +1351,7 @@ void UBlendSpace::TickFollowerSamples(TArray<FBlendSampleData>& SampleDataList,
 
 			if (Sample.Animation->AuthoredSyncMarkers.Num() > 0) // Update followers who can do marker sync, others will be handled later in TickAssetPlayer
 			{
-				Sample.Animation->TickByMarkerAsFollower(SampleDataItem.MarkerTickRecord, Context.MarkerTickContext, SampleDataItem.Time, SampleDataItem.PreviousTime, Context.GetLeaderDelta(), true);
+				Sample.Animation->TickByMarkerAsFollower(SampleDataItem.MarkerTickRecord, Context.MarkerTickContext, SampleDataItem.Time, SampleDataItem.PreviousTime, Context.GetLeaderDelta(), true, nullptr);
 			}
 		}
 	}
diff --git a/Engine/Source/Runtime/Engine/Private/Animation/MirrorDataTable.cpp b/Engine/Source/Runtime/Engine/Private/Animation/MirrorDataTable.cpp
--- a/Engine/Source/Runtime/Engine/Private/Animation/MirrorDataTable.cpp
+++ b/Engine/Source/Runtime/Engine/Private/Animation/MirrorDataTable.cpp
@@ -275,7 +275,7 @@ void UMirrorDataTable::FindReplaceMirroredNames()
 		FName MirroredName = FindReplace(Notify);
 		if (!MirroredName.IsNone() && Skeleton->AnimationNotifies.Contains(MirroredName))
 		{
-			AddMirrorRow(Notify, MirroredName, EMirrorRowType::Notify);
+			AddMirrorRow(Notify, MirroredName, EMirrorRowType::AnimationNotify);
 		}
 	}
 
@@ -394,7 +394,8 @@ void UMirrorDataTable::FillMirrorBoneIndexes(const FReferenceSkeleton& Reference
 
 void UMirrorDataTable::FillMirrorArrays()
 {
-	NotifyToMirrorNotifyMap.Empty();
+	SyncToMirrorSyncMap.Empty();
+	AnimNotifyToMirrorAnimNotifyMap.Empty();
 	if (!Skeleton)
 	{
 		BoneToMirrorBoneIndex.Empty();
@@ -409,13 +410,25 @@ void UMirrorDataTable::FillMirrorArrays()
 	
 	ForeachRow<FMirrorTableRow>(TEXT("UMirrorDataTable::FillMirrorArrays"), [this, &CurveToMirrorCurveMap](const FName& Key, const FMirrorTableRow& Value) mutable
 		{
-			if (Value.MirrorEntryType == EMirrorRowType::Curve)
+			switch (Value.MirrorEntryType)
+			{
+			case  EMirrorRowType::Curve:
 			{
 				CurveToMirrorCurveMap.Add(Value.Name, Value.MirroredName);
+				break;
 			}
-			else if (Value.MirrorEntryType == EMirrorRowType::Notify)
+			case EMirrorRowType::SyncMarker:
+			{
+                SyncToMirrorSyncMap.Add(Value.Name, Value.MirroredName);
+				break;
+            }
+			case EMirrorRowType::AnimationNotify:
 			{
-				NotifyToMirrorNotifyMap.Add(Value.Name, Value.MirroredName);
+                AnimNotifyToMirrorAnimNotifyMap.Add(Value.Name, Value.MirroredName);
+				break;
+            }
+			default:
+				break;
 			}
 		}
 	);
diff --git a/Engine/Source/Runtime/Engine/Private/Animation/MirrorSyncScope.cpp b/Engine/Source/Runtime/Engine/Private/Animation/MirrorSyncScope.cpp
new file mode 100644
--- /dev/null
+++ b/Engine/Source/Runtime/Engine/Private/Animation/MirrorSyncScope.cpp
@@ -0,0 +1,73 @@
+// Copyright Epic Games, Inc. All Rights Reserved.
+
+#include "Animation/MirrorSyncScope.h"
+#include "Animation/AnimSyncScope.h"
+#include "Animation/AnimInstanceProxy.h"
+#include "Animation/MirrorDataTable.h"
+
+IMPLEMENT_ANIMGRAPH_MESSAGE(UE::Anim::FMirrorSyncScope);
+
+namespace UE { namespace Anim {
+
+bool IsMirrorSyncIdentical(const UMirrorDataTable* MirrorTableA, const UMirrorDataTable* MirrorTableB)
+{
+	if (MirrorTableA == MirrorTableB)
+	{
+		return true; 
+	}
+	if (MirrorTableA->SyncToMirrorSyncMap.Num() == MirrorTableB->SyncToMirrorSyncMap.Num())
+	{
+		for (auto& Elem : MirrorTableA->SyncToMirrorSyncMap)
+		{
+			const FName* BValue = MirrorTableB->SyncToMirrorSyncMap.Find(Elem.Key);
+			if( BValue == nullptr || *BValue != Elem.Value)
+			{
+				return false; 
+			}
+		}
+		return true; 
+	}
+	return false; 
+}
+
+FMirrorSyncScope::FMirrorSyncScope(const FAnimationBaseContext& InContext, const UMirrorDataTable* InMirrorDataTable)
+	: MirrorDataTable(InMirrorDataTable)
+{
+	bool bClearMirrorSync = false; 
+	if (FMirrorSyncScope* ParentMirrorScope = InContext.GetMessage<FMirrorSyncScope>())
+	{
+		// Syncing is done when animation sequences are evaluated, which means we need to examine the stack of mirror
+		// nodes to determine if we can simply enable and disable sync mirroring.  
+		OuterScopeMirrorDataTable = ParentMirrorScope->MirrorDataTable;
+		MirrorScopeDepth = ParentMirrorScope->GetMirrorScopeDepth() + 1;
+		if (MirrorScopeDepth % 2 == 0)
+		{
+			bClearMirrorSync = IsMirrorSyncIdentical(ParentMirrorScope->MirrorDataTable, MirrorDataTable);
+		}
+	}
+		 
+	AnimSyncGroupScope = InContext.GetMessage<FAnimSyncGroupScope>();
+	if(AnimSyncGroupScope)
+	{
+		if (bClearMirrorSync)
+		{
+			MirrorDataTable = nullptr;
+		}
+		AnimSyncGroupScope->SetMirror(MirrorDataTable);
+	}
+}
+
+FMirrorSyncScope::~FMirrorSyncScope()
+{
+	if(AnimSyncGroupScope)
+	{
+		AnimSyncGroupScope->SetMirror(OuterScopeMirrorDataTable);
+	}
+}
+
+int32 FMirrorSyncScope::GetMirrorScopeDepth() const
+{
+	return MirrorScopeDepth;
+}
+	
+}}	// namespace UE::Anim
\ No newline at end of file
diff --git a/Engine/Source/Runtime/Engine/Public/Animation/AnimInstanceProxy.h b/Engine/Source/Runtime/Engine/Public/Animation/AnimInstanceProxy.h
--- a/Engine/Source/Runtime/Engine/Public/Animation/AnimInstanceProxy.h
+++ b/Engine/Source/Runtime/Engine/Public/Animation/AnimInstanceProxy.h
@@ -19,7 +19,7 @@
 #include "Logging/TokenizedMessage.h"
 #include "Animation/AnimTrace.h"
 #include "Animation/AnimSync.h"
-
+#include "Animation/AnimSyncScope.h"
 #include "AnimInstanceProxy.generated.h"
 
 class UBlendSpace;
@@ -550,6 +550,7 @@ public:
 	friend struct FAnimationBaseContext;
 	friend struct FAnimTrace;
 	friend struct UE::Anim::FAnimSync;
+	friend class UE::Anim::FAnimSyncGroupScope;
 
 protected:
 	/** Called when our anim instance is being initialized */
@@ -714,6 +715,12 @@ protected:
 
 	FMarkerSyncAnimPosition GetSyncGroupPosition(FName InSyncGroupName) const;
 
+	// Set the mirror data table to sync animations
+	void SetSyncMirror(const UMirrorDataTable* MirrorDataTable)
+	{
+		Sync.SetMirror(MirrorDataTable);
+	}
+	
 	// slot node run-time functions
 	void ReinitializeSlotNodes();
 
diff --git a/Engine/Source/Runtime/Engine/Public/Animation/AnimSync.h b/Engine/Source/Runtime/Engine/Public/Animation/AnimSync.h
--- a/Engine/Source/Runtime/Engine/Public/Animation/AnimSync.h
+++ b/Engine/Source/Runtime/Engine/Public/Animation/AnimSync.h
@@ -8,6 +8,7 @@
 
 struct FAnimationBaseContext;
 struct FAnimInstanceProxy;
+class UMirrorDataTable; 
 
 namespace UE { namespace Anim {
 
@@ -38,6 +39,8 @@ struct ENGINE_API FAnimSync
 	// Adds a tick record in the list for the correct group or the ungrouped array.
 	void AddTickRecord(const FAnimTickRecord& InTickRecord, const FAnimSyncParams& InSyncParams = FAnimSyncParams());
 
+	void SetMirror(const UMirrorDataTable* MirrorTable);
+
 	// Tick all of the asset player tick records that are registered with us
 	void TickAssetPlayerInstances(FAnimInstanceProxy& InProxy, float InDeltaSeconds);
 
@@ -96,6 +99,9 @@ private:
 
 	/** The set of tick groups for this anim instance */
 	FSyncGroupMap SyncGroupMaps[2];
+		
+	/** If not null, the sync markers will be mirrored based on the data in the table **/ 
+	const UMirrorDataTable* MirrorDataTable = nullptr; 
 
 	/** Current sync group buffer index */
 	int32 SyncGroupWriteIndex = 0;
diff --git a/Engine/Source/Runtime/Engine/Public/Animation/AnimSyncScope.h b/Engine/Source/Runtime/Engine/Public/Animation/AnimSyncScope.h
--- a/Engine/Source/Runtime/Engine/Public/Animation/AnimSyncScope.h
+++ b/Engine/Source/Runtime/Engine/Public/Animation/AnimSyncScope.h
@@ -39,6 +39,9 @@ public:
 	// Adds a tick record in the list for the correct group or the ungrouped array.
 	void AddTickRecord(const FAnimTickRecord& InTickRecord, const FAnimSyncParams& InSyncParams = FAnimSyncParams(), const FAnimSyncDebugInfo& InDebugInfo = FAnimSyncDebugInfo());
 
+	// Set a mirror table that any tick records added to the group will use
+	void SetMirror(const UMirrorDataTable* MirrorDataTable);
+
 private:
 	// The node ID that was used when this scope was entered
 	int32 NodeId;
diff --git a/Engine/Source/Runtime/Engine/Public/Animation/AnimTypes.h b/Engine/Source/Runtime/Engine/Public/Animation/AnimTypes.h
--- a/Engine/Source/Runtime/Engine/Public/Animation/AnimTypes.h
+++ b/Engine/Source/Runtime/Engine/Public/Animation/AnimTypes.h
@@ -17,6 +17,7 @@ struct FPassedMarker;
 
 class FMemoryReader;
 class FMemoryWriter;
+class UMirrorDataTable;
 
 // Disable debugging information for shipping and test builds.
 #define ENABLE_ANIM_DEBUG (1 && !(UE_BUILD_SHIPPING || UE_BUILD_TEST))
@@ -642,7 +643,7 @@ struct FMarkerSyncData
 	TArray<FName>				UniqueMarkerNames;
 
 	void GetMarkerIndicesForTime(float CurrentTime, bool bLooping, const TArray<FName>& ValidMarkerNames, FMarkerPair& OutPrevMarker, FMarkerPair& OutNextMarker, float SequenceLength) const;
-	FMarkerSyncAnimPosition GetMarkerSyncPositionfromMarkerIndicies(int32 PrevMarker, int32 NextMarker, float CurrentTime, float SequenceLength) const;
+	FMarkerSyncAnimPosition GetMarkerSyncPositionfromMarkerIndicies(int32 PrevMarker, int32 NextMarker, float CurrentTime, float SequenceLength, const UMirrorDataTable* MirrorTable = nullptr) const;
 	void CollectUniqueNames();
 	void CollectMarkersInRange(float PrevPosition, float NewPosition, TArray<FPassedMarker>& OutMarkersPassedThisTick, float TotalDeltaMove);
 };
diff --git a/Engine/Source/Runtime/Engine/Public/Animation/MirrorSyncScope.h b/Engine/Source/Runtime/Engine/Public/Animation/MirrorSyncScope.h
new file mode 100644
--- /dev/null
+++ b/Engine/Source/Runtime/Engine/Public/Animation/MirrorSyncScope.h
@@ -0,0 +1,30 @@
+// Copyright Epic Games, Inc. All Rights Reserved.
+
+#pragma once
+
+#include "CoreMinimal.h"
+#include "AnimNodeMessages.h"
+
+class UMirrorDataTable;
+
+namespace UE { namespace Anim {
+
+class FAnimSyncGroupScope;
+
+// Scoped graph message used to synchronize mirroring 
+class ENGINE_API FMirrorSyncScope : public IGraphMessage
+{
+	DECLARE_ANIMGRAPH_MESSAGE(FMirrorSyncScope);
+public:
+	FMirrorSyncScope(const FAnimationBaseContext& InContext, const UMirrorDataTable* InMirrorDataTable);
+	virtual ~FMirrorSyncScope();
+	int32 GetMirrorScopeDepth() const;
+	
+private:
+	const UMirrorDataTable* MirrorDataTable = nullptr;
+	int32 MirrorScopeDepth = 1; 
+	const UMirrorDataTable* OuterScopeMirrorDataTable = nullptr;
+	FAnimSyncGroupScope* AnimSyncGroupScope = nullptr; 
+};
+
+}}	// namespace UE::Anim
