Added extra validation around FStreamableTextureResource::TextureReferenceRHI.

diff --git a/Engine/Source/Runtime/Engine/Private/Rendering/StreamableTextureResource.cpp b/Engine/Source/Runtime/Engine/Private/Rendering/StreamableTextureResource.cpp
--- a/Engine/Source/Runtime/Engine/Private/Rendering/StreamableTextureResource.cpp
+++ b/Engine/Source/Runtime/Engine/Private/Rendering/StreamableTextureResource.cpp
@@ -199,7 +199,11 @@ void FStreamableTextureResource::ReleaseRHI()
 {
 	STAT(DecrementTextureStats());
 
-	RHIUpdateTextureReference(TextureReferenceRHI, nullptr);
+	if (ensure(TextureReferenceRHI.IsValid()))
+	{
+		RHIUpdateTextureReference(TextureReferenceRHI, nullptr);
+	}
+
 	TextureRHI.SafeRelease();
 	FTextureResource::ReleaseRHI();
 }
@@ -227,7 +231,10 @@ void FStreamableTextureResource::FinalizeStreaming(FRHITexture* InTextureRHI)
 	}
 
 	TextureRHI = InTextureRHI;
-	RHIUpdateTextureReference(TextureReferenceRHI, TextureRHI);
+	if (ensure(TextureReferenceRHI.IsValid()))
+	{
+		RHIUpdateTextureReference(TextureReferenceRHI, TextureRHI);
+	}
 	State.NumResidentLODs = State.NumRequestedLODs;
 }
 
