Disable deprecation warning in implicitly created copy constructor

diff --git a/Engine/Source/Runtime/Engine/Private/Animation/AnimCompressionTypes.cpp b/Engine/Source/Runtime/Engine/Private/Animation/AnimCompressionTypes.cpp
--- a/Engine/Source/Runtime/Engine/Private/Animation/AnimCompressionTypes.cpp
+++ b/Engine/Source/Runtime/Engine/Private/Animation/AnimCompressionTypes.cpp
@@ -437,6 +437,10 @@ void FUECompressedAnimDataMutable::BuildFinalBuffer(TArray<uint8>& OutCompressed
 	WriteArray(MemoryWriter, CompressedByteStream);
 }
 
+PRAGMA_DISABLE_DEPRECATION_WARNINGS
+ICompressedAnimData::ICompressedAnimData(const ICompressedAnimData&) = default;
+PRAGMA_ENABLE_DEPRECATION_WARNINGS
+
 void ICompressedAnimData::SerializeCompressedData(class FArchive& Ar)
 {
 	Ar << CompressedNumberOfKeys;
diff --git a/Engine/Source/Runtime/Engine/Public/Animation/AnimCompressionTypes.h b/Engine/Source/Runtime/Engine/Public/Animation/AnimCompressionTypes.h
--- a/Engine/Source/Runtime/Engine/Public/Animation/AnimCompressionTypes.h
+++ b/Engine/Source/Runtime/Engine/Public/Animation/AnimCompressionTypes.h
@@ -406,6 +406,9 @@ struct ENGINE_API ICompressedAnimData
 	FAnimationErrorStats BoneCompressionErrorStats;
 #endif
 
+	ICompressedAnimData() = default;
+	ICompressedAnimData(const ICompressedAnimData&);
+
 	/* Virtual interface codecs must implement */
 	virtual ~ICompressedAnimData() {}
 
