Add const keywords where possible to improve the code contract between the SubobjectEditor and its components. Only the Subobject Subsystem should be allow to modify Subobject Data.

#rnx

diff --git a/Engine/Source/Runtime/Engine/Classes/Components/SceneComponent.h b/Engine/Source/Runtime/Engine/Classes/Components/SceneComponent.h
--- a/Engine/Source/Runtime/Engine/Classes/Components/SceneComponent.h
+++ b/Engine/Source/Runtime/Engine/Classes/Components/SceneComponent.h
@@ -1208,7 +1208,7 @@ public:
 	 * Called to see if it's possible to attach another scene component as a child.
 	 * Note: This can be called on template component as well!
 	 */
-	virtual bool CanAttachAsChild(USceneComponent* ChildComponent, FName SocketName) const { return true; }
+	virtual bool CanAttachAsChild(const USceneComponent* ChildComponent, FName SocketName) const { return true; }
 
 	/** Get the extent used when placing this component in the editor, used for 'pulling back' hit. */
 	virtual FBoxSphereBounds GetPlacementExtent() const;
