Fix PVS static analysis warnings about ternary operators.

diff --git a/Engine/Source/Editor/StaticMeshEditor/Private/SStaticMeshEditorViewport.cpp b/Engine/Source/Editor/StaticMeshEditor/Private/SStaticMeshEditorViewport.cpp
--- a/Engine/Source/Editor/StaticMeshEditor/Private/SStaticMeshEditorViewport.cpp
+++ b/Engine/Source/Editor/StaticMeshEditor/Private/SStaticMeshEditorViewport.cpp
@@ -225,14 +225,14 @@ bool SStaticMeshEditorViewport::IsShowNaniteProxyChecked() const
 
 bool SStaticMeshEditorViewport::IsShowNaniteProxyVisible() const
 {
-	const UStaticMesh* PreviewStaticMesh = PreviewMeshComponent ? PreviewMeshComponent->GetStaticMesh() : nullptr;
+	const UStaticMesh* PreviewStaticMesh = PreviewMeshComponent ? ToRawPtr(PreviewMeshComponent->GetStaticMesh()) : nullptr;
 
 	return PreviewStaticMesh && PreviewStaticMesh->NaniteSettings.bEnabled ? true : false;
 }
 
 void SStaticMeshEditorViewport::UpdatePreviewSocketMeshes()
 {
-	UStaticMesh* const PreviewStaticMesh = PreviewMeshComponent ? PreviewMeshComponent->GetStaticMesh() : nullptr;
+	UStaticMesh* const PreviewStaticMesh = PreviewMeshComponent ? ToRawPtr(PreviewMeshComponent->GetStaticMesh()) : nullptr;
 
 	if( PreviewStaticMesh )
 	{
diff --git a/Engine/Source/Editor/UnrealEd/Private/HierarchicalLOD.cpp b/Engine/Source/Editor/UnrealEd/Private/HierarchicalLOD.cpp
--- a/Engine/Source/Editor/UnrealEd/Private/HierarchicalLOD.cpp
+++ b/Engine/Source/Editor/UnrealEd/Private/HierarchicalLOD.cpp
@@ -884,7 +884,7 @@ void FHierarchicalLODBuilder::BuildMeshesForLODActors(bool bForceAll)
 
 						for (ALODActor* LODActor : LODLevelActors[LODIndex])
 						{
-							UStaticMesh* StaticMesh = LODActor->GetStaticMeshComponent() ? LODActor->GetStaticMeshComponent()->GetStaticMesh() : nullptr;
+							UStaticMesh* StaticMesh = LODActor->GetStaticMeshComponent() ? ToRawPtr(LODActor->GetStaticMeshComponent()->GetStaticMesh()) : nullptr;
 							if (StaticMesh == nullptr)
 							{
 								continue;
diff --git a/Engine/Source/Runtime/Engine/Private/HLOD/HLODProxyDesc.cpp b/Engine/Source/Runtime/Engine/Private/HLOD/HLODProxyDesc.cpp
--- a/Engine/Source/Runtime/Engine/Private/HLOD/HLODProxyDesc.cpp
+++ b/Engine/Source/Runtime/Engine/Private/HLOD/HLODProxyDesc.cpp
@@ -154,7 +154,7 @@ bool UHLODProxyDesc::ShouldUpdateDesc(const ALODActor* InLODActor) const
 		return true;
 	}
 
-	UStaticMesh* LocalStaticMesh = InLODActor->StaticMeshComponent ? InLODActor->StaticMeshComponent->GetStaticMesh() : nullptr;
+	UStaticMesh* LocalStaticMesh = InLODActor->StaticMeshComponent ? ToRawPtr(InLODActor->StaticMeshComponent->GetStaticMesh()) : nullptr;
 	if (StaticMesh != LocalStaticMesh)
 	{
 		return true;
diff --git a/Engine/Source/Runtime/Engine/Private/HierarchicalInstancedStaticMesh.cpp b/Engine/Source/Runtime/Engine/Private/HierarchicalInstancedStaticMesh.cpp
--- a/Engine/Source/Runtime/Engine/Private/HierarchicalInstancedStaticMesh.cpp
+++ b/Engine/Source/Runtime/Engine/Private/HierarchicalInstancedStaticMesh.cpp
@@ -2879,7 +2879,7 @@ void UHierarchicalInstancedStaticMeshComponent::ApplyEmpty()
 	SortedInstances.Empty();
 	UnbuiltInstanceBoundsList.Empty();
 	BuiltInstanceBounds.Init();
-	CacheMeshExtendedBounds = GetStaticMesh() && (GetStaticMesh()->IsCompiling() || GetStaticMesh()->HasValidRenderData()) ? GetStaticMesh()->GetBounds() : FBoxSphereBounds(ForceInitToZero);
+	CacheMeshExtendedBounds = (GetStaticMesh() && (GetStaticMesh()->IsCompiling() || GetStaticMesh()->HasValidRenderData())) ? GetStaticMesh()->GetBounds() : FBoxSphereBounds(ForceInitToZero);
 	InstanceUpdateCmdBuffer.Reset();
 	if (PerInstanceRenderData.IsValid())
 	{
diff --git a/Engine/Source/Runtime/Engine/Private/LODActor.cpp b/Engine/Source/Runtime/Engine/Private/LODActor.cpp
--- a/Engine/Source/Runtime/Engine/Private/LODActor.cpp
+++ b/Engine/Source/Runtime/Engine/Private/LODActor.cpp
@@ -828,7 +828,7 @@ void ALODActor::AddSubActors(const TArray<AActor*>& InActors)
 
 			for (UStaticMeshComponent* Component : StaticMeshComponents)
 			{
-				const UStaticMesh* StaticMesh = (Component) ? Component->GetStaticMesh() : nullptr;
+				const UStaticMesh* StaticMesh = (Component) ? ToRawPtr(Component->GetStaticMesh()) : nullptr;
 				if (StaticMesh && StaticMesh->GetRenderData() && StaticMesh->GetRenderData()->LODResources.Num() > 0)
 				{
 					NumTrianglesInSubActors += StaticMesh->GetRenderData()->LODResources[0].GetNumTriangles();
@@ -860,7 +860,7 @@ const bool ALODActor::RemoveSubActor(AActor* InActor)
 			InActor->GetComponents<UStaticMeshComponent>(StaticMeshComponents);
 			for (UStaticMeshComponent* Component : StaticMeshComponents)
 			{
-				const UStaticMesh* StaticMesh = (Component) ? Component->GetStaticMesh() : nullptr;
+				const UStaticMesh* StaticMesh = (Component) ? ToRawPtr(Component->GetStaticMesh()) : nullptr;
 				if (StaticMesh && StaticMesh->GetRenderData() && StaticMesh->GetRenderData()->LODResources.Num() > 0)
 				{
 					NumTrianglesInSubActors -= StaticMesh->GetRenderData()->LODResources[0].GetNumTriangles();
