Address redundant wrapped pointer accesses in USkeletalMesh::InitMorphTargets.

diff --git a/Engine/Source/Runtime/Engine/Private/SkeletalMesh.cpp b/Engine/Source/Runtime/Engine/Private/SkeletalMesh.cpp
--- a/Engine/Source/Runtime/Engine/Private/SkeletalMesh.cpp
+++ b/Engine/Source/Runtime/Engine/Private/SkeletalMesh.cpp
@@ -3564,13 +3564,14 @@ void USkeletalMesh::InitMorphTargets()
 	TRACE_CPUPROFILER_EVENT_SCOPE(USkeletalMesh::InitMorphTargets);
 	GetMorphTargetIndexMap().Empty();
 
-	for (int32 Index = 0; Index < GetMorphTargets().Num(); ++Index)
+	TArray<UMorphTarget*>& MorphTargetsLocal = GetMorphTargets();
+	for (int32 Index = 0; Index < MorphTargetsLocal.Num(); ++Index)
 	{
-		UMorphTarget* MorphTarget = GetMorphTargets()[Index];
+		UMorphTarget* MorphTarget = MorphTargetsLocal[Index];
 		// if we don't have a valid data, just remove it
 		if (!MorphTarget->HasValidData())
 		{
-			GetMorphTargets().RemoveAt(Index);
+			MorphTargetsLocal.RemoveAt(Index);
 			--Index;
 			continue;
 		}
