Fix additional deprecation warnings
Move constructor = default in to header to allow inlining

diff --git a/Engine/Source/Runtime/Engine/Private/Animation/AnimCompressionTypes.cpp b/Engine/Source/Runtime/Engine/Private/Animation/AnimCompressionTypes.cpp
--- a/Engine/Source/Runtime/Engine/Private/Animation/AnimCompressionTypes.cpp
+++ b/Engine/Source/Runtime/Engine/Private/Animation/AnimCompressionTypes.cpp
@@ -437,10 +437,6 @@ void FUECompressedAnimDataMutable::BuildFinalBuffer(TArray<uint8>& OutCompressed
 	WriteArray(MemoryWriter, CompressedByteStream);
 }
 
-PRAGMA_DISABLE_DEPRECATION_WARNINGS
-ICompressedAnimData::ICompressedAnimData(const ICompressedAnimData&) = default;
-PRAGMA_ENABLE_DEPRECATION_WARNINGS
-
 void ICompressedAnimData::SerializeCompressedData(class FArchive& Ar)
 {
 	Ar << CompressedNumberOfKeys;
diff --git a/Engine/Source/Runtime/Engine/Public/Animation/AnimCompressionTypes.h b/Engine/Source/Runtime/Engine/Public/Animation/AnimCompressionTypes.h
--- a/Engine/Source/Runtime/Engine/Public/Animation/AnimCompressionTypes.h
+++ b/Engine/Source/Runtime/Engine/Public/Animation/AnimCompressionTypes.h
@@ -195,6 +195,8 @@ public:
 
 	FCompressibleAnimData(class UAnimSequence* InSeq, const bool bPerformStripping);
 
+	FCompressibleAnimData& operator=(const FCompressibleAnimData&);
+
 	UAnimCurveCompressionSettings* CurveCompressionSettings;
 
 	UAnimBoneCompressionSettings* BoneCompressionSettings;
@@ -296,6 +298,10 @@ private:
 
 };
 
+PRAGMA_DISABLE_DEPRECATION_WARNINGS
+FCompressibleAnimData& FCompressibleAnimData::operator=(const FCompressibleAnimData&) = default;
+PRAGMA_ENABLE_DEPRECATION_WARNINGS
+
 typedef TSharedPtr<FCompressibleAnimData, ESPMode::ThreadSafe> FCompressibleAnimPtr;
 typedef TSharedRef<FCompressibleAnimData, ESPMode::ThreadSafe> FCompressibleAnimRef;
 
@@ -420,6 +426,10 @@ struct ENGINE_API ICompressedAnimData
 	virtual bool IsValid() const = 0;
 };
 
+PRAGMA_DISABLE_DEPRECATION_WARNINGS
+ICompressedAnimData::ICompressedAnimData(const ICompressedAnimData&) = default;
+PRAGMA_ENABLE_DEPRECATION_WARNINGS
+
 struct ENGINE_API FCompressibleAnimDataResult
 {
 	TArray<uint8> CompressedByteStream;
