Avoid duplicate symbol linker errors

diff --git a/Engine/Source/Runtime/Engine/Private/Animation/AnimCompressionTypes.cpp b/Engine/Source/Runtime/Engine/Private/Animation/AnimCompressionTypes.cpp
--- a/Engine/Source/Runtime/Engine/Private/Animation/AnimCompressionTypes.cpp
+++ b/Engine/Source/Runtime/Engine/Private/Animation/AnimCompressionTypes.cpp
@@ -77,6 +77,11 @@ void StripFramesOdd(TArray<ArrayValue>& Keys, const int32 NumFrames)
 	}
 }
 
+PRAGMA_DISABLE_DEPRECATION_WARNINGS
+FCompressibleAnimData& FCompressibleAnimData::operator=(const FCompressibleAnimData&) = default;
+ICompressedAnimData::ICompressedAnimData(const ICompressedAnimData&) = default;
+PRAGMA_ENABLE_DEPRECATION_WARNINGS
+
 FCompressibleAnimData::FCompressibleAnimData(class UAnimSequence* InSeq, const bool bPerformStripping)
 	: CurveCompressionSettings(InSeq->CurveCompressionSettings)
 	, BoneCompressionSettings(InSeq->BoneCompressionSettings)
diff --git a/Engine/Source/Runtime/Engine/Public/Animation/AnimCompressionTypes.h b/Engine/Source/Runtime/Engine/Public/Animation/AnimCompressionTypes.h
--- a/Engine/Source/Runtime/Engine/Public/Animation/AnimCompressionTypes.h
+++ b/Engine/Source/Runtime/Engine/Public/Animation/AnimCompressionTypes.h
@@ -298,10 +298,6 @@ private:
 
 };
 
-PRAGMA_DISABLE_DEPRECATION_WARNINGS
-FCompressibleAnimData& FCompressibleAnimData::operator=(const FCompressibleAnimData&) = default;
-PRAGMA_ENABLE_DEPRECATION_WARNINGS
-
 typedef TSharedPtr<FCompressibleAnimData, ESPMode::ThreadSafe> FCompressibleAnimPtr;
 typedef TSharedRef<FCompressibleAnimData, ESPMode::ThreadSafe> FCompressibleAnimRef;
 
@@ -426,10 +422,6 @@ struct ENGINE_API ICompressedAnimData
 	virtual bool IsValid() const = 0;
 };
 
-PRAGMA_DISABLE_DEPRECATION_WARNINGS
-ICompressedAnimData::ICompressedAnimData(const ICompressedAnimData&) = default;
-PRAGMA_ENABLE_DEPRECATION_WARNINGS
-
 struct ENGINE_API FCompressibleAnimDataResult
 {
 	TArray<uint8> CompressedByteStream;
