Show HLODs in the scene outliner, under the proper level ("HLOD/HLOD0")

diff --git a/Engine/Source/Runtime/Engine/Private/WorldPartition/HLOD/HLODActor.cpp b/Engine/Source/Runtime/Engine/Private/WorldPartition/HLOD/HLODActor.cpp
--- a/Engine/Source/Runtime/Engine/Private/WorldPartition/HLOD/HLODActor.cpp
+++ b/Engine/Source/Runtime/Engine/Private/WorldPartition/HLOD/HLODActor.cpp
@@ -135,6 +135,11 @@ void AWorldPartitionHLOD::PostRegisterAllComponents()
 	{
 		SetupLoadedSubActors();
 	}
+
+#if WITH_EDITOR
+	bListedInSceneOutliner = true;
+	SetFolderPath(*FString::Printf(TEXT("HLOD/HLOD%d"), GetLODLevel()));
+#endif
 }
 
 void AWorldPartitionHLOD::PostUnregisterAllComponents()
