Some fixes in the skeleton remapping.

diff --git a/Engine/Source/Runtime/Engine/Classes/Animation/Skeleton.h b/Engine/Source/Runtime/Engine/Classes/Animation/Skeleton.h
--- a/Engine/Source/Runtime/Engine/Classes/Animation/Skeleton.h
+++ b/Engine/Source/Runtime/Engine/Classes/Animation/Skeleton.h
@@ -888,10 +888,11 @@ public:
 	 *		- and if more than 50 % of bones matches
 	 *  
 	 * @param	InSkelMesh	SkeletalMesh to compare the Skeleton against.
+	 * @param   bDoParentChainCheck When true (the default) this method also compares if chains match with the parent. 
 	 * 
 	 * @return				true if animation set can play on supplied SkeletalMesh, false if not.
 	 */
-	ENGINE_API bool IsCompatibleMesh(const USkeletalMesh* InSkelMesh) const;
+	ENGINE_API bool IsCompatibleMesh(const USkeletalMesh* InSkelMesh, bool bDoParentChainCheck=true) const;
 
 	/** Clears all cache data **/
 	ENGINE_API void ClearCacheData();
diff --git a/Engine/Source/Runtime/Engine/Private/Animation/Skeleton.cpp b/Engine/Source/Runtime/Engine/Private/Animation/Skeleton.cpp
--- a/Engine/Source/Runtime/Engine/Private/Animation/Skeleton.cpp
+++ b/Engine/Source/Runtime/Engine/Private/Animation/Skeleton.cpp
@@ -926,7 +926,7 @@ bool USkeleton::DoesParentChainMatch(int32 StartBoneIndex, const USkeletalMesh*
 	return true;
 }
 
-bool USkeleton::IsCompatibleMesh(const USkeletalMesh* InSkelMesh) const
+bool USkeleton::IsCompatibleMesh(const USkeletalMesh* InSkelMesh, bool bDoParentChainCheck) const
 {
 	// at least % of bone should match 
 	int32 NumOfBoneMatches = 0;
@@ -948,7 +948,7 @@ bool USkeleton::IsCompatibleMesh(const USkeletalMesh* InSkelMesh) const
 			++NumOfBoneMatches;
 
 			// follow the parent chain to verify the chain is same
-			if(!DoesParentChainMatch(SkeletonBoneIndex, InSkelMesh))
+			if(bDoParentChainCheck && !DoesParentChainMatch(SkeletonBoneIndex, InSkelMesh))
 			{
 				UE_LOG(LogAnimation, Verbose, TEXT("%s : Hierarchy does not match."), *MeshBoneName.ToString());
 				return false;
@@ -988,7 +988,7 @@ bool USkeleton::IsCompatibleMesh(const USkeletalMesh* InSkelMesh) const
 			}
 
 			// second follow the parent chain to verify the chain is same
-			if( !DoesParentChainMatch(SkeletonBoneIndex, InSkelMesh) )
+			if (bDoParentChainCheck && !DoesParentChainMatch(SkeletonBoneIndex, InSkelMesh))
 			{
 				UE_LOG(LogAnimation, Verbose, TEXT("%s : Hierarchy does not match."), *MeshBoneName.ToString());
 				return false;
diff --git a/Engine/Source/Runtime/Engine/Private/Components/SkeletalMeshComponent.cpp b/Engine/Source/Runtime/Engine/Private/Components/SkeletalMeshComponent.cpp
--- a/Engine/Source/Runtime/Engine/Private/Components/SkeletalMeshComponent.cpp
+++ b/Engine/Source/Runtime/Engine/Private/Components/SkeletalMeshComponent.cpp
@@ -484,7 +484,9 @@ bool USkeletalMeshComponent::NeedToSpawnAnimScriptInstance() const
 {
 	IAnimClassInterface* AnimClassInterface = IAnimClassInterface::GetFromClass(AnimClass);
 	const USkeleton* AnimSkeleton = (AnimClassInterface) ? AnimClassInterface->GetTargetSkeleton() : nullptr;
-	const bool bAnimSkelValid = !AnimClassInterface || (AnimSkeleton && SkeletalMesh && SkeletalMesh->GetSkeleton()->IsCompatible(AnimSkeleton) && AnimSkeleton->IsCompatibleMesh(SkeletalMesh));
+	const bool bSkeletonCompatible = (SkeletalMesh && AnimSkeleton) ? SkeletalMesh->GetSkeleton()->IsCompatible(AnimSkeleton) : false;
+	const bool bSkelMeshCompatible = (SkeletalMesh && AnimSkeleton) ? AnimSkeleton->IsCompatibleMesh(SkeletalMesh, false) : false;
+	const bool bAnimSkelValid = !AnimClassInterface || (bSkeletonCompatible && bSkelMeshCompatible);
 
 	if (AnimationMode == EAnimationMode::AnimationBlueprint && AnimClass && bAnimSkelValid)
 	{
