Copy/paste error, fixes broken cross-product preshader

diff --git a/Engine/Source/Runtime/Engine/Private/Shader/Preshader.cpp b/Engine/Source/Runtime/Engine/Private/Shader/Preshader.cpp
--- a/Engine/Source/Runtime/Engine/Private/Shader/Preshader.cpp
+++ b/Engine/Source/Runtime/Engine/Private/Shader/Preshader.cpp
@@ -360,7 +360,7 @@ void UE::Shader::EvaluatePreshader(const FUniformExpressionSet* UniformExpressio
 		case EPreshaderOpcode::Max: EvaluateBinaryOp(Stack, UE::Shader::Max); break;
 		case EPreshaderOpcode::Clamp: EvaluateTernaryOp(Stack, UE::Shader::Clamp); break;
 		case EPreshaderOpcode::Dot: EvaluateBinaryOp(Stack, UE::Shader::Dot); break;
-		case EPreshaderOpcode::Cross: EvaluateBinaryOp(Stack, UE::Shader::Dot); break;
+		case EPreshaderOpcode::Cross: EvaluateBinaryOp(Stack, UE::Shader::Cross); break;
 		case EPreshaderOpcode::Sqrt: EvaluateUnaryOp(Stack, UE::Shader::Sqrt); break;
 		case EPreshaderOpcode::Sin: EvaluateUnaryOp(Stack, UE::Shader::Sin); break;
 		case EPreshaderOpcode::Cos: EvaluateUnaryOp(Stack, UE::Shader::Cos); break;
