Fix Unsaved changes result in broken references in PIE

diff --git a/Engine/Source/Runtime/Engine/Private/WorldPartition/WorldPartitionLevelStreamingDynamic.cpp b/Engine/Source/Runtime/Engine/Private/WorldPartition/WorldPartitionLevelStreamingDynamic.cpp
--- a/Engine/Source/Runtime/Engine/Private/WorldPartition/WorldPartitionLevelStreamingDynamic.cpp
+++ b/Engine/Source/Runtime/Engine/Private/WorldPartition/WorldPartitionLevelStreamingDynamic.cpp
@@ -51,6 +51,7 @@ void UWorldPartitionLevelStreamingDynamic::Initialize(const UWorldPartitionRunti
 	check(ChildPackages.Num() == 0);
 	check(!WorldAsset.IsNull());
 
+	bIsModifiedForPIE = InCell.IsModifiedForPIE();
 	bShouldBeAlwaysLoaded = InCell.IsAlwaysLoaded();
 	StreamingPriority = InCell.GetStreamingPriority();
 	ChildPackages = InCell.GetPackages();
@@ -235,20 +236,19 @@ bool UWorldPartitionLevelStreamingDynamic::IssueLoadRequests()
 	{
 		if (ChildPackage.ContainerID == 0)
 		{
-			bool bNeedDup = false;
-
-			FString SubObjectName;
-			FString SubObjectContext;	
-			if (ChildPackage.LoadedPath.ToString().Split(TEXT("."), &SubObjectContext, &SubObjectName, ESearchCase::IgnoreCase, ESearchDir::FromEnd))
+			if (bIsModifiedForPIE)
 			{
-				if (AActor* ActorModifiedForPIE = World->PersistentLevel->ActorsModifiedForPIE.FindRef(*SubObjectName))
+				FString SubObjectName;
+				FString SubObjectContext;
+				if (ChildPackage.LoadedPath.ToString().Split(TEXT("."), &SubObjectContext, &SubObjectName, ESearchCase::IgnoreCase, ESearchDir::FromEnd))
 				{
-					ActorsToDuplicate.Add(ActorModifiedForPIE);
-					bNeedDup = true;
+					if (AActor* ActorModifiedForPIE = World->PersistentLevel->ActorsModifiedForPIE.FindRef(*SubObjectName))
+					{
+						ActorsToDuplicate.Add(ActorModifiedForPIE);
+					}
 				}
 			}
-
-			if (!bNeedDup)
+			else
 			{
 				ChildPackagesToLoad.Add(ChildPackage);
 			}
@@ -258,6 +258,7 @@ bool UWorldPartitionLevelStreamingDynamic::IssueLoadRequests()
 	// Duplicate unsaved actors
 	if (ActorsToDuplicate.Num())
 	{
+		check(ChildPackagesToLoad.Num() == 0);
 		// Create an actor container to make sure duplicated actors will share an outer to properly remap inter-actors references
 		UActorContainer* ActorContainer = NewObject<UActorContainer>(World->PersistentLevel);
 
@@ -303,6 +304,7 @@ bool UWorldPartitionLevelStreamingDynamic::IssueLoadRequests()
 	// Load saved actors
 	if (ChildPackagesToLoad.Num())
 	{
+		check(ActorsToDuplicate.Num() == 0);
 		FWorldPartitionLevelHelper::LoadActors(RuntimeLevel, ChildPackagesToLoad, PackageCache, [this, FinalizeLoading](bool bSucceeded)
 		{
 			if (!bSucceeded)
diff --git a/Engine/Source/Runtime/Engine/Private/WorldPartition/WorldPartitionRuntimeCell.cpp b/Engine/Source/Runtime/Engine/Private/WorldPartition/WorldPartitionRuntimeCell.cpp
--- a/Engine/Source/Runtime/Engine/Private/WorldPartition/WorldPartitionRuntimeCell.cpp
+++ b/Engine/Source/Runtime/Engine/Private/WorldPartition/WorldPartitionRuntimeCell.cpp
@@ -5,6 +5,9 @@
 UWorldPartitionRuntimeCell::UWorldPartitionRuntimeCell(const FObjectInitializer& ObjectInitializer)
 : Super(ObjectInitializer)
 , bIsAlwaysLoaded(false)
+#if WITH_EDITORONLY_DATA
+, bIsModifiedForPIE(false)
+#endif
 {
 }
 
diff --git a/Engine/Source/Runtime/Engine/Private/WorldPartition/WorldPartitionRuntimeSpatialHash.cpp b/Engine/Source/Runtime/Engine/Private/WorldPartition/WorldPartitionRuntimeSpatialHash.cpp
--- a/Engine/Source/Runtime/Engine/Private/WorldPartition/WorldPartitionRuntimeSpatialHash.cpp
+++ b/Engine/Source/Runtime/Engine/Private/WorldPartition/WorldPartitionRuntimeSpatialHash.cpp
@@ -23,6 +23,7 @@
 #include "GameFramework/WorldSettings.h"
 #include "ProfilingDebugging/ScopedTimers.h"
 #include "Engine/World.h"
+#include "Engine/Level.h"
 #include "Engine/Canvas.h"
 #include "DrawDebugHelpers.h"
 #include "DisplayDebugHelpers.h"
@@ -685,7 +686,8 @@ bool UWorldPartitionRuntimeSpatialHash::CreateStreamingGrid(const FSpatialHashRu
 	TRACE_CPUPROFILER_EVENT_SCOPE(UWorldPartitionRuntimeSpatialHash::CreateStreamingGrid);
 
 	UWorldPartition* WorldPartition = GetOuterUWorldPartition();
-	const bool bIsMainWorldPartition = (WorldPartition->GetWorld() == WorldPartition->GetTypedOuter<UWorld>());
+	UWorld* World = WorldPartition->GetWorld();
+	const bool bIsMainWorldPartition = (World == WorldPartition->GetTypedOuter<UWorld>());
 
 	FSpatialHashStreamingGrid& CurrentStreamingGrid = StreamingGrids.AddDefaulted_GetRef();
 	CurrentStreamingGrid.GridName = RuntimeGrid.GridName;
@@ -711,6 +713,7 @@ bool UWorldPartitionRuntimeSpatialHash::CreateStreamingGrid(const FSpatialHashRu
 			const uint32 CellIndex = TempCellMapping.Key;
 			const int32 CellCoordX = CellIndex % TempLevel.GridSize;
 			const int32 CellCoordY = CellIndex / TempLevel.GridSize;
+			bool bCellModifiedForPIE = false;
 
 			const FSquare2DGridHelper::FGridLevel::FGridCell& TempCell = TempLevel.Cells[TempCellMapping.Value];
 
@@ -733,24 +736,39 @@ bool UWorldPartitionRuntimeSpatialHash::CreateStreamingGrid(const FSpatialHashRu
 							continue;
 						}
 
-						// In PIE, Always loaded cell is not generated. Instead, always loaded actors will be added to AlwaysLoadedActorsForPIE.
-						// This will trigger loading/registration of these actors in the PersistentLevel (if not already loaded).
-						// Then, duplication of world for PIE will duplicate only these actors. 
-						// When stopping PIE, WorldPartition will release these FWorldPartitionReferences which 
-						// will unload actors that were not already loaded in the non PIE world.
-						if (bIsCellAlwaysLoaded && bIsMainWorldPartition && (Mode == EWorldPartitionStreamingMode::PIE || Mode == EWorldPartitionStreamingMode::EditorStandalone))
+						if (bIsMainWorldPartition && (Mode == EWorldPartitionStreamingMode::PIE || Mode == EWorldPartitionStreamingMode::EditorStandalone))
 						{
-							if (ActorInstance.ContainerInstance->Container == WorldPartition)
+							const FWorldPartitionActorDescView& ActorDescView = ActorInstance.GetActorDescView();
+
+							// In PIE, Always loaded cell is not generated. Instead, always loaded actors will be added to AlwaysLoadedActorsForPIE.
+							// This will trigger loading/registration of these actors in the PersistentLevel (if not already loaded).
+							// Then, duplication of world for PIE will duplicate only these actors. 
+							// When stopping PIE, WorldPartition will release these FWorldPartitionReferences which 
+							// will unload actors that were not already loaded in the non PIE world.
+							if (bIsCellAlwaysLoaded)
 							{
-								// This will load the actor if it isn't already loaded
-								FWorldPartitionReference Reference(WorldPartition, ActorInstance.Actor);
-								const FWorldPartitionActorDescView& ActorDescView = ActorInstance.GetActorDescView();
-								AActor* AlwaysLoadedActor = FindObject<AActor>(nullptr, *ActorDescView.GetActorPath().ToString());
-								AlwaysLoadedActorsForPIE.Emplace(Reference, AlwaysLoadedActor);
-								continue;
+								if (ActorInstance.ContainerInstance->Container == WorldPartition)
+								{
+									// This will load the actor if it isn't already loaded
+									FWorldPartitionReference Reference(WorldPartition, ActorInstance.Actor);
+									AActor* AlwaysLoadedActor = FindObject<AActor>(nullptr, *ActorDescView.GetActorPath().ToString());
+									AlwaysLoadedActorsForPIE.Emplace(Reference, AlwaysLoadedActor);
+									continue;
+								}
+							} // InPIE, If a Cell as at least one modified actor we need to flag all actors of this cell so they get added to the ActorsModifiedForPIE so they are all duplicated
+							else if(!bCellModifiedForPIE && Mode == EWorldPartitionStreamingMode::PIE && World->PersistentLevel->ActorsModifiedForPIE.Num() > 0)
+							{
+								FString SubObjectName;
+								FString SubObjectContext;
+								
+								ActorDescView.GetActorPath().ToString().Split(TEXT("."), &SubObjectContext, &SubObjectName, ESearchCase::IgnoreCase, ESearchDir::FromEnd);
+								if (World->PersistentLevel->ActorsModifiedForPIE.FindRef(*SubObjectName))
+								{
+									bCellModifiedForPIE = true;
+								}
 							}
 						}
-																
+
 						FilteredActors.Add(ActorInstance);
 					}
 				}
@@ -778,6 +796,7 @@ bool UWorldPartitionRuntimeSpatialHash::CreateStreamingGrid(const FSpatialHashRu
 
 				GridLevel.LayerCells[LayerCellIndex].GridCells.Add(StreamingCell);
 				StreamingCell->SetIsAlwaysLoaded(bIsCellAlwaysLoaded);
+				StreamingCell->SetIsModifiedForPIE(bCellModifiedForPIE);
 				StreamingCell->SetDataLayers(GridCellDataChunk.GetDataLayers());
 				StreamingCell->Level = Level;
 				FBox2D Bounds;
@@ -797,6 +816,20 @@ bool UWorldPartitionRuntimeSpatialHash::CreateStreamingGrid(const FSpatialHashRu
 						ReferencedHLODActors.Add(ActorDescView.GetHLODParent());
 					}
 					StreamingCell->AddActorToCell(ActorDescView, ActorInstance.ContainerInstance->ID, ActorInstance.ContainerInstance->Transform, ActorInstance.ContainerInstance->Container);
+
+					// Add all actors of this cell to the ActorsModifiedForPIE (for Duplication)
+					if (StreamingCell->IsModifiedForPIE())
+					{
+						if (AActor* Actor = FindObject<AActor>(nullptr, *ActorDescView.GetActorPath().ToString()))
+						{
+							World->PersistentLevel->ActorsModifiedForPIE.Add(Actor->GetFName(), Actor);
+						}
+						else
+						{
+							UE_LOG(LogWorldPartitionRuntimeSpatialHash, Error, TEXT(" Failed to add Actor : %s to ActorsModifiedForPIE"), *(ActorDescView.GetActorPath().ToString()));
+						}
+					}
+
 					UE_LOG(LogWorldPartitionRuntimeSpatialHash, Verbose, TEXT("  Actor : %s (%s) (Container %08x) Origin(%s)"), *(ActorDescView.GetActorPath().ToString()), *ActorDescView.GetGuid().ToString(EGuidFormats::UniqueObjectGuid), ActorInstance.ContainerInstance->ID, *FVector2D(ActorInstance.GetOrigin()).ToString());
 				}
 
diff --git a/Engine/Source/Runtime/Engine/Public/WorldPartition/WorldPartitionLevelStreamingDynamic.h b/Engine/Source/Runtime/Engine/Public/WorldPartition/WorldPartitionLevelStreamingDynamic.h
--- a/Engine/Source/Runtime/Engine/Public/WorldPartition/WorldPartitionLevelStreamingDynamic.h
+++ b/Engine/Source/Runtime/Engine/Public/WorldPartition/WorldPartitionLevelStreamingDynamic.h
@@ -63,6 +63,11 @@ private:
 	UPROPERTY()
 	bool bShouldBeAlwaysLoaded;
 
+#if WITH_EDITORONLY_DATA
+	UPROPERTY()
+	bool bIsModifiedForPIE;
+#endif
+
 	UPROPERTY()
 	TWeakObjectPtr<UWorldPartition> OuterWorldPartition;
 };
diff --git a/Engine/Source/Runtime/Engine/Public/WorldPartition/WorldPartitionRuntimeCell.h b/Engine/Source/Runtime/Engine/Public/WorldPartition/WorldPartitionRuntimeCell.h
--- a/Engine/Source/Runtime/Engine/Public/WorldPartition/WorldPartitionRuntimeCell.h
+++ b/Engine/Source/Runtime/Engine/Public/WorldPartition/WorldPartitionRuntimeCell.h
@@ -115,6 +115,9 @@ class UWorldPartitionRuntimeCell : public UObject
 	const TArray<FName>& GetDataLayers() const { return DataLayers; }
 
 #if WITH_EDITOR
+	bool IsModifiedForPIE() const { return bIsModifiedForPIE; };
+	void SetIsModifiedForPIE(bool bInIsModifiedForPIE) { bIsModifiedForPIE = bInIsModifiedForPIE; }
+	
 	void SetDataLayers(const TArray<const UDataLayer*> InDataLayers);
 	void AddCellData(const UWorldPartitionRuntimeCellData* InCellData);
 	virtual void AddActorToCell(const FWorldPartitionActorDescView& ActorDescView, uint32 InContainerID, const FTransform& InContainerTransform, const UActorDescContainer* InContainer) PURE_VIRTUAL(UWorldPartitionRuntimeCell::AddActorToCell,);
@@ -134,6 +137,11 @@ protected:
 	UPROPERTY()
 	bool bIsAlwaysLoaded;
 
+#if WITH_EDITORONLY_DATA
+	UPROPERTY(Transient)
+	bool bIsModifiedForPIE;
+#endif
+
 private:
 	UPROPERTY()
 	TMap<const TSubclassOf<UWorldPartitionRuntimeCellData>, TObjectPtr<const UWorldPartitionRuntimeCellData>> CellDataMap;
