Adding option for linked layers to read montage eval data from their main anim instance

diff --git a/Engine/Source/Runtime/Engine/Classes/Animation/AnimInstance.h b/Engine/Source/Runtime/Engine/Classes/Animation/AnimInstance.h
--- a/Engine/Source/Runtime/Engine/Classes/Animation/AnimInstance.h
+++ b/Engine/Source/Runtime/Engine/Classes/Animation/AnimInstance.h
@@ -447,6 +447,10 @@ class ENGINE_API UAnimInstance : public UObject
 	UPROPERTY(EditDefaultsOnly, Category = Notifies)
 	uint8 bPropagateNotifiesToLinkedInstances : 1;
 
+	/** If true, linked instances will use the main instance's montage data. (i.e. playing a montage on a main instance will play it on the linked layer too.) */
+	UPROPERTY(EditDefaultsOnly, Category = "Montage")
+	uint8 bUseMainInstanceMontageEvaluationData: 1;
+
 private:
 	/** True when Montages are being ticked, and Montage Events should be queued. 
 	 * When Montage are being ticked, we queue AnimNotifies and Events. We trigger notifies first, then Montage events. */
@@ -543,6 +547,12 @@ public:
 	UFUNCTION(BlueprintCallable, Category = "Notifies")
 	void SetPropagateNotifiesToLinkedInstances(bool bSet) { bPropagateNotifiesToLinkedInstances = bSet; }
 
+	UFUNCTION(BlueprintCallable, Category = "Animation Blueprint Linking|Montage")
+	bool IsUsingMainInstanceMontageEvaluationData() const { return bUseMainInstanceMontageEvaluationData; }
+
+	UFUNCTION(BlueprintCallable, Category = "Animation Blueprint Linking|Montage")
+	void SetUseMainInstanceMontageEvaluationData(bool bSet) { bUseMainInstanceMontageEvaluationData = bSet; }
+
 private:
 	// Does this anim instance need immediate update (rather than parallel)?
 	bool NeedsImmediateUpdate(float DeltaSeconds) const;
diff --git a/Engine/Source/Runtime/Engine/Private/Animation/AnimInstance.cpp b/Engine/Source/Runtime/Engine/Private/Animation/AnimInstance.cpp
--- a/Engine/Source/Runtime/Engine/Private/Animation/AnimInstance.cpp
+++ b/Engine/Source/Runtime/Engine/Private/Animation/AnimInstance.cpp
@@ -29,6 +29,7 @@
 #include "Animation/AnimNode_LinkedAnimGraph.h"
 #include "Animation/AnimNode_LinkedInputPose.h"
 #include "Animation/AnimNode_LinkedAnimLayer.h"
+#include "UObject/UE5MainStreamObjectVersion.h"
 
 /** Anim stats */
 
@@ -111,6 +112,7 @@ UAnimInstance::UAnimInstance(const FObjectInitializer& ObjectInitializer)
 
 	bReceiveNotifiesFromLinkedInstances = false;
 	bPropagateNotifiesToLinkedInstances = false;
+	bUseMainInstanceMontageEvaluationData = false;
 }
 
 // this is only used by montage marker based sync
@@ -330,6 +332,15 @@ bool UAnimInstance::UpdateSnapshotAndSkipRemainingUpdate()
 
 void UAnimInstance::UpdateMontage(float DeltaSeconds)
 {
+	// Don't update montages if we are using the main instance's montage eval data and we are not the main instance.
+	if (IsUsingMainInstanceMontageEvaluationData())
+	{
+		if (GetOwningComponent()->GetAnimInstance() != this)
+		{
+			return;
+		}
+	}
+
 	// update montage weight
 	Montage_UpdateWeight(DeltaSeconds);
 
diff --git a/Engine/Source/Runtime/Engine/Private/Animation/AnimInstanceProxy.cpp b/Engine/Source/Runtime/Engine/Private/Animation/AnimInstanceProxy.cpp
--- a/Engine/Source/Runtime/Engine/Private/Animation/AnimInstanceProxy.cpp
+++ b/Engine/Source/Runtime/Engine/Private/Animation/AnimInstanceProxy.cpp
@@ -549,9 +549,10 @@ void FAnimInstanceProxy::InitializeObjects(UAnimInstance* InAnimInstance)
 	DECLARE_SCOPE_HIERARCHICAL_COUNTER_FUNC()
 
 	SkeletalMeshComponent = InAnimInstance->GetSkelMeshComponent();
-	if(SkeletalMeshComponent->GetAnimInstance())
+	if(UAnimInstance* MainAnimInstance = SkeletalMeshComponent->GetAnimInstance())
 	{
-		MainInstanceProxy = &SkeletalMeshComponent->GetAnimInstance()->GetProxyOnAnyThread<FAnimInstanceProxy>();
+		MainInstanceProxy = &MainAnimInstance->GetProxyOnAnyThread<FAnimInstanceProxy>();
+		MainMontageEvaluationData = InAnimInstance->IsUsingMainInstanceMontageEvaluationData() ? &MainInstanceProxy->MontageEvaluationData : &MontageEvaluationData;
 	}
 
 	if (SkeletalMeshComponent->SkeletalMesh != nullptr)
@@ -1284,14 +1285,14 @@ void FAnimInstanceProxy::SlotEvaluatePoseWithBlendProfiles(const FName& SlotNode
 		PerBoneWeightTotals[BoneIndex] = 0.0f;
 	}
 	PerBoneWeightTotalsAdditive = PerBoneWeightTotals;
-	PerBoneWeights.Reset(MontageEvaluationData.Num());
+	PerBoneWeights.Reset(GetMontageEvaluationData().Num());
 
 	//------------------------------------------
 	// Figure out what poses we need to blend.
 	//------------------------------------------
 	// Make a list of all the montages that use the slot we're interested in.
 	// Split this in an additive and non additive list.
-	check(MontageEvaluationData.Num() < 255); // Make sure we're in limits and not blending more than 255 poses (that would indicate another issue anyway)
+	check(GetMontageEvaluationData().Num() < 255); // Make sure we're in limits and not blending more than 255 poses (that would indicate another issue anyway)
 	TArray<FSlotEvaluationPose>& Poses = BlendProfileScratchData.Poses;
 	TArray<FSlotEvaluationPose>& AdditivePoses = BlendProfileScratchData.AdditivePoses;
 	TArray<uint8, TInlineAllocator<8>>& PoseIndices = BlendProfileScratchData.PoseIndices;
@@ -1313,9 +1314,9 @@ void FAnimInstanceProxy::SlotEvaluatePoseWithBlendProfiles(const FName& SlotNode
 	int32 CurrentPoseIndex = 0;
 	float NonAdditiveTotalWeight = 0.0f;
 	float AdditiveTotalWeight = 0.0f;
-	for (int32 Index = 0; Index < MontageEvaluationData.Num(); ++Index)
+	for (int32 Index = 0; Index < GetMontageEvaluationData().Num(); ++Index)
 	{
-		const FMontageEvaluationState& EvalState = MontageEvaluationData[Index];
+		const FMontageEvaluationState& EvalState = GetMontageEvaluationData()[Index];
 		const UAnimMontage* Montage = EvalState.Montage.Get();
 		if (Montage && Montage->IsValidSlot(SlotNodeName))
 		{
@@ -1611,7 +1612,7 @@ void FAnimInstanceProxy::SlotEvaluatePose(const FName& SlotNodeName, const FAnim
 	}
 
 	// Check if we are blending a montage using a blend profile, if so take a special code path for this.
-	for (const FMontageEvaluationState& EvalState : MontageEvaluationData)
+	for (const FMontageEvaluationState& EvalState : GetMontageEvaluationData())
 	{
 		if (EvalState.Montage.IsValid() && EvalState.ActiveBlendProfile && EvalState.Montage->IsValidSlot(SlotNodeName))
 		{
@@ -1639,7 +1640,7 @@ void FAnimInstanceProxy::SlotEvaluatePose(const FName& SlotNodeName, const FAnim
 #if DEBUG_MONTAGEINSTANCE_WEIGHT
 	float TotalWeight = 0.f;
 #endif // DEBUG_MONTAGEINSTANCE_WEIGHT
-	for (const FMontageEvaluationState& EvalState : MontageEvaluationData)
+	for (const FMontageEvaluationState& EvalState : GetMontageEvaluationData())
 	{
 		// If MontageEvaluationData is not valid anymore, pass-through AnimSlot.
 		// This can happen if InitAnim pushes a RefreshBoneTransforms when not rendered,
@@ -1788,7 +1789,7 @@ void FAnimInstanceProxy::GetSlotWeight(const FName& SlotNodeName, float& out_Slo
 #endif
 
 	// first get all the montage instance weight this slot node has
-	for (const FMontageEvaluationState& EvalState : MontageEvaluationData)
+	for (const FMontageEvaluationState& EvalState : GetMontageEvaluationData())
 	{
 		if (EvalState.Montage.IsValid())
 		{
@@ -1848,10 +1849,10 @@ void FAnimInstanceProxy::GetSlotWeight(const FName& SlotNodeName, float& out_Slo
 const FMontageEvaluationState* FAnimInstanceProxy::GetActiveMontageEvaluationState() const
 {
 	// Start from end, as most recent instances are added at the end of the queue.
-	int32 const NumInstances = MontageEvaluationData.Num();
+	int32 const NumInstances = GetMontageEvaluationData().Num();
 	for (int32 InstanceIndex = NumInstances - 1; InstanceIndex >= 0; InstanceIndex--)
 	{
-		const FMontageEvaluationState& EvaluationData = MontageEvaluationData[InstanceIndex];
+		const FMontageEvaluationState& EvaluationData = GetMontageEvaluationData()[InstanceIndex];
 		if (EvaluationData.bIsActive)
 		{
 			return &EvaluationData;
@@ -2221,6 +2222,18 @@ float FAnimInstanceProxy::GetRelevantAnimTimeFraction(int32 MachineIndex, int32
 	return 0.0f;
 }
 
+TArray<FMontageEvaluationState>& FAnimInstanceProxy::GetMontageEvaluationData()
+{
+	check(MainMontageEvaluationData);
+	return *MainMontageEvaluationData;
+}
+
+const TArray<FMontageEvaluationState>& FAnimInstanceProxy::GetMontageEvaluationData() const
+{
+	check(MainMontageEvaluationData);
+	return *MainMontageEvaluationData;
+}
+
 FAnimNode_AssetPlayerBase* FAnimInstanceProxy::GetRelevantAssetPlayerFromState(int32 MachineIndex, int32 StateIndex)
 {
 	FAnimNode_AssetPlayerBase* ResultPlayer = nullptr;
@@ -2646,7 +2659,7 @@ void FAnimInstanceProxy::ResetDynamics()
 #if ANIM_TRACE_ENABLED
 void FAnimInstanceProxy::TraceMontageEvaluationData(const FAnimationUpdateContext& InContext, const FName& InSlotName)
 {
-	for (const FMontageEvaluationState& MontageEvaluationState : MontageEvaluationData)
+	for (const FMontageEvaluationState& MontageEvaluationState : GetMontageEvaluationData())
 	{
 		if (MontageEvaluationState.Montage != nullptr && MontageEvaluationState.Montage->IsValidSlot(InSlotName))
 		{
diff --git a/Engine/Source/Runtime/Engine/Public/Animation/AnimInstanceProxy.h b/Engine/Source/Runtime/Engine/Public/Animation/AnimInstanceProxy.h
--- a/Engine/Source/Runtime/Engine/Public/Animation/AnimInstanceProxy.h
+++ b/Engine/Source/Runtime/Engine/Public/Animation/AnimInstanceProxy.h
@@ -131,6 +131,7 @@ public:
 		, FrameCounterForUpdate(0)
 		, FrameCounterForNodeUpdate(0)
 		, CacheBonesRecursionCounter(0)
+		, MainMontageEvaluationData(&MontageEvaluationData)
 		, bUpdatingRoot(false)
 		, bBoneCachesInvalidated(false)
 		, bShouldExtractRootMotion(false)
@@ -156,6 +157,7 @@ public:
 		, FrameCounterForUpdate(0)
 		, FrameCounterForNodeUpdate(0)
 		, CacheBonesRecursionCounter(0)
+		, MainMontageEvaluationData(&MontageEvaluationData)
 		, bUpdatingRoot(false)
 		, bBoneCachesInvalidated(false)
 		, bShouldExtractRootMotion(false)
@@ -763,13 +765,13 @@ protected:
 		Note that there might be multiple Active at the same time. This will only return the first active one it finds. **/
 	const FMontageEvaluationState* GetActiveMontageEvaluationState() const;
 
-	/** Access montage array data */
-	TArray<FMontageEvaluationState>& GetMontageEvaluationData() { return MontageEvaluationData; }
-
 	TMap<FName, float>& GetSlotGroupInertializationRequestMap() { return SlotGroupInertializationRequestMap; }
 
 	/** Access montage array data */
-	const TArray<FMontageEvaluationState>& GetMontageEvaluationData() const { return MontageEvaluationData; }
+	TArray<FMontageEvaluationState>& GetMontageEvaluationData();
+
+	/** Access montage array data */
+	const TArray<FMontageEvaluationState>& GetMontageEvaluationData() const;
 
 	/** Check whether we have active morph target curves */
 	/** Gets the most relevant asset player in a specified state */
@@ -1037,6 +1039,7 @@ private:
 private:
 	/** Copy of UAnimInstance::MontageInstances data used for update & evaluation */
 	TArray<FMontageEvaluationState> MontageEvaluationData;
+	TArray<FMontageEvaluationState>* MainMontageEvaluationData;
 
 	// Inertialization request for each slot. Mapped with SlotNameToTrackerIndex. Initialized to a value of 1.0f, which indicates no request for this slot.
 	TMap<FName, float> SlotGroupInertializationRequestMap;
