CIS: Add missing 'Abstract' keyword

diff --git a/Engine/Source/Runtime/Engine/Public/WorldPartition/WorldPartitionRuntimeCell.h b/Engine/Source/Runtime/Engine/Public/WorldPartition/WorldPartitionRuntimeCell.h
--- a/Engine/Source/Runtime/Engine/Public/WorldPartition/WorldPartitionRuntimeCell.h
+++ b/Engine/Source/Runtime/Engine/Public/WorldPartition/WorldPartitionRuntimeCell.h
@@ -103,7 +103,7 @@ enum class EWorldPartitionRuntimeCellState : uint8
 /**
  * Represents a PIE/Game streaming cell which points to external actor/data chunk packages
  */
-UCLASS()
+UCLASS(Abstract)
 class UWorldPartitionRuntimeCell : public UObject
 {
 	GENERATED_UCLASS_BODY()
