LevelInstance: Avoid Instance Map/Actor package from showing up in ContentBrowser

diff --git a/Engine/Source/Runtime/Engine/Private/Actor.cpp b/Engine/Source/Runtime/Engine/Private/Actor.cpp
--- a/Engine/Source/Runtime/Engine/Private/Actor.cpp
+++ b/Engine/Source/Runtime/Engine/Private/Actor.cpp
@@ -374,7 +374,7 @@ bool AActor::IsAsset() const
 {
 	// External actors are considered assets, to allow using the asset logic for save dialogs, etc.
 	// Also, they return true even if pending kill, in order to show up as deleted in these dialogs.
-	return IsPackageExternal() && !IsChildActor() && !HasAnyFlags(RF_Transient | RF_ClassDefaultObject);
+	return IsPackageExternal() && !GetPackage()->HasAnyFlags(RF_Transient) && !IsChildActor() && !HasAnyFlags(RF_Transient | RF_ClassDefaultObject);
 }
 
 bool AActor::PreSaveRoot(const TCHAR* InFilename)
diff --git a/Engine/Source/Runtime/Engine/Private/Level.cpp b/Engine/Source/Runtime/Engine/Private/Level.cpp
--- a/Engine/Source/Runtime/Engine/Private/Level.cpp
+++ b/Engine/Source/Runtime/Engine/Private/Level.cpp
@@ -791,7 +791,12 @@ void ULevel::PostLoad()
 		{
 			const FString& ActorPackageName = ActorPackageNames[i];
 
+			// Propagate RF_Transient from the Level Package in case we are loading an instanced level
 			UPackage* ActorPackage = bInstanced ? CreatePackage( *InstancePackageNames[i]) : nullptr;
+			if (ActorPackage && LevelPackage->HasAnyFlags(RF_Transient))
+			{
+				ActorPackage->SetFlags(RF_Transient);
+			}
 
 			ActorPackage = LoadPackage(ActorPackage, *ActorPackageName, bPackageForPIE ? LOAD_PackageForPIE : LOAD_None, nullptr, &InstancingContext);
 
diff --git a/Engine/Source/Runtime/Engine/Private/LevelStreaming.cpp b/Engine/Source/Runtime/Engine/Private/LevelStreaming.cpp
--- a/Engine/Source/Runtime/Engine/Private/LevelStreaming.cpp
+++ b/Engine/Source/Runtime/Engine/Private/LevelStreaming.cpp
@@ -1132,6 +1132,17 @@ bool ULevelStreaming::RequestLevel(UWorld* PersistentWorld, bool bAllowLevelLoad
 			// Kick off async load request.
 			STAT_ADD_CUSTOMMESSAGE_NAME( STAT_NamedMarker, *(FString( TEXT( "RequestLevel - " ) + DesiredPackageName.ToString() )) );
 			TRACE_BOOKMARK(TEXT("RequestLevel - %s"), *DesiredPackageName.ToString());
+
+			// When loading an instanced package we want to avoid it being processed as an asset so we make sure to set it RF_Transient.
+			// If the package is not created here, it will get created by the LoadPackageAsync call. 
+			// Gameworld packages (PIE) are already ignored so we can let LoadPackageAsync do its job.
+			if (!bIsGameWorld && DesiredPackageName != NAME_None && PackagePath.GetPackageFName() != DesiredPackageName)
+			{
+				UPackage* NewPackage = CreatePackage(*DesiredPackageName.ToString());
+				NewPackage->SetFlags(RF_Transient);
+			}
+
+
 			LoadPackageAsync(PackagePath, DesiredPackageName, FLoadPackageAsyncDelegate::CreateUObject(this, &ULevelStreaming::AsyncLevelLoadComplete), nullptr, PackageFlags, PIEInstanceID, GetPriority());
 
 			// streamingServer: server loads everything?
