Fix crash in Cascade due to the scene textures uniform buffer not being set.

The RHI handle for the UB needs to be extracted on the RDG timeline, which is slightly more complex than we'd like, because Cascade lives inside the Engine module, so it cannot declare pass parameters containing FSceneTextureUniformParameters. As a workaround, we add the pass which extracts the handle inside RenderOpaqueFX in the Renderer module and pass the pointer to Cascade through a new function in its interface.

Also cleaned up the unnecessary uniform buffer parameters on PostRenderOpaque().

#rnx

diff --git a/Engine/Plugins/FX/Niagara/Source/Niagara/Classes/NiagaraEmitterInstanceBatcher.h b/Engine/Plugins/FX/Niagara/Source/Niagara/Classes/NiagaraEmitterInstanceBatcher.h
--- a/Engine/Plugins/FX/Niagara/Source/Niagara/Classes/NiagaraEmitterInstanceBatcher.h
+++ b/Engine/Plugins/FX/Niagara/Source/Niagara/Classes/NiagaraEmitterInstanceBatcher.h
@@ -86,8 +86,6 @@ public:
 	virtual void PostRenderOpaque(
 		class FRDGBuilder& GraphBuilder,
 		FRHIUniformBuffer* ViewUniformBuffer,
-		const class FShaderParametersMetadata* SceneTexturesUniformBufferStruct,
-		FRHIUniformBuffer* SceneTexturesUniformBuffer,
 		bool bAllowGPUParticleUpdate) override;
 
 	/**
diff --git a/Engine/Plugins/FX/Niagara/Source/Niagara/Private/NiagaraEmitterInstanceBatcher.cpp b/Engine/Plugins/FX/Niagara/Source/Niagara/Private/NiagaraEmitterInstanceBatcher.cpp
--- a/Engine/Plugins/FX/Niagara/Source/Niagara/Private/NiagaraEmitterInstanceBatcher.cpp
+++ b/Engine/Plugins/FX/Niagara/Source/Niagara/Private/NiagaraEmitterInstanceBatcher.cpp
@@ -483,7 +483,7 @@ void NiagaraEmitterInstanceBatcher::ProcessPendingTicksFlush(FRHICommandListImme
 			FRDGBuilder GraphBuilder(RHICmdList);
 			PreInitViews(GraphBuilder, true);
 			PostInitViews(GraphBuilder, View.ViewUniformBuffer, true);
-			PostRenderOpaque(GraphBuilder, View.ViewUniformBuffer, nullptr, nullptr, true);
+			PostRenderOpaque(GraphBuilder, View.ViewUniformBuffer, true);
 			GraphBuilder.Execute();
 			break;
 		}
@@ -1514,14 +1514,14 @@ void NiagaraEmitterInstanceBatcher::PostInitViews(FRDGBuilder& GraphBuilder, FRH
 	}
 }
 
-void NiagaraEmitterInstanceBatcher::PostRenderOpaque(FRDGBuilder& GraphBuilder, FRHIUniformBuffer* ViewUniformBuffer, const class FShaderParametersMetadata* SceneTexturesUniformBufferStruct, FRHIUniformBuffer* SceneTexturesUniformBuffer, bool bAllowGPUParticleUpdate)
+void NiagaraEmitterInstanceBatcher::PostRenderOpaque(FRDGBuilder& GraphBuilder, FRHIUniformBuffer* ViewUniformBuffer, bool bAllowGPUParticleUpdate)
 {
 	LLM_SCOPE(ELLMTag::Niagara);
 
 	AddPass(
 		GraphBuilder,
 		RDG_EVENT_NAME("Niagara::PostRenderOpaque"),
-		[this, ViewUniformBuffer, SceneTexturesUniformBufferStruct, SceneTexturesUniformBuffer, bAllowGPUParticleUpdate](FRHICommandListImmediate& RHICmdList)
+		[this, ViewUniformBuffer, bAllowGPUParticleUpdate](FRHICommandListImmediate& RHICmdList)
 		{
 			if (bAllowGPUParticleUpdate && FNiagaraUtilities::AllowGPUParticles(GetShaderPlatform()))
 			{
diff --git a/Engine/Source/Runtime/Engine/Private/Particles/FXSystem.cpp b/Engine/Source/Runtime/Engine/Private/Particles/FXSystem.cpp
--- a/Engine/Source/Runtime/Engine/Private/Particles/FXSystem.cpp
+++ b/Engine/Source/Runtime/Engine/Private/Particles/FXSystem.cpp
@@ -455,7 +455,7 @@ void FFXSystem::PreRender(FRDGBuilder& GraphBuilder, FRHIUniformBuffer* ViewUnif
 				PrepareGPUSimulation(RHICmdList);
 
 				RHICmdList.SetCurrentStat(GET_STATID(STAT_CLM_FXPreRender_Simulate));
-				SimulateGPUParticles(RHICmdList, EParticleSimulatePhase::Main, nullptr, nullptr, nullptr, nullptr);
+				SimulateGPUParticles(RHICmdList, EParticleSimulatePhase::Main, nullptr, nullptr);
 
 				RHICmdList.SetCurrentStat(GET_STATID(STAT_CLM_FXPreRender_Finalize));
 				FinalizeGPUSimulation(RHICmdList);
@@ -466,7 +466,7 @@ void FFXSystem::PreRender(FRDGBuilder& GraphBuilder, FRHIUniformBuffer* ViewUnif
 					PrepareGPUSimulation(RHICmdList);
 
 					RHICmdList.SetCurrentStat(GET_STATID(STAT_CLM_FXPreRender_SimulateCDF));
-					SimulateGPUParticles(RHICmdList, EParticleSimulatePhase::CollisionDistanceField, ViewUniformBuffer, GlobalDistanceFieldParameterData, nullptr, nullptr);
+					SimulateGPUParticles(RHICmdList, EParticleSimulatePhase::CollisionDistanceField, ViewUniformBuffer, GlobalDistanceFieldParameterData);
 					//particles rendered during basepass may need to read pos/velocity buffers; must finalize unless we know for sure that nothing in base pass will read it.
 					RHICmdList.SetCurrentStat(GET_STATID(STAT_CLM_FXPreRender_FinalizeCDF));
 					FinalizeGPUSimulation(RHICmdList);
@@ -479,20 +479,16 @@ void FFXSystem::PreRender(FRDGBuilder& GraphBuilder, FRHIUniformBuffer* ViewUnif
 void FFXSystem::PostRenderOpaque(
 	FRDGBuilder& GraphBuilder,
 	FRHIUniformBuffer* ViewUniformBuffer,
-	const FShaderParametersMetadata* SceneTexturesUniformBufferStruct,
-	FRHIUniformBuffer* SceneTexturesUniformBuffer,
 	bool bAllowGPUParticleUpdate)
 {
 	if (RHISupportsGPUParticles() && IsParticleCollisionModeSupported(GetShaderPlatform(), PCM_DepthBuffer) && bAllowGPUParticleUpdate)
 	{
-		AddPass(
-			GraphBuilder,
-			RDG_EVENT_NAME("FFXSystem::PostRenderOpaque"),
-			[this, ViewUniformBuffer, SceneTexturesUniformBufferStruct, SceneTexturesUniformBuffer](FRHICommandListImmediate& RHICmdList)
+		AddPass(GraphBuilder, RDG_EVENT_NAME("FFXSystem::PostRenderOpaque"), 
+			[this, ViewUniformBuffer](FRHICommandListImmediate& RHICmdList)
 			{
 				SCOPED_DRAW_EVENT(RHICmdList, GPUParticles_PostRenderOpaque);
 				PrepareGPUSimulation(RHICmdList);
-				SimulateGPUParticles(RHICmdList, EParticleSimulatePhase::CollisionDepthBuffer, ViewUniformBuffer, NULL, SceneTexturesUniformBufferStruct, SceneTexturesUniformBuffer);
+				SimulateGPUParticles(RHICmdList, EParticleSimulatePhase::CollisionDepthBuffer, ViewUniformBuffer, nullptr);
 				FinalizeGPUSimulation(RHICmdList);
 			}
 		);
diff --git a/Engine/Source/Runtime/Engine/Private/Particles/FXSystemPrivate.h b/Engine/Source/Runtime/Engine/Private/Particles/FXSystemPrivate.h
--- a/Engine/Source/Runtime/Engine/Private/Particles/FXSystemPrivate.h
+++ b/Engine/Source/Runtime/Engine/Private/Particles/FXSystemPrivate.h
@@ -159,8 +159,6 @@ public:
 	virtual void PostRenderOpaque(
 		class FRDGBuilder& GraphBuilder,
 		FRHIUniformBuffer* ViewUniformBuffer,
-		const FShaderParametersMetadata* SceneTexturesUniformBufferStruct,
-		FRHIUniformBuffer* SceneTexturesUniformBuffer,
 		bool bAllowGPUParticleUpdate) override;
 	// End FFXSystemInterface.
 
@@ -221,6 +219,8 @@ public:
 	/** Get the shared SortManager, used in the rendering loop to call FGPUSortManager::OnPreRender() and FGPUSortManager::OnPostRenderOpaque() */
 	virtual FGPUSortManager* GetGPUSortManager() const override;
 
+	virtual void SetSceneTexturesUniformBuffer(FRHIUniformBuffer* InSceneTexturesUniformParams) override { SceneTexturesUniformParams = InSceneTexturesUniformParams; }
+
 private:
 
 	/**
@@ -290,9 +290,7 @@ private:
 		FRHICommandListImmediate& RHICmdList,
 		EParticleSimulatePhase::Type Phase,
 		FRHIUniformBuffer* ViewUniformBuffer,
-		const FGlobalDistanceFieldParameterData* GlobalDistanceFieldParameterData,
-		const FShaderParametersMetadata* SceneTexturesUniformBufferStruct,
-		FRHIUniformBuffer* SceneTexturesUniformBuffer
+		const FGlobalDistanceFieldParameterData* GlobalDistanceFieldParameterData
 		);
 
 	/**
@@ -346,5 +344,7 @@ private:
 	EParticleSimulatePhase::Type PhaseToWaitForTemporalEffect = EParticleSimulatePhase::First;
 	EParticleSimulatePhase::Type PhaseToBroadcastTemporalEffect = EParticleSimulatePhase::First;
 #endif
+
+	FRHIUniformBuffer* SceneTexturesUniformParams = nullptr;
 };
 
diff --git a/Engine/Source/Runtime/Engine/Private/Particles/FXSystemSet.cpp b/Engine/Source/Runtime/Engine/Private/Particles/FXSystemSet.cpp
--- a/Engine/Source/Runtime/Engine/Private/Particles/FXSystemSet.cpp
+++ b/Engine/Source/Runtime/Engine/Private/Particles/FXSystemSet.cpp
@@ -190,11 +190,18 @@ void FFXSystemSet::PreRender(class FRDGBuilder& GraphBuilder, FRHIUniformBuffer*
 	}
 }
 
+void FFXSystemSet::SetSceneTexturesUniformBuffer(FRHIUniformBuffer* InSceneTexturesUniformParams)
+{
+	for (FFXSystemInterface* FXSystem : FXSystems)
+	{
+		check(FXSystem);
+		FXSystem->SetSceneTexturesUniformBuffer(InSceneTexturesUniformParams);
+	}
+}
+
 void FFXSystemSet::PostRenderOpaque(
 	class FRDGBuilder& GraphBuilder,
 	FRHIUniformBuffer* ViewUniformBuffer,
-	const class FShaderParametersMetadata* SceneTexturesUniformBufferStruct,
-	FRHIUniformBuffer* SceneTexturesUniformBuffer,
 	bool bAllowGPUParticleUpdate)
 {
 	for (FFXSystemInterface* FXSystem : FXSystems)
@@ -203,8 +210,6 @@ void FFXSystemSet::PostRenderOpaque(
 		FXSystem->PostRenderOpaque(
 			GraphBuilder,
 			ViewUniformBuffer,
-			SceneTexturesUniformBufferStruct,
-			SceneTexturesUniformBuffer,
 			bAllowGPUParticleUpdate
 		);
 	}
diff --git a/Engine/Source/Runtime/Engine/Private/Particles/FXSystemSet.h b/Engine/Source/Runtime/Engine/Private/Particles/FXSystemSet.h
--- a/Engine/Source/Runtime/Engine/Private/Particles/FXSystemSet.h
+++ b/Engine/Source/Runtime/Engine/Private/Particles/FXSystemSet.h
@@ -48,8 +48,6 @@ public:
 	virtual void PostRenderOpaque(
 		class FRDGBuilder& GraphBuilder,
 		FRHIUniformBuffer* ViewUniformBuffer,
-		const class FShaderParametersMetadata* SceneTexturesUniformBufferStruct,
-		FRHIUniformBuffer* SceneTexturesUniformBuffer,
 		bool bAllowGPUParticleUpdate) override;
 
 	virtual void OnDestroy() override;
@@ -58,6 +56,8 @@ public:
 	/** Get the shared SortManager, used in the rendering loop to call FGPUSortManager::OnPreRender() and FGPUSortManager::OnPostRenderOpaque() */
 	virtual FGPUSortManager* GetGPUSortManager() const override;
 
+	virtual void SetSceneTexturesUniformBuffer(FRHIUniformBuffer* InSceneTexturesUniformParams) override;
+
 protected:
 
 	/** By making the destructor protected, an instance must be destroyed via FFXSystemInterface::Destroy. */
diff --git a/Engine/Source/Runtime/Engine/Private/Particles/ParticleGpuSimulation.cpp b/Engine/Source/Runtime/Engine/Private/Particles/ParticleGpuSimulation.cpp
--- a/Engine/Source/Runtime/Engine/Private/Particles/ParticleGpuSimulation.cpp
+++ b/Engine/Source/Runtime/Engine/Private/Particles/ParticleGpuSimulation.cpp
@@ -1062,9 +1062,7 @@ public:
 		const FParticleAttributesTexture& InRenderAttributesTexture,
 		FRHIUniformBuffer* ViewUniformBuffer,
 		ERHIFeatureLevel::Type FeatureLevel,
-		const FGlobalDistanceFieldParameterData* GlobalDistanceFieldParameterData,
-		const FShaderParametersMetadata* SceneTexturesUniformBufferStruct,
-		FRHIUniformBuffer* SceneTexturesUniformBuffer
+		const FGlobalDistanceFieldParameterData* GlobalDistanceFieldParameterData
 		)
 	{
 		FRHIPixelShader* PixelShaderRHI = RHICmdList.GetBoundPixelShader();
@@ -1080,7 +1078,6 @@ public:
 			check(ViewUniformBuffer != NULL);
 			FGlobalShader::SetParameters<FViewUniformShaderParameters>(RHICmdList, PixelShaderRHI, ViewUniformBuffer);
 
-			SetUniformBufferParameter(RHICmdList, PixelShaderRHI, GetUniformBufferParameter(SceneTexturesUniformBufferStruct), SceneTexturesUniformBuffer);
 			SetTextureParameter(
 				RHICmdList, 
 				PixelShaderRHI,
@@ -1427,7 +1424,6 @@ void ExecuteSimulationCommands(
 	FParticleSimulationResources* ParticleSimulationResources,
 	FRHIUniformBuffer* ViewUniformBuffer,
 	const FGlobalDistanceFieldParameterData* GlobalDistanceFieldParameterData,
-	const FShaderParametersMetadata* SceneTexturesUniformBufferStruct,
 	FRHIUniformBuffer* SceneTexturesUniformBuffer,
 	bool bUseFixDT)
 {
@@ -1440,6 +1436,13 @@ void ExecuteSimulationCommands(
 	SCOPED_DRAW_EVENT(RHICmdList, ParticleSimulation);
 	SCOPED_GPU_STAT(RHICmdList, ParticleSimulation);
 
+	FUniformBufferStaticBindings StaticUniformBuffers;
+	if (SceneTexturesUniformBuffer)
+	{
+		StaticUniformBuffers.AddUniformBuffer(SceneTexturesUniformBuffer);
+	}
+	SCOPED_UNIFORM_BUFFER_STATIC_BINDINGS(RHICmdList, StaticUniformBuffers);
+
 	const float FixDeltaSeconds = CVarGPUParticleFixDeltaSeconds.GetValueOnRenderThread();
 	const FParticleStateTextures& TextureResources = (FixDeltaSeconds <= 0 || bUseFixDT) ? ParticleSimulationResources->GetPreviousStateTextures() : ParticleSimulationResources->GetCurrentStateTextures();
 	const FParticleAttributesTexture& AttributeTexture = ParticleSimulationResources->SimulationAttributesTexture;
@@ -1456,7 +1459,7 @@ void ExecuteSimulationCommands(
 
 	SetGraphicsPipelineState(RHICmdList, GraphicsPSOInit);
 
-	PixelShader->SetParameters(RHICmdList, TextureResources, AttributeTexture, RenderAttributeTexture, ViewUniformBuffer, FeatureLevel, GlobalDistanceFieldParameterData, SceneTexturesUniformBufferStruct, SceneTexturesUniformBuffer);
+	PixelShader->SetParameters(RHICmdList, TextureResources, AttributeTexture, RenderAttributeTexture, ViewUniformBuffer, FeatureLevel, GlobalDistanceFieldParameterData);
 
 	// Draw tiles to perform the simulation step.
 	const int32 CommandCount = SimulationCommands.Num();
@@ -1495,7 +1498,6 @@ void ExecuteSimulationCommands(
 	FParticleSimulationResources* ParticleSimulationResources,
 	FRHIUniformBuffer* ViewUniformBuffer,
 	const FGlobalDistanceFieldParameterData* GlobalDistanceFieldParameterData,
-	const FShaderParametersMetadata* SceneTexturesUniformBufferStruct,
 	FRHIUniformBuffer* SceneTexturesUniformBuffer,
 	EParticleSimulatePhase::Type Phase,
 	bool bUseFixDT)
@@ -1510,7 +1512,6 @@ void ExecuteSimulationCommands(
 			ParticleSimulationResources,
 			ViewUniformBuffer,
 			GlobalDistanceFieldParameterData,
-			SceneTexturesUniformBufferStruct,
 			SceneTexturesUniformBuffer,
 			bUseFixDT);
 	}
@@ -1524,7 +1525,6 @@ void ExecuteSimulationCommands(
 			ParticleSimulationResources,
 			ViewUniformBuffer,
 			GlobalDistanceFieldParameterData,
-			SceneTexturesUniformBufferStruct,
 			SceneTexturesUniformBuffer,
 			bUseFixDT);
 	}
@@ -1538,7 +1538,6 @@ void ExecuteSimulationCommands(
 			ParticleSimulationResources,
 			NULL,
 			GlobalDistanceFieldParameterData,
-			SceneTexturesUniformBufferStruct,
 			SceneTexturesUniformBuffer,
 			bUseFixDT);
 	}
@@ -4648,9 +4647,7 @@ void FFXSystem::SimulateGPUParticles(
 	FRHICommandListImmediate& RHICmdList,
 	EParticleSimulatePhase::Type Phase,
 	FRHIUniformBuffer* ViewUniformBuffer,
-	const FGlobalDistanceFieldParameterData* GlobalDistanceFieldParameterData,
-	const FShaderParametersMetadata* SceneTexturesUniformBufferStruct,
-	FRHIUniformBuffer* SceneTexturesUniformBuffer
+	const FGlobalDistanceFieldParameterData* GlobalDistanceFieldParameterData
 	)
 {
 	LLM_SCOPE(ELLMTag::Particles);
@@ -4921,8 +4918,7 @@ void FFXSystem::SimulateGPUParticles(
 					ParticleSimulationResources,
 					ViewUniformBuffer,
 					GlobalDistanceFieldParameterData,
-					SceneTexturesUniformBufferStruct,
-					SceneTexturesUniformBuffer,
+					SceneTexturesUniformParams,
 					Phase,
 					FixDeltaSeconds > 0
 				);
@@ -5044,8 +5040,7 @@ void FFXSystem::SimulateGPUParticles(
 				ParticleSimulationResources,
 				ViewUniformBuffer,
 				GlobalDistanceFieldParameterData,
-				SceneTexturesUniformBufferStruct,
-				SceneTexturesUniformBuffer,
+				SceneTexturesUniformParams,
 				Phase,
 				false
 			);
diff --git a/Engine/Source/Runtime/Engine/Public/FXSystem.h b/Engine/Source/Runtime/Engine/Public/FXSystem.h
--- a/Engine/Source/Runtime/Engine/Public/FXSystem.h
+++ b/Engine/Source/Runtime/Engine/Public/FXSystem.h
@@ -254,8 +254,6 @@ public:
 	virtual void PostRenderOpaque(
 		class FRDGBuilder& GraphBuilder,
 		FRHIUniformBuffer* ViewUniformBuffer,
-		const class FShaderParametersMetadata* SceneTexturesUniformBufferStruct,
-		FRHIUniformBuffer* SceneTexturesUniformBuffer,
 		bool bAllowGPUParticleUpdate) = 0;
 
 	bool IsPendingKill() const { return bIsPendingKill; }
@@ -263,6 +261,8 @@ public:
 	/** Get the shared SortManager, used in the rendering loop to call FGPUSortManager::OnPreRender() and FGPUSortManager::OnPostRenderOpaque() */
 	virtual FGPUSortManager* GetGPUSortManager() const = 0;
 
+	virtual void SetSceneTexturesUniformBuffer(FRHIUniformBuffer* InSceneTexturesUniformParams) {}
+
 protected:
 	
 	friend class FFXSystemSet;
diff --git a/Engine/Source/Runtime/Renderer/Private/DeferredShadingRenderer.cpp b/Engine/Source/Runtime/Renderer/Private/DeferredShadingRenderer.cpp
--- a/Engine/Source/Runtime/Renderer/Private/DeferredShadingRenderer.cpp
+++ b/Engine/Source/Runtime/Renderer/Private/DeferredShadingRenderer.cpp
@@ -427,14 +427,26 @@ static void RenderOpaqueFX(
 		RDG_GPU_STAT_SCOPE(GraphBuilder, PostRenderOpsFX);
 		RDG_CSV_STAT_EXCLUSIVE_SCOPE(GraphBuilder, RenderOpaqueFX);
 
+		// Add a pass which extracts the RHI handle from the scene textures UB and sends it to the FX system.
+		FRenderOpaqueFXPassParameters* ExtractUBPassParameters = GraphBuilder.AllocParameters<FRenderOpaqueFXPassParameters>();
+		ExtractUBPassParameters->SceneTextures = SceneTexturesUniformBuffer;
+		GraphBuilder.AddPass({}, ExtractUBPassParameters, ERDGPassFlags::Raster | ERDGPassFlags::SkipRenderPass | ERDGPassFlags::NeverCull, [ExtractUBPassParameters, FXSystem](FRHICommandList&)
+		{
+			FXSystem->SetSceneTexturesUniformBuffer(ExtractUBPassParameters->SceneTextures->GetRHIRef());
+		});
+
 		FXSystem->PostRenderOpaque(
 			GraphBuilder,
 			Views[0].ViewUniformBuffer,
-			&FSceneTextureUniformParameters::StaticStructMetadata,
-			nullptr,
 			Views[0].AllowGPUParticleUpdate()
 		);
 
+		// Clear the scene textures UB pointer on the FX system.
+		AddPass(GraphBuilder, [FXSystem](FRHICommandList&)
+		{
+			FXSystem->SetSceneTexturesUniformBuffer(nullptr);
+		});
+
 		if (FGPUSortManager* GPUSortManager = FXSystem->GetGPUSortManager())
 		{
 			GPUSortManager->OnPostRenderOpaque(GraphBuilder);
diff --git a/Engine/Source/Runtime/Renderer/Private/MobileShadingRenderer.cpp b/Engine/Source/Runtime/Renderer/Private/MobileShadingRenderer.cpp
--- a/Engine/Source/Runtime/Renderer/Private/MobileShadingRenderer.cpp
+++ b/Engine/Source/Runtime/Renderer/Private/MobileShadingRenderer.cpp
@@ -731,8 +731,6 @@ void FMobileSceneRenderer::Render(FRDGBuilder& GraphBuilder)
 		FXSystem->PostRenderOpaque(
 			GraphBuilder,
 			Views[0].ViewUniformBuffer,
-			nullptr,
-			nullptr,
 			Views[0].AllowGPUParticleUpdate()
 		);
 		if (FGPUSortManager* GPUSortManager = FXSystem->GetGPUSortManager())
diff --git a/Engine/Source/Runtime/Renderer/Private/SceneHitProxyRendering.cpp b/Engine/Source/Runtime/Renderer/Private/SceneHitProxyRendering.cpp
--- a/Engine/Source/Runtime/Renderer/Private/SceneHitProxyRendering.cpp
+++ b/Engine/Source/Runtime/Renderer/Private/SceneHitProxyRendering.cpp
@@ -658,7 +658,7 @@ void FDeferredShadingSceneRenderer::RenderHitProxies(FRDGBuilder& GraphBuilder)
 		}
 		// Call PostRenderOpaque now as this is irrelevant for when rendering hit proxies.
 		// because we don't tick the particles in the render loop (see last param being "false").
-		FXSystem->PostRenderOpaque(GraphBuilder, Views[0].ViewUniformBuffer, nullptr, nullptr, false);
+		FXSystem->PostRenderOpaque(GraphBuilder, Views[0].ViewUniformBuffer, false);
 		if (GPUSortManager)
 		{
 			GPUSortManager->OnPostRenderOpaque(GraphBuilder);
