PR #7824: Bugfix for UEdGraphNode::BreakAllNodeLinks method (Contributed by [REDACTED_PERSON])

#rnx

diff --git a/Engine/Source/Runtime/Engine/Private/EdGraph/EdGraphNode.cpp b/Engine/Source/Runtime/Engine/Private/EdGraph/EdGraphNode.cpp
--- a/Engine/Source/Runtime/Engine/Private/EdGraph/EdGraphNode.cpp
+++ b/Engine/Source/Runtime/Engine/Private/EdGraph/EdGraphNode.cpp
@@ -422,13 +422,20 @@ void UEdGraphNode::BreakAllNodeLinks()
 	NodeList.Add(this);
 
 	// Iterate over each pin and break all links
-	for(int32 PinIdx=0; PinIdx<Pins.Num(); PinIdx++)
+	for(int32 PinIdx = 0; PinIdx < Pins.Num(); ++PinIdx)
 	{
-		Pins[PinIdx]->BreakAllPinLinks();
-		NodeList.Add(Pins[PinIdx]->GetOwningNode());
+		UEdGraphPin* Pin = Pins[PinIdx];
+
+		// Save all the connected nodes to be notified below
+		for (UEdGraphPin* Connection : Pin->LinkedTo)
+		{
+			NodeList.Add(Connection->GetOwningNode());
+		}
+
+		Pin->BreakAllPinLinks();
 	}
 
-	// Send all nodes that received a new pin connection a notification
+	// Send a notification to all nodes that lost a connection
 	for (UEdGraphNode* Node : NodeList)
 	{
 		Node->NodeConnectionListChanged();
