World Partition - Fixed HLOD to cell mapping

diff --git a/Engine/Source/Runtime/Engine/Private/WorldPartition/HLOD/HLODActor.cpp b/Engine/Source/Runtime/Engine/Private/WorldPartition/HLOD/HLODActor.cpp
--- a/Engine/Source/Runtime/Engine/Private/WorldPartition/HLOD/HLODActor.cpp
+++ b/Engine/Source/Runtime/Engine/Private/WorldPartition/HLOD/HLODActor.cpp
@@ -29,19 +29,11 @@ UPrimitiveComponent* AWorldPartitionHLOD::GetHLODComponent()
 	return Cast<UPrimitiveComponent>(RootComponent);
 }
 
-void AWorldPartitionHLOD::OnCellShown(FName InCellName)
+void AWorldPartitionHLOD::SetVisibility(bool bInVisible)
 {
 	if (GetRootComponent())
 	{
-		GetRootComponent()->SetVisibility(false, true);
-	}
-}
-
-void AWorldPartitionHLOD::OnCellHidden(FName InCellName)
-{
-	if (GetRootComponent())
-	{
-		GetRootComponent()->SetVisibility(true, true);
+		GetRootComponent()->SetVisibility(bInVisible, /*bPropagateToChildren*/ true);
 	}
 }
 
@@ -57,6 +49,19 @@ void AWorldPartitionHLOD::EndPlay(const EEndPlayReason::Type EndPlayReason)
 	Super::EndPlay(EndPlayReason);
 }
 
+void AWorldPartitionHLOD::PostLoad()
+{
+	Super::PostLoad();
+
+#if WITH_EDITOR
+	if (CellName == NAME_None)
+	{
+		// Prior to the addition of the cell name member, the cell name was actually stored in the label
+		CellName = *GetActorLabel();
+	}
+#endif
+}
+
 #if WITH_EDITOR
 
 void AWorldPartitionHLOD::OnSubActorLoaded(const AActor& Actor)
@@ -220,19 +225,9 @@ const TArray<FGuid>& AWorldPartitionHLOD::GetSubActors() const
 	return SubActors;
 }
 
-void AWorldPartitionHLOD::PostActorCreated()
+void AWorldPartitionHLOD::SetCellName(FName InCellName)
 {
-	Super::PostActorCreated();
-	HLODGuid = GetActorGuid();
-}
-
-void AWorldPartitionHLOD::PostDuplicate(bool bDuplicateForPIE)
-{
-	Super::PostDuplicate(bDuplicateForPIE);
-	if (!bDuplicateForPIE)
-	{
-		HLODGuid = GetActorGuid();
-	}
+	CellName = InCellName;
 }
 
 const FBox& AWorldPartitionHLOD::GetHLODBounds() const
@@ -275,18 +270,3 @@ void AWorldPartitionHLOD::BuildHLOD(bool bForceBuild)
 }
 
 #endif // WITH_EDITOR
-
-
-UWorldPartitionRuntimeHLODCellData::UWorldPartitionRuntimeHLODCellData(const FObjectInitializer& ObjectInitializer)
-	: Super(ObjectInitializer)
-{
-}
-
-#if WITH_EDITOR
-
-void UWorldPartitionRuntimeHLODCellData::SetReferencedHLODActors(TArray<FGuid>&& InReferencedHLODActors)
-{
-	ReferencedHLODActors = MoveTemp(InReferencedHLODActors);
-}
-
-#endif // WITH_EDITOR
diff --git a/Engine/Source/Runtime/Engine/Private/WorldPartition/HLOD/HLODBuilder.cpp b/Engine/Source/Runtime/Engine/Private/WorldPartition/HLOD/HLODBuilder.cpp
--- a/Engine/Source/Runtime/Engine/Private/WorldPartition/HLOD/HLODBuilder.cpp
+++ b/Engine/Source/Runtime/Engine/Private/WorldPartition/HLOD/HLODBuilder.cpp
@@ -441,7 +441,8 @@ TArray<AWorldPartitionHLOD*> FHLODBuilderUtilities::CreateHLODActors(FHLODCreati
 			SpawnParams.NameMode = FActorSpawnParameters::ESpawnActorNameMode::Required_Fatal;
 			HLODActor = InCreationParams.WorldPartition->GetWorld()->SpawnActor<AWorldPartitionHLOD>(SpawnParams);
 
-			HLODActor->SetActorLabel(InCreationParams.CellName.ToString());
+			HLODActor->SetActorLabel(FString::Printf(TEXT("HLOD%d_%s"), InCreationParams.HLODLevel, *InCreationParams.CellName.ToString()));
+			HLODActor->SetCellName(InCreationParams.CellName);
 			HLODActor->SetSubActorsHLODLayer(HLODLayer);
 			HLODActor->SetGridIndices(InCreationParams.GridIndexX, InCreationParams.GridIndexY, InCreationParams.GridIndexZ);
 
diff --git a/Engine/Source/Runtime/Engine/Private/WorldPartition/HLOD/HLODSubsystem.cpp b/Engine/Source/Runtime/Engine/Private/WorldPartition/HLOD/HLODSubsystem.cpp
--- a/Engine/Source/Runtime/Engine/Private/WorldPartition/HLOD/HLODSubsystem.cpp
+++ b/Engine/Source/Runtime/Engine/Private/WorldPartition/HLOD/HLODSubsystem.cpp
@@ -7,6 +7,7 @@
 #include "WorldPartition/WorldPartition.h"
 #include "WorldPartition/WorldPartitionSubsystem.h"
 #include "WorldPartition/WorldPartitionRuntimeCell.h"
+#include "WorldPartition/WorldPartitionRuntimeHash.h"
 #include "UObject/UObjectHash.h"
 #include "UObject/UObjectGlobals.h"
 #include "Engine/World.h"
@@ -49,6 +50,18 @@ void UHLODSubsystem::Initialize(FSubsystemCollectionBase& Collection)
 		GetWorld()->PersistentLevel->OnLoadedActorRemovedFromLevelEvent.AddUObject(this, &UHLODSubsystem::OnActorUnloaded);
 	}
 #endif
+
+	if (GetWorld()->IsGameWorld())
+	{
+		TSet<const UWorldPartitionRuntimeCell*> StreamingCells;
+		GetWorld()->GetWorldPartition()->RuntimeHash->GetAllStreamingCells(StreamingCells, /*bIncludeDataLayers*/ true);
+
+		// Build cell to HLOD mapping
+		for (const UWorldPartitionRuntimeCell* Cell : StreamingCells)
+		{
+			CellsHLODMapping.Emplace(Cell->GetFName());
+		}
+	}
 }
 
 #if WITH_EDITOR
@@ -103,18 +116,12 @@ void UHLODSubsystem::RegisterHLODActor(AWorldPartitionHLOD* InWorldPartitionHLOD
 {
 	TRACE_CPUPROFILER_EVENT_SCOPE(UHLODSubsystem::RegisterHLODActor);
 
-	const FGuid& HLODActorGUID = InWorldPartitionHLOD->GetHLODGuid();
-	RegisteredHLODActors.Emplace(HLODActorGUID, InWorldPartitionHLOD);
-
-	TArray<FName> Cells;
-	PendingCellsShown.MultiFind(HLODActorGUID, Cells);
-	if (Cells.Num())
+	FName CellName = InWorldPartitionHLOD->GetCellName();
+	FCellHLODMapping* CellHLODs = CellsHLODMapping.Find(CellName);
+	if (CellHLODs)
 	{
-		for (FName Cell : Cells)
-		{ 
-			InWorldPartitionHLOD->OnCellShown(Cell);
-		}
-		PendingCellsShown.Remove(HLODActorGUID);
+		CellHLODs->LoadedHLODs.Add(InWorldPartitionHLOD);
+		InWorldPartitionHLOD->SetVisibility(!CellHLODs->bIsCellVisible);
 	}
 }
 
@@ -122,53 +129,39 @@ void UHLODSubsystem::UnregisterHLODActor(AWorldPartitionHLOD* InWorldPartitionHL
 {
 	TRACE_CPUPROFILER_EVENT_SCOPE(UHLODSubsystem::UnregisterHLODActor);
 
-	const FGuid& HLODActorGUID = InWorldPartitionHLOD->GetHLODGuid();
-	int32 NumRemoved = RegisteredHLODActors.Remove(HLODActorGUID);
-	check(NumRemoved == 1);
+	FName CellName = InWorldPartitionHLOD->GetCellName();
+	FCellHLODMapping* CellHLODs = CellsHLODMapping.Find(CellName);
+	if (CellHLODs)
+	{
+		int32 NumRemoved = CellHLODs->LoadedHLODs.Remove(InWorldPartitionHLOD);
+		check(NumRemoved == 1);
+	}
 }
 
 void UHLODSubsystem::OnCellShown(const UWorldPartitionRuntimeCell* InCell)
 {
-	const UWorldPartitionRuntimeHLODCellData* HLODCellData = InCell->GetCellData<UWorldPartitionRuntimeHLODCellData>();
-	if (!ensure(HLODCellData))
+	FCellHLODMapping* CellHLODs = CellsHLODMapping.Find(InCell->GetFName());
+	if (CellHLODs)
 	{
-		return;
-	}
+		CellHLODs->bIsCellVisible = true;
 
-	for (const FGuid& HLODActorGUID : HLODCellData->ReferencedHLODActors)
-	{
-		AWorldPartitionHLOD* HLODActor = RegisteredHLODActors.FindRef(HLODActorGUID);
-		if (HLODActor)
-		{
-			HLODActor->OnCellShown(InCell->GetFName());
-		}
-		else
+		for (AWorldPartitionHLOD* HLODActor : CellHLODs->LoadedHLODs)
 		{
-			// Cell was shown before the HLOD
-			PendingCellsShown.Add(HLODActorGUID, InCell->GetFName());
+			HLODActor->SetVisibility(false);
 		}
 	}
 }
 
 void UHLODSubsystem::OnCellHidden(const UWorldPartitionRuntimeCell* InCell)
 {
-	const UWorldPartitionRuntimeHLODCellData* HLODCellData = InCell->GetCellData<UWorldPartitionRuntimeHLODCellData>();
-	if (!ensure(HLODCellData))
+	FCellHLODMapping* CellHLODs = CellsHLODMapping.Find(InCell->GetFName());
+	if (CellHLODs)
 	{
-		return;
-	}
+		CellHLODs->bIsCellVisible = false;
 
-	for (const FGuid& HLODActorGUID : HLODCellData->ReferencedHLODActors)
-	{
-		AWorldPartitionHLOD* HLODActor = RegisteredHLODActors.FindRef(HLODActorGUID);
-		if (HLODActor)
-		{
-			HLODActor->OnCellHidden(InCell->GetFName());
-			check(!PendingCellsShown.Find(HLODActorGUID));
-		}
-		else
+		for (AWorldPartitionHLOD* HLODActor : CellHLODs->LoadedHLODs)
 		{
-			PendingCellsShown.Remove(HLODActorGUID, InCell->GetFName());
+			HLODActor->SetVisibility(true);
 		}
 	}
 }
diff --git a/Engine/Source/Runtime/Engine/Private/WorldPartition/WorldPartitionActorDescView.cpp b/Engine/Source/Runtime/Engine/Private/WorldPartition/WorldPartitionActorDescView.cpp
--- a/Engine/Source/Runtime/Engine/Private/WorldPartition/WorldPartitionActorDescView.cpp
+++ b/Engine/Source/Runtime/Engine/Private/WorldPartition/WorldPartitionActorDescView.cpp
@@ -105,15 +105,4 @@ bool FWorldPartitionActorDescView::GetContainerInstance(const UActorDescContaine
 	return ActorDesc->GetContainerInstance(OutLevelContainer, OutLevelTransform, OutClusterMode);
 }
 
-FGuid FWorldPartitionActorDescView::GetHLODParent() const
-{
-	return HLODParent;
-}
-
-void FWorldPartitionActorDescView::SetHLODParent(const FGuid& InHLODParent)
-{
-	check(!HLODParent.IsValid());
-	check(InHLODParent.IsValid());
-	HLODParent = InHLODParent;
-}
 #endif
\ No newline at end of file
diff --git a/Engine/Source/Runtime/Engine/Private/WorldPartition/WorldPartitionRuntimeHash.cpp b/Engine/Source/Runtime/Engine/Private/WorldPartition/WorldPartitionRuntimeHash.cpp
--- a/Engine/Source/Runtime/Engine/Private/WorldPartition/WorldPartitionRuntimeHash.cpp
+++ b/Engine/Source/Runtime/Engine/Private/WorldPartition/WorldPartitionRuntimeHash.cpp
@@ -109,18 +109,6 @@ void UWorldPartitionRuntimeHash::CreateActorDescViewMap(const UActorDescContaine
 		}
 	}
 
-	// Set HLOD parents into actor desc views
-	for (UActorDescContainer::TConstIterator<AWorldPartitionHLOD> HLODIterator(Container); HLODIterator; ++HLODIterator)
-	{
-		for (const FGuid& SubActor : HLODIterator->GetSubActors())
-		{
-			if (FWorldPartitionActorDescView* SubActorDescView = OutActorDescViewMap.Find(SubActor))
-			{
-				SubActorDescView->SetHLODParent(HLODIterator->GetGuid());
-			}
-		}
-	}
-
 	// Gather all references to external actors from the level script and make them always loaded
 	if (ULevelScriptBlueprint* LevelScriptBlueprint = Container->GetWorld()->PersistentLevel->GetLevelScriptBlueprint(true))
 	{
diff --git a/Engine/Source/Runtime/Engine/Private/WorldPartition/WorldPartitionRuntimeLevelStreamingCell.cpp b/Engine/Source/Runtime/Engine/Private/WorldPartition/WorldPartitionRuntimeLevelStreamingCell.cpp
--- a/Engine/Source/Runtime/Engine/Private/WorldPartition/WorldPartitionRuntimeLevelStreamingCell.cpp
+++ b/Engine/Source/Runtime/Engine/Private/WorldPartition/WorldPartitionRuntimeLevelStreamingCell.cpp
@@ -216,11 +216,8 @@ UWorldPartitionLevelStreamingDynamic* UWorldPartitionRuntimeLevelStreamingCell::
 
 	if (LevelStreaming)
 	{
-		if (const UWorldPartitionRuntimeHLODCellData* HLODCellData = GetCellData<UWorldPartitionRuntimeHLODCellData>())
-		{
-			LevelStreaming->OnLevelShown.AddUniqueDynamic(this, &UWorldPartitionRuntimeLevelStreamingCell::OnLevelShown);
-			LevelStreaming->OnLevelHidden.AddUniqueDynamic(this, &UWorldPartitionRuntimeLevelStreamingCell::OnLevelHidden);
-		}
+		LevelStreaming->OnLevelShown.AddUniqueDynamic(this, &UWorldPartitionRuntimeLevelStreamingCell::OnLevelShown);
+		LevelStreaming->OnLevelHidden.AddUniqueDynamic(this, &UWorldPartitionRuntimeLevelStreamingCell::OnLevelHidden);
 	}
 
 	return LevelStreaming;
diff --git a/Engine/Source/Runtime/Engine/Private/WorldPartition/WorldPartitionRuntimeSpatialHash.cpp b/Engine/Source/Runtime/Engine/Private/WorldPartition/WorldPartitionRuntimeSpatialHash.cpp
--- a/Engine/Source/Runtime/Engine/Private/WorldPartition/WorldPartitionRuntimeSpatialHash.cpp
+++ b/Engine/Source/Runtime/Engine/Private/WorldPartition/WorldPartitionRuntimeSpatialHash.cpp
@@ -805,16 +805,9 @@ bool UWorldPartitionRuntimeSpatialHash::CreateStreamingGrid(const FSpatialHashRu
 
 				UE_LOG(LogWorldPartitionRuntimeSpatialHash, Verbose, TEXT("Cell%s %s Actors = %d Bounds (%s)"), bIsCellAlwaysLoaded ? TEXT(" (AlwaysLoaded)") : TEXT(""), *StreamingCell->GetName(), FilteredActors.Num(), *Bounds.ToString());
 
-				// Keep track of all AWorldPartitionHLOD actors referenced by this cell
-				TSet<FGuid> ReferencedHLODActors;
-
 				for (const FActorInstance& ActorInstance : FilteredActors)
 				{
 					const FWorldPartitionActorDescView& ActorDescView = ActorInstance.GetActorDescView();
-					if (ActorDescView.GetHLODParent().IsValid())
-					{
-						ReferencedHLODActors.Add(ActorDescView.GetHLODParent());
-					}
 					StreamingCell->AddActorToCell(ActorDescView, ActorInstance.ContainerInstance->ID, ActorInstance.ContainerInstance->Transform, ActorInstance.ContainerInstance->Container);
 
 					// Add all actors of this cell to the ActorsModifiedForPIE (for Duplication)
@@ -833,14 +826,6 @@ bool UWorldPartitionRuntimeSpatialHash::CreateStreamingGrid(const FSpatialHashRu
 					UE_LOG(LogWorldPartitionRuntimeSpatialHash, Verbose, TEXT("  Actor : %s (%s) (Container %08x) Origin(%s)"), *(ActorDescView.GetActorPath().ToString()), *ActorDescView.GetGuid().ToString(EGuidFormats::UniqueObjectGuid), ActorInstance.ContainerInstance->ID, *FVector2D(ActorInstance.GetOrigin()).ToString());
 				}
 
-				if (ReferencedHLODActors.Num() > 0)
-				{
-					// Store the referenced HLOD actors as custom cell data
-					UWorldPartitionRuntimeHLODCellData* HLODCellData = NewObject<UWorldPartitionRuntimeHLODCellData>(StreamingCell);
-					HLODCellData->SetReferencedHLODActors(ReferencedHLODActors.Array());
-					StreamingCell->AddCellData(HLODCellData);
-				}
-
 				if (Mode == EWorldPartitionStreamingMode::Cook)
 				{
 					UE_LOG(LogWorldPartitionRuntimeSpatialHash, Log, TEXT("Creating runtime streaming cells %s."), *StreamingCell->GetName());
diff --git a/Engine/Source/Runtime/Engine/Public/WorldPartition/HLOD/HLODActor.h b/Engine/Source/Runtime/Engine/Public/WorldPartition/HLOD/HLODActor.h
--- a/Engine/Source/Runtime/Engine/Public/WorldPartition/HLOD/HLODActor.h
+++ b/Engine/Source/Runtime/Engine/Public/WorldPartition/HLOD/HLODActor.h
@@ -17,10 +17,9 @@ class ENGINE_API AWorldPartitionHLOD : public AActor
 	GENERATED_UCLASS_BODY()
 
 public:
-	void OnCellShown(FName InCellName);
-	void OnCellHidden(FName InCellName);
+	void SetVisibility(bool bInVisible);
 
-	const FGuid& GetHLODGuid() const { return HLODGuid; }
+	inline FName GetCellName() const { return CellName; }
 	inline uint32 GetLODLevel() const { return LODLevel; }
 
 #if WITH_EDITOR
@@ -49,6 +48,7 @@ public:
 		OutGridIndexZ = GridIndexZ;
 	}
 
+	void SetCellName(FName InCellName);
 	inline void SetLODLevel(uint32 InLODLevel) { LODLevel = InLODLevel; }
 
 	const FBox& GetHLODBounds() const;
@@ -59,6 +59,10 @@ public:
 #endif // WITH_EDITOR
 
 protected:
+	//~ Begin UObject Interface.
+	virtual void PostLoad() override;
+	//~ End UObject Interface.
+
 	//~ Begin AActor Interface.
 	virtual void BeginPlay() override;
 	virtual void EndPlay(const EEndPlayReason::Type EndPlayReason) override;
@@ -68,9 +72,6 @@ protected:
 	virtual EActorGridPlacement GetDefaultGridPlacement() const override;
 	virtual TUniquePtr<class FWorldPartitionActorDesc> CreateClassActorDesc() const override;
 
-	virtual void PostActorCreated() override;
-	virtual void PostDuplicate(bool bDuplicateForPIE) override;
-
 	virtual void GetActorBounds(bool bOnlyCollidingComponents, FVector& Origin, FVector& BoxExtent, bool bIncludeFromChildActors) const override;
 	virtual void GetActorLocationBounds(bool bOnlyCollidingComponents, FVector& Origin, FVector& BoxExtent, bool bIncludeFromChildActors) const override;
 #endif
@@ -112,25 +113,11 @@ private:
 	uint32 HLODHash;
 #endif
 
-	UPROPERTY(NonPIEDuplicateTransient, TextExportTransient, NonTransactional)
-	FGuid HLODGuid;
-
 	UPROPERTY()
 	uint32 LODLevel;
+
+	UPROPERTY()
+	FName CellName;
 };
 
 DEFINE_ACTORDESC_TYPE(AWorldPartitionHLOD, FHLODActorDesc);
-
-UCLASS()
-class UWorldPartitionRuntimeHLODCellData : public UWorldPartitionRuntimeCellData
-{
-	GENERATED_UCLASS_BODY()
-
-public:
-#if WITH_EDITOR
-	void SetReferencedHLODActors(TArray<FGuid>&& InReferencedHLODActors);
-#endif
-
-	UPROPERTY()
-	TArray<FGuid> ReferencedHLODActors;
-};
\ No newline at end of file
diff --git a/Engine/Source/Runtime/Engine/Public/WorldPartition/HLOD/HLODSubsystem.h b/Engine/Source/Runtime/Engine/Public/WorldPartition/HLOD/HLODSubsystem.h
--- a/Engine/Source/Runtime/Engine/Public/WorldPartition/HLOD/HLODSubsystem.h
+++ b/Engine/Source/Runtime/Engine/Public/WorldPartition/HLOD/HLODSubsystem.h
@@ -42,11 +42,17 @@ private:
 #endif
 
 private:
-	// Mapping between an HLOD GUID & the loaded actor.
-	TMap<FGuid, AWorldPartitionHLOD*> RegisteredHLODActors;
+	struct FCellHLODMapping
+	{
+		bool						bIsCellVisible;
+		TSet<AWorldPartitionHLOD*>	LoadedHLODs;
 
-	// Cells that where shown before their HLOD was loaded.
-	TMultiMap<FGuid, FName> PendingCellsShown;
+		FCellHLODMapping() : bIsCellVisible(false)
+		{
+		}
+	};
+
+	TMap<FName, FCellHLODMapping> CellsHLODMapping;
 
 #if WITH_EDITORONLY_DATA
 	TMap<FGuid, AWorldPartitionHLOD*> ActorsToHLOD;
diff --git a/Engine/Source/Runtime/Engine/Public/WorldPartition/WorldPartitionActorDescView.h b/Engine/Source/Runtime/Engine/Public/WorldPartition/WorldPartitionActorDescView.h
--- a/Engine/Source/Runtime/Engine/Public/WorldPartition/WorldPartitionActorDescView.h
+++ b/Engine/Source/Runtime/Engine/Public/WorldPartition/WorldPartitionActorDescView.h
@@ -42,9 +42,6 @@ public:
 
 	bool GetContainerInstance(const UActorDescContainer*& OutLevelContainer, FTransform& OutLevelTransform, EContainerClusterMode& OutClusterMode) const;
 
-	FGuid GetHLODParent() const;
-	void SetHLODParent(const FGuid& InHLODParent);
-
 	bool operator==(const FWorldPartitionActorDescView& Other) const
 	{
 		return GetGuid() == Other.GetGuid();
@@ -58,6 +55,5 @@ public:
 protected:
 	const FWorldPartitionActorDesc* ActorDesc;
 	EActorGridPlacement EffectiveGridPlacement;
-	FGuid HLODParent;
 };
 #endif
\ No newline at end of file
