Implemented common "Duplicate" operation using typed elements

 - Moved the logic for managing the offset when pasting/duplicating out of the internal Paste/DuplicateActors implementation, and instead pass the desired offset into the duplicate functions.
 - Made the function for getting the list of elements to operate on for a common action more generic, and added helper functions to delete and duplicate the current selection set.
 - Added a CanDuplicateElements function to UTypedElementWorldInterface.
 - Implemented duplication offset for static mesh instances.

diff --git a/Engine/Source/Editor/LevelEditor/Private/Elements/Actor/ActorElementLevelEditorCommonActionsCustomization.cpp b/Engine/Source/Editor/LevelEditor/Private/Elements/Actor/ActorElementLevelEditorCommonActionsCustomization.cpp
--- a/Engine/Source/Editor/LevelEditor/Private/Elements/Actor/ActorElementLevelEditorCommonActionsCustomization.cpp
+++ b/Engine/Source/Editor/LevelEditor/Private/Elements/Actor/ActorElementLevelEditorCommonActionsCustomization.cpp
@@ -10,17 +10,17 @@
 #include "EditorModeManager.h"
 #include "Toolkits/IToolkitHost.h"
 
-void FActorElementLevelEditorCommonActionsCustomization::GetElementsToDelete(const TTypedElement<UTypedElementWorldInterface>& InElementWorldHandle, const UTypedElementSelectionSet* InSelectionSet, UTypedElementList* OutElementsToDelete)
+void FActorElementLevelEditorCommonActionsCustomization::GetElementsForAction(const TTypedElement<UTypedElementWorldInterface>& InElementWorldHandle, const UTypedElementSelectionSet* InSelectionSet, UTypedElementList* OutElementsToDelete)
 {
 	AActor* Actor = ActorElementDataUtil::GetActorFromHandleChecked(InElementWorldHandle);
 
 	if (UComponentElementSelectionInterface::HasSelectedComponents(InSelectionSet->GetElementList()))
 	{
-		// If we have components selected then we will delete those rather than the actors
+		// If we have components selected then we will perform the action on those rather than the actors
 		return;
 	}
 
-	FTypedElementCommonActionsCustomization::GetElementsToDelete(InElementWorldHandle, InSelectionSet, OutElementsToDelete);
+	FTypedElementCommonActionsCustomization::GetElementsForAction(InElementWorldHandle, InSelectionSet, OutElementsToDelete);
 }
 
 bool FActorElementLevelEditorCommonActionsCustomization::DeleteElements(UTypedElementWorldInterface* InWorldInterface, TArrayView<const FTypedElementHandle> InElementHandles, UWorld* InWorld, UTypedElementSelectionSet* InSelectionSet, const FTypedElementDeletionOptions& InDeletionOptions)
@@ -37,7 +37,7 @@ bool FActorElementLevelEditorCommonActionsCustomization::DeleteElements(UTypedEl
 	return FTypedElementCommonActionsCustomization::DeleteElements(InWorldInterface, InElementHandles, InWorld, InSelectionSet, InDeletionOptions);
 }
 
-void FActorElementLevelEditorCommonActionsCustomization::DuplicateElements(UTypedElementWorldInterface* InWorldInterface, TArrayView<const FTypedElementHandle> InElementHandles, UWorld* InWorld, bool bOffsetLocations, TArray<FTypedElementHandle>& OutNewElements)
+void FActorElementLevelEditorCommonActionsCustomization::DuplicateElements(UTypedElementWorldInterface* InWorldInterface, TArrayView<const FTypedElementHandle> InElementHandles, UWorld* InWorld, const FVector& InLocationOffset, TArray<FTypedElementHandle>& OutNewElements)
 {
 	if (const IToolkitHost* ToolkitHostPtr = GetToolkitHost())
 	{
@@ -48,5 +48,5 @@ void FActorElementLevelEditorCommonActionsCustomization::DuplicateElements(UType
 		}
 	}
 
-	FTypedElementCommonActionsCustomization::DuplicateElements(InWorldInterface, InElementHandles, InWorld, bOffsetLocations, OutNewElements);
+	FTypedElementCommonActionsCustomization::DuplicateElements(InWorldInterface, InElementHandles, InWorld, InLocationOffset, OutNewElements);
 }
diff --git a/Engine/Source/Editor/LevelEditor/Private/Elements/Actor/ActorElementLevelEditorCommonActionsCustomization.h b/Engine/Source/Editor/LevelEditor/Private/Elements/Actor/ActorElementLevelEditorCommonActionsCustomization.h
--- a/Engine/Source/Editor/LevelEditor/Private/Elements/Actor/ActorElementLevelEditorCommonActionsCustomization.h
+++ b/Engine/Source/Editor/LevelEditor/Private/Elements/Actor/ActorElementLevelEditorCommonActionsCustomization.h
@@ -8,7 +8,7 @@
 class FActorElementLevelEditorCommonActionsCustomization : public FTypedElementCommonActionsCustomization, public FTypedElementAssetEditorToolkitHostMixin
 {
 public:
-	virtual void GetElementsToDelete(const TTypedElement<UTypedElementWorldInterface>& InElementWorldHandle, const UTypedElementSelectionSet* InSelectionSet, UTypedElementList* OutElementsToDelete) override;
+	virtual void GetElementsForAction(const TTypedElement<UTypedElementWorldInterface>& InElementWorldHandle, const UTypedElementSelectionSet* InSelectionSet, UTypedElementList* OutElementsToDelete) override;
 	virtual bool DeleteElements(UTypedElementWorldInterface* InWorldInterface, TArrayView<const FTypedElementHandle> InElementHandles, UWorld* InWorld, UTypedElementSelectionSet* InSelectionSet, const FTypedElementDeletionOptions& InDeletionOptions) override;
-	virtual void DuplicateElements(UTypedElementWorldInterface* InWorldInterface, TArrayView<const FTypedElementHandle> InElementHandles, UWorld* InWorld, bool bOffsetLocations, TArray<FTypedElementHandle>& OutNewElements) override;
+	virtual void DuplicateElements(UTypedElementWorldInterface* InWorldInterface, TArrayView<const FTypedElementHandle> InElementHandles, UWorld* InWorld, const FVector& InLocationOffset, TArray<FTypedElementHandle>& OutNewElements) override;
 };
diff --git a/Engine/Source/Editor/LevelEditor/Private/Elements/Component/ComponentElementLevelEditorCommonActionsCustomization.cpp b/Engine/Source/Editor/LevelEditor/Private/Elements/Component/ComponentElementLevelEditorCommonActionsCustomization.cpp
--- a/Engine/Source/Editor/LevelEditor/Private/Elements/Component/ComponentElementLevelEditorCommonActionsCustomization.cpp
+++ b/Engine/Source/Editor/LevelEditor/Private/Elements/Component/ComponentElementLevelEditorCommonActionsCustomization.cpp
@@ -21,7 +21,7 @@ bool FComponentElementLevelEditorCommonActionsCustomization::DeleteElements(UTyp
 	return FTypedElementCommonActionsCustomization::DeleteElements(InWorldInterface, InElementHandles, InWorld, InSelectionSet, InDeletionOptions);
 }
 
-void FComponentElementLevelEditorCommonActionsCustomization::DuplicateElements(UTypedElementWorldInterface* InWorldInterface, TArrayView<const FTypedElementHandle> InElementHandles, UWorld* InWorld, bool bOffsetLocations, TArray<FTypedElementHandle>& OutNewElements)
+void FComponentElementLevelEditorCommonActionsCustomization::DuplicateElements(UTypedElementWorldInterface* InWorldInterface, TArrayView<const FTypedElementHandle> InElementHandles, UWorld* InWorld, const FVector& InLocationOffset, TArray<FTypedElementHandle>& OutNewElements)
 {
 	if (const IToolkitHost* ToolkitHostPtr = GetToolkitHost())
 	{
@@ -32,5 +32,5 @@ void FComponentElementLevelEditorCommonActionsCustomization::DuplicateElements(U
 		}
 	}
 
-	FTypedElementCommonActionsCustomization::DuplicateElements(InWorldInterface, InElementHandles, InWorld, bOffsetLocations, OutNewElements);
+	FTypedElementCommonActionsCustomization::DuplicateElements(InWorldInterface, InElementHandles, InWorld, InLocationOffset, OutNewElements);
 }
diff --git a/Engine/Source/Editor/LevelEditor/Private/Elements/Component/ComponentElementLevelEditorCommonActionsCustomization.h b/Engine/Source/Editor/LevelEditor/Private/Elements/Component/ComponentElementLevelEditorCommonActionsCustomization.h
--- a/Engine/Source/Editor/LevelEditor/Private/Elements/Component/ComponentElementLevelEditorCommonActionsCustomization.h
+++ b/Engine/Source/Editor/LevelEditor/Private/Elements/Component/ComponentElementLevelEditorCommonActionsCustomization.h
@@ -9,5 +9,5 @@ class FComponentElementLevelEditorCommonActionsCustomization : public FTypedElem
 {
 public:
 	virtual bool DeleteElements(UTypedElementWorldInterface* InWorldInterface, TArrayView<const FTypedElementHandle> InElementHandles, UWorld* InWorld, UTypedElementSelectionSet* InSelectionSet, const FTypedElementDeletionOptions& InDeletionOptions) override;
-	virtual void DuplicateElements(UTypedElementWorldInterface* InWorldInterface, TArrayView<const FTypedElementHandle> InElementHandles, UWorld* InWorld, bool bOffsetLocations, TArray<FTypedElementHandle>& OutNewElements) override;
+	virtual void DuplicateElements(UTypedElementWorldInterface* InWorldInterface, TArrayView<const FTypedElementHandle> InElementHandles, UWorld* InWorld, const FVector& InLocationOffset, TArray<FTypedElementHandle>& OutNewElements) override;
 };
diff --git a/Engine/Source/Editor/LevelEditor/Private/LevelEditorActions.cpp b/Engine/Source/Editor/LevelEditor/Private/LevelEditorActions.cpp
--- a/Engine/Source/Editor/LevelEditor/Private/LevelEditorActions.cpp
+++ b/Engine/Source/Editor/LevelEditor/Private/LevelEditorActions.cpp
@@ -1502,41 +1502,30 @@ bool FLevelEditorActionCallbacks::Duplicate_CanExecute()
 		return false;
 	}
 	
-	// If we can copy, we can duplicate
-	bool bCanCopy = false;
-	if (GEditor->GetSelectedComponentCount() > 0)
+	static const FName NAME_LevelEditor = "LevelEditor";
+	if (TSharedPtr<ILevelEditor> LevelEditor = FModuleManager::GetModuleChecked<FLevelEditorModule>(NAME_LevelEditor).GetLevelEditorInstance().Pin())
 	{
-		TArray<UActorComponent*> SelectedComponents;
-		for (FSelectionIterator It(GEditor->GetSelectedComponentIterator()); It; ++It)
-		{
-			SelectedComponents.Add(CastChecked<UActorComponent>(*It));
-		}
+		bool bCanDuplicate = false;
 
-		bCanCopy = FComponentEditorUtils::CanCopyComponents(SelectedComponents);
-	}
-	else
-	{
-		UWorld* World = GetWorld();
-		if (World)
+		const UTypedElementSelectionSet* SelectionSet = LevelEditor->GetElementSelectionSet();
+		SelectionSet->ForEachSelectedElement<UTypedElementWorldInterface>([&bCanDuplicate](const TTypedElement<UTypedElementWorldInterface>& InWorldElement)
 		{
-			bCanCopy = GUnrealEd->CanCopySelectedActorsToClipboard(World);
-		}
-	}
+			bCanDuplicate |= InWorldElement.CanDuplicateElement();
+			return !bCanDuplicate;
+		});
 
-	if (!bCanCopy)
-	{
-		TWeakPtr<class ILevelEditor> LevelEditor = FModuleManager::GetModuleChecked<FLevelEditorModule>(TEXT("LevelEditor")).GetLevelEditorInstance();
-		if (LevelEditor.IsValid())
+		if (!bCanDuplicate)
 		{
-			TSharedPtr<class ISceneOutliner> SceneOutlinerPtr = LevelEditor.Pin()->GetSceneOutliner();
-			if (SceneOutlinerPtr.IsValid())
+			if (TSharedPtr<ISceneOutliner> SceneOutlinerPtr = LevelEditor->GetSceneOutliner())
 			{
-				bCanCopy = SceneOutlinerPtr->Copy_CanExecute();
+				bCanDuplicate = SceneOutlinerPtr->Copy_CanExecute(); // If we can copy, we can duplicate
 			}
 		}
+
+		return bCanDuplicate;
 	}
 
-	return bCanCopy;
+	return false;
 }
 
 bool FLevelEditorActionCallbacks::Delete_CanExecute()
diff --git a/Engine/Source/Editor/UnrealEd/Classes/Editor/EditorEngine.h b/Engine/Source/Editor/UnrealEd/Classes/Editor/EditorEngine.h
--- a/Engine/Source/Editor/UnrealEd/Classes/Editor/EditorEngine.h
+++ b/Engine/Source/Editor/UnrealEd/Classes/Editor/EditorEngine.h
@@ -2223,7 +2223,7 @@ public:
 	/** 
 	 * @return - returns the currently selected positional snap grid setting
 	 */
-	float GetGridSize();
+	float GetGridSize() const;
 
 	/** 
 	 * @return - if the grid size is part of the 1,2,4,8,16,.. list or not
@@ -2285,6 +2285,11 @@ public:
 
 	float GetGridInterval();
 
+	/**
+	 * Get the location offset applied by the grid based on the grid size and active viewport.
+	 */
+	FVector GetGridLocationOffset(bool bUniformOffset) const;
+
 	/**
 	 * Access the array of grid interval options
 	 */
diff --git a/Engine/Source/Editor/UnrealEd/Classes/Editor/UnrealEdEngine.h b/Engine/Source/Editor/UnrealEd/Classes/Editor/UnrealEdEngine.h
--- a/Engine/Source/Editor/UnrealEd/Classes/Editor/UnrealEdEngine.h
+++ b/Engine/Source/Editor/UnrealEd/Classes/Editor/UnrealEdEngine.h
@@ -573,12 +573,12 @@ public:
 	 *
 	 * @param	OutPastedActors		List of all the actors that were pasted
 	 * @param	InWorld				World context
+	 * @param	LocationOffset		Offset to apply to actor locations after they're created
 	 * @param	bDuplicate			Is this a duplicate operation (as opposed to a real paste)?
-	 * @param	bOffsetLocations	Should the actor locations be offset after they are created?
 	 * @param	bWarnIfHidden		If true displays a warning if the destination level is hidden
 	 * @param	SourceData			If != NULL, use instead of clipboard data
 	 */
-	virtual void PasteActors(TArray<AActor*>& OutPastedActors, UWorld* InWorld, bool bDuplicate, bool bOffsetLocations, bool bWarnIfHidden, const FString* SourceData = nullptr);
+	virtual void PasteActors(TArray<AActor*>& OutPastedActors, UWorld* InWorld, const FVector& LocationOffset, bool bDuplicate, bool bWarnIfHidden, const FString* SourceData = nullptr);
 
 	/**
 	 * Duplicates selected actors.  Handles the case where you are trying to duplicate PrefabInstance actors.
@@ -603,9 +603,9 @@ public:
 	 * @param	InActorsToDuplicate	Array of actors to duplicate
 	 * @param	OutNewActors		List of all the actors that were duplicated
 	 * @param	InLevel				Level to place duplicate
-	 * @param	bUseOffset			Should the actor locations be offset after they are created?
+	 * @param	LocationOffset		Offset to apply to actor locations after they're created
 	 */
-	virtual void DuplicateActors(const TArray<AActor*>& InActorsToDuplicate, TArray<AActor*>& OutNewActors, ULevel* InLevel, bool bOffsetLocations);
+	virtual void DuplicateActors(const TArray<AActor*>& InActorsToDuplicate, TArray<AActor*>& OutNewActors, ULevel* InLevel, const FVector& LocationOffset);
 
 	/**
 	 * Replace all selected brushes with the default brush.
diff --git a/Engine/Source/Editor/UnrealEd/Private/EditorActor.cpp b/Engine/Source/Editor/UnrealEd/Private/EditorActor.cpp
--- a/Engine/Source/Editor/UnrealEd/Private/EditorActor.cpp
+++ b/Engine/Source/Editor/UnrealEd/Private/EditorActor.cpp
@@ -171,33 +171,6 @@ void UUnrealEdEngine::CopyActors(const TArray<AActor*>& InActorsToCopy, UWorld*
 	}
 }
 
-
-/**
- * Creates offsets for locations based on the editor grid size and active viewport.
- */
-static FVector CreateLocationOffset(bool bDuplicate, bool bOffsetLocations)
-{
-	const float Offset = static_cast<float>( bOffsetLocations ? GEditor->GetGridSize() : 0 );
-	FVector LocationOffset(Offset,Offset,Offset);
-	if ( bDuplicate && GCurrentLevelEditingViewportClient )
-	{
-		switch( GCurrentLevelEditingViewportClient->ViewportType )
-		{
-		case LVT_OrthoXZ:
-			LocationOffset = FVector(Offset,0.f,Offset);
-			break;
-		case LVT_OrthoYZ:
-			LocationOffset = FVector(0.f,Offset,Offset);
-			break;
-		default:
-			LocationOffset = FVector(Offset,Offset,0.f);
-			break;
-		}
-	}
-	return LocationOffset;
-}
-
-
 bool UUnrealEdEngine::WarnIfDestinationLevelIsHidden( UWorld* InWorld ) const
 {
 	bool bShouldLoadHiddenLevels = true;
@@ -290,7 +263,7 @@ void UUnrealEdEngine::edactPasteSelected(UWorld* InWorld, bool bDuplicate, bool
 	else
 	{
 		TArray<AActor*> PastedActors;
-		PasteActors(PastedActors, InWorld, bDuplicate, bOffsetLocations, bWarnIfHidden, SourceData);
+		PasteActors(PastedActors, InWorld, bOffsetLocations ? GEditor->GetGridLocationOffset(/*bUniformOffset*/true) : FVector::ZeroVector, bDuplicate, bWarnIfHidden, SourceData);
 
 		if (PastedActors.Num() > 0)
 		{
@@ -333,7 +306,7 @@ void UUnrealEdEngine::PasteComponents(TArray<UActorComponent*>& OutPastedCompone
 	}
 }
 
-void UUnrealEdEngine::PasteActors(TArray<AActor*>& OutPastedActors, UWorld* InWorld, bool bDuplicate, bool bOffsetLocations, bool bWarnIfHidden, const FString* SourceData)
+void UUnrealEdEngine::PasteActors(TArray<AActor*>& OutPastedActors, UWorld* InWorld, const FVector& LocationOffset, bool bDuplicate, bool bWarnIfHidden, const FString* SourceData)
 {
 	OutPastedActors.Reset();
 
@@ -345,9 +318,6 @@ void UUnrealEdEngine::PasteActors(TArray<AActor*>& OutPastedActors, UWorld* InWo
 
 	const FScopedBusyCursor BusyCursor;
 
-	// Create a location offset.
-	const FVector LocationOffset = CreateLocationOffset(bDuplicate, bOffsetLocations);
-
 	FCachedActorLabels ActorLabels(InWorld);
 
 	// Get pasted text.
@@ -444,9 +414,9 @@ public:
 	 *
 	 * @param	OutNewActors			[out] Newly created actors are appended to this list.
 	 * @param	DestLevel				The level to duplicate the actors in this job to.
-	 * @param	bOffsetLocations		Passed to edactPasteSelected; true if new actor locations should be offset.
+	 * @param	bOffsetLocations		@see PasteActors.
 	 */
-	void DuplicateActorsToLevel(TArray<AActor*>& OutNewActors, ULevel* DestLevel, bool bOffsetLocations)
+	void DuplicateActorsToLevel(TArray<AActor*>& OutNewActors, ULevel* DestLevel, const FVector& LocationOffset)
 	{
 		// Check neither level is locked
 		if (FLevelUtils::IsLevelLocked(SrcLevel))
@@ -484,7 +454,7 @@ public:
 			DestLevel->OwningWorld->SetCurrentLevel(LevelPartitionScope.GetLevel());
 
 			TArray<AActor*> PastedActors;
-			GUnrealEd->PasteActors(PastedActors, DestLevel->OwningWorld, /*bDuplicate*/true, bOffsetLocations, /*bWarnIfHidden*/true, &ScratchData);
+			GUnrealEd->PasteActors(PastedActors, DestLevel->OwningWorld, LocationOffset, /*bDuplicate*/true, /*bWarnIfHidden*/true, &ScratchData);
 			
 			DestLevel->OwningWorld->SetCurrentLevel(CurrentDestLevel);
 
@@ -536,7 +506,7 @@ void UUnrealEdEngine::edactDuplicateSelected( ULevel* InLevel, bool bOffsetLocat
 		ActorSelection->GetSelectedObjects<AActor>(SelectedActors);
 
 		TArray<AActor*> NewActors;
-		DuplicateActors(SelectedActors, NewActors, InLevel, bOffsetLocations);
+		DuplicateActors(SelectedActors, NewActors, InLevel, bOffsetLocations ? GEditor->GetGridLocationOffset(/*bUniformOffset*/false) : FVector::ZeroVector);
 
 		if (NewActors.Num() > 0)
 		{
@@ -583,7 +553,7 @@ void UUnrealEdEngine::DuplicateComponents(const TArray<UActorComponent*>& InComp
 	}
 }
 
-void UUnrealEdEngine::DuplicateActors(const TArray<AActor*>& InActorsToDuplicate, TArray<AActor*>& OutNewActors, ULevel* InLevel, bool bOffsetLocations)
+void UUnrealEdEngine::DuplicateActors(const TArray<AActor*>& InActorsToDuplicate, TArray<AActor*>& OutNewActors, ULevel* InLevel, const FVector& LocationOffset)
 {
 	OutNewActors.Reset(InActorsToDuplicate.Num());
 
@@ -608,7 +578,7 @@ void UUnrealEdEngine::DuplicateActors(const TArray<AActor*>& InActorsToDuplicate
 	for (const auto& DuplicateJobPair : DuplicateJobs)
 	{
 		FDuplicateJob& Job = *DuplicateJobPair.Value;
-		Job.DuplicateActorsToLevel(OutNewActors, InLevel, bOffsetLocations);
+		Job.DuplicateActorsToLevel(OutNewActors, InLevel, LocationOffset);
 	}
 
 	// Finally, cleanup.
diff --git a/Engine/Source/Editor/UnrealEd/Private/EditorConstraints.cpp b/Engine/Source/Editor/UnrealEd/Private/EditorConstraints.cpp
--- a/Engine/Source/Editor/UnrealEd/Private/EditorConstraints.cpp
+++ b/Engine/Source/Editor/UnrealEd/Private/EditorConstraints.cpp
@@ -5,12 +5,13 @@
 =============================================================================*/
 
 #include "CoreMinimal.h"
+#include "LevelEditorViewport.h"
 #include "Settings/LevelEditorViewportSettings.h"
 #include "Editor/EditorEngine.h"
 #include "Editor.h"
 #include "EditorSupportDelegates.h"
 
-float UEditorEngine::GetGridSize()
+float UEditorEngine::GetGridSize() const
 {
 	const TArray<float>& PosGridSizes = GetCurrentPositionGridArray();
 	const int32 CurrentPosGridSize = GetDefault<ULevelEditorViewportSettings>()->CurrentPosGridSize;
@@ -163,4 +164,27 @@ float UEditorEngine::GetGridInterval()
 	return IntervalValue;
 }
 
+FVector UEditorEngine::GetGridLocationOffset(bool bUniformOffset) const
+{
+	const float Offset = GetDefault<ULevelEditorViewportSettings>()->GridEnabled ? GetGridSize() : 0.0f;
 
+	FVector LocationOffset(Offset, Offset, Offset);
+	if (!bUniformOffset && GCurrentLevelEditingViewportClient)
+	{
+		switch (GCurrentLevelEditingViewportClient->ViewportType)
+		{
+		case LVT_OrthoXZ:
+			LocationOffset = FVector(Offset, 0.f, Offset);
+			break;
+
+		case LVT_OrthoYZ:
+			LocationOffset = FVector(0.f, Offset, Offset);
+			break;
+
+		default:
+			LocationOffset = FVector(Offset, Offset, 0.f);
+			break;
+		}
+	}
+	return LocationOffset;
+}
diff --git a/Engine/Source/Editor/UnrealEd/Private/Elements/Actor/ActorElementEditorWorldInterface.cpp b/Engine/Source/Editor/UnrealEd/Private/Elements/Actor/ActorElementEditorWorldInterface.cpp
--- a/Engine/Source/Editor/UnrealEd/Private/Elements/Actor/ActorElementEditorWorldInterface.cpp
+++ b/Engine/Source/Editor/UnrealEd/Private/Elements/Actor/ActorElementEditorWorldInterface.cpp
@@ -74,7 +74,13 @@ bool UActorElementEditorWorldInterface::DeleteElements(TArrayView<const FTypedEl
 		&& GUnrealEd->DeleteActors(ActorsToDelete, InWorld, InSelectionSet, InDeletionOptions.VerifyDeletionCanHappen(), InDeletionOptions.WarnAboutReferences(), InDeletionOptions.WarnAboutSoftReferences());
 }
 
-void UActorElementEditorWorldInterface::DuplicateElements(TArrayView<const FTypedElementHandle> InElementHandles, UWorld* InWorld, bool bOffsetLocations, TArray<FTypedElementHandle>& OutNewElements)
+bool UActorElementEditorWorldInterface::CanDuplicateElement(const FTypedElementHandle& InElementHandle)
+{
+	AActor* Actor = ActorElementDataUtil::GetActorFromHandle(InElementHandle);
+	return Actor != nullptr; // All actors can be duplicated
+}
+
+void UActorElementEditorWorldInterface::DuplicateElements(TArrayView<const FTypedElementHandle> InElementHandles, UWorld* InWorld, const FVector& InLocationOffset, TArray<FTypedElementHandle>& OutNewElements)
 {
 	const TArray<AActor*> ActorsToDuplicate = ActorElementDataUtil::GetActorsFromHandles(InElementHandles);
 
@@ -84,7 +90,7 @@ void UActorElementEditorWorldInterface::DuplicateElements(TArrayView<const FType
 
 		TArray<AActor*> NewActors;
 		ABrush::SetSuppressBSPRegeneration(true);
-		GUnrealEd->DuplicateActors(ActorsToDuplicate, NewActors, InWorld->GetCurrentLevel(), bOffsetLocations);
+		GUnrealEd->DuplicateActors(ActorsToDuplicate, NewActors, InWorld->GetCurrentLevel(), InLocationOffset);
 		ABrush::SetSuppressBSPRegeneration(false);
 
 		FEditorDelegates::OnDuplicateActorsEnd.Broadcast();
diff --git a/Engine/Source/Editor/UnrealEd/Private/Elements/Actor/ActorElementEditorWorldInterface.h b/Engine/Source/Editor/UnrealEd/Private/Elements/Actor/ActorElementEditorWorldInterface.h
--- a/Engine/Source/Editor/UnrealEd/Private/Elements/Actor/ActorElementEditorWorldInterface.h
+++ b/Engine/Source/Editor/UnrealEd/Private/Elements/Actor/ActorElementEditorWorldInterface.h
@@ -18,5 +18,6 @@ public:
 	virtual void NotifyMovementEnded(const FTypedElementHandle& InElementHandle) override;
 	virtual bool CanDeleteElement(const FTypedElementHandle& InElementHandle) override;
 	virtual bool DeleteElements(TArrayView<const FTypedElementHandle> InElementHandles, UWorld* InWorld, UTypedElementSelectionSet* InSelectionSet, const FTypedElementDeletionOptions& InDeletionOptions) override;
-	virtual void DuplicateElements(TArrayView<const FTypedElementHandle> InElementHandles, UWorld* InWorld, bool bOffsetLocations, TArray<FTypedElementHandle>& OutNewElements) override;
+	virtual bool CanDuplicateElement(const FTypedElementHandle& InElementHandle) override;
+	virtual void DuplicateElements(TArrayView<const FTypedElementHandle> InElementHandles, UWorld* InWorld, const FVector& InLocationOffset, TArray<FTypedElementHandle>& OutNewElements) override;
 };
diff --git a/Engine/Source/Editor/UnrealEd/Private/Elements/Component/ComponentElementEditorWorldInterface.cpp b/Engine/Source/Editor/UnrealEd/Private/Elements/Component/ComponentElementEditorWorldInterface.cpp
--- a/Engine/Source/Editor/UnrealEd/Private/Elements/Component/ComponentElementEditorWorldInterface.cpp
+++ b/Engine/Source/Editor/UnrealEd/Private/Elements/Component/ComponentElementEditorWorldInterface.cpp
@@ -9,6 +9,7 @@
 #include "Editor.h"
 #include "UnrealEdGlobals.h"
 #include "Editor/UnrealEdEngine.h"
+#include "Kismet2/ComponentEditorUtils.h"
 
 void UComponentElementEditorWorldInterface::NotifyMovementStarted(const FTypedElementHandle& InElementHandle)
 {
@@ -57,7 +58,13 @@ bool UComponentElementEditorWorldInterface::DeleteElements(TArrayView<const FTyp
 		&& GUnrealEd->DeleteComponents(ComponentsToDelete, InSelectionSet, InDeletionOptions.VerifyDeletionCanHappen());
 }
 
-void UComponentElementEditorWorldInterface::DuplicateElements(TArrayView<const FTypedElementHandle> InElementHandles, UWorld* InWorld, bool bOffsetLocations, TArray<FTypedElementHandle>& OutNewElements)
+bool UComponentElementEditorWorldInterface::CanDuplicateElement(const FTypedElementHandle& InElementHandle)
+{
+	UActorComponent* Component = ComponentElementDataUtil::GetComponentFromHandle(InElementHandle);
+	return Component && FComponentEditorUtils::CanCopyComponent(Component); // If we can copy, we can duplicate
+}
+
+void UComponentElementEditorWorldInterface::DuplicateElements(TArrayView<const FTypedElementHandle> InElementHandles, UWorld* InWorld, const FVector& InLocationOffset, TArray<FTypedElementHandle>& OutNewElements)
 {
 	const TArray<UActorComponent*> ComponentsToDuplicate = ComponentElementDataUtil::GetComponentsFromHandles(InElementHandles);
 
diff --git a/Engine/Source/Editor/UnrealEd/Private/Elements/Component/ComponentElementEditorWorldInterface.h b/Engine/Source/Editor/UnrealEd/Private/Elements/Component/ComponentElementEditorWorldInterface.h
--- a/Engine/Source/Editor/UnrealEd/Private/Elements/Component/ComponentElementEditorWorldInterface.h
+++ b/Engine/Source/Editor/UnrealEd/Private/Elements/Component/ComponentElementEditorWorldInterface.h
@@ -16,5 +16,6 @@ public:
 	virtual void NotifyMovementEnded(const FTypedElementHandle& InElementHandle) override;
 	virtual bool CanDeleteElement(const FTypedElementHandle& InElementHandle) override;
 	virtual bool DeleteElements(TArrayView<const FTypedElementHandle> InElementHandles, UWorld* InWorld, UTypedElementSelectionSet* InSelectionSet, const FTypedElementDeletionOptions& InDeletionOptions) override;
-	virtual void DuplicateElements(TArrayView<const FTypedElementHandle> InElementHandles, UWorld* InWorld, bool bOffsetLocations, TArray<FTypedElementHandle>& OutNewElements) override;
+	virtual bool CanDuplicateElement(const FTypedElementHandle& InElementHandle) override;
+	virtual void DuplicateElements(TArrayView<const FTypedElementHandle> InElementHandles, UWorld* InWorld, const FVector& InLocationOffset, TArray<FTypedElementHandle>& OutNewElements) override;
 };
diff --git a/Engine/Source/Editor/UnrealEd/Private/Elements/SMInstance/SMInstanceElementEditorWorldInterface.cpp b/Engine/Source/Editor/UnrealEd/Private/Elements/SMInstance/SMInstanceElementEditorWorldInterface.cpp
--- a/Engine/Source/Editor/UnrealEd/Private/Elements/SMInstance/SMInstanceElementEditorWorldInterface.cpp
+++ b/Engine/Source/Editor/UnrealEd/Private/Elements/SMInstance/SMInstanceElementEditorWorldInterface.cpp
@@ -67,7 +67,13 @@ bool USMInstanceElementEditorWorldInterface::DeleteElements(TArrayView<const FTy
 	return false;
 }
 
-FTypedElementHandle USMInstanceElementEditorWorldInterface::DuplicateElement(const FTypedElementHandle& InElementHandle, UWorld* InWorld, bool bOffsetLocations)
+bool USMInstanceElementEditorWorldInterface::CanDuplicateElement(const FTypedElementHandle& InElementHandle)
+{
+	const FSMInstanceId SMInstance = SMInstanceElementDataUtil::GetSMInstanceFromHandle(InElementHandle);
+	return SMInstance && CanEditSMInstance(SMInstance);
+}
+
+FTypedElementHandle USMInstanceElementEditorWorldInterface::DuplicateElement(const FTypedElementHandle& InElementHandle, UWorld* InWorld, const FVector& InLocationOffset)
 {
 	if (const FSMInstanceId SMInstance = SMInstanceElementDataUtil::GetSMInstanceFromHandle(InElementHandle))
 	{
@@ -75,6 +81,7 @@ FTypedElementHandle USMInstanceElementEditorWorldInterface::DuplicateElement(con
 		{
 			FTransform NewInstanceTransform = FTransform::Identity;
 			SMInstance.ISMComponent->GetInstanceTransform(SMInstance.InstanceIndex, NewInstanceTransform);
+			NewInstanceTransform.SetTranslation(NewInstanceTransform.GetTranslation() + InLocationOffset);
 
 			const int32 NewInstanceIndex = SMInstance.ISMComponent->AddInstance(NewInstanceTransform);
 			return UEngineElementsLibrary::AcquireEditorSMInstanceElementHandle(FSMInstanceId{ SMInstance.ISMComponent, NewInstanceIndex });
@@ -84,7 +91,7 @@ FTypedElementHandle USMInstanceElementEditorWorldInterface::DuplicateElement(con
 	return FTypedElementHandle();
 }
 
-void USMInstanceElementEditorWorldInterface::DuplicateElements(TArrayView<const FTypedElementHandle> InElementHandles, UWorld* InWorld, bool bOffsetLocations, TArray<FTypedElementHandle>& OutNewElements)
+void USMInstanceElementEditorWorldInterface::DuplicateElements(TArrayView<const FTypedElementHandle> InElementHandles, UWorld* InWorld, const FVector& InLocationOffset, TArray<FTypedElementHandle>& OutNewElements)
 {
 	const TArray<FSMInstanceId> SMInstancesToDuplicate = SMInstanceElementDataUtil::GetSMInstancesFromHandles(InElementHandles);
 
@@ -109,6 +116,7 @@ void USMInstanceElementEditorWorldInterface::DuplicateElements(TArrayView<const
 			{
 				FTransform& NewInstanceTransform = NewInstanceTransforms.Add_GetRef(FTransform::Identity);
 				BatchedInstancesToDuplicatePair.Key->GetInstanceTransform(InstanceIndex, NewInstanceTransform);
+				NewInstanceTransform.SetTranslation(NewInstanceTransform.GetTranslation() + InLocationOffset);
 			}
 
 			const TArray<int32> NewInstanceIndices = BatchedInstancesToDuplicatePair.Key->AddInstances(NewInstanceTransforms, /*bShouldReturnIndices*/true);
diff --git a/Engine/Source/Editor/UnrealEd/Private/Elements/SMInstance/SMInstanceElementEditorWorldInterface.h b/Engine/Source/Editor/UnrealEd/Private/Elements/SMInstance/SMInstanceElementEditorWorldInterface.h
--- a/Engine/Source/Editor/UnrealEd/Private/Elements/SMInstance/SMInstanceElementEditorWorldInterface.h
+++ b/Engine/Source/Editor/UnrealEd/Private/Elements/SMInstance/SMInstanceElementEditorWorldInterface.h
@@ -14,6 +14,7 @@ public:
 	virtual bool CanDeleteElement(const FTypedElementHandle& InElementHandle) override;
 	virtual bool DeleteElement(const FTypedElementHandle& InElementHandle, UWorld* InWorld, UTypedElementSelectionSet* InSelectionSet, const FTypedElementDeletionOptions& InDeletionOptions) override;
 	virtual bool DeleteElements(TArrayView<const FTypedElementHandle> InElementHandles, UWorld* InWorld, UTypedElementSelectionSet* InSelectionSet, const FTypedElementDeletionOptions& InDeletionOptions) override;
-	virtual FTypedElementHandle DuplicateElement(const FTypedElementHandle& InElementHandle, UWorld* InWorld, bool bOffsetLocations) override;
-	virtual void DuplicateElements(TArrayView<const FTypedElementHandle> InElementHandles, UWorld* InWorld, bool bOffsetLocations, TArray<FTypedElementHandle>& OutNewElements) override;
+	virtual bool CanDuplicateElement(const FTypedElementHandle& InElementHandle) override;
+	virtual FTypedElementHandle DuplicateElement(const FTypedElementHandle& InElementHandle, UWorld* InWorld, const FVector& InLocationOffset) override;
+	virtual void DuplicateElements(TArrayView<const FTypedElementHandle> InElementHandles, UWorld* InWorld, const FVector& InLocationOffset, TArray<FTypedElementHandle>& OutNewElements) override;
 };
diff --git a/Engine/Source/Editor/UnrealEd/Private/LevelEditorViewport.cpp b/Engine/Source/Editor/UnrealEd/Private/LevelEditorViewport.cpp
--- a/Engine/Source/Editor/UnrealEd/Private/LevelEditorViewport.cpp
+++ b/Engine/Source/Editor/UnrealEd/Private/LevelEditorViewport.cpp
@@ -2594,7 +2594,7 @@ bool FLevelEditorViewportClient::InputWidgetDelta(FViewport* InViewport, EAxisLi
 							if (CommonActions)
 							{
 								const UTypedElementList* ElementsToManipulate = GetElementsToManipulate();
-								TArray<FTypedElementHandle> DuplicatedElements = CommonActions->DuplicateElements(ElementsToManipulate, GetWorld(), /*bOffsetLocations*/false);
+								TArray<FTypedElementHandle> DuplicatedElements = CommonActions->DuplicateElements(ElementsToManipulate, GetWorld(), FVector::ZeroVector);
 
 								// Exclusively select the new elements, so that future gizmo interaction manipulates those items instead
 								if (DuplicatedElements.Num() > 0)
diff --git a/Engine/Source/Editor/UnrealEd/Private/UnrealEdSrv.cpp b/Engine/Source/Editor/UnrealEd/Private/UnrealEdSrv.cpp
--- a/Engine/Source/Editor/UnrealEd/Private/UnrealEdSrv.cpp
+++ b/Engine/Source/Editor/UnrealEd/Private/UnrealEdSrv.cpp
@@ -2495,15 +2495,9 @@ bool UUnrealEdEngine::Exec_Actor( UWorld* InWorld, const TCHAR* Str, FOutputDevi
 			if (UTypedElementCommonActions* CommonActions = LevelEditor->GetCommonActions())
 			{
 				UTypedElementSelectionSet* SelectionSet = LevelEditor->GetMutableElementSelectionSet();
-
-				TStrongObjectPtr<UTypedElementList> ElementsToDelete(UTypedElementRegistry::GetInstance()->CreateElementList());
-				CommonActions->GetSelectedElementsToDelete(SelectionSet, ElementsToDelete.Get());
-				if (ElementsToDelete->Num() > 0)
-				{
-					const FScopedTransaction Transaction(NSLOCTEXT("UnrealEd", "DeleteElements", "Delete Elements"));
-					CommonActions->DeleteElements(ElementsToDelete.Get(), InWorld, SelectionSet, FTypedElementDeletionOptions());
-					ElementsToDelete->Reset();
-				}
+				
+				const FScopedTransaction Transaction(NSLOCTEXT("UnrealEd", "DeleteElements", "Delete Elements"));
+				CommonActions->DeleteSelectedElements(SelectionSet, InWorld, FTypedElementDeletionOptions());
 			}
 		}
 		return true;
@@ -2674,7 +2668,21 @@ bool UUnrealEdEngine::Exec_Actor( UWorld* InWorld, const TCHAR* Str, FOutputDevi
 	}
 	else if( FParse::Command(&Str,TEXT("DUPLICATE")) )
 	{
-		DuplicateSelectedActors(InWorld);
+		if (TSharedPtr<ILevelEditor> LevelEditor = GCurrentLevelEditingViewportClient->ParentLevelEditor.Pin())
+		{
+			if (UTypedElementCommonActions* CommonActions = LevelEditor->GetCommonActions())
+			{
+				UTypedElementSelectionSet* SelectionSet = LevelEditor->GetMutableElementSelectionSet();
+
+				const FScopedTransaction Transaction(NSLOCTEXT("UnrealEd", "DuplicateElements", "Duplicate Elements"));
+				const TArray<FTypedElementHandle> DuplicatedElements = CommonActions->DuplicateSelectedElements(SelectionSet, InWorld, GEditor->GetGridLocationOffset(/*bUniformOffset*/false));
+				if (DuplicatedElements.Num() > 0)
+				{
+					SelectionSet->SetSelection(DuplicatedElements, FTypedElementSelectionOptions());
+					SelectionSet->NotifyPendingChanges();
+				}
+			}
+		}
 		return true;
 	}
 	else if( FParse::Command(&Str, TEXT("ALIGN")) )
diff --git a/Engine/Source/Runtime/Engine/Private/Elements/Framework/TypedElementCommonActions.cpp b/Engine/Source/Runtime/Engine/Private/Elements/Framework/TypedElementCommonActions.cpp
--- a/Engine/Source/Runtime/Engine/Private/Elements/Framework/TypedElementCommonActions.cpp
+++ b/Engine/Source/Runtime/Engine/Private/Elements/Framework/TypedElementCommonActions.cpp
@@ -5,7 +5,9 @@
 #include "Elements/Framework/TypedElementRegistry.h"
 #include "Elements/Framework/TypedElementUtil.h"
 
-void FTypedElementCommonActionsCustomization::GetElementsToDelete(const TTypedElement<UTypedElementWorldInterface>& InElementWorldHandle, const UTypedElementSelectionSet* InSelectionSet, UTypedElementList* OutElementsToDelete)
+#include "UObject/GCObjectScopeGuard.h"
+
+void FTypedElementCommonActionsCustomization::GetElementsForAction(const TTypedElement<UTypedElementWorldInterface>& InElementWorldHandle, const UTypedElementSelectionSet* InSelectionSet, UTypedElementList* OutElementsToDelete)
 {
 	OutElementsToDelete->Add(InElementWorldHandle);
 }
@@ -15,30 +17,30 @@ bool FTypedElementCommonActionsCustomization::DeleteElements(UTypedElementWorldI
 	return InWorldInterface->DeleteElements(InElementHandles, InWorld, InSelectionSet, InDeletionOptions);
 }
 
-void FTypedElementCommonActionsCustomization::DuplicateElements(UTypedElementWorldInterface* InWorldInterface, TArrayView<const FTypedElementHandle> InElementHandles, UWorld* InWorld, bool bOffsetLocations, TArray<FTypedElementHandle>& OutNewElements)
+void FTypedElementCommonActionsCustomization::DuplicateElements(UTypedElementWorldInterface* InWorldInterface, TArrayView<const FTypedElementHandle> InElementHandles, UWorld* InWorld, const FVector& InLocationOffset, TArray<FTypedElementHandle>& OutNewElements)
 {
-	InWorldInterface->DuplicateElements(InElementHandles, InWorld, bOffsetLocations, OutNewElements);
+	InWorldInterface->DuplicateElements(InElementHandles, InWorld, InLocationOffset, OutNewElements);
 }
 
 
-void UTypedElementCommonActions::GetSelectedElementsToDelete(const UTypedElementSelectionSet* InSelectionSet, UTypedElementList* OutElementsToDelete) const
+void UTypedElementCommonActions::GetSelectedElementsForAction(const UTypedElementSelectionSet* InSelectionSet, UTypedElementList* OutElementsToDelete) const
 {
 	OutElementsToDelete->Reset();
 	InSelectionSet->ForEachSelectedElement<UTypedElementWorldInterface>([this, InSelectionSet, OutElementsToDelete](const TTypedElement<UTypedElementWorldInterface>& InElementWorldHandle)
 	{
 		FTypedElementCommonActionsElement CommonActionsElement(InElementWorldHandle, GetInterfaceCustomizationByTypeId(InElementWorldHandle.GetId().GetTypeId()));
 		check(CommonActionsElement.IsSet());
-		CommonActionsElement.GetElementsToDelete(InSelectionSet, OutElementsToDelete);
+		CommonActionsElement.GetElementsForAction(InSelectionSet, OutElementsToDelete);
 		return true;
 	});
 }
 
-bool UTypedElementCommonActions::DeleteElements(const TArray<FTypedElementHandle>& ElementHandles, UWorld* World, UTypedElementSelectionSet* InSelectionSet, const FTypedElementDeletionOptions& InDeletionOptions)
+bool UTypedElementCommonActions::DeleteElements(const TArray<FTypedElementHandle>& ElementHandles, UWorld* World, UTypedElementSelectionSet* SelectionSet, const FTypedElementDeletionOptions& DeletionOptions)
 {
-	return DeleteElements(MakeArrayView(ElementHandles), World, InSelectionSet, InDeletionOptions);
+	return DeleteElements(MakeArrayView(ElementHandles), World, SelectionSet, DeletionOptions);
 }
 
-bool UTypedElementCommonActions::DeleteElements(TArrayView<const FTypedElementHandle> ElementHandles, UWorld* World, UTypedElementSelectionSet* InSelectionSet, const FTypedElementDeletionOptions& InDeletionOptions)
+bool UTypedElementCommonActions::DeleteElements(TArrayView<const FTypedElementHandle> ElementHandles, UWorld* World, UTypedElementSelectionSet* SelectionSet, const FTypedElementDeletionOptions& DeletionOptions)
 {
 	TMap<FTypedHandleTypeId, TArray<FTypedElementHandle>> ElementsToDuplicateByType;
 	TypedElementUtil::BatchElementsByType(ElementHandles, ElementsToDuplicateByType);
@@ -54,14 +56,14 @@ bool UTypedElementCommonActions::DeleteElements(TArrayView<const FTypedElementHa
 		UTypedElementWorldInterface* WorldInterface = Registry->GetElementInterface<UTypedElementWorldInterface>(ElementsByTypePair.Key);
 		if (CommonActionsCustomization && WorldInterface)
 		{
-			bSuccess |= CommonActionsCustomization->DeleteElements(WorldInterface, ElementsByTypePair.Value, World, InSelectionSet, InDeletionOptions);
+			bSuccess |= CommonActionsCustomization->DeleteElements(WorldInterface, ElementsByTypePair.Value, World, SelectionSet, DeletionOptions);
 		}
 	}
 
 	return bSuccess;
 }
 
-bool UTypedElementCommonActions::DeleteElements(const UTypedElementList* ElementList, UWorld* World, UTypedElementSelectionSet* InSelectionSet, const FTypedElementDeletionOptions& InDeletionOptions)
+bool UTypedElementCommonActions::DeleteElements(const UTypedElementList* ElementList, UWorld* World, UTypedElementSelectionSet* SelectionSet, const FTypedElementDeletionOptions& DeletionOptions)
 {
 	TMap<FTypedHandleTypeId, TArray<FTypedElementHandle>> ElementsToDuplicateByType;
 	TypedElementUtil::BatchElementsByType(ElementList, ElementsToDuplicateByType);
@@ -77,19 +79,29 @@ bool UTypedElementCommonActions::DeleteElements(const UTypedElementList* Element
 		UTypedElementWorldInterface* WorldInterface = Registry->GetElementInterface<UTypedElementWorldInterface>(ElementsByTypePair.Key);
 		if (CommonActionsCustomization && WorldInterface)
 		{
-			bSuccess |= CommonActionsCustomization->DeleteElements(WorldInterface, ElementsByTypePair.Value, World, InSelectionSet, InDeletionOptions);
+			bSuccess |= CommonActionsCustomization->DeleteElements(WorldInterface, ElementsByTypePair.Value, World, SelectionSet, DeletionOptions);
 		}
 	}
 
 	return bSuccess;
 }
 
-TArray<FTypedElementHandle> UTypedElementCommonActions::DuplicateElements(const TArray<FTypedElementHandle>& ElementHandles, UWorld* World, bool bOffsetLocations)
+bool UTypedElementCommonActions::DeleteSelectedElements(UTypedElementSelectionSet* SelectionSet, UWorld* World, const FTypedElementDeletionOptions& DeletionOptions)
+{
+	TGCObjectScopeGuard<UTypedElementList> ElementsForAction(UTypedElementRegistry::GetInstance()->CreateElementList());
+	GetSelectedElementsForAction(SelectionSet, ElementsForAction.Get());
+
+	const bool bSuccess = DeleteElements(ElementsForAction.Get(), World, SelectionSet, DeletionOptions);
+	ElementsForAction.Get()->Reset();
+	return bSuccess;
+}
+
+TArray<FTypedElementHandle> UTypedElementCommonActions::DuplicateElements(const TArray<FTypedElementHandle>& ElementHandles, UWorld* World, const FVector& LocationOffset)
 {
-	return DuplicateElements(MakeArrayView(ElementHandles), World, bOffsetLocations);
+	return DuplicateElements(MakeArrayView(ElementHandles), World, LocationOffset);
 }
 
-TArray<FTypedElementHandle> UTypedElementCommonActions::DuplicateElements(TArrayView<const FTypedElementHandle> ElementHandles, UWorld* World, bool bOffsetLocations)
+TArray<FTypedElementHandle> UTypedElementCommonActions::DuplicateElements(TArrayView<const FTypedElementHandle> ElementHandles, UWorld* World, const FVector& LocationOffset)
 {
 	TMap<FTypedHandleTypeId, TArray<FTypedElementHandle>> ElementsToDuplicateByType;
 	TypedElementUtil::BatchElementsByType(ElementHandles, ElementsToDuplicateByType);
@@ -104,14 +116,14 @@ TArray<FTypedElementHandle> UTypedElementCommonActions::DuplicateElements(TArray
 		UTypedElementWorldInterface* WorldInterface = Registry->GetElementInterface<UTypedElementWorldInterface>(ElementsByTypePair.Key);
 		if (CommonActionsCustomization && WorldInterface)
 		{
-			CommonActionsCustomization->DuplicateElements(WorldInterface, ElementsByTypePair.Value, World, bOffsetLocations, NewElements);
+			CommonActionsCustomization->DuplicateElements(WorldInterface, ElementsByTypePair.Value, World, LocationOffset, NewElements);
 		}
 	}
 
 	return NewElements;
 }
 
-TArray<FTypedElementHandle> UTypedElementCommonActions::DuplicateElements(const UTypedElementList* ElementList, UWorld* World, bool bOffsetLocations)
+TArray<FTypedElementHandle> UTypedElementCommonActions::DuplicateElements(const UTypedElementList* ElementList, UWorld* World, const FVector& LocationOffset)
 {
 	TMap<FTypedHandleTypeId, TArray<FTypedElementHandle>> ElementsToDuplicateByType;
 	TypedElementUtil::BatchElementsByType(ElementList, ElementsToDuplicateByType);
@@ -126,13 +138,23 @@ TArray<FTypedElementHandle> UTypedElementCommonActions::DuplicateElements(const
 		UTypedElementWorldInterface* WorldInterface = Registry->GetElementInterface<UTypedElementWorldInterface>(ElementsByTypePair.Key);
 		if (CommonActionsCustomization && WorldInterface)
 		{
-			CommonActionsCustomization->DuplicateElements(WorldInterface, ElementsByTypePair.Value, World, bOffsetLocations, NewElements);
+			CommonActionsCustomization->DuplicateElements(WorldInterface, ElementsByTypePair.Value, World, LocationOffset, NewElements);
 		}
 	}
 
 	return NewElements;
 }
 
+TArray<FTypedElementHandle> UTypedElementCommonActions::DuplicateSelectedElements(const UTypedElementSelectionSet* SelectionSet, UWorld* World, const FVector& LocationOffset)
+{
+	TGCObjectScopeGuard<UTypedElementList> ElementsForAction(UTypedElementRegistry::GetInstance()->CreateElementList());
+	GetSelectedElementsForAction(SelectionSet, ElementsForAction.Get());
+
+	TArray<FTypedElementHandle> NewElements = DuplicateElements(ElementsForAction.Get(), World, LocationOffset);
+	ElementsForAction.Get()->Reset();
+	return NewElements;
+}
+
 FTypedElementCommonActionsElement UTypedElementCommonActions::ResolveCommonActionsElement(const FTypedElementHandle& InElementHandle) const
 {
 	return InElementHandle
diff --git a/Engine/Source/Runtime/Engine/Public/Elements/Framework/TypedElementCommonActions.h b/Engine/Source/Runtime/Engine/Public/Elements/Framework/TypedElementCommonActions.h
--- a/Engine/Source/Runtime/Engine/Public/Elements/Framework/TypedElementCommonActions.h
+++ b/Engine/Source/Runtime/Engine/Public/Elements/Framework/TypedElementCommonActions.h
@@ -18,9 +18,9 @@ public:
 	virtual ~FTypedElementCommonActionsCustomization() = default;
 
 	//~ See UTypedElementCommonActions for API docs
-	virtual void GetElementsToDelete(const TTypedElement<UTypedElementWorldInterface>& InElementWorldHandle, const UTypedElementSelectionSet* InSelectionSet, UTypedElementList* OutElementsToDelete);
+	virtual void GetElementsForAction(const TTypedElement<UTypedElementWorldInterface>& InElementWorldHandle, const UTypedElementSelectionSet* InSelectionSet, UTypedElementList* OutElementsToDelete);
 	virtual bool DeleteElements(UTypedElementWorldInterface* InWorldInterface, TArrayView<const FTypedElementHandle> InElementHandles, UWorld* InWorld, UTypedElementSelectionSet* InSelectionSet, const FTypedElementDeletionOptions& InDeletionOptions);
-	virtual void DuplicateElements(UTypedElementWorldInterface* InWorldInterface, TArrayView<const FTypedElementHandle> InElementHandles, UWorld* InWorld, bool bOffsetLocations, TArray<FTypedElementHandle>& OutNewElements);
+	virtual void DuplicateElements(UTypedElementWorldInterface* InWorldInterface, TArrayView<const FTypedElementHandle> InElementHandles, UWorld* InWorld, const FVector& InLocationOffset, TArray<FTypedElementHandle>& OutNewElements);
 };
 
 /**
@@ -55,7 +55,7 @@ public:
 	}
 
 	//~ See UTypedElementCommonActions for API docs
-	void GetElementsToDelete(const UTypedElementSelectionSet* InSelectionSet, UTypedElementList* OutElementsToDelete) { CommonActionsCustomization->GetElementsToDelete(ElementWorldHandle, InSelectionSet, OutElementsToDelete); }
+	void GetElementsForAction(const UTypedElementSelectionSet* InSelectionSet, UTypedElementList* OutElementsToDelete) { CommonActionsCustomization->GetElementsForAction(ElementWorldHandle, InSelectionSet, OutElementsToDelete); }
 
 private:
 	TTypedElement<UTypedElementWorldInterface> ElementWorldHandle;
@@ -73,19 +73,26 @@ class ENGINE_API UTypedElementCommonActions : public UObject, public TTypedEleme
 
 public:
 	/**
-	 * Get the elements from the given selection set that should be deleted
+	 * Get the elements from the given selection set that a common action should operate on.
+	 * @note This allows coarse filtering of the elements based on the selection state (eg, favoring components over actors), however it doesn't mean that the returned elements will actually be valid to perform a given action on.
 	 */
-	void GetSelectedElementsToDelete(const UTypedElementSelectionSet* InSelectionSet, UTypedElementList* OutElementsToDelete) const;
+	void GetSelectedElementsForAction(const UTypedElementSelectionSet* InSelectionSet, UTypedElementList* OutElementsForAction) const;
 
 	UFUNCTION(BlueprintCallable, Category="TypedElementFramework|Common")
-	bool DeleteElements(const TArray<FTypedElementHandle>& ElementHandles, UWorld* World, UTypedElementSelectionSet* InSelectionSet, const FTypedElementDeletionOptions& InDeletionOptions);
-	bool DeleteElements(TArrayView<const FTypedElementHandle> ElementHandles, UWorld* World, UTypedElementSelectionSet* InSelectionSet, const FTypedElementDeletionOptions& InDeletionOptions);
-	bool DeleteElements(const UTypedElementList* ElementList, UWorld* World, UTypedElementSelectionSet* InSelectionSet, const FTypedElementDeletionOptions& InDeletionOptions);
+	bool DeleteElements(const TArray<FTypedElementHandle>& ElementHandles, UWorld* World, UTypedElementSelectionSet* SelectionSet, const FTypedElementDeletionOptions& DeletionOptions);
+	bool DeleteElements(TArrayView<const FTypedElementHandle> ElementHandles, UWorld* World, UTypedElementSelectionSet* SelectionSet, const FTypedElementDeletionOptions& DeletionOptions);
+	bool DeleteElements(const UTypedElementList* ElementList, UWorld* World, UTypedElementSelectionSet* SelectionSet, const FTypedElementDeletionOptions& DeletionOptions);
 
 	UFUNCTION(BlueprintCallable, Category="TypedElementFramework|Common")
-	TArray<FTypedElementHandle> DuplicateElements(const TArray<FTypedElementHandle>& ElementHandles, UWorld* World, bool bOffsetLocations);
-	TArray<FTypedElementHandle> DuplicateElements(TArrayView<const FTypedElementHandle> ElementHandles, UWorld* World, bool bOffsetLocations);
-	TArray<FTypedElementHandle> DuplicateElements(const UTypedElementList* ElementList, UWorld* World, bool bOffsetLocations);
+	bool DeleteSelectedElements(UTypedElementSelectionSet* SelectionSet, UWorld* World, const FTypedElementDeletionOptions& DeletionOptions);
+
+	UFUNCTION(BlueprintCallable, Category="TypedElementFramework|Common")
+	TArray<FTypedElementHandle> DuplicateElements(const TArray<FTypedElementHandle>& ElementHandles, UWorld* World, const FVector& LocationOffset);
+	TArray<FTypedElementHandle> DuplicateElements(TArrayView<const FTypedElementHandle> ElementHandles, UWorld* World, const FVector& LocationOffset);
+	TArray<FTypedElementHandle> DuplicateElements(const UTypedElementList* ElementList, UWorld* World, const FVector& LocationOffset);
+
+	UFUNCTION(BlueprintCallable, Category="TypedElementFramework|Common")
+	TArray<FTypedElementHandle> DuplicateSelectedElements(const UTypedElementSelectionSet* SelectionSet, UWorld* World, const FVector& LocationOffset);
 
 private:
 	/**
diff --git a/Engine/Source/Runtime/Engine/Public/Elements/Interfaces/TypedElementWorldInterface.h b/Engine/Source/Runtime/Engine/Public/Elements/Interfaces/TypedElementWorldInterface.h
--- a/Engine/Source/Runtime/Engine/Public/Elements/Interfaces/TypedElementWorldInterface.h
+++ b/Engine/Source/Runtime/Engine/Public/Elements/Interfaces/TypedElementWorldInterface.h
@@ -165,15 +165,21 @@ public:
 	 */
 	virtual bool DeleteElements(TArrayView<const FTypedElementHandle> InElementHandles, UWorld* InWorld, UTypedElementSelectionSet* InSelectionSet, const FTypedElementDeletionOptions& InDeletionOptions) { return false; }
 
+	/**
+	 * Can the given element be duplicated?
+	 */
+	UFUNCTION(BlueprintCallable, Category="TypedElementInterfaces|World")
+	virtual bool CanDuplicateElement(const FTypedElementHandle& InElementHandle) { return false; }
+
 	/**
 	 * Duplicate the given element.
 	 * @note Default version calls DuplicateElements with a single element.
 	 */
 	UFUNCTION(BlueprintCallable, Category="TypedElementInterfaces|World")
-	virtual FTypedElementHandle DuplicateElement(const FTypedElementHandle& InElementHandle, UWorld* InWorld, bool bOffsetLocations)
+	virtual FTypedElementHandle DuplicateElement(const FTypedElementHandle& InElementHandle, UWorld* InWorld, const FVector& InLocationOffset)
 	{
 		TArray<FTypedElementHandle> NewElements;
-		DuplicateElements(MakeArrayView(&InElementHandle, 1), InWorld, bOffsetLocations, NewElements);
+		DuplicateElements(MakeArrayView(&InElementHandle, 1), InWorld, InLocationOffset, NewElements);
 		return NewElements.Num() > 0 ? MoveTemp(NewElements[0]) : FTypedElementHandle();
 	}
 
@@ -181,7 +187,7 @@ public:
 	 * Duplicate the given set of elements.
 	 * @note If you want to duplicate an array of elements that are potentially different types, you probably want to use the higher-level UTypedElementCommonActions::DuplicateElements function instead.
 	 */
-	virtual void DuplicateElements(TArrayView<const FTypedElementHandle> InElementHandles, UWorld* InWorld, bool bOffsetLocations, TArray<FTypedElementHandle>& OutNewElements)
+	virtual void DuplicateElements(TArrayView<const FTypedElementHandle> InElementHandles, UWorld* InWorld, const FVector& InLocationOffset, TArray<FTypedElementHandle>& OutNewElements)
 	{
 	}
 };
@@ -207,5 +213,6 @@ struct TTypedElement<UTypedElementWorldInterface> : public TTypedElementBase<UTy
 	bool FindSuitableTransformAlongPath(const FVector& InPathStart, const FVector& InPathEnd, const FCollisionShape& InTestShape, TArrayView<const FTypedElementHandle> InElementsToIgnore, FTransform& OutSuitableTransform) const { return InterfacePtr->FindSuitableTransformAlongPath(*this, InPathStart, InPathEnd, InTestShape, InElementsToIgnore, OutSuitableTransform); }
 	bool CanDeleteElement() const { return InterfacePtr->CanDeleteElement(*this); }
 	bool DeleteElement(UWorld* InWorld, UTypedElementSelectionSet* InSelectionSet, const FTypedElementDeletionOptions& InDeletionOptions) const { return InterfacePtr->DeleteElement(*this, InWorld, InSelectionSet, InDeletionOptions); }
-	FTypedElementHandle DuplicateElement(UWorld* InWorld, bool bOffsetLocations) const { return InterfacePtr->DuplicateElement(*this, InWorld, bOffsetLocations); }
+	bool CanDuplicateElement() const { return InterfacePtr->CanDuplicateElement(*this); }
+	FTypedElementHandle DuplicateElement(UWorld* InWorld, const FVector& InLocationOffset) const { return InterfacePtr->DuplicateElement(*this, InWorld, InLocationOffset); }
 };
