Material shader compilation fixes.
Fixed inconsistent capitalization of AsinDeriv, AsinFastDeriv, AcosDeriv, AcosFastDeriv, AtanDeriv, AtanFastDeriv
Fixed RotateScaleOffsetTexCoords so RotationScale/Offset are promoted to derivative types when they constant (EDerivativeStatus::Zero).

diff --git a/Engine/Shaders/Public/ShaderVersion.ush b/Engine/Shaders/Public/ShaderVersion.ush
--- a/Engine/Shaders/Public/ShaderVersion.ush
+++ b/Engine/Shaders/Public/ShaderVersion.ush
@@ -2,4 +2,4 @@
 // Each time the console command is executed it generates a new GUID. As this file is included
 // in Platform.ush (which should be included in any shader) it allows to invalidate the shader DDC.
 // 
-// GUID = 6FF70D5281BB4AEA9C90FD3CF0E5BE00
\ No newline at end of file
+// GUID = 8FF90D8281CB4FEB9C95FD3CF0E5BE07
\ No newline at end of file
diff --git a/Engine/Source/Runtime/Engine/Private/Materials/HLSLMaterialDerivativeAutogen.cpp b/Engine/Source/Runtime/Engine/Private/Materials/HLSLMaterialDerivativeAutogen.cpp
--- a/Engine/Source/Runtime/Engine/Private/Materials/HLSLMaterialDerivativeAutogen.cpp
+++ b/Engine/Source/Runtime/Engine/Private/Materials/HLSLMaterialDerivativeAutogen.cpp
@@ -538,22 +538,22 @@ int32 FMaterialDerivativeAutogen::GenerateExpressionFunc1(FHLSLMaterialTranslato
 			DstToken = TEXT("TanDeriv(") + SrcToken + TEXT(")");
 			break;
 		case EFunc1::Asin:
-			DstToken = TEXT("ASinDeriv(") + SrcToken + TEXT(")");
+			DstToken = TEXT("AsinDeriv(") + SrcToken + TEXT(")");
 			break;
 		case EFunc1::AsinFast:
-			DstToken = TEXT("ASinFastDeriv(") + SrcToken + TEXT(")");
+			DstToken = TEXT("AsinFastDeriv(") + SrcToken + TEXT(")");
 			break;
 		case EFunc1::Acos:
-			DstToken = TEXT("ACosDeriv(") + SrcToken + TEXT(")");
+			DstToken = TEXT("AcosDeriv(") + SrcToken + TEXT(")");
 			break;
 		case EFunc1::AcosFast:
-			DstToken = TEXT("ACosFastDeriv(") + SrcToken + TEXT(")");
+			DstToken = TEXT("AcosFastDeriv(") + SrcToken + TEXT(")");
 			break;
 		case EFunc1::Atan:
-			DstToken = TEXT("ATanDeriv(") + SrcToken + TEXT(")");
+			DstToken = TEXT("AtanDeriv(") + SrcToken + TEXT(")");
 			break;
 		case EFunc1::AtanFast:
-			DstToken = TEXT("ATanFastDeriv(") + SrcToken + TEXT(")");
+			DstToken = TEXT("AtanFastDeriv(") + SrcToken + TEXT(")");
 			break;
 		case EFunc1::Sqrt:
 			DstToken = TEXT("SqrtDeriv(") + SrcToken + TEXT(")");
@@ -880,14 +880,15 @@ int32 FMaterialDerivativeAutogen::GenerateRotateScaleOffsetTexCoordsFunc(FHLSLMa
 
 	if (!bAllZeroDeriv && IsDerivativeValid(TexCoordDerivInfo.DerivativeStatus) && IsDerivativeValid(RotationScaleDerivInfo.DerivativeStatus) && IsDerivativeValid(OffsetDerivInfo.DerivativeStatus))
 	{
-		const uint32 ResultTypeIndex = GetDerivTypeIndex(ResultType);
 		FString TexCoordDeriv		= Translator.GetParameterCodeDeriv(TexCoord, CompiledPDV_Analytic);
 		FString RotationScaleDeriv	= Translator.GetParameterCodeDeriv(RotationScale, CompiledPDV_Analytic);
 		FString OffsetDeriv			= Translator.GetParameterCodeDeriv(Offset, CompiledPDV_Analytic);
 
-		TexCoordDeriv		= CoerceValueDeriv(TexCoordDeriv,		TexCoordDerivInfo.TypeIndex,		TexCoordDerivInfo.DerivativeStatus,			ResultTypeIndex);
+		TexCoordDeriv		= CoerceValueDeriv(TexCoordDeriv,		TexCoordDerivInfo.TypeIndex,		TexCoordDerivInfo.DerivativeStatus,			1);
+		RotationScaleDeriv	= CoerceValueDeriv(RotationScaleDeriv,	RotationScaleDerivInfo.TypeIndex,	RotationScaleDerivInfo.DerivativeStatus,	3);
+		OffsetDeriv			= CoerceValueDeriv(OffsetDeriv,			OffsetDerivInfo.TypeIndex,			OffsetDerivInfo.DerivativeStatus,			1);
 
-		FString AnalyticString = FString::Printf(TEXT("RotateScaleOffsetTexCoordsDeriv(%s, %s, SwizzleDeriv2(%s, xy))"), *TexCoordDeriv, *RotationScaleDeriv, *OffsetDeriv);
+		FString AnalyticString = FString::Printf(TEXT("RotateScaleOffsetTexCoordsDeriv(%s, %s, %s)"), *TexCoordDeriv, *RotationScaleDeriv, *OffsetDeriv);
 
 		check(NumResultComponents <= 4);
 		bRotateScaleOffsetTexCoords = true;
@@ -1724,4 +1725,4 @@ FString FMaterialDerivativeAutogen::ApplyUnMirror(FString Value, bool bUnMirrorU
 	return Value;	
 }
 
-#endif // WITH_EDITORONLY_DATA
\ No newline at end of file
+#endif // WITH_EDITORONLY_DATA
