World Partition - Fixed HLOD visibility in the editor (!PIE)
* HLOD actors would previously be hidden in the editor as soon as one actor used to generate them would be loaded... The issue is that if all underlying actors had been replaced for some reason, the HLOD would not be hidden.
* For now, hide HLODs at all time in !PIE... They can be shown from the HLOD outliner if needed.

diff --git a/Engine/Source/Runtime/Engine/Private/WorldPartition/HLOD/HLODActor.cpp b/Engine/Source/Runtime/Engine/Private/WorldPartition/HLOD/HLODActor.cpp
--- a/Engine/Source/Runtime/Engine/Private/WorldPartition/HLOD/HLODActor.cpp
+++ b/Engine/Source/Runtime/Engine/Private/WorldPartition/HLOD/HLODActor.cpp
@@ -20,7 +20,6 @@ AWorldPartitionHLOD::AWorldPartitionHLOD(const FObjectInitializer& ObjectInitial
 #if WITH_EDITORONLY_DATA
 	HLODHash = 0;
 	HLODBounds = FBox(EForceInit::ForceInit);
-	bListedInSceneOutliner = false;
 #endif
 }
 
@@ -64,107 +63,13 @@ void AWorldPartitionHLOD::PostLoad()
 
 #if WITH_EDITOR
 
-void AWorldPartitionHLOD::OnSubActorLoaded(const AActor& Actor)
-{
-	bool bIsAlreadyInSet = false;
-	LoadedSubActors.Add(&Actor, &bIsAlreadyInSet);
-
-	if (!bIsAlreadyInSet)
-	{
-		if (LoadedSubActors.Num() == 1)
-		{
-			UpdateVisibility();
-		}
-	}
-}
-
-void AWorldPartitionHLOD::OnSubActorUnloaded(const AActor& Actor)
-{
-	LoadedSubActors.Remove(&Actor);
-
-	// If HLOD has no more sub actors, ensure it is drawn at all time
-	if (LoadedSubActors.IsEmpty())
-	{
-		UpdateVisibility();
-	}
-}
-
-void AWorldPartitionHLOD::SetupLoadedSubActors()
-{
-	UWorld* World = GetWorld();
-	if (World && World->IsEditorWorld() && !World->IsPlayInEditor())
-	{
-		LoadedSubActors.Empty();
-
-		UWorldPartition* WorldPartition = World->GetWorldPartition();
-		check(WorldPartition);
-
-		for (const FGuid& SubActorGuid : SubActors)
-		{
-			const FWorldPartitionActorDesc* SubActorDesc = WorldPartition->GetActorDesc(SubActorGuid);
-			AActor* SubActor = SubActorDesc && SubActorDesc->IsLoaded() ? SubActorDesc->GetActor() : nullptr;
-			if (SubActor)
-			{
-				OnSubActorLoaded(*SubActor);
-			}
-		}
-
-		UpdateVisibility();
-	}
-}
-
-void AWorldPartitionHLOD::ResetLoadedSubActors()
-{
-	UWorld* World = GetWorld();
-	if (World && !World->IsGameWorld())
-	{
-		for (const TWeakObjectPtr<const AActor>& SubActor : LoadedSubActors)
-		{
-			if (!SubActor.IsValid())
-			{
-				continue;
-			}
-		}
-
-		LoadedSubActors.Empty();
-
-		UpdateVisibility();
-	}
-}
-
 void AWorldPartitionHLOD::PostRegisterAllComponents()
 {
 	Super::PostRegisterAllComponents();
 
-	if (!IsTemplate())
-	{
-		SetupLoadedSubActors();
-	}
-
-#if WITH_EDITOR
+	SetIsTemporarilyHiddenInEditor(true);
 	bListedInSceneOutliner = true;
 	SetFolderPath(*FString::Printf(TEXT("HLOD/HLOD%d"), GetLODLevel()));
-#endif
-}
-
-void AWorldPartitionHLOD::PostUnregisterAllComponents()
-{
-	if (!IsTemplate())
-	{
-		ResetLoadedSubActors();
-	}
-
-	Super::PostUnregisterAllComponents();
-}
-
-void AWorldPartitionHLOD::UpdateVisibility()
-{
-	SetIsTemporarilyHiddenInEditor(HasLoadedSubActors());
-}
-
-bool AWorldPartitionHLOD::HasLoadedSubActors() const
-{
-	return !LoadedSubActors.IsEmpty();
 }
 
 EActorGridPlacement AWorldPartitionHLOD::GetDefaultGridPlacement() const
@@ -213,11 +118,7 @@ void AWorldPartitionHLOD::SetHLODPrimitives(const TArray<UPrimitiveComponent*>&
 
 void AWorldPartitionHLOD::SetSubActors(const TArray<FGuid>& InSubActors)
 {
-	ResetLoadedSubActors();
-
 	SubActors = InSubActors;
-
-	SetupLoadedSubActors();
 }
 
 const TArray<FGuid>& AWorldPartitionHLOD::GetSubActors() const
diff --git a/Engine/Source/Runtime/Engine/Private/WorldPartition/HLOD/HLODBuilder.cpp b/Engine/Source/Runtime/Engine/Private/WorldPartition/HLOD/HLODBuilder.cpp
--- a/Engine/Source/Runtime/Engine/Private/WorldPartition/HLOD/HLODBuilder.cpp
+++ b/Engine/Source/Runtime/Engine/Private/WorldPartition/HLOD/HLODBuilder.cpp
@@ -442,6 +442,7 @@ TArray<AWorldPartitionHLOD*> FHLODBuilderUtilities::CreateHLODActors(FHLODCreati
 			HLODActor = InCreationParams.WorldPartition->GetWorld()->SpawnActor<AWorldPartitionHLOD>(SpawnParams);
 
 			HLODActor->SetActorLabel(FString::Printf(TEXT("HLOD%d_%s"), InCreationParams.HLODLevel, *InCreationParams.CellName.ToString()));
+			HLODActor->SetFolderPath(*FString::Printf(TEXT("HLOD/HLOD%d"), InCreationParams.HLODLevel));
 			HLODActor->SetCellName(InCreationParams.CellName);
 			HLODActor->SetSubActorsHLODLayer(HLODLayer);
 			HLODActor->SetGridIndices(InCreationParams.GridIndexX, InCreationParams.GridIndexY, InCreationParams.GridIndexZ);
diff --git a/Engine/Source/Runtime/Engine/Private/WorldPartition/HLOD/HLODSubsystem.cpp b/Engine/Source/Runtime/Engine/Private/WorldPartition/HLOD/HLODSubsystem.cpp
--- a/Engine/Source/Runtime/Engine/Private/WorldPartition/HLOD/HLODSubsystem.cpp
+++ b/Engine/Source/Runtime/Engine/Private/WorldPartition/HLOD/HLODSubsystem.cpp
@@ -43,14 +43,6 @@ void UHLODSubsystem::Initialize(FSubsystemCollectionBase& Collection)
 {
 	Super::Initialize(Collection);
 
-#if WITH_EDITOR
-	if (!IsRunningCommandlet() && !GetWorld()->IsGameWorld())
-	{
-		GetWorld()->PersistentLevel->OnLoadedActorAddedToLevelEvent.AddUObject(this, &UHLODSubsystem::OnActorLoaded);
-		GetWorld()->PersistentLevel->OnLoadedActorRemovedFromLevelEvent.AddUObject(this, &UHLODSubsystem::OnActorUnloaded);
-	}
-#endif
-
 	if (GetWorld()->IsGameWorld())
 	{
 		TSet<const UWorldPartitionRuntimeCell*> StreamingCells;
@@ -64,54 +56,6 @@ void UHLODSubsystem::Initialize(FSubsystemCollectionBase& Collection)
 	}
 }
 
-#if WITH_EDITOR
-void UHLODSubsystem::OnActorLoaded(AActor& Actor)
-{
-	// If loading an HLOD actor, keep a map of Actor -> HLODActors
-	if (AWorldPartitionHLOD* HLODActor = Cast<AWorldPartitionHLOD>(&Actor))
-	{
-		for (const FGuid& SubActorGuid : HLODActor->GetSubActors())
-		{
-			// Keep track of unloaded subactors
-			check(!ActorsToHLOD.Contains(SubActorGuid));
-			ActorsToHLOD.Emplace(SubActorGuid, HLODActor);
-		}
-	}
-
-	// If HLOD for this actor is already loaded, notify the HLOD that it should be hidden
-	AWorldPartitionHLOD** HLODActorPtr = ActorsToHLOD.Find(Actor.GetActorGuid());
-	if (HLODActorPtr)
-	{
-		(*HLODActorPtr)->OnSubActorLoaded(Actor);
-	}
-}
-
-void UHLODSubsystem::OnActorUnloaded(AActor& Actor)
-{
-	UWorld* World = GetWorld();
-	UWorldPartition* WorldPartition = World->GetWorldPartition();
-
-	// If bound to an HLOD, clear reference
-	{
-		AWorldPartitionHLOD** HLODActorPtr = ActorsToHLOD.Find(Actor.GetActorGuid());
-		if (HLODActorPtr)
-		{
-			(*HLODActorPtr)->OnSubActorUnloaded(Actor);
-		}
-	}
-
-	// If unloading an HLOD actor, clear map of Actor -> HLODActors
-	if (AWorldPartitionHLOD* HLODActor = Cast<AWorldPartitionHLOD>(&Actor))
-	{
-		for (const FGuid& SubActorGuid : HLODActor->GetSubActors())
-		{
-			int32 NumRemoved = ActorsToHLOD.Remove(SubActorGuid);
-			check(NumRemoved == 1);
-		}
-	}
-}
-#endif
-
 void UHLODSubsystem::RegisterHLODActor(AWorldPartitionHLOD* InWorldPartitionHLOD)
 {
 	TRACE_CPUPROFILER_EVENT_SCOPE(UHLODSubsystem::RegisterHLODActor);
diff --git a/Engine/Source/Runtime/Engine/Public/WorldPartition/HLOD/HLODActor.h b/Engine/Source/Runtime/Engine/Public/WorldPartition/HLOD/HLODActor.h
--- a/Engine/Source/Runtime/Engine/Public/WorldPartition/HLOD/HLODActor.h
+++ b/Engine/Source/Runtime/Engine/Public/WorldPartition/HLOD/HLODActor.h
@@ -23,9 +23,6 @@ public:
 	inline uint32 GetLODLevel() const { return LODLevel; }
 
 #if WITH_EDITOR
-	void OnSubActorLoaded(const AActor& Actor);
-	void OnSubActorUnloaded(const AActor& Actor);
-
 	void SetHLODPrimitives(const TArray<UPrimitiveComponent*>& InHLODPrimitives);
 
 	void SetSubActors(const TArray<FGuid>& InSubActors);
@@ -68,7 +65,6 @@ protected:
 	virtual void EndPlay(const EEndPlayReason::Type EndPlayReason) override;
 #if WITH_EDITOR
 	virtual void PostRegisterAllComponents() override;
-	virtual void PostUnregisterAllComponents() override;
 	virtual EActorGridPlacement GetDefaultGridPlacement() const override;
 	virtual TUniquePtr<class FWorldPartitionActorDesc> CreateClassActorDesc() const override;
 
@@ -79,14 +75,6 @@ protected:
 
 	UPrimitiveComponent* GetHLODComponent();
 
-#if WITH_EDITOR
-	void UpdateVisibility();
-	bool HasLoadedSubActors() const;
-
-	void ResetLoadedSubActors();
-	void SetupLoadedSubActors();
-#endif // WITH_EDITOR
-
 private:
 #if WITH_EDITORONLY_DATA
 	UPROPERTY()
@@ -104,8 +92,6 @@ private:
 	UPROPERTY()
 	int64 GridIndexZ;
 
-	TSet<TWeakObjectPtr<const AActor>> LoadedSubActors;
-
 	UPROPERTY()
 	FBox HLODBounds;
 
diff --git a/Engine/Source/Runtime/Engine/Public/WorldPartition/HLOD/HLODSubsystem.h b/Engine/Source/Runtime/Engine/Public/WorldPartition/HLOD/HLODSubsystem.h
--- a/Engine/Source/Runtime/Engine/Public/WorldPartition/HLOD/HLODSubsystem.h
+++ b/Engine/Source/Runtime/Engine/Public/WorldPartition/HLOD/HLODSubsystem.h
@@ -36,10 +36,6 @@ public:
 	void OnCellHidden(const UWorldPartitionRuntimeCell* InCell);
 	
 private:
-#if WITH_EDITOR
-	void OnActorLoaded(AActor& Actor);
-	void OnActorUnloaded(AActor& Actor);
-#endif
 
 private:
 	struct FCellHLODMapping
@@ -53,8 +49,4 @@ private:
 	};
 
 	TMap<FName, FCellHLODMapping> CellsHLODMapping;
-
-#if WITH_EDITORONLY_DATA
-	TMap<FGuid, AWorldPartitionHLOD*> ActorsToHLOD;
-#endif
 };
