Skip bulk data serialization in FixUpForPIE
  - All static mesh bulk data was serialized for nothing including all HLODs
  - 1m19s to 22.5s to PIE in P_Construct_WP

diff --git a/Engine/Source/Runtime/Engine/Private/Level.cpp b/Engine/Source/Runtime/Engine/Private/Level.cpp
--- a/Engine/Source/Runtime/Engine/Private/Level.cpp
+++ b/Engine/Source/Runtime/Engine/Private/Level.cpp
@@ -2664,6 +2664,7 @@ void ULevel::FixupForPIE(int32 InPIEInstanceID, TFunctionRef<void(int32, FSoftOb
 			, CustomFixupFunction(InCustomFixupFunction)
 		{
 			this->SetIsSaving(true);
+			this->ArShouldSkipBulkData = true;
 		}
 
 		FArchive& operator<<(FSoftObjectPath& Value)
