Temp fix for invalid class in editor standalone build.

diff --git a/Engine/Source/Runtime/Engine/Private/WorldPartition/ActorDescContainer.cpp b/Engine/Source/Runtime/Engine/Private/WorldPartition/ActorDescContainer.cpp
--- a/Engine/Source/Runtime/Engine/Private/WorldPartition/ActorDescContainer.cpp
+++ b/Engine/Source/Runtime/Engine/Private/WorldPartition/ActorDescContainer.cpp
@@ -124,7 +124,13 @@ TUniquePtr<FWorldPartitionActorDesc> UActorDescContainer::GetActorDescriptor(con
 
 			NewActorDesc->Init(this, ActorDescInitData);
 			
-			check(bIsValidClass || (NewActorDesc->GetActorIsEditorOnly() && IsRunningGame()));
+			//check(bIsValidClass || (NewActorDesc->GetActorIsEditorOnly() && IsRunningGame()));
+
+			if (!bIsValidClass)
+			{
+				return nullptr;
+			}
+
 			return NewActorDesc;
 		}
 	}
