UStruct property default initializations

//UE5/Main - UE.EditorAutomation(RunTest=Editor) Win64 - Uninitialized script struct members found AnimGraphAttributes.h and other anim files

//UE5/Main - UE.EditorAutomation(RunTest=Editor) Win64 - Uninitialized script struct members found DataSourceFiltering.h

Fix BlendSpace errors in ShooterGame

diff --git a/Engine/Plugins/Developer/TraceSourceFiltering/Source/SourceFilteringCore/Public/DataSourceFiltering.h b/Engine/Plugins/Developer/TraceSourceFiltering/Source/SourceFilteringCore/Public/DataSourceFiltering.h
--- a/Engine/Plugins/Developer/TraceSourceFiltering/Source/SourceFilteringCore/Public/DataSourceFiltering.h
+++ b/Engine/Plugins/Developer/TraceSourceFiltering/Source/SourceFilteringCore/Public/DataSourceFiltering.h
@@ -53,7 +53,7 @@ public:
 
 	/** Flag as to whether or not any derived classes from ActorClass should also be considered when filtering */
 	UPROPERTY(EditAnywhere, Category = TraceSourceFiltering)
-	bool bIncludeDerivedClasses;
+	bool bIncludeDerivedClasses = false;
 
 	inline bool operator==(const FActorClassFilter Rhs)
 	{
diff --git a/Engine/Source/Editor/AnimGraph/Public/AnimGraphAttributes.h b/Engine/Source/Editor/AnimGraph/Public/AnimGraphAttributes.h
--- a/Engine/Source/Editor/AnimGraph/Public/AnimGraphAttributes.h
+++ b/Engine/Source/Editor/AnimGraph/Public/AnimGraphAttributes.h
@@ -67,15 +67,15 @@ struct FAnimGraphAttributeDesc
 
 	// How the attribute is displayed in the graph. This overrides any settings per node type.
 	UPROPERTY(config)
-	EAnimGraphAttributesDisplayMode DisplayMode;
+	EAnimGraphAttributesDisplayMode DisplayMode = EAnimGraphAttributesDisplayMode::Automatic;
 
 	// Blendability of the attribute
 	UPROPERTY(config)
-	EAnimGraphAttributeBlend Blend;
+	EAnimGraphAttributeBlend Blend = EAnimGraphAttributeBlend::Blendable;
 
 	// How to sort the attribute
 	UPROPERTY(config)
-	int32 SortOrder;
+	int32 SortOrder = INDEX_NONE;
 };
 
 UCLASS(Config = EditorPerProjectUserSettings)
diff --git a/Engine/Source/Runtime/AnimGraphRuntime/Public/AnimNodes/AnimNode_Sync.h b/Engine/Source/Runtime/AnimGraphRuntime/Public/AnimNodes/AnimNode_Sync.h
--- a/Engine/Source/Runtime/AnimGraphRuntime/Public/AnimNodes/AnimNode_Sync.h
+++ b/Engine/Source/Runtime/AnimGraphRuntime/Public/AnimNodes/AnimNode_Sync.h
@@ -28,7 +28,7 @@ private:
 
 	// The role this player can assume within the group
 	UPROPERTY(EditAnywhere, Category = Settings)
-	TEnumAsByte<EAnimGroupRole::Type> GroupRole;
+	TEnumAsByte<EAnimGroupRole::Type> GroupRole = EAnimGroupRole::CanBeLeader;
 
 private:
 	// FAnimNode_Base
diff --git a/Engine/Source/Runtime/Engine/Classes/Animation/AnimData/AnimDataModel.h b/Engine/Source/Runtime/Engine/Classes/Animation/AnimData/AnimDataModel.h
--- a/Engine/Source/Runtime/Engine/Classes/Animation/AnimData/AnimDataModel.h
+++ b/Engine/Source/Runtime/Engine/Classes/Animation/AnimData/AnimDataModel.h
@@ -23,7 +23,7 @@ struct ENGINE_API FBoneAnimationTrack
 
 	/** Index corresponding to the bone this track corresponds to within the target USkeleton */
 	UPROPERTY(VisibleAnywhere, BlueprintReadOnly, Category = "Model")
-	int32 BoneTreeIndex;
+	int32 BoneTreeIndex = INDEX_NONE;
 
 	/** Name of the bone this track corresponds to */
 	UPROPERTY(VisibleAnywhere, BlueprintReadOnly, Category = "Model")
diff --git a/Engine/Source/Runtime/Engine/Classes/Animation/AnimData/AnimDataNotifications.h b/Engine/Source/Runtime/Engine/Classes/Animation/AnimData/AnimDataNotifications.h
--- a/Engine/Source/Runtime/Engine/Classes/Animation/AnimData/AnimDataNotifications.h
+++ b/Engine/Source/Runtime/Engine/Classes/Animation/AnimData/AnimDataNotifications.h
@@ -114,7 +114,7 @@ struct FAnimationTrackAddedPayload : public FAnimationTrackPayload
 
 	/** Index of the track (bone) which was added */
 	UPROPERTY(BlueprintReadOnly, Category = Payload)
-	int32 TrackIndex;
+	int32 TrackIndex = INDEX_NONE;
 };
 
 USTRUCT(BlueprintType)
@@ -124,15 +124,15 @@ struct FSequenceLengthChangedPayload : public FEmptyPayload
 
 	/** Previous playable length for the Model */
 	UPROPERTY(BlueprintReadOnly, Category = Payload)
-	float PreviousLength;
+	float PreviousLength = 0.f;
 
 	/** Time at which the change in length has been made */
 	UPROPERTY(BlueprintReadOnly, Category = Payload)
-	float T0;
+	float T0 = 0.f;
 
 	/** Length of time which is inserted or removed starting at T0 */
 	UPROPERTY(BlueprintReadOnly, Category = Payload)
-	float T1;
+	float T1 = 0.f;
 };
 
 USTRUCT(BlueprintType)
@@ -166,11 +166,11 @@ struct FCurveScaledPayload : public FCurvePayload
 
 	/** Factor with which the curve was scaled */
 	UPROPERTY(BlueprintReadOnly, Category = Payload)
-	float Factor;
+	float Factor = 0.f;
 	
 	/** Time used as the origin when scaling the curve */
 	UPROPERTY(BlueprintReadOnly, Category = Payload)
-	float Origin;
+	float Origin = 0.f;
 };
 
 USTRUCT(BlueprintType)
@@ -190,7 +190,7 @@ struct FCurveFlagsChangedPayload : public FCurvePayload
 
 	/** Old flags mask for the curve */
 	UPROPERTY(BlueprintReadOnly, Category = Payload)
-	int32 OldFlags;
+	int32 OldFlags = 0;
 };
 
 USTRUCT(BlueprintType)
diff --git a/Engine/Source/Runtime/Engine/Classes/Animation/BlendSpace.h b/Engine/Source/Runtime/Engine/Classes/Animation/BlendSpace.h
--- a/Engine/Source/Runtime/Engine/Classes/Animation/BlendSpace.h
+++ b/Engine/Source/Runtime/Engine/Classes/Animation/BlendSpace.h
@@ -166,11 +166,11 @@ public:
 
 	/** Indices into the samples */
 	UPROPERTY()
-	int32 SampleIndices[NUM_VERTICES];
+	int32 SampleIndices[NUM_VERTICES] = { INDEX_NONE, INDEX_NONE };
 
 	/** The vertices are in the normalized space - i.e. in the range 0-1. */
 	UPROPERTY()
-	float Vertices[NUM_VERTICES];
+	float Vertices[NUM_VERTICES]  = { 0.f, 0.f };
 };
 
 USTRUCT()
@@ -181,24 +181,24 @@ struct FBlendSpaceTriangleEdgeInfo
 public:
 	/** Edge normal faces out */
 	UPROPERTY()
-	FVector2D Normal;
+	FVector2D Normal = FVector2D(0.f);
 
 	UPROPERTY()
-	int32 NeighbourTriangleIndex;
+	int32 NeighbourTriangleIndex = INDEX_NONE;
 
 	/**
 	* IF there is no neighbor, then (a) we're on the perimeter and (b) these will be the indices of
 	* triangles along the perimeter (next to the start and end of this edge, respectively) 
 	*/
 	UPROPERTY()
-	int32 AdjacentPerimeterTriangleIndices[2];
+	int32 AdjacentPerimeterTriangleIndices[2] = { INDEX_NONE, INDEX_NONE };
 
 	/**
 	 * The vertex index of the associated AdjacentPerimeterTriangle such that the perimeter edge is
 	 * from this vertex to the next.
 	 */
 	UPROPERTY()
-	int32 AdjacentPerimeterVertexIndices[2];
+	int32 AdjacentPerimeterVertexIndices[2]  = { INDEX_NONE, INDEX_NONE };
 };
 
 /**
@@ -218,12 +218,12 @@ public:
 
 	/** Indices into the samples */
 	UPROPERTY(EditAnywhere, Category = EditorElement)
-	int32 SampleIndices[NUM_VERTICES];
+	int32 SampleIndices[NUM_VERTICES] = { INDEX_NONE, INDEX_NONE, INDEX_NONE };
 
 	/** The vertices are in the normalized space - i.e. in the range 0-1. */
 	UPROPERTY(EditAnywhere, Category = EditorElement)
-	FVector2D Vertices[NUM_VERTICES];
-
+	FVector2D Vertices[NUM_VERTICES] = { FVector2D(0.f), FVector2D(0.f), FVector2D(0.f) };
+	
 	/** Info for the edge starting at the vertex index and going (anti-clockwise) to the next vertex */
 	UPROPERTY(EditAnywhere, Category = EditorElement)
 	FBlendSpaceTriangleEdgeInfo EdgeInfo[NUM_VERTICES];
diff --git a/Engine/Source/Runtime/Engine/Classes/Animation/BuiltInAttributeTypes.h b/Engine/Source/Runtime/Engine/Classes/Animation/BuiltInAttributeTypes.h
--- a/Engine/Source/Runtime/Engine/Classes/Animation/BuiltInAttributeTypes.h
+++ b/Engine/Source/Runtime/Engine/Classes/Animation/BuiltInAttributeTypes.h
@@ -19,7 +19,7 @@ struct FFloatAnimationAttribute
 	GENERATED_BODY()
 
 	UPROPERTY()
-	float Value;
+	float Value = 0.f;
 
 	void Accumulate(const FFloatAnimationAttribute& Attribute, float Weight, EAdditiveAnimationType AdditiveType)
 	{
@@ -52,7 +52,7 @@ struct FIntegerAnimationAttribute
 	GENERATED_BODY()
 
 	UPROPERTY()
-	int32 Value;
+	int32 Value = 0;
 
 	void Accumulate(const FIntegerAnimationAttribute& Attribute, float Weight, EAdditiveAnimationType AdditiveType)
 	{
diff --git a/Engine/Source/Runtime/Engine/Classes/Animation/PoseAsset.h b/Engine/Source/Runtime/Engine/Classes/Animation/PoseAsset.h
--- a/Engine/Source/Runtime/Engine/Classes/Animation/PoseAsset.h
+++ b/Engine/Source/Runtime/Engine/Classes/Animation/PoseAsset.h
@@ -60,10 +60,10 @@ struct FPoseAssetInfluence
 	GENERATED_USTRUCT_BODY()
 
 	UPROPERTY()
-	int32 PoseIndex;
+	int32 PoseIndex = INDEX_NONE;
 
 	UPROPERTY()
-	int32 BoneTransformIndex;
+	int32 BoneTransformIndex = INDEX_NONE;
 };
 
 USTRUCT()
diff --git a/Engine/Source/Runtime/Engine/Public/Animation/AnimInstanceProxy.h b/Engine/Source/Runtime/Engine/Public/Animation/AnimInstanceProxy.h
--- a/Engine/Source/Runtime/Engine/Public/Animation/AnimInstanceProxy.h
+++ b/Engine/Source/Runtime/Engine/Public/Animation/AnimInstanceProxy.h
@@ -104,7 +104,7 @@ struct FQueuedDrawDebugItem
 	FVector2D TextScale = FVector2D(0.f);
 
 	UPROPERTY(Transient)
-	TEnumAsByte<ESceneDepthPriorityGroup> DepthPriority;
+	TEnumAsByte<ESceneDepthPriorityGroup> DepthPriority = SDPG_World;
 };
 
 /** Proxy object passed around during animation tree update in lieu of a UAnimInstance */
