WorldPartition: make sure to load unloaded actors in ModifiedForPIE cell like before

diff --git a/Engine/Source/Runtime/Engine/Private/WorldPartition/WorldPartitionLevelStreamingDynamic.cpp b/Engine/Source/Runtime/Engine/Private/WorldPartition/WorldPartitionLevelStreamingDynamic.cpp
--- a/Engine/Source/Runtime/Engine/Private/WorldPartition/WorldPartitionLevelStreamingDynamic.cpp
+++ b/Engine/Source/Runtime/Engine/Private/WorldPartition/WorldPartitionLevelStreamingDynamic.cpp
@@ -236,6 +236,7 @@ bool UWorldPartitionLevelStreamingDynamic::IssueLoadRequests()
 	{
 		if (ChildPackage.ContainerID == 0)
 		{
+			bool bNeedDup = false;
 			if (bIsModifiedForPIE)
 			{
 				FString SubObjectName;
@@ -245,10 +246,12 @@ bool UWorldPartitionLevelStreamingDynamic::IssueLoadRequests()
 					if (AActor* ActorModifiedForPIE = World->PersistentLevel->ActorsModifiedForPIE.FindRef(*SubObjectName))
 					{
 						ActorsToDuplicate.Add(ActorModifiedForPIE);
+						bNeedDup = true;
 					}
 				}
 			}
-			else
+			
+			if(!bNeedDup)
 			{
 				ChildPackagesToLoad.Add(ChildPackage);
 			}
@@ -258,7 +261,6 @@ bool UWorldPartitionLevelStreamingDynamic::IssueLoadRequests()
 	// Duplicate unsaved actors
 	if (ActorsToDuplicate.Num())
 	{
-		check(ChildPackagesToLoad.Num() == 0);
 		// Create an actor container to make sure duplicated actors will share an outer to properly remap inter-actors references
 		UActorContainer* ActorContainer = NewObject<UActorContainer>(World->PersistentLevel);
 
@@ -304,7 +306,6 @@ bool UWorldPartitionLevelStreamingDynamic::IssueLoadRequests()
 	// Load saved actors
 	if (ChildPackagesToLoad.Num())
 	{
-		check(ActorsToDuplicate.Num() == 0);
 		FWorldPartitionLevelHelper::LoadActors(RuntimeLevel, ChildPackagesToLoad, PackageCache, [this, FinalizeLoading](bool bSucceeded)
 		{
 			if (!bSucceeded)
