//UE5/Release-Engine-Staging - Compile UnrealEditor Win64 - 'CustomAttributesGuid': undeclared identifier

Ensure that animation attributes are remapped when retargetting an animation
#fix added remapping functionality to AnimDataController for Animation Attributes
#misc added functionality to duplicate an attribute, fixed some typos, added automated tests for DuplicateAttribute

diff --git a/Engine/Source/Runtime/Engine/Classes/Animation/AnimData/AnimDataController.h b/Engine/Source/Runtime/Engine/Classes/Animation/AnimData/AnimDataController.h
--- a/Engine/Source/Runtime/Engine/Classes/Animation/AnimData/AnimDataController.h
+++ b/Engine/Source/Runtime/Engine/Classes/Animation/AnimData/AnimDataController.h
@@ -180,7 +180,7 @@ public:
 	* @param	BoneName			Bone name of the track which should be removed
 	* @param	bShouldTransact		Whether or not any undo-redo changes should be generated
 	*
-	* @return	Whether or not the removal was succesful
+	* @return	Whether or not the removal was successful
 	*/
 	UFUNCTION(BlueprintCallable, Category = AnimationData)
 	bool RemoveBoneTrack(FName BoneName, bool bShouldTransact = true);
@@ -203,7 +203,7 @@ public:
 	* @param	ScalingKeys			Array of keys for the scale component
 	* @param	bShouldTransact		Whether or not any undo-redo changes should be generated
 	*
-	* @return	Whether or not the keys were succesfully set
+	* @return	Whether or not the keys were successfully set
 	*/
 	UFUNCTION(BlueprintCallable, Category = AnimationData)
 	bool SetBoneTrackKeys(FName BoneName, const TArray<FVector>& PositionalKeys, const TArray<FQuat>& RotationalKeys, const TArray<FVector>& ScalingKeys, bool bShouldTransact = true);
@@ -215,7 +215,7 @@ public:
 	* @param	CurveFlags			Flags to be set for the curve
 	* @param	bShouldTransact		Whether or not any undo-redo changes should be generated
 	*
-	* @return	Whether or not the curve was succesfully added
+	* @return	Whether or not the curve was successfully added
 	*/
 	UFUNCTION(BlueprintCallable, Category = CurveData)
 	bool AddCurve(const FAnimationCurveIdentifier& CurveId, int32 CurveFlags = 0x00000004, bool bShouldTransact = true);
@@ -227,7 +227,7 @@ public:
 	* @param	NewCurveId			Identifier for the to-be-added curve
 	* @param	bShouldTransact		Whether or not any undo-redo changes should be generated
 	*
-	* @return	Whether or not the curve was succesfully duplicated
+	* @return	Whether or not the curve was successfully duplicated
 	*/
 	UFUNCTION(BlueprintCallable, Category = CurveData)
 	bool DuplicateCurve(const FAnimationCurveIdentifier& CopyCurveId, const FAnimationCurveIdentifier& NewCurveId, bool bShouldTransact = true);
@@ -239,7 +239,7 @@ public:
 	* @param	CopyCurveId			Identifier for the to-be-removed curve
 	* @param	bShouldTransact		Whether or not any undo-redo changes should be generated
 	*
-	* @return	Whether or not the curve was succesfully removed
+	* @return	Whether or not the curve was successfully removed
 	*/
 	UFUNCTION(BlueprintCallable, Category = CurveData)
 	bool RemoveCurve(const FAnimationCurveIdentifier& CurveId, bool bShouldTransact = true);
@@ -261,7 +261,7 @@ public:
 	* @param	bState				State of the flag to be, true=set/false=not set
 	* @param	bShouldTransact		Whether or not any undo-redo changes should be generated
 	*
-	* @return	Whether or not the flag state was succesfully set
+	* @return	Whether or not the flag state was successfully set
 	*/
 	UFUNCTION(BlueprintCallable, Category = CurveData)
 	bool SetCurveFlag(const FAnimationCurveIdentifier& CurveId, EAnimAssetCurveFlags Flag, bool bState = true, bool bShouldTransact = true);
@@ -273,7 +273,7 @@ public:
 	* @param	Flags				Flag mask with which the existings flags are to be replaced
 	* @param	bShouldTransact		Whether or not any undo-redo changes should be generated
 	*
-	* @return	Whether or not the flag mask was succesfully set
+	* @return	Whether or not the flag mask was successfully set
 	*/
 	UFUNCTION(BlueprintCallable, Category = CurveData)
 	bool SetCurveFlags(const FAnimationCurveIdentifier& CurveId, int32 Flags, bool bShouldTransact = true);
@@ -286,7 +286,7 @@ public:
 	* @param	TimeKeys			Time Keys with which the existing keys are to be replaced
 	* @param	bShouldTransact		Whether or not any undo-redo changes should be generated
 	*
-	* @return	Whether or not the transform curve keys were succesfully set
+	* @return	Whether or not the transform curve keys were successfully set
 	*/
 	UFUNCTION(BlueprintCallable, Category = CurveData)
 	bool SetTransformCurveKeys(const FAnimationCurveIdentifier& CurveId, const TArray<FTransform>& TransformValues, const TArray<float>& TimeKeys, bool bShouldTransact = true);
@@ -300,7 +300,7 @@ public:
 	* @param	Value				Value of the key to be set
 	* @param	bShouldTransact		Whether or not any undo-redo changes should be generated
 	*
-	* @return	Whether or not the transform curve key was succesfully set
+	* @return	Whether or not the transform curve key was successfully set
 	*/
 	UFUNCTION(BlueprintCallable, Category = CurveData)
 	bool SetTransformCurveKey(const FAnimationCurveIdentifier& CurveId, float Time, const FTransform& Value, bool bShouldTransact = true);
@@ -312,7 +312,7 @@ public:
 	* @param	Time				Time of the key to be removed
 	* @param	bShouldTransact		Whether or not any undo-redo changes should be generated
 	*
-	* @return	Whether or not the transform curve key was succesfully removed
+	* @return	Whether or not the transform curve key was successfully removed
 	*/
 	UFUNCTION(BlueprintCallable, Category = CurveData)
 	bool RemoveTransformCurveKey(const FAnimationCurveIdentifier& CurveId, float Time, bool bShouldTransact = true);
@@ -324,7 +324,7 @@ public:
 	* @param	NewCurveId			Time of the key to be removed
 	* @param	bShouldTransact		Whether or not any undo-redo changes should be generated
 	*
-	* @return	Whether or not the curve was succesfully renamed
+	* @return	Whether or not the curve was successfully renamed
 	*/
 	UFUNCTION(BlueprintCallable, Category = CurveData)
 	bool RenameCurve(const FAnimationCurveIdentifier& CurveToRenameId, const FAnimationCurveIdentifier& NewCurveId, bool bShouldTransact = true);
@@ -337,7 +337,7 @@ public:
 	* @param	Color				Color to which the curve is to be set
 	* @param	bShouldTransact		Whether or not any undo-redo changes should be generated
 	*
-	* @return	Whether or not the curve color was succesfully changed
+	* @return	Whether or not the curve color was successfully changed
 	*/
 	UFUNCTION(BlueprintCallable, Category = CurveData)
 	bool SetCurveColor(const FAnimationCurveIdentifier& CurveId, FLinearColor Color, bool bShouldTransact = true);
@@ -350,7 +350,7 @@ public:
 	* @param	Factor				Factor with which the curve is supposed to be scaled
 	* @param	bShouldTransact		Whether or not any undo-redo changes should be generated
 	*
-	* @return	Whether or not scaling the curve was succesful
+	* @return	Whether or not scaling the curve was successful
 	*/
 	UFUNCTION(BlueprintCallable, Category = CurveData)
 	bool ScaleCurve(const FAnimationCurveIdentifier& CurveId, float Origin, float Factor, bool bShouldTransact = true);
@@ -363,7 +363,7 @@ public:
 	* @param	Key					Key to be set
 	* @param	bShouldTransact		Whether or not any undo-redo changes should be generated
 	*
-	* @return	Whether or not the curve key was succesfully set
+	* @return	Whether or not the curve key was successfully set
 	*/
 	UFUNCTION(BlueprintCallable, Category = CurveData)
 	bool SetCurveKey(const FAnimationCurveIdentifier& CurveId, const FRichCurveKey& Key, bool bShouldTransact = true);
@@ -375,7 +375,7 @@ public:
 	* @param	Time				Time of the key to be removed
 	* @param	bShouldTransact		Whether or not any undo-redo changes should be generated
 	*
-	* @return	Whether or not the curve key was succesfully removed
+	* @return	Whether or not the curve key was successfully removed
 	*/
 	UFUNCTION(BlueprintCallable, Category = CurveData)
 	bool RemoveCurveKey(const FAnimationCurveIdentifier& CurveId, float Time, bool bShouldTransact = true);
@@ -387,7 +387,7 @@ public:
 	* @param	CurveKeys			Keys with which the existing keys are to be replaced
 	* @param	bShouldTransact		Whether or not any undo-redo changes should be generated
 	*
-	* @return	Whether or not replacing curve keys was succesful
+	* @return	Whether or not replacing curve keys was successful
 	*/
 	UFUNCTION(BlueprintCallable, Category = CurveData)
 	bool SetCurveKeys(const FAnimationCurveIdentifier& CurveId, const TArray<FRichCurveKey>& CurveKeys, bool bShouldTransact = true);
@@ -421,6 +421,15 @@ public:
 	*/
 	bool RemoveBoneTracksMissingFromSkeleton(const USkeleton* Skeleton, bool bShouldTransact = true);
 
+	/**
+	* Removes any bone attribute for which the name was not found in the provided skeleton. Broadcasts a EAnimDataModelNotifyType::AttributeRemoved for each attribute which was not found in the skeleton, wrapped within BracketOpened/BracketClosed notifies.
+	* Updates any bone attribute for which the bone index is different in the provided skeleton. Broadcasts a EAnimDataModelNotifyType::AttributeAdded and EAnimDataModelNotifyType::AttributeRemove for each attribute which was remapped
+	*
+	* @param	Skeleton			Skeleton to retrieve the bone information from
+	* @param	bShouldTransact		Whether or not any undo-redo changes should be generated
+	*/
+	void UpdateAttributesFromSkeleton(const USkeleton* Skeleton, bool bShouldTransact = true);
+
 	/**
 	* Broadcast a EAnimDataModelNotifyType::Populated notify.
 	*/
@@ -443,7 +452,7 @@ public:
 	* @param	AttributeIdentifier		Identifier for the to-be-added attribute
 	* @param	bShouldTransact			Whether or not any undo-redo changes should be generated
 	*
-	* @return	Whether or not the attribute was succesfully added
+	* @return	Whether or not the attribute was successfully added
 	*/
 	UFUNCTION(BlueprintCallable, Category = CurveData)
 	bool AddAttribute(const FAnimationAttributeIdentifier& AttributeIdentifier, bool bShouldTransact = true);
@@ -454,7 +463,7 @@ public:
 	* @param	AttributeIdentifier		Identifier for the to-be-removed attribute
 	* @param	bShouldTransact			Whether or not any undo-redo changes should be generated
 	*
-	* @return	Whether or not the attribute was succesfully removed
+	* @return	Whether or not the attribute was successfully removed
 	*/
 	UFUNCTION(BlueprintCallable, Category = CurveData)
 	bool RemoveAttribute(const FAnimationAttributeIdentifier& AttributeIdentifier, bool bShouldTransact = true);
@@ -489,7 +498,7 @@ public:
 	* @param	KeyValue				Value (templated) of the to-be-set key
 	* @param	bShouldTransact			Whether or not any undo-redo changes should be generated
 	*
-	* @return	Whether or not the key was succesfully set
+	* @return	Whether or not the key was successfully set
 	*/
 	template<typename AttributeType>
 	bool SetTypedAttributeKey(const FAnimationAttributeIdentifier& AttributeIdentifier, float Time, AttributeType& KeyValue, bool bShouldTransact = true)
@@ -506,7 +515,7 @@ public:
 	* @param	KeyValue				Value of the to-be-set key
 	* @param	bShouldTransact			Whether or not any undo-redo changes should be generated
 	*
-	* @return	Whether or not the key was succesfully set
+	* @return	Whether or not the key was successfully set
 	*/
 	bool SetAttributeKey(const FAnimationAttributeIdentifier& AttributeIdentifier, float Time, const void* KeyValue, bool bShouldTransact = true)
 	{
@@ -521,7 +530,7 @@ public:
 	* @param	KeyValues				Values with which the existing key values are to be replaced
 	* @param	bShouldTransact			Whether or not any undo-redo changes should be generated
 	*
-	* @return	Whether or not replacing the attribute keys was succesful
+	* @return	Whether or not replacing the attribute keys was successful
 	*/
 	bool SetAttributeKeys(const FAnimationAttributeIdentifier& AttributeIdentifier, TArrayView<const float> Times, TArrayView<const void*> KeyValues, bool bShouldTransact = true)
 	{
@@ -536,7 +545,7 @@ public:
 	* @param	KeyValues				Values (templated) with which the existing key values are to be replaced
 	* @param	bShouldTransact			Whether or not any undo-redo changes should be generated
 	*
-	* @return	Whether or not replacing the attribute keys was succesful
+	* @return	Whether or not replacing the attribute keys was successful
 	*/
 	template<typename AttributeType>
 	bool SetTypedAttributeKeys(const FAnimationAttributeIdentifier& AttributeIdentifier, TArrayView<const float> Times, TArrayView<AttributeType> KeyValues, bool bShouldTransact = true)
@@ -557,10 +566,22 @@ public:
 	* @param	Time					Time of the key to be removed
 	* @param	bShouldTransact			Whether or not any undo-redo changes should be generated
 	*
-	* @return	Whether or not the attribute key was succesfully removed
+	* @return	Whether or not the attribute key was successfully removed
 	*/
 	UFUNCTION(BlueprintCallable, Category = CurveData)
 	bool RemoveAttributeKey(const FAnimationAttributeIdentifier& AttributeIdentifier, float Time, bool bShouldTransact = true);
+
+	/**
+	* Duplicated the attribute (curve) with the identifier. Broadcasts a EAnimDataModelNotifyType::AttributeAdded notify if successful.
+	*
+	* @param	AttributeIdentifier			Identifier for the to-be-duplicated attribute
+	* @param	NewAttributeIdentifier		Identifier for the to-be-added attribute
+	* @param	bShouldTransact				Whether or not any undo-redo changes should be generated
+	*
+	* @return	Whether or not the attribute was successfully duplicated
+	*/
+	UFUNCTION(BlueprintCallable, Category = CurveData)
+	bool DuplicateAttribute(const FAnimationAttributeIdentifier& AttributeIdentifier, const FAnimationAttributeIdentifier& NewAttributeIdentifier, bool bShouldTransact = true);
 protected:
 	/** Functionality used by FOpenBracketAction and FCloseBracketAction to broadcast their equivalent notifies without actually opening a bracket. */
 	void NotifyBracketOpen();
diff --git a/Engine/Source/Runtime/Engine/Classes/Animation/AnimData/AttributeIdentifier.h b/Engine/Source/Runtime/Engine/Classes/Animation/AnimData/AttributeIdentifier.h
--- a/Engine/Source/Runtime/Engine/Classes/Animation/AnimData/AttributeIdentifier.h
+++ b/Engine/Source/Runtime/Engine/Classes/Animation/AnimData/AttributeIdentifier.h
@@ -80,6 +80,7 @@ protected:
 	FSoftObjectPath ScriptStructPath;
 
 	friend class UAnimationAttributeIdentifierExtensions;
+	friend class UAnimDataController;
 };
 
 template<>
diff --git a/Engine/Source/Runtime/Engine/Private/Animation/AnimDataController.cpp b/Engine/Source/Runtime/Engine/Private/Animation/AnimDataController.cpp
--- a/Engine/Source/Runtime/Engine/Private/Animation/AnimDataController.cpp
+++ b/Engine/Source/Runtime/Engine/Private/Animation/AnimDataController.cpp
@@ -393,6 +393,53 @@ bool UAnimDataController::RemoveBoneTracksMissingFromSkeleton(const USkeleton* S
 	return false;
 }
 
+void UAnimDataController::UpdateAttributesFromSkeleton(const USkeleton* Skeleton, bool bShouldTransact)
+{
+	ValidateModel();
+	
+	if (Skeleton)
+	{
+		// Verifying that bone (names) for attribute data exist on new skeleton
+		TArray<FAnimationAttributeIdentifier> ToRemoveIdentifiers;
+		TArray<TPair<FAnimationAttributeIdentifier, int32>> ToDuplicateIdentifiers;
+			    
+		for (const FAnimatedBoneAttribute& Attribute : Model->AnimatedBoneAttributes)
+		{
+			const int32 NewBoneIndex = Skeleton->GetReferenceSkeleton().FindBoneIndex(Attribute.Identifier.GetBoneName());
+
+			if (NewBoneIndex == INDEX_NONE)
+			{
+				ToRemoveIdentifiers.Add(Attribute.Identifier);
+			}
+			else if(NewBoneIndex != Attribute.Identifier.GetBoneIndex())
+			{						
+				ToDuplicateIdentifiers.Add(TPair<FAnimationAttributeIdentifier, int32>(Attribute.Identifier, NewBoneIndex));
+			}
+		}
+
+		if (ToRemoveIdentifiers.Num() || ToDuplicateIdentifiers.Num())
+		{
+			CONDITIONAL_BRACKET(LOCTEXT("VerifyAttributeBoneNames", "Remapping Animation Attribute Data"));
+			for (const FAnimationAttributeIdentifier& Identifier : ToRemoveIdentifiers)
+			{
+				RemoveAttribute(Identifier);
+			}
+			
+			for (const TPair<FAnimationAttributeIdentifier, int32>& Pair : ToDuplicateIdentifiers)
+			{
+				FAnimationAttributeIdentifier NewIdentifier = Pair.Key;
+				NewIdentifier.BoneIndex = Pair.Value;
+
+				DuplicateAttribute(Pair.Key, NewIdentifier);
+				RemoveAttribute(Pair.Key);
+			}	
+		}
+	}
+	else
+	{
+		ReportError(LOCTEXT("InvalidSkeletonError", "Invalid USkeleton supplied"));
+	}
+}
 
 void UAnimDataController::ResetModel(bool bShouldTransact /*= true*/)
 {
@@ -1809,6 +1856,56 @@ bool UAnimDataController::RemoveAttributeKey(const FAnimationAttributeIdentifier
 	return false;
 }
 
+bool UAnimDataController::DuplicateAttribute(const FAnimationAttributeIdentifier& AttributeIdentifier, const FAnimationAttributeIdentifier& NewAttributeIdentifier, bool bShouldTransact)
+{
+	ValidateModel();
+
+	if (AttributeIdentifier.IsValid() && NewAttributeIdentifier.IsValid())
+	{
+		if (AttributeIdentifier.GetType() == NewAttributeIdentifier.GetType())
+		{
+			const FAnimatedBoneAttribute* ExistingAttributePtr = Model->FindAttribute(NewAttributeIdentifier);
+			if (ExistingAttributePtr == nullptr)
+			{
+				if(const FAnimatedBoneAttribute* AttributePtr = Model->FindAttribute(AttributeIdentifier))
+				{
+					CONDITIONAL_TRANSACTION(LOCTEXT("DuplicateAttribute", "Duplicating Animation Attribute"));
+
+					FAnimatedBoneAttribute& DuplicateAttribute = Model->AnimatedBoneAttributes.AddDefaulted_GetRef();
+					DuplicateAttribute.Identifier = NewAttributeIdentifier;
+					DuplicateAttribute.Curve = AttributePtr->Curve;
+
+					FAttributeAddedPayload Payload;
+					Payload.Identifier = NewAttributeIdentifier;
+					Model->Notify(EAnimDataModelNotifyType::AttributeAdded, Payload);
+
+					CONDITIONAL_ACTION(UE::Anim::FRemoveAtributeAction, NewAttributeIdentifier);
+					
+					return true;
+				}
+				else
+				{
+					ReportWarningf(LOCTEXT("AttributeNameToDuplicateNotFoundWarning", "Could not find attribute with name {0} and type {1} for duplication"), FText::FromName(AttributeIdentifier.GetName()), FText::FromString(AttributeIdentifier.GetType()->GetName()));
+				}
+			}
+			else
+			{
+				ReportWarningf(LOCTEXT("ExistingAttributeWarning", "Attribute with name {0} already exists"), FText::FromName(AttributeIdentifier.GetName()));
+			}
+		}
+		else
+		{
+			ReportWarningf(LOCTEXT("InvalidAttributeTypeWarning", "Attribute types do not match: {0} ({1})"), FText::FromString(AttributeIdentifier.GetType()->GetName()), FText::FromString(AttributeIdentifier.GetType()->GetName()));
+		}
+	}
+	else
+	{
+		ReportWarningf(LOCTEXT("InvalidAttributeIdentifierWarning", "Invalid attribute identifier(s) provided: {0} and/or {1}"), FText::FromName(AttributeIdentifier.GetName()), FText::FromName(AttributeIdentifier.GetName()));
+	}
+
+	return false;
+}
+
 #endif // WITH_EDITOR
 
 #undef LOCTEXT_NAMESPACE // "AnimDataController"
diff --git a/Engine/Source/Runtime/Engine/Private/Animation/AnimSequence.cpp b/Engine/Source/Runtime/Engine/Private/Animation/AnimSequence.cpp
--- a/Engine/Source/Runtime/Engine/Private/Animation/AnimSequence.cpp
+++ b/Engine/Source/Runtime/Engine/Private/Animation/AnimSequence.cpp
@@ -2851,6 +2851,8 @@ void UAnimSequence::RemapTracksToNewSkeleton( USkeleton* NewSkeleton, bool bConv
 	ValidateModel();
 
 	UAnimDataController::FScopedBracket ScopedBracket(Controller, LOCTEXT("RemapTrackToNewSkeleton_Bracket", "Remapping Bone Animation tracks to new Skeleton"));
+	
+	Controller->UpdateAttributesFromSkeleton(NewSkeleton);
 
 	// @Todo : currently additive will work fine since we don't bake anything except when we extract
 	// but in the future if we bake this can be problem
