Remove invalid log

diff --git a/Engine/Source/Runtime/Engine/Private/WorldPartition/WorldPartitionRuntimeSpatialHash.cpp b/Engine/Source/Runtime/Engine/Private/WorldPartition/WorldPartitionRuntimeSpatialHash.cpp
--- a/Engine/Source/Runtime/Engine/Private/WorldPartition/WorldPartitionRuntimeSpatialHash.cpp
+++ b/Engine/Source/Runtime/Engine/Private/WorldPartition/WorldPartitionRuntimeSpatialHash.cpp
@@ -817,10 +817,6 @@ bool UWorldPartitionRuntimeSpatialHash::CreateStreamingGrid(const FSpatialHashRu
 						{
 							World->PersistentLevel->ActorsModifiedForPIE.Add(Actor->GetFName(), Actor);
 						}
-						else
-						{
-							UE_LOG(LogWorldPartitionRuntimeSpatialHash, Error, TEXT(" Failed to add Actor : %s to ActorsModifiedForPIE"), *(ActorDescView.GetActorPath().ToString()));
-						}
 					}
 
 					UE_LOG(LogWorldPartitionRuntimeSpatialHash, Verbose, TEXT("  Actor : %s (%s) (Container %08x) Origin(%s)"), *(ActorDescView.GetActorPath().ToString()), *ActorDescView.GetGuid().ToString(EGuidFormats::UniqueObjectGuid), ActorInstance.ContainerInstance->ID, *FVector2D(ActorInstance.GetOrigin()).ToString());
