Improve BlendSpace Smoothing, adding options like SpringDamper, Ease In/Out and Exponential. Make existing blendspace assets use the grid for interpolation by default. Tidyup to make UI naming etc (more) consistent.

diff --git a/Engine/Source/Editor/Persona/Private/Customization/BlendSpaceDetails.cpp b/Engine/Source/Editor/Persona/Private/Customization/BlendSpaceDetails.cpp
--- a/Engine/Source/Editor/Persona/Private/Customization/BlendSpaceDetails.cpp
+++ b/Engine/Source/Editor/Persona/Private/Customization/BlendSpaceDetails.cpp
@@ -89,7 +89,41 @@ void FBlendSpaceDetails::CustomizeDetails(class IDetailLayoutBuilder& DetailBuil
 			if (AxisIndex < HideIndex)
 			{
 				Groups[AxisIndex]->AddPropertyRow(BlendParameter.ToSharedRef());
-				Groups[AxisIndex]->AddPropertyRow(InterpolationParameter.ToSharedRef());
+				// Don't add InterpolationParameter in the same way as BlendParameter, because it would add the
+				// elements as customizations that we can't subsequently customize. We will add them individually
+				// below.
+
+				TSharedPtr<IPropertyHandle> InterpolationTime = InterpolationParameter->GetChildHandle(
+                    GET_MEMBER_NAME_CHECKED(FInterpolationParameter, InterpolationTime));
+				TSharedPtr<IPropertyHandle> DampingRatio = InterpolationParameter->GetChildHandle(
+                    GET_MEMBER_NAME_CHECKED(FInterpolationParameter, DampingRatio));
+				TSharedPtr<IPropertyHandle> MaxSpeed = InterpolationParameter->GetChildHandle(
+                    GET_MEMBER_NAME_CHECKED(FInterpolationParameter, MaxSpeed));
+				TSharedPtr<IPropertyHandle> InterpolationType = InterpolationParameter->GetChildHandle(
+                    GET_MEMBER_NAME_CHECKED(FInterpolationParameter, InterpolationType));
+
+				// Custom edit condition for MaxSpeed
+				TAttribute<bool> MaxSpeedEditCondition = TAttribute<bool>::Create(
+                    [this, InterpolationTime, InterpolationType]()
+                    {
+                    uint8 IntType;
+                    InterpolationType->GetValue(IntType);
+                    EFilterInterpolationType Type = (EFilterInterpolationType) IntType;
+                    float Time;
+                    InterpolationTime->GetValue(Time);
+                    if (Time > 0.0f && (Type == EFilterInterpolationType::BSIT_SpringDamper ||
+                        Type == EFilterInterpolationType::BSIT_ExponentialDecay))
+                    {
+                        return true;
+                    }
+                    return false;
+                });
+
+				Groups[AxisIndex]->AddPropertyRow(InterpolationTime.ToSharedRef());
+				Groups[AxisIndex]->AddPropertyRow(InterpolationType.ToSharedRef());
+				Groups[AxisIndex]->AddPropertyRow(DampingRatio.ToSharedRef());
+				IDetailPropertyRow& MaxSpeedProperty = Groups[AxisIndex]->AddPropertyRow(MaxSpeed.ToSharedRef());
+				MaxSpeedProperty.EditCondition(MaxSpeedEditCondition, nullptr);
 			}
 			else
 			{
diff --git a/Engine/Source/Editor/Persona/Private/SAnimationBlendSpaceGridWidget.cpp b/Engine/Source/Editor/Persona/Private/SAnimationBlendSpaceGridWidget.cpp
--- a/Engine/Source/Editor/Persona/Private/SAnimationBlendSpaceGridWidget.cpp
+++ b/Engine/Source/Editor/Persona/Private/SAnimationBlendSpaceGridWidget.cpp
@@ -447,7 +447,7 @@ void SBlendSpaceGridWidget::PaintSampleKeys(
 			FVector2D GridPosition = SampleValueToScreenPosition(PreviewFilteredPosition) - (PreviewSize * .5f);
 			FSlateDrawElement::MakeBox(
 				OutDrawElements, FilteredPositionLayer, AllottedGeometry.ToPaintGeometry(GridPosition, PreviewSize), 
-				PreviewBrush, ESlateDrawEffect::None, PreviewKeyColor.GetSpecifiedColor() * 0.5);
+				PreviewBrush, ESlateDrawEffect::None, PreviewKeyColor.GetSpecifiedColor() * 0.7);
 		}
 
 		// Always draw the preview position
diff --git a/Engine/Source/Runtime/Core/Public/Math/UnrealMathUtility.h b/Engine/Source/Runtime/Core/Public/Math/UnrealMathUtility.h
--- a/Engine/Source/Runtime/Core/Public/Math/UnrealMathUtility.h
+++ b/Engine/Source/Runtime/Core/Public/Math/UnrealMathUtility.h
@@ -400,6 +400,13 @@ public:
 		return A*A;
 	}
 
+	/** Cubes the value */
+	template< class T > 
+    static FORCEINLINE T Cube( const T A )
+	{
+		return A*A*A;
+	}
+
 	/** Clamps X to be between Min and Max, inclusive */
 	template< class T > 
 	UE_NODISCARD static FORCEINLINE T Clamp( const T X, const T Min, const T Max )
@@ -1140,22 +1147,24 @@ public:
 	 * 
 	 * Note that InSmoothingTime is the time lag when tracking constant motion (so if you tracked a predicted position 
 	 * TargetVel * InSmoothingTime ahead, you'd match the target position)
+	 *
+	 * An approximation is used that is accurate so long as InDeltaTime < 0.5 * InSmoothingTime
 	 * 
 	 * When starting from zero velocity, the maximum velocity is reached after time InSmoothingTime * 0.5
 	 *
 	 * @param  InOutValue      The value to be smoothed
 	 * @param  InOutValueRate  The rate of change of the value
-	 * @para   InTargetValue   The target to smooth towards
+	 * @param  InTargetValue   The target to smooth towards
 	 * @param  InDeltaTime     Time interval
 	 * @param  InSmoothingTime Timescale over which to smooth. Larger values result in more smoothed behaviour. Can be zero.
 	 */
 	template< class T >
 	static void CriticallyDampedSmoothing(
-		T& InOutValue,
-		T& InOutValueRate,
-		const T& InTargetValue,
-		float InDeltaTime,
-		float InSmoothingTime)
+		T&          InOutValue,
+		T&          InOutValueRate,
+		const T&    InTargetValue,
+		const float InDeltaTime,
+		const float InSmoothingTime)
 	{
 		if (InSmoothingTime > 0.0f)
 		{
@@ -1174,7 +1183,7 @@ public:
 			// SmoothingTime is 
 			float W = 2.0f / InSmoothingTime;
 			float A = W * InDeltaTime;
-			// Taylor expansion for Exp(-W * InDeltaTime)
+			// Taylor expansion for Exp(-W * InDeltaTime), adjusted to be more accurate for the expected delta time
 			float Exp = 1.0f / (1.0f + A + 0.48f * A * A + 0.235f * A * A * A);
 			T X0 = InOutValue - InTargetValue;
 			T B = (InOutValueRate + X0 * W) * InDeltaTime;
@@ -1193,6 +1202,106 @@ public:
 		}
 	}
 
+	/**
+	* Smooths a value using exponential damping towards a target. Works for any type that supports basic arithmetic operations.
+	* 
+	 * An approximation is used that is accurate so long as InDeltaTime < 0.5 * InSmoothingTime
+	* 
+	* @param  InOutValue      The value to be smoothed
+	* @param  InTargetValue   The target to smooth towards
+	* @param  InDeltaTime     Time interval
+	* @param  InSmoothingTime Timescale over which to smooth. Larger values result in more smoothed behaviour. Can be zero.
+	*/
+	template< class T >
+	static void ExponentialSmoothing(
+	    T&          InOutValue,
+	    const T&    InTargetValue,
+	    const float InDeltaTime,
+	    const float InSmoothingTime)
+	{
+		if (InSmoothingTime > 0.0f)
+		{
+			const float A = InDeltaTime / InSmoothingTime;
+			// Taylor expansion for Exp(-W * InDeltaTime), adjusted to be more accurate for the expected delta time
+			float Exp = 1.0f / (1.0f + A + 0.48f * A * A + 0.235f * A * A * A);
+			InOutValue = InTargetValue + (InOutValue - InTargetValue) * Exp;
+		}
+		else
+		{
+			InOutValue = InTargetValue;
+		}
+	}
+
+	/**
+	* Smooths a value using a spring damper towards a target. Works for any type that supports basic arithmetic operations.
+	* 
+	* The implementation is exact, so with approximations for Exp/Sin/Cos
+	* 
+	* @param  InOutValue      The value to be smoothed
+	* @param  InTargetValue   The target to smooth towards
+	* @param  InDeltaTime     Time interval
+	* @param  InSmoothingTime Timescale over which to smooth. Larger values result in more smoothed behaviour. Can be zero.
+	* @param  InDampingRatio  1 is critical damping. <1 results in under-damped motion (i.e. with overshoot), and >1 results in over-damped motion. 
+	*/
+	template< class T >
+	static void SpringDamperSmoothing(
+	    T&          InOutValue,
+	    T&          InOutValueRate,
+	    const T&    InTargetValue,
+	    const float InDeltaTime,
+	    const float InSmoothingTime,
+	    const float InDampingRatio)
+	{
+		if (InDeltaTime <= 0.0f)
+		{
+			return;
+		}
+		if (InSmoothingTime <= 0.0f)
+		{
+			InOutValueRate = (InTargetValue - InOutValue) / InDeltaTime;
+			InOutValue = InTargetValue;
+			return;
+		}
+		float W = 2.0f / InSmoothingTime;
+		T Err = InOutValue - InTargetValue;
+		// Handle the three cases separately
+		if (InDampingRatio > 1.0f) // Overdamped
+		{
+			const float WD = W * FMath::Sqrt(FMath::Square(InDampingRatio) - 1.0f);
+			const T C2 = -(InOutValueRate + (W * InDampingRatio - WD) * Err) / (2.0f * WD);
+			const T C1 = Err - C2;
+			const float A1 = (WD - InDampingRatio * W);
+			const float A2 = -(WD + InDampingRatio * W);
+			const float E1 = FMath::Exp(A1 * InDeltaTime);
+			const float E2 = FMath::Exp(A2 * InDeltaTime);
+			InOutValue = E1 * C1 + E2 * C2;
+			InOutValueRate = E1 * C1 * A1 + E2 * C2 * A2;
+			InOutValue += InTargetValue;
+		}
+		else if (InDampingRatio < 1.0f) // Underdamped
+		{
+			const float WD = W * FMath::Sqrt(1.0f - FMath::Square(InDampingRatio));
+			const T A = Err;
+			const T B = (InOutValueRate + Err * (InDampingRatio * W)) / WD;
+			const float C = FMath::Cos(WD * InDeltaTime);
+			const float S = FMath::Sin(WD * InDeltaTime);
+			const float E = FMath::Exp(-InDampingRatio * W * InDeltaTime);
+			InOutValue = E * (A * C + B * S);
+			InOutValueRate = -InOutValue * InDampingRatio * W;
+			InOutValueRate += E * (B * (WD * C) - A * (WD * S));
+			InOutValue += InTargetValue;
+		}
+		else // Critical damping
+		{
+			const T& C1 = Err;
+			T C2 = InOutValueRate + Err * W;
+			float E = FMath::Exp(-W * InDeltaTime);
+			InOutValue = (C1 + C2 * InDeltaTime) * E;
+			InOutValueRate = (C2 - C1 * W - C2 * (W * InDeltaTime)) * E;
+			InOutValue += InTargetValue;
+		}
+	}
+	
 	/**
 	 * Simple function to create a pulsating scalar value
 	 *
diff --git a/Engine/Source/Runtime/Core/Public/UObject/UE5MainStreamObjectVersion.h b/Engine/Source/Runtime/Core/Public/UObject/UE5MainStreamObjectVersion.h
--- a/Engine/Source/Runtime/Core/Public/UObject/UE5MainStreamObjectVersion.h
+++ b/Engine/Source/Runtime/Core/Public/UObject/UE5MainStreamObjectVersion.h
@@ -103,7 +103,10 @@ struct CORE_API FUE5MainStreamObjectVersion
 		// Use of triangulation at runtime in BlendSpace
 		BlendSpaceRuntimeTriangulation,
 
-		// -----<new versions can be added above this line>-------------------------------------------------
+		// Fix to the Cubic smoothing, plus introduction of new smoothing types
+        BlendSpaceSmoothingImprovements,
+
+        // -----<new versions can be added above this line>-------------------------------------------------
 		VersionPlusOne,
 		LatestVersion = VersionPlusOne - 1
 	};
diff --git a/Engine/Source/Runtime/Engine/Classes/Animation/BlendSpace.h b/Engine/Source/Runtime/Engine/Classes/Animation/BlendSpace.h
--- a/Engine/Source/Runtime/Engine/Classes/Animation/BlendSpace.h
+++ b/Engine/Source/Runtime/Engine/Classes/Animation/BlendSpace.h
@@ -38,13 +38,33 @@ struct FInterpolationParameter
 {
 	GENERATED_BODY()
 
-	/** Interpolation Time for input, when it gets input, it will use this time to interpolate to target, used for smoother interpolation. */
-	UPROPERTY(EditAnywhere, Category=Parameter)
+	/**
+	 * Smoothing Time used to move smoothly across the blendpsace from the current parameters to the target
+	 * parameters. The different Smoothing Types will treat this in different ways, but in general a value of
+	 * zero will disable all smoothing, and larger values will smooth more.
+	 */
+	UPROPERTY(EditAnywhere, DisplayName = "Smoothing Time", Category=Parameter, meta = (ClampMin = "0"))
 	float InterpolationTime = 0.f;
 
-	/** Type of interpolation used for filtering the input value to decide how to get to target. */
-	UPROPERTY(EditAnywhere, Category=Parameter)
-	TEnumAsByte<EFilterInterpolationType> InterpolationType = EFilterInterpolationType::BSIT_Average;
+	/**
+	 * Damping ratio - only used when the type is set to SpringDamper. A value of 1 will move quickly and
+	 * smoothly to the target, without overshooting. Values below 1 can be used to encourage some overshoot,
+	 * which can make pose transitions look more natural.
+	 */
+	UPROPERTY(EditAnywhere, Category=Parameter, meta = (ClampMin = "0", EditCondition = "InterpolationType == EFilterInterpolationType::BSIT_SpringDamper && InterpolationTime > 0"))
+	float DampingRatio = 1.f;
+
+	/**
+	 * Maximum speed, in real units. For example, if this axis is degrees then you could use a value of 90 to
+	 * limit the turn rate to 90 degrees per second. Only used when greater than zero and the type is
+	 * set to SpringDamper or Exponential.
+	 */
+	UPROPERTY(EditAnywhere, Category=Parameter, meta = (ClampMin = "0"))
+	float MaxSpeed = 0.f;
+
+	/** Type of smoothing used for filtering the input value to decide how to get to target. */
+	UPROPERTY(EditAnywhere, DisplayName = "Smoothing Type", Category=Parameter, meta = (EditCondition = "InterpolationTime > 0"))
+	TEnumAsByte<EFilterInterpolationType> InterpolationType = EFilterInterpolationType::BSIT_SpringDamper;
 };
 
 USTRUCT()
@@ -357,7 +377,7 @@ struct FPerBoneInterpolation
 	* 
 	* If set, the value overrides the overall Sample Weight Speed which will no longer affect this bone.
 	*/
-	UPROPERTY(EditAnywhere, Category=FPerBoneInterpolation, meta=(DisplayName="Sample Weight Speed"))
+	UPROPERTY(EditAnywhere, Category=FPerBoneInterpolation, meta=(DisplayName="Weight Speed"))
 	float InterpolationSpeedPerSec;
 
 	FPerBoneInterpolation()
@@ -461,9 +481,12 @@ public:
 	*/
 	ENGINE_API bool GetSamplesFromBlendInput(const FVector &BlendInput, TArray<FBlendSampleData> & OutSampleDataList, int32& InOutCachedTriangulationIndex, bool bCombineAnimations) const;
 
-	/** Initialize BlendSpace for runtime. It needs certain data to be reinitialized per instsance **/
+	/** Initialize BlendSpace filtering for runtime. **/
 	ENGINE_API void InitializeFilter(FBlendFilter* Filter) const;
 
+	/** Update BlendSpace filtering parameters - values that don't require a full initialization **/
+	ENGINE_API void UpdateFilterParams(FBlendFilter* Filter) const;
+
 	/** Returns the blend input after clamping and/or wrapping */
 	ENGINE_API FVector GetClampedAndWrappedBlendInput(const FVector& BlendInput) const;
 
@@ -620,20 +643,10 @@ public:
 	UPROPERTY()
 	bool bRotationBlendInMeshSpace;
 
-	/** Input interpolation parameter for all 3 axis, for each axis input, decide how you'd like to interpolate input to*/
+	/** Input Smoothing parameters for each input axis */
 	UPROPERTY(EditAnywhere, Category = InputInterpolation)
 	FInterpolationParameter	InterpolationParam[3];
 
-#if WITH_EDITORONLY_DATA
-	/** Preview Base pose for additive BlendSpace **/
-	UPROPERTY(EditAnywhere, Category = AdditiveSettings)
-	TObjectPtr<UAnimSequence> PreviewBasePose;
-#endif // WITH_EDITORONLY_DATA
-
-	/** This animation length changes based on current input (resulting in different blend time)**/
-	UPROPERTY(transient)
-	float AnimLength;
-
 	/**
 	* If greater than zero, this is the speed at which the sample weights are allowed to change.
 	* 
@@ -650,14 +663,24 @@ public:
 	* Smaller values mean slower adjustments of the sample weights, and thus more smoothing. However, a 
 	* value of zero disables this smoothing entirely.
 	*/
-	UPROPERTY(EditAnywhere, Category = SampleInterpolation, meta = (DisplayName = "Sample Weight Speed"))
+	UPROPERTY(EditAnywhere, Category = SampleSmoothing, meta = (DisplayName = "Weight Speed", ClampMin = "0"))
 	float TargetWeightInterpolationSpeedPerSec;
 
 	/**
 	 * If set then this eases in/out the sample weight adjustments, using the speed to determine how much smoothing to apply.
 	 */
-	UPROPERTY(EditAnywhere, Category = SampleInterpolation, meta = (DisplayName = "Sample Weight Speed Smoothing "))
+	UPROPERTY(EditAnywhere, Category = SampleSmoothing, meta = (DisplayName = "Smoothing"))
 	bool bTargetWeightInterpolationEaseInOut = true;
+	
+#if WITH_EDITORONLY_DATA
+	/** Preview Base pose for additive BlendSpace **/
+	UPROPERTY(EditAnywhere, Category = AdditiveSettings)
+	TObjectPtr<UAnimSequence> PreviewBasePose;
+#endif // WITH_EDITORONLY_DATA
+
+	/** This animation length changes based on current input (resulting in different blend time)**/
+	UPROPERTY(transient)
+	float AnimLength;
 
 	/** The current mode used by the BlendSpace to decide which animation notifies to fire. Valid options are:
 	- AllAnimations - All notify events will fire
@@ -668,7 +691,7 @@ public:
 	TEnumAsByte<ENotifyTriggerMode::Type> NotifyTriggerMode;
 
 	/** If true then interpolation is done via a grid at runtime. If false the interpolation uses the triangulation. */
-	UPROPERTY(EditAnywhere, Category = InputInterpolation)
+	UPROPERTY(EditAnywhere, Category = InputInterpolation, meta = (DisplayName="Use Grid"))
 	bool bInterpolateUsingGrid = false;
 
 	/** Preferred edge direction when the triangulation has to make an arbitrary choice */
@@ -679,11 +702,12 @@ protected:
 
 #if WITH_EDITORONLY_DATA
 	/**
-	* Per bone interpolation speed settings, which affect the specified bone and all its descendants in the skeleton.
-	* These act as overrides to the global interpolation speed, which means the global interpolation speed does not affect these bones.
-	* Note that they also override each other - so a per-bone setting on the chest will not affect the hand if there is a per-bone setting on the arm.
+	* Per bone sample smoothing settings, which affect the specified bone and all its descendants in the skeleton.
+	* These act as overrides to the global sample smoothing speed, which means the global sample smoothing speed does
+	* not affect these bones. Note that they also override each other - so a per-bone setting on the chest will not
+	* affect the hand if there is a per-bone setting on the arm.
 	*/
-	UPROPERTY(EditAnywhere, Category = SampleInterpolation, meta = (DisplayName="Per Bone Overrides"))
+	UPROPERTY(EditAnywhere, Category = SampleSmoothing, meta = (DisplayName="Per Bone Overrides"))
 	TArray<FPerBoneInterpolation> PerBoneBlend;
 #endif
 
@@ -714,7 +738,11 @@ protected:
 	UPROPERTY(EditAnywhere, Category = BlendParametersTest)
 	struct FBlendParameter BlendParameters[3];
 
-	/** If you have input interpolation, which axis to drive animation speed (scale) - i.e. for locomotion animation, speed axis will drive animation speed (thus scale) */
+	/**
+	 * If you have input smoothing, which axis to scale the animation speed. E.g. for locomotion animation, 
+	 * the speed axis will scale the animation speed in order to make up the difference between the target 
+	 * and the result of blending the samples.
+	 */
 	UPROPERTY(EditAnywhere, Category = InputInterpolation)
 	TEnumAsByte<EBlendSpaceAxis> AxisToScaleAnimation;
 
diff --git a/Engine/Source/Runtime/Engine/Classes/Animation/BlendSpace1D.h b/Engine/Source/Runtime/Engine/Classes/Animation/BlendSpace1D.h
--- a/Engine/Source/Runtime/Engine/Classes/Animation/BlendSpace1D.h
+++ b/Engine/Source/Runtime/Engine/Classes/Animation/BlendSpace1D.h
@@ -24,7 +24,11 @@ public:
 	bool bDisplayEditorVertically_DEPRECATED;
 #endif
 
-	/** Drive animation speed by blend input position */
+	/**
+	 * If you have input smoothing, whether to scale the animation speed. E.g. for locomotion animation, 
+	 * the speed axis will scale the animation speed in order to make up the difference between the target 
+	 * and the result of blending the samples.
+	 */
 	UPROPERTY(EditAnywhere, Category = InputInterpolation)
 	bool bScaleAnimation;
 
diff --git a/Engine/Source/Runtime/Engine/Classes/Engine/EngineTypes.h b/Engine/Source/Runtime/Engine/Classes/Engine/EngineTypes.h
--- a/Engine/Source/Runtime/Engine/Classes/Engine/EngineTypes.h
+++ b/Engine/Source/Runtime/Engine/Classes/Engine/EngineTypes.h
@@ -963,9 +963,12 @@ enum ECollisionResponse
 UENUM()
 enum EFilterInterpolationType
 {
-	BSIT_Average UMETA(DisplayName = "Averaged Interpolation"),
-	BSIT_Linear UMETA(DisplayName = "Linear Interpolation"),
-	BSIT_Cubic UMETA(DisplayName = "Cubic Interpolation"),
+	BSIT_Average UMETA(DisplayName = "Averaged"),
+	BSIT_Linear UMETA(DisplayName = "Linear"),
+	BSIT_Cubic UMETA(DisplayName = "Cubic"),
+	BSIT_EaseInOut UMETA(DisplayName = "Ease In/Out"),
+	BSIT_ExponentialDecay UMETA(DisplayName = "Exponential"),
+	BSIT_SpringDamper UMETA(DisplayName = "Spring Damper"),
 	BSIT_MAX
 };
 
diff --git a/Engine/Source/Runtime/Engine/Private/Animation/AnimInterpFilter.cpp b/Engine/Source/Runtime/Engine/Private/Animation/AnimInterpFilter.cpp
--- a/Engine/Source/Runtime/Engine/Private/Animation/AnimInterpFilter.cpp
+++ b/Engine/Source/Runtime/Engine/Private/Animation/AnimInterpFilter.cpp
@@ -1,29 +1,26 @@
 // Copyright Epic Games, Inc. All Rights Reserved.
 
-/*=============================================================================
-	InterpFilter.h
-=============================================================================*/ 
-
 #include "AnimInterpFilter.h"
 
-////////////////////////////
+//======================================================================================================================
 // FFIRFilter
-
+//======================================================================================================================
 float FFIRFilter::GetInterpolationCoefficient (EFilterInterpolationType InterpolationType, int32 CoefficientIndex) const
 {
-	float Step = GetStep();
+	const float Step = GetStep();
 
 	switch (InterpolationType)
 	{
-	case BSIT_Average:
-		return Step;	
-	case BSIT_Linear:
-		return Step*CoefficientIndex;
-	case BSIT_Cubic:
-		return Step*Step*Step*CoefficientIndex;
+		case BSIT_Average:
+			return Step;	
+		case BSIT_Linear:
+			return Step*CoefficientIndex;
+		case BSIT_Cubic:
+			return Step*Step*Step*CoefficientIndex;
+		default:
+			// Note that BSIT_EaseInOut is not supported
+			return 0.f;
 	}
-
-	return 0.f;
 }
 
 void FFIRFilter::CalculateCoefficient(EFilterInterpolationType InterpolationType)
@@ -87,6 +84,9 @@ float FFIRFilter::CalculateFilteredOutput() const
 	return Output;
 }
 
+//======================================================================================================================
+// FFIRFilterTimeBased
+//======================================================================================================================
 int32 FFIRFilterTimeBased::GetSafeCurrentStackIndex()
 {
 	// if valid range
@@ -119,7 +119,7 @@ int32 FFIRFilterTimeBased::GetSafeCurrentStackIndex()
 
 	// if current one isn't available anymore 
 	// that means we need more stack
-	int32 NewIndex = FilterWindow.Num();
+	const int32 NewIndex = FilterWindow.Num();
 	FilterWindow.AddZeroed(5);
 	return NewIndex;
 }
@@ -128,12 +128,12 @@ void FFIRFilterTimeBased::RefreshValidFilters()
 {
 	NumValidFilter = 0;
 
-	if ( TimeDuration > 0.f )
+	if ( WindowDuration > 0.f )
 	{
 		// run validation test
 		for (int32 I=0; I<FilterWindow.Num(); ++I)
 		{
-			FilterWindow[I].CheckValidation(CurrentTime, TimeDuration);
+			FilterWindow[I].CheckValidation(CurrentTime, WindowDuration);
 			if ( FilterWindow[I].IsValid() )
 			{
 				++NumValidFilter;
@@ -142,31 +142,97 @@ void FFIRFilterTimeBased::RefreshValidFilters()
 	}
 }
 
-float FFIRFilterTimeBased::GetFilteredData(float Input, float DeltaTime)
+float FFIRFilterTimeBased::UpdateAndGetFilteredData(float Input, float DeltaTime)
 {
 	if (DeltaTime > KINDA_SMALL_NUMBER)
 	{
 		float Result;
 		CurrentTime += DeltaTime;
 
-		if (IsValid())
+		switch (InterpolationType)
 		{
-			RefreshValidFilters();
-
-			CurrentStackIndex = GetSafeCurrentStackIndex();
-			FilterWindow[CurrentStackIndex].SetInput(Input, CurrentTime);
-			Result = CalculateFilteredOutput();
-			CurrentStackIndex = CurrentStackIndex + 1;
-			if (CurrentStackIndex > FilterWindow.Num() - 1)
+			case BSIT_ExponentialDecay:
 			{
-				CurrentStackIndex = 0;
+				if (FilterWindow.Num() != 1)
+				{
+					FilterWindow.Empty(1);
+					FilterWindow.Push(FFilterData(Input, 0.0f));
+				}
+				const float OrigValue = FilterWindow[0].Input;
+				FMath::ExponentialSmoothing(FilterWindow[0].Input, Input, DeltaTime, WindowDuration / EULERS_NUMBER);
+				if (MaxSpeed > 0.0f)
+				{
+					// Clamp the speed
+					FilterWindow[0].Input = FMath::Clamp(FilterWindow[0].Input, OrigValue - MaxSpeed * DeltaTime,
+                                                         OrigValue + MaxSpeed * DeltaTime); 
+				}
+				Result = FilterWindow[0].Input;
 			}
+			break;
+			case BSIT_SpringDamper:
+			{
+				if (FilterWindow.Num() != 2)
+				{
+					FilterWindow.Empty(2);
+					// [0] element is the value, [1] element is the rate
+					FilterWindow.Push(FFilterData(Input, 0.0f));
+					FilterWindow.Push(FFilterData(0.0f, 0.0f));
+				}
+				const float OrigValue = FilterWindow[0].Input;
+				FMath::SpringDamperSmoothing(FilterWindow[0].Input, FilterWindow[1].Input, Input, DeltaTime,
+				                      WindowDuration / EULERS_NUMBER, DampingRatio);
+				if (MaxSpeed > 0.0f)
+				{
+					// Clamp the speed
+					FilterWindow[0].Input = FMath::Clamp(FilterWindow[0].Input, OrigValue - MaxSpeed * DeltaTime,
+                                                         OrigValue + MaxSpeed * DeltaTime); 
+					FilterWindow[1].Input = FMath::Clamp(FilterWindow[1].Input, -MaxSpeed, MaxSpeed);
+				}
+				if (bClamp)
+				{
+					// Clamp the value
+					if (FilterWindow[0].Input > MaxValue)
+					{
+						FilterWindow[0].Input = MaxValue;
+						if (FilterWindow[1].Input > 0.0f)
+						{
+							FilterWindow[1].Input = 0.0f;
+						}
+					}
+					if (FilterWindow[0].Input < MinValue)
+					{
+						FilterWindow[0].Input = MinValue;
+						if (FilterWindow[1].Input < 0.0f)
+						{
+							FilterWindow[1].Input = 0.0f;
+						}
+					}
+				}
+				Result = FilterWindow[0].Input;
+			}
+			break;
+			default:
+			{
+				if (IsValid())
+				{
+					RefreshValidFilters();
+
+					CurrentStackIndex = GetSafeCurrentStackIndex();
+					FilterWindow[CurrentStackIndex].SetInput(Input, CurrentTime);
+					Result = CalculateFilteredOutput();
+					CurrentStackIndex = CurrentStackIndex + 1;
+					if (CurrentStackIndex > FilterWindow.Num() - 1)
+					{
+						CurrentStackIndex = 0;
+					}
+				}
+				else
+				{
+					Result = Input;
+				}
+			}
+			break;
 		}
-		else
-		{
-			Result = Input;
-		}
-
 		LastOutput = Result;
 		return Result;
 	}
@@ -174,21 +240,26 @@ float FFIRFilterTimeBased::GetFilteredData(float Input, float DeltaTime)
 	return LastOutput;
 }
 
-float FFIRFilterTimeBased::GetInterpolationCoefficient (FFilterData & Data)
+float FFIRFilterTimeBased::GetInterpolationCoefficient(const FFilterData& Data) const
 {
 	if (Data.IsValid())
 	{
-		float Diff = Data.Diff(CurrentTime);
-		if (Diff<=TimeDuration)
+		const float Diff = Data.Diff(CurrentTime);
+		if (Diff<=WindowDuration)
 		{
 			switch(InterpolationType)
 			{
 			case BSIT_Average:
 				return 1.f;
 			case BSIT_Linear:
-				return 1.f - Diff/TimeDuration;
+				return 1.f - Diff/WindowDuration;
 			case BSIT_Cubic:
-				return 1.f - Diff*Diff*Diff/TimeDuration;
+				return 1.f - FMath::Cube(Diff/WindowDuration);
+			case BSIT_EaseInOut:
+				// Quadratic that starts and ends at 0, and reaches 1 half way through the window
+				return 1.0f - 4.0f * FMath::Square(Diff/WindowDuration - 0.5f);
+			default:
+				break;
 			}
 		}
 	}
@@ -201,24 +272,15 @@ float FFIRFilterTimeBased::CalculateFilteredOutput()
 	check ( IsValid() );
 	float SumCoefficient = 0.f;
 	float SumInputs = 0.f;
-	int32 NumValidSamples = 0;
 	for (int32 I=0; I<FilterWindow.Num(); ++I)
 	{
-		float Coefficient = GetInterpolationCoefficient(FilterWindow[I]);
+		const float Coefficient = GetInterpolationCoefficient(FilterWindow[I]);
 		if (Coefficient > 0.f)
 		{
 			SumCoefficient += Coefficient;
 			SumInputs += Coefficient * FilterWindow[I].Input;
-			++NumValidSamples;
 		}
 	}
-
-	// if number of samples are 1, do not normalize
-// 	if (NumValidSamples == 1)
-// 	{
-// 		return SumInputs;
-// 	}
-
-	return (SumCoefficient>0.f)?SumInputs/SumCoefficient : 0.f;
+	return SumCoefficient > 0.f ? SumInputs/SumCoefficient : 0.f;
 }
 
diff --git a/Engine/Source/Runtime/Engine/Private/Animation/BlendSpace.cpp b/Engine/Source/Runtime/Engine/Private/Animation/BlendSpace.cpp
--- a/Engine/Source/Runtime/Engine/Private/Animation/BlendSpace.cpp
+++ b/Engine/Source/Runtime/Engine/Private/Animation/BlendSpace.cpp
@@ -80,12 +80,38 @@ void UBlendSpace::Serialize(FArchive& Ar)
 		}
 	}
 
-	if (Ar.IsLoading() && (Ar.CustomVer(FUE5MainStreamObjectVersion::GUID) < FUE5MainStreamObjectVersion::BlendSpaceRuntimeTriangulation))
+	if (Ar.IsLoading() && (Ar.CustomVer(FUE5MainStreamObjectVersion::GUID) <
+		FUE5MainStreamObjectVersion::BlendSpaceRuntimeTriangulation))
 	{
+		// Make old blend spaces use the grid
+		bInterpolateUsingGrid = true;
+		// Force the data to be updated
 		ResampleData();
-		// Preserve the constant blend for old blend spaces, but allow the ease in/out default for new ones.
+		// Preserve the constant sample weight speed for old blend spaces, but allow the ease
+		// in/out default for new ones.
 		bTargetWeightInterpolationEaseInOut = false;
 	}
+
+	if (Ar.IsLoading() && (Ar.CustomVer(FUE5MainStreamObjectVersion::GUID) <
+		FUE5MainStreamObjectVersion::BlendSpaceSmoothingImprovements))
+	{
+		// If it's an old asset but has the current default smoothing and that smoothing is in use, then it must
+		// have been using the old default so switch to that.
+		for (int Index = 0 ; Index != 3 ; ++Index)
+		{
+			if (InterpolationParam[Index].InterpolationType == EFilterInterpolationType::BSIT_SpringDamper &&
+                InterpolationParam[Index].InterpolationTime > 0.0f)
+			{
+				InterpolationParam[Index].InterpolationType = EFilterInterpolationType::BSIT_Average;
+			}
+			else if (InterpolationParam[Index].InterpolationType == EFilterInterpolationType::BSIT_Cubic)
+			{
+				// Cubic was broken since it was not cubing the interpolation time
+				InterpolationParam[Index].InterpolationTime = FMath::Pow(
+					InterpolationParam[Index].InterpolationTime, 1.0f / 3.0f);
+			}
+		}
+	}
 #endif // WITH_EDITOR
 }
 
@@ -155,7 +181,6 @@ void UBlendSpace::PostEditChangeProperty(struct FPropertyChangedEvent& PropertyC
 			}
 		}
 	}
-
 	Super::PostEditChangeProperty(PropertyChangedEvent);
 }
 #endif // WITH_EDITOR
@@ -943,9 +968,31 @@ void UBlendSpace::InitializeFilter(FBlendFilter* Filter) const
 {
 	if (Filter)
 	{
-		Filter->FilterPerAxis[0].Initialize(InterpolationParam[0].InterpolationTime, InterpolationParam[0].InterpolationType);
-		Filter->FilterPerAxis[1].Initialize(InterpolationParam[1].InterpolationTime, InterpolationParam[1].InterpolationType);
-		Filter->FilterPerAxis[2].Initialize(InterpolationParam[2].InterpolationTime, InterpolationParam[2].InterpolationType);
+		for (int FilterIndex = 0 ; FilterIndex != 3 ; ++FilterIndex)
+		{
+			Filter->FilterPerAxis[FilterIndex].Initialize(InterpolationParam[FilterIndex].InterpolationTime,
+			                                              InterpolationParam[FilterIndex].InterpolationType,
+			                                              InterpolationParam[FilterIndex].DampingRatio,
+			                                              BlendParameters[FilterIndex].Min,
+			                                              BlendParameters[FilterIndex].Max,
+			                                              InterpolationParam[FilterIndex].MaxSpeed,
+			                                              !BlendParameters[FilterIndex].bWrapInput);
+		}
+	}
+}
+
+void UBlendSpace::UpdateFilterParams(FBlendFilter* Filter) const
+{
+	if (Filter)
+	{
+		for (int FilterIndex = 0 ; FilterIndex != 3 ; ++FilterIndex)
+		{
+			Filter->FilterPerAxis[FilterIndex].SetParams(InterpolationParam[FilterIndex].DampingRatio,
+			                                             BlendParameters[FilterIndex].Min,
+			                                             BlendParameters[FilterIndex].Max,
+			                                             InterpolationParam[FilterIndex].MaxSpeed,
+			                                             !BlendParameters[FilterIndex].bWrapInput);
+		}
 	}
 }
 
@@ -1618,17 +1665,20 @@ FVector UBlendSpace::FilterInput(FBlendFilter* Filter, const FVector& BlendInput
 	// Check 
 	for (int32 AxisIndex = 0; AxisIndex < 3; ++AxisIndex)
 	{
-		if (Filter->FilterPerAxis[AxisIndex].NeedsUpdate(InterpolationParam[AxisIndex].InterpolationType, InterpolationParam[AxisIndex].InterpolationTime))
+		if (Filter->FilterPerAxis[AxisIndex].NeedsUpdate(InterpolationParam[AxisIndex].InterpolationType,
+		                                                 InterpolationParam[AxisIndex].InterpolationTime))
 		{
 			InitializeFilter(Filter);
 			break;
 		}
 	}
+	// Note that if we expose the damping ratio etc as pins, this should be called outside of the editor too.
+	UpdateFilterParams(Filter);
 #endif
 	FVector FilteredBlendInput;
-	FilteredBlendInput.X = Filter->FilterPerAxis[0].GetFilteredData(BlendInput.X, DeltaTime);
-	FilteredBlendInput.Y = Filter->FilterPerAxis[1].GetFilteredData(BlendInput.Y, DeltaTime);
-	FilteredBlendInput.Z = Filter->FilterPerAxis[2].GetFilteredData(BlendInput.Z, DeltaTime);
+	FilteredBlendInput.X = Filter->FilterPerAxis[0].UpdateAndGetFilteredData(BlendInput.X, DeltaTime);
+	FilteredBlendInput.Y = Filter->FilterPerAxis[1].UpdateAndGetFilteredData(BlendInput.Y, DeltaTime);
+	FilteredBlendInput.Z = Filter->FilterPerAxis[2].UpdateAndGetFilteredData(BlendInput.Z, DeltaTime);
 	return FilteredBlendInput;
 }
 
diff --git a/Engine/Source/Runtime/Engine/Public/AnimInterpFilter.h b/Engine/Source/Runtime/Engine/Public/AnimInterpFilter.h
--- a/Engine/Source/Runtime/Engine/Public/AnimInterpFilter.h
+++ b/Engine/Source/Runtime/Engine/Public/AnimInterpFilter.h
@@ -2,14 +2,15 @@
 
 /*=============================================================================
 	InterpFilter.h
-=============================================================================*/ 
+=============================================================================*/
 
 #pragma once
 
 #include "CoreMinimal.h"
 #include "Engine/EngineTypes.h"
 
-struct FFIRFilter
+//======================================================================================================================
+struct UE_DEPRECATED(5.0, "FFIRFilter is no longer used or supported") FFIRFilter
 {
 public:
 	FFIRFilter()
@@ -60,7 +61,7 @@ private:
 	float CalculateFilteredOutput() const;
 };
 
-
+//======================================================================================================================
 struct FFilterData
 {
 	float Input;
@@ -69,31 +70,39 @@ struct FFilterData
 	FFilterData()
 	{
 	}
-	void CheckValidation(float CurrentTime, float ValidationWindow)
+
+	FFilterData(const float InInput, const float InTime)
+		: Input(InInput)
+		, Time(InTime)
+	{
+	}
+
+	void CheckValidation(const float CurrentTime, const float ValidationWindow)
 	{
-		if ( Diff(CurrentTime) > ValidationWindow )
+		if (Diff(CurrentTime) > ValidationWindow)
 		{
 			Time = 0.f;
 		}
 	}
 
-	bool IsValid()
+	bool IsValid() const
 	{
 		return (Time > 0.f);
 	}
 
-	float Diff(float InTime)
+	float Diff(const float InTime) const
 	{
 		return (InTime - Time);
 	}
 
-	void SetInput(float InData, float InTime)
+	void SetInput(const float InData, const float InTime)
 	{
 		Input = InData;
 		Time = InTime;
 	}
 };
 
+//======================================================================================================================
 struct ENGINE_API FFIRFilterTimeBased
 {
 public:
@@ -101,47 +110,71 @@ public:
 	{
 	}
 
-	FFIRFilterTimeBased(float Duration, EFilterInterpolationType InInterpolationType)
+	FFIRFilterTimeBased(float InWindowDuration, EFilterInterpolationType InInterpolationType, float InDampingRatio,
+	                    float InMin, float InMax, float InMaxSpeed, bool bInClamp)
 	{
-		Initialize(Duration, InInterpolationType);
+		Initialize(InWindowDuration, InInterpolationType, InDampingRatio, InMin, InMax, InMaxSpeed, bInClamp);
 	}
 
-	void Initialize(float WindowDuration, EFilterInterpolationType BlendType)
+	void Initialize(float InWindowDuration, EFilterInterpolationType InInterpolationType, float InDampingRatio,
+	                float InMinValue, float InMaxValue, float InMaxSpeed, bool bInClamp)
 	{
 		FilterWindow.Empty(10);
 		FilterWindow.AddZeroed(10);
-		InterpolationType = BlendType;
+		InterpolationType = InInterpolationType;
 		NumValidFilter = 0;
 		CurrentStackIndex = 0;
-		TimeDuration = WindowDuration;
+		WindowDuration = InWindowDuration;
+		DampingRatio = InDampingRatio;
+		MinValue = InMinValue;
+		MaxValue = InMaxValue;
+		MaxSpeed = InMaxSpeed;
+		bClamp = bInClamp;
 		CurrentTime = 0.f;
 		LastOutput = 0.f;
 	}
 
-	float GetFilteredData(float Input, float DeltaTime);
-	bool IsValid() const { return TimeDuration > 0.f; }
-	float				LastOutput;
+	// These parameters can be modified at runtime without needing to re-initialize
+	void SetParams(float InDampingRatio, float InMinValue, float InMaxValue, float InMaxSpeed, bool bInClamp)
+	{
+		DampingRatio = InDampingRatio;
+		MinValue = InMinValue;
+		MaxValue = InMaxValue;
+		MaxSpeed = InMaxSpeed;
+		bClamp = bInClamp;
+	}
 
-	void SetWindowDuration(float WindowDuration)
+	// This adds Input to the stack representing an update of DeltaTime, and returns the new filtered value. 
+	float UpdateAndGetFilteredData(float Input, float DeltaTime);
+	bool IsValid() const { return WindowDuration > 0.f; }
+	float LastOutput;
+
+	void SetWindowDuration(float InWindowDuration)
 	{
-		TimeDuration = WindowDuration;
+		WindowDuration = InWindowDuration;
 	}
 
 #if WITH_EDITOR
-	bool NeedsUpdate(const EFilterInterpolationType InType, const float InTime)
+	bool NeedsUpdate(const EFilterInterpolationType InType, const float InWindowDuration)
 	{
-		return InterpolationType != InType || TimeDuration != InTime;
+		return InterpolationType != InType || WindowDuration != InWindowDuration;
 	}
 #endif // WITH_EDITOR
+
 private:
-	EFilterInterpolationType	InterpolationType;
-	int32						CurrentStackIndex;
-	float						TimeDuration;
-	int32						NumValidFilter;
-	float						CurrentTime;
-	TArray<FFilterData>			FilterWindow;
-
-	float GetInterpolationCoefficient (FFilterData & Data);
+	EFilterInterpolationType InterpolationType;
+	int32 CurrentStackIndex;
+	float WindowDuration;
+	float DampingRatio;
+	float MinValue;
+	float MaxValue;
+	float MaxSpeed;
+	bool  bClamp;
+	int32 NumValidFilter;
+	float CurrentTime;
+	TArray<FFilterData> FilterWindow;
+
+	float GetInterpolationCoefficient(const FFilterData& Data) const;
 	float CalculateFilteredOutput();
 	int32 GetSafeCurrentStackIndex();
 	void RefreshValidFilters();
