Allow setting grid placement property for packed level instances.

diff --git a/Engine/Source/Runtime/Engine/Public/LevelInstance/Packed/PackedLevelInstanceActor.h b/Engine/Source/Runtime/Engine/Public/LevelInstance/Packed/PackedLevelInstanceActor.h
--- a/Engine/Source/Runtime/Engine/Public/LevelInstance/Packed/PackedLevelInstanceActor.h
+++ b/Engine/Source/Runtime/Engine/Public/LevelInstance/Packed/PackedLevelInstanceActor.h
@@ -44,6 +44,9 @@ public:
 
 	virtual bool CanEditChange(const FProperty* InProperty) const override;
 
+	// AActor overrides
+	virtual EActorGridPlacement GetDefaultGridPlacement() const override { return EActorGridPlacement::None; }
+
 	virtual bool IsHiddenEd() const override;
 
 	void DestroyPackedComponents();
