Fixed WorldPartition actor SoftObjectPath fixup (properly remaps actor to streaming cells)

diff --git a/Engine/Source/Runtime/Engine/Private/WorldPartition/WorldPartition.cpp b/Engine/Source/Runtime/Engine/Private/WorldPartition/WorldPartition.cpp
--- a/Engine/Source/Runtime/Engine/Private/WorldPartition/WorldPartition.cpp
+++ b/Engine/Source/Runtime/Engine/Private/WorldPartition/WorldPartition.cpp
@@ -29,6 +29,7 @@
 #include "ScopedTransaction.h"
 #include "UnrealEdMisc.h"
 #include "WorldPartition/IWorldPartitionEditorModule.h"
+#include "WorldPartition/WorldPartitionLevelHelper.h"
 #include "WorldPartition/WorldPartitionLevelStreamingDynamic.h"
 #include "WorldPartition/WorldPartitionEditorHash.h"
 #include "WorldPartition/WorldPartitionEditorSpatialHash.h"
@@ -149,7 +150,6 @@ void UWorldPartition::OnBeginPlay(EWorldPartitionStreamingMode Mode)
 {
 	GenerateStreaming(Mode);
 	RuntimeHash->OnBeginPlay(Mode);
-	GetStreamingPolicy()->OnBeginPlay();
 }
 
 void UWorldPartition::OnEndPlay()
@@ -158,7 +158,8 @@ void UWorldPartition::OnEndPlay()
 
 	FlushStreaming();
 	RuntimeHash->OnEndPlay();
-	GetStreamingPolicy()->OnEndPlay();
+
+	GetStreamingPolicy()->ClearActorToCellRemapping();
 
 	World->PersistentLevel->ActorsModifiedForPIE.Empty();
 }
@@ -329,9 +330,13 @@ void UWorldPartition::Initialize(UWorld* InWorld, const FTransform& InTransform)
 	}
 
 #if WITH_EDITOR
-	if (IsRunningGame())
+	if (!bEditorOnly)
 	{
-		OnBeginPlay(EWorldPartitionStreamingMode::EditorStandalone);
+		if (IsRunningGame())
+		{
+			OnBeginPlay(EWorldPartitionStreamingMode::EditorStandalone);
+		}
+		GetStreamingPolicy()->PrepareActorToCellRemapping();
 	}
 #endif
 }
@@ -988,7 +993,15 @@ void UWorldPartition::DrawRuntimeHashPreview()
 bool UWorldPartition::GenerateStreaming(EWorldPartitionStreamingMode Mode, TArray<FString>* OutPackagesToGenerate)
 {
 	check(RuntimeHash);
-	return RuntimeHash->GenerateRuntimeStreaming(Mode, GetStreamingPolicy(), OutPackagesToGenerate);
+	bool Result = RuntimeHash->GenerateRuntimeStreaming(Mode, GetStreamingPolicy(), OutPackagesToGenerate);
+	if (Mode == EWorldPartitionStreamingMode::Cook)
+	{
+		// For Cook, prepare actor to cell remapping
+		GetStreamingPolicy()->PrepareActorToCellRemapping();
+		// Apply remapping of Persistent Level's SoftObjectPaths
+		FWorldPartitionLevelHelper::RemapLevelSoftObjectPaths(World->PersistentLevel, this);
+	}
+	return Result;
 }
 
 bool UWorldPartition::PopulateGeneratedPackageForCook(UPackage* InPackage, const FString& InPackageRelativePath, const FString& InPackageCookName)
@@ -1095,10 +1108,7 @@ void UWorldPartition::DumpActorDescs(const FString& Path)
 
 void UWorldPartition::RemapSoftObjectPath(FSoftObjectPath& ObjectPath)
 {
-	if (GetWorld()->IsGameWorld())
-	{
-		GetStreamingPolicy()->RemapSoftObjectPath(ObjectPath);
-	}
+	GetStreamingPolicy()->RemapSoftObjectPath(ObjectPath);
 }
 
 FBox UWorldPartition::GetWorldBounds() const
diff --git a/Engine/Source/Runtime/Engine/Private/WorldPartition/WorldPartitionLevelHelper.cpp b/Engine/Source/Runtime/Engine/Private/WorldPartition/WorldPartitionLevelHelper.cpp
--- a/Engine/Source/Runtime/Engine/Private/WorldPartition/WorldPartitionLevelHelper.cpp
+++ b/Engine/Source/Runtime/Engine/Private/WorldPartition/WorldPartitionLevelHelper.cpp
@@ -13,6 +13,7 @@
 #include "UnrealEngine.h"
 #include "WorldPartition/ActorDescContainer.h"
 #include "WorldPartition/WorldPartitionActorDesc.h"
+#include "WorldPartition/WorldPartition.h"
 #include "LevelUtils.h"
 #include "Templates/SharedPointer.h"
 
@@ -64,6 +65,40 @@ void FWorldPartitionLevelHelper::MoveExternalActorsToLevel(const TArray<FWorldPa
 	}
 }
 
+void FWorldPartitionLevelHelper::RemapLevelSoftObjectPaths(ULevel* InLevel, UWorldPartition* InWorldPartition)
+{
+	check(InLevel);
+	check(InWorldPartition);
+
+	struct FSoftPathFixupSerializer : public FArchiveUObject
+	{
+		FSoftPathFixupSerializer(TFunctionRef<void(FSoftObjectPath&)> InCustomFixupFunction)
+			: CustomFixupFunction(InCustomFixupFunction)
+		{
+			this->SetIsSaving(true);
+		}
+
+		FArchive& operator<<(FSoftObjectPath& Value)
+		{
+			if (!Value.IsNull())
+			{
+				CustomFixupFunction(Value);
+			}
+			return *this;
+		}
+
+		TFunctionRef<void(FSoftObjectPath&)> CustomFixupFunction;
+	};
+
+	FSoftPathFixupSerializer FixupSerializer([InWorldPartition](FSoftObjectPath& ObjectPath) { InWorldPartition->RemapSoftObjectPath(ObjectPath); });
+	TArray<UObject*> SubObjects;
+	GetObjectsWithOuter(InLevel, SubObjects);
+	for (UObject* Object : SubObjects)
+	{
+		Object->Serialize(FixupSerializer);
+	}
+}
+
 /**
  * Creates an empty Level used in World Partition
  */
diff --git a/Engine/Source/Runtime/Engine/Private/WorldPartition/WorldPartitionLevelStreamingDynamic.cpp b/Engine/Source/Runtime/Engine/Private/WorldPartition/WorldPartitionLevelStreamingDynamic.cpp
--- a/Engine/Source/Runtime/Engine/Private/WorldPartition/WorldPartitionLevelStreamingDynamic.cpp
+++ b/Engine/Source/Runtime/Engine/Private/WorldPartition/WorldPartitionLevelStreamingDynamic.cpp
@@ -337,6 +337,12 @@ void UWorldPartitionLevelStreamingDynamic::FinalizeRuntimeLevel()
 	UPackage* RuntimePackage = RuntimeLevel->GetPackage();
 	RuntimePackage->MarkAsFullyLoaded();
 
+	if (OuterWorld->IsGameWorld())
+	{
+		// Remap Runtime Level's SoftObjectPaths
+		FWorldPartitionLevelHelper::RemapLevelSoftObjectPaths(RuntimeLevel, OuterWorldPartition.Get());
+	}
+
 	if (OuterWorld->IsPlayInEditor())
 	{
 		int32 PIEInstanceID = GetPackage()->PIEInstanceID;
@@ -348,44 +354,7 @@ void UWorldPartitionLevelStreamingDynamic::FinalizeRuntimeLevel()
 		FixupLazyPointersAr << RuntimeLevel;
 
 		// PIE Fixup SoftObjectPaths
-		RuntimeLevel->FixupForPIE(PIEInstanceID, [&](int32 InPIEInstanceID, FSoftObjectPath& ObjectPath)
-		{
-			OuterWorldPartition->RemapSoftObjectPath(ObjectPath);
-		});
-	}
-	else if (IsRunningGame())
-	{
-		check(OuterWorld->IsGameWorld());
-
-		struct FSoftPathFixupSerializer : public FArchiveUObject
-		{
-			FSoftPathFixupSerializer(TFunctionRef<void(FSoftObjectPath&)> InCustomFixupFunction)
-			: CustomFixupFunction(InCustomFixupFunction)
-			{
-				this->SetIsSaving(true);
-			}
-
-			FArchive& operator<<(FSoftObjectPath& Value)
-			{
-				if (!Value.IsNull())
-				{
-					CustomFixupFunction(Value);
-				}
-				return *this;
-			}
-
-			TFunctionRef<void(FSoftObjectPath&)> CustomFixupFunction;
-		};
-
-		FSoftPathFixupSerializer FixupSerializer([&](FSoftObjectPath& ObjectPath) { OuterWorldPartition->RemapSoftObjectPath(ObjectPath); });
-		
-		TArray<UObject*> SubObjects;
-		GetObjectsWithOuter(RuntimeLevel, SubObjects);
-		
-		for (UObject* Object : SubObjects)
-		{
-			Object->Serialize(FixupSerializer);
-		}
+		RuntimeLevel->FixupForPIE(PIEInstanceID);
 	}
 
 	SetLoadedLevel(RuntimeLevel);
diff --git a/Engine/Source/Runtime/Engine/Private/WorldPartition/WorldPartitionLevelStreamingPolicy.cpp b/Engine/Source/Runtime/Engine/Private/WorldPartition/WorldPartitionLevelStreamingPolicy.cpp
--- a/Engine/Source/Runtime/Engine/Private/WorldPartition/WorldPartitionLevelStreamingPolicy.cpp
+++ b/Engine/Source/Runtime/Engine/Private/WorldPartition/WorldPartitionLevelStreamingPolicy.cpp
@@ -199,7 +199,7 @@ TSubclassOf<UWorldPartitionRuntimeCell> UWorldPartitionLevelStreamingPolicy::Get
 	return UWorldPartitionRuntimeLevelStreamingCell::StaticClass();
 }
 
-void UWorldPartitionLevelStreamingPolicy::OnBeginPlay()
+void UWorldPartitionLevelStreamingPolicy::PrepareActorToCellRemapping()
 {
 	TSet<const UWorldPartitionRuntimeCell*> StreamingCells;
 	WorldPartition->RuntimeHash->GetAllStreamingCells(StreamingCells, /*bIncludeDataLayers*/ true);
@@ -216,15 +216,15 @@ void UWorldPartitionLevelStreamingPolicy::OnBeginPlay()
 	}
 }
 
-void UWorldPartitionLevelStreamingPolicy::OnEndPlay()
+void UWorldPartitionLevelStreamingPolicy::ClearActorToCellRemapping()
 {
 	ActorToCellRemapping.Empty();
 }
 
 void UWorldPartitionLevelStreamingPolicy::RemapSoftObjectPath(FSoftObjectPath& ObjectPath)
 {
-	// Once we fix up the path to the proper streaming Level, normal flow of FixupForPIE will adapt the path for PIE.
-	FString SrcPath = ObjectPath.ToString();
+	// Make sure to work on non-PIE path (can happen for modified actors in PIE)
+	FString SrcPath = UWorld::RemovePIEPrefix(ObjectPath.ToString());
 
 	FName* CellName = ActorToCellRemapping.Find(FName(*SrcPath));
 	if (!CellName)
@@ -254,7 +254,14 @@ void UWorldPartitionLevelStreamingPolicy::RemapSoftObjectPath(FSoftObjectPath& O
 		{
 			const FString ObjectNameWithoutPackage = ShortPackageOuterAndName.Mid(DelimiterIdx + 1);
 			const FString PackagePath = UWorldPartitionLevelStreamingPolicy::GetCellPackagePath(*CellName, WorldPartition->GetWorld());
-			FString NewPath = FString::Printf(TEXT("%s.%s"), *PackagePath, *ObjectNameWithoutPackage);
+			FString PrefixPath;
+			if (IsRunningCookCommandlet())
+			{
+				//@todo_ow: Temporary workaround. This information should be provided by the COTFS
+				const UPackage* Package = GetOuterUWorldPartition()->GetWorld()->GetPackage();
+				PrefixPath = FString::Printf(TEXT("%s/%s/_Generated_"), *FPackageName::GetLongPackagePath(Package->GetPathName()), *FPackageName::GetShortName(Package->GetName()));
+			}
+			FString NewPath = FString::Printf(TEXT("%s%s.%s"), *PrefixPath, *PackagePath, *ObjectNameWithoutPackage);
 			ObjectPath.SetPath(MoveTemp(NewPath));
 		}
 	}
diff --git a/Engine/Source/Runtime/Engine/Private/WorldPartition/WorldPartitionRuntimeLevelStreamingCell.cpp b/Engine/Source/Runtime/Engine/Private/WorldPartition/WorldPartitionRuntimeLevelStreamingCell.cpp
--- a/Engine/Source/Runtime/Engine/Private/WorldPartition/WorldPartitionRuntimeLevelStreamingCell.cpp
+++ b/Engine/Source/Runtime/Engine/Private/WorldPartition/WorldPartitionRuntimeLevelStreamingCell.cpp
@@ -152,6 +152,8 @@ bool UWorldPartitionRuntimeLevelStreamingCell::PopulateGeneratedPackageForCook(U
 		ULevel* NewLevel = FWorldPartitionLevelHelper::CreateEmptyLevelForRuntimeCell(OuterWorld, NewLevelStreaming->GetWorldAsset().ToString(), InPackage);
 		check(NewLevel->GetPackage() == InPackage);
 		FWorldPartitionLevelHelper::MoveExternalActorsToLevel(Packages, NewLevel);
+		// Remap Level's SoftObjectPaths
+		FWorldPartitionLevelHelper::RemapLevelSoftObjectPaths(NewLevel, WorldPartition);
 	}
 	return true;
 }
diff --git a/Engine/Source/Runtime/Engine/Public/WorldPartition/WorldPartition.h b/Engine/Source/Runtime/Engine/Public/WorldPartition/WorldPartition.h
--- a/Engine/Source/Runtime/Engine/Public/WorldPartition/WorldPartition.h
+++ b/Engine/Source/Runtime/Engine/Public/WorldPartition/WorldPartition.h
@@ -136,11 +136,9 @@ public:
 	bool AreEditorCellsLoaded(const FBox& Box);
 	bool RefreshLoadedEditorCells();
 
-	// PIE/Game Methods
-	void RemapSoftObjectPath(FSoftObjectPath& ObjectPath);
-
 	// PIE/Game/Cook Methods
 	bool GenerateStreaming(EWorldPartitionStreamingMode Mode, TArray<FString>* OutPackagesToGenerate = nullptr);
+	void RemapSoftObjectPath(FSoftObjectPath& ObjectPath);
 
 	// Cook Methods
 	bool PopulateGeneratedPackageForCook(UPackage* InPackage, const FString& InPackageRelativePath, const FString& InPackageCookName);
diff --git a/Engine/Source/Runtime/Engine/Public/WorldPartition/WorldPartitionLevelHelper.h b/Engine/Source/Runtime/Engine/Public/WorldPartition/WorldPartitionLevelHelper.h
--- a/Engine/Source/Runtime/Engine/Public/WorldPartition/WorldPartitionLevelHelper.h
+++ b/Engine/Source/Runtime/Engine/Public/WorldPartition/WorldPartitionLevelHelper.h
@@ -15,12 +15,14 @@
 #include "WorldPartition/WorldPartitionRuntimeCell.h"
 
 class FWorldPartitionPackageCache;
+class UWorldPartition;
 
 class FWorldPartitionLevelHelper
 {
 public:
 	static ULevel* CreateEmptyLevelForRuntimeCell(const UWorld* InWorld, const FString& InWorldAssetName, UPackage* DestPackage = nullptr);
 	static void MoveExternalActorsToLevel(const TArray<FWorldPartitionRuntimeCellObjectMapping>& InChildPackages, ULevel* InLevel);
+	static void RemapLevelSoftObjectPaths(ULevel* InLevel, UWorldPartition* InWorldPartition);
 	
 	static bool LoadActors(ULevel* InDestLevel, TArrayView<FWorldPartitionRuntimeCellObjectMapping> InActorPackages, FWorldPartitionPackageCache& InPackageCache, TFunction<void(bool)> InCompletionCallback, bool bLoadForPlay, bool bLoadAsync = true, FLinkerInstancingContext* InOutInstancingContext = nullptr);
 private:
diff --git a/Engine/Source/Runtime/Engine/Public/WorldPartition/WorldPartitionLevelStreamingPolicy.h b/Engine/Source/Runtime/Engine/Public/WorldPartition/WorldPartitionLevelStreamingPolicy.h
--- a/Engine/Source/Runtime/Engine/Public/WorldPartition/WorldPartitionLevelStreamingPolicy.h
+++ b/Engine/Source/Runtime/Engine/Public/WorldPartition/WorldPartitionLevelStreamingPolicy.h
@@ -28,8 +28,8 @@ public:
 
 #if WITH_EDITOR
 	virtual TSubclassOf<class UWorldPartitionRuntimeCell> GetRuntimeCellClass() const override;
-	virtual void OnBeginPlay() override;
-	virtual void OnEndPlay() override;
+	virtual void PrepareActorToCellRemapping() override;
+	virtual void ClearActorToCellRemapping() override;
 	virtual void RemapSoftObjectPath(FSoftObjectPath& ObjectPath) override;
 	static FString GetCellPackagePath(const FName& InCellName, const UWorld* InWorld);
 #endif
diff --git a/Engine/Source/Runtime/Engine/Public/WorldPartition/WorldPartitionStreamingPolicy.h b/Engine/Source/Runtime/Engine/Public/WorldPartition/WorldPartitionStreamingPolicy.h
--- a/Engine/Source/Runtime/Engine/Public/WorldPartition/WorldPartitionStreamingPolicy.h
+++ b/Engine/Source/Runtime/Engine/Public/WorldPartition/WorldPartitionStreamingPolicy.h
@@ -36,8 +36,8 @@ public:
 	virtual TSubclassOf<class UWorldPartitionRuntimeCell> GetRuntimeCellClass() const PURE_VIRTUAL(UWorldPartitionStreamingPolicy::GetRuntimeCellClass, return UWorldPartitionRuntimeCell::StaticClass(); );
 
 	// PIE/Game methods
-	virtual void OnBeginPlay() {}
-	virtual void OnEndPlay() {}
+	virtual void PrepareActorToCellRemapping() {}
+	virtual void ClearActorToCellRemapping() {}
 	virtual void RemapSoftObjectPath(FSoftObjectPath& ObjectPath) {}
 #endif
 
