Fix crash when cooking because package containing HLOD static mesh bulk data are moved during cook while async compilation tries to deserialize them

diff --git a/Engine/Source/Runtime/Engine/Private/WorldPartition/WorldPartitionLevelHelper.cpp b/Engine/Source/Runtime/Engine/Private/WorldPartition/WorldPartitionLevelHelper.cpp
--- a/Engine/Source/Runtime/Engine/Private/WorldPartition/WorldPartitionLevelHelper.cpp
+++ b/Engine/Source/Runtime/Engine/Private/WorldPartition/WorldPartitionLevelHelper.cpp
@@ -14,6 +14,7 @@
 #include "WorldPartition/ActorDescContainer.h"
 #include "WorldPartition/WorldPartitionActorDesc.h"
 #include "WorldPartition/WorldPartition.h"
+#include "StaticMeshCompiler.h"
 #include "LevelUtils.h"
 #include "Templates/SharedPointer.h"
 
@@ -40,6 +41,13 @@ UWorld::InitializationValues FWorldPartitionLevelHelper::GetWorldInitializationV
  */
 void FWorldPartitionLevelHelper::MoveExternalActorsToLevel(const TArray<FWorldPartitionRuntimeCellObjectMapping>& InChildPackages, ULevel* InLevel)
 {
+	TRACE_CPUPROFILER_EVENT_SCOPE(FWorldPartitionLevelHelper::MoveExternalActorsToLevel);
+
+	// We can't have async compilation still going on while we move actors as this is going to ResetLoaders which will move bulkdata around that
+	// might still be used by async compilation. 
+	// #TODO_DC Revisit once virtualbulkdata are enabled
+	FStaticMeshCompilingManager::Get().FinishAllCompilation();
+
 	check(InLevel);
 	UPackage* LevelPackage = InLevel->GetPackage();
 	
