Build fix

diff --git a/Engine/Source/Runtime/Engine/Private/WorldPartition/WorldPartitionLevelHelper.cpp b/Engine/Source/Runtime/Engine/Private/WorldPartition/WorldPartitionLevelHelper.cpp
--- a/Engine/Source/Runtime/Engine/Private/WorldPartition/WorldPartitionLevelHelper.cpp
+++ b/Engine/Source/Runtime/Engine/Private/WorldPartition/WorldPartitionLevelHelper.cpp
@@ -11,6 +11,7 @@
 
 #include "FileHelpers.h"
 #include "UnrealEngine.h"
+#include "UObject/UObjectHash.h"
 #include "WorldPartition/ActorDescContainer.h"
 #include "WorldPartition/WorldPartitionActorDesc.h"
 #include "WorldPartition/WorldPartition.h"
