Skip bulkdata serialization of HLOD during map loading
  - 1m34s to 1m13s when starting P_Construct_WP in -game

diff --git a/Engine/Source/Runtime/Engine/Private/WorldPartition/WorldPartitionLevelHelper.cpp b/Engine/Source/Runtime/Engine/Private/WorldPartition/WorldPartitionLevelHelper.cpp
--- a/Engine/Source/Runtime/Engine/Private/WorldPartition/WorldPartitionLevelHelper.cpp
+++ b/Engine/Source/Runtime/Engine/Private/WorldPartition/WorldPartitionLevelHelper.cpp
@@ -76,6 +76,8 @@ void FWorldPartitionLevelHelper::MoveExternalActorsToLevel(const TArray<FWorldPa
 
 void FWorldPartitionLevelHelper::RemapLevelSoftObjectPaths(ULevel* InLevel, UWorldPartition* InWorldPartition)
 {
+	TRACE_CPUPROFILER_EVENT_SCOPE(FWorldPartitionLevelHelper::RemapLevelSoftObjectPaths);
+
 	check(InLevel);
 	check(InWorldPartition);
 
@@ -85,6 +87,7 @@ void FWorldPartitionLevelHelper::RemapLevelSoftObjectPaths(ULevel* InLevel, UWor
 			: CustomFixupFunction(InCustomFixupFunction)
 		{
 			this->SetIsSaving(true);
+			this->ArShouldSkipBulkData = true;
 		}
 
 		FArchive& operator<<(FSoftObjectPath& Value)
