Fix GPU field nor working on some consoles

diff --git a/Engine/Plugins/Experimental/ChaosNiagara/Source/ChaosNiagara/Private/NiagaraDataInterfacePhysicsField.cpp b/Engine/Plugins/Experimental/ChaosNiagara/Source/ChaosNiagara/Private/NiagaraDataInterfacePhysicsField.cpp
--- a/Engine/Plugins/Experimental/ChaosNiagara/Source/ChaosNiagara/Private/NiagaraDataInterfacePhysicsField.cpp
+++ b/Engine/Plugins/Experimental/ChaosNiagara/Source/ChaosNiagara/Private/NiagaraDataInterfacePhysicsField.cpp
@@ -204,10 +204,6 @@ void FNDIPhysicsFieldProxy::ConsumePerInstanceDataFromGameThread(void* PerInstan
 	{
 		TargetData->FieldResource = SourceData->FieldResource;
 	}
-	else
-	{
-		UE_LOG(LogPhysicsField, Log, TEXT("ConsumePerInstanceDataFromGameThread() ... could not find %d"), Instance);
-	}
 }
 
 void FNDIPhysicsFieldProxy::InitializePerInstanceData(const FNiagaraSystemInstanceID& SystemInstance)
diff --git a/Engine/Source/Runtime/Engine/Private/PhysicsField/PhysicsFieldComponent.cpp b/Engine/Source/Runtime/Engine/Private/PhysicsField/PhysicsFieldComponent.cpp
--- a/Engine/Source/Runtime/Engine/Private/PhysicsField/PhysicsFieldComponent.cpp
+++ b/Engine/Source/Runtime/Engine/Private/PhysicsField/PhysicsFieldComponent.cpp
@@ -113,7 +113,7 @@ void InitInternalBuffer(const uint32 ElementCount, FRWBuffer& OutputBuffer)
 
 		if (OutputBuffer.UAV)
 		{
-			FRHICommandListExecutor::GetImmediateCommandList().Transition(FRHITransitionInfo(OutputBuffer.UAV, ERHIAccess::UAVCompute, ERHIAccess::UAVCompute));
+			FRHICommandListExecutor::GetImmediateCommandList().Transition(FRHITransitionInfo(OutputBuffer.UAV, ERHIAccess::SRVCompute, ERHIAccess::UAVCompute));
 			if (PixelFormat == EPixelFormat::PF_R32_FLOAT)
 			{
 				FRHICommandListExecutor::GetImmediateCommandList().ClearUAVFloat(OutputBuffer.UAV, FVector4(ForceInitToZero));
@@ -122,23 +122,33 @@ void InitInternalBuffer(const uint32 ElementCount, FRWBuffer& OutputBuffer)
 			{
 				FRHICommandListExecutor::GetImmediateCommandList().ClearUAVUint(OutputBuffer.UAV, FUintVector4(ForceInitToZero));
 			}
-			FRHICommandListExecutor::GetImmediateCommandList().Transition(FRHITransitionInfo(OutputBuffer.UAV, ERHIAccess::UAVCompute, ERHIAccess::UAVCompute));
+			FRHICommandListExecutor::GetImmediateCommandList().Transition(FRHITransitionInfo(OutputBuffer.UAV, ERHIAccess::UAVCompute, ERHIAccess::SRVCompute));
+
 		}
 	}
 }
 
 template<typename BufferType, int ElementSize, EPixelFormat PixelFormat>
-void UpdateInternalBuffer(const uint32 ElementCount, const BufferType* InputData, FRWBuffer& OutputBuffer)
+void UpdateInternalBuffer(const uint32 ElementCount, const BufferType* InputData, FRWBuffer& OutputBuffer, const bool bInitField = false)
 {
 	if (ElementCount > 0 && InputData)
 	{
 		const uint32 BufferCount = ElementCount * ElementSize;
 		const uint32 BufferBytes = sizeof(BufferType) * BufferCount;
 
+		FRHICommandListExecutor::GetImmediateCommandList().Transition(FRHITransitionInfo(OutputBuffer.UAV, ERHIAccess::SRVCompute, ERHIAccess::UAVCompute));
+
+		if(bInitField)
+		{
+			OutputBuffer.Initialize(TEXT("FPhysicsFieldResource"), sizeof(BufferType), BufferCount, PixelFormat, BUF_Static);
+		}
+
 		void* OutputData = RHILockBuffer(OutputBuffer.Buffer, 0, BufferBytes, RLM_WriteOnly);
 
 		FMemory::Memcpy(OutputData, InputData, BufferBytes);
 		RHIUnlockBuffer(OutputBuffer.Buffer);
+
+		FRHICommandListExecutor::GetImmediateCommandList().Transition(FRHITransitionInfo(OutputBuffer.UAV, ERHIAccess::UAVCompute, ERHIAccess::SRVCompute));
 	}
 }
 
@@ -530,6 +540,7 @@ void FPhysicsFieldResource::UpdateResource(FRHICommandListImmediate& RHICmdList,
 {
 	SCOPE_CYCLE_COUNTER(STAT_PhysicsFields_UpdateResource_RT);
 	const int32 CellsCount = FieldInfos.ClipmapCount * EFieldPhysicsType::Field_PhysicsType_Max;
+
 	{
 		SCOPED_DRAW_EVENT(RHICmdList, PhysicsFields_ClearClipmap);
 		SCOPED_GPU_STAT(RHICmdList, PhysicsFields_ClearClipmap);
@@ -547,11 +558,9 @@ void FPhysicsFieldResource::UpdateResource(FRHICommandListImmediate& RHICmdList,
 	{
 		SCOPED_DRAW_EVENT(RHICmdList, PhysicsFields_UpdateBuffers);
 		SCOPED_GPU_STAT(RHICmdList, PhysicsFields_UpdateBuffers);
-		InitInternalBuffer<float, 1, EPixelFormat::PF_R32_FLOAT>(NodesParamsDatas.Num(), NodesParams);
-		InitInternalBuffer<int32, 1, EPixelFormat::PF_R32_SINT>(NodesOffsetsDatas.Num(), NodesOffsets);
 
-		UpdateInternalBuffer<float, 1, EPixelFormat::PF_R32_FLOAT>(NodesParamsDatas.Num(), NodesParamsDatas.GetData(), NodesParams);
-		UpdateInternalBuffer<int32, 1, EPixelFormat::PF_R32_SINT>(NodesOffsetsDatas.Num(), NodesOffsetsDatas.GetData(), NodesOffsets);
+		UpdateInternalBuffer<float, 1, EPixelFormat::PF_R32_FLOAT>(NodesParamsDatas.Num(), NodesParamsDatas.GetData(), NodesParams, true);
+		UpdateInternalBuffer<int32, 1, EPixelFormat::PF_R32_SINT>(NodesOffsetsDatas.Num(), NodesOffsetsDatas.GetData(), NodesOffsets, true);
 		UpdateInternalBuffer<int32, 1, EPixelFormat::PF_R32_SINT>(TargetsOffsetsDatas.Num(), TargetsOffsetsDatas.GetData(), TargetsOffsets);
 
 		FieldInfos.ClipmapCenter = FieldInfos.ViewOrigin;
