Merge identitcal code in UWorldPartitionLevelStreamingDynamic::OnCleanupLevel.

diff --git a/Engine/Source/Runtime/Engine/Private/WorldPartition/WorldPartitionLevelStreamingDynamic.cpp b/Engine/Source/Runtime/Engine/Private/WorldPartition/WorldPartitionLevelStreamingDynamic.cpp
--- a/Engine/Source/Runtime/Engine/Private/WorldPartition/WorldPartitionLevelStreamingDynamic.cpp
+++ b/Engine/Source/Runtime/Engine/Private/WorldPartition/WorldPartitionLevelStreamingDynamic.cpp
@@ -383,38 +383,32 @@ void UWorldPartitionLevelStreamingDynamic::OnCleanupLevel()
 
 		RuntimeLevel->OnCleanupLevel.Remove(OnCleanupLevelDelegateHandle);
 
-		// Clears RF_Standalone flag on objects in package (Metadata)
-		UPackage* RuntimePackage = RuntimeLevel->GetPackage();
-		ForEachObjectWithPackage(RuntimePackage, [](UObject* Object)
+		auto TrashPackage = [](UPackage* Package)
 		{
-			Object->ClearFlags(RF_Standalone);
-			return true;
-		}, false);
-
-		// Rename package to avoid having to deal with pending kill objects in subsequent RequestLevel
-		FName NewUniqueTrashName = MakeUniqueObjectName(nullptr, UPackage::StaticClass(), FName(*FString::Printf(TEXT("%s_Trashed"), *RuntimePackage->GetName())));
-		RuntimePackage->Rename(*NewUniqueTrashName.ToString(), nullptr, REN_ForceNoResetLoaders | REN_DontCreateRedirectors | REN_NonTransactional | REN_DoNotDirty);
-
-		// Apply same logic on Actor Packages
+			// Clears RF_Standalone flag on objects in package (UMetaData)
+			ForEachObjectWithPackage(Package, [](UObject* Object)
+			{
+				Object->ClearFlags(RF_Standalone);
+				return true;
+			}, false);
+
+			// Rename package to avoid having to deal with pending kill objects in subsequent RequestLevel calls
+			FName NewPackageName = MakeUniqueObjectName(nullptr, UPackage::StaticClass(), FName(*FString::Printf(TEXT("%s_Trashed"), *Package->GetName())));
+			Package->Rename(*NewPackageName.ToString(), nullptr, REN_ForceNoResetLoaders | REN_DontCreateRedirectors | REN_NonTransactional | REN_DoNotDirty);
+		};
+		
+		TrashPackage(RuntimeLevel->GetPackage());
 		for (AActor* Actor : RuntimeLevel->Actors)
 		{
 			if (UPackage* ActorPackage = Actor ? Actor->GetExternalPackage() : nullptr)
 			{
-				ForEachObjectWithPackage(ActorPackage, [&](UObject* Object)
-				{
-					Object->ClearFlags(RF_Standalone);
-					return true;
-				}, false);
-
-				NewUniqueTrashName = MakeUniqueObjectName(nullptr, UPackage::StaticClass(), FName(*FString::Printf(TEXT("%s_Trashed"), *ActorPackage->GetName())));
-				ActorPackage->Rename(*NewUniqueTrashName.ToString(), nullptr, REN_ForceNoResetLoaders | REN_DontCreateRedirectors | REN_NonTransactional | REN_DoNotDirty);
+				TrashPackage(ActorPackage);
 			}
 		}
 
 		RuntimeLevel = nullptr;
 	}
 }
-
 #endif
 
 /*
