CoerceValueRaw emits an error instead of asserting. Fixes issue where connecting incompatible types to a dot node would cause a crash.

diff --git a/Engine/Source/Runtime/Engine/Private/Materials/HLSLMaterialDerivativeAutogen.cpp b/Engine/Source/Runtime/Engine/Private/Materials/HLSLMaterialDerivativeAutogen.cpp
--- a/Engine/Source/Runtime/Engine/Private/Materials/HLSLMaterialDerivativeAutogen.cpp
+++ b/Engine/Source/Runtime/Engine/Private/Materials/HLSLMaterialDerivativeAutogen.cpp
@@ -209,7 +209,7 @@ void FMaterialDerivativeAutogen::EnableGeneratedDepencencies()
 }
 
 // given a string, convert it from type to type
-FString FMaterialDerivativeAutogen::CoerceValueRaw(const FString& Token, int32 SrcType, EDerivativeStatus SrcStatus, int32 DstType)
+FString FMaterialDerivativeAutogen::CoerceValueRaw(FHLSLMaterialTranslator& Translator, const FString& Token, int32 SrcType, EDerivativeStatus SrcStatus, int32 DstType)
 {
 	check(IsDerivTypeIndexValid(SrcType));
 	check(IsDerivTypeIndexValid(DstType));
@@ -228,26 +228,33 @@ FString FMaterialDerivativeAutogen::CoerceValueRaw(const FString& Token, int32 S
 	}
 	else
 	{
-		check(SrcType == 0); // can only coerce a float1
-		if (DstType == 0)
+		if (SrcType == 0)
 		{
-			Ret = FString::Printf( TEXT("MaterialFloat(%s)"), *Ret);
-		}
-		else if (DstType == 1)
-		{
-			Ret = FString::Printf( TEXT("MaterialFloat2(%s,%s)"), *Ret, *Ret);
-		}
-		else if (DstType == 2)
-		{
-			Ret = FString::Printf( TEXT("MaterialFloat3(%s,%s,%s)"), *Ret, *Ret, *Ret);
-		}
-		else if (DstType == 3)
-		{
-			Ret = FString::Printf( TEXT("MaterialFloat4(%s,%s,%s,%s)"), *Ret, *Ret, *Ret, *Ret);
+			if (DstType == 0)
+			{
+				Ret = FString::Printf( TEXT("MaterialFloat(%s)"), *Ret);
+			}
+			else if (DstType == 1)
+			{
+				Ret = FString::Printf( TEXT("MaterialFloat2(%s,%s)"), *Ret, *Ret);
+			}
+			else if (DstType == 2)
+			{
+				Ret = FString::Printf( TEXT("MaterialFloat3(%s,%s,%s)"), *Ret, *Ret, *Ret);
+			}
+			else if (DstType == 3)
+			{
+				Ret = FString::Printf( TEXT("MaterialFloat4(%s,%s,%s,%s)"), *Ret, *Ret, *Ret, *Ret);
+			}
+			else
+			{
+				check(0);
+			}
 		}
 		else
 		{
-			check(0);
+			Translator.Errorf(TEXT("Coercion failed: %s: %s -> %s"), *Token, GetFloatVectorName(SrcType), GetFloatVectorName(DstType));
+			return TEXT("");
 		}
 	}
 	return Ret;
@@ -426,7 +433,7 @@ int32 FMaterialDerivativeAutogen::GenerateExpressionFunc1(FHLSLMaterialTranslato
 		// we are in analytic mode, we may need to get the value.
 		if (Index == CompiledPDV_Analytic)
 		{
-			SrcToken = CoerceValueRaw(SrcToken, SrcDerivInfo.TypeIndex, SrcDerivInfo.DerivativeStatus, SrcDerivInfo.TypeIndex);
+			SrcToken = CoerceValueRaw(Translator, SrcToken, SrcDerivInfo.TypeIndex, SrcDerivInfo.DerivativeStatus, SrcDerivInfo.TypeIndex);
 		}
 
 		FString DstToken;
@@ -692,8 +699,8 @@ int32 FMaterialDerivativeAutogen::GenerateExpressionFunc2(FHLSLMaterialTranslato
 		// we are in analytic mode, we may need to get the value.
 		if (Index == CompiledPDV_Analytic)
 		{
-			LhsToken = CoerceValueRaw(LhsToken, LhsDerivInfo.TypeIndex, LhsDerivInfo.DerivativeStatus, IntermediaryTypeIndex);
-			RhsToken = CoerceValueRaw(RhsToken, RhsDerivInfo.TypeIndex, RhsDerivInfo.DerivativeStatus, IntermediaryTypeIndex);
+			LhsToken = CoerceValueRaw(Translator, LhsToken, LhsDerivInfo.TypeIndex, LhsDerivInfo.DerivativeStatus, IntermediaryTypeIndex);
+			RhsToken = CoerceValueRaw(Translator, RhsToken, RhsDerivInfo.TypeIndex, RhsDerivInfo.DerivativeStatus, IntermediaryTypeIndex);
 		}
 
 		FString DstToken;
diff --git a/Engine/Source/Runtime/Engine/Private/Materials/HLSLMaterialDerivativeAutogen.h b/Engine/Source/Runtime/Engine/Private/Materials/HLSLMaterialDerivativeAutogen.h
--- a/Engine/Source/Runtime/Engine/Private/Materials/HLSLMaterialDerivativeAutogen.h
+++ b/Engine/Source/Runtime/Engine/Private/Materials/HLSLMaterialDerivativeAutogen.h
@@ -114,7 +114,7 @@ private:
 	static int32 GetFunc1ReturnNumComponents(int32 SrcTypeIndex, EFunc1 Op);
 	static int32 GetFunc2ReturnNumComponents(int32 LhsTypeIndex, int32 RhsTypeIndex, EFunc2 Op);
 
-	FString CoerceValueRaw(const FString& Token, int32 SrcType, EDerivativeStatus SrcStatus, int32 DstType);
+	FString CoerceValueRaw(FHLSLMaterialTranslator& Translator, const FString& Token, int32 SrcType, EDerivativeStatus SrcStatus, int32 DstType);
 	FString CoerceValueDeriv(const FString& Token, int32 SrcType, EDerivativeStatus SrcStatus, int32 DstType);
 
 	void EnableGeneratedDepencencies();
