Fix HLODs generation when there's no existing HLOD actors in the map.

diff --git a/Engine/Source/Runtime/Engine/Public/WorldPartition/HLOD/HLODActor.h b/Engine/Source/Runtime/Engine/Public/WorldPartition/HLOD/HLODActor.h
--- a/Engine/Source/Runtime/Engine/Public/WorldPartition/HLOD/HLODActor.h
+++ b/Engine/Source/Runtime/Engine/Public/WorldPartition/HLOD/HLODActor.h
@@ -22,6 +22,8 @@ public:
 	inline FName GetCellName() const { return CellName; }
 	inline uint32 GetLODLevel() const { return LODLevel; }
 
+	virtual bool IsHLODRelevant() const override { return true; }
+
 #if WITH_EDITOR
 	void SetHLODPrimitives(const TArray<UPrimitiveComponent*>& InHLODPrimitives);
 
