Fixed WorldPartition SoftObjectPath remapping at PIE : was only applied on streamed runtime levels (not in PersistentLevel at startup).

diff --git a/Engine/Source/Runtime/Engine/Classes/Engine/World.h b/Engine/Source/Runtime/Engine/Classes/Engine/World.h
--- a/Engine/Source/Runtime/Engine/Classes/Engine/World.h
+++ b/Engine/Source/Runtime/Engine/Classes/Engine/World.h
@@ -3828,8 +3828,8 @@ public:
 	/** Given a loaded editor UWorld, duplicate it for play in editor purposes with OwningWorld as the world with the persistent level. */
 	static UWorld* DuplicateWorldForPIE(const FString& PackageName, UWorld* OwningWorld);
 
-	/** Given a string, return that string with any PIE prefix removed */
-	static FString RemovePIEPrefix(const FString &Source);
+	/** Given a string, return that string with any PIE prefix removed. Optionally returns the PIE Instance ID. */
+	static FString RemovePIEPrefix(const FString &Source, int32* OutPIEInstanceID = nullptr);
 
 	/** Given a package, locate the UWorld contained within if one exists */
 	static UWorld* FindWorldInPackage(UPackage* Package);
diff --git a/Engine/Source/Runtime/Engine/Private/World.cpp b/Engine/Source/Runtime/Engine/Private/World.cpp
--- a/Engine/Source/Runtime/Engine/Private/World.cpp
+++ b/Engine/Source/Runtime/Engine/Private/World.cpp
@@ -7153,11 +7153,14 @@ FString UWorld::GetAddressURL() const
 	return FString::Printf( TEXT("%s"), *URL.GetHostPortString() );
 }
 
-FString UWorld::RemovePIEPrefix(const FString &Source)
+FString UWorld::RemovePIEPrefix(const FString &Source, int32* OutPIEInstanceID)
 {
 	// PIE prefix is: UEDPIE_X_MapName (where X is some decimal number)
 	const FString LookFor = PLAYWORLD_PACKAGE_PREFIX;
-
+	if (OutPIEInstanceID)
+	{
+		*OutPIEInstanceID = INDEX_NONE;
+	}
 	int32 idx = Source.Find(LookFor);
 	if (idx >= 0)
 	{
@@ -7167,6 +7170,7 @@ FString UWorld::RemovePIEPrefix(const FString &Source)
 			UE_LOG(LogWorld, Warning, TEXT("Looks like World path invalid PIE prefix (expected '_' characeter after PIE prefix): %s"), *Source);
 			return Source;
 		}
+		int32 PIEInstanceIDStartIndex = end + 1;
 		for (++end; (end < Source.Len()) && (Source[end] != '_'); ++end)
 		{
 			if ((Source[end] < '0') || (Source[end] > '9'))
@@ -7180,6 +7184,12 @@ FString UWorld::RemovePIEPrefix(const FString &Source)
 			UE_LOG(LogWorld, Warning, TEXT("Looks like World path invalid PIE prefix (can't find end of PIE prefix): %s"), *Source);
 			return Source;
 		}
+		if (OutPIEInstanceID && (end > PIEInstanceIDStartIndex))
+		{
+			const int32 PIEInstanceIDCount = end - PIEInstanceIDStartIndex;
+			const FString PIEInstanceIDStr = Source.Mid(PIEInstanceIDStartIndex, PIEInstanceIDCount);
+			TTypeFromString<int32>::FromString(*OutPIEInstanceID, *PIEInstanceIDStr);
+		}
 		const FString Prefix = Source.Left(idx);
 		const FString Suffix = Source.Right(Source.Len() - end - 1);
 		return Prefix + Suffix;
diff --git a/Engine/Source/Runtime/Engine/Private/WorldPartition/WorldPartition.cpp b/Engine/Source/Runtime/Engine/Private/WorldPartition/WorldPartition.cpp
--- a/Engine/Source/Runtime/Engine/Private/WorldPartition/WorldPartition.cpp
+++ b/Engine/Source/Runtime/Engine/Private/WorldPartition/WorldPartition.cpp
@@ -337,6 +337,8 @@ void UWorldPartition::Initialize(UWorld* InWorld, const FTransform& InTransform)
 			OnBeginPlay(EWorldPartitionStreamingMode::EditorStandalone);
 		}
 		GetStreamingPolicy()->PrepareActorToCellRemapping();
+		// Apply remapping of Persistent Level's SoftObjectPaths
+		FWorldPartitionLevelHelper::RemapLevelSoftObjectPaths(World->PersistentLevel, this);
 	}
 #endif
 }
diff --git a/Engine/Source/Runtime/Engine/Private/WorldPartition/WorldPartitionLevelStreamingDynamic.cpp b/Engine/Source/Runtime/Engine/Private/WorldPartition/WorldPartitionLevelStreamingDynamic.cpp
--- a/Engine/Source/Runtime/Engine/Private/WorldPartition/WorldPartitionLevelStreamingDynamic.cpp
+++ b/Engine/Source/Runtime/Engine/Private/WorldPartition/WorldPartitionLevelStreamingDynamic.cpp
@@ -337,12 +337,6 @@ void UWorldPartitionLevelStreamingDynamic::FinalizeRuntimeLevel()
 	UPackage* RuntimePackage = RuntimeLevel->GetPackage();
 	RuntimePackage->MarkAsFullyLoaded();
 
-	if (OuterWorld->IsGameWorld())
-	{
-		// Remap Runtime Level's SoftObjectPaths
-		FWorldPartitionLevelHelper::RemapLevelSoftObjectPaths(RuntimeLevel, OuterWorldPartition.Get());
-	}
-
 	if (OuterWorld->IsPlayInEditor())
 	{
 		int32 PIEInstanceID = GetPackage()->PIEInstanceID;
@@ -354,7 +348,17 @@ void UWorldPartitionLevelStreamingDynamic::FinalizeRuntimeLevel()
 		FixupLazyPointersAr << RuntimeLevel;
 
 		// PIE Fixup SoftObjectPaths
-		RuntimeLevel->FixupForPIE(PIEInstanceID);
+		RuntimeLevel->FixupForPIE(PIEInstanceID, [&](int32 InPIEInstanceID, FSoftObjectPath& ObjectPath)
+		{
+			// Remap Runtime Level's SoftObjectPath before each PIE Fixup to avoid doing 2 passes of serialization
+			OuterWorldPartition->RemapSoftObjectPath(ObjectPath);
+		});
+	}
+	else if (OuterWorld->IsGameWorld())
+	{
+		check(IsRunningGame());
+		// Remap Runtime Level's SoftObjectPaths
+		FWorldPartitionLevelHelper::RemapLevelSoftObjectPaths(RuntimeLevel, OuterWorldPartition.Get());
 	}
 
 	SetLoadedLevel(RuntimeLevel);
diff --git a/Engine/Source/Runtime/Engine/Private/WorldPartition/WorldPartitionLevelStreamingPolicy.cpp b/Engine/Source/Runtime/Engine/Private/WorldPartition/WorldPartitionLevelStreamingPolicy.cpp
--- a/Engine/Source/Runtime/Engine/Private/WorldPartition/WorldPartitionLevelStreamingPolicy.cpp
+++ b/Engine/Source/Runtime/Engine/Private/WorldPartition/WorldPartitionLevelStreamingPolicy.cpp
@@ -224,7 +224,8 @@ void UWorldPartitionLevelStreamingPolicy::ClearActorToCellRemapping()
 void UWorldPartitionLevelStreamingPolicy::RemapSoftObjectPath(FSoftObjectPath& ObjectPath)
 {
 	// Make sure to work on non-PIE path (can happen for modified actors in PIE)
-	FString SrcPath = UWorld::RemovePIEPrefix(ObjectPath.ToString());
+	int32 PIEInstanceID = INDEX_NONE;
+	FString SrcPath = UWorld::RemovePIEPrefix(ObjectPath.ToString(), &PIEInstanceID);
 
 	FName* CellName = ActorToCellRemapping.Find(FName(*SrcPath));
 	if (!CellName)
@@ -252,8 +253,9 @@ void UWorldPartitionLevelStreamingPolicy::RemapSoftObjectPath(FSoftObjectPath& O
 		int32 DelimiterIdx;
 		if (ShortPackageOuterAndName.FindChar('.', DelimiterIdx))
 		{
+			UWorld* World = WorldPartition->GetWorld();
 			const FString ObjectNameWithoutPackage = ShortPackageOuterAndName.Mid(DelimiterIdx + 1);
-			const FString PackagePath = UWorldPartitionLevelStreamingPolicy::GetCellPackagePath(*CellName, WorldPartition->GetWorld());
+			const FString PackagePath = UWorldPartitionLevelStreamingPolicy::GetCellPackagePath(*CellName, World);
 			FString PrefixPath;
 			if (IsRunningCookCommandlet())
 			{
@@ -263,6 +265,11 @@ void UWorldPartitionLevelStreamingPolicy::RemapSoftObjectPath(FSoftObjectPath& O
 			}
 			FString NewPath = FString::Printf(TEXT("%s%s.%s"), *PrefixPath, *PackagePath, *ObjectNameWithoutPackage);
 			ObjectPath.SetPath(MoveTemp(NewPath));
+			// Put back PIE prefix
+			if (World->IsPlayInEditor() && (PIEInstanceID != INDEX_NONE))
+			{
+				ObjectPath.FixupForPIE(PIEInstanceID);
+			}
 		}
 	}
 }
