Add meaningful Insights trace affecting either map loading or PIE

diff --git a/Engine/Plugins/Editor/StylusInput/Source/StylusInput/Private/StylusInputModule.cpp b/Engine/Plugins/Editor/StylusInput/Source/StylusInput/Private/StylusInputModule.cpp
--- a/Engine/Plugins/Editor/StylusInput/Source/StylusInput/Private/StylusInputModule.cpp
+++ b/Engine/Plugins/Editor/StylusInput/Source/StylusInput/Private/StylusInputModule.cpp
@@ -99,6 +99,8 @@ void UStylusInputSubsystem::RemoveMessageHandler(IStylusMessageHandler& InHandle
 
 void UStylusInputSubsystem::Tick(float DeltaTime)
 {
+	TRACE_CPUPROFILER_EVENT_SCOPE(UStylusInputSubsystem::Tick);
+
 	if (InputInterface.IsValid())
 	{
 		InputInterface->Tick();
diff --git a/Engine/Source/Editor/DataLayerEditor/Private/DataLayer/DataLayerEditorSubsystem.cpp b/Engine/Source/Editor/DataLayerEditor/Private/DataLayer/DataLayerEditorSubsystem.cpp
--- a/Engine/Source/Editor/DataLayerEditor/Private/DataLayer/DataLayerEditorSubsystem.cpp
+++ b/Engine/Source/Editor/DataLayerEditor/Private/DataLayer/DataLayerEditorSubsystem.cpp
@@ -805,6 +805,8 @@ bool UDataLayerEditorSubsystem::UpdateActorVisibility(AActor* Actor, bool& bOutS
 
 bool UDataLayerEditorSubsystem::UpdateAllActorsVisibility(const bool bNotifySelectionChange, const bool bRedrawViewports)
 {
+	TRACE_CPUPROFILER_EVENT_SCOPE(UDataLayerEditorSubsystem::UpdateAllActorsVisibility);
+
 	bool bSelectionChanged = false;
 	bool bChangesOccurred = false;
 	for (AActor* Actor : FActorRange(GetWorld()))
diff --git a/Engine/Source/Editor/UnrealEd/Private/EditorEngine.cpp b/Engine/Source/Editor/UnrealEd/Private/EditorEngine.cpp
--- a/Engine/Source/Editor/UnrealEd/Private/EditorEngine.cpp
+++ b/Engine/Source/Editor/UnrealEd/Private/EditorEngine.cpp
@@ -4180,6 +4180,8 @@ void UEditorEngine::RunDeferredMarkForAddFiles(bool)
 
 bool UEditorEngine::InitializePhysicsSceneForSaveIfNecessary(UWorld* World, bool &bOutForceInitialized)
 {
+	TRACE_CPUPROFILER_EVENT_SCOPE(UEditorEngine::InitializePhysicsSceneForSaveIfNecessary);
+
 	// We need a physics scene at save time in case code does traces during onsave events.
 	bool bHasPhysicsScene = false;
 
diff --git a/Engine/Source/Editor/UnrealEd/Private/UnrealEdMisc.cpp b/Engine/Source/Editor/UnrealEd/Private/UnrealEdMisc.cpp
--- a/Engine/Source/Editor/UnrealEd/Private/UnrealEdMisc.cpp
+++ b/Engine/Source/Editor/UnrealEd/Private/UnrealEdMisc.cpp
@@ -741,6 +741,8 @@ void FUnrealEdMisc::TickAssetAnalytics()
 {
 	if( bIsAssetAnalyticsPending )
 	{
+		TRACE_CPUPROFILER_EVENT_SCOPE(FUnrealEdMisc::TickAssetAnalytics);
+
 		FAssetRegistryModule& AssetRegistryModule = FModuleManager::LoadModuleChecked<FAssetRegistryModule>(AssetRegistryName);
 
 		if( !AssetRegistryModule.Get().IsLoadingAssets())
diff --git a/Engine/Source/Runtime/AssetRegistry/Private/AssetRegistry.cpp b/Engine/Source/Runtime/AssetRegistry/Private/AssetRegistry.cpp
--- a/Engine/Source/Runtime/AssetRegistry/Private/AssetRegistry.cpp
+++ b/Engine/Source/Runtime/AssetRegistry/Private/AssetRegistry.cpp
@@ -2305,6 +2305,8 @@ const TSet<FName>& UAssetRegistryImpl::GetCachedEmptyPackages() const
 
 void UAssetRegistryImpl::ScanPathsSynchronousInternal(const TArray<FString>& InDirs, const TArray<FString>& InFiles, bool bForceRescan, bool bIgnoreBlackListScanFilters, TArray<FName>* OutFoundAssets)
 {
+	TRACE_CPUPROFILER_EVENT_SCOPE(UAssetRegistryImpl::ScanPathsSynchronousInternal);
+
 	// Note this function can be reentered due to arbitrary code execution in construction of FAssetData
 	LLM_SCOPE(ELLMTag::AssetRegistry);
 	const double SearchStartTime = FPlatformTime::Seconds();
@@ -2976,6 +2978,8 @@ void UAssetRegistryImpl::RemovePackageData(const FName PackageName)
 
 void UAssetRegistryImpl::OnDirectoryChanged(const TArray<FFileChangeData>& FileChanges)
 {
+	TRACE_CPUPROFILER_EVENT_SCOPE(UAssetRegistryImpl::OnDirectoryChanged);
+
 	// Take local copy of FileChanges array as we wish to collapse pairs of 'Removed then Added' FileChangeData
 	// entries into a single 'Modified' entry.
 	TArray<FFileChangeData> FileChangesProcessed(FileChanges);
diff --git a/Engine/Source/Runtime/Engine/Private/AI/NavigationSystemBase.cpp b/Engine/Source/Runtime/Engine/Private/AI/NavigationSystemBase.cpp
--- a/Engine/Source/Runtime/Engine/Private/AI/NavigationSystemBase.cpp
+++ b/Engine/Source/Runtime/Engine/Private/AI/NavigationSystemBase.cpp
@@ -67,6 +67,7 @@ namespace FNavigationSystem
 
 	void AddNavigationSystemToWorld(UWorld& WorldOwner, const FNavigationSystemRunMode RunMode, UNavigationSystemConfig* NavigationSystemConfig, const bool bInitializeForWorld, const bool bOverridePreviousNavSys)
 	{
+		TRACE_CPUPROFILER_EVENT_SCOPE(FNavigationSystem::AddNavigationSystemToWorld);
 		UE_LOG(LogNavigation, VeryVerbose, TEXT("%s (WorldOwner: %s)"), ANSI_TO_TCHAR(__FUNCTION__), *WorldOwner.GetOuter()->GetName());
 
 		const FNavigationSystemRunMode ResolvedRunMode = (RunMode == FNavigationSystemRunMode::InferFromWorldMode) ? FindRunModeFromWorldType(WorldOwner) : RunMode;
diff --git a/Engine/Source/Runtime/Engine/Private/ActorConstruction.cpp b/Engine/Source/Runtime/Engine/Private/ActorConstruction.cpp
--- a/Engine/Source/Runtime/Engine/Private/ActorConstruction.cpp
+++ b/Engine/Source/Runtime/Engine/Private/ActorConstruction.cpp
@@ -228,33 +228,34 @@ void AActor::RerunConstructionScripts()
 	// don't allow (re)running construction scripts on dying actors and Actors that seamless traveled 
 	// were constructed in the previous level and should not have construction scripts rerun
 	bool bAllowReconstruction = !bActorSeamlessTraveled && !IsPendingKill() && !HasAnyFlags(RF_BeginDestroyed|RF_FinishDestroyed);
-#if WITH_EDITOR
-	if(bAllowReconstruction && GIsEditor)
+	if (bAllowReconstruction)
 	{
-		// Don't allow reconstruction if we're still in the middle of construction.
-		bAllowReconstruction = !bActorIsBeingConstructed;
-		if (ensureMsgf(bAllowReconstruction, TEXT("Attempted to rerun construction scripts on an Actor that isn't fully constructed yet (%s)."), *GetFullName()))
+		TRACE_CPUPROFILER_EVENT_SCOPE(AActor::RerunConstructionScripts);
+#if WITH_EDITOR
+		if(GIsEditor)
 		{
-			// Generate the blueprint hierarchy for this actor
-			TArray<UBlueprint*> ParentBPStack;
-			bAllowReconstruction = UBlueprint::GetBlueprintHierarchyFromClass(GetClass(), ParentBPStack);
-			if (bAllowReconstruction)
+			// Don't allow reconstruction if we're still in the middle of construction.
+			bAllowReconstruction = !bActorIsBeingConstructed;
+			if (ensureMsgf(bAllowReconstruction, TEXT("Attempted to rerun construction scripts on an Actor that isn't fully constructed yet (%s)."), *GetFullName()))
 			{
-				for (int i = ParentBPStack.Num() - 1; i > 0 && bAllowReconstruction; --i)
+				// Generate the blueprint hierarchy for this actor
+				TArray<UBlueprint*> ParentBPStack;
+				bAllowReconstruction = UBlueprint::GetBlueprintHierarchyFromClass(GetClass(), ParentBPStack);
+				if (bAllowReconstruction)
 				{
-					const UBlueprint* ParentBP = ParentBPStack[i];
-					if (ParentBP && ParentBP->bBeingCompiled)
+					for (int i = ParentBPStack.Num() - 1; i > 0 && bAllowReconstruction; --i)
 					{
-						// don't allow (re)running construction scripts if a parent BP is being compiled
-						bAllowReconstruction = false;
+						const UBlueprint* ParentBP = ParentBPStack[i];
+						if (ParentBP && ParentBP->bBeingCompiled)
+						{
+							// don't allow (re)running construction scripts if a parent BP is being compiled
+							bAllowReconstruction = false;
+						}
 					}
 				}
 			}
 		}
-	}
 #endif
-	if(bAllowReconstruction)
-	{
 		// Child Actors can be customized in many ways by their parents construction scripts and rerunning directly on them would wipe
 		// that out. So instead we redirect up the hierarchy
 		if (IsChildActor())
diff --git a/Engine/Source/Runtime/Engine/Private/Components/PrimitiveComponent.cpp b/Engine/Source/Runtime/Engine/Private/Components/PrimitiveComponent.cpp
--- a/Engine/Source/Runtime/Engine/Private/Components/PrimitiveComponent.cpp
+++ b/Engine/Source/Runtime/Engine/Private/Components/PrimitiveComponent.cpp
@@ -474,6 +474,8 @@ void UPrimitiveComponent::GetStreamingRenderAssetInfoWithNULLRemoval(FStreamingT
 	// Ignore components that are fully initialized but have no scene proxy (hidden primitive or non game primitive)
 	if (!IsRegistered() || !IsRenderStateCreated() || SceneProxy)
 	{
+		TRACE_CPUPROFILER_EVENT_SCOPE(UPrimitiveComponent::GetStreamingRenderAssetInfoWithNULLRemoval);
+
 		GetStreamingRenderAssetInfo(LevelContext, OutStreamingRenderAssets);
 		for (int32 Index = 0; Index < OutStreamingRenderAssets.Num(); Index++)
 		{
diff --git a/Engine/Source/Runtime/Engine/Private/Level.cpp b/Engine/Source/Runtime/Engine/Private/Level.cpp
--- a/Engine/Source/Runtime/Engine/Private/Level.cpp
+++ b/Engine/Source/Runtime/Engine/Private/Level.cpp
@@ -1169,6 +1169,8 @@ static void SortActorsHierarchy(TArray<AActor*>& Actors, UObject* Level)
 DECLARE_CYCLE_STAT(TEXT("Deferred Init Bodies"), STAT_DeferredUpdateBodies, STATGROUP_Physics);
 void ULevel::IncrementalUpdateComponents(int32 NumComponentsToUpdate, bool bRerunConstructionScripts, FRegisterComponentContext* Context)
 {
+	TRACE_CPUPROFILER_EVENT_SCOPE(ULevel::IncrementalUpdateComponents);
+
 	// A value of 0 means that we want to update all components.
 	if (NumComponentsToUpdate != 0)
 	{
diff --git a/Engine/Source/Runtime/Engine/Private/RawIndexBuffer.cpp b/Engine/Source/Runtime/Engine/Private/RawIndexBuffer.cpp
--- a/Engine/Source/Runtime/Engine/Private/RawIndexBuffer.cpp
+++ b/Engine/Source/Runtime/Engine/Private/RawIndexBuffer.cpp
@@ -372,6 +372,7 @@ void FRawStaticIndexBuffer::CopyRHIForStreaming(const FRawStaticIndexBuffer& Oth
 
 void FRawStaticIndexBuffer::InitRHI()
 {
+	TRACE_CPUPROFILER_EVENT_SCOPE(FRawStaticIndexBuffer::InitRHI);
 	IndexBufferRHI = CreateRHIBuffer_RenderThread();
 }
 
diff --git a/Engine/Source/Runtime/Engine/Private/Texture2D.cpp b/Engine/Source/Runtime/Engine/Private/Texture2D.cpp
--- a/Engine/Source/Runtime/Engine/Private/Texture2D.cpp
+++ b/Engine/Source/Runtime/Engine/Private/Texture2D.cpp
@@ -1215,6 +1215,8 @@ FVirtualTexture2DResource::~FVirtualTexture2DResource()
 
 void FVirtualTexture2DResource::InitRHI()
 {
+	TRACE_CPUPROFILER_EVENT_SCOPE(FVirtualTexture2DResource::InitRHI);
+
 	check(TextureOwner);
 
 	const uint32 MaxAnisotropy = FMath::Min<int32>(VirtualTextureScalability::GetMaxAnisotropy(), VTData->TileBorderSize);
diff --git a/Engine/Source/Runtime/Engine/Private/World.cpp b/Engine/Source/Runtime/Engine/Private/World.cpp
--- a/Engine/Source/Runtime/Engine/Private/World.cpp
+++ b/Engine/Source/Runtime/Engine/Private/World.cpp
@@ -2048,6 +2048,8 @@ void UWorld::ClearWorldComponents()
 
 void UWorld::UpdateWorldComponents(bool bRerunConstructionScripts, bool bCurrentLevelOnly, FRegisterComponentContext* Context)
 {
+	TRACE_CPUPROFILER_EVENT_SCOPE(UWorld::UpdateWorldComponents);
+
 	if ( !IsRunningDedicatedServer() )
 	{
 		if(!LineBatcher)
@@ -2127,6 +2129,8 @@ void UWorld::UpdateWorldComponents(bool bRerunConstructionScripts, bool bCurrent
 
 bool UWorld::UpdateCullDistanceVolumes(AActor* ActorToUpdate, UPrimitiveComponent* ComponentToUpdate)
 {
+	TRACE_CPUPROFILER_EVENT_SCOPE(UWorld::UpdateCullDistanceVolumes);
+
 	// Map that will store new max draw distance for every primitive
 	TMap<UPrimitiveComponent*,float> CompToNewMaxDrawMap;
 	bool bUpdatedDrawDistances = false;
diff --git a/Engine/Source/Runtime/RenderCore/Private/RenderResource.cpp b/Engine/Source/Runtime/RenderCore/Private/RenderResource.cpp
--- a/Engine/Source/Runtime/RenderCore/Private/RenderResource.cpp
+++ b/Engine/Source/Runtime/RenderCore/Private/RenderResource.cpp
@@ -76,6 +76,8 @@ void FRenderResource::ChangeFeatureLevel(ERHIFeatureLevel::Type NewFeatureLevel)
 
 void FRenderResource::InitResource()
 {
+	TRACE_CPUPROFILER_EVENT_SCOPE(FRenderResource::InitResource);
+	
 	check(IsInRenderingThread());
 	if (ListIndex == INDEX_NONE)
 	{
diff --git a/Engine/Source/Runtime/Renderer/Private/VT/VirtualTextureSystem.cpp b/Engine/Source/Runtime/Renderer/Private/VT/VirtualTextureSystem.cpp
--- a/Engine/Source/Runtime/Renderer/Private/VT/VirtualTextureSystem.cpp
+++ b/Engine/Source/Runtime/Renderer/Private/VT/VirtualTextureSystem.cpp
@@ -1886,6 +1886,7 @@ void FVirtualTextureSystem::SubmitPreMappedRequests(FRDGBuilder& GraphBuilder, E
 
 void FVirtualTextureSystem::SubmitRequests(FRDGBuilder& GraphBuilder, ERHIFeatureLevel::Type FeatureLevel, FMemStack& MemStack, FUniqueRequestList* RequestList, bool bAsync)
 {
+	TRACE_CPUPROFILER_EVENT_SCOPE(FVirtualTextureSystem::SubmitRequests);
 	LLM_SCOPE(ELLMTag::VirtualTextureSystem);
 
 	// Allocate space to hold the physical address we allocate for each page load (1 page per layer per request)
