UE4 comment update.

diff --git a/Engine/Source/Runtime/Engine/Private/Commandlets/PluginCommandlet.cpp b/Engine/Source/Runtime/Engine/Private/Commandlets/PluginCommandlet.cpp
--- a/Engine/Source/Runtime/Engine/Private/Commandlets/PluginCommandlet.cpp
+++ b/Engine/Source/Runtime/Engine/Private/Commandlets/PluginCommandlet.cpp
@@ -22,7 +22,7 @@ DEFINE_LOG_CATEGORY_STATIC(LogPluginCommandlet, Log, All);
  *
  * Example:
  *	Plugin Enable NetcodeUnitTest
- *	Plugin Enable NetcodeUnitTest,NUTUnrealEngine4
+ *	Plugin Enable NetcodeUnitTest,NUTUnrealEngine
  */
 
 UPluginCommandlet::UPluginCommandlet(const FObjectInitializer& ObjectInitializer)
