Matinee: Remove UE4 -> UE

diff --git a/Engine/Source/Runtime/Engine/Classes/Matinee/MatineeActor.h b/Engine/Source/Runtime/Engine/Classes/Matinee/MatineeActor.h
--- a/Engine/Source/Runtime/Engine/Classes/Matinee/MatineeActor.h
+++ b/Engine/Source/Runtime/Engine/Classes/Matinee/MatineeActor.h
@@ -152,7 +152,7 @@ class ENGINE_API AMatineeActor : public AActor
 	UPROPERTY(EditAnywhere, BlueprintReadWrite, Category = Cinematic)
 	uint32 bHideHud:1;
 
-	/** @todo UE4 matinee - shouldnt be directly editable.  Needs a nice interface in matinee */
+	/** @todo UE matinee - shouldnt be directly editable.  Needs a nice interface in matinee */
 	UPROPERTY(replicated)
 	TArray<struct FInterpGroupActorInfo> GroupActorInfos;
 
