Fixed issue where player viewpoint data may not be sent as updates in networked games, when a client player has an override controller on a vehicle they do not actually possess. Also left the door open for any similar cases where the PawnMovementComponent does not have a direct reference to the owner's camera manager.

Adaptation of PR #7626

diff --git a/Engine/Plugins/Experimental/ChaosVehiclesPlugin/Source/ChaosVehicles/Public/ChaosVehicleMovementComponent.h b/Engine/Plugins/Experimental/ChaosVehiclesPlugin/Source/ChaosVehicles/Public/ChaosVehicleMovementComponent.h
--- a/Engine/Plugins/Experimental/ChaosVehiclesPlugin/Source/ChaosVehicles/Public/ChaosVehicleMovementComponent.h
+++ b/Engine/Plugins/Experimental/ChaosVehiclesPlugin/Source/ChaosVehicles/Public/ChaosVehicleMovementComponent.h
@@ -1200,7 +1200,7 @@ protected:
 	// Setup
 
 	/** Get our controller */
-	AController* GetController() const;
+	virtual AController* GetController() const override;
 
 	/** Get the mesh this vehicle is tied to */
 	class UMeshComponent* GetMesh() const;
diff --git a/Engine/Plugins/Runtime/PhysXVehicles/Source/PhysXVehicles/Public/WheeledVehicleMovementComponent.h b/Engine/Plugins/Runtime/PhysXVehicles/Source/PhysXVehicles/Public/WheeledVehicleMovementComponent.h
--- a/Engine/Plugins/Runtime/PhysXVehicles/Source/PhysXVehicles/Public/WheeledVehicleMovementComponent.h
+++ b/Engine/Plugins/Runtime/PhysXVehicles/Source/PhysXVehicles/Public/WheeledVehicleMovementComponent.h
@@ -499,7 +499,7 @@ public:
 
 protected:
 
-	AController* GetController() const;
+	virtual AController* GetController() const override;
 
 	// replicated state of vehicle 
 	UPROPERTY(Transient, Replicated)
diff --git a/Engine/Source/Runtime/Engine/Classes/GameFramework/PawnMovementComponent.h b/Engine/Source/Runtime/Engine/Classes/GameFramework/PawnMovementComponent.h
--- a/Engine/Source/Runtime/Engine/Classes/GameFramework/PawnMovementComponent.h
+++ b/Engine/Source/Runtime/Engine/Classes/GameFramework/PawnMovementComponent.h
@@ -88,8 +88,11 @@ protected:
 	UPROPERTY(Transient, DuplicateTransient)
 	TObjectPtr<class APawn> PawnOwner;
 
+	/** Returns this component's associated controller, typically from the owning Pawn. May be null. May be overridden for special handling when the controller isn't paired with the Pawn that owns this component. */
+	virtual AController* GetController() const;
+
 	/**
-	 * Attempts to mark the PlayerCameraManager as dirty.
+	 * Attempts to mark the PlayerCameraManager as dirty, if the controller has one.
 	 * This will have no effect if called from the server.
 	 */
 	void MarkForClientCameraUpdate();
diff --git a/Engine/Source/Runtime/Engine/Private/Components/PawnMovementComponent.cpp b/Engine/Source/Runtime/Engine/Private/Components/PawnMovementComponent.cpp
--- a/Engine/Source/Runtime/Engine/Private/Components/PawnMovementComponent.cpp
+++ b/Engine/Source/Runtime/Engine/Private/Components/PawnMovementComponent.cpp
@@ -96,17 +96,25 @@ void UPawnMovementComponent::OnTeleported()
 	}
 }
 
+AController* UPawnMovementComponent::GetController() const
+{
+	if (PawnOwner)
+	{
+		return PawnOwner->GetController();
+	}
+
+	return nullptr;
+}
+
 void UPawnMovementComponent::MarkForClientCameraUpdate()
 {
-	if(PawnOwner)
+
+	if (APlayerController* PlayerController = Cast<APlayerController>(GetController()))
 	{
-		if (APlayerController* PlayerController = Cast<APlayerController>(PawnOwner->GetController()))
+		APlayerCameraManager* PlayerCameraManager = PlayerController->PlayerCameraManager;
+		if (PlayerCameraManager != nullptr && PlayerCameraManager->bUseClientSideCameraUpdates)
 		{
-			APlayerCameraManager* PlayerCameraManager = PlayerController->PlayerCameraManager;
-			if (PlayerCameraManager != nullptr && PlayerCameraManager->bUseClientSideCameraUpdates)
-			{
-				PlayerCameraManager->bShouldSendClientSideCameraUpdate = true;
-			}
+			PlayerCameraManager->bShouldSendClientSideCameraUpdate = true;
 		}
-	}	
+	}
 }
\ No newline at end of file
