Build fix

diff --git a/Engine/Source/Runtime/Engine/Public/ComputeFramework/ComputeKernelFromText.h b/Engine/Source/Runtime/Engine/Public/ComputeFramework/ComputeKernelFromText.h
--- a/Engine/Source/Runtime/Engine/Public/ComputeFramework/ComputeKernelFromText.h
+++ b/Engine/Source/Runtime/Engine/Public/ComputeFramework/ComputeKernelFromText.h
@@ -31,8 +31,6 @@ public:
 	UPROPERTY()
 	uint64 SourceHash = 0;
 
-#if WITH_EDITOR
-	FString KernelSourceText;
 
 	FString GetEntryPoint() const override
 	{
@@ -41,18 +39,22 @@ public:
 
 	FString GetSource() const override
 	{
+#if WITH_EDITOR
 		return KernelSourceText;
+#else
+		return {};
+#endif
 	}
 
 	uint64 GetSourceHashCode() const override
 	{
 		return SourceHash;
 	}
-#endif
 
 	void PostLoad() override;
 
 #if WITH_EDITOR
+	FString KernelSourceText;
 	FFilePath PrevSourceFile;
 
 	virtual void PreEditChange(FEditPropertyChain& PropertyAboutToChange) override;
