Fixed missing MarkAsSuccessfullyFinished() in FTextureStreamOut.

diff --git a/Engine/Source/Runtime/Engine/Private/Streaming/TextureStreamIn.cpp b/Engine/Source/Runtime/Engine/Private/Streaming/TextureStreamIn.cpp
--- a/Engine/Source/Runtime/Engine/Private/Streaming/TextureStreamIn.cpp
+++ b/Engine/Source/Runtime/Engine/Private/Streaming/TextureStreamIn.cpp
@@ -385,8 +385,11 @@ void FTextureStreamIn::FinalizeNewMips(const FContext& Context)
 	DECLARE_SCOPE_CYCLE_COUNTER(TEXT("FTextureStreamIn::FinalizeNewMips"), STAT_TextureStreamIn_FinalizeNewMips, STATGROUP_StreamingDetails);
 
 	// Execute
-	const bool bSuccessfullyFinalized = DoFinalizeNewMips(Context);
-	if (!bSuccessfullyFinalized)
+	if (DoFinalizeNewMips(Context))
+	{
+		MarkAsSuccessfullyFinished();
+	}
+	else
 	{
 		MarkAsCancelled();
 	}
@@ -394,12 +397,6 @@ void FTextureStreamIn::FinalizeNewMips(const FContext& Context)
 	// Schedule the next update step.
 	EThreadType NextThread = GetMipAllocatorThread(FTextureMipAllocator::ETickState::FinalizeMips);
 
-	if (bSuccessfullyFinalized && NextThread == TT_None)
-	{
-		// RHI resource swap has happened so treat the stream-in as successful
-		MarkAsSuccessfullyFinished();
-	}
-
 	if (NextThread != TT_None) // Loop on this state.
 	{
 		PushTask(Context, NextThread, SRA_UPDATE_CALLBACK(FinalizeNewMips), GetCancelThread(), SRA_UPDATE_CALLBACK(Cancel));
diff --git a/Engine/Source/Runtime/Engine/Private/Streaming/TextureStreamOut.cpp b/Engine/Source/Runtime/Engine/Private/Streaming/TextureStreamOut.cpp
--- a/Engine/Source/Runtime/Engine/Private/Streaming/TextureStreamOut.cpp
+++ b/Engine/Source/Runtime/Engine/Private/Streaming/TextureStreamOut.cpp
@@ -155,7 +155,11 @@ void FTextureStreamOut::FinalizeNewMips(const FContext& Context)
 	DECLARE_SCOPE_CYCLE_COUNTER(TEXT("FTextureStreamOut::FinalizeNewMips"), STAT_TextureStreamIn_FinalizeNewMips, STATGROUP_StreamingDetails);
 
 	// Execute
-	if (!DoFinalizeNewMips(Context))
+	if (DoFinalizeNewMips(Context))
+	{
+		MarkAsSuccessfullyFinished();
+	}
+	else
 	{
 		MarkAsCancelled();
 	}
