OFPA: Fix autosave for external actors

diff --git a/Engine/Source/Editor/UnrealEd/Private/FileHelpers.cpp b/Engine/Source/Editor/UnrealEd/Private/FileHelpers.cpp
--- a/Engine/Source/Editor/UnrealEd/Private/FileHelpers.cpp
+++ b/Engine/Source/Editor/UnrealEd/Private/FileHelpers.cpp
@@ -2866,7 +2866,8 @@ EAutosaveContentPackagesResult::Type FEditorFileUtils::AutosaveMapEx(const FStri
 						{
 							return EAutosaveContentPackagesResult::Failure;
 						}
-
+						// We saved an actor
+						bResult = true;
 						// Re-mark the package as dirty, because autosaving it will have cleared the dirty flag
 						ExternalPackage->MarkPackageDirty();
 					}
diff --git a/Engine/Source/Editor/UnrealEd/Private/PackageAutoSaver.cpp b/Engine/Source/Editor/UnrealEd/Private/PackageAutoSaver.cpp
--- a/Engine/Source/Editor/UnrealEd/Private/PackageAutoSaver.cpp
+++ b/Engine/Source/Editor/UnrealEd/Private/PackageAutoSaver.cpp
@@ -380,7 +380,7 @@ void FPackageAutoSaver::UpdateDirtyListsForPackage(UPackage* Pkg)
 		if ( !IsAutoSaving() )
 		{
 			// Add package into the appropriate list (map or content)
-			if (UWorld::FindWorldInPackage(Pkg))
+			if (UWorld::IsWorldOrExternalActorPackage(Pkg))
 			{
 				DirtyMapsForAutoSave.Add(Pkg);
 			}
diff --git a/Engine/Source/Runtime/Engine/Classes/Engine/World.h b/Engine/Source/Runtime/Engine/Classes/Engine/World.h
--- a/Engine/Source/Runtime/Engine/Classes/Engine/World.h
+++ b/Engine/Source/Runtime/Engine/Classes/Engine/World.h
@@ -3834,6 +3834,9 @@ public:
 	/** Given a package, locate the UWorld contained within if one exists */
 	static UWorld* FindWorldInPackage(UPackage* Package);
 
+	/** Given a package, return if package contains UWorld or External Actor */
+	static bool IsWorldOrExternalActorPackage(UPackage* Package);
+
 	/** If the specified package contains a redirector to a UWorld, that UWorld is returned. Otherwise, nullptr is returned. */
 	static UWorld* FollowWorldRedirectorInPackage(UPackage* Package, UObjectRedirector** OptionalOutRedirector = nullptr);
 
diff --git a/Engine/Source/Runtime/Engine/Private/World.cpp b/Engine/Source/Runtime/Engine/Private/World.cpp
--- a/Engine/Source/Runtime/Engine/Private/World.cpp
+++ b/Engine/Source/Runtime/Engine/Private/World.cpp
@@ -7219,6 +7219,30 @@ UWorld* UWorld::FindWorldInPackage(UPackage* Package)
 	return RetVal;
 }
 
+bool UWorld::IsWorldOrExternalActorPackage(UPackage* Package)
+{
+	bool bRetVal = false;
+	ForEachObjectWithPackage(Package, [&bRetVal](UObject* PackageObject)
+	{
+		if (PackageObject->IsA<UWorld>())
+		{
+			bRetVal = true;
+			return false;
+		}
+		else if (AActor* Actor = Cast<AActor>(PackageObject))
+		{
+			if (Actor->IsPackageExternal())
+			{
+				bRetVal = true;
+				return false;
+			}
+		}
+		return true;
+	}, false);
+
+	return bRetVal;
+}
+
 UWorld* UWorld::FollowWorldRedirectorInPackage(UPackage* Package, UObjectRedirector** OptionalOutRedirector)
 {
 	UWorld* RetVal = nullptr;
