Fix ctrl + left mouse click deselect intermittently not working due to outdated hit proxy cached data caused by CL 15373844.

diff --git a/Engine/Source/Runtime/Engine/Classes/Components/PrimitiveComponent.h b/Engine/Source/Runtime/Engine/Classes/Components/PrimitiveComponent.h
--- a/Engine/Source/Runtime/Engine/Classes/Components/PrimitiveComponent.h
+++ b/Engine/Source/Runtime/Engine/Classes/Components/PrimitiveComponent.h
@@ -538,10 +538,6 @@ public:
 	UPROPERTY()
 	uint8 bUseEditorCompositing:1;
 
-	/** Set to true while the editor is moving the component, which notifies the Renderer to track velocities even if the component is Static. */
-	UPROPERTY(Transient, DuplicateTransient)
-	uint8 bIsBeingMovedByEditor:1;
-
 	/** If true, this component will be rendered in the CustomDepth pass (usually used for outlines) */
 	UPROPERTY(EditAnywhere, AdvancedDisplay, BlueprintReadOnly, Category=Rendering, meta=(DisplayName = "Render CustomDepth Pass"))
 	uint8 bRenderCustomDepth:1;
diff --git a/Engine/Source/Runtime/Engine/Private/Components/PrimitiveComponent.cpp b/Engine/Source/Runtime/Engine/Private/Components/PrimitiveComponent.cpp
--- a/Engine/Source/Runtime/Engine/Private/Components/PrimitiveComponent.cpp
+++ b/Engine/Source/Runtime/Engine/Private/Components/PrimitiveComponent.cpp
@@ -360,7 +360,6 @@ UPrimitiveComponent::UPrimitiveComponent(const FObjectInitializer& ObjectInitial
 	bNeverDistanceCull = false;
 
 	bUseEditorCompositing = false;
-	bIsBeingMovedByEditor = false;
 
 	SetGenerateOverlapEvents(true);
 	bMultiBodyOverlap = false;
@@ -1591,8 +1590,10 @@ uint64 UPrimitiveComponent::GetHiddenEditorViews() const
 
 void UPrimitiveComponent::SetIsBeingMovedByEditor(bool bIsBeingMoved)
 {
-	bIsBeingMovedByEditor = bIsBeingMoved;
-	MarkRenderStateDirty();
+	if (SceneProxy)
+	{
+		SceneProxy->SetIsBeingMovedByEditor_GameThread(bIsBeingMoved);
+	}
 }
 #endif// WITH_EDITOR
 
diff --git a/Engine/Source/Runtime/Engine/Private/PrimitiveSceneProxy.cpp b/Engine/Source/Runtime/Engine/Private/PrimitiveSceneProxy.cpp
--- a/Engine/Source/Runtime/Engine/Private/PrimitiveSceneProxy.cpp
+++ b/Engine/Source/Runtime/Engine/Private/PrimitiveSceneProxy.cpp
@@ -171,7 +171,7 @@ FPrimitiveSceneProxy::FPrimitiveSceneProxy(const UPrimitiveComponent* InComponen
 ,	bSelectable(InComponent->bSelectable)
 ,	bHasPerInstanceHitProxies(InComponent->bHasPerInstanceHitProxies)
 ,	bUseEditorCompositing(InComponent->bUseEditorCompositing)
-,	bIsBeingMovedByEditor(InComponent->bIsBeingMovedByEditor)
+,	bIsBeingMovedByEditor(false)
 ,	bReceiveMobileCSMShadows(InComponent->bReceiveMobileCSMShadows)
 ,	bRenderCustomDepth(InComponent->bRenderCustomDepth)
 ,	bVisibleInSceneCaptureOnly(InComponent->bVisibleInSceneCaptureOnly)
@@ -735,6 +735,27 @@ void FPrimitiveSceneProxy::SetHiddenEdViews_RenderThread( uint64 InHiddenEditorV
 #endif
 }
 
+#if WITH_EDITOR
+void FPrimitiveSceneProxy::SetIsBeingMovedByEditor_GameThread(bool bIsBeingMoved)
+{
+	check(IsInGameThread());
+
+	FPrimitiveSceneProxy* PrimitiveSceneProxy = this;
+	ENQUEUE_RENDER_COMMAND(SetIsBeingMovedByEditor)(
+		[PrimitiveSceneProxy, bIsBeingMoved](FRHICommandListImmediate& RHICmdList)
+		{
+			PrimitiveSceneProxy->SetIsBeingMovedByEditor_RenderThread(bIsBeingMoved);
+		});
+}
+
+void FPrimitiveSceneProxy::SetIsBeingMovedByEditor_RenderThread(bool bIsBeingMoved)
+{
+	check(IsInRenderingThread());
+	bIsBeingMovedByEditor = bIsBeingMoved;
+}
+
+#endif
+
 void FPrimitiveSceneProxy::SetCollisionEnabled_GameThread(const bool bNewEnabled)
 {
 	check(IsInGameThread());
diff --git a/Engine/Source/Runtime/Engine/Public/PrimitiveSceneProxy.h b/Engine/Source/Runtime/Engine/Public/PrimitiveSceneProxy.h
--- a/Engine/Source/Runtime/Engine/Public/PrimitiveSceneProxy.h
+++ b/Engine/Source/Runtime/Engine/Public/PrimitiveSceneProxy.h
@@ -168,6 +168,14 @@ public:
 	 */
 	void SetHiddenEdViews_GameThread( uint64 InHiddenEditorViews );
 
+#if WITH_EDITOR
+	/**
+	 * Enqueue and update for the render thread to notify it that the editor is currently moving the owning component with gizmos.
+	 */
+	void SetIsBeingMovedByEditor_GameThread(bool bIsBeingMoved);
+	void SetIsBeingMovedByEditor_RenderThread(bool bIsBeingMoved);
+#endif	// WITH_EDITOR
+
 	/** @return True if the primitive is visible in the given View. */
 	ENGINE_API bool IsShown(const FSceneView* View) const;
 
