Replace 'NotPlaceable' by 'Abstract' to fix CIS

diff --git a/Engine/Source/Runtime/Engine/Public/ISMPartition/ISMPartitionActor.h b/Engine/Source/Runtime/Engine/Public/ISMPartition/ISMPartitionActor.h
--- a/Engine/Source/Runtime/Engine/Public/ISMPartition/ISMPartitionActor.h
+++ b/Engine/Source/Runtime/Engine/Public/ISMPartition/ISMPartitionActor.h
@@ -16,7 +16,7 @@ DECLARE_LOG_CATEGORY_EXTERN(LogISMPartition, Log, All);
 
 /** Actor base class for instance containers placed on a grid.
 	See UActorPartitionSubsystem. */
-UCLASS(NotPlaceable)
+UCLASS(Abstract)
 class ENGINE_API AISMPartitionActor : public APartitionActor
 {
 	GENERATED_UCLASS_BODY()
