CIS Fix: Add missing Abstract keyword

diff --git a/Engine/Source/Runtime/Engine/Public/WorldPartition/WorldPartitionRuntimeSpatialHashCell.h b/Engine/Source/Runtime/Engine/Public/WorldPartition/WorldPartitionRuntimeSpatialHashCell.h
--- a/Engine/Source/Runtime/Engine/Public/WorldPartition/WorldPartitionRuntimeSpatialHashCell.h
+++ b/Engine/Source/Runtime/Engine/Public/WorldPartition/WorldPartitionRuntimeSpatialHashCell.h
@@ -5,7 +5,7 @@
 #include "WorldPartition/WorldPartitionRuntimeCell.h"
 #include "WorldPartitionRuntimeSpatialHashCell.generated.h"
 
-UCLASS(Within = WorldPartition)
+UCLASS(Abstract, Within = WorldPartition)
 class UWorldPartitionRuntimeSpatialHashCell : public UWorldPartitionRuntimeCell
 {
 	GENERATED_UCLASS_BODY()
