Fixed compiler error on Android ES3.1

diff --git a/Engine/Source/Runtime/Engine/Private/Materials/HLSLMaterialDerivativeAutogen.cpp b/Engine/Source/Runtime/Engine/Private/Materials/HLSLMaterialDerivativeAutogen.cpp
--- a/Engine/Source/Runtime/Engine/Private/Materials/HLSLMaterialDerivativeAutogen.cpp
+++ b/Engine/Source/Runtime/Engine/Private/Materials/HLSLMaterialDerivativeAutogen.cpp
@@ -1360,7 +1360,8 @@ FString FMaterialDerivativeAutogen::GenerateUsedFunctions(FHLSLMaterialTranslato
 					Ret += TEXT("{") LINE_TERMINATOR;
 					Ret += TEXT("\t") + BaseName + TEXT(" Ret;") LINE_TERMINATOR;
 					Ret += TEXT("\tRet.Value = abs(A.Value);") LINE_TERMINATOR;
-					Ret += TEXT("\t") + FieldName + TEXT(" dFdA = (A.Value >= 0.0f ? 1.0f : -1.0f);") LINE_TERMINATOR;
+					Ret += TEXT("\t") + FieldName + TEXT(" One = 1.0f;") LINE_TERMINATOR;
+					Ret += TEXT("\t") + FieldName + TEXT(" dFdA = (A.Value >= 0.0f ? One : -One);") LINE_TERMINATOR;
 					Ret += TEXT("\tRet.Ddx = dFdA * A.Ddx;") LINE_TERMINATOR;
 					Ret += TEXT("\tRet.Ddy = dFdA * A.Ddy;") LINE_TERMINATOR;
 					Ret += TEXT("\treturn Ret;") LINE_TERMINATOR;
