Split GetAssetRegistryTags world partition specific data into GetAssetRegistryTagsForPackageData.
This fixes what 15732838 was trying to fix.

diff --git a/Engine/Source/Runtime/CoreUObject/Private/UObject/SavePackage/SavePackageUtilities.cpp b/Engine/Source/Runtime/CoreUObject/Private/UObject/SavePackage/SavePackageUtilities.cpp
--- a/Engine/Source/Runtime/CoreUObject/Private/UObject/SavePackage/SavePackageUtilities.cpp
+++ b/Engine/Source/Runtime/CoreUObject/Private/UObject/SavePackage/SavePackageUtilities.cpp
@@ -2278,6 +2278,7 @@ namespace AssetRegistry
 
 			TArray<UObject::FAssetRegistryTag> SourceTags;
 			Object->GetAssetRegistryTags(SourceTags);
+			Object->GetExternalActorExtendedAssetRegistryTags(SourceTags);
 
 			TArray<UObject::FAssetRegistryTag> Tags;
 			for (UObject::FAssetRegistryTag& SourceTag : SourceTags)
diff --git a/Engine/Source/Runtime/CoreUObject/Public/UObject/Object.h b/Engine/Source/Runtime/CoreUObject/Public/UObject/Object.h
--- a/Engine/Source/Runtime/CoreUObject/Public/UObject/Object.h
+++ b/Engine/Source/Runtime/CoreUObject/Public/UObject/Object.h
@@ -767,6 +767,13 @@ public:
 	 */
 	virtual void GetAssetRegistryTags(TArray<FAssetRegistryTag>& OutTags) const;
 
+	/**
+	 * Temporary interim solution to gather external actors asset registry data.
+	 *
+	 * @param	OutTags		A list of key-value pairs associated with this object and their types
+	 */
+	virtual void GetExternalActorExtendedAssetRegistryTags(TArray<FAssetRegistryTag>& OutTags) const {}
+
 	/** Gathers a list of asset registry tags for an FAssetData  */
 	void GetAssetRegistryTags(FAssetData& Out) const;
 
diff --git a/Engine/Source/Runtime/Engine/Classes/GameFramework/Actor.h b/Engine/Source/Runtime/Engine/Classes/GameFramework/Actor.h
--- a/Engine/Source/Runtime/Engine/Classes/GameFramework/Actor.h
+++ b/Engine/Source/Runtime/Engine/Classes/GameFramework/Actor.h
@@ -1849,6 +1849,7 @@ public:
 #if WITH_EDITOR
 	virtual bool Modify(bool bAlwaysMarkDirty = true) override;
 	virtual void GetAssetRegistryTags(TArray<FAssetRegistryTag>& OutTags) const override;
+	virtual void GetExternalActorExtendedAssetRegistryTags(TArray<FAssetRegistryTag>& OutTags) const override;
 	virtual bool NeedsLoadForTargetPlatform(const ITargetPlatform* TargetPlatform) const;
 	virtual void PreEditChange(FProperty* PropertyThatWillChange) override;
 	virtual bool CanEditChange(const FProperty* InProperty) const;
diff --git a/Engine/Source/Runtime/Engine/Private/Actor.cpp b/Engine/Source/Runtime/Engine/Private/Actor.cpp
--- a/Engine/Source/Runtime/Engine/Private/Actor.cpp
+++ b/Engine/Source/Runtime/Engine/Private/Actor.cpp
@@ -429,6 +429,20 @@ void AActor::GetAssetRegistryTags(TArray<FAssetRegistryTag>& OutTags) const
 {
 	Super::GetAssetRegistryTags(OutTags);
 
+	if (IsPackageExternal() && !IsChildActor())
+	{
+		if (ActorLabel.Len() > 0)
+		{
+			static FName NAME_ActorLabel(TEXT("ActorLabel"));
+			OutTags.Add(UObject::FAssetRegistryTag(NAME_ActorLabel, ActorLabel, UObject::FAssetRegistryTag::TT_Hidden));
+		}
+	}
+}
+
+void AActor::GetExternalActorExtendedAssetRegistryTags(TArray<FAssetRegistryTag>& OutTags) const
+{
+	Super::GetExternalActorExtendedAssetRegistryTags(OutTags);
+
 	if (IsPackageExternal() && !IsChildActor())
 	{
 		TUniquePtr<FWorldPartitionActorDesc> ActorDesc(CreateActorDesc());
@@ -442,12 +456,6 @@ void AActor::GetAssetRegistryTags(TArray<FAssetRegistryTag>& OutTags) const
 		const FString ActorMetaData = FBase64::Encode(SerializedData);
 		static FName NAME_ActorMetaData(TEXT("ActorMetaData"));
 		OutTags.Add(UObject::FAssetRegistryTag(NAME_ActorMetaData, ActorMetaData, UObject::FAssetRegistryTag::TT_Hidden));
-
-		if (ActorLabel.Len() > 0)
-		{
-			static FName NAME_ActorLabel(TEXT("ActorLabel"));
-			OutTags.Add(UObject::FAssetRegistryTag(NAME_ActorLabel, ActorLabel, UObject::FAssetRegistryTag::TT_Hidden));
-		}
 	}
 }
 
