At cook Empty Alwaysloaded cell when moving actors to PersistentLevel. This fixes bad softobjectpath remapping when cooking alwaysloaded actors.

diff --git a/Engine/Source/Runtime/Engine/Private/WorldPartition/WorldPartitionRuntimeLevelStreamingCell.cpp b/Engine/Source/Runtime/Engine/Private/WorldPartition/WorldPartitionRuntimeLevelStreamingCell.cpp
--- a/Engine/Source/Runtime/Engine/Private/WorldPartition/WorldPartitionRuntimeLevelStreamingCell.cpp
+++ b/Engine/Source/Runtime/Engine/Private/WorldPartition/WorldPartitionRuntimeLevelStreamingCell.cpp
@@ -127,6 +127,9 @@ void UWorldPartitionRuntimeLevelStreamingCell::MoveAlwaysLoadedContentToPersiste
 
 		UWorld* OuterWorld = GetOuterUWorldPartition()->GetTypedOuter<UWorld>();
 		FWorldPartitionLevelHelper::MoveExternalActorsToLevel(Packages, OuterWorld->PersistentLevel);
+
+		// Empty cell's package list (this ensures that no one can rely on cell's content).
+		Packages.Empty();
 	}
 }
 
