ASpatialHashRuntimeGridInfo - Now !EditorOnly

diff --git a/Engine/Source/Runtime/Engine/Public/WorldPartition/WorldPartitionRuntimeSpatialHash.h b/Engine/Source/Runtime/Engine/Public/WorldPartition/WorldPartitionRuntimeSpatialHash.h
--- a/Engine/Source/Runtime/Engine/Public/WorldPartition/WorldPartitionRuntimeSpatialHash.h
+++ b/Engine/Source/Runtime/Engine/Public/WorldPartition/WorldPartitionRuntimeSpatialHash.h
@@ -179,11 +179,6 @@ class ASpatialHashRuntimeGridInfo : public AInfo
 {
 	GENERATED_UCLASS_BODY()
 
-public:
-	/** Begin UObject override */
-	virtual bool IsEditorOnly() const override { return true; }
-	/** End UObject override */
-
 public:
 	UPROPERTY(EditAnywhere, Category=Settings)
 	FSpatialHashRuntimeGrid	GridSettings;
