WorldPartition - Fixed HLOD visibility in -game

diff --git a/Engine/Source/Runtime/Engine/Private/WorldPartition/HLOD/HLODSubsystem.cpp b/Engine/Source/Runtime/Engine/Private/WorldPartition/HLOD/HLODSubsystem.cpp
--- a/Engine/Source/Runtime/Engine/Private/WorldPartition/HLOD/HLODSubsystem.cpp
+++ b/Engine/Source/Runtime/Engine/Private/WorldPartition/HLOD/HLODSubsystem.cpp
@@ -41,27 +41,49 @@ bool UHLODSubsystem::ShouldCreateSubsystem(UObject* Outer) const
 
 void UHLODSubsystem::Initialize(FSubsystemCollectionBase& Collection)
 {
+	// Ensure the WorldPartitionSubsystem gets created before the HLODSubsystem
+	Collection.InitializeDependency<UWorldPartitionSubsystem>();
+
 	Super::Initialize(Collection);
 
 	if (GetWorld()->IsGameWorld())
 	{
-		TSet<const UWorldPartitionRuntimeCell*> StreamingCells;
-		GetWorld()->GetWorldPartition()->RuntimeHash->GetAllStreamingCells(StreamingCells, /*bIncludeDataLayers*/ true);
+		GetWorld()->GetSubsystem<UWorldPartitionSubsystem>()->OnWorldPartitionRegistered.AddUObject(this, &UHLODSubsystem::OnWorldPartitionRegistered);
+		GetWorld()->GetSubsystem<UWorldPartitionSubsystem>()->OnWorldPartitionUnregistered.AddUObject(this, &UHLODSubsystem::OnWorldPartitionUnregistered);
+	}
+}
 
-		// Build cell to HLOD mapping
-		for (const UWorldPartitionRuntimeCell* Cell : StreamingCells)
-		{
-			CellsHLODMapping.Emplace(Cell->GetFName());
-		}
+void UHLODSubsystem::OnWorldPartitionRegistered(UWorldPartition* InWorldPartition)
+{
+	TSet<const UWorldPartitionRuntimeCell*> StreamingCells;
+	InWorldPartition->RuntimeHash->GetAllStreamingCells(StreamingCells, /*bIncludeDataLayers*/ true);
+
+	check(!CellsHLODMapping.Contains(InWorldPartition));
+	TMap<FName, FCellHLODMapping>& CellsMapping = CellsHLODMapping.Emplace(InWorldPartition);
+
+	// Build cell to HLOD mapping
+	for (const UWorldPartitionRuntimeCell* Cell : StreamingCells)
+	{
+		CellsMapping.Emplace(Cell->GetFName());
 	}
 }
 
+void UHLODSubsystem::OnWorldPartitionUnregistered(UWorldPartition* InWorldPartition)
+{
+	check(CellsHLODMapping.Contains(InWorldPartition));
+	CellsHLODMapping.Remove(InWorldPartition);
+}
+
+
 void UHLODSubsystem::RegisterHLODActor(AWorldPartitionHLOD* InWorldPartitionHLOD)
 {
 	TRACE_CPUPROFILER_EVENT_SCOPE(UHLODSubsystem::RegisterHLODActor);
 
+	UWorldPartition* OwnerPartition = InWorldPartitionHLOD->GetWorld()->GetWorldPartition();
+	check(OwnerPartition);
+
 	FName CellName = InWorldPartitionHLOD->GetCellName();
-	FCellHLODMapping* CellHLODs = CellsHLODMapping.Find(CellName);
+	FCellHLODMapping* CellHLODs = CellsHLODMapping.FindChecked(OwnerPartition).Find(CellName);
 	if (CellHLODs)
 	{
 		CellHLODs->LoadedHLODs.Add(InWorldPartitionHLOD);
@@ -73,8 +95,11 @@ void UHLODSubsystem::UnregisterHLODActor(AWorldPartitionHLOD* InWorldPartitionHL
 {
 	TRACE_CPUPROFILER_EVENT_SCOPE(UHLODSubsystem::UnregisterHLODActor);
 
+	UWorldPartition* OwnerPartition = InWorldPartitionHLOD->GetWorld()->GetWorldPartition();
+	check(OwnerPartition);
+
 	FName CellName = InWorldPartitionHLOD->GetCellName();
-	FCellHLODMapping* CellHLODs = CellsHLODMapping.Find(CellName);
+	FCellHLODMapping* CellHLODs = CellsHLODMapping.FindChecked(OwnerPartition).Find(CellName);
 	if (CellHLODs)
 	{
 		int32 NumRemoved = CellHLODs->LoadedHLODs.Remove(InWorldPartitionHLOD);
@@ -84,7 +109,10 @@ void UHLODSubsystem::UnregisterHLODActor(AWorldPartitionHLOD* InWorldPartitionHL
 
 void UHLODSubsystem::OnCellShown(const UWorldPartitionRuntimeCell* InCell)
 {
-	FCellHLODMapping* CellHLODs = CellsHLODMapping.Find(InCell->GetFName());
+	UWorldPartition* OwnerPartition = InCell->GetTypedOuter<UWorldPartition>();
+	check(OwnerPartition);
+
+	FCellHLODMapping* CellHLODs = CellsHLODMapping.FindChecked(OwnerPartition).Find(InCell->GetFName());
 	if (CellHLODs)
 	{
 		CellHLODs->bIsCellVisible = true;
@@ -98,7 +126,10 @@ void UHLODSubsystem::OnCellShown(const UWorldPartitionRuntimeCell* InCell)
 
 void UHLODSubsystem::OnCellHidden(const UWorldPartitionRuntimeCell* InCell)
 {
-	FCellHLODMapping* CellHLODs = CellsHLODMapping.Find(InCell->GetFName());
+	UWorldPartition* OwnerPartition = InCell->GetTypedOuter<UWorldPartition>();
+	check(OwnerPartition);
+
+	FCellHLODMapping* CellHLODs = CellsHLODMapping.FindChecked(OwnerPartition).Find(InCell->GetFName());
 	if (CellHLODs)
 	{
 		CellHLODs->bIsCellVisible = false;
diff --git a/Engine/Source/Runtime/Engine/Private/WorldPartition/WorldPartitionSubsystem.cpp b/Engine/Source/Runtime/Engine/Private/WorldPartition/WorldPartitionSubsystem.cpp
--- a/Engine/Source/Runtime/Engine/Private/WorldPartition/WorldPartitionSubsystem.cpp
+++ b/Engine/Source/Runtime/Engine/Private/WorldPartition/WorldPartitionSubsystem.cpp
@@ -127,6 +127,7 @@ void UWorldPartitionSubsystem::RegisterWorldPartition(UWorldPartition* WorldPart
 	{
 		check(WorldPartition->IsInitialized());
 		RegisteredWorldPartitions.Add(WorldPartition);
+		OnWorldPartitionRegistered.Broadcast(WorldPartition);
 	}
 }
 
@@ -135,6 +136,7 @@ void UWorldPartitionSubsystem::UnregisterWorldPartition(UWorldPartition* WorldPa
 	if (ensure(RegisteredWorldPartitions.Contains(WorldPartition)))
 	{
 		RegisteredWorldPartitions.Remove(WorldPartition);
+		OnWorldPartitionUnregistered.Broadcast(WorldPartition);
 	}
 }
 
diff --git a/Engine/Source/Runtime/Engine/Public/WorldPartition/HLOD/HLODSubsystem.h b/Engine/Source/Runtime/Engine/Public/WorldPartition/HLOD/HLODSubsystem.h
--- a/Engine/Source/Runtime/Engine/Public/WorldPartition/HLOD/HLODSubsystem.h
+++ b/Engine/Source/Runtime/Engine/Public/WorldPartition/HLOD/HLODSubsystem.h
@@ -8,9 +8,9 @@
 
 #include "HLODSubsystem.generated.h"
 
-class UWorldPartitionSubsystem;
-class AWorldPartitionHLOD;
+class UWorldPartition;
 class UWorldPartitionRuntimeCell;
+class AWorldPartitionHLOD;
 
 /**
  * UHLODSubsystem
@@ -36,6 +36,8 @@ public:
 	void OnCellHidden(const UWorldPartitionRuntimeCell* InCell);
 	
 private:
+	void OnWorldPartitionRegistered(UWorldPartition* InWorldPartition);
+	void OnWorldPartitionUnregistered(UWorldPartition* InWorldPartition);
 
 private:
 	struct FCellHLODMapping
@@ -48,5 +50,5 @@ private:
 		}
 	};
 
-	TMap<FName, FCellHLODMapping> CellsHLODMapping;
+	TMap<UWorldPartition*, TMap<FName, FCellHLODMapping>> CellsHLODMapping;
 };
diff --git a/Engine/Source/Runtime/Engine/Public/WorldPartition/WorldPartitionSubsystem.h b/Engine/Source/Runtime/Engine/Public/WorldPartition/WorldPartitionSubsystem.h
--- a/Engine/Source/Runtime/Engine/Public/WorldPartition/WorldPartitionSubsystem.h
+++ b/Engine/Source/Runtime/Engine/Public/WorldPartition/WorldPartitionSubsystem.h
@@ -43,6 +43,10 @@ public:
 	virtual TStatId GetStatId() const override;
 	//~End FTickableGameObject
 
+	DECLARE_MULTICAST_DELEGATE_OneParam(FWorldPartitionRegisterDelegate, UWorldPartition*);
+	FWorldPartitionRegisterDelegate OnWorldPartitionRegistered;
+	FWorldPartitionRegisterDelegate OnWorldPartitionUnregistered;
+
 	UFUNCTION(BlueprintCallable, Category = Streaming)
 	bool IsStreamingCompleted(EWorldPartitionRuntimeCellState QueryState, const TArray<FWorldPartitionStreamingQuerySource>& QuerySources, bool bExactState) const;
 
