Implement CheckForErrors for partitioned worlds
This will run a set of checks on the actors, even if they aren't loaded in the editor.

diff --git a/Engine/Source/Editor/UnrealEd/Private/EditorServer.cpp b/Engine/Source/Editor/UnrealEd/Private/EditorServer.cpp
--- a/Engine/Source/Editor/UnrealEd/Private/EditorServer.cpp
+++ b/Engine/Source/Editor/UnrealEd/Private/EditorServer.cpp
@@ -118,6 +118,7 @@
 #include "InstancedFoliageActor.h"
 #include "IMovieSceneCapture.h"
 #include "MovieSceneCaptureModule.h"
+#include "WorldPartition/WorldPartition.h"
 
 #include "Kismet2/KismetEditorUtilities.h"
 #include "PropertyEditorModule.h"
@@ -3990,6 +3991,11 @@ bool UEditorEngine::Map_Check( UWorld* InWorld, const TCHAR* Str, FOutputDevice&
 		FMessageLog("MapCheck").Warning()->AddToken(FTextToken::Create(FText::Format(LOCTEXT("MapCheck_Message_TextureStreamingNeedsRebuild", "Texture streaming needs to be rebuilt ({0} Components, {1} Resource Refs), run 'Build Texture Streaming'."), InWorld->NumTextureStreamingUnbuiltComponents, InWorld->NumTextureStreamingDirtyResources)));
 	}
 
+	if (const UWorldPartition* WorldPartition = InWorld->GetWorldPartition())
+	{
+		WorldPartition->CheckForErrors();
+	}
+
 	GWarn->StatusUpdate( 0, ProgressDenominator, CheckMapLocText );
 
 	int32 LastUpdateCount = 0;
diff --git a/Engine/Source/Runtime/Engine/Classes/GameFramework/Actor.h b/Engine/Source/Runtime/Engine/Classes/GameFramework/Actor.h
--- a/Engine/Source/Runtime/Engine/Classes/GameFramework/Actor.h
+++ b/Engine/Source/Runtime/Engine/Classes/GameFramework/Actor.h
@@ -842,6 +842,7 @@ public:
 private:
 	friend class FActorDescList;
 	friend class UActorDescContainer;
+	friend class FWorldPartitionActorViewProxy;
 
 	/**
 	 * Creates an uninitialized actor descriptor. Can be overriden.
diff --git a/Engine/Source/Runtime/Engine/Private/WorldPartition/WorldPartition.cpp b/Engine/Source/Runtime/Engine/Private/WorldPartition/WorldPartition.cpp
--- a/Engine/Source/Runtime/Engine/Private/WorldPartition/WorldPartition.cpp
+++ b/Engine/Source/Runtime/Engine/Private/WorldPartition/WorldPartition.cpp
@@ -42,7 +42,7 @@
 
 DEFINE_LOG_CATEGORY(LogWorldPartition);
 
-#define LOCTEXT_NAMESPACE "WorldPartitionEditor"
+#define LOCTEXT_NAMESPACE "WorldPartition"
 
 #if WITH_EDITOR
 static FAutoConsoleCommand DumpActorDescs(
@@ -1108,6 +1108,14 @@ void UWorldPartition::DumpActorDescs(const FString& Path)
 	}
 }
 
+void UWorldPartition::CheckForErrors() const
+{
+	if (RuntimeHash)
+	{
+		RuntimeHash->CheckForErrors();
+	}
+}
+
 void UWorldPartition::RemapSoftObjectPath(FSoftObjectPath& ObjectPath)
 {
 	GetStreamingPolicy()->RemapSoftObjectPath(ObjectPath);
diff --git a/Engine/Source/Runtime/Engine/Private/WorldPartition/WorldPartitionActorDescViewProxy.cpp b/Engine/Source/Runtime/Engine/Private/WorldPartition/WorldPartitionActorDescViewProxy.cpp
new file mode 100644
--- /dev/null
+++ b/Engine/Source/Runtime/Engine/Private/WorldPartition/WorldPartitionActorDescViewProxy.cpp
@@ -0,0 +1,17 @@
+// Copyright Epic Games, Inc. All Rights Reserved.
+
+#include "WorldPartition/WorldPartitionActorDescViewProxy.h"
+
+#if WITH_EDITOR
+#include "WorldPartition/WorldPartitionActorDesc.h"
+
+FWorldPartitionActorViewProxy::FWorldPartitionActorViewProxy(const FWorldPartitionActorDesc* InActorDesc)
+	: FWorldPartitionActorDescView(InActorDesc)
+{
+	if (AActor* Actor = ActorDesc->GetActor())
+	{
+		CachedActorDesc = Actor->CreateActorDesc();
+		ActorDesc = CachedActorDesc.Get();
+	}
+}
+#endif
\ No newline at end of file
diff --git a/Engine/Source/Runtime/Engine/Private/WorldPartition/WorldPartitionRuntimeHash.cpp b/Engine/Source/Runtime/Engine/Private/WorldPartition/WorldPartitionRuntimeHash.cpp
--- a/Engine/Source/Runtime/Engine/Private/WorldPartition/WorldPartitionRuntimeHash.cpp
+++ b/Engine/Source/Runtime/Engine/Private/WorldPartition/WorldPartitionRuntimeHash.cpp
@@ -6,10 +6,16 @@
 #include "Engine/World.h"
 #include "Engine/LevelScriptBlueprint.h"
 #include "ActorReferencesUtils.h"
+
 #if WITH_EDITOR
 #include "WorldPartition/WorldPartitionHandle.h"
+#include "WorldPartition/WorldPartitionActorDescViewProxy.h"
+#include "Misc/UObjectToken.h"
+#include "Misc/MapErrors.h"
 #endif
 
+#define LOCTEXT_NAMESPACE "WorldPartition"
+
 UWorldPartitionRuntimeHash::UWorldPartitionRuntimeHash(const FObjectInitializer& ObjectInitializer)
 	: Super(ObjectInitializer)
 {}
@@ -132,4 +138,91 @@ void UWorldPartitionRuntimeHash::ChangeActorDescViewGridPlacement(FWorldPartitio
 		UE_LOG(LogWorldPartition, Verbose, TEXT("Actor '%s' grid placement changed to %d"), *ActorDescView.GetActorLabel().ToString(), (int32)GridPlacement);
 	}
 }
+
+void UWorldPartitionRuntimeHash::CheckForErrors() const
+{
+	auto GetActorLabel = [](const FWorldPartitionActorDescView& ActorDescView) -> FString
+	{
+		const FName ActorLabel = ActorDescView.GetActorLabel();
+		if (!ActorLabel.IsNone())
+		{
+			return ActorLabel.ToString();
+		}
+
+		const FString ActorPath = ActorDescView.GetActorPath().ToString();
+
+		FString SubObjectName;
+		FString SubObjectContext;
+		if (FString(ActorPath).Split(TEXT("."), &SubObjectContext, &SubObjectName))
+		{
+			return SubObjectName;
+		}
+
+		return ActorPath;
+	};
+
+	TMap<FGuid, FWorldPartitionActorViewProxy> ActorDescList;
+	for (UActorDescContainer::TConstIterator<> ActorDescIt(GetOuterUWorldPartition()); ActorDescIt; ++ActorDescIt)
+	{
+		ActorDescList.Emplace(ActorDescIt->GetGuid(), *ActorDescIt);
+	}
+
+	for (auto& ActorDescListPair: ActorDescList)
+	{
+		const FWorldPartitionActorViewProxy& ActorDescView = ActorDescListPair.Value;
+
+		if (!ActorDescView.GetActorIsEditorOnly())
+		{
+			for (const FGuid ActorDescRefGuid : ActorDescView.GetReferences())
+			{
+				if (const FWorldPartitionActorViewProxy* ActorDescRefView = ActorDescList.Find(ActorDescRefGuid))
+				{
+					const bool bIsActorDescAlwaysLoaded = ActorDescView.GetGridPlacement() == EActorGridPlacement::AlwaysLoaded;
+					const bool bIsActorDescRefAlwaysLoaded = ActorDescRefView->GetGridPlacement() == EActorGridPlacement::AlwaysLoaded;
+
+					if (bIsActorDescAlwaysLoaded != bIsActorDescRefAlwaysLoaded)
+					{
+						const FText StreamedActor(LOCTEXT("MapCheck_WorldPartition_StreamedActor", "Streamed actor"));
+						const FText AlwaysLoadedActor(LOCTEXT("MapCheck_WorldPartition_AlwaysLoadedActor", "Always loaded actor"));
+
+						TSharedRef<FTokenizedMessage> Error = FMessageLog("MapCheck").Warning()
+							->AddToken(FTextToken::Create(bIsActorDescAlwaysLoaded ? AlwaysLoadedActor : StreamedActor))
+							->AddToken(FAssetNameToken::Create(GetActorLabel(ActorDescView)))
+							->AddToken(FTextToken::Create(LOCTEXT("MapCheck_WorldPartition_References", "references")))
+							->AddToken(FTextToken::Create(bIsActorDescRefAlwaysLoaded ? AlwaysLoadedActor : StreamedActor)								)
+							->AddToken(FAssetNameToken::Create(GetActorLabel(*ActorDescRefView)))
+							->AddToken(FMapErrorToken::Create(FName(TEXT("WorldPartition_StreamedActorReferenceAlwaysLoadedActor_CheckForErrors"))));
+					}
+
+					TArray<FName> ActorDescLayers = ActorDescView.GetDataLayers();
+					ActorDescLayers.Sort([](const FName& A, const FName& B) { return A.FastLess(B); });
+
+					TArray<FName> ActorDescRefLayers = ActorDescRefView->GetDataLayers();
+					ActorDescRefLayers.Sort([](const FName& A, const FName& B) { return A.FastLess(B); });
+
+					if (ActorDescLayers != ActorDescRefLayers)
+					{
+						TSharedRef<FTokenizedMessage> Error = FMessageLog("MapCheck").Error()
+							->AddToken(FTextToken::Create(LOCTEXT("MapCheck_WorldPartition_Actor", "Actor")))
+							->AddToken(FAssetNameToken::Create(GetActorLabel(ActorDescView)))
+							->AddToken(FTextToken::Create(LOCTEXT("MapCheck_WorldPartition_ReferenceActorInOtherDataLayers", "references an actor in a different set of data layers")))
+							->AddToken(FAssetNameToken::Create(GetActorLabel(*ActorDescRefView)))
+							->AddToken(FMapErrorToken::Create(FName(TEXT("WorldPartition_ActorReferenceActorInAnotherDataLayer_CheckForErrors"))));
+					}
+				}
+				else
+				{
+					TSharedRef<FTokenizedMessage> Error = FMessageLog("MapCheck").Warning()
+						->AddToken(FTextToken::Create(LOCTEXT("MapCheck_WorldPartition_Actor", "Actor")))
+						->AddToken(FAssetNameToken::Create(GetActorLabel(ActorDescView)))
+						->AddToken(FTextToken::Create(LOCTEXT("MapCheck_WorldPartition_HaveMissingRefsTo", "have missing references to")))
+						->AddToken(FAssetNameToken::Create(GetActorLabel(*ActorDescRefView)))
+						->AddToken(FMapErrorToken::Create(FName(TEXT("WorldPartition_MissingActorReference_CheckForErrors"))));
+				}
+			}
+		}
+	}
+}
 #endif
+
+#undef LOCTEXT_NAMESPACE
\ No newline at end of file
diff --git a/Engine/Source/Runtime/Engine/Public/WorldPartition/WorldPartition.h b/Engine/Source/Runtime/Engine/Public/WorldPartition/WorldPartition.h
--- a/Engine/Source/Runtime/Engine/Public/WorldPartition/WorldPartition.h
+++ b/Engine/Source/Runtime/Engine/Public/WorldPartition/WorldPartition.h
@@ -154,6 +154,8 @@ public:
 	// Debugging Methods
 	void DrawRuntimeHashPreview();
 	void DumpActorDescs(const FString& Path);
+
+	void CheckForErrors() const;
 #endif
 
 public:
diff --git a/Engine/Source/Runtime/Engine/Public/WorldPartition/WorldPartitionActorDescViewProxy.h b/Engine/Source/Runtime/Engine/Public/WorldPartition/WorldPartitionActorDescViewProxy.h
new file mode 100644
--- /dev/null
+++ b/Engine/Source/Runtime/Engine/Public/WorldPartition/WorldPartitionActorDescViewProxy.h
@@ -0,0 +1,18 @@
+// Copyright Epic Games, Inc. All Rights Reserved.
+#pragma once
+
+#include "WorldPartitionActorDescView.h"
+
+#if WITH_EDITOR
+/**
+ * A view on top of an actor desc, which reverts to the actor if its loaded.
+ */
+class FWorldPartitionActorViewProxy : public FWorldPartitionActorDescView
+{
+public:
+	FWorldPartitionActorViewProxy(const FWorldPartitionActorDesc* InActorDesc);
+
+private:
+	TUniquePtr<FWorldPartitionActorDesc> CachedActorDesc;
+};
+#endif
\ No newline at end of file
diff --git a/Engine/Source/Runtime/Engine/Public/WorldPartition/WorldPartitionRuntimeHash.h b/Engine/Source/Runtime/Engine/Public/WorldPartition/WorldPartitionRuntimeHash.h
--- a/Engine/Source/Runtime/Engine/Public/WorldPartition/WorldPartitionRuntimeHash.h
+++ b/Engine/Source/Runtime/Engine/Public/WorldPartition/WorldPartitionRuntimeHash.h
@@ -31,6 +31,7 @@ class ENGINE_API UWorldPartitionRuntimeHash : public UObject
 	virtual bool GenerateHLOD(ISourceControlHelper* SourceControlHelper, bool bCreateActorsOnly) { return false; }
 	virtual bool GenerateNavigationData() { return false; }
 	virtual FName GetActorRuntimeGrid(const AActor* Actor) const { return NAME_None; }
+	virtual void CheckForErrors() const;
 	virtual void DrawPreview() const {}
 
 	// PIE/Game methods
