WorldPartition: Fix Clustering not assigning the correct GridPlacement to referenced Actor

diff --git a/Engine/Source/Runtime/Engine/Private/WorldPartition/RuntimeSpatialHash/RuntimeSpatialHashHLOD.cpp b/Engine/Source/Runtime/Engine/Private/WorldPartition/RuntimeSpatialHash/RuntimeSpatialHashHLOD.cpp
--- a/Engine/Source/Runtime/Engine/Private/WorldPartition/RuntimeSpatialHash/RuntimeSpatialHashHLOD.cpp
+++ b/Engine/Source/Runtime/Engine/Private/WorldPartition/RuntimeSpatialHash/RuntimeSpatialHashHLOD.cpp
@@ -561,7 +561,7 @@ bool UWorldPartitionRuntimeSpatialHash::GenerateHLOD(ISourceControlHelper* Sourc
 		for (const FGuid& HLODActorGuid : GridsHLODActors[HLODGridName])
 		{
 			const FWorldPartitionActorDescView& HLODActorDescView = MainContainerInstance.GetActorDescView(HLODActorGuid);
-			FActorCluster& NewHLODActorCluster = HLODActorClusters.Emplace_GetRef(WorldPartition->GetWorld(), HLODActorDescView, HLODActorDescView.GetGridPlacement());
+			FActorCluster& NewHLODActorCluster = HLODActorClusters.Emplace_GetRef(WorldPartition->GetWorld(), HLODActorDescView);
 			FActorClusterInstance& NewHLODActorClusterInstance = HLODActorClusterInstances.Emplace_GetRef(&NewHLODActorCluster, &MainContainerInstance);
 			HLODActorClusterInstancePtrs.Add(&NewHLODActorClusterInstance);
 		};
diff --git a/Engine/Source/Runtime/Engine/Private/WorldPartition/WorldPartitionActorCluster.cpp b/Engine/Source/Runtime/Engine/Private/WorldPartition/WorldPartitionActorCluster.cpp
--- a/Engine/Source/Runtime/Engine/Private/WorldPartition/WorldPartitionActorCluster.cpp
+++ b/Engine/Source/Runtime/Engine/Private/WorldPartition/WorldPartitionActorCluster.cpp
@@ -34,8 +34,8 @@ TSet<const UDataLayer*> GetDataLayers(UWorld* InWorld, const LayerNameContainer&
 	return DataLayers;
 }
 
-FActorCluster::FActorCluster(UWorld* InWorld, const FWorldPartitionActorDescView& InActorDescView, EActorGridPlacement InGridPlacement)
-	: GridPlacement(InGridPlacement)
+FActorCluster::FActorCluster(UWorld* InWorld, const FWorldPartitionActorDescView& InActorDescView)
+	: GridPlacement(InActorDescView.GetGridPlacement())
 	, RuntimeGrid(InActorDescView.GetRuntimeGrid())
 	, Bounds(InActorDescView.GetBounds())
 {
@@ -188,7 +188,7 @@ const FWorldPartitionActorDescView& FActorInstance::GetActorDescView() const
 	return ContainerInstance->GetActorDescView(Actor);
 }
 
-void CreateActorCluster(const FWorldPartitionActorDescView& ActorDescView, EActorGridPlacement GridPlacement, TMap<FGuid, FActorCluster*>& ActorToActorCluster, TSet<FActorCluster*>& ActorClustersSet, const FActorContainerInstance& ContainerInstance)
+void CreateActorCluster(const FWorldPartitionActorDescView& ActorDescView, TMap<FGuid, FActorCluster*>& ActorToActorCluster, TSet<FActorCluster*>& ActorClustersSet, const FActorContainerInstance& ContainerInstance)
 {
 	// Don't include references from editor-only actors
 	if (!ActorDescView.GetActorIsEditorOnly())
@@ -200,7 +200,7 @@ void CreateActorCluster(const FWorldPartitionActorDescView& ActorDescView, EActo
 		FActorCluster* ActorCluster = ActorToActorCluster.FindRef(ActorGuid);
 		if (!ActorCluster)
 		{
-			ActorCluster = new FActorCluster(World, ActorDescView, GridPlacement);
+			ActorCluster = new FActorCluster(World, ActorDescView);
 			ActorClustersSet.Add(ActorCluster);
 			ActorToActorCluster.Add(ActorGuid, ActorCluster);
 		}
@@ -230,7 +230,7 @@ void CreateActorCluster(const FWorldPartitionActorDescView& ActorDescView, EActo
 					else
 					{
 						// Put Reference in Actor's cluster
-						ActorCluster->Add(FActorCluster(World, *ReferenceActorDescView, GridPlacement));
+						ActorCluster->Add(FActorCluster(World, *ReferenceActorDescView));
 					}
 
 					// Map its cluster
@@ -331,11 +331,10 @@ const TArray<FActorCluster>& FActorClusterContext::CreateActorClustersImpl(const
 	for (auto& ActorDescViewPair : ContainerInstance.ActorDescViewMap)
 	{
 		const FWorldPartitionActorDescView& ActorDescView = ActorDescViewPair.Value;
-		EActorGridPlacement GridPlacement = ActorDescView.GetGridPlacement();
 
 		if (!FilterPredicate.IsSet() || FilterPredicate.GetValue()(ActorDescView))
 		{
-			CreateActorCluster(ActorDescView, GridPlacement, ActorToActorCluster, ActorClustersSet, ContainerInstance);
+			CreateActorCluster(ActorDescView, ActorToActorCluster, ActorClustersSet, ContainerInstance);
 		}
 	}
 				
diff --git a/Engine/Source/Runtime/Engine/Public/WorldPartition/WorldPartitionActorCluster.h b/Engine/Source/Runtime/Engine/Public/WorldPartition/WorldPartitionActorCluster.h
--- a/Engine/Source/Runtime/Engine/Public/WorldPartition/WorldPartitionActorCluster.h
+++ b/Engine/Source/Runtime/Engine/Public/WorldPartition/WorldPartitionActorCluster.h
@@ -27,7 +27,7 @@ struct FActorCluster
 	TSet<const UDataLayer*>		DataLayers;
 	FDataLayersID				DataLayersID;
 
-	FActorCluster(UWorld* InWorld, const FWorldPartitionActorDescView& InActorDescView, EActorGridPlacement InGridPlacement);
+	FActorCluster(UWorld* InWorld, const FWorldPartitionActorDescView& InActorDescView);
 	void Add(const FActorCluster& InActorCluster);
 
 };
