Disabled Lumen in editor preview windows, even the ones that use EnableAdvancedFeatures(), since Lumen Scene uses fixed size atlases

diff --git a/Engine/Source/Runtime/Engine/Public/ShowFlags.h b/Engine/Source/Runtime/Engine/Public/ShowFlags.h
--- a/Engine/Source/Runtime/Engine/Public/ShowFlags.h
+++ b/Engine/Source/Runtime/Engine/Public/ShowFlags.h
@@ -245,6 +245,8 @@ struct FEngineShowFlags
 		SetLightShafts(true);
 		SetPostProcessMaterial(true);
 		SetDistanceFieldAO(true);
+		SetLumenGlobalIllumination(false);
+		SetLumenReflections(false);
 	}
 
 	bool IsVisualizeCalibrationEnabled() const
