Fix bIsPersistent not being propagated properly (UE-110824)

- Log more information about the material DDC.
- Also removed some leftover code that wasn't necessary.

diff --git a/Engine/Source/Editor/UnrealEd/Private/PreviewMaterial.cpp b/Engine/Source/Editor/UnrealEd/Private/PreviewMaterial.cpp
--- a/Engine/Source/Editor/UnrealEd/Private/PreviewMaterial.cpp
+++ b/Engine/Source/Editor/UnrealEd/Private/PreviewMaterial.cpp
@@ -225,6 +225,7 @@ public:
 	 * Should shaders compiled for this material be saved to disk?
 	 */
 	virtual bool IsPersistent() const override { return false; }
+	virtual FString GetAssetName() const override { return FString::Printf(TEXT("Preview:%s"), *FMaterialResource::GetAssetName()); }
 };
 
 /** Implementation of Preview Material functions*/
diff --git a/Engine/Source/Runtime/Engine/Private/Materials/MaterialShader.cpp b/Engine/Source/Runtime/Engine/Private/Materials/MaterialShader.cpp
--- a/Engine/Source/Runtime/Engine/Private/Materials/MaterialShader.cpp
+++ b/Engine/Source/Runtime/Engine/Private/Materials/MaterialShader.cpp
@@ -165,7 +165,7 @@ FString GetBlendModeString(EBlendMode BlendMode)
 
 #if WITH_EDITOR
 /** Creates a string key for the derived data cache given a shader map id. */
-static FString GetMaterialShaderMapKeyString(const FMaterialShaderMapId& ShaderMapId, EShaderPlatform Platform, const ITargetPlatform* TargetPlatform)
+static FString GetMaterialShaderMapKeyString(const FMaterialShaderMapId& ShaderMapId, EShaderPlatform Platform)
 {
 	FName Format = LegacyShaderPlatformToShaderFormat(Platform);
 	FString ShaderMapKeyString = Format.ToString() + TEXT("_") + FString(FString::FromInt(GetTargetPlatformManagerRef().ShaderFormatVersion(Format))) + TEXT("_");
@@ -177,6 +177,17 @@ static FString GetMaterialShaderMapKeyString(const FMaterialShaderMapId& ShaderM
 
 	return FDerivedDataCacheInterface::BuildCacheKey(TEXT("MATSM"), MATERIALSHADERMAP_DERIVEDDATA_VER, *ShaderMapKeyString);
 }
+
+static FString FSHA1_HashString(const FString& Key)
+{
+	FSHAHash DDCKeyHash;
+	FSHA1 HashState;
+	HashState.UpdateWithString(*Key, Key.Len());
+	HashState.Final();
+	HashState.GetHash(&DDCKeyHash.Hash[0]);
+	return DDCKeyHash.ToString();
+}
+
 #endif // WITH_EDITOR
 
 /** Called for every material shader to update the appropriate stats. */
@@ -1167,15 +1178,8 @@ void FMaterialShaderMap::LoadFromDerivedDataCache(const FMaterial* Material, con
 			COOK_STAT(auto Timer = MaterialShaderCookStats::UsageStats.TimeSyncWork());
 
 			TArray<uint8> CachedData;
-			const FString DataKey = GetMaterialShaderMapKeyString(ShaderMapId, InPlatform, TargetPlatform);
-			{
-				FSHAHash DDCKeyHash;
-				FSHA1 HashState;
-				HashState.UpdateWithString(*DataKey, DataKey.Len());
-				HashState.Final();
-				HashState.GetHash(&DDCKeyHash.Hash[0]);
-				OutDDCKeyDesc = DDCKeyHash.ToString();
-			}
+			const FString DataKey = GetMaterialShaderMapKeyString(ShaderMapId, InPlatform);
+			OutDDCKeyDesc = FSHA1_HashString(DataKey);
 
 			bool CheckCache = true;
 
@@ -1206,7 +1210,6 @@ void FMaterialShaderMap::LoadFromDerivedDataCache(const FMaterial* Material, con
 				InOutShaderMap->Serialize(Ar);
 				//InOutShaderMap->RegisterSerializedShaders(false);
 		
-				const FString InDataKey = GetMaterialShaderMapKeyString(InOutShaderMap->GetShaderMapId(), InPlatform, TargetPlatform);
 				checkSlow(InOutShaderMap->GetShaderMapId() == ShaderMapId);
 
 				// Register in the global map
@@ -1223,14 +1226,16 @@ void FMaterialShaderMap::LoadFromDerivedDataCache(const FMaterial* Material, con
 	}
 }
 
-void FMaterialShaderMap::SaveToDerivedDataCache(const ITargetPlatform* TargetPlatform)
+void FMaterialShaderMap::SaveToDerivedDataCache()
 {
 	COOK_STAT(auto Timer = MaterialShaderCookStats::UsageStats.TimeSyncWork());
 	TArray<uint8> SaveData;
 	FMemoryWriter Ar(SaveData, true);
 	Serialize(Ar);
 
-	GetDerivedDataCacheRef().Put(*GetMaterialShaderMapKeyString(ShaderMapId, GetShaderPlatform(), TargetPlatform), SaveData, FStringView(GetFriendlyName()));
+	FString DataKey = GetMaterialShaderMapKeyString(ShaderMapId, GetShaderPlatform());
+	UE_LOG(LogMaterial, Verbose, TEXT("Saving material DDC for key hash %s"), *FSHA1_HashString(DataKey));
+	GetDerivedDataCacheRef().Put(*DataKey, SaveData, FStringView(GetFriendlyName()));
 	COOK_STAT(Timer.AddMiss(SaveData.Num()));
 }
 #endif // WITH_EDITOR
@@ -1446,6 +1451,8 @@ void FMaterialShaderMap::ReleaseCompilingId()
 void FMaterialShaderMap::AddCompilingDependency(FMaterial* Material)
 {
 	CompilingMaterialDependencies.AddUnique(Material);
+	// if any of our dependencies is persistent, we're persistent
+	bIsPersistent |= Material->IsPersistent();
 }
 
 void FMaterialShaderMap::RemoveCompilingDependency(FMaterial* Material)
@@ -2556,7 +2563,7 @@ FMaterialShaderMap::FMaterialShaderMap() :
 	bRegistered(false),
 	bCompilationFinalized(true),
 	bCompiledSuccessfully(true),
-	bIsPersistent(true)
+	bIsPersistent(false)
 {
 	checkSlow(IsInGameThread() || IsAsyncLoading());
 #if ALLOW_SHADERMAP_DEBUG_DATA
@@ -2618,6 +2625,7 @@ FMaterialShaderMap* FMaterialShaderMap::AcquireFinalizedClone()
 	Clone->ShaderMapId = ShaderMapId;
 	Clone->bCompilationFinalized = bCompilationFinalized;
 	Clone->bCompiledSuccessfully = bCompiledSuccessfully;
+	Clone->bIsPersistent = bIsPersistent;
 	Clone->AssignCopy(*this);
 #if WITH_EDITOR
 	Clone->AssociateWithAssets(GetAssociatedAssets());
diff --git a/Engine/Source/Runtime/Engine/Private/Materials/MaterialShared.cpp b/Engine/Source/Runtime/Engine/Private/Materials/MaterialShared.cpp
--- a/Engine/Source/Runtime/Engine/Private/Materials/MaterialShared.cpp
+++ b/Engine/Source/Runtime/Engine/Private/Materials/MaterialShared.cpp
@@ -2149,6 +2149,10 @@ bool FMaterial::CacheShaders(const FMaterialShaderMapId& ShaderMapId, EShaderPla
 			TRefCountPtr<FMaterialShaderMap> LoadedShaderMap;
 			FMaterialShaderMap::LoadFromDerivedDataCache(this, ShaderMapId, Platform, TargetPlatform, LoadedShaderMap, DDCKeyHash);
 			ShaderMap = LoadedShaderMap;
+			if (LoadedShaderMap)
+			{
+				UE_LOG(LogMaterial, Verbose, TEXT("Loaded shaders for %s from DDC (key hash: %s)"), *GetAssetName(), *DDCKeyHash);
+			}
 		}
 
 		check(!ShaderMap || ShaderMap->GetFrozenContentSize() > 0u);
@@ -3126,7 +3130,7 @@ void FMaterialRenderProxy::CacheUniformExpressions_GameThread(bool bRecreateUnif
 {
 	if (FApp::CanEverRender())
 	{
-		UE_LOG(LogMaterial, Verbose, TEXT("Caching uniform expressions for material: %s"), *GetFriendlyName());
+		UE_LOG(LogMaterial, VeryVerbose, TEXT("Caching uniform expressions for material: %s"), *GetFriendlyName());
 
 		FMaterialRenderProxy* RenderProxy = this;
 		ENQUEUE_RENDER_COMMAND(FCacheUniformExpressionsCommand)(
diff --git a/Engine/Source/Runtime/Engine/Private/ShaderCompiler/ShaderCompiler.cpp b/Engine/Source/Runtime/Engine/Private/ShaderCompiler/ShaderCompiler.cpp
--- a/Engine/Source/Runtime/Engine/Private/ShaderCompiler/ShaderCompiler.cpp
+++ b/Engine/Source/Runtime/Engine/Private/ShaderCompiler/ShaderCompiler.cpp
@@ -3655,7 +3655,7 @@ void FShaderCompilingManager::ProcessCompiledShaderMaps(
 					ShaderMapToUseForRendering->bCompilationFinalized = true;
 					if (ShaderMapToUseForRendering->bIsPersistent)
 					{
-						ShaderMapToUseForRendering->SaveToDerivedDataCache(CompilingShaderMap->PendingCompilerEnvironment->TargetPlatform);
+						ShaderMapToUseForRendering->SaveToDerivedDataCache();
 					}
 				}
 
diff --git a/Engine/Source/Runtime/Engine/Public/MaterialShared.h b/Engine/Source/Runtime/Engine/Public/MaterialShared.h
--- a/Engine/Source/Runtime/Engine/Public/MaterialShared.h
+++ b/Engine/Source/Runtime/Engine/Public/MaterialShared.h
@@ -1137,7 +1137,7 @@ public:
 
 #if WITH_EDITOR
 	/** Saves this shader map to the derived data cache. */
-	void SaveToDerivedDataCache(const ITargetPlatform* TargetPlatform);
+	void SaveToDerivedDataCache();
 #endif
 
 	/** Backs up any FShaders in this shader map to memory through serialization and clears FShader references. */
