Add informative log when encountering a special condition causing a texture to be rebuilt

diff --git a/Engine/Source/Runtime/Engine/Private/TextureDerivedDataTask.cpp b/Engine/Source/Runtime/Engine/Private/TextureDerivedDataTask.cpp
--- a/Engine/Source/Runtime/Engine/Private/TextureDerivedDataTask.cpp
+++ b/Engine/Source/Runtime/Engine/Private/TextureDerivedDataTask.cpp
@@ -480,10 +480,20 @@ void FTextureCacheDerivedDataWorker::DoWork()
 		if (bForDDC)
 		{
 			bSucceeded = DerivedData->TryLoadMips(0, nullptr, &Texture);
+
+			if (!bSucceeded)
+			{
+				UE_LOG(LogTexture, Display, TEXT("Texture %s is missing mips. The texture will be rebuilt."), *Texture.GetFullName());
+			}
 		}
 		else if (bInlineMips)
 		{
 			bSucceeded = DerivedData->TryInlineMipData(BuildSettingsPerLayer[0].LODBiasWithCinematicMips, &Texture);
+
+			if (!bSucceeded)
+			{
+				UE_LOG(LogTexture, Display, TEXT("Texture %s is missing inline mips. The texture will be rebuilt."), *Texture.GetFullName());
+			}
 		}
 		else
 		{
@@ -492,10 +502,19 @@ void FTextureCacheDerivedDataWorker::DoWork()
 				bSucceeded =	DerivedData->VTData != nullptr &&
 								DerivedData->VTData->IsInitialized() &&
 								DerivedData->AreDerivedVTChunksAvailable();
+
+				if (!bSucceeded)
+				{
+					UE_LOG(LogTexture, Display, TEXT("Texture %s is missing VT Chunks. The texture will be rebuilt."), *Texture.GetFullName());
+				}
 			}
 			else
 			{
 				bSucceeded = DerivedData->AreDerivedMipsAvailable();
+				if (!bSucceeded)
+				{
+					UE_LOG(LogTexture, Display, TEXT("Texture %s is missing derived mips. The texture will be rebuilt."), *Texture.GetFullName());
+				}
 
 				if (bSucceeded && BuildSettingsPerLayer.Num() > 0)
 				{
@@ -530,7 +549,7 @@ void FTextureCacheDerivedDataWorker::DoWork()
 			bSucceeded = DerivedData->VTData->ValidateData(Texture.GetPathName(), false);
 			if (!bSucceeded)
 			{
-				UE_LOG(LogTexture, Error, TEXT("Texture %s has corrupt Virtual Texture compression. The texture will be rebuild."), *Texture.GetFullName());
+				UE_LOG(LogTexture, Error, TEXT("Texture %s has corrupt Virtual Texture compression. The texture will be rebuilt."), *Texture.GetFullName());
 				bInvalidVirtualTextureCompression = true;
 			}
 		}
