Removing UE4 references

diff --git a/Engine/Source/Runtime/Engine/Classes/Animation/AnimSet.h b/Engine/Source/Runtime/Engine/Classes/Animation/AnimSet.h
--- a/Engine/Source/Runtime/Engine/Classes/Animation/AnimSet.h
+++ b/Engine/Source/Runtime/Engine/Classes/Animation/AnimSet.h
@@ -64,7 +64,7 @@ public:
 #endif // WITH_EDITORONLY_DATA
 private:
 	/** Non-serialised cache of linkups between different skeletal meshes and this AnimSet. */
-	// @UE4: Do not change private - they will go away
+	// Do not change private - they will go away
 	UPROPERTY(transient)
 	TArray<struct FAnimSetMeshLinkup> LinkupCache;
 
@@ -73,22 +73,22 @@ private:
 	 *	This is basically a cooked down version of UseTranslationBoneNames for speed.
 	 *	Size matches the number of tracks.
 	 */
-	// @UE4: Do not change private - they will go away
+	// Do not change private - they will go away
 	UPROPERTY(transient)
 	TArray<uint8> BoneUseAnimTranslation;
 
 	/** Cooked down version of ForceMeshTranslationBoneNames */
-	// @UE4: Do not change private - they will go away
+	// Do not change private - they will go away
 	UPROPERTY(transient)
 	TArray<uint8> ForceUseMeshTranslation;
 
 	/** Names of bones that should use translation from the animation, if bAnimRotationOnly is set. */
-	// @UE4: Do not change private - they will go away
+	// Do not change private - they will go away
 	UPROPERTY(EditAnywhere, Category=AnimSet)
 	TArray<FName> UseTranslationBoneNames;
 
 	/** List of bones which are ALWAYS going to use their translation from the mesh and not the animation. */
-	// @UE4: Do not change private - they will go away
+	// Do not change private - they will go away
 	UPROPERTY(EditAnywhere, Category=AnimSet)
 	TArray<FName> ForceMeshTranslationBoneNames;
 
@@ -113,7 +113,7 @@ public:
 
 public:
 	/** Runtime built mapping table between SkeletalMeshes, and LinkupCache array indices. */
-	// @UE4: Do change private - they will go away
+	// Do change private - they will go away
 	TMap<FName,int32> SkelMesh2LinkupCache;
 
 	//~ Begin UObject Interface
diff --git a/Engine/Source/Runtime/Engine/Classes/Engine/SkeletalMesh.h b/Engine/Source/Runtime/Engine/Classes/Engine/SkeletalMesh.h
--- a/Engine/Source/Runtime/Engine/Classes/Engine/SkeletalMesh.h
+++ b/Engine/Source/Runtime/Engine/Classes/Engine/SkeletalMesh.h
@@ -2216,7 +2216,7 @@ public:
 #if WITH_EDITOR && WITH_APEX_CLOTHING
 	/** 
 	 * Take clothing assets that were imported using APEX files before we moved away from the APEX simulation
-	 * framework and upgrade them to UE4 UClothingAssets. This will bind the new assets to the mesh so the
+	 * framework and upgrade them to UClothingAssets. This will bind the new assets to the mesh so the
 	 * clothing remains working as before.
 	 */
 	void UpgradeOldClothingAssets();
diff --git a/Engine/Source/Runtime/Engine/Private/Animation/AnimPhysicsSolver.cpp b/Engine/Source/Runtime/Engine/Private/Animation/AnimPhysicsSolver.cpp
--- a/Engine/Source/Runtime/Engine/Private/Animation/AnimPhysicsSolver.cpp
+++ b/Engine/Source/Runtime/Engine/Private/Animation/AnimPhysicsSolver.cpp
@@ -1,6 +1,6 @@
 // Copyright Epic Games, Inc. All Rights Reserved.
 
-// This code is based upon and adapted to UE4 from the code 
+// This code is based upon and adapted to Unreal from the code 
 // provided in the Sandbox project here:
 // [REDACTED_PROFILE_URL]/sandbox
 
diff --git a/Engine/Source/Runtime/Engine/Private/Animation/AnimationRuntime.cpp b/Engine/Source/Runtime/Engine/Private/Animation/AnimationRuntime.cpp
--- a/Engine/Source/Runtime/Engine/Private/Animation/AnimationRuntime.cpp
+++ b/Engine/Source/Runtime/Engine/Private/Animation/AnimationRuntime.cpp
@@ -1,7 +1,7 @@
 // Copyright Epic Games, Inc. All Rights Reserved.
 
 /*=============================================================================
-	AnimationUE4.cpp: Animation runtime utilities
+	AnimationRuntime.cpp: Animation runtime utilities
 =============================================================================*/ 
 
 #include "AnimationRuntime.h"
diff --git a/Engine/Source/Runtime/Engine/Public/Animation/AnimPhysicsSolver.h b/Engine/Source/Runtime/Engine/Public/Animation/AnimPhysicsSolver.h
--- a/Engine/Source/Runtime/Engine/Public/Animation/AnimPhysicsSolver.h
+++ b/Engine/Source/Runtime/Engine/Public/Animation/AnimPhysicsSolver.h
@@ -1,6 +1,6 @@
 // Copyright Epic Games, Inc. All Rights Reserved.
 
-// This code is based upon and adapted to UE4 from the code 
+// This code is based upon and adapted to Unreal from the code 
 // provided in the Sandbox project here:
 // [REDACTED_PROFILE_URL]/sandbox
 
