Change default actors grid placement to location instead of bounds, and automatically update to bounds for actors larger than a streaming cell.

diff --git a/Engine/Source/Runtime/Engine/Private/Actor.cpp b/Engine/Source/Runtime/Engine/Private/Actor.cpp
--- a/Engine/Source/Runtime/Engine/Private/Actor.cpp
+++ b/Engine/Source/Runtime/Engine/Private/Actor.cpp
@@ -175,7 +175,7 @@ void AActor::InitializeDefaults()
 #endif // (CSV_PROFILER && !UE_BUILD_SHIPPING)
 
 #if WITH_EDITORONLY_DATA
-	GridPlacement = EActorGridPlacement::Bounds;
+	GridPlacement = EActorGridPlacement::Location;
 #endif
 }
 
diff --git a/Engine/Source/Runtime/Engine/Private/WorldPartition/WorldPartitionRuntimeSpatialHash.cpp b/Engine/Source/Runtime/Engine/Private/WorldPartition/WorldPartitionRuntimeSpatialHash.cpp
--- a/Engine/Source/Runtime/Engine/Private/WorldPartition/WorldPartitionRuntimeSpatialHash.cpp
+++ b/Engine/Source/Runtime/Engine/Private/WorldPartition/WorldPartitionRuntimeSpatialHash.cpp
@@ -659,22 +659,25 @@ void UWorldPartitionRuntimeSpatialHash::CreateActorDescViewMap(const UActorDescC
 
 		if (!ActorDescView.GetActorIsEditorOnly())
 		{
-			if (ActorDescView.GetGridPlacement() == EActorGridPlacement::Bounds)
+			if (ActorDescView.GetActorClass()->GetDefaultObject<AActor>()->GetDefaultGridPlacement() == EActorGridPlacement::None)
 			{
-				const int32* GridIndexPtr = GridsMapping.Find(ActorDescView.GetRuntimeGrid());
-				const int32 GridIndex = GridIndexPtr ? *GridIndexPtr : 0;
-				const FSpatialHashRuntimeGrid& RuntimeGrid = Grids[GridIndex];
-				const float CellArea = RuntimeGrid.CellSize * RuntimeGrid.CellSize;
-				const FBox2D ActorBounds2D = FBox2D(FVector2D(ActorDescView.GetBounds().Min), FVector2D(ActorDescView.GetBounds().Max));
-				const float ActorBoundsArea = ActorBounds2D.GetArea();
-
-				if (ActorBoundsArea < CellArea / 16.0f)
+				if (ActorDescView.GetGridPlacement() != EActorGridPlacement::AlwaysLoaded)
 				{
-					ChangeActorDescViewGridPlacement(ActorDescView, EActorGridPlacement::Location);
-				}
-				else if (ActorBoundsArea > WorldBoundsArea / 2.0f)
-				{
-					ChangeActorDescViewGridPlacement(ActorDescView, EActorGridPlacement::AlwaysLoaded);
+					const int32* GridIndexPtr = GridsMapping.Find(ActorDescView.GetRuntimeGrid());
+					const int32 GridIndex = GridIndexPtr ? *GridIndexPtr : 0;
+					const FSpatialHashRuntimeGrid& RuntimeGrid = Grids[GridIndex];
+					const float CellArea = RuntimeGrid.CellSize * RuntimeGrid.CellSize;
+					const FBox2D ActorBounds2D = FBox2D(FVector2D(ActorDescView.GetBounds().Min), FVector2D(ActorDescView.GetBounds().Max));
+					const float ActorBoundsArea = ActorBounds2D.GetArea();
+
+					if (ActorBoundsArea > WorldBoundsArea / 2.0f)
+					{
+						ChangeActorDescViewGridPlacement(ActorDescView, EActorGridPlacement::AlwaysLoaded);
+					}
+					else if (ActorBoundsArea > CellArea)
+					{
+						ChangeActorDescViewGridPlacement(ActorDescView, EActorGridPlacement::Bounds);
+					}
 				}
 			}
 		}
