Fixing potential heap corruption issue where there is a mismatch of FMalloc and delete in the root motion serialization code. This fix follows a pattern used elsewhere where a custom deleter is used that cleans up properly.

diff --git a/Engine/Source/Runtime/Engine/Private/GameFramework/RootMotionSource.cpp b/Engine/Source/Runtime/Engine/Private/GameFramework/RootMotionSource.cpp
--- a/Engine/Source/Runtime/Engine/Private/GameFramework/RootMotionSource.cpp
+++ b/Engine/Source/Runtime/Engine/Private/GameFramework/RootMotionSource.cpp
@@ -1970,6 +1970,18 @@ void FRootMotionSourceGroup::UpdateStateFrom(const FRootMotionSourceGroup& Group
 	}
 }
 
+struct FRootMotionSourceDeleter
+{
+	FORCEINLINE void operator()(FRootMotionSource* Object) const
+	{
+		check(Object);
+		UScriptStruct* ScriptStruct = Object->GetScriptStruct();
+		check(ScriptStruct);
+		ScriptStruct->DestroyStruct(Object);
+		FMemory::Free(Object);
+	}
+};
+
 void FRootMotionSourceGroup::NetSerializeRMSArray(FArchive& Ar, class UPackageMap* Map, bool& bOutSuccess, TArray< TSharedPtr<FRootMotionSource> >& RootMotionSourceArray)
 {
 	uint8 SourcesNum;
@@ -2027,7 +2039,7 @@ void FRootMotionSourceGroup::NetSerializeRMSArray(FArchive& Ar, class UPackageMa
 						FRootMotionSource* NewSource = (FRootMotionSource*)FMemory::Malloc(ScriptStruct->GetCppStructOps()->GetSize());
 						ScriptStruct->InitializeStruct(NewSource);
 
-						RootMotionSourceArray[i] = TSharedPtr<FRootMotionSource>(NewSource);
+						RootMotionSourceArray[i] = TSharedPtr<FRootMotionSource>(NewSource, FRootMotionSourceDeleter());
 					}
 				}
 
