Implements a valid CompilePreview for strata node.
Follow up to 15797843.

diff --git a/Engine/Source/Runtime/Engine/Classes/Materials/MaterialExpressionStrata.h b/Engine/Source/Runtime/Engine/Classes/Materials/MaterialExpressionStrata.h
--- a/Engine/Source/Runtime/Engine/Classes/Materials/MaterialExpressionStrata.h
+++ b/Engine/Source/Runtime/Engine/Classes/Materials/MaterialExpressionStrata.h
@@ -27,6 +27,10 @@ UCLASS(MinimalAPI, collapsecategories, hidecategories = Object, Abstract, Displa
 class UMaterialExpressionStrataBSDF : public UMaterialExpression
 {
 	GENERATED_UCLASS_BODY()
+
+#if WITH_EDITOR
+		virtual int32 CompilePreview(class FMaterialCompiler* Compiler, int32 OutputIndex) override;
+#endif
 };
 
 UCLASS(MinimalAPI, collapsecategories, hidecategories = Object, DisplayName = "Strata Slab BSDF")
diff --git a/Engine/Source/Runtime/Engine/Private/Materials/MaterialExpressions.cpp b/Engine/Source/Runtime/Engine/Private/Materials/MaterialExpressions.cpp
--- a/Engine/Source/Runtime/Engine/Private/Materials/MaterialExpressions.cpp
+++ b/Engine/Source/Runtime/Engine/Private/Materials/MaterialExpressions.cpp
@@ -19725,6 +19725,13 @@ UMaterialExpressionStrataBSDF::UMaterialExpressionStrataBSDF(const FObjectInitia
 {
 }
 
+int32 UMaterialExpressionStrataBSDF::CompilePreview(class FMaterialCompiler* Compiler, int32 OutputIndex)
+{
+	// Strata nodes cannot be previewed due to the complex Strata type.
+	// So we override the default implementation calling Compile and we simply return a default black preview color.
+	return Compiler->Constant(0.0f);
+}
+
 UMaterialExpressionStrataSlabBSDF::UMaterialExpressionStrataSlabBSDF(const FObjectInitializer& ObjectInitializer)
 	: Super(ObjectInitializer)
 {
