Added extra logging when storing empty texture bulk data in DDC.

diff --git a/Engine/Source/Runtime/Engine/Private/TextureDerivedData.cpp b/Engine/Source/Runtime/Engine/Private/TextureDerivedData.cpp
--- a/Engine/Source/Runtime/Engine/Private/TextureDerivedData.cpp
+++ b/Engine/Source/Runtime/Engine/Private/TextureDerivedData.cpp
@@ -609,15 +609,6 @@ uint32 PutDerivedDataInCache(FTexturePlatformData* DerivedData, const FString& D
 	GetTextureDerivedDataKeyFromSuffix(DerivedDataKeySuffix, DerivedDataKey);
 
 	FString LogString;
-	if (UE_LOG_ACTIVE(LogTexture,Verbose))
-	{
-		LogString = FString::Printf(
-			TEXT("Storing texture in DDC:\n  Name: %s\n  Key: %s\n  Format: %s\n"),
-			*FString(TextureName),
-			*DerivedDataKey,
-			GPixelFormats[DerivedData->PixelFormat].Name
-			);
-	}
 
 	// Write out individual mips to the derived data cache.
 	const int32 MipCount = DerivedData->Mips.Num();
@@ -630,8 +621,19 @@ uint32 PutDerivedDataInCache(FTexturePlatformData* DerivedData, const FString& D
 		const bool bInline = (MipIndex >= FirstInlineMip);
 		GetTextureDerivedMipKey(MipIndex, Mip, DerivedDataKeySuffix, MipDerivedDataKey);
 
-		if (UE_LOG_ACTIVE(LogTexture,Verbose))
+		const bool bDDCError = !bInline && !Mip.BulkData.GetBulkDataSize();
+		if (UE_LOG_ACTIVE(LogTexture,Verbose) || bDDCError)
 		{
+			if (LogString.IsEmpty())
+			{
+				LogString = FString::Printf(
+					TEXT("Storing texture in DDC:\n  Name: %s\n  Key: %s\n  Format: %s\n"),
+					*FString(TextureName),
+					*DerivedDataKey,
+					GPixelFormats[DerivedData->PixelFormat].Name
+				);
+			}
+
 			LogString += FString::Printf(TEXT("  Mip%d %dx%d %d bytes%s %s\n"),
 				MipIndex,
 				Mip.SizeX,
@@ -642,6 +644,11 @@ uint32 PutDerivedDataInCache(FTexturePlatformData* DerivedData, const FString& D
 				);
 		}
 
+		if (bDDCError)
+		{
+			UE_LOG(LogTexture, Fatal, TEXT("Error %s"), *LogString);
+		}
+
 		// Note that calling StoreInDerivedDataCache() also calls RemoveBulkData().
 		// This means that the resource needs to load differently inlined mips and non inlined mips.
 		if (!bInline)
