Speed up ODSC shader map processing from ~27.5 ms to ~6.1 ms on average.

- Pass in 0 flags to FMaterialUpdateContext.  This avoids redundant FlushRenderingCommands call which already happens in ProcessCookOnTheFlyShaders.
- Only reattach components that that use the materials passed into the context to avoid re-creating every component in the scene.

diff --git a/Engine/Source/Runtime/Engine/Private/ShaderCompiler/ShaderCompiler.cpp b/Engine/Source/Runtime/Engine/Private/ShaderCompiler/ShaderCompiler.cpp
--- a/Engine/Source/Runtime/Engine/Private/ShaderCompiler/ShaderCompiler.cpp
+++ b/Engine/Source/Runtime/Engine/Private/ShaderCompiler/ShaderCompiler.cpp
@@ -5130,7 +5130,7 @@ void ProcessCookOnTheFlyShaders(bool bReloadGlobalShaders, const TArray<uint8>&
 	if (MeshMaterialMaps.Num() > 0)
 	{
 		// this will stop the rendering thread, and reattach components, in the destructor
-		FMaterialUpdateContext UpdateContext;
+		FMaterialUpdateContext UpdateContext(0);
 
 		// parse the shaders
 		FMemoryReader MemoryReader(MeshMaterialMaps, true);
