Fixed WorldPartition not using modified actors in PIE : runtime cell bIsModifiedForPIE flag was not duplicated at PIE.

diff --git a/Engine/Source/Runtime/Engine/Public/WorldPartition/WorldPartitionRuntimeCell.h b/Engine/Source/Runtime/Engine/Public/WorldPartition/WorldPartitionRuntimeCell.h
--- a/Engine/Source/Runtime/Engine/Public/WorldPartition/WorldPartitionRuntimeCell.h
+++ b/Engine/Source/Runtime/Engine/Public/WorldPartition/WorldPartitionRuntimeCell.h
@@ -138,7 +138,7 @@ protected:
 	bool bIsAlwaysLoaded;
 
 #if WITH_EDITORONLY_DATA
-	UPROPERTY(Transient)
+	UPROPERTY()
 	bool bIsModifiedForPIE;
 #endif
 
