UStaticMeshComponent::IsHLODRelevant() - Use !IsVisible() instead of bHiddenInGame

diff --git a/Engine/Source/Runtime/Engine/Private/Components/StaticMeshComponent.cpp b/Engine/Source/Runtime/Engine/Private/Components/StaticMeshComponent.cpp
--- a/Engine/Source/Runtime/Engine/Private/Components/StaticMeshComponent.cpp
+++ b/Engine/Source/Runtime/Engine/Private/Components/StaticMeshComponent.cpp
@@ -2470,7 +2470,7 @@ bool UStaticMeshComponent::IsHLODRelevant() const
 		return false;
 	}
 
-	if (bHiddenInGame)
+	if (!IsVisible())
 	{
 		return false;
 	}
