UPrimitiveComponent::SetLODParentPrimitive() - Skip unneeded call to ShouldGenerateAutoLOD() if InLODParentPrimitive is null

diff --git a/Engine/Source/Runtime/Engine/Private/Components/PrimitiveComponent.cpp b/Engine/Source/Runtime/Engine/Private/Components/PrimitiveComponent.cpp
--- a/Engine/Source/Runtime/Engine/Private/Components/PrimitiveComponent.cpp
+++ b/Engine/Source/Runtime/Engine/Private/Components/PrimitiveComponent.cpp
@@ -3840,7 +3840,7 @@ void UPrimitiveComponent::SetLODParentPrimitive(UPrimitiveComponent * InLODParen
 		return nullptr;
 	}();
 
-	if (!GIsEditor || ShouldGenerateAutoLOD(ParentLODActor ? ParentLODActor->LODLevel - 1 : INDEX_NONE))
+	if (!GIsEditor || !InLODParentPrimitive || ShouldGenerateAutoLOD(ParentLODActor ? ParentLODActor->LODLevel - 1 : INDEX_NONE))
 #endif
 	{
 		// @todo, what do we do with old parent. We can't just reset undo parent because the parent might be used by other primitive
