Changed how world partition actor clustering handles 2 clusters with different sets of Data Layers.
If they each have a valid set of Data Layers, the actor cluster will merge all Data Layers.
If one of the cluster doesn't have Data Layers, the resulting actor cluster won't have Data Layers.
Also added logging to keep track of what actor clustering is doing.
Note that these cases are already reported as errors/warnings during Map_Check.

diff --git a/Engine/Source/Runtime/Engine/Private/WorldPartition/WorldPartitionActorCluster.cpp b/Engine/Source/Runtime/Engine/Private/WorldPartition/WorldPartitionActorCluster.cpp
--- a/Engine/Source/Runtime/Engine/Private/WorldPartition/WorldPartitionActorCluster.cpp
+++ b/Engine/Source/Runtime/Engine/Private/WorldPartition/WorldPartitionActorCluster.cpp
@@ -12,8 +12,6 @@
 #include "WorldPartition/WorldPartitionRuntimeHash.h"
 #include "Misc/HashBuilder.h"
 
-DEFINE_LOG_CATEGORY(LogWorldPartitionActorCluster);
-
 template<class LayerNameContainer>
 TSet<const UDataLayer*> GetDataLayers(UWorld* InWorld, const LayerNameContainer& DataLayerNames)
 {
@@ -48,9 +46,6 @@ FActorCluster::FActorCluster(UWorld* InWorld, const FWorldPartitionActorDescView
 
 void FActorCluster::Add(const FActorCluster& InActorCluster)
 {
-	// Merge Actors
-	Actors.Append(InActorCluster.Actors);
-
 	// Merge RuntimeGrid
 	RuntimeGrid = RuntimeGrid == InActorCluster.RuntimeGrid ? RuntimeGrid : NAME_None;
 
@@ -82,16 +77,58 @@ void FActorCluster::Add(const FActorCluster& InActorCluster)
 		}
 	}
 
-	// Merge DataLayers
 	if (DataLayersID != InActorCluster.DataLayersID)
 	{
-		for (const UDataLayer* DataLayer : InActorCluster.DataLayers)
+		if (DataLayers.Num() && InActorCluster.DataLayers.Num())
+		{
+			// Merge Data Layers
+			for (const UDataLayer* DataLayer : InActorCluster.DataLayers)
+			{
+				check(DataLayer->IsDynamicallyLoaded());
+				DataLayers.Add(DataLayer);
+			}
+			UE_SUPPRESS(LogWorldPartition, Verbose,
+			{
+				UE_LOG(LogWorldPartition, Verbose, TEXT("Merging Data Layers for clustered actors with different sets of Data Layers."));
+				UE_LOG(LogWorldPartition, Verbose, TEXT("1st clustered actors :"));
+				for (const FGuid& ActorGuid : Actors)
+				{
+					UE_LOG(LogWorldPartition, Verbose, TEXT("   - Actor: (%s)"), *ActorGuid.ToString(EGuidFormats::UniqueObjectGuid));
+				}
+				UE_LOG(LogWorldPartition, Verbose, TEXT("2nd clustered actors :"));
+				for (const FGuid& ActorGuid : InActorCluster.Actors)
+				{
+					UE_LOG(LogWorldPartition, Verbose, TEXT("   - Actor: (%s)"), *ActorGuid.ToString(EGuidFormats::UniqueObjectGuid));
+				}
+			});
+		}
+		else
 		{
-			check(DataLayer->IsDynamicallyLoaded());
-			DataLayers.Add(DataLayer);
+			// In the case where one of the actor cluster has no Data Layer and the other does, 
+			// the merged actor cluster removes all Data Layers
+			UE_SUPPRESS(LogWorldPartition, Verbose,
+			{
+				UE_LOG(LogWorldPartition, Verbose, TEXT("Removing Data Layers for clustered actors because they are referenced by or are referencing other actors with no Data Layer."));
+				UE_LOG(LogWorldPartition, Verbose, TEXT("Clustered actors with Data Layers :"));
+				const FActorCluster * ActorClusterWithDataLayer = DataLayers.Num() ? this : &InActorCluster;
+				for (const FGuid& ActorGuid : ActorClusterWithDataLayer->Actors)
+				{
+					UE_LOG(LogWorldPartition, Verbose, TEXT("   - Actor: (%s)"), *ActorGuid.ToString(EGuidFormats::UniqueObjectGuid));
+				}
+				UE_LOG(LogWorldPartition, Verbose, TEXT("Clustered actors without Data Layer :"));
+				const FActorCluster * ActorClusterWithoutDataLayer = DataLayers.Num() ? &InActorCluster : this;
+				for (const FGuid& ActorGuid : ActorClusterWithoutDataLayer->Actors)
+				{
+					UE_LOG(LogWorldPartition, Verbose, TEXT("   - Actor: (%s)"), *ActorGuid.ToString(EGuidFormats::UniqueObjectGuid));
+				}
+			});
+			DataLayers.Empty();
 		}
 		DataLayersID = FDataLayersID(DataLayers.Array());
 	}
+
+	// Merge Actors
+	Actors.Append(InActorCluster.Actors);
 }
 
 FActorClusterInstance::FActorClusterInstance(const FActorCluster* InCluster, const FActorContainerInstance* InContainerInstance)
@@ -300,7 +337,7 @@ void FActorClusterContext::CreateContainerInstanceRecursive(uint32 ID, const FTr
 		*ParentBounds += Bounds;
 	}
 
-	UE_LOG(LogWorldPartitionActorCluster, Verbose, TEXT("ContainerInstance (%08x) Bounds (%s) Package (%s)"), ID, *Bounds.TransformBy(Transform).ToString(), *Container->GetContainerPackage().ToString());
+	UE_LOG(LogWorldPartition, Verbose, TEXT("ContainerInstance (%08x) Bounds (%s) Package (%s)"), ID, *Bounds.TransformBy(Transform).ToString(), *Container->GetContainerPackage().ToString());
 	ContainerInstances.Add(FActorContainerInstance(ID, Transform, Bounds, DataLayers, ClusterMode, Container, MoveTemp(ChildContainers), MoveTemp(ActorDescViewMap)));
 }
 
diff --git a/Engine/Source/Runtime/Engine/Public/WorldPartition/WorldPartitionActorCluster.h b/Engine/Source/Runtime/Engine/Public/WorldPartition/WorldPartitionActorCluster.h
--- a/Engine/Source/Runtime/Engine/Public/WorldPartition/WorldPartitionActorCluster.h
+++ b/Engine/Source/Runtime/Engine/Public/WorldPartition/WorldPartitionActorCluster.h
@@ -13,8 +13,6 @@ class FWorldPartitionActorDesc;
 class UWorldPartitionRuntimeHash;
 enum class EContainerClusterMode : uint8;
 
-DECLARE_LOG_CATEGORY_EXTERN(LogWorldPartitionActorCluster, Log, All);
-
 /**
  * List of actors bound together based on clustering rules. (mainly object references)
  */
