Merged all world partition logs to LogWorldPartition.

diff --git a/Engine/Source/Runtime/Engine/Private/WorldPartition/RuntimeSpatialHash/RuntimeSpatialHashGridHelper.cpp b/Engine/Source/Runtime/Engine/Private/WorldPartition/RuntimeSpatialHash/RuntimeSpatialHashGridHelper.cpp
--- a/Engine/Source/Runtime/Engine/Private/WorldPartition/RuntimeSpatialHash/RuntimeSpatialHashGridHelper.cpp
+++ b/Engine/Source/Runtime/Engine/Private/WorldPartition/RuntimeSpatialHash/RuntimeSpatialHashGridHelper.cpp
@@ -30,7 +30,7 @@ FSquare2DGridHelper::FSquare2DGridHelper(const FBox& InWorldBounds, const FVecto
 	}
 	else
 	{
-		UE_LOG(LogWorldPartitionRuntimeSpatialHash, Warning, TEXT("Invalid world bounds, grid partitioning will use a runtime grid with 1 cell."));
+		UE_LOG(LogWorldPartition, Warning, TEXT("Invalid world bounds, grid partitioning will use a runtime grid with 1 cell."));
 	}
 
 	check(FMath::IsPowerOfTwo(GridSize));
@@ -98,7 +98,7 @@ int32 FSquare2DGridHelper::ForEachIntersectingCells(const FSphere& InSphere, TFu
 #if WITH_EDITOR
 void FSquare2DGridHelper::ValidateSingleActorReferer()
 {
-	UE_SCOPED_TIMER(TEXT("ValidateSingleActorReferer"), LogWorldPartitionRuntimeSpatialHash);
+	UE_SCOPED_TIMER(TEXT("ValidateSingleActorReferer"), LogWorldPartition);
 
 	TSet<FActorInstance> ActorUsage;
 	for (int32 Level = 0; Level < Levels.Num() - 1; Level++)
@@ -130,7 +130,7 @@ FSquare2DGridHelper GetGridHelper(const FBox& WorldBounds, int32 GridCellSize)
 
 FSquare2DGridHelper GetPartitionedActors(const UWorldPartition* WorldPartition, const FBox& WorldBounds, const FSpatialHashRuntimeGrid& Grid, const TArray<const FActorClusterInstance*>& GridActors)
 {
-	UE_SCOPED_TIMER(TEXT("GetPartitionedActors"), LogWorldPartitionRuntimeSpatialHash);
+	UE_SCOPED_TIMER(TEXT("GetPartitionedActors"), LogWorldPartition);
 
 	//
 	// Create the hierarchical grids for the game
@@ -143,7 +143,7 @@ FSquare2DGridHelper GetPartitionedActors(const UWorldPartition* WorldPartition,
 		LastGridLevel.ForEachIntersectingCells(WorldBounds, [&IntersectingCellCount](const FIntVector2& Coords) { ++IntersectingCellCount; });
 		if (!ensure(IntersectingCellCount == 1))
 		{
-			UE_LOG(LogWorldPartitionRuntimeSpatialHash, Warning, TEXT("Can't find grid cell that encompasses world bounds."));
+			UE_LOG(LogWorldPartition, Warning, TEXT("Can't find grid cell that encompasses world bounds."));
 		}
 	}
 
@@ -215,13 +215,13 @@ FSquare2DGridHelper GetPartitionedActors(const UWorldPartition* WorldPartition,
 			check(0);
 		}
 
-		if (UE_LOG_ACTIVE(LogWorldPartitionRuntimeSpatialHash, Verbose))
+		if (UE_LOG_ACTIVE(LogWorldPartition, Verbose))
 		{
 			if (ActorCluster->Actors.Num() > 1)
 			{
 				static UEnum* ActorGridPlacementEnum = FindObject<UEnum>(ANY_PACKAGE, TEXT("EActorGridPlacement"));
 
-				UE_LOG(LogWorldPartitionRuntimeSpatialHash, Verbose, TEXT("Clustered %d actors (%s%s%s), generated shared BV of [%d x %d] (meters)"),
+				UE_LOG(LogWorldPartition, Verbose, TEXT("Clustered %d actors (%s%s%s), generated shared BV of [%d x %d] (meters)"),
 					ActorCluster->Actors.Num(),
 					*ActorGridPlacementEnum->GetNameStringByValue((int64)GridPlacement),
 					bAlwaysLoadedPromotedCluster ? TEXT(":PromotedCluster") : TEXT(""),
@@ -232,9 +232,9 @@ FSquare2DGridHelper GetPartitionedActors(const UWorldPartition* WorldPartition,
 				for (const FGuid& ActorGuid : ActorCluster->Actors)
 				{
 					const FWorldPartitionActorDescView& ActorDescView = ClusterInstance->ContainerInstance->GetActorDescView(ActorGuid);
-					UE_LOG(LogWorldPartitionRuntimeSpatialHash, Verbose, TEXT("   - Actor: %s (%s)"), *ActorDescView.GetActorPath().ToString(), *ActorGuid.ToString(EGuidFormats::UniqueObjectGuid));
-					UE_LOG(LogWorldPartitionRuntimeSpatialHash, Verbose, TEXT("         Package: %s"), *ActorDescView.GetActorPackage().ToString());
-					UE_LOG(LogWorldPartitionRuntimeSpatialHash, Verbose, TEXT("         Container (%08x): %s"), ClusterInstance->ContainerInstance->ID, *ClusterInstance->ContainerInstance->Container->GetContainerPackage().ToString())
+					UE_LOG(LogWorldPartition, Verbose, TEXT("   - Actor: %s (%s)"), *ActorDescView.GetActorPath().ToString(), *ActorGuid.ToString(EGuidFormats::UniqueObjectGuid));
+					UE_LOG(LogWorldPartition, Verbose, TEXT("         Package: %s"), *ActorDescView.GetActorPackage().ToString());
+					UE_LOG(LogWorldPartition, Verbose, TEXT("         Container (%08x): %s"), ClusterInstance->ContainerInstance->ID, *ClusterInstance->ContainerInstance->Container->GetContainerPackage().ToString())
 				}
 			}
 		}
diff --git a/Engine/Source/Runtime/Engine/Private/WorldPartition/WorldPartitionRuntimeSpatialHash.cpp b/Engine/Source/Runtime/Engine/Private/WorldPartition/WorldPartitionRuntimeSpatialHash.cpp
--- a/Engine/Source/Runtime/Engine/Private/WorldPartition/WorldPartitionRuntimeSpatialHash.cpp
+++ b/Engine/Source/Runtime/Engine/Private/WorldPartition/WorldPartitionRuntimeSpatialHash.cpp
@@ -43,8 +43,6 @@
 extern UNREALED_API class UEditorEngine* GEditor;
 #endif //WITH_EDITOR
 
-DEFINE_LOG_CATEGORY(LogWorldPartitionRuntimeSpatialHash);
-
 static int32 GShowRuntimeSpatialHashGridLevel = 0;
 static FAutoConsoleVariableRef CVarShowRuntimeSpatialHashGridLevel(
 	TEXT("wp.Runtime.ShowRuntimeSpatialHashGridLevel"),
@@ -523,11 +521,11 @@ bool UWorldPartitionRuntimeSpatialHash::GenerateStreaming(EWorldPartitionStreami
 	TRACE_CPUPROFILER_EVENT_SCOPE(UWorldPartitionRuntimeSpatialHash::GenerateStreaming);
 	UWorldPartition* WorldPartition = GetOuterUWorldPartition();
 
-	UE_SCOPED_TIMER(TEXT("GenerateStreaming"), LogWorldPartitionRuntimeSpatialHash);
+	UE_SCOPED_TIMER(TEXT("GenerateStreaming"), LogWorldPartition);
 	
 	if (!Grids.Num())
 	{
-		UE_LOG(LogWorldPartitionRuntimeSpatialHash, Error, TEXT("Invalid partition grids setup"));
+		UE_LOG(LogWorldPartition, Error, TEXT("Invalid partition grids setup"));
 		return false;
 	}
 
@@ -574,7 +572,7 @@ bool UWorldPartitionRuntimeSpatialHash::GenerateStreaming(EWorldPartitionStreami
 		int32* FoundIndex = GridsMapping.Find(ClusterInstance.Cluster->RuntimeGrid);
 		if (!FoundIndex)
 		{
-			UE_LOG(LogWorldPartitionRuntimeSpatialHash, Error, TEXT("Invalid partition grid '%s' referenced by actor cluster"), *ClusterInstance.Cluster->RuntimeGrid.ToString());
+			UE_LOG(LogWorldPartition, Error, TEXT("Invalid partition grid '%s' referenced by actor cluster"), *ClusterInstance.Cluster->RuntimeGrid.ToString());
 		}
 
 		int32 GridIndex = FoundIndex ? *FoundIndex : 0;
@@ -734,7 +732,7 @@ bool UWorldPartitionRuntimeSpatialHash::CreateStreamingGrid(const FSpatialHashRu
 						if (ActorInstance.ShouldStripFromStreaming())
 						{
 							const FWorldPartitionActorDescView& ActorDescView = ActorInstance.GetActorDescView();
-							UE_LOG(LogWorldPartitionRuntimeSpatialHash, Verbose, TEXT("Stripping Actor %s (%s) from streaming grid (Container %08x)"),
+							UE_LOG(LogWorldPartition, Verbose, TEXT("Stripping Actor %s (%s) from streaming grid (Container %08x)"),
 								*(ActorDescView.GetActorPath().ToString()), *ActorInstance.Actor.ToString(EGuidFormats::UniqueObjectGuid), ActorInstance.ContainerInstance->ID);
 							continue;
 						}
@@ -806,7 +804,7 @@ bool UWorldPartitionRuntimeSpatialHash::CreateStreamingGrid(const FSpatialHashRu
 				verify(TempLevel.GetCellBounds(FIntVector2(CellCoordX, CellCoordY), Bounds));
 				StreamingCell->Position = FVector(Bounds.GetCenter(), 0.f);
 
-				UE_LOG(LogWorldPartitionRuntimeSpatialHash, Verbose, TEXT("Cell%s %s Actors = %d Bounds (%s)"), bIsCellAlwaysLoaded ? TEXT(" (AlwaysLoaded)") : TEXT(""), *StreamingCell->GetName(), FilteredActors.Num(), *Bounds.ToString());
+				UE_LOG(LogWorldPartition, Verbose, TEXT("Cell%s %s Actors = %d Bounds (%s)"), bIsCellAlwaysLoaded ? TEXT(" (AlwaysLoaded)") : TEXT(""), *StreamingCell->GetName(), FilteredActors.Num(), *Bounds.ToString());
 
 				for (const FActorInstance& ActorInstance : FilteredActors)
 				{
@@ -822,12 +820,12 @@ bool UWorldPartitionRuntimeSpatialHash::CreateStreamingGrid(const FSpatialHashRu
 						}
 					}
 
-					UE_LOG(LogWorldPartitionRuntimeSpatialHash, Verbose, TEXT("  Actor : %s (%s) (Container %08x) Origin(%s)"), *(ActorDescView.GetActorPath().ToString()), *ActorDescView.GetGuid().ToString(EGuidFormats::UniqueObjectGuid), ActorInstance.ContainerInstance->ID, *FVector2D(ActorInstance.GetOrigin()).ToString());
+					UE_LOG(LogWorldPartition, Verbose, TEXT("  Actor : %s (%s) (Container %08x) Origin(%s)"), *(ActorDescView.GetActorPath().ToString()), *ActorDescView.GetGuid().ToString(EGuidFormats::UniqueObjectGuid), ActorInstance.ContainerInstance->ID, *FVector2D(ActorInstance.GetOrigin()).ToString());
 				}
 
 				if (Mode == EWorldPartitionStreamingMode::Cook)
 				{
-					UE_LOG(LogWorldPartitionRuntimeSpatialHash, Log, TEXT("Creating runtime streaming cells %s."), *StreamingCell->GetName());
+					UE_LOG(LogWorldPartition, Log, TEXT("Creating runtime streaming cells %s."), *StreamingCell->GetName());
 
 					if (StreamingCell->GetActorCount())
 					{
@@ -840,7 +838,7 @@ bool UWorldPartitionRuntimeSpatialHash::CreateStreamingGrid(const FSpatialHashRu
 						{
 							if (!OutPackagesToGenerate)
 							{
-								UE_LOG(LogWorldPartitionRuntimeSpatialHash, Error, TEXT("Error creating runtime streaming cells for cook, OutPackagesToGenerate is null."));
+								UE_LOG(LogWorldPartition, Error, TEXT("Error creating runtime streaming cells for cook, OutPackagesToGenerate is null."));
 								return false;
 							}
 
diff --git a/Engine/Source/Runtime/Engine/Public/WorldPartition/WorldPartitionRuntimeSpatialHash.h b/Engine/Source/Runtime/Engine/Public/WorldPartition/WorldPartitionRuntimeSpatialHash.h
--- a/Engine/Source/Runtime/Engine/Public/WorldPartition/WorldPartitionRuntimeSpatialHash.h
+++ b/Engine/Source/Runtime/Engine/Public/WorldPartition/WorldPartitionRuntimeSpatialHash.h
@@ -14,8 +14,6 @@
 #include "WorldPartitionRuntimeSpatialHashGridPreviewer.h"
 #include "WorldPartitionRuntimeSpatialHash.generated.h"
 
-DECLARE_LOG_CATEGORY_EXTERN(LogWorldPartitionRuntimeSpatialHash, Log, All);
-
 USTRUCT()
 struct FSpatialHashStreamingGridLayerCell
 {
