Added wp.Editor.SetLogWorldPartitionVerbosity to change world partition logs verbosity.

diff --git a/Engine/Source/Runtime/Engine/Private/WorldPartition/WorldPartition.cpp b/Engine/Source/Runtime/Engine/Private/WorldPartition/WorldPartition.cpp
--- a/Engine/Source/Runtime/Engine/Private/WorldPartition/WorldPartition.cpp
+++ b/Engine/Source/Runtime/Engine/Private/WorldPartition/WorldPartition.cpp
@@ -66,6 +66,25 @@ static FAutoConsoleCommand DumpActorDescs(
 	})
 );
 
+static FAutoConsoleCommand SetLogWorldPartitionVerbosity(
+	TEXT("wp.Editor.SetLogWorldPartitionVerbosity"),
+	TEXT("Change the WorldPartition verbosity log verbosity."),
+	FConsoleCommandWithArgsDelegate::CreateLambda([](const TArray<FString>& Args)
+	{
+		if (Args.Num() == 1)
+		{
+			if (Args[0].Contains(TEXT("Verbose")))
+			{
+				LogWorldPartition.SetVerbosity(ELogVerbosity::Verbose);
+			}
+			else
+			{
+				LogWorldPartition.SetVerbosity(LogWorldPartition.GetCompileTimeVerbosity());
+			}
+		}
+	})
+);
+
 // Helper class to avoid sending global events until all cells updates are processed.
 struct FWorldPartionCellUpdateContext
 {
