WorldPartition HLODs - Fixed packed level instances components being added twice
* Added check for SupportsLoading() before calling BlockLoadLevelInstance() on the level instance...
* Also verified that underlying actors are IsHLODRelevant() before including them

The BlockLoadLevelInstance() call might not be actually needed, to be investigated

diff --git a/Engine/Source/Runtime/Engine/Private/WorldPartition/HLOD/HLODBuilder.cpp b/Engine/Source/Runtime/Engine/Private/WorldPartition/HLOD/HLODBuilder.cpp
--- a/Engine/Source/Runtime/Engine/Private/WorldPartition/HLOD/HLODBuilder.cpp
+++ b/Engine/Source/Runtime/Engine/Private/WorldPartition/HLOD/HLODBuilder.cpp
@@ -84,7 +84,10 @@ public:
 				if (ALevelInstance* LevelInstance = Cast<ALevelInstance>(LoadedActor))
 				{
 					// Wait for level instance loading
-					LevelInstance->GetLevelInstanceSubsystem()->BlockLoadLevelInstance(LevelInstance);
+					if (LevelInstance->SupportsLoading())
+					{
+						LevelInstance->GetLevelInstanceSubsystem()->BlockLoadLevelInstance(LevelInstance);
+					}
 				}
 
 				OutActors.Add(LoadedActor);
@@ -106,9 +109,10 @@ public:
 	{
 		TArray<UPrimitiveComponent*> PrimitiveComponents;
 
-		auto GatherPrimitivesFromActor = [&PrimitiveComponents](const AActor* Actor)
+		auto GatherPrimitivesFromActor = [&PrimitiveComponents](const AActor* Actor, const AActor* ParentActor = nullptr)
 		{
-			const bool bIncludeTransientActors = Actor->IsA<ALevelInstance>();
+			const TCHAR* Padding = ParentActor ? TEXT("    ") : TEXT("");
+			UE_LOG(LogHLODBuilder, Verbose, TEXT("%s* Adding components from actor %s"), Padding, *Actor->GetName());
 			for (UActorComponent* SubComponent : Actor->GetComponents())
 			{
 				if (SubComponent && SubComponent->IsHLODRelevant())
@@ -116,6 +120,15 @@ public:
 					if (UPrimitiveComponent* PrimitiveComponent = Cast<UPrimitiveComponent>(SubComponent))
 					{
 						PrimitiveComponents.Add(PrimitiveComponent);
+
+						if (UInstancedStaticMeshComponent* ISMC = Cast<UInstancedStaticMeshComponent>(PrimitiveComponent))
+						{
+							UE_LOG(LogHLODBuilder, Verbose, TEXT("%s    * %s [%d instances]"), Padding, *ISMC->GetStaticMesh()->GetName(), ISMC->GetInstanceCount());
+						}
+						else if (UStaticMeshComponent* SMC = Cast<UStaticMeshComponent>(PrimitiveComponent))
+						{
+							UE_LOG(LogHLODBuilder, Verbose, TEXT("%s    * %s"), Padding, *SMC->GetStaticMesh()->GetName());
+						}
 					}
 					else
 					{
@@ -139,7 +152,10 @@ public:
 			// Gather primitives from underlying actors
 			for (const AActor* UnderlyingActor : UnderlyingActors)
 			{
-				GatherPrimitivesFromActor(UnderlyingActor);
+				if (UnderlyingActor->IsHLODRelevant())
+				{
+					GatherPrimitivesFromActor(UnderlyingActor, Actor);
+				}
 			}
 		}
 
