Fix Static Analysis Win64.

diff --git a/Engine/Source/Runtime/Engine/Private/WorldPartition/WorldPartitionRuntimeHash.cpp b/Engine/Source/Runtime/Engine/Private/WorldPartition/WorldPartitionRuntimeHash.cpp
--- a/Engine/Source/Runtime/Engine/Private/WorldPartition/WorldPartitionRuntimeHash.cpp
+++ b/Engine/Source/Runtime/Engine/Private/WorldPartition/WorldPartitionRuntimeHash.cpp
@@ -216,7 +216,7 @@ void UWorldPartitionRuntimeHash::CheckForErrors() const
 						->AddToken(FTextToken::Create(LOCTEXT("MapCheck_WorldPartition_Actor", "Actor")))
 						->AddToken(FAssetNameToken::Create(GetActorLabel(ActorDescView)))
 						->AddToken(FTextToken::Create(LOCTEXT("MapCheck_WorldPartition_HaveMissingRefsTo", "have missing references to")))
-						->AddToken(FAssetNameToken::Create(GetActorLabel(*ActorDescRefView)))
+						->AddToken(FTextToken::Create(FText::FromString(ActorDescRefGuid.ToString())))
 						->AddToken(FMapErrorToken::Create(FName(TEXT("WorldPartition_MissingActorReference_CheckForErrors"))));
 				}
 			}
