WorldPartition: IsStreamingCompleted Query now ignores cells that are loaded by a DataLayer not part of the Query

diff --git a/Engine/Source/Runtime/Engine/Private/WorldPartition/WorldPartitionRuntimeSpatialHash.cpp b/Engine/Source/Runtime/Engine/Private/WorldPartition/WorldPartitionRuntimeSpatialHash.cpp
--- a/Engine/Source/Runtime/Engine/Private/WorldPartition/WorldPartitionRuntimeSpatialHash.cpp
+++ b/Engine/Source/Runtime/Engine/Private/WorldPartition/WorldPartitionRuntimeSpatialHash.cpp
@@ -102,7 +102,7 @@ const FSquare2DGridHelper& FSpatialHashStreamingGrid::GetGridHelper() const
 	return *GridHelper;
 }
 
-void FSpatialHashStreamingGrid::GetCells(const TArray<FWorldPartitionStreamingQuerySource>& QuerySources, TSet<const UWorldPartitionRuntimeCell*>& OutCells) const
+void FSpatialHashStreamingGrid::GetCells(const FWorldPartitionStreamingQuerySource& QuerySource, TSet<const UWorldPartitionRuntimeCell*>& OutCells) const
 {
 	auto ShouldAddCell = [](const UWorldPartitionRuntimeCell* Cell, const FWorldPartitionStreamingQuerySource& QuerySource)
 	{
@@ -122,41 +122,38 @@ void FSpatialHashStreamingGrid::GetCells(const TArray<FWorldPartitionStreamingQu
 	};
 
 	const FSquare2DGridHelper& Helper = GetGridHelper();
-	for (const FWorldPartitionStreamingQuerySource& QuerySource : QuerySources)
-	{
 		// Spatial Query
-		if (QuerySource.bSpatialQuery)
+	if (QuerySource.bSpatialQuery)
+	{
+		const FSphere GridSphere(QuerySource.Location, QuerySource.bUseLoadingRangeRadius ? GetLoadingRange() : QuerySource.Radius);
+		Helper.ForEachIntersectingCells(GridSphere, [&](const FIntVector& Coords)
 		{
-			const FSphere GridSphere(QuerySource.Location, QuerySource.bUseLoadingRangeRadius? GetLoadingRange() : QuerySource.Radius);
-			Helper.ForEachIntersectingCells(GridSphere, [&](const FIntVector& Coords)
+			if (const int32* LayerCellIndexPtr = GridLevels[Coords.Z].LayerCellsMapping.Find(Coords.Y * Helper.Levels[Coords.Z].GridSize + Coords.X))
+			{
+				const FSpatialHashStreamingGridLayerCell& LayerCell = GridLevels[Coords.Z].LayerCells[*LayerCellIndexPtr];
+				for (const UWorldPartitionRuntimeCell* Cell : LayerCell.GridCells)
 				{
-					if (const int32* LayerCellIndexPtr = GridLevels[Coords.Z].LayerCellsMapping.Find(Coords.Y * Helper.Levels[Coords.Z].GridSize + Coords.X))
+					if (ShouldAddCell(Cell, QuerySource))
 					{
-						const FSpatialHashStreamingGridLayerCell& LayerCell = GridLevels[Coords.Z].LayerCells[*LayerCellIndexPtr];
-						for (const UWorldPartitionRuntimeCell* Cell : LayerCell.GridCells)
-						{
-							if (ShouldAddCell(Cell, QuerySource))
-							{
-								OutCells.Add(Cell);
-							}
-						}
+						OutCells.Add(Cell);
 					}
-				});
-		}
+				}
+			}
+		});
+	}
 
-		// Non Spatial (always included)
-		const int32 TopLevel = GridLevels.Num() - 1;
-		for (const FSpatialHashStreamingGridLayerCell& LayerCell : GridLevels[TopLevel].LayerCells)
+	// Non Spatial (always included)
+	const int32 TopLevel = GridLevels.Num() - 1;
+	for (const FSpatialHashStreamingGridLayerCell& LayerCell : GridLevels[TopLevel].LayerCells)
+	{
+		for (const UWorldPartitionRuntimeCell* Cell : LayerCell.GridCells)
 		{
-			for (const UWorldPartitionRuntimeCell* Cell : LayerCell.GridCells)
+			if (ShouldAddCell(Cell, QuerySource))
 			{
-				if (ShouldAddCell(Cell, QuerySource))
-				{
-					OutCells.Add(Cell);
-				}
+				OutCells.Add(Cell);
 			}
 		}
-	}
+	}	
 }
 
 void FSpatialHashStreamingGrid::GetCells(const TArray<FWorldPartitionStreamingSource>& Sources, const UDataLayerSubsystem* DataLayerSubsystem, TSet<const UWorldPartitionRuntimeCell*>& OutActivateCells, TSet<const UWorldPartitionRuntimeCell*>& OutLoadCells) const
@@ -947,11 +944,11 @@ int32 UWorldPartitionRuntimeSpatialHash::GetAllStreamingCells(TSet<const UWorldP
 	return Cells.Num();
 }
 
-bool UWorldPartitionRuntimeSpatialHash::GetStreamingCells(const TArray<FWorldPartitionStreamingQuerySource>& Sources, TSet<const UWorldPartitionRuntimeCell*>& OutCells) const 
+bool UWorldPartitionRuntimeSpatialHash::GetStreamingCells(const FWorldPartitionStreamingQuerySource& QuerySource, TSet<const UWorldPartitionRuntimeCell*>& OutCells) const 
 {
 	for (const FSpatialHashStreamingGrid& StreamingGrid : StreamingGrids)
 	{
-		StreamingGrid.GetCells(Sources, OutCells);
+		StreamingGrid.GetCells(QuerySource, OutCells);
 	}
 
 	return !!OutCells.Num();
diff --git a/Engine/Source/Runtime/Engine/Private/WorldPartition/WorldPartitionStreamingPolicy.cpp b/Engine/Source/Runtime/Engine/Private/WorldPartition/WorldPartitionStreamingPolicy.cpp
--- a/Engine/Source/Runtime/Engine/Private/WorldPartition/WorldPartitionStreamingPolicy.cpp
+++ b/Engine/Source/Runtime/Engine/Private/WorldPartition/WorldPartitionStreamingPolicy.cpp
@@ -9,6 +9,7 @@
 #include "WorldPartition/WorldPartitionRuntimeHash.h"
 #include "WorldPartition/WorldPartitionStreamingSource.h"
 #include "WorldPartition/WorldPartition.h"
+#include "WorldPartition/DataLayer/DataLayerSubsystem.h"
 #include "GameFramework/PlayerController.h"
 #include "Engine/LocalPlayer.h"
 #include "Engine/World.h"
@@ -188,19 +189,38 @@ void UWorldPartitionStreamingPolicy::UpdateStreamingState()
 
 bool UWorldPartitionStreamingPolicy::IsStreamingCompleted(EWorldPartitionRuntimeCellState QueryState, const TArray<FWorldPartitionStreamingQuerySource>& QuerySources, bool bExactState) const
 {
-	TSet<const UWorldPartitionRuntimeCell*> Cells;
-	WorldPartition->RuntimeHash->GetStreamingCells(QuerySources, Cells);
-	
-	for (const UWorldPartitionRuntimeCell* Cell : Cells)
+	const UDataLayerSubsystem* DataLayerSubsystem = WorldPartition->GetWorld()->GetSubsystem<UDataLayerSubsystem>();
+
+	for (const FWorldPartitionStreamingQuerySource& QuerySource : QuerySources)
 	{
-		EWorldPartitionRuntimeCellState CellState = GetCurrentStateForCell(Cell);
-		if (CellState != QueryState)
+		TSet<const UWorldPartitionRuntimeCell*> Cells;
+		WorldPartition->RuntimeHash->GetStreamingCells(QuerySource, Cells);
+
+		for (const UWorldPartitionRuntimeCell* Cell : Cells)
 		{
-			if (bExactState || CellState < QueryState)
+			EWorldPartitionRuntimeCellState CellState = GetCurrentStateForCell(Cell);
+			if (CellState != QueryState)
 			{
-				return false;
+				bool bSkipCell = false;
+				// If we are querying for Unloaded/Loaded but a Cell is part of a data layer outside of the query that is activated do not consider it
+				if (QueryState < CellState)
+				{
+					for (const FName& CellDataLayer : Cell->GetDataLayers())
+					{
+						if (!QuerySource.DataLayers.Contains(CellDataLayer) && DataLayerSubsystem->GetDataLayerStateByName(CellDataLayer) > EDataLayerState::Unloaded)
+						{
+							bSkipCell = true;
+							break;
+						}
+					}
+				}
+								
+				if (!bSkipCell && (bExactState || CellState < QueryState))
+				{
+					return false;
+				}
 			}
-		}
+		}	
 	}
 
 	return true;
diff --git a/Engine/Source/Runtime/Engine/Public/WorldPartition/DataLayer/DataLayer.h b/Engine/Source/Runtime/Engine/Public/WorldPartition/DataLayer/DataLayer.h
--- a/Engine/Source/Runtime/Engine/Public/WorldPartition/DataLayer/DataLayer.h
+++ b/Engine/Source/Runtime/Engine/Public/WorldPartition/DataLayer/DataLayer.h
@@ -19,6 +19,8 @@ enum class EDataLayerState : uint8
 	Activated
 };
 
+static_assert(EDataLayerState::Unloaded < EDataLayerState::Loaded && EDataLayerState::Loaded < EDataLayerState::Activated, "Streaming Query code is dependent on this being true");
+
 UCLASS(Config = Engine, PerObjectConfig, Within = WorldDataLayers, BlueprintType)
 class ENGINE_API UDataLayer : public UObject
 {
diff --git a/Engine/Source/Runtime/Engine/Public/WorldPartition/WorldPartitionRuntimeCell.h b/Engine/Source/Runtime/Engine/Public/WorldPartition/WorldPartitionRuntimeCell.h
--- a/Engine/Source/Runtime/Engine/Public/WorldPartition/WorldPartitionRuntimeCell.h
+++ b/Engine/Source/Runtime/Engine/Public/WorldPartition/WorldPartitionRuntimeCell.h
@@ -100,6 +100,8 @@ enum class EWorldPartitionRuntimeCellState : uint8
 	Activated
 };
 
+static_assert(EWorldPartitionRuntimeCellState::Unloaded < EWorldPartitionRuntimeCellState::Loaded && EWorldPartitionRuntimeCellState::Loaded < EWorldPartitionRuntimeCellState::Activated, "Streaming Query code is dependent on this being true");
+
 /**
  * Represents a PIE/Game streaming cell which points to external actor/data chunk packages
  */
diff --git a/Engine/Source/Runtime/Engine/Public/WorldPartition/WorldPartitionRuntimeHash.h b/Engine/Source/Runtime/Engine/Public/WorldPartition/WorldPartitionRuntimeHash.h
--- a/Engine/Source/Runtime/Engine/Public/WorldPartition/WorldPartitionRuntimeHash.h
+++ b/Engine/Source/Runtime/Engine/Public/WorldPartition/WorldPartitionRuntimeHash.h
@@ -41,7 +41,7 @@ class ENGINE_API UWorldPartitionRuntimeHash : public UObject
 
 	// Streaming interface
 	virtual int32 GetAllStreamingCells(TSet<const UWorldPartitionRuntimeCell*>& Cells, bool bIncludeDataLayers = false) const { return 0; }
-	virtual bool GetStreamingCells(const TArray<FWorldPartitionStreamingQuerySource>& Sources, TSet<const UWorldPartitionRuntimeCell*>& OutCells) const { return false; }
+	virtual bool GetStreamingCells(const FWorldPartitionStreamingQuerySource& QuerySource, TSet<const UWorldPartitionRuntimeCell*>& OutCells) const { return false; }
 	virtual bool GetStreamingCells(const TArray<FWorldPartitionStreamingSource>& Sources, TSet<const UWorldPartitionRuntimeCell*>& OutActivateCells, TSet<const UWorldPartitionRuntimeCell*>& OutLoadCells) const { return false; };
 	virtual void SortStreamingCellsByImportance(const TSet<const UWorldPartitionRuntimeCell*>& InCells, const TArray<FWorldPartitionStreamingSource>& InSources, TArray<const UWorldPartitionRuntimeCell*, TInlineAllocator<256>>& OutSortedCells) const {}
 
diff --git a/Engine/Source/Runtime/Engine/Public/WorldPartition/WorldPartitionRuntimeSpatialHash.h b/Engine/Source/Runtime/Engine/Public/WorldPartition/WorldPartitionRuntimeSpatialHash.h
--- a/Engine/Source/Runtime/Engine/Public/WorldPartition/WorldPartitionRuntimeSpatialHash.h
+++ b/Engine/Source/Runtime/Engine/Public/WorldPartition/WorldPartitionRuntimeSpatialHash.h
@@ -124,7 +124,7 @@ struct FSpatialHashStreamingGrid
 	}
 
 	// Used by PIE/Game
-	void GetCells(const TArray<FWorldPartitionStreamingQuerySource>& QuerySources, TSet<const UWorldPartitionRuntimeCell*>& OutCells) const;
+	void GetCells(const FWorldPartitionStreamingQuerySource& QuerySource, TSet<const UWorldPartitionRuntimeCell*>& OutCells) const;
 	void GetCells(const TArray<FWorldPartitionStreamingSource>& Sources, const class UDataLayerSubsystem* DataLayerSubsystem, TSet<const UWorldPartitionRuntimeCell*>& OutActivateCells, TSet<const UWorldPartitionRuntimeCell*>& OutLoadCells) const;
 	void GetAlwaysLoadedCells(const UDataLayerSubsystem* DataLayerSubsystem, TSet<const UWorldPartitionRuntimeCell*>& OutActivateCells, TSet<const UWorldPartitionRuntimeCell*>& OutLoadCells) const;
 	void Draw2D(UCanvas* Canvas, const TArray<FWorldPartitionStreamingSource>& Sources, const FBox& Region, const FBox2D& GridScreenBounds, TFunctionRef<FVector2D(const FVector2D&)> WorldToScreen) const;
@@ -209,7 +209,7 @@ public:
 
 	// streaming interface
 	virtual int32 GetAllStreamingCells(TSet<const UWorldPartitionRuntimeCell*>& Cells, bool bIncludeDataLayers = false) const override;
-	virtual bool GetStreamingCells(const TArray<FWorldPartitionStreamingQuerySource>& Sources, TSet<const UWorldPartitionRuntimeCell*>& OutCells) const override;
+	virtual bool GetStreamingCells(const FWorldPartitionStreamingQuerySource& QuerySource, TSet<const UWorldPartitionRuntimeCell*>& OutCells) const override;
 	virtual bool GetStreamingCells(const TArray<FWorldPartitionStreamingSource>& Sources, TSet<const UWorldPartitionRuntimeCell*>& OutActivateCells, TSet<const UWorldPartitionRuntimeCell*>& OutLoadCells) const override;
 	virtual void SortStreamingCellsByImportance(const TSet<const UWorldPartitionRuntimeCell*>& InCells, const TArray<FWorldPartitionStreamingSource>& InSources, TArray<const UWorldPartitionRuntimeCell*, TInlineAllocator<256>>& OutSortedCells) const override;
 
