The torso of Metahuman_001 appears to flicker during runtime of the MH Sample scene in UE5EA
#fix reset pose states to prevent invalid pose influences, ensure that bone transform is always initialised with first weighted influence (which could be not the first one)

diff --git a/Engine/Source/Runtime/Engine/Private/Animation/PoseAsset.cpp b/Engine/Source/Runtime/Engine/Private/Animation/PoseAsset.cpp
--- a/Engine/Source/Runtime/Engine/Private/Animation/PoseAsset.cpp
+++ b/Engine/Source/Runtime/Engine/Private/Animation/PoseAsset.cpp
@@ -631,6 +631,7 @@ bool UPoseAsset::GetAnimationPose(struct FAnimationPoseData& OutAnimationPoseDat
 		WeightedPoseIndices.Reset();
 
 		TArray<bool>& WeightedPoses = EvalData.WeightedPoses;
+		WeightedPoses.Reset();
 		WeightedPoses.SetNumZeroed(NumPoses, false);
 
 		float TotalWeight = 0.f;
@@ -714,6 +715,7 @@ bool UPoseAsset::GetAnimationPose(struct FAnimationPoseData& OutAnimationPoseDat
 				if (NumInfluences)
 				{
 					float TotalLocalWeight = 0.f;
+					bool bSet = false;
 					// Only loop over poses known to influence this track its final transform
 					for (int32 Index = 0; Index < NumInfluences; ++Index)
 					{
@@ -731,9 +733,10 @@ bool UPoseAsset::GetAnimationPose(struct FAnimationPoseData& OutAnimationPoseDat
 							const FTransform& BonePose = PoseData.LocalSpacePose[BonePoseIndex];
 
 							// Set weighted value For first pose, if applicable
-							if(Index == 0 && !bAdditivePose)
+							if(!bSet && !bAdditivePose)
 							{								
 								OutBoneTransform = BonePose * ScalarRegister(Weight);
+								bSet = true;
 							}
 							else
 							{
