Remove deprecated functions with UE4 in the deprecation message

diff --git a/Engine/Source/Runtime/Engine/Classes/Engine/BlueprintCore.h b/Engine/Source/Runtime/Engine/Classes/Engine/BlueprintCore.h
--- a/Engine/Source/Runtime/Engine/Classes/Engine/BlueprintCore.h
+++ b/Engine/Source/Runtime/Engine/Classes/Engine/BlueprintCore.h
@@ -34,14 +34,6 @@ private:
 	FGuid BlueprintGuid;
 
 public:
-	UE_DEPRECATED(4.22, "The minimum UE4 object version implies the Blueprint generated class is always authoritative. It is no longer necessary to explicitly set it.")
-	void SetLegacyGeneratedClassIsAuthoritative() {}
-
-	UE_DEPRECATED(4.22, "The minimum UE4 object version implies the Blueprint generated class is always authoritative. It is no longer necessary to explicitly check it.")
-	bool IsGeneratedClassAuthoritative()
-	{
-		return true;
-	}
 
 	virtual void Serialize( FArchive& Ar ) override;
 
