Fix for editor-only crash when previewing a sound after changing USoundWave properties

diff --git a/Engine/Source/Runtime/AudioCodecEngine/Private/DecoderInputBackCompat.cpp b/Engine/Source/Runtime/AudioCodecEngine/Private/DecoderInputBackCompat.cpp
--- a/Engine/Source/Runtime/AudioCodecEngine/Private/DecoderInputBackCompat.cpp
+++ b/Engine/Source/Runtime/AudioCodecEngine/Private/DecoderInputBackCompat.cpp
@@ -85,9 +85,10 @@ namespace Audio
 	void FBackCompatInput::SeekToTime(const float InSeconds)
 	{
 		ICompressedAudioInfo* Info = GetInfo();
-		ensure(Info);
-
-		Info->SeekToTime(InSeconds);
+		if (ensure(Info))
+		{
+			Info->SeekToTime(InSeconds);
+		}
 	}
 
 	bool FBackCompatInput::IsEndOfStream() const
diff --git a/Engine/Source/Runtime/Engine/Private/SoundWave.cpp b/Engine/Source/Runtime/Engine/Private/SoundWave.cpp
--- a/Engine/Source/Runtime/Engine/Private/SoundWave.cpp
+++ b/Engine/Source/Runtime/Engine/Private/SoundWave.cpp
@@ -1188,11 +1188,15 @@ void USoundWave::EnsureZerothChunkIsLoaded()
 	}
 
 #if WITH_EDITOR 
-	if (RunningPlatformData.IsValid())
+
+	if (!RunningPlatformData.IsValid())
 	{
-		CachePlatformData(false);
+		// we may be in the middle of garbage collection, don't access RunningPlatformData
+		return;
 	}
 
+	CachePlatformData(false);
+
 	// If we're running the editor, we'll need to retrieve the chunked audio from the DDC:
 	uint8* TempChunkBuffer = nullptr;
 	int32 ChunkSizeInBytes = RunningPlatformData->GetChunkFromDDC(0, &TempChunkBuffer, true);
@@ -2445,6 +2449,10 @@ bool USoundWave::HasCookedAmplitudeEnvelopeData() const
 
 FSoundWaveProxyPtr USoundWave::CreateSoundWaveProxy()
 {
+#if WITH_EDITORONLY_DATA
+	EnsureZerothChunkIsLoaded();
+#endif // #if WITH_EDITORONLY_DATA
+
 	return MakeShared<FSoundWaveProxy, ESPMode::ThreadSafe>(this);
 }
 
@@ -3101,12 +3109,10 @@ uint32 FSoundWaveProxy::GetSizeOfChunk(uint32 ChunkIndex) const
 
 TArrayView<const uint8> FSoundWaveProxy::GetZerothChunk(const FSoundWaveProxyPtr& SoundWaveProxy, bool bForImmediatePlayback)
 {
-	if (ensure(SoundWaveProxy.IsValid()))
+	if (ensure(SoundWaveProxy.IsValid()) && SoundWaveProxy->IsZerothChunkDataLoaded())
 	{
 		if (SoundWaveProxy->bShouldUseStreamCaching)
 		{
-			check(SoundWaveProxy->IsZerothChunkDataLoaded());
-
 			if (SoundWaveProxy->GetNumChunks() > 1)
 			{
 				// Prime first chunk for playback
