Fix inconsistent logic for using distance fields needing to access per-instance data in ISM/HISM.

diff --git a/Engine/Source/Runtime/Engine/Private/InstancedStaticMesh.cpp b/Engine/Source/Runtime/Engine/Private/InstancedStaticMesh.cpp
--- a/Engine/Source/Runtime/Engine/Private/InstancedStaticMesh.cpp
+++ b/Engine/Source/Runtime/Engine/Private/InstancedStaticMesh.cpp
@@ -1286,14 +1286,17 @@ void FInstancedStaticMeshSceneProxy::GetDistancefieldInstanceData(TArray<FMatrix
 {
 	ObjectLocalToWorldTransforms.Reset();
 
-	const uint32 NumInstances = InstancedRenderData.PerInstanceRenderData->InstanceBuffer.GetNumInstances();
-	for (uint32 InstanceIndex = 0; InstanceIndex < NumInstances; InstanceIndex++)
+	if (ensureMsgf(InstancedRenderData.PerInstanceRenderData->InstanceBuffer.RequireCPUAccess, TEXT("GetDistancefieldInstanceData requires a CPU copy of the per-instance data to be accessible. Possible mismatch in ComponentRequestsCPUAccess / IncludePrimitiveInDistanceFieldSceneData filtering.")))
 	{
-		FMatrix InstanceToLocal;
-		InstancedRenderData.PerInstanceRenderData->InstanceBuffer.GetInstanceTransform(InstanceIndex, InstanceToLocal);
-		InstanceToLocal.M[3][3] = 1.0f;
+		const uint32 NumInstances = InstancedRenderData.PerInstanceRenderData->InstanceBuffer.GetNumInstances();
+		for (uint32 InstanceIndex = 0; InstanceIndex < NumInstances; InstanceIndex++)
+		{
+			FMatrix InstanceToLocal;
+			InstancedRenderData.PerInstanceRenderData->InstanceBuffer.GetInstanceTransform(InstanceIndex, InstanceToLocal);
+			InstanceToLocal.M[3][3] = 1.0f;
 
-		ObjectLocalToWorldTransforms.Add(InstanceToLocal * GetLocalToWorld());
+			ObjectLocalToWorldTransforms.Add(InstanceToLocal * GetLocalToWorld());
+		}
 	}
 }
 
@@ -3296,11 +3299,17 @@ static bool ComponentRequestsCPUAccess(UInstancedStaticMeshComponent* InComponen
 	{
 		if ((FeatureLevel > ERHIFeatureLevel::ES3_1) || IsMobileDistanceFieldEnabled(GMaxRHIShaderPlatform))
 		{
-			static const auto CVar = IConsoleManager::Get().FindTConsoleVariableDataInt(TEXT("r.GenerateMeshDistanceFields"));
-
-			bNeedsCPUAccess |= (InComponent->CastShadow && InComponent->bAffectDistanceFieldLighting
-				// Distance field algorithms need access to instance data on the CPU
-				&& (CVar->GetValueOnAnyThread(true) != 0 || StaticMesh->bGenerateMeshDistanceField));
+			// Mirror the conditions used in the FPrimitiveSceneProxy since these are used in IncludePrimitiveInDistanceFieldSceneData in RendererScene.cpp to filter the 
+			// primitives that are included in the distance field scene. If these are not in sync, the host copy may be discarded and thus crashing in the distance field update.
+			auto ShaderPlatform = GetFeatureLevelShaderPlatform(FeatureLevel);
+			bNeedsCPUAccess |= PrimitiveNeedsDistanceFieldSceneData(
+				ShouldAllPrimitivesHaveDistanceField(ShaderPlatform),
+				/* bCastsDynamicIndirectShadow */ InComponent->bCastDynamicShadow && InComponent->CastShadow && InComponent->bCastDistanceFieldIndirectShadow && InComponent->Mobility != EComponentMobility::Static,
+				InComponent->bAffectDistanceFieldLighting,
+				true, /* conservatively overestimate DrawInGame - it has complex logic in the Proxy. */
+				InComponent->bCastHiddenShadow,
+				/* bCastsDynamicShadow */ InComponent->bCastDynamicShadow && InComponent->CastShadow && !InComponent->GetShadowIndirectOnly(),
+				InComponent->bAffectDynamicIndirectLighting);
 		}
 
 		// Check Nanite
diff --git a/Engine/Source/Runtime/Engine/Public/SceneManagement.h b/Engine/Source/Runtime/Engine/Public/SceneManagement.h
--- a/Engine/Source/Runtime/Engine/Public/SceneManagement.h
+++ b/Engine/Source/Runtime/Engine/Public/SceneManagement.h
@@ -1052,6 +1052,31 @@ inline bool DoesPlatformSupportDistanceFieldAO(EShaderPlatform Platform)
 	return DoesPlatformSupportDistanceFields(Platform);
 }
 
+inline bool ShouldAllPrimitivesHaveDistanceField(EShaderPlatform ShaderPlatform)
+{
+	static const auto CVar = IConsoleManager::Get().FindTConsoleVariableDataInt(TEXT("r.GenerateMeshDistanceFields"));
+
+	return (DoesPlatformSupportDistanceFieldAO(ShaderPlatform) || DoesPlatformSupportDistanceFieldShadowing(ShaderPlatform)) && CVar->GetValueOnGameThread() != 0 && IsUsingDistanceFields(ShaderPlatform);
+}
+
+/**
+ * Centralized decision function to avoid diverging logic.
+ */
+inline bool PrimitiveNeedsDistanceFieldSceneData(bool bTrackAllPrimitives,
+	bool bCastsDynamicIndirectShadow,
+	bool bAffectsDistanceFieldLighting,
+	bool bIsDrawnInGame,
+	bool bCastsHiddenShadow,
+	bool bCastsDynamicShadow,
+	bool bAffectsDynamicIndirectLighting)
+{
+	return (bTrackAllPrimitives || bCastsDynamicIndirectShadow)
+		&& bAffectsDistanceFieldLighting
+		&& (bIsDrawnInGame || bCastsHiddenShadow)
+		&& (bCastsDynamicShadow || bAffectsDynamicIndirectLighting);
+}
+
+
 BEGIN_GLOBAL_SHADER_PARAMETER_STRUCT(FMobileReflectionCaptureShaderParameters,ENGINE_API)
 	SHADER_PARAMETER(FVector4, Params) // x - inv average brightness, y - sky cubemap max mip, z - Max value for RGBM, w - unused
 	SHADER_PARAMETER_TEXTURE(TextureCube, Texture)
diff --git a/Engine/Source/Runtime/Renderer/Private/RendererScene.cpp b/Engine/Source/Runtime/Renderer/Private/RendererScene.cpp
--- a/Engine/Source/Runtime/Renderer/Private/RendererScene.cpp
+++ b/Engine/Source/Runtime/Renderer/Private/RendererScene.cpp
@@ -379,9 +379,7 @@ FDistanceFieldSceneData::FDistanceFieldSceneData(EShaderPlatform ShaderPlatform)
 
 	HeightFieldObjectBuffers = nullptr;
 
-	static const auto CVar = IConsoleManager::Get().FindTConsoleVariableDataInt(TEXT("r.GenerateMeshDistanceFields"));
-
-	bTrackAllPrimitives = (DoesPlatformSupportDistanceFieldAO(ShaderPlatform) || DoesPlatformSupportDistanceFieldShadowing(ShaderPlatform)) && CVar->GetValueOnGameThread() != 0 && IsUsingDistanceFields(ShaderPlatform);
+	bTrackAllPrimitives = ShouldAllPrimitivesHaveDistanceField(ShaderPlatform);
 
 	bCanUse16BitObjectIndices = RHISupportsBufferLoadTypeConversion(ShaderPlatform);
 
@@ -395,10 +393,14 @@ FDistanceFieldSceneData::~FDistanceFieldSceneData()
 
 bool IncludePrimitiveInDistanceFieldSceneData(bool bTrackAllPrimitives, const FPrimitiveSceneProxy* Proxy)
 {
-	return (bTrackAllPrimitives || Proxy->CastsDynamicIndirectShadow())
-		&& Proxy->AffectsDistanceFieldLighting()
-		&& (Proxy->IsDrawnInGame() || Proxy->CastsHiddenShadow())
-		&& (Proxy->CastsDynamicShadow() || Proxy->AffectsDynamicIndirectLighting());
+	return PrimitiveNeedsDistanceFieldSceneData(
+		bTrackAllPrimitives, 
+		Proxy->CastsDynamicIndirectShadow(), 
+		Proxy->AffectsDistanceFieldLighting(), 
+		Proxy->IsDrawnInGame(),  
+		Proxy->CastsHiddenShadow(), 
+		Proxy->CastsDynamicShadow(),
+		Proxy->AffectsDynamicIndirectLighting());
 }
 
 void FDistanceFieldSceneData::AddPrimitive(FPrimitiveSceneInfo* InPrimitive)
