Fix trivial UE4 in comments

diff --git a/Engine/Source/Runtime/Engine/Classes/Engine/EngineTypes.h b/Engine/Source/Runtime/Engine/Classes/Engine/EngineTypes.h
--- a/Engine/Source/Runtime/Engine/Classes/Engine/EngineTypes.h
+++ b/Engine/Source/Runtime/Engine/Classes/Engine/EngineTypes.h
@@ -1818,7 +1818,7 @@ struct FLightmassPrimitiveSettings
 
 	friend bool operator==(const FLightmassPrimitiveSettings& A, const FLightmassPrimitiveSettings& B)
 	{
-		//@todo UE4. Do we want a little 'leeway' in joining 
+		//@todo Do we want a little 'leeway' in joining 
 		if ((A.bUseTwoSidedLighting != B.bUseTwoSidedLighting) ||
 			(A.bShadowIndirectOnly != B.bShadowIndirectOnly) || 
 			(A.bUseEmissiveForStaticLighting != B.bUseEmissiveForStaticLighting) || 
@@ -1964,7 +1964,7 @@ struct FSwarmDebugOptions
 	{
 	}
 
-	//@todo UE4. For some reason, the global instance is not initializing to the default settings...
+	//@todo For some reason, the global instance is not initializing to the default settings...
 	// Be sure to update this function to properly set the desired initial values!!!!
 	void Touch();
 };
diff --git a/Engine/Source/Runtime/Engine/Private/Camera/CameraAnim.cpp b/Engine/Source/Runtime/Engine/Private/Camera/CameraAnim.cpp
--- a/Engine/Source/Runtime/Engine/Private/Camera/CameraAnim.cpp
+++ b/Engine/Source/Runtime/Engine/Private/Camera/CameraAnim.cpp
@@ -45,7 +45,7 @@ bool UCameraAnim::CreateFromInterpGroup(class UInterpGroup* SrcGroup, class AMat
 	if (CameraInterpGroup != SrcGroup)
 	{
 		// copy the source interp group for use in the CameraAnim
-		// @fixme jf: fixed this potentially creating an object of UInterpGroup and raw casting it to InterpGroupCamera.  No source data in UE4 to test though.
+		// @fixme jf: fixed this potentially creating an object of UInterpGroup and raw casting it to InterpGroupCamera.  No source data to test though.
 		CameraInterpGroup = Cast<UInterpGroupCamera>(StaticDuplicateObject(SrcGroup, this, NAME_None, RF_AllFlags, UInterpGroupCamera::StaticClass()));
 
 		if (CameraInterpGroup)
diff --git a/Engine/Source/Runtime/Engine/Private/Components/ActorComponent.cpp b/Engine/Source/Runtime/Engine/Private/Components/ActorComponent.cpp
--- a/Engine/Source/Runtime/Engine/Private/Components/ActorComponent.cpp
+++ b/Engine/Source/Runtime/Engine/Private/Components/ActorComponent.cpp
@@ -939,7 +939,7 @@ void UActorComponent::ConsolidatedPostEditChange(const FPropertyChangedEvent& Pr
 	// Don't do do a full recreate in this situation, and instead simply detach.
 	if( IsPendingKill() )
 	{
-		// @todo UE4 james should this call UnregisterComponent instead to remove itself from the RegisteredComponents array on the owner?
+		// @todo james should this call UnregisterComponent instead to remove itself from the RegisteredComponents array on the owner?
 		ExecuteUnregisterEvents();
 		WorldPrivate = nullptr;
 	}
diff --git a/Engine/Source/Runtime/Engine/Private/Components/BoxComponent.cpp b/Engine/Source/Runtime/Engine/Private/Components/BoxComponent.cpp
--- a/Engine/Source/Runtime/Engine/Private/Components/BoxComponent.cpp
+++ b/Engine/Source/Runtime/Engine/Private/Components/BoxComponent.cpp
@@ -48,7 +48,7 @@ bool InvalidateOrUpdateBoxBodySetup(BodySetupType& ShapeBodySetup, bool bUseArch
 	check(ShapeBodySetup->AggGeom.BoxElems.Num() == 1);
 	FKBoxElem* se = ShapeBodySetup->AggGeom.BoxElems.GetData();
 
-	// @todo UE4 do we allow this now?
+	// @todo do we allow this now?
 	// check for malformed values
 	if (BoxExtent.X < KINDA_SMALL_NUMBER)
 	{
diff --git a/Engine/Source/Runtime/Engine/Private/Components/CharacterMovementComponent.cpp b/Engine/Source/Runtime/Engine/Private/Components/CharacterMovementComponent.cpp
--- a/Engine/Source/Runtime/Engine/Private/Components/CharacterMovementComponent.cpp
+++ b/Engine/Source/Runtime/Engine/Private/Components/CharacterMovementComponent.cpp
@@ -1095,7 +1095,7 @@ void UCharacterMovementComponent::SetMovementMode(EMovementMode NewMovementMode,
 	// Handle change in movement mode
 	OnMovementModeChanged(PrevMovementMode, PrevCustomMode);
 
-	// @todo UE4 do we need to disable ragdoll physics here? Should this function do nothing if in ragdoll?
+	// @todo do we need to disable ragdoll physics here? Should this function do nothing if in ragdoll?
 }
 
 
@@ -1977,7 +1977,7 @@ void UCharacterMovementComponent::MaybeSaveBaseLocation()
 	}
 }
 
-// @todo UE4 - handle lift moving up and down through encroachment
+// @todo handle lift moving up and down through encroachment
 void UCharacterMovementComponent::UpdateBasedMovement(float DeltaSeconds)
 {
 	if (!HasValidData())
diff --git a/Engine/Source/Runtime/Engine/Private/Components/SkeletalMeshComponent.cpp b/Engine/Source/Runtime/Engine/Private/Components/SkeletalMeshComponent.cpp
--- a/Engine/Source/Runtime/Engine/Private/Components/SkeletalMeshComponent.cpp
+++ b/Engine/Source/Runtime/Engine/Private/Components/SkeletalMeshComponent.cpp
@@ -1725,8 +1725,7 @@ void USkeletalMeshComponent::ComputeRequiredBones(TArray<FBoneIndexType>& OutReq
 		MergeInBoneIndexArrays(OutRequiredBones, PhysAssetBones);
 	}
 
-	// Make sure that bones with per-poly collision are also always updated.
-	// TODO UE4
+	// TODO - Make sure that bones with per-poly collision are also always updated.
 
 	// Purge invisible bones and their children
 	// this has to be done before mirror table check/physics body checks
diff --git a/Engine/Source/Runtime/Engine/Private/Components/SkinnedMeshComponent.cpp b/Engine/Source/Runtime/Engine/Private/Components/SkinnedMeshComponent.cpp
--- a/Engine/Source/Runtime/Engine/Private/Components/SkinnedMeshComponent.cpp
+++ b/Engine/Source/Runtime/Engine/Private/Components/SkinnedMeshComponent.cpp
@@ -1168,8 +1168,7 @@ FBoxSphereBounds USkinnedMeshComponent::CalcMeshBound(const FVector& RootOffset,
 		NewBounds = FBoxSphereBounds(LocalToWorld.GetLocation(), FVector::ZeroVector, 0.f);
 	}
 
-	// Add bounds of any per-poly collision data.
-	// TODO UE4
+	// TODO - Add bounds of any per-poly collision data.
 
 	NewBounds.BoxExtent *= BoundsScale;
 	NewBounds.SphereRadius *= BoundsScale;
diff --git a/Engine/Source/Runtime/Engine/Private/World.cpp b/Engine/Source/Runtime/Engine/Private/World.cpp
--- a/Engine/Source/Runtime/Engine/Private/World.cpp
+++ b/Engine/Source/Runtime/Engine/Private/World.cpp
@@ -6839,7 +6839,7 @@ UWorld* FSeamlessTravelHandler::Tick()
 				}
 
 				// Copy cheat flags if the game info is present
-				// @todo UE4 FIXMELH - see if this exists, it should not since it's created in GameMode or it's garbage info
+				// @todo FIXMELH - see if this exists, it should not since it's created in GameMode or it's garbage info
 				if (LoadedWorld->NetworkManager != nullptr)
 				{
 					LoadedWorld->NetworkManager->bHasStandbyCheatTriggered = bHasStandbyCheatTriggered;
