Remove UE4 references from elements internal to this file

diff --git a/Engine/Source/Runtime/Engine/Private/GameplayStatics.cpp b/Engine/Source/Runtime/Engine/Private/GameplayStatics.cpp
--- a/Engine/Source/Runtime/Engine/Private/GameplayStatics.cpp
+++ b/Engine/Source/Runtime/Engine/Private/GameplayStatics.cpp
@@ -53,7 +53,7 @@
 
 #define LOCTEXT_NAMESPACE "GameplayStatics"
 
-static const int UE4_SAVEGAME_FILE_TYPE_TAG = 0x53415647;		// "SAVG"
+static const int UE_SAVEGAME_FILE_TYPE_TAG = 0x53415647;		// "SAVG"
 
 struct FSaveGameFileVersion
 {
@@ -104,7 +104,7 @@ FSaveGameHeader::FSaveGameHeader()
 {}
 
 FSaveGameHeader::FSaveGameHeader(TSubclassOf<USaveGame> ObjectType)
-	: FileTypeTag(UE4_SAVEGAME_FILE_TYPE_TAG)
+	: FileTypeTag(UE_SAVEGAME_FILE_TYPE_TAG)
 	, SaveGameFileVersion(FSaveGameFileVersion::LatestVersion)
 	, PackageFileUE4Version(GPackageFileUE4Version)
 	, SavedEngineVersion(FEngineVersion::Current())
@@ -135,7 +135,7 @@ void FSaveGameHeader::Read(FMemoryReader& MemoryReader)
 
 	MemoryReader << FileTypeTag;
 
-	if (FileTypeTag != UE4_SAVEGAME_FILE_TYPE_TAG)
+	if (FileTypeTag != UE_SAVEGAME_FILE_TYPE_TAG)
 	{
 		// this is an old saved game, back up the file pointer to the beginning and assume version 1
 		MemoryReader.Seek(0);
@@ -176,7 +176,7 @@ void FSaveGameHeader::Read(FMemoryReader& MemoryReader)
 void FSaveGameHeader::Write(FMemoryWriter& MemoryWriter)
 {
 	// write file type tag. identifies this file type and indicates it's using proper versioning
-	// since older UE4 versions did not version this data.
+	// since older UE versions did not version this data.
 	MemoryWriter << FileTypeTag;
 
 	// Write version for this file format
