WorldPartition HLODs - Use the MainContainerInstance bound rather than the world bound

diff --git a/Engine/Source/Runtime/Engine/Private/WorldPartition/RuntimeSpatialHash/RuntimeSpatialHashHLOD.cpp b/Engine/Source/Runtime/Engine/Private/WorldPartition/RuntimeSpatialHash/RuntimeSpatialHashHLOD.cpp
--- a/Engine/Source/Runtime/Engine/Private/WorldPartition/RuntimeSpatialHash/RuntimeSpatialHashHLOD.cpp
+++ b/Engine/Source/Runtime/Engine/Private/WorldPartition/RuntimeSpatialHash/RuntimeSpatialHashHLOD.cpp
@@ -92,7 +92,7 @@ static void DeletePackage(FWorldPartitionActorDesc* ActorDesc, ISourceControlHel
 	}
 }
 
-static TArray<FGuid> GenerateHLODsForGrid(UWorldPartition* WorldPartition, const FSpatialHashRuntimeGrid& RuntimeGrid, uint32 HLODLevel, FHLODCreationContext& Context, ISourceControlHelper* SourceControlHelper, bool bCreateActorsOnly, const TArray<const FActorClusterInstance*>& ClusterInstances)
+static TArray<FGuid> GenerateHLODsForGrid(UWorldPartition* WorldPartition, const FActorContainerInstance& MainContainerInstance, const FSpatialHashRuntimeGrid& RuntimeGrid, uint32 HLODLevel, FHLODCreationContext& Context, ISourceControlHelper* SourceControlHelper, bool bCreateActorsOnly, const TArray<const FActorClusterInstance*>& ClusterInstances)
 {
 	auto HasExceededMaxMemory = []()
 	{
@@ -124,7 +124,7 @@ static TArray<FGuid> GenerateHLODsForGrid(UWorldPartition* WorldPartition, const
 		);
 	};
 
-	const FBox WorldBounds = WorldPartition->GetWorldBounds();
+	const FBox WorldBounds = MainContainerInstance.Bounds;
 
 	const FSquare2DGridHelper PartitionedActors = GetPartitionedActors(WorldPartition, WorldBounds, RuntimeGrid, ClusterInstances);
 	const FSquare2DGridHelper::FGridLevel::FGridCell& AlwaysLoadedCell = PartitionedActors.GetAlwaysLoadedCell();
@@ -495,7 +495,7 @@ bool UWorldPartitionRuntimeSpatialHash::GenerateHLOD(ISourceControlHelper* Sourc
 	}), /* bInIncludeChildContainers=*/ false);
 
 	FActorContainerInstance& MainContainerInstance = *ClusterContext.GetClusterInstance(WorldPartition);
-		
+
 	TArray<TArray<const FActorClusterInstance*>> GridsClusters;
 	GridsClusters.InsertDefaulted(0, Grids.Num());
 
@@ -519,7 +519,7 @@ bool UWorldPartitionRuntimeSpatialHash::GenerateHLOD(ISourceControlHelper* Sourc
 	auto GenerateHLODs = [&GridsHLODActors, &MainContainerInstance, WorldPartition, &GridsDepth, &Context, SourceControlHelper, bCreateActorsOnly](const FSpatialHashRuntimeGrid& RuntimeGrid, uint32 HLODLevel, const TArray<const FActorClusterInstance*>& ActorClusterInstances)
 	{
 		// Generate HLODs for this grid
-		TArray<FGuid> HLODActors = GenerateHLODsForGrid(WorldPartition, RuntimeGrid, HLODLevel, Context, SourceControlHelper, bCreateActorsOnly, ActorClusterInstances);
+		TArray<FGuid> HLODActors = GenerateHLODsForGrid(WorldPartition, MainContainerInstance, RuntimeGrid, HLODLevel, Context, SourceControlHelper, bCreateActorsOnly, ActorClusterInstances);
 
 		for (const FGuid& HLODActorGuid : HLODActors)
 		{
