Change shader jobs' input hash to Blake3.

- Also unified FSHA* and FBlake3* functions somewhat.

diff --git a/Engine/Source/Runtime/Core/Private/Hash/Blake3.cpp b/Engine/Source/Runtime/Core/Private/Hash/Blake3.cpp
--- a/Engine/Source/Runtime/Core/Private/Hash/Blake3.cpp
+++ b/Engine/Source/Runtime/Core/Private/Hash/Blake3.cpp
@@ -1,6 +1,7 @@
 // Copyright Epic Games, Inc. All Rights Reserved.
 
 #include "Hash/Blake3.h"
+#include "Containers/UnrealString.h"
 
 #include "blake3.h"
 
@@ -37,4 +38,9 @@ FBlake3Hash FBlake3::HashBuffer(const void* Data, uint64 Size)
 	return Hash.Finalize();
 }
 
+FString LexToString(const FBlake3Hash& Hash)
+{
+	return BytesToHex((const uint8*)Hash.GetBytes(), sizeof(decltype(Hash.GetBytes())));
+}
+
 ///////////////////////////////////////////////////////////////////////////////////////////////////////////////////////
diff --git a/Engine/Source/Runtime/Core/Public/Hash/Blake3.h b/Engine/Source/Runtime/Core/Public/Hash/Blake3.h
--- a/Engine/Source/Runtime/Core/Public/Hash/Blake3.h
+++ b/Engine/Source/Runtime/Core/Public/Hash/Blake3.h
@@ -137,6 +137,9 @@ inline FArchive& operator<<(FArchive& Ar, FBlake3Hash& Hash)
 	return Ar;
 }
 
+/** Stringifies the hash for textual output (e.g. to log file) */
+CORE_API FString LexToString(const FBlake3Hash& Hash);
+
 inline uint32 GetTypeHash(const FBlake3Hash& Hash)
 {
 	return *reinterpret_cast<const uint32*>(Hash.GetBytes());
diff --git a/Engine/Source/Runtime/Core/Public/Misc/SecureHash.h b/Engine/Source/Runtime/Core/Public/Misc/SecureHash.h
--- a/Engine/Source/Runtime/Core/Public/Misc/SecureHash.h
+++ b/Engine/Source/Runtime/Core/Public/Misc/SecureHash.h
@@ -279,6 +279,15 @@ public:
 	// Finalize hash and report
 	void Final();
 
+	// Finalize hash and return it
+	FSHAHash Finalize()
+	{
+		Final();
+		FSHAHash Digest;
+		GetHash(reinterpret_cast<uint8*>(&Digest));
+		return Digest;
+	}
+
 	// Report functions: as pre-formatted and raw data
 	void GetHash(uint8 *puDest);
 
@@ -291,6 +300,20 @@ public:
 	 */
 	static void HashBuffer(const void* Data, uint64 DataSize, uint8* OutHash);
 
+	/**
+	 * Calculate the hash on a single block and return it
+	 *
+	 * @param Data Input data to hash
+	 * @param DataSize Size of the Data block
+	 * @return Resulting digest
+	 */
+	static FSHAHash HashBuffer(const void* Data, uint64 DataSize)
+	{
+		FSHAHash Hash;
+		HashBuffer(Data, DataSize, Hash.Hash);
+		return Hash;
+	}
+
 	/**
 	 * Generate the HMAC (Hash-based Message Authentication Code) for a block of data.
 	 * https://en.wikipedia.org/wiki/Hash-based_message_authentication_code
diff --git a/Engine/Source/Runtime/Core/Public/Serialization/MemoryHasher.h b/Engine/Source/Runtime/Core/Public/Serialization/MemoryHasher.h
--- a/Engine/Source/Runtime/Core/Public/Serialization/MemoryHasher.h
+++ b/Engine/Source/Runtime/Core/Public/Serialization/MemoryHasher.h
@@ -12,6 +12,7 @@
 #include "CoreGlobals.h"
 #include "Serialization/MemoryArchive.h"
 #include "Misc/SecureHash.h"
+#include "Hash/Blake3.h"
 
 /**
  * Archive for hashing arbitrary data
@@ -32,16 +33,9 @@ public:
 		BuilderState.Update(reinterpret_cast<uint8*>(Data), Num);
 	}
 
-	void Finalize()
+	HashDigest Finalize()
 	{
-		BuilderState.Final();
-	}
-
-	HashDigest GetHash()
-	{
-		HashDigest Digest;
-		BuilderState.GetHash(reinterpret_cast<uint8*>(&Digest));
-		return Digest;
+		return BuilderState.Finalize();
 	}
 
 	/**
@@ -64,3 +58,4 @@ protected:
 };
 
 using FMemoryHasherSHA1 = TMemoryHasher<FSHA1, FSHAHash>;
+using FMemoryHasherBlake3 = TMemoryHasher<FBlake3, FBlake3Hash>;
diff --git a/Engine/Source/Runtime/Engine/Private/ShaderCompiler/ShaderCompiler.cpp b/Engine/Source/Runtime/Engine/Private/ShaderCompiler/ShaderCompiler.cpp
--- a/Engine/Source/Runtime/Engine/Private/ShaderCompiler/ShaderCompiler.cpp
+++ b/Engine/Source/Runtime/Engine/Private/ShaderCompiler/ShaderCompiler.cpp
@@ -365,7 +365,7 @@ int32 FShaderCompileJobCollection::RemoveAllPendingJobsWithId(uint32 InId)
 		{
 			// Also look into the jobs that are cached
 			// Since each entry in DuplicateJobsWaitList is a list, and the head node can be removed, we essentially have to rebuild it
-			for (TMap<FSHAHash, FShaderCommonCompileJob*>::TIterator Iter(DuplicateJobsWaitList); Iter; ++Iter)
+			for (TMap<FShaderCommonCompileJob::FInputHash, FShaderCommonCompileJob*>::TIterator Iter(DuplicateJobsWaitList); Iter; ++Iter)
 			{
 				FShaderCommonCompileJob* ListHead = Iter.Value();
 				FShaderCommonCompileJob* NewListHead = ListHead;
@@ -439,12 +439,12 @@ void FShaderCompileJobCollection::SubmitJobs(const TArray<FShaderCommonCompileJo
 				bool bNewJob = true;
 				if (ShaderCompiler::IsJobCacheEnabled())
 				{
-					const FSHAHash& InputHash = Job->GetInputHash();
+					const FShaderCommonCompileJob::FInputHash& InputHash = Job->GetInputHash();
 
 					// see if we can find the job in the cache first
 					if (TArray<uint8>* ExistingOutput = CompletedJobsCache.Find(InputHash))
 					{
-						UE_LOG(LogShaderCompilers, UE_SHADERCACHE_LOG_LEVEL, TEXT("There is already a cached job with the ihash %s, processing the new one immediately."), *InputHash.ToString());
+						UE_LOG(LogShaderCompilers, UE_SHADERCACHE_LOG_LEVEL, TEXT("There is already a cached job with the ihash %s, processing the new one immediately."), *LexToString(InputHash));
 						FMemoryReader MemReader(*ExistingOutput);
 						Job->SerializeOutput(MemReader);
 
@@ -456,7 +456,7 @@ void FShaderCompileJobCollection::SubmitJobs(const TArray<FShaderCommonCompileJo
 					// see if another job with the same input hash is being worked on
 					else if (FShaderCommonCompileJob** DuplicateInFlight = JobsInFlight.Find(InputHash))
 					{
-						UE_LOG(LogShaderCompilers, UE_SHADERCACHE_LOG_LEVEL, TEXT("There is an outstanding job with the ihash %s, not submitting another one (adding to wait list)."), *InputHash.ToString());
+						UE_LOG(LogShaderCompilers, UE_SHADERCACHE_LOG_LEVEL, TEXT("There is an outstanding job with the ihash %s, not submitting another one (adding to wait list)."), *LexToString(InputHash));
 
 						// because of the cloned jobs, we need to maintain a separate mapping
 						FShaderCommonCompileJob** WaitListHead = DuplicateJobsWaitList.Find(InputHash);
@@ -536,7 +536,7 @@ void FShaderCompileJobCollection::AddToCacheAndProcessPending(FShaderCommonCompi
 
 	ensureMsgf(FinishedJob->bInputHashSet, TEXT("Finished job didn't have input hash set, was shader compiler jobs cache toggled runtime?"));
 
-	const FSHAHash& InputHash = FinishedJob->GetInputHash();
+	const FShaderCommonCompileJob::FInputHash& InputHash = FinishedJob->GetInputHash();
 	TArray<uint8> Output;
 	FMemoryWriter Writer(Output);
 	FinishedJob->SerializeOutput(Writer);
@@ -569,7 +569,7 @@ void FShaderCompileJobCollection::AddToCacheAndProcessPending(FShaderCommonCompi
 
 		if (NumOutstandingJobsWithSameHash > 0)
 		{
-			UE_LOG(LogShaderCompilers, UE_SHADERCACHE_LOG_LEVEL, TEXT("Processed %d outstanding jobs with the same ihash %s."), NumOutstandingJobsWithSameHash, *InputHash.ToString());
+			UE_LOG(LogShaderCompilers, UE_SHADERCACHE_LOG_LEVEL, TEXT("Processed %d outstanding jobs with the same ihash %s."), NumOutstandingJobsWithSameHash, *LexToString(InputHash));
 		}
 	}
 
@@ -6363,7 +6363,7 @@ void ProcessCompiledGlobalShaders(const TArray<FShaderCommonCompileJobPtr>& Comp
 	}
 }
 
-FSHAHash FShaderCompileJob::GetInputHash()
+FShaderCommonCompileJob::FInputHash FShaderCompileJob::GetInputHash()
 {
 	if (bInputHashSet)
 	{
@@ -6378,8 +6378,8 @@ FSHAHash FShaderCompileJob::GetInputHash()
 		Archive << Input.Environment;
 
 		// hash the source file so changes to files during the development are picked up
-		const FSHAHash &SourceHash = GetShaderFileHash(*Input.VirtualSourceFilePath, Input.Target.GetPlatform());
-		Archive << const_cast<FSHAHash &>(SourceHash);
+		const FSHAHash& SourceHash = GetShaderFileHash(*Input.VirtualSourceFilePath, Input.Target.GetPlatform());
+		Archive << const_cast<FSHAHash&>(SourceHash);
 
 		for (TMap<FString, FThreadSafeSharedStringPtr>::TConstIterator It(Input.Environment.IncludeVirtualPathToExternalContentsMap); It; ++It)
 		{
@@ -6405,10 +6405,9 @@ FSHAHash FShaderCompileJob::GetInputHash()
 	};
 
 	// use faster hasher that doesn't allocate memory
-	FMemoryHasherSHA1 MemHasher;
+	FMemoryHasherBlake3 MemHasher;
 	SerializeInputs(MemHasher);
-	MemHasher.Finalize();
-	InputHash = MemHasher.GetHash();
+	InputHash = MemHasher.Finalize();
 
 	if (GShaderCompilerDumpCompileJobInputs)
 	{
@@ -6434,11 +6433,10 @@ FSHAHash FShaderCompileJob::GetInputHash()
 		DumpAr->Serialize(MemoryBlob.GetData(), MemoryBlob.Num());
 
 		// as an additional debugging feature, make sure that the hash is the same as calculated by the memhasher
-		FSHAHash Check;
-		FSHA1::HashBuffer(MemoryBlob.GetData(), MemoryBlob.Num(), Check.Hash);
+		FBlake3Hash Check = FBlake3::HashBuffer(MemoryBlob.GetData(), MemoryBlob.Num());
 		if (Check != InputHash)
 		{
-			UE_LOG(LogShaderCompilers, Error, TEXT("Job input hash disagrees between FMemoryHasherSHA1 (%s) and FMemoryWriter + FSHA1 (%s, which was dumped to disk)"), *InputHash.ToString(), *Check.ToString());
+			UE_LOG(LogShaderCompilers, Error, TEXT("Job input hash disagrees between FMemoryHasherSHA1 (%s) and FMemoryWriter + FSHA1 (%s, which was dumped to disk)"), *LexToString(InputHash), *LexToString(Check));
 		}
 	}
 
@@ -6473,26 +6471,24 @@ void FShaderCompileJob::SerializeOutput(FArchive& Ar)
 	}
 }
 
-FSHAHash FShaderPipelineCompileJob::GetInputHash()
+FShaderCommonCompileJob::FInputHash FShaderPipelineCompileJob::GetInputHash()
 {
 	if (bInputHashSet)
 	{
 		return InputHash;
 	}
 
-	FSHA1 Hasher;
-
+	FBlake3 Hasher;
 	for (int32 Index = 0; Index < StageJobs.Num(); ++Index)
 	{
 		if (StageJobs[Index])
 		{
-			FSHAHash StageHash = StageJobs[Index]->GetInputHash();
-			Hasher.Update(&StageHash.Hash[0], sizeof(StageHash.Hash));
+			FShaderCommonCompileJob::FInputHash StageHash = StageJobs[Index]->GetInputHash();
+			Hasher.Update(StageHash.GetBytes(), sizeof(decltype(StageHash.GetBytes())));
 		}
 	}
 
-	Hasher.Final();
-	Hasher.GetHash(&InputHash.Hash[0]);
+	InputHash = Hasher.Finalize();
 
 	bInputHashSet = true;
 	return InputHash;
@@ -6527,7 +6523,7 @@ FShaderJobCache::FJobCachedOutput* FShaderJobCache::Find(const FJobInputHash& Ha
 
 			FStoredOutput** CannedOutput = Outputs.Find(*OutputHash);
 			// we should not allow a dangling input to output mapping to exist
-			checkf(CannedOutput != nullptr, TEXT("Inconsistency in FShaderJobCache - cache record for ihash %s exists, but output cannot be found."), *Hash.ToString());
+			checkf(CannedOutput != nullptr, TEXT("Inconsistency in FShaderJobCache - cache record for ihash %s exists, but output cannot be found."), *LexToString(Hash));
 			// update the output hit count
 			(*CannedOutput)->NumHits++;
 			return &(*CannedOutput)->JobOutput;
@@ -6558,8 +6554,7 @@ void FShaderJobCache::Add(const FJobInputHash& Hash, const FJobCachedOutput& Con
 		return;
 	}
 
-	FSHAHash OutputHash;
-	FSHA1::HashBuffer(Contents.GetData(), Contents.Num(), OutputHash.Hash);
+	FJobOutputHash OutputHash = FBlake3::HashBuffer(Contents.GetData(), Contents.Num());
 
 	// add the record
 	InputHashToOutput.Add(Hash, OutputHash);
@@ -6596,7 +6591,7 @@ void FShaderJobCache::Add(const FJobInputHash& Hash, const FJobCachedOutput& Con
 					{
 						MemoryThatWillBeUsed -= static_cast<uint64>(Iter.Value()->JobOutput.Num());
 
-						FSHAHash RemovedOutputHash = Iter.Key();
+						FJobOutputHash RemovedOutputHash = Iter.Key();
 						Iter.RemoveCurrent();
 
 						// remove all mappings
diff --git a/Engine/Source/Runtime/Engine/Public/ShaderCompiler.h b/Engine/Source/Runtime/Engine/Public/ShaderCompiler.h
--- a/Engine/Source/Runtime/Engine/Public/ShaderCompiler.h
+++ b/Engine/Source/Runtime/Engine/Public/ShaderCompiler.h
@@ -23,6 +23,7 @@
 #include "Misc/ScopeRWLock.h"
 #include "Containers/HashTable.h"
 #include "Containers/List.h"
+#include "Misc/Blake3.h"
 
 class FShaderCompileJob;
 class FShaderPipelineCompileJob;
@@ -107,6 +108,8 @@ using FPendingShaderMapCompileResultsPtr = TRefCountPtr<FPendingShaderMapCompile
 class FShaderCommonCompileJob : public TIntrusiveLinkedList<FShaderCommonCompileJob>
 {
 public:
+	using FInputHash = FBlake3Hash;
+
 	FPendingShaderMapCompileResultsPtr PendingShaderMap;
 
 	mutable FThreadSafeCounter NumRefs;
@@ -134,7 +137,7 @@ public:
 	/** Whether we hashed the inputs */
 	uint8 bInputHashSet : 1;
 	/** Hash of all the job inputs */
-	FSHAHash InputHash;
+	FInputHash InputHash;
 
 	uint32 AddRef() const
 	{
@@ -156,7 +159,7 @@ public:
 	}
 
 	/** Returns hash of all inputs for this job (needed for caching). */
-	virtual FSHAHash GetInputHash() { return FSHAHash(); }
+	virtual FInputHash GetInputHash() { return FInputHash(); }
 
 	/** Serializes (and deserializes) the output for caching purposes. */
 	virtual void SerializeOutput(FArchive& Ar) {}
@@ -245,7 +248,7 @@ public:
 	// List of pipelines that are sharing this job.
 	TMap<const FVertexFactoryType*, TArray<const FShaderPipelineType*>> SharingPipelines;
 
-	virtual ENGINE_API FSHAHash GetInputHash() override;
+	virtual ENGINE_API FInputHash GetInputHash() override;
 	virtual ENGINE_API void SerializeOutput(FArchive& Ar) override;
 
 	FShaderCompileJob(uint32 InHash, uint32 InId, EShaderCompileJobPriority InPriroity, const FShaderCompileJobKey& InKey) :
@@ -284,7 +287,7 @@ public:
 	TArray<TRefCountPtr<FShaderCompileJob>> StageJobs;
 	bool bFailedRemovingUnused;
 
-	virtual ENGINE_API FSHAHash GetInputHash() override;
+	virtual ENGINE_API FInputHash GetInputHash() override;
 	virtual ENGINE_API void SerializeOutput(FArchive& Ar) override;
 
 	FShaderPipelineCompileJob(uint32 InHash, uint32 InId, EShaderCompileJobPriority InPriroity, const FShaderPipelineCompileJobKey& InKey);
@@ -322,7 +325,7 @@ inline void FShaderCommonCompileJob::Destroy() const
 class FShaderJobCache
 {
 public:
-	using FJobInputHash = FSHAHash;
+	using FJobInputHash = FShaderCommonCompileJob::FInputHash;
 	using FJobCachedOutput = TArray<uint8>;
 
 	/** Looks for the job in the cache, returns null if not found */
@@ -341,7 +344,7 @@ public:
 	uint64 GetCurrentMemoryBudget() const;
 
 private:
-	using FJobOutputHash = FSHAHash;
+	using FJobOutputHash = FBlake3Hash;
 	struct FStoredOutput
 	{
 		/** How many times this output is referenced by the cached jobs */
@@ -502,10 +505,10 @@ private:
 	mutable FRWLock Lock;
 
 	/** Map of input hash to the jobs that we decided to execute. Note that mapping will miss cloned jobs (to avoid being a multimap). */
-	TMap<FSHAHash, FShaderCommonCompileJob*> JobsInFlight;
+	TMap<FShaderCommonCompileJob::FInputHash, FShaderCommonCompileJob*> JobsInFlight;
 
 	/** Map of input hash to the jobs that we delayed because a job with the same hash was executing. Each job is a head of a linked list of jobs with the same input hash (ihash) */
-	TMap<FSHAHash, FShaderCommonCompileJob*> DuplicateJobsWaitList;
+	TMap<FShaderCommonCompileJob::FInputHash, FShaderCommonCompileJob*> DuplicateJobsWaitList;
 
 	/** Cache for the completed jobs.*/
 	FShaderJobCache CompletedJobsCache;
