OFPA: Add new static ULevel::GetExternalActorsFolderName() method to avoid proliferation of "__ExternalActors__" in code

diff --git a/Engine/Source/Editor/SourceControlWindows/Private/SSourceControlCommon.cpp b/Engine/Source/Editor/SourceControlWindows/Private/SSourceControlCommon.cpp
--- a/Engine/Source/Editor/SourceControlWindows/Private/SSourceControlCommon.cpp
+++ b/Engine/Source/Editor/SourceControlWindows/Private/SSourceControlCommon.cpp
@@ -12,6 +12,8 @@
 #include "Widgets/Images/SLayeredImage.h"
 #include "Widgets/Layout/SBox.h"
 
+#include "Engine/Level.h"
+
 #define LOCTEXT_NAMESPACE "SourceControlChangelist"
 
 //////////////////////////////////////////////////////////////////////////
@@ -66,7 +68,7 @@ FFileTreeItem::FFileTreeItem(FSourceControlStateRef InFileState, bool bBeautifyP
 	{
 		// At the moment, getting the asset data from non-external assets yields issues with the package path
 		// so we will fall down to our recovery (below) instead
-		if (Filename.Contains("__ExternalActors__"))
+		if (Filename.Contains(ULevel::GetExternalActorsFolderName()))
 		{
 			const int64 MaxFetchSize = (1 << 20); // 1MB
 			// In the case of shelved "marked for delete", we'll piggy back on the non-shelved file
diff --git a/Engine/Source/Runtime/Engine/Classes/Engine/Level.h b/Engine/Source/Runtime/Engine/Classes/Engine/Level.h
--- a/Engine/Source/Runtime/Engine/Classes/Engine/Level.h
+++ b/Engine/Source/Runtime/Engine/Classes/Engine/Level.h
@@ -1017,6 +1017,12 @@ public:
 	 */
 	static ENGINE_API FString GetExternalActorsPath(UPackage* InLevelPackage, const FString& InPackageShortName = FString());
 
+	/**
+	 * Get the folder name from which all external actors paths are created
+	 * @return folder name
+	 */
+	static ENGINE_API const TCHAR* GetExternalActorsFolderName();
+		
 	/**
 	 * Create an package for this actor
 	 * @param InActorPath the fully qualified actor path, in the format: 'Outermost.Outer.Name'
diff --git a/Engine/Source/Runtime/Engine/Private/Level.cpp b/Engine/Source/Runtime/Engine/Private/Level.cpp
--- a/Engine/Source/Runtime/Engine/Private/Level.cpp
+++ b/Engine/Source/Runtime/Engine/Private/Level.cpp
@@ -2356,6 +2356,11 @@ TArray<FString> ULevel::GetOnDiskExternalActorPackages() const
 	return ActorPackageNames;
 }
 
+const TCHAR* ULevel::GetExternalActorsFolderName()
+{
+	return TEXT("__ExternalActors__");
+}
+
 TArray<UPackage*> ULevel::GetLoadedExternalActorPackages() const
 {
 	TSet<UPackage*> ActorPackages;
@@ -2390,7 +2395,7 @@ FString ULevel::GetExternalActorsPath(const FString& InLevelPackageName, const F
 		FString MountPoint, PackagePath, ShortName;
 		if (FPackageName::SplitLongPackageName(InLevelPackageName, MountPoint, PackagePath, ShortName))
 		{
-			OutExternalActorsPath = FString::Printf(TEXT("%s__ExternalActors__/%s%s"), *MountPoint, *PackagePath, InPackageShortName.IsEmpty() ? *ShortName : *InPackageShortName);
+			OutExternalActorsPath = FString::Printf(TEXT("%s%s/%s%s"), *MountPoint, GetExternalActorsFolderName(), *PackagePath, InPackageShortName.IsEmpty() ? *ShortName : *InPackageShortName);
 			return true;
 		}
 		return false;
