Remove old UE references in comments

#rnx

diff --git a/Engine/Source/Runtime/Engine/Classes/Engine/LevelScriptBlueprint.h b/Engine/Source/Runtime/Engine/Classes/Engine/LevelScriptBlueprint.h
--- a/Engine/Source/Runtime/Engine/Classes/Engine/LevelScriptBlueprint.h
+++ b/Engine/Source/Runtime/Engine/Classes/Engine/LevelScriptBlueprint.h
@@ -12,7 +12,7 @@
 /**
  * A level blueprint is a specialized type of blueprint. It is used to house
  * global, level-wide logic. In a level blueprint, you can operate on specific 
- * level-actor instances through blueprint's node-based interface. UE3 users 
+ * level-actor instances through blueprint's node-based interface. Unreal Engine users 
  * should be familiar with this concept, as it is very similar to Kismet.
  *
  * @see https://docs.unrealengine.com/latest/INT/Engine/Blueprints/UserGuide/Types/LevelBlueprint/index.html
