Fix HISM race condition

diff --git a/Engine/Source/Runtime/Engine/Classes/Engine/InstancedStaticMesh.h b/Engine/Source/Runtime/Engine/Classes/Engine/InstancedStaticMesh.h
--- a/Engine/Source/Runtime/Engine/Classes/Engine/InstancedStaticMesh.h
+++ b/Engine/Source/Runtime/Engine/Classes/Engine/InstancedStaticMesh.h
@@ -477,6 +477,7 @@ public:
 #if WITH_EDITOR
 	,	bHasSelectedInstances(false)
 #endif
+	,	StaticMeshBounds(StaticMesh->GetBounds())
 	{
 #if WITH_EDITOR
 		for (int32 InstanceIndex = 0; InstanceIndex < InComponent->SelectedInstances.Num() && !bHasSelectedInstances; ++InstanceIndex)
@@ -496,6 +497,7 @@ public:
 	}
 
 	// FPrimitiveSceneProxy interface.
+	virtual void CreateRenderThreadResources() override;
 
 	virtual void DestroyRenderThreadResources() override;
 
@@ -612,6 +614,8 @@ protected:
 	/** Common path for the Get*MeshElement functions */
 	void SetupInstancedMeshBatch(int32 LODIndex, int32 BatchIndex, FMeshBatch& OutMeshBatch) const;
 
+	/** Untransformed bounds of the static mesh */
+	FBoxSphereBounds StaticMeshBounds;
 private:
 
 	void SetupProxy(UInstancedStaticMeshComponent* InComponent);
diff --git a/Engine/Source/Runtime/Engine/Private/HierarchicalInstancedStaticMesh.cpp b/Engine/Source/Runtime/Engine/Private/HierarchicalInstancedStaticMesh.cpp
--- a/Engine/Source/Runtime/Engine/Private/HierarchicalInstancedStaticMesh.cpp
+++ b/Engine/Source/Runtime/Engine/Private/HierarchicalInstancedStaticMesh.cpp
@@ -894,42 +894,6 @@ public:
 		, CaptureTag(0)
 #endif
 	{
-#if GPUCULL_TODO
-		if (UseGPUScene(GetScene().GetShaderPlatform(), GetScene().GetFeatureLevel()))
-		{
-			if (InstancedRenderData.PerInstanceRenderData.IsValid() && InstancedRenderData.PerInstanceRenderData->InstanceBuffer_GameThread.IsValid())
-			{
-				const FStaticMeshInstanceData& StaticMeshInstanceData = *InstancedRenderData.PerInstanceRenderData->InstanceBuffer_GameThread;
-
-					Instances.SetNum(StaticMeshInstanceData.GetNumInstances());
-					if (Instances.Num() == 0 && ClusterTree.Num() != 0)
-					{
-						UE_LOG(LogTemp, Display, TEXT("Instances.Num() == 0 && ClusterTree.Num() != 0"));
-					}
-
-				for (int32 InstanceIndex = 0; InstanceIndex < Instances.Num(); ++InstanceIndex)
-				{
-					FMatrix InstanceTransform;
-					StaticMeshInstanceData.GetInstanceTransform(InstanceIndex, InstanceTransform);
-					InstanceTransform.M[3][3] = 1.0f;
-
-					FPrimitiveInstance& Instance = Instances[InstanceIndex];
-					Instance.PrimitiveId = ~uint32(0);
-					Instance.InstanceToLocal = InstanceTransform;
-					Instance.LocalToInstance = InstanceTransform.Inverse();
-					// Filled in during GPU-Scene update...
-					Instance.LocalToWorld.SetIdentity();
-					Instance.RenderBounds = InComponent->GetStaticMesh()->GetBounds();
-					Instance.LocalBounds = Instance.RenderBounds.TransformBy(Instance.InstanceToLocal);
-				}
-			}
-			else
-			{
-				ensure(ClusterTree.Num() == 0);
-			}
-		}
-#endif // GPUCULL_TODO
-
 		SetupOcclusion(InComponent);
 	}
 
diff --git a/Engine/Source/Runtime/Engine/Private/InstancedStaticMesh.cpp b/Engine/Source/Runtime/Engine/Private/InstancedStaticMesh.cpp
--- a/Engine/Source/Runtime/Engine/Private/InstancedStaticMesh.cpp
+++ b/Engine/Source/Runtime/Engine/Private/InstancedStaticMesh.cpp
@@ -1207,6 +1207,71 @@ void FInstancedStaticMeshSceneProxy::SetupProxy(UInstancedStaticMeshComponent* I
 #endif
 }
 
+
+void FInstancedStaticMeshSceneProxy::CreateRenderThreadResources()
+{
+	FStaticMeshSceneProxy::CreateRenderThreadResources();
+
+#if GPUCULL_TODO
+	if (UseGPUScene(GetScene().GetShaderPlatform(), GetScene().GetFeatureLevel()))
+	{
+		bSupportsInstanceDataBuffer = true;
+		// TODO: can the PerInstanceRenderData ever not be valid here?
+		if (ensure(InstancedRenderData.PerInstanceRenderData.IsValid()))
+		{
+			const FStaticMeshInstanceBuffer& InstanceBuffer = InstancedRenderData.PerInstanceRenderData->InstanceBuffer;
+			ensureMsgf(InstanceBuffer.RequireCPUAccess, TEXT("GPU-Scene instance culling requires CPU access to instance data for setup."));
+
+			// This happens when this is actually a HISM and the data is not present in the component (which is true for landscape grass
+			// which manages its own setup.
+			if (Instances.Num() == 0)
+			{
+				Instances.SetNum(InstanceBuffer.GetNumInstances());
+			}
+
+			// NOTE: we set up partial data in the construction of ISM proxy (yep, awful but the equally awful way the InstanceBuffer is maintained means complete data is not available)
+			if (Instances.Num() == InstanceBuffer.GetNumInstances())
+			{
+				FVector4 InstanceTransformVec[3];
+				FVector4 InstanceLightMapAndShadowMapUVBias = FVector4(ForceInitToZero);
+				FVector4 InstanceOrigin = FVector::ZeroVector;
+
+				for (int32 InstanceIndex = 0; InstanceIndex < Instances.Num(); ++InstanceIndex)
+				{
+					InstanceBuffer.GetInstanceShaderValues(
+						InstanceIndex,
+						InstanceTransformVec,
+						InstanceLightMapAndShadowMapUVBias,
+						InstanceOrigin
+					);
+
+					FMatrix LocalToPrimitiveTransform;
+					LocalToPrimitiveTransform.M[0][0] = InstanceTransformVec[0][0];
+					LocalToPrimitiveTransform.M[0][1] = InstanceTransformVec[0][1];
+					LocalToPrimitiveTransform.M[0][2] = InstanceTransformVec[0][2];
+					LocalToPrimitiveTransform.M[0][3] = 0.f;
+					LocalToPrimitiveTransform.M[1][0] = InstanceTransformVec[1][0];
+					LocalToPrimitiveTransform.M[1][1] = InstanceTransformVec[1][1];
+					LocalToPrimitiveTransform.M[1][2] = InstanceTransformVec[1][2];
+					LocalToPrimitiveTransform.M[1][3] = 0.f;
+					LocalToPrimitiveTransform.M[2][0] = InstanceTransformVec[2][0];
+					LocalToPrimitiveTransform.M[2][1] = InstanceTransformVec[2][1];
+					LocalToPrimitiveTransform.M[2][2] = InstanceTransformVec[2][2];
+					LocalToPrimitiveTransform.M[2][3] = 0.f;
+					LocalToPrimitiveTransform.M[3][0] = InstanceOrigin.X;
+					LocalToPrimitiveTransform.M[3][1] = InstanceOrigin.Y;
+					LocalToPrimitiveTransform.M[3][2] = InstanceOrigin.Z;
+					LocalToPrimitiveTransform.M[3][3] = 1.f;
+
+					SetupPrimitiveInstance(Instances[InstanceIndex], LocalToPrimitiveTransform, StaticMeshBounds, InstanceLightMapAndShadowMapUVBias, /*PerInstanceRandom */InstanceOrigin.W);
+				}
+			}
+		}
+	}
+#endif
+}
+
+
 void FInstancedStaticMeshSceneProxy::DestroyRenderThreadResources()
 {
 	InstancedRenderData.ReleaseResources(&GetScene(), StaticMesh);
diff --git a/Engine/Source/Runtime/Engine/Public/InstanceUniformShaderParameters.h b/Engine/Source/Runtime/Engine/Public/InstanceUniformShaderParameters.h
--- a/Engine/Source/Runtime/Engine/Public/InstanceUniformShaderParameters.h
+++ b/Engine/Source/Runtime/Engine/Public/InstanceUniformShaderParameters.h
@@ -47,6 +47,23 @@ struct FPrimitiveInstance
 	uint32 Flags;
 };
 
+/**
+ * Use to initialize the required fields in an FPrimitiveInstance, much of the data is derived at upload time and need not be set (should be refactored).
+ */
+FORCEINLINE void SetupPrimitiveInstance(FPrimitiveInstance& PrimitiveInstance,
+	const FMatrix& LocalToPrimitive /* Transform from the mesh-local space to the primitive/component local space, AKA the instance transform. */,
+	const FBoxSphereBounds& LocalInstanceBounds /* Local bounds of the instanced mesh. */,
+	const FVector4 &LightMapAndShadowMapUVBias,
+	float PerInstanceRandom,
+	const FNaniteInfo &NaniteInfo = FNaniteInfo())
+{
+	PrimitiveInstance.InstanceToLocal = LocalToPrimitive;
+	PrimitiveInstance.RenderBounds = LocalInstanceBounds;
+	PrimitiveInstance.LightMapAndShadowMapUVBias = LightMapAndShadowMapUVBias;
+	PrimitiveInstance.PerInstanceRandom = PerInstanceRandom;
+	PrimitiveInstance.NaniteInfo = NaniteInfo;
+}
+
 /** 
  * The uniform shader parameters associated with a primitive instance.
  * Note: Must match FInstanceSceneData in shaders.
