Fix NonUnity UnrealEditor Win64

diff --git a/Engine/Source/Runtime/Engine/Private/WorldPartition/WorldPartitionRuntimeHash.cpp b/Engine/Source/Runtime/Engine/Private/WorldPartition/WorldPartitionRuntimeHash.cpp
--- a/Engine/Source/Runtime/Engine/Private/WorldPartition/WorldPartitionRuntimeHash.cpp
+++ b/Engine/Source/Runtime/Engine/Private/WorldPartition/WorldPartitionRuntimeHash.cpp
@@ -10,6 +10,7 @@
 #if WITH_EDITOR
 #include "WorldPartition/WorldPartitionHandle.h"
 #include "WorldPartition/WorldPartitionActorDescViewProxy.h"
+#include "Logging/MessageLog.h"
 #include "Misc/UObjectToken.h"
 #include "Misc/MapErrors.h"
 #endif
@@ -135,7 +136,7 @@ void UWorldPartitionRuntimeHash::ChangeActorDescViewGridPlacement(FWorldPartitio
 	if (ActorDescView.EffectiveGridPlacement != GridPlacement)
 	{
 		ActorDescView.EffectiveGridPlacement = GridPlacement;
-		UE_LOG(LogWorldPartition, Verbose, TEXT("Actor '%s' grid placement changed to %d"), *ActorDescView.GetActorLabel().ToString(), (int32)GridPlacement);
+		UE_LOG(LogWorldPartition, Warning, TEXT("Actor '%s' grid placement changed to %d"), *ActorDescView.GetActorLabel().ToString(), (int32)GridPlacement);
 	}
 }
 
