Implements SetGlobalShaderCacheOverrideDirectory

diff --git a/Engine/Source/Developer/TargetPlatform/Public/Interfaces/ITargetDevice.h b/Engine/Source/Developer/TargetPlatform/Public/Interfaces/ITargetDevice.h
--- a/Engine/Source/Developer/TargetPlatform/Public/Interfaces/ITargetDevice.h
+++ b/Engine/Source/Developer/TargetPlatform/Public/Interfaces/ITargetDevice.h
@@ -387,6 +387,16 @@ public:
 	 * @param ProcessIdentifier The bundle id
 	 */
 	virtual bool TerminateLaunchedProcess(const FString & ProcessIdentifier) { return false;  };
+	
+	/**
+	 * Execute console command on the device to reload the global shader map
+	 * 
+	 * @param GlobalShaderMapDirectory path to directory that contains global shader map
+	 */
+	virtual void ReloadGlobalShadersMap(const FString& GlobalShaderMapDirectory) const
+	{
+		return ExecuteConsoleCommand(TEXT("ReloadGlobalShaders"));
+	}
 
 public:
 
diff --git a/Engine/Source/Editor/UnrealEd/Private/Commandlets/CookGlobalShaderCommandlet.cpp b/Engine/Source/Editor/UnrealEd/Private/Commandlets/CookGlobalShaderCommandlet.cpp
--- a/Engine/Source/Editor/UnrealEd/Private/Commandlets/CookGlobalShaderCommandlet.cpp
+++ b/Engine/Source/Editor/UnrealEd/Private/Commandlets/CookGlobalShaderCommandlet.cpp
@@ -204,7 +204,7 @@ int32 UCookGlobalShadersCommandlet::Main(const FString& Params)
 		if (bCopySucceeded && bExecuteReload && TargetDevice.IsValid())
 		{
 			UE_LOG(LogCookGlobalShaders, Log, TEXT("Sending Reload Command..."));
-			TargetDevice->ExecuteConsoleCommand("ReloadGlobalShaders");
+			TargetDevice->ReloadGlobalShadersMap(OutputDir / TEXT("Engine"));
 		}
 	}
 	UE_LOG(LogCookGlobalShaders, Log, TEXT("Complete"));
diff --git a/Engine/Source/Runtime/Engine/Private/ShaderCompiler/ShaderCompiler.cpp b/Engine/Source/Runtime/Engine/Private/ShaderCompiler/ShaderCompiler.cpp
--- a/Engine/Source/Runtime/Engine/Private/ShaderCompiler/ShaderCompiler.cpp
+++ b/Engine/Source/Runtime/Engine/Private/ShaderCompiler/ShaderCompiler.cpp
@@ -5634,9 +5634,18 @@ static FPreLoadFile GGlobalShaderPreLoadFile(*(FString(TEXT("../../../Engine"))
 
 static const ITargetPlatform* GGlobalShaderTargetPlatform[SP_NumPlatforms] = { nullptr };
 
+static FString GGlobalShaderCacheOverrideDirectory;
+
 static FString GetGlobalShaderCacheOverrideFilename(EShaderPlatform Platform)
 {
-	return FString(TEXT("Engine")) / TEXT("OverrideGlobalShaderCache-") + LegacyShaderPlatformToShaderFormat(Platform).ToString() + TEXT(".bin");
+	FString DirectoryPrefix = FString(TEXT("Engine")) / TEXT("OverrideGlobalShaderCache-");
+
+	if (!GGlobalShaderCacheOverrideDirectory.IsEmpty())
+	{
+		DirectoryPrefix = GGlobalShaderCacheOverrideDirectory / TEXT("GlobalShaderCache-");
+	}
+
+	return DirectoryPrefix + LegacyShaderPlatformToShaderFormat(Platform).ToString() + TEXT(".bin");
 }
 
 static FString GetGlobalShaderCacheFilename(EShaderPlatform Platform)
@@ -5846,7 +5855,26 @@ void CompileGlobalShaderMap(EShaderPlatform Platform, const ITargetPlatform* Tar
 			{
 				FString OverrideGlobalShaderCacheFilename = GetGlobalShaderCacheOverrideFilename(Platform);
 				FPaths::MakeStandardFilename(OverrideGlobalShaderCacheFilename);
-				bLoadedFromCacheFile = FFileHelper::LoadFileToArray(GlobalShaderData, *OverrideGlobalShaderCacheFilename, FILEREAD_Silent);
+
+				bool bFileExist = IFileManager::Get().FileExists(*OverrideGlobalShaderCacheFilename);
+
+				if (!bFileExist)
+				{
+					UE_LOG(LogShaders, Log, TEXT("%s doesn't exists"), *OverrideGlobalShaderCacheFilename);
+				}
+				else
+				{
+					bLoadedFromCacheFile = FFileHelper::LoadFileToArray(GlobalShaderData, *OverrideGlobalShaderCacheFilename, FILEREAD_Silent);
+
+					if (bLoadedFromCacheFile)
+					{
+						UE_LOG(LogShaders, Log, TEXT("%s has been loaded successfully"), *OverrideGlobalShaderCacheFilename);
+					}
+					else
+					{
+						UE_LOG(LogShaders, Error, TEXT("%s failed to load"), *OverrideGlobalShaderCacheFilename);
+					}
+				}
 			}
 
 			// is the data already loaded?
@@ -6004,6 +6032,23 @@ static FAutoConsoleCommand CCmdReloadGlobalShaders = FAutoConsoleCommand(
 	FConsoleCommandDelegate::CreateStatic(ReloadGlobalShaders)
 );
 
+void SetGlobalShaderCacheOverrideDirectory(const TArray<FString>& Args)
+{
+	if (Args.Num() == 0)
+	{
+		UE_LOG(LogShaders, Error, TEXT("Failed to set GGlobalShaderCacheOverrideDirectory without any arguments"));
+		return; 
+	}
+	
+	GGlobalShaderCacheOverrideDirectory = Args[0];
+	UE_LOG(LogShaders, Log, TEXT("GGlobalShaderCacheOverrideDirectory = %s"), *GGlobalShaderCacheOverrideDirectory);
+}
+
+static FAutoConsoleCommand CCmdSetGlobalShaderCacheOverrideDirectory = FAutoConsoleCommand(
+	TEXT("SetGlobalShaderCacheOverrideDirectory"),
+	TEXT("Set the directory to read the override global shader map file from."),
+	FConsoleCommandWithArgsDelegate::CreateStatic(SetGlobalShaderCacheOverrideDirectory));
+
 bool RecompileChangedShadersForPlatform(const FString& PlatformName)
 {
 	// figure out what shader platforms to recompile
