Only include newly created/unsaved modified actors in to the actor descriptors list in PIE.
For some reason, some actors were dirtied when loading the map during HLOD generation and it created a grid mismatch between HLOD and runtime.

diff --git a/Engine/Source/Runtime/Engine/Private/WorldPartition/WorldPartition.cpp b/Engine/Source/Runtime/Engine/Private/WorldPartition/WorldPartition.cpp
--- a/Engine/Source/Runtime/Engine/Private/WorldPartition/WorldPartition.cpp
+++ b/Engine/Source/Runtime/Engine/Private/WorldPartition/WorldPartition.cpp
@@ -32,9 +32,7 @@
 #include "WorldPartition/WorldPartitionLevelHelper.h"
 #include "WorldPartition/WorldPartitionLevelStreamingDynamic.h"
 #include "WorldPartition/WorldPartitionEditorHash.h"
-#include "WorldPartition/WorldPartitionEditorSpatialHash.h"
 #include "WorldPartition/WorldPartitionRuntimeHash.h"
-#include "WorldPartition/WorldPartitionRuntimeSpatialHash.h"
 #include "WorldPartition/WorldPartitionEditorPerProjectUserSettings.h"
 #include "Modules/ModuleManager.h"
 #include "GameDelegates.h"
@@ -135,6 +133,7 @@ UWorldPartition::UWorldPartition(const FObjectInitializer& ObjectInitializer)
 	, WorldPartitionEditor(nullptr)
 	, bForceGarbageCollection(false)
 	, bForceGarbageCollectionPurge(false)
+	, bIsPIE(false)
 #endif
 	, InitState(EWorldPartitionInitState::Uninitialized)
 	, InstanceTransform(FTransform::Identity)
@@ -151,6 +150,9 @@ UWorldPartition::UWorldPartition(const FObjectInitializer& ObjectInitializer)
 #if WITH_EDITOR
 void UWorldPartition::OnPreBeginPIE(bool bStartSimulate)
 {
+	check(!bIsPIE);
+	bIsPIE = true;
+
 	// Gather all modified/newly created actors for PIE so they can be duplicated when streaming their cell
 	for (AActor* Actor : World->PersistentLevel->Actors)
 	{
@@ -165,6 +167,12 @@ void UWorldPartition::OnPreBeginPIE(bool bStartSimulate)
 	OnBeginPlay(EWorldPartitionStreamingMode::PIE);
 }
 
+void UWorldPartition::OnPrePIEEnded(bool bWasSimulatingInEditor)
+{
+	check(bIsPIE);
+	bIsPIE = false;
+}
+
 void UWorldPartition::OnBeginPlay(EWorldPartitionStreamingMode Mode)
 {
 	GenerateStreaming(Mode);
@@ -223,9 +231,9 @@ void UWorldPartition::Initialize(UWorld* InWorld, const FTransform& InTransform)
 	{
 		if (!EditorHash)
 		{
-			UWorldPartitionEditorSpatialHash* EditorSpatialHash = NewObject<UWorldPartitionEditorSpatialHash>(this, NAME_None, RF_Transactional);
-			EditorSpatialHash->SetDefaultValues();
-			EditorHash = EditorSpatialHash;
+			UClass* EditorHashClass = FindObject<UClass>(ANY_PACKAGE, TEXT("WorldPartitionEditorSpatialHash"));
+			EditorHash = NewObject<UWorldPartitionEditorHash>(this, EditorHashClass);
+			EditorHash->SetDefaultValues();
 		}
 
 		EditorHash->Initialize();
@@ -233,9 +241,9 @@ void UWorldPartition::Initialize(UWorld* InWorld, const FTransform& InTransform)
 
 	if (!RuntimeHash)
 	{
-		UWorldPartitionRuntimeSpatialHash* RuntimeSpatialHash = NewObject<UWorldPartitionRuntimeSpatialHash>(this);
-		RuntimeSpatialHash->SetDefaultValues();
-		RuntimeHash = RuntimeSpatialHash;
+		UClass* RuntimeHashClass = FindObject<UClass>(ANY_PACKAGE, TEXT("WorldPartitionRuntimeSpatialHash"));
+		RuntimeHash = NewObject<UWorldPartitionRuntimeHash>(this, RuntimeHashClass);
+		RuntimeHash->SetDefaultValues();
 	}
 
 	if (bEditorOnly || IsRunningGame())
@@ -497,6 +505,7 @@ void UWorldPartition::RegisterDelegates()
 	if (GEditor && !IsTemplate())
 	{
 		FEditorDelegates::PreBeginPIE.AddUObject(this, &UWorldPartition::OnPreBeginPIE);
+		FEditorDelegates::PrePIEEnded.AddUObject(this, &UWorldPartition::OnPrePIEEnded);
 		FGameDelegates::Get().GetEndPlayMapDelegate().AddUObject(this, &UWorldPartition::OnEndPlay);
 	}
 }
@@ -507,6 +516,7 @@ void UWorldPartition::UnregisterDelegates()
 	if (GEditor && !IsTemplate())
 	{
 		FEditorDelegates::PreBeginPIE.RemoveAll(this);
+		FEditorDelegates::PrePIEEnded.RemoveAll(this);
 		FGameDelegates::Get().GetEndPlayMapDelegate().RemoveAll(this);
 	}
 }
diff --git a/Engine/Source/Runtime/Engine/Private/WorldPartition/WorldPartitionRuntimeHash.cpp b/Engine/Source/Runtime/Engine/Private/WorldPartition/WorldPartitionRuntimeHash.cpp
--- a/Engine/Source/Runtime/Engine/Private/WorldPartition/WorldPartitionRuntimeHash.cpp
+++ b/Engine/Source/Runtime/Engine/Private/WorldPartition/WorldPartitionRuntimeHash.cpp
@@ -80,22 +80,28 @@ void UWorldPartitionRuntimeHash::CreateActorDescViewMap(const UActorDescContaine
 {
 	// Build the actor desc view map
 	OutActorDescViewMap.Empty();
+
+	const bool bIsPIE = GetOuterUWorldPartition()->bIsPIE;
+
 	for (UActorDescContainer::TConstIterator<> ActorDescIt(Container); ActorDescIt; ++ActorDescIt)
 	{
-		if (AActor* Actor = ActorDescIt->GetActor())
+		if (bIsPIE)
 		{
-			// Don't include deleted, unsaved actors
-			if (Actor->IsPendingKill())
+			if (AActor* Actor = ActorDescIt->GetActor())
 			{
-				continue;
-			}
+				// Don't include deleted, unsaved actors
+				if (Actor->IsPendingKill())
+				{
+					continue;
+				}
 
-			// Handle dirty, unsaved actors
-			if (Actor->GetPackage()->IsDirty())
-			{
-				FWorldPartitionActorDesc* ActorDesc = ModifiedActorDescListForPIE.AddActor(Actor);
-				OutActorDescViewMap.Emplace(ActorDesc->GetGuid(), ActorDesc);
-				continue;
+				// Handle dirty, unsaved actors
+				if (Actor->GetPackage()->IsDirty())
+				{
+					FWorldPartitionActorDesc* ActorDesc = ModifiedActorDescListForPIE.AddActor(Actor);
+					OutActorDescViewMap.Emplace(ActorDesc->GetGuid(), ActorDesc);
+					continue;
+				}
 			}
 		}
 
@@ -103,15 +109,18 @@ void UWorldPartitionRuntimeHash::CreateActorDescViewMap(const UActorDescContaine
 		OutActorDescViewMap.Emplace(ActorDescIt->GetGuid(), *ActorDescIt);
 	}
 
-	// Append new unsaved actors for the persistent level
-	if (Container->GetContainerPackage() == GetWorld()->PersistentLevel->GetPackage()->GetLoadedPath().GetPackageFName())
+	if (bIsPIE)
 	{
-		for (auto& ActorPair : GetWorld()->PersistentLevel->ActorsModifiedForPIE)
+		// Append new unsaved actors for the persistent level
+		if (Container->GetContainerPackage() == GetWorld()->PersistentLevel->GetPackage()->GetLoadedPath().GetPackageFName())
 		{
-			if (ActorPair.Value->GetPackage()->HasAnyPackageFlags(PKG_NewlyCreated))
+			for (auto& ActorPair : GetWorld()->PersistentLevel->ActorsModifiedForPIE)
 			{
-				FWorldPartitionActorDesc* ActorDesc = ModifiedActorDescListForPIE.AddActor(ActorPair.Value);
-				OutActorDescViewMap.Emplace(ActorDesc->GetGuid(), ActorDesc);
+				if (ActorPair.Value->GetPackage()->HasAnyPackageFlags(PKG_NewlyCreated))
+				{
+					FWorldPartitionActorDesc* ActorDesc = ModifiedActorDescListForPIE.AddActor(ActorPair.Value);
+					OutActorDescViewMap.Emplace(ActorDesc->GetGuid(), ActorDesc);
+				}
 			}
 		}
 	}
diff --git a/Engine/Source/Runtime/Engine/Public/WorldPartition/WorldPartition.h b/Engine/Source/Runtime/Engine/Public/WorldPartition/WorldPartition.h
--- a/Engine/Source/Runtime/Engine/Public/WorldPartition/WorldPartition.h
+++ b/Engine/Source/Runtime/Engine/Public/WorldPartition/WorldPartition.h
@@ -95,6 +95,7 @@ private:
 
 	// PIE/Game Methods
 	void OnPreBeginPIE(bool bStartSimulate);
+	void OnPrePIEEnded(bool bWasSimulatingInEditor);
 	void OnBeginPlay(EWorldPartitionStreamingMode Mode);
 	void OnEndPlay();
 
@@ -196,6 +197,7 @@ public:
 #if WITH_EDITOR
 	bool bForceGarbageCollection;
 	bool bForceGarbageCollectionPurge;
+	bool bIsPIE;
 #endif
 
 #if WITH_EDITORONLY_DATA
