Editgrate crash fix with getting current first LOD index when no LOD's are present

diff --git a/Engine/Source/Runtime/Engine/Private/StaticMesh.cpp b/Engine/Source/Runtime/Engine/Private/StaticMesh.cpp
--- a/Engine/Source/Runtime/Engine/Private/StaticMesh.cpp
+++ b/Engine/Source/Runtime/Engine/Private/StaticMesh.cpp
@@ -1767,22 +1767,30 @@ void FStaticMeshRenderData::AllocateLODResources(int32 NumLODs)
 int32 FStaticMeshRenderData::GetFirstValidLODIdx(int32 MinIdx) const
 {
 	const int32 LODCount = LODResources.Num();
-	int32 LODIndex = FMath::Clamp<int32>(MinIdx, 0, LODCount - 1);
-
-	if (FPlatformProperties::RequiresCookedData())
+	int32 LODIndex = INDEX_NONE;
+	if (LODCount > 0)
 	{
-		// When cooked, BuffersSize can be used reliably to know if there is valid data, even though num vertices could be non zero.
-		// This happens because UStaticMesh::MinLOD is platform specific, where as num vertices is the same for all platforms.
-		while (LODIndex < LODCount && !LODResources[LODIndex].GetNumVertices() && LODResources[LODIndex].BuffersSize)
+		LODIndex = FMath::Clamp<int32>(MinIdx, 0, LODCount - 1);
+		if (FPlatformProperties::RequiresCookedData())
 		{
-			++LODIndex;
+			// When cooked, BuffersSize can be used reliably to know if there is valid data, even though num vertices could be non zero.
+			// This happens because UStaticMesh::MinLOD is platform specific, where as num vertices is the same for all platforms.
+			while (LODIndex < LODCount && !LODResources[LODIndex].GetNumVertices() && LODResources[LODIndex].BuffersSize)
+			{
+				++LODIndex;
+			}
 		}
-	}
-	else
-	{
-		while (LODIndex < LODCount && !LODResources[LODIndex].GetNumVertices())
+		else
+		{
+			while (LODIndex < LODCount && !LODResources[LODIndex].GetNumVertices())
+			{
+				++LODIndex;
+			}
+		}
+
+		if (LODIndex >= LODCount)
 		{
-			++LODIndex;
+			LODIndex = INDEX_NONE;
 		}
 	}
 
diff --git a/Engine/Source/Runtime/Engine/Public/StaticMeshResources.h b/Engine/Source/Runtime/Engine/Public/StaticMeshResources.h
--- a/Engine/Source/Runtime/Engine/Public/StaticMeshResources.h
+++ b/Engine/Source/Runtime/Engine/Public/StaticMeshResources.h
@@ -724,7 +724,8 @@ public:
 	 */
 	FORCEINLINE const FStaticMeshLODResources* GetCurrentFirstLOD(int32 MinLODIdx) const
 	{
-		return &LODResources[GetCurrentFirstLODIdx(MinLODIdx)];
+		const int32 LODIdx = GetCurrentFirstLODIdx(MinLODIdx);
+		return LODResources.IsValidIndex(LODIdx) ? &LODResources[LODIdx] : nullptr;
 	}
 
 private:
