Added the "wp.Runtime.HLOD" command line option to enable / disable loading & rendering of HLODs

diff --git a/Engine/Source/Runtime/Engine/Private/WorldPartition/HLOD/HLODActor.cpp b/Engine/Source/Runtime/Engine/Private/WorldPartition/HLOD/HLODActor.cpp
--- a/Engine/Source/Runtime/Engine/Private/WorldPartition/HLOD/HLODActor.cpp
+++ b/Engine/Source/Runtime/Engine/Private/WorldPartition/HLOD/HLODActor.cpp
@@ -171,3 +171,10 @@ void AWorldPartitionHLOD::BuildHLOD(bool bForceBuild)
 }
 
 #endif // WITH_EDITOR
+
+
+UWorldPartitionRuntimeHLODCellData::UWorldPartitionRuntimeHLODCellData(const FObjectInitializer& ObjectInitializer)
+	: Super(ObjectInitializer)
+{
+}
+
diff --git a/Engine/Source/Runtime/Engine/Private/WorldPartition/WorldPartitionRuntimeSpatialHash.cpp b/Engine/Source/Runtime/Engine/Private/WorldPartition/WorldPartitionRuntimeSpatialHash.cpp
--- a/Engine/Source/Runtime/Engine/Private/WorldPartition/WorldPartitionRuntimeSpatialHash.cpp
+++ b/Engine/Source/Runtime/Engine/Private/WorldPartition/WorldPartitionRuntimeSpatialHash.cpp
@@ -812,6 +812,9 @@ bool UWorldPartitionRuntimeSpatialHash::CreateStreamingGrid(const FSpatialHashRu
 
 				UE_LOG(LogWorldPartition, Verbose, TEXT("Cell%s %s Actors = %d Bounds (%s)"), bIsCellAlwaysLoaded ? TEXT(" (AlwaysLoaded)") : TEXT(""), *StreamingCell->GetName(), FilteredActors.Num(), *Bounds.ToString());
 
+				// Keep track of all AWorldPartitionHLOD actors referenced by this cell
+				int32 NumHLODActors = 0;
+
 				for (const FActorInstance& ActorInstance : FilteredActors)
 				{
 					const FWorldPartitionActorDescView& ActorDescView = ActorInstance.GetActorDescView();
@@ -826,9 +829,24 @@ bool UWorldPartitionRuntimeSpatialHash::CreateStreamingGrid(const FSpatialHashRu
 						}
 					}
 
+					if (ActorDescView.GetActorClass()->IsChildOf<AWorldPartitionHLOD>())
+					{
+						NumHLODActors++;
+					}
+
 					UE_LOG(LogWorldPartition, Verbose, TEXT("  Actor : %s (%s) (Container %08x) Origin(%s)"), *(ActorDescView.GetActorPath().ToString()), *ActorDescView.GetGuid().ToString(EGuidFormats::UniqueObjectGuid), ActorInstance.ContainerInstance->ID, *FVector2D(ActorInstance.GetOrigin()).ToString());
 				}
 
+				if (NumHLODActors > 0)
+				{
+					// HLOD cells should contain only HLODs
+					check(NumHLODActors == FilteredActors.Num());
+
+					// Mark this cell as containing HLOD data
+					UWorldPartitionRuntimeHLODCellData* HLODCellData = NewObject<UWorldPartitionRuntimeHLODCellData>(StreamingCell);
+					StreamingCell->AddCellData(HLODCellData);
+				}
+
 				if (Mode == EWorldPartitionStreamingMode::Cook)
 				{
 					UE_LOG(LogWorldPartition, Log, TEXT("Creating runtime streaming cells %s."), *StreamingCell->GetName());
diff --git a/Engine/Source/Runtime/Engine/Private/WorldPartition/WorldPartitionStreamingPolicy.cpp b/Engine/Source/Runtime/Engine/Private/WorldPartition/WorldPartitionStreamingPolicy.cpp
--- a/Engine/Source/Runtime/Engine/Private/WorldPartition/WorldPartitionStreamingPolicy.cpp
+++ b/Engine/Source/Runtime/Engine/Private/WorldPartition/WorldPartitionStreamingPolicy.cpp
@@ -10,6 +10,7 @@
 #include "WorldPartition/WorldPartitionStreamingSource.h"
 #include "WorldPartition/WorldPartition.h"
 #include "WorldPartition/DataLayer/DataLayerSubsystem.h"
+#include "WorldPartition/HLOD/HLODActor.h"
 #include "GameFramework/PlayerController.h"
 #include "Engine/LocalPlayer.h"
 #include "Engine/World.h"
@@ -18,12 +19,21 @@
 #include "DrawDebugHelpers.h"
 #include "DisplayDebugHelpers.h"
 #include "RenderUtils.h"
+#include "Algo/RemoveIf.h"
 
 #if WITH_EDITOR
 #include "Editor.h"
 #include "LevelEditorViewport.h"
 #endif
 
+int32 GWorldPartitionShowHLODs = 1;
+static FAutoConsoleVariableRef CVarWorldPartitionShowHLODs(
+	TEXT("wp.Runtime.HLOD"),
+	GWorldPartitionShowHLODs,
+	TEXT("Turn on/off loading & rendering of world partition HLODs."),
+	ECVF_Scalability
+);
+
 UWorldPartitionStreamingPolicy::UWorldPartitionStreamingPolicy(const FObjectInitializer& ObjectInitializer)
 : Super(ObjectInitializer) 
 , bIsServerLoadingDone(false)
@@ -147,6 +157,23 @@ void UWorldPartitionStreamingPolicy::UpdateStreamingState()
 			LoadStreamingCells = LoadStreamingCells.Difference(ActivateStreamingCells);
 		}
 
+		if (!GWorldPartitionShowHLODs)
+		{
+			// Remove all HLOD cells from the Activate & Load cells
+			auto RemoveHLODCells = [](TSet<const UWorldPartitionRuntimeCell*>& Cells)
+			{
+				for (auto It = Cells.CreateIterator(); It; ++It)
+				{
+					if ((*It)->HasCellData<UWorldPartitionRuntimeHLODCellData>())
+					{
+						It.RemoveCurrent();
+					}
+				}
+			};
+			RemoveHLODCells(ActivateStreamingCells);
+			RemoveHLODCells(LoadStreamingCells);
+		}
+
 		// Determine cells to load/unload
 		TSet<const UWorldPartitionRuntimeCell*> ToActivateCells = ActivateStreamingCells.Difference(ActivatedCells);
 		TSet<const UWorldPartitionRuntimeCell*> ToLoadCells = LoadStreamingCells.Difference(LoadedCells);
diff --git a/Engine/Source/Runtime/Engine/Public/WorldPartition/HLOD/HLODActor.h b/Engine/Source/Runtime/Engine/Public/WorldPartition/HLOD/HLODActor.h
--- a/Engine/Source/Runtime/Engine/Public/WorldPartition/HLOD/HLODActor.h
+++ b/Engine/Source/Runtime/Engine/Public/WorldPartition/HLOD/HLODActor.h
@@ -109,3 +109,9 @@ private:
 };
 
 DEFINE_ACTORDESC_TYPE(AWorldPartitionHLOD, FHLODActorDesc);
+
+UCLASS()
+class UWorldPartitionRuntimeHLODCellData : public UWorldPartitionRuntimeCellData
+{
+	GENERATED_UCLASS_BODY()
+};
diff --git a/Engine/Source/Runtime/Engine/Public/WorldPartition/WorldPartitionRuntimeCell.h b/Engine/Source/Runtime/Engine/Public/WorldPartition/WorldPartitionRuntimeCell.h
--- a/Engine/Source/Runtime/Engine/Public/WorldPartition/WorldPartitionRuntimeCell.h
+++ b/Engine/Source/Runtime/Engine/Public/WorldPartition/WorldPartitionRuntimeCell.h
@@ -134,6 +134,7 @@ class UWorldPartitionRuntimeCell : public UObject
 
 	const UWorldPartitionRuntimeCellData* GetCellData(const TSubclassOf<UWorldPartitionRuntimeCellData> InCellDataClass) const;
 	template <class T> inline const T* GetCellData() const { return Cast<const T>(GetCellData(T::StaticClass())); }
+	template <class T> inline bool HasCellData() const { return GetCellData<T>() != nullptr; }
 
 protected:
 	UPROPERTY()
