Fixed issue with time-mapping curves being used on root motion sources, where the fraction of the move would be submitted as a time-based input. This only worked for curves with min/max keys at 0 and 1 second. Now the curves is range-scaled to match the motion duration.

This also fixes an issue in FRootMotionSource_MoveToDynamicForce where the time-mapping curve was not being considered when using bRestrictSpeedToExpected.

Adapted from PR #6611

diff --git a/Engine/Source/Runtime/Engine/Private/GameFramework/RootMotionSource.cpp b/Engine/Source/Runtime/Engine/Private/GameFramework/RootMotionSource.cpp
--- a/Engine/Source/Runtime/Engine/Private/GameFramework/RootMotionSource.cpp
+++ b/Engine/Source/Runtime/Engine/Private/GameFramework/RootMotionSource.cpp
@@ -908,9 +908,15 @@ void FRootMotionSource_MoveToDynamicForce::PrepareRootMotion
 	if (Duration > SMALL_NUMBER && MovementTickTime > SMALL_NUMBER)
 	{
 		float MoveFraction = (GetTime() + SimulationTime) / Duration;
+		
+		// only used if TimeMappingCurve is present
+		float MinMappingCurveTime(0.f);
+		float MaxMappingCurveTime(1.f); 
+
 		if (TimeMappingCurve)
 		{
-			MoveFraction = TimeMappingCurve->GetFloatValue(MoveFraction);
+			TimeMappingCurve->GetTimeRange(MinMappingCurveTime, MaxMappingCurveTime);
+			MoveFraction = TimeMappingCurve->GetFloatValue(FMath::GetRangeValue(FVector2D(MinMappingCurveTime, MaxMappingCurveTime), MoveFraction));
 		}
 
 		FVector CurrentTargetLocation = FMath::Lerp<FVector, float>(StartLocation, TargetLocation, MoveFraction);
@@ -923,7 +929,12 @@ void FRootMotionSource_MoveToDynamicForce::PrepareRootMotion
 		if (bRestrictSpeedToExpected && !Force.IsNearlyZero(KINDA_SMALL_NUMBER))
 		{
 			// Calculate expected current location (if we didn't have collision and moved exactly where our velocity should have taken us)
-			const float PreviousMoveFraction = GetTime() / Duration;
+			float PreviousMoveFraction = GetTime() / Duration;
+			if (TimeMappingCurve)
+			{
+				PreviousMoveFraction = TimeMappingCurve->GetFloatValue(FMath::GetRangeValue(FVector2D(MinMappingCurveTime, MaxMappingCurveTime), PreviousMoveFraction));
+			}
+
 			FVector CurrentExpectedLocation = FMath::Lerp<FVector, float>(StartLocation, TargetLocation, PreviousMoveFraction);
 			CurrentExpectedLocation += GetPathOffsetInWorldSpace(PreviousMoveFraction);
 
@@ -1150,8 +1161,11 @@ void FRootMotionSource_JumpForce::PrepareRootMotion
 
 		if (TimeMappingCurve)
 		{
-			CurrentMoveFraction = TimeMappingCurve->GetFloatValue(CurrentTimeFraction);
-			TargetMoveFraction = TimeMappingCurve->GetFloatValue(TargetTimeFraction);
+			float MinMappingCurveTime, MaxMappingCurveTime;
+			TimeMappingCurve->GetTimeRange(MinMappingCurveTime, MaxMappingCurveTime);
+
+			CurrentMoveFraction = TimeMappingCurve->GetFloatValue(FMath::GetRangeValue(FVector2D(MinMappingCurveTime, MaxMappingCurveTime), CurrentMoveFraction));
+			TargetMoveFraction  = TimeMappingCurve->GetFloatValue(FMath::GetRangeValue(FVector2D(MinMappingCurveTime, MaxMappingCurveTime), TargetMoveFraction));
 		}
 
 		const FVector CurrentRelativeLocation = GetRelativeLocation(CurrentMoveFraction);
