Temp fix to resolve sequencer soft references in PIE.
A more robust solution will be implemented in the future.

diff --git a/Engine/Source/Runtime/Engine/Private/WorldPartition/WorldPartition.cpp b/Engine/Source/Runtime/Engine/Private/WorldPartition/WorldPartition.cpp
--- a/Engine/Source/Runtime/Engine/Private/WorldPartition/WorldPartition.cpp
+++ b/Engine/Source/Runtime/Engine/Private/WorldPartition/WorldPartition.cpp
@@ -916,20 +916,31 @@ void UWorldPartition::BeginDestroy()
 #if WITH_EDITOR
 bool UWorldPartition::LoadSubobject(const TCHAR* SubObjectPath, UObject*& OutObject, bool bOnlyTestExistence)
 {
-	for (UActorDescContainer::TIterator<> ActorDescIterator(this); ActorDescIterator; ++ActorDescIterator)
+	if (GetWorld()->IsGameWorld())
 	{
-		FWorldPartitionActorDesc* ActorDesc = *ActorDescIterator;
-
-		if (FString(ActorDesc->ActorPath.ToString()).EndsWith(SubObjectPath))
+		if (UObject* SubObject = StreamingPolicy->GetSubObject(SubObjectPath))
 		{
-			if (!bOnlyTestExistence)
+			OutObject = bOnlyTestExistence ? nullptr : SubObject;
+			return true;
+		}
+	}
+	else
+	{
+		for (UActorDescContainer::TIterator<> ActorDescIterator(this); ActorDescIterator; ++ActorDescIterator)
+		{
+			FWorldPartitionActorDesc* ActorDesc = *ActorDescIterator;
+
+			if (FString(ActorDesc->ActorPath.ToString()).EndsWith(SubObjectPath))
 			{
-				LoadedSubobjects.Emplace(this, ActorDesc->GetGuid());
-			}
+				if (!bOnlyTestExistence)
+				{
+					LoadedSubobjects.Emplace(this, ActorDesc->GetGuid());
+				}
 
-			OutObject = ActorDescIterator->GetActor();
+				OutObject = ActorDescIterator->GetActor();
 
-			return true;
+				return true;
+			}
 		}
 	}
 
diff --git a/Engine/Source/Runtime/Engine/Private/WorldPartition/WorldPartitionLevelStreamingPolicy.cpp b/Engine/Source/Runtime/Engine/Private/WorldPartition/WorldPartitionLevelStreamingPolicy.cpp
--- a/Engine/Source/Runtime/Engine/Private/WorldPartition/WorldPartitionLevelStreamingPolicy.cpp
+++ b/Engine/Source/Runtime/Engine/Private/WorldPartition/WorldPartitionLevelStreamingPolicy.cpp
@@ -212,6 +212,14 @@ void UWorldPartitionLevelStreamingPolicy::PrepareActorToCellRemapping()
 		for (const FWorldPartitionRuntimeCellObjectMapping& CellObjectMap : StreamingCell->GetPackages())
 		{
 			ActorToCellRemapping.Add(CellObjectMap.Path, StreamingCell->GetFName());
+
+			FString SubObjectName;
+			FString SubObjectContext;
+			verify(CellObjectMap.Path.ToString().Split(TEXT("."), &SubObjectContext, &SubObjectName));
+
+			int32 DotPos = SubObjectName.Find(TEXT("."), ESearchCase::IgnoreCase, ESearchDir::FromEnd);
+			check(DotPos != INDEX_NONE);
+			ActorToCellRemappingForPIE.Add(*SubObjectName.Right(SubObjectName.Len() - DotPos - 1), StreamingCell->GetFName());
 		}
 	}
 }
@@ -274,4 +282,21 @@ void UWorldPartitionLevelStreamingPolicy::RemapSoftObjectPath(FSoftObjectPath& O
 	}
 }
 
+UObject* UWorldPartitionLevelStreamingPolicy::GetSubObject(const TCHAR* SubObjectPath)
+{
+	FString SrcPath = UWorld::RemovePIEPrefix(SubObjectPath);
+
+	if (FName* CellName = ActorToCellRemappingForPIE.Find(FName(*SrcPath)))
+	{
+		if (UWorldPartitionRuntimeLevelStreamingCell* Cell = (UWorldPartitionRuntimeLevelStreamingCell*)StaticFindObject(UWorldPartitionRuntimeLevelStreamingCell::StaticClass(), GetOuterUWorldPartition(), *(CellName->ToString())))
+		{
+			if (UWorldPartitionLevelStreamingDynamic* LevelStraming = Cell->GetLevelStreaming())
+			{
+				return StaticFindObject(UObject::StaticClass(), LevelStraming->GetRuntimeLevel(), SubObjectPath);
+			}
+		}
+	}
+
+	return nullptr;
+}
 #endif
\ No newline at end of file
diff --git a/Engine/Source/Runtime/Engine/Public/WorldPartition/WorldPartitionLevelStreamingDynamic.h b/Engine/Source/Runtime/Engine/Public/WorldPartition/WorldPartitionLevelStreamingDynamic.h
--- a/Engine/Source/Runtime/Engine/Public/WorldPartition/WorldPartitionLevelStreamingDynamic.h
+++ b/Engine/Source/Runtime/Engine/Public/WorldPartition/WorldPartitionLevelStreamingDynamic.h
@@ -41,6 +41,8 @@ class ENGINE_API UWorldPartitionLevelStreamingDynamic : public ULevelStreamingDy
 	virtual bool RequestLevel(UWorld* PersistentWorld, bool bAllowLevelLoadRequests, EReqLevelBlock BlockPolicy) override;
 	virtual void BeginDestroy() override;
 
+	ULevel* GetRuntimeLevel() const { return RuntimeLevel; }
+
 private:
 	void CreateRuntimeLevel();
 	bool IssueLoadRequests();
diff --git a/Engine/Source/Runtime/Engine/Public/WorldPartition/WorldPartitionLevelStreamingPolicy.h b/Engine/Source/Runtime/Engine/Public/WorldPartition/WorldPartitionLevelStreamingPolicy.h
--- a/Engine/Source/Runtime/Engine/Public/WorldPartition/WorldPartitionLevelStreamingPolicy.h
+++ b/Engine/Source/Runtime/Engine/Public/WorldPartition/WorldPartitionLevelStreamingPolicy.h
@@ -31,6 +31,7 @@ public:
 	virtual void PrepareActorToCellRemapping() override;
 	virtual void ClearActorToCellRemapping() override;
 	virtual void RemapSoftObjectPath(FSoftObjectPath& ObjectPath) override;
+	virtual UObject* GetSubObject(const TCHAR* SubObjectPath) override;
 	static FString GetCellPackagePath(const FName& InCellName, const UWorld* InWorld);
 #endif
 
@@ -42,5 +43,6 @@ private:
 
 #if WITH_EDITOR
 	TMap<FName, FName> ActorToCellRemapping;
+	TMap<FName, FName> ActorToCellRemappingForPIE;
 #endif
 };
\ No newline at end of file
diff --git a/Engine/Source/Runtime/Engine/Public/WorldPartition/WorldPartitionStreamingPolicy.h b/Engine/Source/Runtime/Engine/Public/WorldPartition/WorldPartitionStreamingPolicy.h
--- a/Engine/Source/Runtime/Engine/Public/WorldPartition/WorldPartitionStreamingPolicy.h
+++ b/Engine/Source/Runtime/Engine/Public/WorldPartition/WorldPartitionStreamingPolicy.h
@@ -39,6 +39,7 @@ public:
 	virtual void PrepareActorToCellRemapping() {}
 	virtual void ClearActorToCellRemapping() {}
 	virtual void RemapSoftObjectPath(FSoftObjectPath& ObjectPath) {}
+	virtual UObject* GetSubObject(const TCHAR* SubObjectPath) { return nullptr; }
 #endif
 
 	const TArray<FWorldPartitionStreamingSource>& GetStreamingSources() const { return StreamingSources; }
