Do not take an async lock on ThumbnailInfo as it is not needed and can cause stalls from the content browser trying to access them

diff --git a/Engine/Source/Runtime/Engine/Classes/Engine/SkeletalMesh.h b/Engine/Source/Runtime/Engine/Classes/Engine/SkeletalMesh.h
--- a/Engine/Source/Runtime/Engine/Classes/Engine/SkeletalMesh.h
+++ b/Engine/Source/Runtime/Engine/Classes/Engine/SkeletalMesh.h
@@ -2234,6 +2234,11 @@ public:
 	*/
 	void RemoveLegacyClothingSections();
 
+	/*
+	* Handle some common preparation steps between async post load and async build
+	*/
+	void PrepareForAsyncCompilation();
+
 	/** Returns false if there is currently an async task running */
 	bool IsAsyncTaskComplete() const;
 
diff --git a/Engine/Source/Runtime/Engine/Private/SkeletalMesh.cpp b/Engine/Source/Runtime/Engine/Private/SkeletalMesh.cpp
--- a/Engine/Source/Runtime/Engine/Private/SkeletalMesh.cpp
+++ b/Engine/Source/Runtime/Engine/Private/SkeletalMesh.cpp
@@ -434,26 +434,6 @@ extern const FString& GetSkeletalMeshDerivedDataVersion();
 
 namespace SkeletalMeshImpl
 {
-	static void PrepareForAsyncCompilation()
-	{
-		// Make sure statics are initialized before calling from multiple threads
-		GetSkeletalMeshDerivedDataVersion();
-
-		// Make sure the target platform is properly initialized before accessing it from multiple threads
-		ITargetPlatformManagerModule& TargetPlatformManager = GetTargetPlatformManagerRef();
-		ITargetPlatform* RunningPlatform = TargetPlatformManager.GetRunningTargetPlatform();
-		check(RunningPlatform);
-
-		// Ensure those modules are loaded on the main thread - we'll need them in async tasks
-		FModuleManager::Get().LoadModuleChecked<IMeshUtilities>(TEXT("MeshUtilities"));
-		FModuleManager::Get().LoadModuleChecked<IMeshReductionManagerModule>(TEXT("MeshReductionInterface"));
-		IMeshBuilderModule::GetForRunningPlatform();
-		for (const ITargetPlatform* TargetPlatform : TargetPlatformManager.GetActiveTargetPlatforms())
-		{
-			IMeshBuilderModule::GetForPlatform(TargetPlatform);
-		}
-	}
-
 	// Condition specifically designed to detect if we're going to enter the non-thread safe part of FLODUtilities::SimplifySkeletalMeshLOD while building.
 	static bool HasInlineReductions(USkeletalMesh* SkeletalMesh)
 	{
@@ -1998,6 +1978,37 @@ void USkeletalMesh::WaitUntilAsyncPropertyReleased(ESkeletalMeshAsyncProperties
 	}
 }
 
+void USkeletalMesh::PrepareForAsyncCompilation()
+{
+	// Make sure statics are initialized before calling from multiple threads
+	GetSkeletalMeshDerivedDataVersion();
+
+	// Make sure the target platform is properly initialized before accessing it from multiple threads
+	ITargetPlatformManagerModule& TargetPlatformManager = GetTargetPlatformManagerRef();
+	ITargetPlatform* RunningPlatform = TargetPlatformManager.GetRunningTargetPlatform();
+	check(RunningPlatform);
+
+	// Ensure those modules are loaded on the main thread - we'll need them in async tasks
+	FModuleManager::Get().LoadModuleChecked<IMeshUtilities>(TEXT("MeshUtilities"));
+	FModuleManager::Get().LoadModuleChecked<IMeshReductionManagerModule>(TEXT("MeshReductionInterface"));
+	IMeshBuilderModule::GetForRunningPlatform();
+	for (const ITargetPlatform* TargetPlatform : TargetPlatformManager.GetActiveTargetPlatforms())
+	{
+		IMeshBuilderModule::GetForPlatform(TargetPlatform);
+	}
+
+	// Release any property that are not touched by async build/postload here
+
+	// The properties are still protected so if an async step tries to 
+	// use them without protection, it will assert and will mean we have
+	// to either avoid touching them asynchronously or we need to remove
+	// the release property here at the cost of maybe causing more stalls
+	// from the game thread.
+
+	// Not touched during async build and can cause stalls when the content browser refresh its tiles.
+	ReleaseAsyncProperty(ESkeletalMeshAsyncProperties::ThumbnailInfo);
+}
+
 void USkeletalMesh::Build()
 {
 	TRACE_CPUPROFILER_EVENT_SCOPE(USkeletalMesh::Build);
@@ -2016,7 +2027,7 @@ void USkeletalMesh::Build()
 	const bool bHasInlineReductions = SkeletalMeshImpl::HasInlineReductions(this);
 	if (!bHasInlineReductions && FSkeletalMeshCompilingManager::Get().IsAsyncCompilationAllowed(this))
 	{
-		SkeletalMeshImpl::PrepareForAsyncCompilation();
+		PrepareForAsyncCompilation();
 
 		FQueuedThreadPool* SkeletalMeshThreadPool = FSkeletalMeshCompilingManager::Get().GetThreadPool();
 		EQueuedWorkPriority BasePriority = FSkeletalMeshCompilingManager::Get().GetBasePriority(this);
@@ -3129,7 +3140,7 @@ void USkeletalMesh::PostLoad()
 	const bool bForceSynchronous = Context.bIsPreSkeletalMeshBuildRefactor || bHasInlineReductions;
 	if (!bForceSynchronous && FSkeletalMeshCompilingManager::Get().IsAsyncCompilationAllowed(this))
 	{
-		SkeletalMeshImpl::PrepareForAsyncCompilation();
+		PrepareForAsyncCompilation();
 
 		FQueuedThreadPool* SkeletalMeshThreadPool = FSkeletalMeshCompilingManager::Get().GetThreadPool();
 		EQueuedWorkPriority BasePriority = FSkeletalMeshCompilingManager::Get().GetBasePriority(this);
