Level Instance
- Fix WorldSettings actor not being ignored when packing (use the DiscardActors to skip actors, not only report)
- Ignore Visualization Components

diff --git a/Engine/Source/Runtime/Engine/Private/LevelInstance/Packed/LevelInstanceISMPacker.cpp b/Engine/Source/Runtime/Engine/Private/LevelInstance/Packed/LevelInstanceISMPacker.cpp
--- a/Engine/Source/Runtime/Engine/Private/LevelInstance/Packed/LevelInstanceISMPacker.cpp
+++ b/Engine/Source/Runtime/Engine/Private/LevelInstance/Packed/LevelInstanceISMPacker.cpp
@@ -26,9 +26,12 @@ void FLevelInstanceISMPacker::GetPackClusters(FPackedLevelInstanceBuilderContext
 
 	for(UStaticMeshComponent* StaticMeshComponent : StaticMeshComponents)
 	{
-		FLevelInstancePackerClusterID ClusterID(MakeUnique<FLevelInstanceISMPackerCluster>(GetID(), StaticMeshComponent));
+		if (InContext.ShouldPackComponent(StaticMeshComponent))
+		{
+			FLevelInstancePackerClusterID ClusterID(MakeUnique<FLevelInstanceISMPackerCluster>(GetID(), StaticMeshComponent));
 
-		InContext.FindOrAddCluster(MoveTemp(ClusterID), StaticMeshComponent);
+			InContext.FindOrAddCluster(MoveTemp(ClusterID), StaticMeshComponent);
+		}
 	}
 }
 
diff --git a/Engine/Source/Runtime/Engine/Private/LevelInstance/Packed/PackedLevelInstanceBuilder.cpp b/Engine/Source/Runtime/Engine/Private/LevelInstance/Packed/PackedLevelInstanceBuilder.cpp
--- a/Engine/Source/Runtime/Engine/Private/LevelInstance/Packed/PackedLevelInstanceBuilder.cpp
+++ b/Engine/Source/Runtime/Engine/Private/LevelInstance/Packed/PackedLevelInstanceBuilder.cpp
@@ -38,11 +38,14 @@
 
 void FPackedLevelInstanceBuilderContext::ClusterLevelActor(AActor* InActor)
 {
-	PerActorClusteredComponents.FindOrAdd(InActor);
-
-	for (const auto& Pair : Packers)
+	if (!ActorDiscards.Contains(InActor))
 	{
-		Pair.Value->GetPackClusters(*this, InActor);
+		PerActorClusteredComponents.FindOrAdd(InActor);
+
+		for (const auto& Pair : Packers)
+		{
+			Pair.Value->GetPackClusters(*this, InActor);
+		}
 	}
 }
 
@@ -147,8 +150,6 @@ TSharedPtr<FPackedLevelInstanceBuilder> FPackedLevelInstanceBuilder::CreateDefau
 		Builder->ClassDiscards.Add(ChaosDebugClass);
 	}
 
-	Builder->ClassDiscards.Add(AWorldSettings::StaticClass());
-	
 	// Root Components that are SceneComponents (not child class of)
 	Builder->ClassDiscards.Add(USceneComponent::StaticClass());
 	
@@ -184,7 +185,8 @@ bool FPackedLevelInstanceBuilder::PackActor(APackedLevelInstance* InPackedLevelI
 
 	ULevelStreaming* SourceLevelStreaming = ULevelStreaming::FindStreamingLevel(SourceLevel);
 	AWorldSettings* WorldSettings = SourceLevel->GetWorldSettings();
-	
+	Context.DiscardActor(WorldSettings);
+
 	Context.SetLevelTransform(SourceLevelStreaming->LevelTransform);
 	Context.SetPivotOffset(WorldSettings->LevelInstancePivotOffset);
 	
@@ -211,6 +213,11 @@ bool FPackedLevelInstanceBuilder::PackActor(APackedLevelInstance* InPackedLevelI
 	return true;
 }
 
+bool FPackedLevelInstanceBuilderContext::ShouldPackComponent(UActorComponent* ActorComponent) const
+{
+	return ActorComponent && !ActorComponent->IsVisualizationComponent();
+}
+
 void FPackedLevelInstanceBuilderContext::Report(FMessageLog& LevelInstanceLog) const
 {
 	TSet<UActorComponent*> NotClusteredComponents;
@@ -244,7 +251,7 @@ void FPackedLevelInstanceBuilderContext::Report(FMessageLog& LevelInstanceLog) c
 		uint32 WarningCount = 0;
 		for (UActorComponent* Component : NotClusteredComponents)
 		{
-			if (Component == nullptr)
+			if (!ShouldPackComponent(Component))
 			{
 				continue;
 			}
diff --git a/Engine/Source/Runtime/Engine/Public/LevelInstance/Packed/PackedLevelInstanceBuilder.h b/Engine/Source/Runtime/Engine/Public/LevelInstance/Packed/PackedLevelInstanceBuilder.h
--- a/Engine/Source/Runtime/Engine/Public/LevelInstance/Packed/PackedLevelInstanceBuilder.h
+++ b/Engine/Source/Runtime/Engine/Public/LevelInstance/Packed/PackedLevelInstanceBuilder.h
@@ -88,6 +88,7 @@ public:
 	const FTransform& GetLevelTransform() const { return LevelTransform; }
 	const FVector GetPivotOffset() const { return PivotOffset; }
 
+	bool ShouldPackComponent(UActorComponent* InActorComponent) const;
 private:
 	const TMap<FLevelInstancePackerID, TUniquePtr<ILevelInstancePacker>>& Packers;
 	const TSet<UClass*>& ClassDiscards;
