Fix non-unity error with FXSystem.h

diff --git a/Engine/Source/Runtime/Engine/Public/FXSystem.h b/Engine/Source/Runtime/Engine/Public/FXSystem.h
--- a/Engine/Source/Runtime/Engine/Public/FXSystem.h
+++ b/Engine/Source/Runtime/Engine/Public/FXSystem.h
@@ -12,12 +12,13 @@
 #include "RenderGraphDefinitions.h"
 
 class FCanvas;
-class FGPUSpriteResources;
-class UVectorFieldComponent;
+class FGPUSortManager;
 struct FGPUSpriteEmitterInfo;
 struct FGPUSpriteResourceData;
+class FGPUSpriteResources;
 struct FParticleEmitterInstance;
-class FGPUSortManager;
+class FRHIUniformBuffer;
+class UVectorFieldComponent;
 
 /*-----------------------------------------------------------------------------
 	Forward declarations.
