Fix FWorldPartitionStreamingQuerySource inverted comments.

diff --git a/Engine/Source/Runtime/Engine/Public/WorldPartition/WorldPartitionStreamingSource.h b/Engine/Source/Runtime/Engine/Public/WorldPartition/WorldPartitionStreamingSource.h
--- a/Engine/Source/Runtime/Engine/Public/WorldPartition/WorldPartitionStreamingSource.h
+++ b/Engine/Source/Runtime/Engine/Public/WorldPartition/WorldPartitionStreamingSource.h
@@ -38,11 +38,11 @@ struct FWorldPartitionStreamingQuerySource
 		, bSpatialQuery(true)
 	{}
 
-	/* Location to query. (not used if bNonSpatialQuery is true) */
+	/* Location to query. (not used if bSpatialQuery is false) */
 	UPROPERTY(BlueprintReadWrite, EditAnywhere, Category = "Query")
 	FVector Location;
 
-	/* Radius to query. (not used if bNonSpatialQuery is true) */
+	/* Radius to query. (not used if bSpatialQuery is false) */
 	UPROPERTY(BlueprintReadWrite, EditAnywhere, Category = "Query")
 	float Radius;
 
