Changed WorldPartition runtime cell names so that even when growing world bounds, names are stable.

diff --git a/Engine/Source/Runtime/Engine/Private/WorldPartition/RuntimeSpatialHash/RuntimeSpatialHashGridHelper.h b/Engine/Source/Runtime/Engine/Private/WorldPartition/RuntimeSpatialHash/RuntimeSpatialHashGridHelper.h
--- a/Engine/Source/Runtime/Engine/Private/WorldPartition/RuntimeSpatialHash/RuntimeSpatialHashGridHelper.h
+++ b/Engine/Source/Runtime/Engine/Private/WorldPartition/RuntimeSpatialHash/RuntimeSpatialHashGridHelper.h
@@ -405,6 +405,28 @@ struct FSquare2DGridHelper
 		return false;
 	}
 
+	/**
+	 * Returns the cell global coordinates
+	 *
+	 * @return true if the specified coord was valid
+	 */
+	inline bool GetCellGlobalCoords(const FIntVector& InCoords, FIntVector& OutGlobalCoords) const
+	{
+		if (Levels.IsValidIndex(InCoords.Z))
+		{
+			const FGridLevel& GridLevel = Levels[InCoords.Z];
+			if (GridLevel.IsValidCoords(FIntVector2(InCoords.X, InCoords.Y)))
+			{
+				int32 CoordOffset = Levels[InCoords.Z].GridSize >> 1;
+				OutGlobalCoords = InCoords;
+				OutGlobalCoords.X -= CoordOffset;
+				OutGlobalCoords.Y -= CoordOffset;
+				return true;
+			}
+		}
+		return false;
+	}
+
 #if WITH_EDITOR
 	// Runs a function on all cells
 	void ForEachCells(TFunctionRef<void(const FSquare2DGridHelper::FGridLevel::FGridCell&)> InOperation) const;
diff --git a/Engine/Source/Runtime/Engine/Private/WorldPartition/RuntimeSpatialHash/RuntimeSpatialHashHLOD.cpp b/Engine/Source/Runtime/Engine/Private/WorldPartition/RuntimeSpatialHash/RuntimeSpatialHashHLOD.cpp
--- a/Engine/Source/Runtime/Engine/Private/WorldPartition/RuntimeSpatialHash/RuntimeSpatialHashHLOD.cpp
+++ b/Engine/Source/Runtime/Engine/Private/WorldPartition/RuntimeSpatialHash/RuntimeSpatialHashHLOD.cpp
@@ -174,8 +174,9 @@ static TArray<FGuid> GenerateHLODsForGrid(UWorldPartition* WorldPartition, const
 			if (bShouldGenerateHLODs)
 			{
 				SlowTask.EnterProgressFrame(1);
-
-				FName CellName = UWorldPartitionRuntimeSpatialHash::GetCellName(WorldPartition, RuntimeGrid.GridName, CellCoord.Z, CellCoord.X, CellCoord.Y, GridCellDataChunk.GetDataLayersID());
+				FIntVector CellGlobalCoord;
+				verify(PartitionedActors.GetCellGlobalCoords(CellCoord, CellGlobalCoord));
+				FName CellName = UWorldPartitionRuntimeSpatialHash::GetCellName(WorldPartition, RuntimeGrid.GridName, CellGlobalCoord, GridCellDataChunk.GetDataLayersID());
 
 				TRACE_CPUPROFILER_EVENT_SCOPE_TEXT(*CellName.ToString());
 
diff --git a/Engine/Source/Runtime/Engine/Private/WorldPartition/RuntimeSpatialHash/RuntimeSpatialHashNavigation.cpp b/Engine/Source/Runtime/Engine/Private/WorldPartition/RuntimeSpatialHash/RuntimeSpatialHashNavigation.cpp
--- a/Engine/Source/Runtime/Engine/Private/WorldPartition/RuntimeSpatialHash/RuntimeSpatialHashNavigation.cpp
+++ b/Engine/Source/Runtime/Engine/Private/WorldPartition/RuntimeSpatialHash/RuntimeSpatialHashNavigation.cpp
@@ -52,7 +52,7 @@ bool UWorldPartitionRuntimeSpatialHash::GenerateNavigationData()
 
 	const FSquare2DGridHelper::FGridLevel& GridLevelHelper = GridHelper.Levels[GridLevel];
 
-	GridLevelHelper.ForEachCells([GridLevel, &GridLevelHelper, RuntimeGrid, WorldPartition, &ActorCount, World, NavSystem, &ValidNavigationDataChunkActors, this](const FIntVector2& CellCoord)
+	GridLevelHelper.ForEachCells([GridLevel, &GridHelper, &GridLevelHelper, RuntimeGrid, WorldPartition, &ActorCount, World, NavSystem, &ValidNavigationDataChunkActors, this](const FIntVector2& CellCoord)
 		{
 			FBox2D CellBounds;
 			GridLevelHelper.GetCellBounds(CellCoord, CellBounds);
@@ -94,7 +94,9 @@ bool UWorldPartitionRuntimeSpatialHash::GenerateNavigationData()
 				UE_LOG(LogWorldPartitionRuntimeSpatialHashNav, VeryVerbose, TEXT("Setting ChunkActorBounds to %s"), *ChunkActorBounds.ToString());
 				DataChunkActor->SetDataChunkActorBounds(ChunkActorBounds);
 
-				const FName CellName = GetCellName(RuntimeGrid.GridName, GridLevel, CellCoord.X, CellCoord.Y, GridCellDataChunk->GetDataLayersID());
+				FIntVector CellGlobalCoord;
+				verify(GridHelper.GetCellGlobalCoords(FIntVector(CellCoord.X, CellCoord.Y, GridLevel), CellGlobalCoord));
+				const FName CellName = GetCellName(RuntimeGrid.GridName, CellGlobalCoord, GridCellDataChunk->GetDataLayersID());
 				DataChunkActor->SetActorLabel(FString::Printf(TEXT("NavDataChunkActor_%s_%s"), *GetName(), *CellName.ToString()));
 
 				// Set target grid
diff --git a/Engine/Source/Runtime/Engine/Private/WorldPartition/WorldPartitionRuntimeSpatialHash.cpp b/Engine/Source/Runtime/Engine/Private/WorldPartition/WorldPartitionRuntimeSpatialHash.cpp
--- a/Engine/Source/Runtime/Engine/Private/WorldPartition/WorldPartitionRuntimeSpatialHash.cpp
+++ b/Engine/Source/Runtime/Engine/Private/WorldPartition/WorldPartitionRuntimeSpatialHash.cpp
@@ -333,7 +333,9 @@ void FSpatialHashStreamingGrid::Draw2D(UCanvas* Canvas, const TArray<FWorldParti
 			}
 			else
 			{
-				FString GridInfoText = FString::Printf(TEXT("X%02d_Y%02d"), Coords.X, Coords.Y);
+				FIntVector CellGlobalCoords;
+				verify(Helper.GetCellGlobalCoords(FIntVector(Coords.X, Coords.Y, GridLevel), CellGlobalCoords));
+				const FString GridInfoText = UWorldPartitionRuntimeSpatialHash::GetCellCoordString(CellGlobalCoords);
 				float TextWidth, TextHeight;
 				Canvas->SetDrawColor(255, 255, 0);
 				Canvas->StrLen(GEngine->GetTinyFont(), GridInfoText, TextWidth, TextHeight);
@@ -599,18 +601,23 @@ bool UWorldPartitionRuntimeSpatialHash::GenerateStreaming(EWorldPartitionStreami
 	return true;
 }
 
-FName UWorldPartitionRuntimeSpatialHash::GetCellName(UWorldPartition* WorldPartition, FName InGridName, int32 InLevel, int32 InCellX, int32 InCellY, const FDataLayersID& InDataLayerID)
+FString UWorldPartitionRuntimeSpatialHash::GetCellCoordString(const FIntVector& InCellGlobalCoord)
+{
+	return FString::Printf(TEXT("L%d_X%d_Y%d"), InCellGlobalCoord.Z, InCellGlobalCoord.X, InCellGlobalCoord.Y);
+}
+
+FName UWorldPartitionRuntimeSpatialHash::GetCellName(UWorldPartition* WorldPartition, FName InGridName, const FIntVector& InCellGlobalCoord, const FDataLayersID& InDataLayerID)
 {
 	const FString PackageName = FPackageName::GetShortName(WorldPartition->GetPackage());
 	const FString PackageNameNoPIEPrefix = UWorld::RemovePIEPrefix(PackageName);
 
-	return FName(*FString::Printf(TEXT("WPRT_%s_%s_Cell_L%d_X%02d_Y%02d_DL%X"), *PackageNameNoPIEPrefix, *InGridName.ToString(), InLevel, InCellX, InCellY, InDataLayerID.GetHash()));
+	return FName(*FString::Printf(TEXT("WPRT_%s_%s_Cell_%s_DL%X"), *PackageNameNoPIEPrefix, *InGridName.ToString(), *GetCellCoordString(InCellGlobalCoord), InDataLayerID.GetHash()));
 }
 
-FName UWorldPartitionRuntimeSpatialHash::GetCellName(FName InGridName, int32 InLevel, int32 InCellX, int32 InCellY, const FDataLayersID& InDataLayerID) const
+FName UWorldPartitionRuntimeSpatialHash::GetCellName(FName InGridName, const FIntVector& InCellGlobalCoord, const FDataLayersID& InDataLayerID) const
 {
 	UWorldPartition* WorldPartition = GetOuterUWorldPartition();
-	return UWorldPartitionRuntimeSpatialHash::GetCellName(WorldPartition, InGridName, InLevel, InCellX, InCellY, InDataLayerID);
+	return UWorldPartitionRuntimeSpatialHash::GetCellName(WorldPartition, InGridName, InCellGlobalCoord, InDataLayerID);
 }
 
 void UWorldPartitionRuntimeSpatialHash::CreateActorDescViewMap(const UActorDescContainer* Container, TMap<FGuid, FWorldPartitionActorDescView>& OutActorDescViewMap) const
@@ -785,7 +792,9 @@ bool UWorldPartitionRuntimeSpatialHash::CreateStreamingGrid(const FSpatialHashRu
 					continue;
 				}
 				
-				FName CellName = GetCellName(CurrentStreamingGrid.GridName, Level, CellCoordX, CellCoordY, GridCellDataChunk.GetDataLayersID());
+				FIntVector CellGlobalCoords;
+				verify(PartionedActors.GetCellGlobalCoords(FIntVector(CellCoordX, CellCoordY, Level), CellGlobalCoords));
+				FName CellName = GetCellName(CurrentStreamingGrid.GridName, CellGlobalCoords, GridCellDataChunk.GetDataLayersID());
 
 				UWorldPartitionRuntimeSpatialHashCell* StreamingCell = NewObject<UWorldPartitionRuntimeSpatialHashCell>(WorldPartition, StreamingPolicy->GetRuntimeCellClass(), CellName);
 
diff --git a/Engine/Source/Runtime/Engine/Public/WorldPartition/WorldPartitionRuntimeSpatialHash.h b/Engine/Source/Runtime/Engine/Public/WorldPartition/WorldPartitionRuntimeSpatialHash.h
--- a/Engine/Source/Runtime/Engine/Public/WorldPartition/WorldPartitionRuntimeSpatialHash.h
+++ b/Engine/Source/Runtime/Engine/Public/WorldPartition/WorldPartitionRuntimeSpatialHash.h
@@ -203,8 +203,8 @@ public:
 	virtual FName GetActorRuntimeGrid(const AActor* Actor) const override;
 	virtual void DrawPreview() const override;
 
-	FName GetCellName(FName InGridName, int32 InLevel, int32 InCellX, int32 InCellY, const FDataLayersID& InDataLayerID) const;
-	static FName GetCellName(UWorldPartition* WorldPartition, FName InGridName, int32 InLevel, int32 InCellX, int32 InCellY, const FDataLayersID& InDataLayerID);
+	FName GetCellName(FName InGridName, const FIntVector& InCellGlobalCoord, const FDataLayersID& InDataLayerID) const;
+	static FName GetCellName(UWorldPartition* WorldPartition, FName InGridName, const FIntVector& InCellGlobalCoord, const FDataLayersID& InDataLayerID);
 
 	virtual void CreateActorDescViewMap(const UActorDescContainer* Container, TMap<FGuid, FWorldPartitionActorDescView>& OutActorDescViewMap) const override;
 #endif
@@ -214,6 +214,8 @@ public:
 	virtual bool GetStreamingCells(const FWorldPartitionStreamingQuerySource& QuerySource, TSet<const UWorldPartitionRuntimeCell*>& OutCells) const override;
 	virtual bool GetStreamingCells(const TArray<FWorldPartitionStreamingSource>& Sources, TSet<const UWorldPartitionRuntimeCell*>& OutActivateCells, TSet<const UWorldPartitionRuntimeCell*>& OutLoadCells) const override;
 	virtual void SortStreamingCellsByImportance(const TSet<const UWorldPartitionRuntimeCell*>& InCells, const TArray<FWorldPartitionStreamingSource>& InSources, TArray<const UWorldPartitionRuntimeCell*, TInlineAllocator<256>>& OutSortedCells) const override;
+	
+	static FString GetCellCoordString(const FIntVector& InCellGlobalCoord);
 
 protected:
 
