Hide HLODs referencing non-existent cells

diff --git a/Engine/Source/Runtime/Engine/Private/WorldPartition/HLOD/HLODSubsystem.cpp b/Engine/Source/Runtime/Engine/Private/WorldPartition/HLOD/HLODSubsystem.cpp
--- a/Engine/Source/Runtime/Engine/Private/WorldPartition/HLOD/HLODSubsystem.cpp
+++ b/Engine/Source/Runtime/Engine/Private/WorldPartition/HLOD/HLODSubsystem.cpp
@@ -16,6 +16,8 @@
 
 #define LOCTEXT_NAMESPACE "HLODSubsystem"
 
+DEFINE_LOG_CATEGORY_STATIC(LogHLODSubsystem, Log, All);
+
 UHLODSubsystem::UHLODSubsystem()
 	: UWorldSubsystem()
 {
@@ -55,51 +57,49 @@ void UHLODSubsystem::Initialize(FSubsystemCollectionBase& Collection)
 
 void UHLODSubsystem::OnWorldPartitionRegistered(UWorldPartition* InWorldPartition)
 {
+	check(InWorldPartition == GetWorld()->GetWorldPartition());
+	check(CellsHLODMapping.IsEmpty());
+
 	TSet<const UWorldPartitionRuntimeCell*> StreamingCells;
 	InWorldPartition->RuntimeHash->GetAllStreamingCells(StreamingCells, /*bIncludeDataLayers*/ true);
 
-	check(!CellsHLODMapping.Contains(InWorldPartition));
-	TMap<FName, FCellHLODMapping>& CellsMapping = CellsHLODMapping.Emplace(InWorldPartition);
-
 	// Build cell to HLOD mapping
 	for (const UWorldPartitionRuntimeCell* Cell : StreamingCells)
 	{
-		CellsMapping.Emplace(Cell->GetFName());
+		CellsHLODMapping.Emplace(Cell->GetFName());
 	}
 }
 
 void UHLODSubsystem::OnWorldPartitionUnregistered(UWorldPartition* InWorldPartition)
 {
-	check(CellsHLODMapping.Contains(InWorldPartition));
-	CellsHLODMapping.Remove(InWorldPartition);
+	check(InWorldPartition == GetWorld()->GetWorldPartition());
+	CellsHLODMapping.Reset();
 }
 
-
 void UHLODSubsystem::RegisterHLODActor(AWorldPartitionHLOD* InWorldPartitionHLOD)
 {
 	TRACE_CPUPROFILER_EVENT_SCOPE(UHLODSubsystem::RegisterHLODActor);
 
-	UWorldPartition* OwnerPartition = InWorldPartitionHLOD->GetWorld()->GetWorldPartition();
-	check(OwnerPartition);
-
 	FName CellName = InWorldPartitionHLOD->GetCellName();
-	FCellHLODMapping* CellHLODs = CellsHLODMapping.FindChecked(OwnerPartition).Find(CellName);
+	FCellHLODMapping* CellHLODs = CellsHLODMapping.Find(CellName);
 	if (CellHLODs)
 	{
 		CellHLODs->LoadedHLODs.Add(InWorldPartitionHLOD);
 		InWorldPartitionHLOD->SetVisibility(!CellHLODs->bIsCellVisible);
 	}
+	else
+	{
+		UE_LOG(LogHLODSubsystem, Warning, TEXT("Found HLOD referencing nonexistent cell '%s'"), *CellName.ToString());
+		InWorldPartitionHLOD->SetVisibility(false);
+	}
 }
 
 void UHLODSubsystem::UnregisterHLODActor(AWorldPartitionHLOD* InWorldPartitionHLOD)
 {
 	TRACE_CPUPROFILER_EVENT_SCOPE(UHLODSubsystem::UnregisterHLODActor);
 
-	UWorldPartition* OwnerPartition = InWorldPartitionHLOD->GetWorld()->GetWorldPartition();
-	check(OwnerPartition);
-
 	FName CellName = InWorldPartitionHLOD->GetCellName();
-	FCellHLODMapping* CellHLODs = CellsHLODMapping.FindChecked(OwnerPartition).Find(CellName);
+	FCellHLODMapping* CellHLODs = CellsHLODMapping.Find(CellName);
 	if (CellHLODs)
 	{
 		int32 NumRemoved = CellHLODs->LoadedHLODs.Remove(InWorldPartitionHLOD);
@@ -109,35 +109,21 @@ void UHLODSubsystem::UnregisterHLODActor(AWorldPartitionHLOD* InWorldPartitionHL
 
 void UHLODSubsystem::OnCellShown(const UWorldPartitionRuntimeCell* InCell)
 {
-	UWorldPartition* OwnerPartition = InCell->GetTypedOuter<UWorldPartition>();
-	check(OwnerPartition);
-
-	FCellHLODMapping* CellHLODs = CellsHLODMapping.FindChecked(OwnerPartition).Find(InCell->GetFName());
-	if (CellHLODs)
+	FCellHLODMapping& CellHLODs = CellsHLODMapping.FindChecked(InCell->GetFName());
+	CellHLODs.bIsCellVisible = true;
+	for (AWorldPartitionHLOD* HLODActor : CellHLODs.LoadedHLODs)
 	{
-		CellHLODs->bIsCellVisible = true;
-
-		for (AWorldPartitionHLOD* HLODActor : CellHLODs->LoadedHLODs)
-		{
-			HLODActor->SetVisibility(false);
-		}
+		HLODActor->SetVisibility(false);
 	}
 }
 
 void UHLODSubsystem::OnCellHidden(const UWorldPartitionRuntimeCell* InCell)
 {
-	UWorldPartition* OwnerPartition = InCell->GetTypedOuter<UWorldPartition>();
-	check(OwnerPartition);
-
-	FCellHLODMapping* CellHLODs = CellsHLODMapping.FindChecked(OwnerPartition).Find(InCell->GetFName());
-	if (CellHLODs)
+	FCellHLODMapping& CellHLODs = CellsHLODMapping.FindChecked(InCell->GetFName());
+	CellHLODs.bIsCellVisible = false;
+	for (AWorldPartitionHLOD* HLODActor : CellHLODs.LoadedHLODs)
 	{
-		CellHLODs->bIsCellVisible = false;
-
-		for (AWorldPartitionHLOD* HLODActor : CellHLODs->LoadedHLODs)
-		{
-			HLODActor->SetVisibility(true);
-		}
+		HLODActor->SetVisibility(true);
 	}
 }
 
diff --git a/Engine/Source/Runtime/Engine/Public/WorldPartition/HLOD/HLODSubsystem.h b/Engine/Source/Runtime/Engine/Public/WorldPartition/HLOD/HLODSubsystem.h
--- a/Engine/Source/Runtime/Engine/Public/WorldPartition/HLOD/HLODSubsystem.h
+++ b/Engine/Source/Runtime/Engine/Public/WorldPartition/HLOD/HLODSubsystem.h
@@ -50,5 +50,5 @@ private:
 		}
 	};
 
-	TMap<UWorldPartition*, TMap<FName, FCellHLODMapping>> CellsHLODMapping;
+	TMap<FName, FCellHLODMapping> CellsHLODMapping;
 };
