Fixing UE4 comments for PhysX/Engine

diff --git a/Engine/Source/Runtime/Engine/Private/PhysicsEngine/BodyInstance.cpp b/Engine/Source/Runtime/Engine/Private/PhysicsEngine/BodyInstance.cpp
--- a/Engine/Source/Runtime/Engine/Private/PhysicsEngine/BodyInstance.cpp
+++ b/Engine/Source/Runtime/Engine/Private/PhysicsEngine/BodyInstance.cpp
@@ -1593,7 +1593,7 @@ void FBodyInstance::TermBody(bool bNeverDeferRelease)
 	}
 
 
-	// @TODO UE4: Release spring body here
+	// @TODO: Release spring body here
 
 	CurrentSceneState = BodyInstanceSceneState::NotAdded;
 	BodySetup = NULL;
diff --git a/Engine/Source/Runtime/Engine/Private/PhysicsEngine/PhysAnim.cpp b/Engine/Source/Runtime/Engine/Private/PhysicsEngine/PhysAnim.cpp
--- a/Engine/Source/Runtime/Engine/Private/PhysicsEngine/PhysAnim.cpp
+++ b/Engine/Source/Runtime/Engine/Private/PhysicsEngine/PhysAnim.cpp
@@ -645,7 +645,7 @@ void USkeletalMeshComponent::UpdateKinematicBonesToAnim(const TArray<FTransform>
 								// if uniform, we'll use BoneTranform
 								if(MeshScale3D.IsUniform())
 								{
-									// @todo UE4 should we update scale when it's simulated?
+									// @todo should we update scale when it's simulated?
 									BodyInst->UpdateBodyScale(BoneTransform.GetScale3D());
 								}
 								else
diff --git a/Engine/Source/Runtime/Engine/Private/UnrealEngine.cpp b/Engine/Source/Runtime/Engine/Private/UnrealEngine.cpp
--- a/Engine/Source/Runtime/Engine/Private/UnrealEngine.cpp
+++ b/Engine/Source/Runtime/Engine/Private/UnrealEngine.cpp
@@ -1800,7 +1800,7 @@ void UEngine::Init(IEngineLoop* InEngineLoop)
 	EngineStats.Add(FEngineStatFuncs(TEXT("STAT_SoundCues"), TEXT("STATCAT_Engine"), FText::GetEmpty(), FEngineStatRender::CreateUObject(this, &UEngine::RenderStatSoundCues), FEngineStatToggle::CreateUObject(this, &UEngine::ToggleStatSoundCues)));
 	EngineStats.Add(FEngineStatFuncs(TEXT("STAT_SoundWaves"), TEXT("STATCAT_Engine"), FText::GetEmpty(), FEngineStatRender::CreateUObject(this, &UEngine::RenderStatSoundWaves), FEngineStatToggle::CreateUObject(this, &UEngine::ToggleStatSoundWaves)));
 #endif // !UE_BUILD_SHIPPING
-	/* @todo UE4 physx fix this once we have convexelem drawing again
+	/* @todo physx fix this once we have convexelem drawing again
 	EngineStats.Add(FEngineStatFuncs(TEXT("STAT_LevelMap"), TEXT("STATCAT_Engine"), FText::GetEmpty(), &UEngine::RenderStatLevelMap, NULL));
 	*/
 	EngineStats.Add(FEngineStatFuncs(TEXT("STAT_Detailed"), TEXT("STATCAT_Engine"), FText::GetEmpty(), FEngineStatRender(), FEngineStatToggle::CreateUObject(this, &UEngine::ToggleStatDetailed)));
@@ -9753,7 +9753,7 @@ static void DrawVolumeOnCanvas(const AVolume* Volume, FCanvas* Canvas, const FVe
 		{
 			FKConvexElem& ConvElem = Volume->GetBrushComponent()->BrushBodySetup->AggGeom.ConvexElems[ConIdx];
 
-#if 0 // @todo UE4 physx fix this once we have convexelem drawing again
+#if 0 // @todo physx fix this once we have convexelem drawing again
 			// Draw each triangle that makes up the convex hull
 			const int32 NumTris = ConvElem.FaceTriData.Num()/3;
 			for(int32 i=0; i<NumTris; i++)
diff --git a/Engine/Source/Runtime/Engine/Public/CustomPhysXPayload.h b/Engine/Source/Runtime/Engine/Public/CustomPhysXPayload.h
--- a/Engine/Source/Runtime/Engine/Public/CustomPhysXPayload.h
+++ b/Engine/Source/Runtime/Engine/Public/CustomPhysXPayload.h
@@ -12,22 +12,22 @@ namespace physx
 	class PxRigidActor;
 }
 
-/** This interface allows plugins to sync between physx sim results and UE4 data. */
+/** This interface allows plugins to sync between physx sim results and Unreal Engine data. */
 struct FCustomPhysXSyncActors
 {
 	virtual ~FCustomPhysXSyncActors() {}
 
-	/** Syncing between PhysX and UE4 can be tricky. The challenge is that PhysX gives us an array of raw pointers that require updating. However, UE4 has many mechanisms for triggering arbitrary callbacks.
+	/** Syncing between PhysX and Unreal Engine can be tricky. The challenge is that PhysX gives us an array of raw pointers that require updating. However, Unreal Engine has many mechanisms for triggering arbitrary callbacks.
 	  * Since the callbacks can invalidate the PhysX data it's critical to break the update into two steps.
-	  * The first step moves the data from PhysX into UE4 safe structures (see TWeakObjectPtr).
-	  * The second step updates the actual UE4 data (components, actors, etc...)
+	  * The first step moves the data from PhysX into Unreal Engine safe structures (see TWeakObjectPtr).
+	  * The second step updates the actual Unreal Engine data (components, actors, etc...)
 	  * It's up to the client code to ensure this is handled correctly. No PhysX raw data should be stored. Do not update component/actor data in step 1 unless you're absolutely sure it will not trigger a callback
 	  */
 
-	/** This is step 1 of the sync. Do not store any raw PhysX data or update any UE4 component/actors which may trigger a callback. Ideally this should build an array of TWeakObjectPtr<UYourComponent>, Transform pairs */
+	/** This is step 1 of the sync. Do not store any raw PhysX data or update any Unreal Engine component/actors which may trigger a callback. Ideally this should build an array of TWeakObjectPtr<UYourComponent>, Transform pairs */
 	virtual void BuildSyncData_AssumesLocked(const TArray<physx::PxRigidActor*>& RigidActors) = 0;
 
-	/** This is step 2 of the sync. It's safe to update any UE4 components/actors or trigger callbacks, as long as the code within this function handles it gracefully.
+	/** This is step 2 of the sync. It's safe to update any Unreal Engine components/actors or trigger callbacks, as long as the code within this function handles it gracefully.
 	  * For example if a callback destroys some components you will update, make sure to use TWeakObjectPtr */
 	virtual void FinalizeSync() = 0;
 
