Reconciling difference between Main and Early-Access w.r.t CL 15851594 to prevent move/edit P4 operations

diff --git a/Engine/Source/Developer/AnimationDataController/Private/AnimDataController.cpp b/Engine/Source/Developer/AnimationDataController/Private/AnimDataController.cpp
--- a/Engine/Source/Developer/AnimationDataController/Private/AnimDataController.cpp
+++ b/Engine/Source/Developer/AnimationDataController/Private/AnimDataController.cpp
@@ -1552,7 +1552,7 @@ void UAnimDataController::ResizeAttributes(float NewLength, bool bInserted, floa
 			Values.Add(Keys[KeyIndex].GetValuePtr<void>());
 		}
 		
-		SetAttributeKeys(Attribute.Identifier, Times, Values, bShouldTransact);
+		SetAttributeKeys(Attribute.Identifier, Times, Values, Attribute.Curve.GetScriptStruct(), bShouldTransact);
 	}
 }
 
diff --git a/Engine/Source/Developer/AnimationDataController/Private/AnimDataControllerActions.cpp b/Engine/Source/Developer/AnimationDataController/Private/AnimDataControllerActions.cpp
--- a/Engine/Source/Developer/AnimationDataController/Private/AnimDataControllerActions.cpp
+++ b/Engine/Source/Developer/AnimationDataController/Private/AnimDataControllerActions.cpp
@@ -425,7 +425,7 @@ TUniquePtr<FChange> FAddAtributeAction::ExecuteInternal(UAnimDataModel* Model, U
 	});
 
 	Controller->AddAttribute(AttributeId, false);
-	Controller->SetAttributeKeys(AttributeId, MakeArrayView(Times), MakeArrayView(VoidValues), false);
+	Controller->SetAttributeKeys(AttributeId, MakeArrayView(Times), MakeArrayView(VoidValues), AttributeId.GetType(), false);
 
 	return MakeUnique<FRemoveAtributeAction>(AttributeId);
 }
@@ -452,7 +452,7 @@ FString FRemoveAtributeAction::ToStringInternal() const
 
 TUniquePtr<FChange> FAddAtributeKeyAction::ExecuteInternal(UAnimDataModel* Model, UAnimDataController* Controller)
 {
-	Controller->SetAttributeKey(AttributeId, Key.Time, Key.GetValuePtr<void>(), false);
+	Controller->SetAttributeKey(AttributeId, Key.Time, Key.GetValuePtr<void>(), AttributeId.GetType(), false);
 
 	return MakeUnique<FRemoveAtributeKeyAction>(AttributeId, Key.Time);
 }
@@ -470,7 +470,7 @@ TUniquePtr<FChange> FSetAtributeKeyAction::ExecuteInternal(UAnimDataModel* Model
 	ensure(Handle != FKeyHandle::Invalid());
 
 	FAttributeKey CurrentKey = Attribute.Curve.GetKey(Handle);
-	Controller->SetAttributeKey(AttributeId, Key.Time, Key.GetValuePtr<void>(), false);
+	Controller->SetAttributeKey(AttributeId, Key.Time, Key.GetValuePtr<void>(), AttributeId.GetType(), false);
 
 	return MakeUnique<FSetAtributeKeyAction>(AttributeId, CurrentKey);
 }
@@ -520,7 +520,7 @@ TUniquePtr<FChange> FSetAtributeKeysAction::ExecuteInternal(UAnimDataModel* Mode
 		return Key.Time;
 	});
 
-	Controller->SetAttributeKeys(AttributeId, MakeArrayView(Times), MakeArrayView(VoidValues), false);
+	Controller->SetAttributeKeys(AttributeId, MakeArrayView(Times), MakeArrayView(VoidValues), AttributeId.GetType(), false);
 
 	return InverseAction;
 }
diff --git a/Engine/Source/Developer/AnimationDataController/Public/AnimDataController.h b/Engine/Source/Developer/AnimationDataController/Public/AnimDataController.h
--- a/Engine/Source/Developer/AnimationDataController/Public/AnimDataController.h
+++ b/Engine/Source/Developer/AnimationDataController/Public/AnimDataController.h
@@ -76,22 +76,41 @@ public:
 	virtual void UpdateCurveNamesFromSkeleton(const USkeleton* Skeleton, ERawCurveTrackTypes SupportedCurveType, bool bShouldTransact = true) override;
 	virtual void FindOrAddCurveNamesOnSkeleton(USkeleton* Skeleton, ERawCurveTrackTypes SupportedCurveType, bool bShouldTransact = true) override;
 	virtual bool RemoveBoneTracksMissingFromSkeleton(const USkeleton* Skeleton, bool bShouldTransact = true) override;
+	virtual void UpdateAttributesFromSkeleton(const USkeleton* Skeleton, bool bShouldTransact = true) override;
 	virtual void NotifyPopulated() override;
 	virtual void ResetModel(bool bShouldTransact = true) override;
-	
+	virtual bool AddAttribute(const FAnimationAttributeIdentifier& AttributeIdentifier, bool bShouldTransact = true) override;	
+    virtual bool RemoveAttribute(const FAnimationAttributeIdentifier& AttributeIdentifier, bool bShouldTransact = true) override;
+    virtual int32 RemoveAllAttributesForBone(const FName& BoneName, bool bShouldTransact = true) override;
+    virtual int32 RemoveAllAttributes(bool bShouldTransact = true) override;
+	virtual bool SetAttributeKey(const FAnimationAttributeIdentifier& AttributeIdentifier, float Time, const void* KeyValue, const UScriptStruct* TypeStruct, bool bShouldTransact = true) override
+	{
+		return SetAttributeKey_Internal(AttributeIdentifier, Time, KeyValue, TypeStruct, bShouldTransact);
+	}
+	virtual bool SetAttributeKeys(const FAnimationAttributeIdentifier& AttributeIdentifier, TArrayView<const float> Times, TArrayView<const void*> KeyValues, const UScriptStruct* TypeStruct, bool bShouldTransact = true) override
+	{
+		return SetAttributeKeys_Internal(AttributeIdentifier, Times, KeyValues, TypeStruct, bShouldTransact);
+	}
+    virtual bool RemoveAttributeKey(const FAnimationAttributeIdentifier& AttributeIdentifier, float Time, bool bShouldTransact = true) override;	
+	virtual bool DuplicateAttribute(const FAnimationAttributeIdentifier& AttributeIdentifier, const FAnimationAttributeIdentifier& NewAttributeIdentifier, bool bShouldTransact = true) override;
 protected:
 	virtual void NotifyBracketOpen() override;
 	virtual void NotifyBracketClosed() override;
 	/** End IAnimationDataController overrides */
 	
 private:
+	/** Internal functionality for setting Attribute curve key(s) */
+	bool SetAttributeKey_Internal(const FAnimationAttributeIdentifier& AttributeIdentifier, float Time, const void* KeyValue, const UScriptStruct* TypeStruct, bool bShouldTransact = true);
+	bool SetAttributeKeys_Internal(const FAnimationAttributeIdentifier& AttributeIdentifier, TArrayView<const float> Times, TArrayView<const void*> KeyValues, const UScriptStruct* TypeStruct, bool bShouldTransact = true);
+
 	/** Returns whether or not the supplied curve type is supported by the controller functionality */
 	const bool IsSupportedCurveType(ERawCurveTrackTypes CurveType) const;
 	/** Returns the string representation of the provided curve enum type value */
 	FString GetCurveTypeValueName(ERawCurveTrackTypes InType) const;
 	
-	/** Resizes the curve data stored on the model according to the provided new length and time at which to insert or remove time */
+	/** Resizes the curve/attribute data stored on the model according to the provided new length and time at which to insert or remove time */
 	void ResizeCurves(float NewLength, bool bInserted, float T0, float T1, bool bShouldTransact = true);
+	void ResizeAttributes(float NewLength, bool bInserted, float T0, float T1, bool bShouldTransact = true);
 
 	/** Ensures that a valid model is currently targeted */
 	void ValidateModel() const;
diff --git a/Engine/Source/Editor/SequenceRecorder/Private/AnimationRecorder.cpp b/Engine/Source/Editor/SequenceRecorder/Private/AnimationRecorder.cpp
--- a/Engine/Source/Editor/SequenceRecorder/Private/AnimationRecorder.cpp
+++ b/Engine/Source/Editor/SequenceRecorder/Private/AnimationRecorder.cpp
@@ -537,8 +537,8 @@ void FAnimationRecorder::ProcessRecordedTimes(UAnimSequence* AnimSequence, USkel
 	TArray<FString> Slates(&Slate, 1);
 	TArray<float> SlateTimes(&Times[0], 1);
 
-	UAnimDataController* Controller = AnimSequence->GetController();
-	UAnimDataController::FScopedBracket(Controller, LOCTEXT("AddTimeCodeAttributesBracket", "Adding Time Code attributes"));
+	IAnimationDataController& Controller = AnimSequence->GetController();
+	IAnimationDataController::FScopedBracket(Controller, LOCTEXT("AddTimeCodeAttributesBracket", "Adding Time Code attributes"));
 
 	for (int32 BoneIndex = 0; BoneIndex < SpaceBases->Num(); ++BoneIndex)
 	{
diff --git a/Engine/Source/Editor/UnrealEd/Private/SkeletalMeshEdit.cpp b/Engine/Source/Editor/UnrealEd/Private/SkeletalMeshEdit.cpp
--- a/Engine/Source/Editor/UnrealEd/Private/SkeletalMeshEdit.cpp
+++ b/Engine/Source/Editor/UnrealEd/Private/SkeletalMeshEdit.cpp
@@ -1543,7 +1543,7 @@ bool UnFbx::FFbxImporter::ImportAnimation(USkeleton* Skeleton, UAnimSequence * D
 
 	if (ImportOptions->bDeleteExistingNonCurveCustomAttributes)
 	{
-		Controller->RemoveAllAttributes();
+		Controller.RemoveAllAttributes();
 	}
 	
 	const bool bReimportWarnings = GetDefault<UEditorPerProjectUserSettings>()->bAnimationReimportWarnings;
diff --git a/Engine/Source/Runtime/Engine/Classes/Animation/AnimData/IAnimationDataController.h b/Engine/Source/Runtime/Engine/Classes/Animation/AnimData/IAnimationDataController.h
--- a/Engine/Source/Runtime/Engine/Classes/Animation/AnimData/IAnimationDataController.h
+++ b/Engine/Source/Runtime/Engine/Classes/Animation/AnimData/IAnimationDataController.h
@@ -5,10 +5,12 @@
 #include "CoreTypes.h"
 #include "UObject/Interface.h"
 #include "CurveIdentifier.h"
+#include "AttributeIdentifier.h"
+#include "Algo/Transform.h"
+
 #include "IAnimationDataController.generated.h"
 
 class UAssetUserData;
-
 class UAnimDataModel;
 
 namespace UE {
@@ -387,7 +389,7 @@ public:
 	* @param	CurveKeys			Keys with which the existing keys are to be replaced
 	* @param	bShouldTransact		Whether or not any undo-redo changes should be generated
 	*
-	* @return	Whether or not replacing curve keys was succesful
+	* @return	Whether or not replacing curve keys was successful
 	*/
 	UFUNCTION(BlueprintCallable, Category = CurveData)
 	virtual bool SetCurveKeys(const FAnimationCurveIdentifier& CurveId, const TArray<FRichCurveKey>& CurveKeys, bool bShouldTransact = true) = 0;
@@ -427,7 +429,7 @@ public:
 	* @param	Skeleton			Skeleton to retrieve the bone information from
 	* @param	bShouldTransact		Whether or not any undo-redo changes should be generated
 	*/
-	void UpdateAttributesFromSkeleton(const USkeleton* Skeleton, bool bShouldTransact = true);
+	virtual void UpdateAttributesFromSkeleton(const USkeleton* Skeleton, bool bShouldTransact = true) = 0;
 
 	/**
 	* Broadcast a EAnimDataModelNotifyType::Populated notify.
@@ -442,7 +444,7 @@ public:
 	*	- Frame rate to 30fps, broadcasts a EAnimDataModelNotifyType::FrameRateChanged
 	*
 	* @param	bShouldTransact		Whether or not any undo-redo changes should be generated
-	*/	
+	*/
 	virtual void ResetModel(bool bShouldTransact = true) = 0;
 
 	/**
@@ -453,8 +455,8 @@ public:
 	*
 	* @return	Whether or not the attribute was successfully added
 	*/
-	UFUNCTION(BlueprintCallable, Category = CurveData)
-	bool AddAttribute(const FAnimationAttributeIdentifier& AttributeIdentifier, bool bShouldTransact = true);
+	UFUNCTION(BlueprintCallable, Category = AttributeData)
+	virtual bool AddAttribute(const FAnimationAttributeIdentifier& AttributeIdentifier, bool bShouldTransact = true) = 0;
 
 	/**
 	* Removes an attribute, if found, with the provided information. Broadcasts a EAnimDataModelNotifyType::AttributeRemoved notify if successful.
@@ -464,8 +466,8 @@ public:
 	*
 	* @return	Whether or not the attribute was successfully removed
 	*/
-	UFUNCTION(BlueprintCallable, Category = CurveData)
-	bool RemoveAttribute(const FAnimationAttributeIdentifier& AttributeIdentifier, bool bShouldTransact = true);
+	UFUNCTION(BlueprintCallable, Category = AttributeData)
+	virtual bool RemoveAttribute(const FAnimationAttributeIdentifier& AttributeIdentifier, bool bShouldTransact = true) = 0;
 
 	/**
 	* Removes all attributes for the specified bone name, if any. Broadcasts a EAnimDataModelNotifyType::AttributeRemoved notify for each removed attribute.
@@ -475,8 +477,8 @@ public:
 	*
 	* @return	Total number of removes attributes
 	*/
-	UFUNCTION(BlueprintCallable, Category = CurveData)
-	int32 RemoveAllAttributesForBone(const FName& BoneName, bool bShouldTransact = true);
+	UFUNCTION(BlueprintCallable, Category = AttributeData)
+	virtual int32 RemoveAllAttributesForBone(const FName& BoneName, bool bShouldTransact = true) = 0;
 
 	/**
 	* Removes all stored attributes. Broadcasts a EAnimDataModelNotifyType::AttributeRemoved notify for each removed attribute.
@@ -485,8 +487,8 @@ public:
 	*
 	* @return	Total number of removes attributes
 	*/
-	UFUNCTION(BlueprintCallable, Category = CurveData)
-	int32 RemoveAllAttributes(bool bShouldTransact = true);	
+	UFUNCTION(BlueprintCallable, Category = AttributeData)
+	virtual int32 RemoveAllAttributes(bool bShouldTransact = true) = 0;	
 
 	/**
 	* Sets a single key for the attribute with provided identifier. Broadcasts a EAnimDataModelNotifyType::AttributeChanged notify if successful.
@@ -502,7 +504,7 @@ public:
 	template<typename AttributeType>
 	bool SetTypedAttributeKey(const FAnimationAttributeIdentifier& AttributeIdentifier, float Time, AttributeType& KeyValue, bool bShouldTransact = true)
 	{
-		return SetAttributeKey_Internal(AttributeIdentifier, Time, (const void*)&KeyValue, AttributeType::StaticStruct(), bShouldTransact);
+		return SetAttributeKey(AttributeIdentifier, Time, (const void*)&KeyValue, AttributeType::StaticStruct(), bShouldTransact);
 	}
 
 	/**
@@ -516,10 +518,7 @@ public:
 	*
 	* @return	Whether or not the key was successfully set
 	*/
-	bool SetAttributeKey(const FAnimationAttributeIdentifier& AttributeIdentifier, float Time, const void* KeyValue, bool bShouldTransact = true)
-	{
-		return SetAttributeKey_Internal(AttributeIdentifier, Time, KeyValue, AttributeIdentifier.GetType(), bShouldTransact);
-	}
+	virtual bool SetAttributeKey(const FAnimationAttributeIdentifier& AttributeIdentifier, float Time, const void* KeyValue, const UScriptStruct* TypeStruct, bool bShouldTransact = true) = 0;
 	
 	/**
 	* Replace the keys for the attribute with provided identifier. Broadcasts a EAnimDataModelNotifyType::AttributeChanged notify if successful.
@@ -531,10 +530,7 @@ public:
 	*
 	* @return	Whether or not replacing the attribute keys was successful
 	*/
-	bool SetAttributeKeys(const FAnimationAttributeIdentifier& AttributeIdentifier, TArrayView<const float> Times, TArrayView<const void*> KeyValues, bool bShouldTransact = true)
-	{
-		return SetAttributeKeys_Internal(AttributeIdentifier, Times, KeyValues, AttributeIdentifier.GetType(), bShouldTransact);
-	}
+	virtual bool SetAttributeKeys(const FAnimationAttributeIdentifier& AttributeIdentifier, TArrayView<const float> Times, TArrayView<const void*> KeyValues, const UScriptStruct* TypeStruct, bool bShouldTransact = true) = 0;
 
 	/**
 	* Replace the keys for the attribute with provided identifier. Broadcasts a EAnimDataModelNotifyType::AttributeChanged notify if successful.
@@ -548,14 +544,14 @@ public:
 	*/
 	template<typename AttributeType>
 	bool SetTypedAttributeKeys(const FAnimationAttributeIdentifier& AttributeIdentifier, TArrayView<const float> Times, TArrayView<AttributeType> KeyValues, bool bShouldTransact = true)
-	{		
+	{
 		TArray<const void*> KeyValuePtrs;
 		Algo::Transform(KeyValues, KeyValuePtrs, [](const AttributeType& Value)
 		{
 			return (const void*)&Value;
 		});
 
-		return SetAttributeKeys_Internal(AttributeIdentifier, Times, MakeArrayView(KeyValuePtrs), AttributeType::StaticStruct(), bShouldTransact);
+		return SetAttributeKeys(AttributeIdentifier, Times, MakeArrayView(KeyValuePtrs), AttributeType::StaticStruct(), bShouldTransact);
 	}
 
 	/**
@@ -567,8 +563,8 @@ public:
 	*
 	* @return	Whether or not the attribute key was successfully removed
 	*/
-	UFUNCTION(BlueprintCallable, Category = CurveData)
-	bool RemoveAttributeKey(const FAnimationAttributeIdentifier& AttributeIdentifier, float Time, bool bShouldTransact = true);
+	UFUNCTION(BlueprintCallable, Category = AttributeData)
+	virtual bool RemoveAttributeKey(const FAnimationAttributeIdentifier& AttributeIdentifier, float Time, bool bShouldTransact = true) = 0;
 
 	/**
 	* Duplicated the attribute (curve) with the identifier. Broadcasts a EAnimDataModelNotifyType::AttributeAdded notify if successful.
@@ -579,8 +575,8 @@ public:
 	*
 	* @return	Whether or not the attribute was successfully duplicated
 	*/
-	UFUNCTION(BlueprintCallable, Category = CurveData)
-	bool DuplicateAttribute(const FAnimationAttributeIdentifier& AttributeIdentifier, const FAnimationAttributeIdentifier& NewAttributeIdentifier, bool bShouldTransact = true);
+	UFUNCTION(BlueprintCallable, Category = AttributeData)
+	virtual bool DuplicateAttribute(const FAnimationAttributeIdentifier& AttributeIdentifier, const FAnimationAttributeIdentifier& NewAttributeIdentifier, bool bShouldTransact = true) = 0;
 protected:
 	/** Functionality used by FOpenBracketAction and FCloseBracketAction to broadcast their equivalent notifies without actually opening a bracket. */
 	virtual void NotifyBracketOpen() = 0;
diff --git a/Engine/Source/Runtime/Engine/Classes/Animation/AnimSequenceBase.h b/Engine/Source/Runtime/Engine/Classes/Animation/AnimSequenceBase.h
--- a/Engine/Source/Runtime/Engine/Classes/Animation/AnimSequenceBase.h
+++ b/Engine/Source/Runtime/Engine/Classes/Animation/AnimSequenceBase.h
@@ -29,7 +29,6 @@ enum ETypeAdvanceAnim
 struct FAnimationPoseData;
 struct FAnimDataModelNotifPayload;
 class UAnimDataModel;
-class UAnimDataController;
 enum class EAnimDataModelNotifyType : uint8;
 
 UCLASS(abstract, BlueprintType)
diff --git a/Engine/Source/Runtime/Engine/Classes/Animation/BuiltInAttributeTypes.h b/Engine/Source/Runtime/Engine/Classes/Animation/BuiltInAttributeTypes.h
--- a/Engine/Source/Runtime/Engine/Classes/Animation/BuiltInAttributeTypes.h
+++ b/Engine/Source/Runtime/Engine/Classes/Animation/BuiltInAttributeTypes.h
@@ -7,7 +7,6 @@
 #include "Animation/AttributeTraits.h"
 #include "AnimationRuntime.h"
 
-#include "Animation/AnimData/AnimDataController.h"
 #include "Algo/Transform.h"
 
 #include "BuiltInAttributeTypes.generated.h"
@@ -173,8 +172,8 @@ namespace UE
 		{
 			const FAnimationAttributeIdentifier Identifier = UAnimationAttributeIdentifierExtensions::CreateAttributeIdentifier(AnimSequenceBase, AttributeName, BoneName, AttributeType::StaticStruct());
 
-			UAnimDataController* Controller = AnimSequenceBase->GetController();
-			if (Controller->AddAttribute(Identifier))
+			IAnimationDataController& Controller = AnimSequenceBase->GetController();
+			if (Controller.AddAttribute(Identifier))
 			{
 				TArray<AttributeType> AttributeValues; 
 				Algo::Transform(Values, AttributeValues, [](const ValueType& Value)
@@ -184,7 +183,7 @@ namespace UE
 					return Attribute;
 				});
 
-				return Controller->SetTypedAttributeKeys<AttributeType>(Identifier, Keys, MakeArrayView(AttributeValues));
+				return Controller.SetTypedAttributeKeys<AttributeType>(Identifier, Keys, MakeArrayView(AttributeValues));
 			}
 
 			return false;
diff --git a/Engine/Source/Runtime/Engine/Private/Animation/AnimComposite.cpp b/Engine/Source/Runtime/Engine/Private/Animation/AnimComposite.cpp
--- a/Engine/Source/Runtime/Engine/Private/Animation/AnimComposite.cpp
+++ b/Engine/Source/Runtime/Engine/Private/Animation/AnimComposite.cpp
@@ -6,7 +6,7 @@
 
 #include "Animation/AnimComposite.h"
 #include "Animation/AnimationPoseData.h"
-#include "Animation/CustomAttributesRuntime.h"
+#include "Animation/AttributesRuntime.h"
 
 UAnimComposite::UAnimComposite(const FObjectInitializer& ObjectInitializer)
 	: Super(ObjectInitializer)
diff --git a/Engine/Source/Runtime/Engine/Private/Animation/AnimSequence.cpp b/Engine/Source/Runtime/Engine/Private/Animation/AnimSequence.cpp
--- a/Engine/Source/Runtime/Engine/Private/Animation/AnimSequence.cpp
+++ b/Engine/Source/Runtime/Engine/Private/Animation/AnimSequence.cpp
@@ -4701,7 +4701,7 @@ void UAnimSequence::RemoveCustomAttribute(const FName& BoneName, const FName& At
 
 	if (IdentifiersToRemove.Num())
 	{
-		UAnimDataController::FScopedBracket Bracket(Controller, LOCTEXT("RemoveCustomAttributeDeprecated", "Removing Custom Attributes by Bone and Name"));
+		IAnimationDataController::FScopedBracket Bracket(Controller, LOCTEXT("RemoveCustomAttributeDeprecated", "Removing Custom Attributes by Bone and Name"));
 		for (const FAnimationAttributeIdentifier& AttributeIdentifier : IdentifiersToRemove)
 		{
 			Controller->RemoveAttribute(AttributeIdentifier);
@@ -4814,7 +4814,7 @@ void UAnimSequence::MoveAttributesToModel()
 	PRAGMA_DISABLE_DEPRECATION_WARNINGS
 	if (TargetSkeleton && PerBoneCustomAttributeData.Num())
 	{
-		UAnimDataController::FScopedBracket Bracket(Controller, LOCTEXT("MoveAttributesToModel", "Moving legacy Custom Attributes to Model"));
+		IAnimationDataController::FScopedBracket Bracket(Controller, LOCTEXT("MoveAttributesToModel", "Moving legacy Custom Attributes to Model"));
 
 		for (const FCustomAttributePerBoneData& PerBoneData : PerBoneCustomAttributeData)
 		{
