Fix VT lightmap invalidation issue when r.SupportLowQualityLightmaps == 0

diff --git a/Engine/Source/Runtime/Engine/Private/LightMap.cpp b/Engine/Source/Runtime/Engine/Private/LightMap.cpp
--- a/Engine/Source/Runtime/Engine/Private/LightMap.cpp
+++ b/Engine/Source/Runtime/Engine/Private/LightMap.cpp
@@ -2924,7 +2924,10 @@ UTexture2D* FLightMap2D::GetAOMaterialMaskTexture() const
 bool FLightMap2D::IsVirtualTextureValid() const
 {
 #if WITH_EDITOR
-	if (VirtualTextures[0] && VirtualTextures[1])
+	static const auto CVarSupportLowQualityLightmap = IConsoleManager::Get().FindTConsoleVariableDataInt(TEXT("r.SupportLowQualityLightmaps"));
+	const bool bAllowLowQualityLightMaps = (!CVarSupportLowQualityLightmap) || (CVarSupportLowQualityLightmap->GetValueOnAnyThread() != 0);
+
+	if (VirtualTextures[0] && (!bAllowLowQualityLightMaps || bAllowLowQualityLightMaps && VirtualTextures[1]))
 #else
 	if (VirtualTextures[0] || VirtualTextures[1])
 #endif
@@ -3466,11 +3469,11 @@ IAllocatedVirtualTexture* FLightmapResourceCluster::AcquireAllocatedVT() const
 	const ULightMapVirtualTexture2D* VirtualTexture = Input.LightMapVirtualTextures[LightmapIndex];
 	
 #if WITH_EDITOR
-	// Compilation is still pending, this function will be called back once compilation finishes.
+		// Compilation is still pending, this function will be called back once compilation finishes.
 	if (VirtualTexture && VirtualTexture->IsCompiling())
-	{
-		return nullptr;
-	}
+		{
+			return nullptr;
+		}
 #endif
 
 	if (!AllocatedVT && VirtualTexture && VirtualTexture->Resource)
