Fix CIS

diff --git a/Engine/Source/Runtime/Engine/Private/LightMap.cpp b/Engine/Source/Runtime/Engine/Private/LightMap.cpp
--- a/Engine/Source/Runtime/Engine/Private/LightMap.cpp
+++ b/Engine/Source/Runtime/Engine/Private/LightMap.cpp
@@ -2927,7 +2927,7 @@ bool FLightMap2D::IsVirtualTextureValid() const
 	static const auto CVarSupportLowQualityLightmap = IConsoleManager::Get().FindTConsoleVariableDataInt(TEXT("r.SupportLowQualityLightmaps"));
 	const bool bAllowLowQualityLightMaps = (!CVarSupportLowQualityLightmap) || (CVarSupportLowQualityLightmap->GetValueOnAnyThread() != 0);
 
-	if (VirtualTextures[0] && (!bAllowLowQualityLightMaps || bAllowLowQualityLightMaps && VirtualTextures[1]))
+	if (VirtualTextures[0] && (!bAllowLowQualityLightMaps || (bAllowLowQualityLightMaps && VirtualTextures[1])))
 #else
 	if (VirtualTextures[0] || VirtualTextures[1])
 #endif
