AudioStreamCaching: De-escalating Warning log to Display when a chunk is requested from the DDC but not available.  This edge-case is handled gracefully and the DDC will be manually populated if appropriate.
# okforgithub public

diff --git a/Engine/Source/Runtime/Engine/Private/AudioDerivedData.cpp b/Engine/Source/Runtime/Engine/Private/AudioDerivedData.cpp
--- a/Engine/Source/Runtime/Engine/Private/AudioDerivedData.cpp
+++ b/Engine/Source/Runtime/Engine/Private/AudioDerivedData.cpp
@@ -804,7 +804,7 @@ int32 FStreamedAudioPlatformData::GetChunkFromDDC(int32 ChunkIndex, uint8** OutC
 {
 	if (Chunks.Num() == 0)
 	{
-		UE_LOG(LogAudioDerivedData, Warning, TEXT("No streamed audio chunks found!"));
+		UE_LOG(LogAudioDerivedData, Display, TEXT("No streamed audio chunks found!"));
 		return 0;
 	}
 
