SoundWaveProxy: Fix for "no zeroth chunk data" ensure in editor.

diff --git a/Engine/Source/Runtime/Engine/Private/SoundWave.cpp b/Engine/Source/Runtime/Engine/Private/SoundWave.cpp
--- a/Engine/Source/Runtime/Engine/Private/SoundWave.cpp
+++ b/Engine/Source/Runtime/Engine/Private/SoundWave.cpp
@@ -2458,6 +2458,10 @@ FSoundWaveProxyPtr USoundWave::CreateSoundWaveProxy()
 
 TUniquePtr<Audio::IProxyData> USoundWave::CreateNewProxyData(const Audio::FProxyDataInitParams& InitParams)
 {
+#if WITH_EDITORONLY_DATA
+	EnsureZerothChunkIsLoaded();
+#endif // #if WITH_EDITORONLY_DATA
+
 	return MakeUnique<FSoundWaveProxy>(this);
 }
 
