Adding parsing on import of wave cue points and labels

diff --git a/Engine/Source/Editor/AudioEditor/Private/Factories/SoundFactory.cpp b/Engine/Source/Editor/AudioEditor/Private/Factories/SoundFactory.cpp
--- a/Engine/Source/Editor/AudioEditor/Private/Factories/SoundFactory.cpp
+++ b/Engine/Source/Editor/AudioEditor/Private/Factories/SoundFactory.cpp
@@ -517,6 +517,19 @@ UObject* USoundFactory::CreateObject
 		// Store the current file path and timestamp for re-import purposes
 		Sound->AssetImportData->Update(CurrentFilename);
 
+		// Setup the cue points
+		Sound->CuePoints.Reset(WaveInfo.WaveCues.Num());
+
+		for (FWaveCue& WaveCue : WaveInfo.WaveCues)
+		{
+			FSoundWaveCuePoint NewCuePoint;
+			NewCuePoint.CuePointID = (int32)WaveCue.CuePointID;
+			NewCuePoint.FrameLength = (int32)WaveCue.SampleLength;
+			NewCuePoint.FramePosition = (int32)WaveCue.Position;
+			NewCuePoint.Label = WaveCue.Label;
+			Sound->CuePoints.Add(NewCuePoint);
+		}
+
 		// Compressed data is now out of date.
 		const bool bRebuildStreamingChunks = FPlatformCompressionUtilities::IsCurrentPlatformUsingStreamCaching();
 		Sound->InvalidateCompressedData(true /* bFreeResources */, bRebuildStreamingChunks);
diff --git a/Engine/Source/Runtime/Engine/Classes/Sound/SoundWave.h b/Engine/Source/Runtime/Engine/Classes/Sound/SoundWave.h
--- a/Engine/Source/Runtime/Engine/Classes/Sound/SoundWave.h
+++ b/Engine/Source/Runtime/Engine/Classes/Sound/SoundWave.h
@@ -277,6 +277,29 @@ enum class ESoundWaveFFTSize : uint8
 	VeryLarge_2048,
 };
 
+// Struct defining a cue point in a sound wave asset
+USTRUCT()
+struct FSoundWaveCuePoint
+{
+	GENERATED_USTRUCT_BODY()
+
+	// Unique identifier for the wave cue point
+	UPROPERTY(Category = Info, VisibleAnywhere)
+	int32 CuePointID = 0;
+
+	// The label for the cue point
+	UPROPERTY(Category = Info, VisibleAnywhere)
+	FString Label;
+
+	// The frame position of the cue point
+	UPROPERTY(Category = Info, VisibleAnywhere)
+	int32 FramePosition = 0;
+
+	// The frame length of the cue point (non-zero if it's a region)
+	UPROPERTY(Category = Info, VisibleAnywhere)
+	int32 FrameLength = 0;
+};
+
 struct ISoundWaveClient
 {
 	ISoundWaveClient() {}
@@ -313,6 +336,8 @@ public:
 	uint32 GetNumChannels() const { return NumChannels; }
 	uint32 GetNumChunks() const;
 	uint32 GetSizeOfChunk(uint32 ChunkIndex) const;
+	int32 GetNumFrames() const { return NumFrames; }
+	float GetDuration() const { return Duration; }
 
 	void ReleaseCompressedAudio();
 
@@ -353,6 +378,8 @@ private:
 	float SampleRate;
 	uint32 NumChannels;
 	uint32 NumChunks;
+	float Duration;
+	int32 NumFrames;
 	FAudioChunkHandle FirstChunk;
 	ESoundWaveLoadingBehavior LoadingBehavior;
 
@@ -378,35 +405,35 @@ class ENGINE_API USoundWave : public USoundBase, public IAudioProxyDataFactory
 public:
 	/** Platform agnostic compression quality. 1..100 with 1 being best compression and 100 being best quality. */
 	UPROPERTY(EditAnywhere, Category = "Format|Quality", meta = (DisplayName = "Compression", ClampMin = "1", ClampMax = "100"), AssetRegistrySearchable)
-		int32 CompressionQuality;
+	int32 CompressionQuality;
 
 	/** Priority of this sound when streaming (lower priority streams may not always play) */
 	UPROPERTY(EditAnywhere, Category = "Playback|Streaming", meta = (ClampMin = 0))
-		int32 StreamingPriority;
+	int32 StreamingPriority;
 
 	/** Quality of sample rate conversion for platforms that opt into resampling during cook. The sample rate for each enumeration is definable per platform in platform target settings. */
 	UPROPERTY(EditAnywhere, Category = "Format|Quality")
-		ESoundwaveSampleRateSettings SampleRateQuality;
+	ESoundwaveSampleRateSettings SampleRateQuality;
 
 	/** Type of buffer this wave uses. Set once on load */
 	TEnumAsByte<EDecompressionType> DecompressionType;
 
 	UPROPERTY(EditAnywhere, Category = Sound, meta = (DisplayName = "Group"))
-		TEnumAsByte<ESoundGroup> SoundGroup;
+	TEnumAsByte<ESoundGroup> SoundGroup;
 
 	/** If set, when played directly (not through a sound cue) the wave will be played looping. */
 	UPROPERTY(EditAnywhere, Category = Sound, AssetRegistrySearchable)
-		uint8 bLooping : 1;
+	uint8 bLooping : 1;
 
 	/** Whether this sound can be streamed to avoid increased memory usage. If using Stream Caching, use Loading Behavior instead to control memory usage. */
 	UPROPERTY(EditAnywhere, Category = "Playback|Streaming", meta = (DisplayName = "Force Streaming"))
-		uint8 bStreaming : 1;
+	uint8 bStreaming : 1;
 
 	/** Whether this sound supports seeking. This requires recooking with a codec which supports seekability and streaming. */
 	UPROPERTY(EditAnywhere, Category = "Playback|Streaming", meta = (DisplayName = "Seekable", EditCondition = "bStreaming"))
-		uint8 bSeekableStreaming : 1;
+	uint8 bSeekableStreaming : 1;
 
-// Loading behavior members are lazily initialized in const getters
+	// Loading behavior members are lazily initialized in const getters
 	/** Specifies how and when compressed audio data is loaded for asset if stream caching is enabled. */
 	UPROPERTY(EditAnywhere, Category = "Loading", meta = (DisplayName = "Loading Behavior Override"))
 	mutable ESoundWaveLoadingBehavior LoadingBehavior;
@@ -431,26 +458,26 @@ public:
 
 	/** If set to true if this sound is considered to contain mature/adult content. */
 	UPROPERTY(EditAnywhere, BlueprintReadOnly, Category = Subtitles, AssetRegistrySearchable)
-		uint8 bMature : 1;
+	uint8 bMature : 1;
 
 	/** If set to true will disable automatic generation of line breaks - use if the subtitles have been split manually. */
 	UPROPERTY(EditAnywhere, BlueprintReadOnly, Category = Subtitles)
-		uint8 bManualWordWrap : 1;
+	uint8 bManualWordWrap : 1;
 
 	/** If set to true the subtitles display as a sequence of single lines as opposed to multiline. */
 	UPROPERTY(EditAnywhere, BlueprintReadOnly, Category = Subtitles)
-		uint8 bSingleLine : 1;
+	uint8 bSingleLine : 1;
 
 #if WITH_EDITORONLY_DATA
 	UPROPERTY()
-		uint8 bVirtualizeWhenSilent_DEPRECATED : 1;
+	uint8 bVirtualizeWhenSilent_DEPRECATED : 1;
 #endif // WITH_EDITORONLY_DATA
 
 	/** Whether or not this source is ambisonics file format. If set, sound always uses the
 	  * 'Master Ambisonics Submix' as set in the 'Audio' category of Project Settings'
 	  * and ignores submix if provided locally or in the referenced SoundClass. */
 	UPROPERTY(EditAnywhere, Category = Format)
-		uint8 bIsAmbisonics : 1;
+	uint8 bIsAmbisonics : 1;
 
 	/** Whether this SoundWave was decompressed from OGG. */
 	uint8 bDecompressedFromOgg : 1;
@@ -494,50 +521,50 @@ public:
 #if WITH_EDITORONLY_DATA
 	/** Specify a sound to use for the baked analysis. Will default to this USoundWave if not set. */
 	UPROPERTY(EditAnywhere, Category = "Analysis")
-		TObjectPtr<USoundWave> OverrideSoundToUseForAnalysis;
+	TObjectPtr<USoundWave> OverrideSoundToUseForAnalysis;
 
 	/**
 		Whether or not we should treat the sound wave used for analysis (this or the override) as looping while performing analysis.
 		A looping sound may include the end of the file for inclusion in analysis for envelope and FFT analysis.
 	*/
 	UPROPERTY(EditAnywhere, Category = "Analysis")
-		uint8 TreatFileAsLoopingForAnalysis : 1;
+	uint8 TreatFileAsLoopingForAnalysis : 1;
 
 	/** Whether or not to enable cook-time baked FFT analysis. */
 	UPROPERTY(EditAnywhere, Category = "Analysis|FFT")
-		uint8 bEnableBakedFFTAnalysis : 1;
+	uint8 bEnableBakedFFTAnalysis : 1;
 
 	/** Whether or not to enable cook-time amplitude envelope analysis. */
 	UPROPERTY(EditAnywhere, Category = "Analysis|Envelope")
-		uint8 bEnableAmplitudeEnvelopeAnalysis : 1;
+	uint8 bEnableAmplitudeEnvelopeAnalysis : 1;
 
 	/** The FFT window size to use for fft analysis. */
 	UPROPERTY(EditAnywhere, Category = "Analysis|FFT", meta = (EditCondition = "bEnableBakedFFTAnalysis"))
-		ESoundWaveFFTSize FFTSize;
+	ESoundWaveFFTSize FFTSize;
 
 	/** How many audio frames analyze at a time. */
 	UPROPERTY(EditAnywhere, Category = "Analysis|FFT", meta = (EditCondition = "bEnableBakedFFTAnalysis", ClampMin = "512", UIMin = "512"))
-		int32 FFTAnalysisFrameSize;
+	int32 FFTAnalysisFrameSize;
 
 	/** Attack time in milliseconds of the spectral envelope follower. */
 	UPROPERTY(EditAnywhere, Category = "Analysis|FFT", meta = (EditCondition = "bEnableBakedFFTAnalysis", ClampMin = "0", UIMin = "0"))
-		int32 FFTAnalysisAttackTime;
+	int32 FFTAnalysisAttackTime;
 
 	/** Release time in milliseconds of the spectral envelope follower. */
 	UPROPERTY(EditAnywhere, Category = "Analysis|FFT", meta = (EditCondition = "bEnableBakedFFTAnalysis", ClampMin = "0", UIMin = "0"))
-		int32 FFTAnalysisReleaseTime;
+	int32 FFTAnalysisReleaseTime;
 
 	/** How many audio frames to average a new envelope value. Larger values use less memory for audio envelope data but will result in lower envelope accuracy. */
 	UPROPERTY(EditAnywhere, Category = "Analysis|Envelope", meta = (EditCondition = "bEnableAmplitudeEnvelopeAnalysis", ClampMin = "512", UIMin = "512"))
-		int32 EnvelopeFollowerFrameSize;
+	int32 EnvelopeFollowerFrameSize;
 
 	/** The attack time in milliseconds. Describes how quickly the envelope analyzer responds to increasing amplitudes. */
 	UPROPERTY(EditAnywhere, Category = "Analysis|Envelope", meta = (EditCondition = "bEnableAmplitudeEnvelopeAnalysis", ClampMin = "0", UIMin = "0"))
-		int32 EnvelopeFollowerAttackTime;
+	int32 EnvelopeFollowerAttackTime;
 
 	/** The release time in milliseconds. Describes how quickly the envelope analyzer responds to decreasing amplitudes. */
 	UPROPERTY(EditAnywhere, Category = "Analysis|Envelope", meta = (EditCondition = "bEnableAmplitudeEnvelopeAnalysis", ClampMin = "0", UIMin = "0"))
-		int32 EnvelopeFollowerReleaseTime;
+	int32 EnvelopeFollowerReleaseTime;
 #endif // WITH_EDITORONLY_DATA
 
 	/** Modulation Settings */
@@ -546,15 +573,15 @@ public:
 
 	/** The frequencies (in hz) to analyze when doing baked FFT analysis. */
 	UPROPERTY(EditAnywhere, Category = "Analysis|FFT", meta = (EditCondition = "bEnableBakedFFTAnalysis"))
-		TArray<float> FrequenciesToAnalyze;
+	TArray<float> FrequenciesToAnalyze;
 
 	/** The cooked spectral time data. */
 	UPROPERTY()
-		TArray<FSoundWaveSpectralTimeData> CookedSpectralTimeData;
+	TArray<FSoundWaveSpectralTimeData> CookedSpectralTimeData;
 
 	/** The cooked cooked envelope data. */
 	UPROPERTY()
-		TArray<FSoundWaveEnvelopeTimeData> CookedEnvelopeTimeData;
+	TArray<FSoundWaveEnvelopeTimeData> CookedEnvelopeTimeData;
 
 	/** Helper function to get interpolated cooked FFT data for a given time value. */
 	bool GetInterpolatedCookedFFTDataForTime(float InTime, uint32& InOutLastIndex, TArray<FSoundWaveSpectralData>& OutData, bool bLoop);
@@ -589,7 +616,7 @@ public:
 
 	/** Use this to override how much audio data is loaded when this USoundWave is loaded. */
 	UPROPERTY(EditAnywhere, AdvancedDisplay, Category = "Loading")
-		int32 InitialChunkSize;
+	int32 InitialChunkSize;
 
 private:
 
@@ -617,32 +644,36 @@ public:
 
 	/** A localized version of the text that is actually spoken phonetically in the audio. */
 	UPROPERTY(EditAnywhere, BlueprintReadOnly, Category = Subtitles)
-		FString SpokenText;
+	FString SpokenText;
 
 	/** The priority of the subtitle. */
 	UPROPERTY(EditAnywhere, BlueprintReadOnly, Category = Subtitles)
-		float SubtitlePriority;
+	float SubtitlePriority;
 
 	/** Playback volume of sound 0 to 1 - Default is 1.0. */
 	UPROPERTY(Category = Sound, meta = (ClampMin = "0.0"), EditAnywhere)
-		float Volume;
+	float Volume;
 
 	/** Playback pitch for sound. */
 	UPROPERTY(Category = Sound, meta = (ClampMin = "0.125", ClampMax = "4.0"), EditAnywhere)
-		float Pitch;
+	float Pitch;
 
 	/** Number of channels of multichannel data; 1 or 2 for regular mono and stereo files */
 	UPROPERTY(Category = Info, AssetRegistrySearchable, VisibleAnywhere)
-		int32 NumChannels;
+	int32 NumChannels;
+
+	/** Cue point data */
+	UPROPERTY(Category = Info, VisibleAnywhere)
+	TArray<FSoundWaveCuePoint> CuePoints;
 
 #if WITH_EDITORONLY_DATA
 	/** Offsets into the bulk data for the source wav data */
 	UPROPERTY()
-		TArray<int32> ChannelOffsets;
+	TArray<int32> ChannelOffsets;
 
 	/** Sizes of the bulk data for the source wav data */
 	UPROPERTY()
-		TArray<int32> ChannelSizes;
+	TArray<int32> ChannelSizes;
 
 #endif // WITH_EDITORONLY_DATA
 
@@ -650,7 +681,7 @@ protected:
 
 	/** Cached sample rate for displaying in the tools */
 	UPROPERTY(Category = Info, AssetRegistrySearchable, VisibleAnywhere)
-		int32 SampleRate;
+	int32 SampleRate;
 
 public:
 
@@ -668,23 +699,24 @@ public:
 	 * as the contents of the subtitle is commonly identical to what is spoken.
 	 */
 	UPROPERTY(EditAnywhere, BlueprintReadOnly, Category = Subtitles)
-		TArray<struct FSubtitleCue> Subtitles;
+	TArray<struct FSubtitleCue> Subtitles;
 
 #if WITH_EDITORONLY_DATA
 	/** Provides contextual information for the sound to the translator. */
 	UPROPERTY(EditAnywhere, Category = Subtitles)
-		FString Comment;
+	FString Comment;
 
 #endif // WITH_EDITORONLY_DATA
 
 #if WITH_EDITORONLY_DATA
 	UPROPERTY()
-		FString SourceFilePath_DEPRECATED;
+	FString SourceFilePath_DEPRECATED;
+	
 	UPROPERTY()
-		FString SourceFileTimestamp_DEPRECATED;
+	FString SourceFileTimestamp_DEPRECATED;
 
 	UPROPERTY(VisibleAnywhere, Instanced, Category = ImportSettings)
-		TObjectPtr<class UAssetImportData> AssetImportData;
+	TObjectPtr<class UAssetImportData> AssetImportData;
 
 #endif // WITH_EDITORONLY_DATA
 
@@ -692,11 +724,11 @@ protected:
 
 	/** Curves associated with this sound wave */
 	UPROPERTY(EditAnywhere, Category = SoundWave, AdvancedDisplay)
-		TObjectPtr<class UCurveTable> Curves;
+	TObjectPtr<class UCurveTable> Curves;
 
 	/** Hold a reference to our internal curve so we can switch back to it if we want to */
 	UPROPERTY()
-		TObjectPtr<class UCurveTable> InternalCurves;
+	TObjectPtr<class UCurveTable> InternalCurves;
 
 	/** Potential strong handle to the first chunk of audio data. Can be released via ReleaseCompressedAudioData. */
 	FAudioChunkHandle FirstChunk;
diff --git a/Engine/Source/Runtime/Engine/Private/Audio.cpp b/Engine/Source/Runtime/Engine/Private/Audio.cpp
--- a/Engine/Source/Runtime/Engine/Private/Audio.cpp
+++ b/Engine/Source/Runtime/Engine/Private/Audio.cpp
@@ -1110,12 +1110,6 @@ FString FWaveInstance::GetName() const
 }
 
 
-/*-----------------------------------------------------------------------------
-	WaveModInfo implementation - downsampling of wave files.
------------------------------------------------------------------------------*/
-
-//  Macros to convert 4 bytes to a Riff-style ID uint32.
-//  Todo: make these endian independent !!!
 
 #define UE_MAKEFOURCC(ch0, ch1, ch2, ch3)\
 	((uint32)(uint8)(ch0) | ((uint32)(uint8)(ch1) << 8) |\
@@ -1155,6 +1149,62 @@ struct FRiffFormatChunk
 	uint16   cbSize;            // The count in bytes of the size of extra information (after cbSize).
 };
 
+// ChunkID: 'cue ' 
+// A cue chunk specifies one or more sample offsets which are often used to mark noteworthy sections of audio. For example, 
+// the beginning and end of a verse in a song may have cue points to make them easier to find. The cue chunk is optional and 
+// if included, a single cue chunk should specify all cue points for the "WAVE" chunk. 
+// No more than one cue chunk is allowed in a "WAVE" chunk.
+struct FRiffCueChunk
+{
+	uint32 ChunkID;			// 'cue '
+	uint32 ChunkDataSize;	// Depends on the number of cue points
+	uint32 NumCuePoints;	// Number of cue points in the list
+};
+
+struct FRiffCuePointChunk
+{
+	uint32 CueID;			// Unique ID value for the cue point
+	uint32 Position;		// Play order position
+	uint32 DataChunkID;		// RIFF ID of corresponding data chunk
+	uint32 ChunkStart;		// Byte offset of data chunk
+	uint32 BlockStart;		// Byte offset of sample of first channel
+	uint32 SampleOffset;	// Byte offset to sample byte of first channel
+};
+
+struct FRiffListChunk
+{
+	uint32 ChunkID;			// 'list'
+	uint32 ChunkDataSize;	// Depends on contained text
+	uint32 TypeID;			// always 'adtl'
+};
+
+struct FRiffLabelChunk
+{
+	uint32 ChunkID;			// 'labl'
+	uint32 ChunkDataSize;	// depends on contained text
+	uint32 CuePointID;		// Cue Point ID associated with the label
+};
+
+struct FRiffNoteChunk
+{
+	uint32 ChunkID;			// 'note'
+	uint32 ChunkDataSize;	// Depends on size of contained text
+	uint32 CuePointID;		// ID associated with the note
+};
+
+struct FRiffLabeledTextChunk
+{
+	uint32 ChunkID;			// 'ltxt'
+	uint32 ChunkDataSize;	// Depends on contained text
+	uint32 CuePointID;		// ID associated with the labeled text
+	uint32 SampleLength;	// Defines how many samples from the cue point the region or section spans
+	uint32 PurposeID;		// Unused: Specifies what the text is used for. 'scrp' means script text. 'capt' means close-caption.
+	uint16 Country;			// Unused 
+	uint16 Language;		// Unused
+	uint16 Dialect;			// Unused 
+	uint16 CodePage;		// Unused
+};
+
 // FExtendedFormatChunk subformat GUID.
 struct FSubformatGUID
 {
@@ -1204,48 +1254,49 @@ bool FWaveModInfo::ReadWaveInfo( const uint8* WaveData, int32 WaveDataSize, FStr
 {
 	FRiffFormatChunk* FmtChunk;
 	FExtendedFormatChunk* FmtChunkEx = nullptr;
-	FRiffWaveHeaderChunk* RiffHdr = (FRiffWaveHeaderChunk* )WaveData;
+	FRiffWaveHeaderChunk* RiffHdr = (FRiffWaveHeaderChunk*)WaveData;
 	WaveDataEnd = WaveData + WaveDataSize;
 
-	if( WaveDataSize == 0 )
+	if (WaveDataSize == 0)
 	{
-		return( false );
+		return false;
 	}
 
 	// Verify we've got a real 'WAVE' header.
 #if PLATFORM_LITTLE_ENDIAN
-	if( RiffHdr->wID != UE_mmioFOURCC( 'W','A','V','E' ) )
+	if (RiffHdr->wID != UE_mmioFOURCC('W','A','V','E'))
 	{
 		if (ErrorReason) *ErrorReason = TEXT("Invalid WAVE file.");
-		return( false );
+		return false;
 	}
 #else
-	if( ( RiffHdr->wID != ( UE_mmioFOURCC( 'W','A','V','E' ) ) ) &&
-	     ( RiffHdr->wID != ( UE_mmioFOURCC( 'E','V','A','W' ) ) ) )
+	if ((RiffHdr->wID != (UE_mmioFOURCC('W','A','V','E'))) &&
+	     (RiffHdr->wID != (UE_mmioFOURCC('E','V','A','W'))))
 	{
 		ErrorReason = TEXT("Invalid WAVE file.")
-		return( false );
+			return false;
 	}
 
-	bool AlreadySwapped = ( RiffHdr->wID == ( UE_mmioFOURCC('W','A','V','E' ) ) );
-	if( !AlreadySwapped )
+	bool AlreadySwapped = (RiffHdr->wID == (UE_mmioFOURCC('W','A','V','E')));
+	if (!AlreadySwapped)
 	{
-		RiffHdr->rID = INTEL_ORDER32( RiffHdr->rID );
-		RiffHdr->ChunkLen = INTEL_ORDER32( RiffHdr->ChunkLen );
-		RiffHdr->wID = INTEL_ORDER32( RiffHdr->wID );
+		RiffHdr->rID = INTEL_ORDER32(RiffHdr->rID);
+		RiffHdr->ChunkLen = INTEL_ORDER32(RiffHdr->ChunkLen);
+		RiffHdr->wID = INTEL_ORDER32(RiffHdr->wID);
 	}
 #endif
 
-	FRiffChunkOld* RiffChunk = ( FRiffChunkOld* )&WaveData[3 * 4];
+	FRiffChunkOld* RiffChunkStart = (FRiffChunkOld*)&WaveData[3 * 4];
+	FRiffChunkOld* RiffChunk = RiffChunkStart;
 	pMasterSize = &RiffHdr->ChunkLen;
 
 	// Look for the 'fmt ' chunk.
-	while( ( ( ( uint8* )RiffChunk + 8 ) < WaveDataEnd ) && ( INTEL_ORDER32( RiffChunk->ChunkID ) != UE_mmioFOURCC( 'f','m','t',' ' ) ) )
+	while ((((uint8*)RiffChunk + 8) < WaveDataEnd) && (INTEL_ORDER32(RiffChunk->ChunkID) != UE_mmioFOURCC('f','m','t',' ')))
 	{
-		RiffChunk = ( FRiffChunkOld* )( ( uint8* )RiffChunk + Pad16Bit( INTEL_ORDER32( RiffChunk->ChunkLen ) ) + 8 );
+		RiffChunk = (FRiffChunkOld*)((uint8*)RiffChunk + Pad16Bit(INTEL_ORDER32(RiffChunk->ChunkLen)) + 8);
 	}
 
-	if( INTEL_ORDER32( RiffChunk->ChunkID ) != UE_mmioFOURCC( 'f','m','t',' ' ) )
+	if (INTEL_ORDER32(RiffChunk->ChunkID) != UE_mmioFOURCC('f','m','t',' '))
 	{
 		#if !PLATFORM_LITTLE_ENDIAN  // swap them back just in case.
 			if( !AlreadySwapped )
@@ -1259,16 +1310,16 @@ bool FWaveModInfo::ReadWaveInfo( const uint8* WaveData, int32 WaveDataSize, FStr
 		return( false );
 	}
 
-	FmtChunk = ( FRiffFormatChunk* )( ( uint8* )RiffChunk + 8 );
+	FmtChunk = (FRiffFormatChunk*)((uint8*)RiffChunk + 8);
 #if !PLATFORM_LITTLE_ENDIAN
-	if( !AlreadySwapped )
+	if (!AlreadySwapped)
 	{
-		FmtChunk->wFormatTag = INTEL_ORDER16( FmtChunk->wFormatTag );
-		FmtChunk->nChannels = INTEL_ORDER16( FmtChunk->nChannels );
-		FmtChunk->nSamplesPerSec = INTEL_ORDER32( FmtChunk->nSamplesPerSec );
-		FmtChunk->nAvgBytesPerSec = INTEL_ORDER32( FmtChunk->nAvgBytesPerSec );
-		FmtChunk->nBlockAlign = INTEL_ORDER16( FmtChunk->nBlockAlign );
-		FmtChunk->wBitsPerSample = INTEL_ORDER16( FmtChunk->wBitsPerSample );
+		FmtChunk->wFormatTag = INTEL_ORDER16(FmtChunk->wFormatTag);
+		FmtChunk->nChannels = INTEL_ORDER16(FmtChunk->nChannels);
+		FmtChunk->nSamplesPerSec = INTEL_ORDER32(FmtChunk->nSamplesPerSec);
+		FmtChunk->nAvgBytesPerSec = INTEL_ORDER32(FmtChunk->nAvgBytesPerSec);
+		FmtChunk->nBlockAlign = INTEL_ORDER16(FmtChunk->nBlockAlign);
+		FmtChunk->wBitsPerSample = INTEL_ORDER16(FmtChunk->wBitsPerSample);
 	}
 #endif
 	pBitsPerSample = &FmtChunk->wBitsPerSample;
@@ -1278,7 +1329,7 @@ bool FWaveModInfo::ReadWaveInfo( const uint8* WaveData, int32 WaveDataSize, FStr
 	pChannels = &FmtChunk->nChannels;
 	pFormatTag = &FmtChunk->wFormatTag;
 
-	if(OutFormatHeader != NULL)
+	if (OutFormatHeader != NULL)
 	{
 		*OutFormatHeader = FmtChunk;
 	}
@@ -1333,29 +1384,29 @@ bool FWaveModInfo::ReadWaveInfo( const uint8* WaveData, int32 WaveDataSize, FStr
 		pFormatTag = reinterpret_cast<uint16*>(&FmtChunkEx->SubFormat.Data1);
 	}
 
-	// re-initalize the RiffChunk pointer
-	RiffChunk = ( FRiffChunkOld* )&WaveData[3 * 4];
+	// re-initialize the RiffChunk pointer
+	RiffChunk = RiffChunkStart;
 
 	// Look for the 'data' chunk.
-	while( ( ( ( uint8* )RiffChunk + 8 ) <= WaveDataEnd ) && ( INTEL_ORDER32( RiffChunk->ChunkID ) != UE_mmioFOURCC( 'd','a','t','a' ) ) )
+	while ((((uint8*)RiffChunk + 8 ) <= WaveDataEnd ) && (INTEL_ORDER32(RiffChunk->ChunkID) != UE_mmioFOURCC('d','a','t','a')))
 	{
-		RiffChunk = ( FRiffChunkOld* )( ( uint8* )RiffChunk + Pad16Bit( INTEL_ORDER32( RiffChunk->ChunkLen ) ) + 8 );
+		RiffChunk = (FRiffChunkOld*)((uint8*)RiffChunk + Pad16Bit(INTEL_ORDER32(RiffChunk->ChunkLen)) + 8);
 	}
 
-	if( INTEL_ORDER32( RiffChunk->ChunkID ) != UE_mmioFOURCC( 'd','a','t','a' ) )
+	if (INTEL_ORDER32(RiffChunk->ChunkID) != UE_mmioFOURCC('d','a','t','a'))
 	{
 		#if !PLATFORM_LITTLE_ENDIAN  // swap them back just in case.
-			if( !AlreadySwapped )
+			if (!AlreadySwapped)
 			{
-				RiffHdr->rID = INTEL_ORDER32( RiffHdr->rID );
-				RiffHdr->ChunkLen = INTEL_ORDER32( RiffHdr->ChunkLen );
-				RiffHdr->wID = INTEL_ORDER32( RiffHdr->wID );
-				FmtChunk->wFormatTag = INTEL_ORDER16( FmtChunk->wFormatTag );
-				FmtChunk->nChannels = INTEL_ORDER16( FmtChunk->nChannels );
-				FmtChunk->nSamplesPerSec = INTEL_ORDER32( FmtChunk->nSamplesPerSec );
-				FmtChunk->nAvgBytesPerSec = INTEL_ORDER32( FmtChunk->nAvgBytesPerSec );
-				FmtChunk->nBlockAlign = INTEL_ORDER16( FmtChunk->nBlockAlign );
-				FmtChunk->wBitsPerSample = INTEL_ORDER16( FmtChunk->wBitsPerSample );
+				RiffHdr->rID = INTEL_ORDER32(RiffHdr->rID);
+				RiffHdr->ChunkLen = INTEL_ORDER32(RiffHdr->ChunkLen);
+				RiffHdr->wID = INTEL_ORDER32(RiffHdr->wID);
+				FmtChunk->wFormatTag = INTEL_ORDER16(FmtChunk->wFormatTag);
+				FmtChunk->nChannels = INTEL_ORDER16(FmtChunk->nChannels);
+				FmtChunk->nSamplesPerSec = INTEL_ORDER32(FmtChunk->nSamplesPerSec);
+				FmtChunk->nAvgBytesPerSec = INTEL_ORDER32(FmtChunk->nAvgBytesPerSec);
+				FmtChunk->nBlockAlign = INTEL_ORDER16(FmtChunk->nBlockAlign);
+				FmtChunk->wBitsPerSample = INTEL_ORDER16(FmtChunk->wBitsPerSample);
 				if (FmtChunkEx != nullptr)
 				{
 					FmtChunkEx->Samples.wValidBitsPerSample = INTEL_ORDER16(FmtChunkEx->Samples.wValidBitsPerSample);
@@ -1367,67 +1418,67 @@ bool FWaveModInfo::ReadWaveInfo( const uint8* WaveData, int32 WaveDataSize, FStr
 			}
 		#endif
 		if (ErrorReason) *ErrorReason = TEXT("Invalid WAVE file.");
-		return( false );
+		return false;
 	}
 
 #if !PLATFORM_LITTLE_ENDIAN  // swap them back just in case.
-	if( AlreadySwapped ) // swap back into Intel order for chunk search...
+	if (AlreadySwapped) // swap back into Intel order for chunk search...
 	{
-		RiffChunk->ChunkLen = INTEL_ORDER32( RiffChunk->ChunkLen );
+		RiffChunk->ChunkLen = INTEL_ORDER32(RiffChunk->ChunkLen);
 	}
 #endif
 
-	SampleDataStart = ( uint8* )RiffChunk + 8;
+	SampleDataStart = (uint8*)RiffChunk + 8;
 	pWaveDataSize = &RiffChunk->ChunkLen;
 	SampleDataSize = INTEL_ORDER32( RiffChunk->ChunkLen );
 	SampleDataEnd = SampleDataStart + SampleDataSize;
 
-	if( !InHeaderDataOnly && ( uint8* )SampleDataEnd > ( uint8* )WaveDataEnd )
+	if (!InHeaderDataOnly && (uint8*)SampleDataEnd > (uint8*)WaveDataEnd)
 	{
 		UE_LOG(LogAudio, Warning, TEXT( "Wave data chunk is too big!" ) );
 
 		// Fix it up by clamping data chunk.
-		SampleDataEnd = ( uint8* )WaveDataEnd;
+		SampleDataEnd = (uint8*)WaveDataEnd;
 		SampleDataSize = SampleDataEnd - SampleDataStart;
 		RiffChunk->ChunkLen = INTEL_ORDER32( SampleDataSize );
 	}
 
-	if (   *pFormatTag != 0x0001 // WAVE_FORMAT_PCM
+	if (*pFormatTag != 0x0001 // WAVE_FORMAT_PCM
 		&& *pFormatTag != 0x0002 // WAVE_FORMAT_ADPCM
 		&& *pFormatTag != 0x0011) // WAVE_FORMAT_DVI_ADPCM
 	{
 		ReportImportFailure();
 		if (ErrorReason) *ErrorReason = TEXT("Unsupported wave file format.  Only PCM, ADPCM, and DVI ADPCM can be imported.");
-		return( false );
+		return false;
 	}
 
-	if(!InHeaderDataOnly)
+	if (!InHeaderDataOnly)
 	{
-		if( ( uint8* )SampleDataEnd > ( uint8* )WaveDataEnd )
+		if ((uint8*)SampleDataEnd > (uint8*)WaveDataEnd)
 		{
-			UE_LOG(LogAudio, Warning, TEXT( "Wave data chunk is too big!" ) );
+			UE_LOG(LogAudio, Warning, TEXT("Wave data chunk is too big!" ));
 
 			// Fix it up by clamping data chunk.
-			SampleDataEnd = ( uint8* )WaveDataEnd;
+			SampleDataEnd = (uint8*)WaveDataEnd;
 			SampleDataSize = SampleDataEnd - SampleDataStart;
-			RiffChunk->ChunkLen = INTEL_ORDER32( SampleDataSize );
+			RiffChunk->ChunkLen = INTEL_ORDER32(SampleDataSize);
 		}
 
 		NewDataSize = SampleDataSize;
 
 		#if !PLATFORM_LITTLE_ENDIAN
-		if( !AlreadySwapped )
+		if (!AlreadySwapped)
 		{
-			if( FmtChunk->wBitsPerSample == 16 )
+			if (FmtChunk->wBitsPerSample == 16)
 			{
-				for( uint16* i = ( uint16* )SampleDataStart; i < ( uint16* )SampleDataEnd; i++ )
+				for (uint16* i = (uint16*)SampleDataStart; i < (uint16*)SampleDataEnd; i++)
 				{
-					*i = INTEL_ORDER16( *i );
+					*i = INTEL_ORDER16(*i);
 				}
 			}
-			else if( FmtChunk->wBitsPerSample == 32 )
+			else if (FmtChunk->wBitsPerSample == 32)
 			{
-				for( uint32* i = ( uint32* )SampleDataStart; i < ( uint32* )SampleDataEnd; i++ )
+				for (uint32* i = (uint32*)SampleDataStart; i < (uint32*)SampleDataEnd; i++)
 				{
 					*i = INTEL_ORDER32( *i );
 				}
@@ -1438,10 +1489,95 @@ bool FWaveModInfo::ReadWaveInfo( const uint8* WaveData, int32 WaveDataSize, FStr
 
 	// Couldn't byte swap this before, since it'd throw off the chunk search.
 #if !PLATFORM_LITTLE_ENDIAN
-	*pWaveDataSize = INTEL_ORDER32( *pWaveDataSize );
+	*pWaveDataSize = INTEL_ORDER32(*pWaveDataSize);
 #endif
 
-	return( true );
+	// Look for the cue chunks
+	RiffChunk = RiffChunkStart;
+
+	// Look for the 'cue ' chunk.
+	while ((((uint8*)RiffChunk + 8) <= WaveDataEnd) && (INTEL_ORDER32(RiffChunk->ChunkID) != UE_mmioFOURCC('c', 'u', 'e', ' ')))
+	{
+		RiffChunk = (FRiffChunkOld*)((uint8*)RiffChunk + Pad16Bit(INTEL_ORDER32(RiffChunk->ChunkLen)) + 8);
+	}
+
+	// Cue chunks are optional
+	if (INTEL_ORDER32(RiffChunk->ChunkID) == UE_mmioFOURCC('c', 'u', 'e', ' '))
+	{
+		FRiffCueChunk* CueChunk = (FRiffCueChunk*)((uint8*)RiffChunk);
+
+		WaveCues.Reset(CueChunk->NumCuePoints);
+
+		// Get to the first cue point chunk
+		FRiffCuePointChunk* CuePointChunks = (FRiffCuePointChunk*)((uint8*)RiffChunk + sizeof(FRiffCueChunk));
+		for (uint32 CuePointId = 0; CuePointId < CueChunk->NumCuePoints; ++CuePointId)
+		{		
+			FWaveCue NewCue;
+			NewCue.CuePointID = CuePointChunks[CuePointId].CueID;
+			NewCue.Position = CuePointChunks[CuePointId].Position;
+			WaveCues.Add(NewCue);
+		
+		}
+		
+		// Now look for labl chunk for labels for cues
+		RiffChunk = RiffChunkStart;
+
+		// Look for the 'list' chunk.
+		while ((((uint8*)RiffChunk + 8) <= WaveDataEnd) && (INTEL_ORDER32(RiffChunk->ChunkID) != UE_mmioFOURCC('L', 'I', 'S', 'T')))
+		{
+			RiffChunk = (FRiffChunkOld*)((uint8*)RiffChunk + Pad16Bit(INTEL_ORDER32(RiffChunk->ChunkLen)) + 8);
+		}
+
+		// Label chunks are also optional
+		if (INTEL_ORDER32(RiffChunk->ChunkID) == UE_mmioFOURCC('L', 'I', 'S', 'T'))
+		{
+			FRiffListChunk* ListChunk = (FRiffListChunk*)((uint8*)RiffChunk);
+			
+			RiffChunk = (FRiffChunkOld*)(ListChunk + 1);
+			const uint8* ListChunkEnd = (uint8*)RiffChunk + ListChunk->ChunkDataSize + 4;
+			while (((uint8*)RiffChunk + 8) <= ListChunkEnd)
+			{
+				// Labeled Text Chunk
+				// This information is often displayed in marked regions of a waveform in digital audio editors.
+				if (INTEL_ORDER32(RiffChunk->ChunkID) == UE_mmioFOURCC('l', 't', 'x', 't'))
+				{
+					FRiffLabeledTextChunk* LabeledTextChunk = (FRiffLabeledTextChunk*)((uint8*)RiffChunk);
+					for (FWaveCue& WaveCue : WaveCues)
+					{
+						if (WaveCue.CuePointID == LabeledTextChunk->CuePointID)
+						{
+							WaveCue.SampleLength = LabeledTextChunk->SampleLength;
+							break;
+						}
+					}
+				}
+				// Labeled Chunk
+				else if (INTEL_ORDER32(RiffChunk->ChunkID) == UE_mmioFOURCC('l', 'a', 'b', 'l'))
+				{
+					FRiffLabelChunk* LabelChunk = (FRiffLabelChunk*)((uint8*)RiffChunk);
+
+					for (FWaveCue& WaveCue : WaveCues)
+					{
+						if (WaveCue.CuePointID == LabelChunk->CuePointID)
+						{
+							char* LabelCharText = (char*)(LabelChunk + 1);
+							while (*LabelCharText != '\0')
+							{
+								WaveCue.Label.AppendChar(*LabelCharText);
+								++LabelCharText;
+							}
+							WaveCue.Label.AppendChar('\0');
+							break;
+						}
+					}
+				}
+				RiffChunk = (FRiffChunkOld*)((uint8*)RiffChunk + Pad16Bit(INTEL_ORDER32(RiffChunk->ChunkLen)) + 8);
+			}
+		}
+	}
+
+
+	return true;
 }
 
 bool FWaveModInfo::ReadWaveHeader(const uint8* RawWaveData, int32 Size, int32 Offset )
diff --git a/Engine/Source/Runtime/Engine/Public/Audio.h b/Engine/Source/Runtime/Engine/Public/Audio.h
--- a/Engine/Source/Runtime/Engine/Public/Audio.h
+++ b/Engine/Source/Runtime/Engine/Public/Audio.h
@@ -807,9 +807,18 @@ public:
 #endif //ENABLE_AUDIO_DEBUG
 };
 
-/*-----------------------------------------------------------------------------
-	FWaveModInfo. 
------------------------------------------------------------------------------*/
+// Data representing a cue in a wave file
+struct FWaveCue
+{
+	// Unique identifying gvalue for the cue
+	uint32 CuePointID = 0;
+	// Sample offset associated with the cue point
+	uint32 Position = 0;
+	// Cue label
+	FString Label;
+	// If this is a region, it will have a duration (sample length)
+	uint32 SampleLength = 0;
+};
 
 //
 // Structure for in-memory interpretation and modification of WAVE sound structures.
@@ -835,6 +844,9 @@ public:
 
 	uint32  NewDataSize;
 
+	// List of cues parsed from the wave file
+	TArray<FWaveCue> WaveCues;
+
 	// Constructor.
 	FWaveModInfo()
 	{
