GetLinkerUE4Version renamed to GetLinkerUEVersion

diff --git a/Engine/Source/Editor/AnimGraph/Private/AnimGraphNode_SpringBone.cpp b/Engine/Source/Editor/AnimGraph/Private/AnimGraphNode_SpringBone.cpp
--- a/Engine/Source/Editor/AnimGraph/Private/AnimGraphNode_SpringBone.cpp
+++ b/Engine/Source/Editor/AnimGraph/Private/AnimGraphNode_SpringBone.cpp
@@ -16,7 +16,7 @@ void UAnimGraphNode_SpringBone::PostLoad()
 {
 	Super::PostLoad();
 
-	if (GetLinkerUE4Version() < VER_UE4_REPLACE_SPRING_NOZ_PROPERTY)
+	if (GetLinkerUEVersion() < VER_UE4_REPLACE_SPRING_NOZ_PROPERTY)
 	{
 		//Fix up z translation from old NoZSpring property
 		Node.bTranslateZ = !Node.bNoZSpring_DEPRECATED;
diff --git a/Engine/Source/Editor/AnimGraph/Private/AnimStateTransitionNode.cpp b/Engine/Source/Editor/AnimGraph/Private/AnimStateTransitionNode.cpp
--- a/Engine/Source/Editor/AnimGraph/Private/AnimStateTransitionNode.cpp
+++ b/Engine/Source/Editor/AnimGraph/Private/AnimStateTransitionNode.cpp
@@ -112,7 +112,7 @@ void UAnimStateTransitionNode::PostLoad()
 		FAnimStateTransitionNodeSharedCrossfadeHelper().MakeSureGuidExists(this);
 	}
 
-	if(GetLinkerUE4Version() < VER_UE4_ADDED_NON_LINEAR_TRANSITION_BLENDS)
+	if(GetLinkerUEVersion() < VER_UE4_ADDED_NON_LINEAR_TRANSITION_BLENDS)
 	{
 		switch(CrossfadeMode_DEPRECATED)
 		{
diff --git a/Engine/Source/Editor/BlueprintGraph/Private/K2Node.cpp b/Engine/Source/Editor/BlueprintGraph/Private/K2Node.cpp
--- a/Engine/Source/Editor/BlueprintGraph/Private/K2Node.cpp
+++ b/Engine/Source/Editor/BlueprintGraph/Private/K2Node.cpp
@@ -148,12 +148,12 @@ bool UK2Node::HasNonEditorOnlyReferences() const
 
 void UK2Node::FixupPinDefaultValues()
 {
-	const int32 LinkerUE4Version = GetLinkerUE4Version();
+	const int32 LinkerUEVersion = GetLinkerUEVersion();
 	const int32 LinkerFrameworkVersion = GetLinkerCustomVersion(FFrameworkObjectVersion::GUID);
 	const UEdGraphSchema_K2* K2Schema = GetDefault<UEdGraphSchema_K2>();
 
 	// Swap "new" default error tolerance value with zero on vector/rotator equality nodes, in order to preserve current behavior in existing blueprints.
-	if(LinkerUE4Version < VER_UE4_BP_MATH_VECTOR_EQUALITY_USES_EPSILON)
+	if(LinkerUEVersion < VER_UE4_BP_MATH_VECTOR_EQUALITY_USES_EPSILON)
 	{
 		static const FString VectorsEqualFunctionEpsilonPinName = TEXT("KismetMathLibrary.EqualEqual_VectorVector.ErrorTolerance");
 		static const FString VectorsNotEqualFunctionEpsilonPinName = TEXT("KismetMathLibrary.NotEqual_VectorVector.ErrorTolerance");
diff --git a/Engine/Source/Editor/BlueprintGraph/Private/K2Node_InputAction.cpp b/Engine/Source/Editor/BlueprintGraph/Private/K2Node_InputAction.cpp
--- a/Engine/Source/Editor/BlueprintGraph/Private/K2Node_InputAction.cpp
+++ b/Engine/Source/Editor/BlueprintGraph/Private/K2Node_InputAction.cpp
@@ -30,7 +30,7 @@ void UK2Node_InputAction::PostLoad()
 {
 	Super::PostLoad();
 
-	if (GetLinkerUE4Version() < VER_UE4_BLUEPRINT_INPUT_BINDING_OVERRIDES)
+	if (GetLinkerUEVersion() < VER_UE4_BLUEPRINT_INPUT_BINDING_OVERRIDES)
 	{
 		// Don't change existing behaviors
 		bOverrideParentBinding = false;
diff --git a/Engine/Source/Editor/BlueprintGraph/Private/K2Node_InputAxisEvent.cpp b/Engine/Source/Editor/BlueprintGraph/Private/K2Node_InputAxisEvent.cpp
--- a/Engine/Source/Editor/BlueprintGraph/Private/K2Node_InputAxisEvent.cpp
+++ b/Engine/Source/Editor/BlueprintGraph/Private/K2Node_InputAxisEvent.cpp
@@ -41,7 +41,7 @@ void UK2Node_InputAxisEvent::PostLoad()
 {
 	Super::PostLoad();
 
-	if (GetLinkerUE4Version() < VER_UE4_BLUEPRINT_INPUT_BINDING_OVERRIDES)
+	if (GetLinkerUEVersion() < VER_UE4_BLUEPRINT_INPUT_BINDING_OVERRIDES)
 	{
 		// Don't change existing behaviors
 		bOverrideParentBinding = false;
diff --git a/Engine/Source/Editor/BlueprintGraph/Private/K2Node_InputKey.cpp b/Engine/Source/Editor/BlueprintGraph/Private/K2Node_InputKey.cpp
--- a/Engine/Source/Editor/BlueprintGraph/Private/K2Node_InputKey.cpp
+++ b/Engine/Source/Editor/BlueprintGraph/Private/K2Node_InputKey.cpp
@@ -37,7 +37,7 @@ void UK2Node_InputKey::PostLoad()
 {
 	Super::PostLoad();
 
-	if (GetLinkerUE4Version() < VER_UE4_BLUEPRINT_INPUT_BINDING_OVERRIDES)
+	if (GetLinkerUEVersion() < VER_UE4_BLUEPRINT_INPUT_BINDING_OVERRIDES)
 	{
 		// Don't change existing behaviors
 		bOverrideParentBinding = false;
diff --git a/Engine/Source/Editor/BlueprintGraph/Private/K2Node_InputTouch.cpp b/Engine/Source/Editor/BlueprintGraph/Private/K2Node_InputTouch.cpp
--- a/Engine/Source/Editor/BlueprintGraph/Private/K2Node_InputTouch.cpp
+++ b/Engine/Source/Editor/BlueprintGraph/Private/K2Node_InputTouch.cpp
@@ -24,7 +24,7 @@ void UK2Node_InputTouch::PostLoad()
 {
 	Super::PostLoad();
 
-	if (GetLinkerUE4Version() < VER_UE4_BLUEPRINT_INPUT_BINDING_OVERRIDES)
+	if (GetLinkerUEVersion() < VER_UE4_BLUEPRINT_INPUT_BINDING_OVERRIDES)
 	{
 		// Don't change existing behaviors
 		bOverrideParentBinding = false;
diff --git a/Engine/Source/Editor/UMGEditor/Private/WidgetBlueprint.cpp b/Engine/Source/Editor/UMGEditor/Private/WidgetBlueprint.cpp
--- a/Engine/Source/Editor/UMGEditor/Private/WidgetBlueprint.cpp
+++ b/Engine/Source/Editor/UMGEditor/Private/WidgetBlueprint.cpp
@@ -858,7 +858,7 @@ void UWidgetBlueprint::PostLoad()
 		Widget->ConnectEditorData();
 	});
 
-	if( GetLinkerUE4Version() < VER_UE4_FIXUP_WIDGET_ANIMATION_CLASS )
+	if( GetLinkerUEVersion() < VER_UE4_FIXUP_WIDGET_ANIMATION_CLASS )
 	{
 		// Fixup widget animations.
 		for( auto& OldAnim : AnimationData_DEPRECATED )
@@ -881,7 +881,7 @@ void UWidgetBlueprint::PostLoad()
 		AnimationData_DEPRECATED.Empty();
 	}
 
-	if ( GetLinkerUE4Version() < VER_UE4_RENAME_WIDGET_VISIBILITY )
+	if ( GetLinkerUEVersion() < VER_UE4_RENAME_WIDGET_VISIBILITY )
 	{
 		static const FName Visiblity(TEXT("Visiblity"));
 		static const FName Visibility(TEXT("Visibility"));
diff --git a/Engine/Source/Runtime/CoreUObject/Private/UObject/ObjectBaseUtility.cpp b/Engine/Source/Runtime/CoreUObject/Private/UObject/ObjectBaseUtility.cpp
--- a/Engine/Source/Runtime/CoreUObject/Private/UObject/ObjectBaseUtility.cpp
+++ b/Engine/Source/Runtime/CoreUObject/Private/UObject/ObjectBaseUtility.cpp
@@ -16,7 +16,7 @@
  *			a) this object is a native only class, or
  *			b) this object's linker has been detached, in which case it is already fully loaded
  */
-int32 UObjectBaseUtility::GetLinkerUE4Version() const
+int32 UObjectBaseUtility::GetLinkerUEVersion() const
 {
 	FLinkerLoad* Loader = GetLinker();
 
@@ -90,7 +90,7 @@ int32 UObjectBaseUtility::GetLinkerCustomVersion(FGuid CustomVersionKey) const
  *			a) this object is a native only class, or
  *			b) this object's linker has been detached, in which case it is already fully loaded
  */
-int32 UObjectBaseUtility::GetLinkerLicenseeUE4Version() const
+int32 UObjectBaseUtility::GetLinkerLicenseeUEVersion() const
 {
 	FLinkerLoad* Loader = GetLinker();
 
diff --git a/Engine/Source/Runtime/CoreUObject/Public/UObject/UObjectBaseUtility.h b/Engine/Source/Runtime/CoreUObject/Public/UObject/UObjectBaseUtility.h
--- a/Engine/Source/Runtime/CoreUObject/Public/UObject/UObjectBaseUtility.h
+++ b/Engine/Source/Runtime/CoreUObject/Public/UObject/UObjectBaseUtility.h
@@ -616,7 +616,12 @@ public:
 	 *			a) this object is a native only class, or
 	 *			b) this object's linker has been detached, in which case it is already fully loaded
 	 */
-	int32 GetLinkerUE4Version() const;
+	int32 GetLinkerUEVersion() const;
+
+
+	UE_DEPRECATED(5.0, "Use GetLinkerUEVersion instead")
+	inline int32 GetLinkerUE4Version() const { return GetLinkerUEVersion(); }
+
 
 	/**
 	 * Returns the licensee version of the linker for this object.
@@ -627,7 +632,11 @@ public:
 	 *			a) this object is a native only class, or
 	 *			b) this object's linker has been detached, in which case it is already fully loaded
 	 */
-	int32 GetLinkerLicenseeUE4Version() const;
+	int32 GetLinkerLicenseeUEVersion() const;
+
+	UE_DEPRECATED(5.0, "Use GetLinkerLicenseeUEVersion instead")
+	inline int32 GetLinkerLicenseeUE4Version() const { return GetLinkerLicenseeUEVersion(); }
+
 
 	/**
 	 * Returns the custom version of the linker for this object corresponding to the given custom version key.
diff --git a/Engine/Source/Runtime/Engine/Private/Actor.cpp b/Engine/Source/Runtime/Engine/Private/Actor.cpp
--- a/Engine/Source/Runtime/Engine/Private/Actor.cpp
+++ b/Engine/Source/Runtime/Engine/Private/Actor.cpp
@@ -844,7 +844,7 @@ void AActor::PostLoad()
 		Owner->Children.Add(this);
 	}
 
-	if (GetLinkerUE4Version() < VER_UE4_PRIVATE_REMOTE_ROLE)
+	if (GetLinkerUEVersion() < VER_UE4_PRIVATE_REMOTE_ROLE)
 	{
 		bReplicates = (RemoteRole != ROLE_None);
 	}
diff --git a/Engine/Source/Runtime/Engine/Private/Animation/AnimSequence.cpp b/Engine/Source/Runtime/Engine/Private/Animation/AnimSequence.cpp
--- a/Engine/Source/Runtime/Engine/Private/Animation/AnimSequence.cpp
+++ b/Engine/Source/Runtime/Engine/Private/Animation/AnimSequence.cpp
@@ -841,7 +841,7 @@ void UAnimSequence::PostLoad()
 
 	// I have to do this first thing in here
 	// so that remove all NaNs before even being read
-	if(GetLinkerUE4Version() < VER_UE4_ANIMATION_REMOVE_NANS)
+	if(GetLinkerUEVersion() < VER_UE4_ANIMATION_REMOVE_NANS)
 	{
 		RemoveNaNTracks();
 	}
diff --git a/Engine/Source/Runtime/Engine/Private/Animation/AnimSequenceBase.cpp b/Engine/Source/Runtime/Engine/Private/Animation/AnimSequenceBase.cpp
--- a/Engine/Source/Runtime/Engine/Private/Animation/AnimSequenceBase.cpp
+++ b/Engine/Source/Runtime/Engine/Private/Animation/AnimSequenceBase.cpp
@@ -110,7 +110,7 @@ void UAnimSequenceBase::PostLoad()
 	// Convert Notifies to new data
 	if( GIsEditor && Notifies.Num() > 0 )
 	{
-		if(GetLinkerUE4Version() < VER_UE4_CLEAR_NOTIFY_TRIGGERS)
+		if(GetLinkerUEVersion() < VER_UE4_CLEAR_NOTIFY_TRIGGERS)
 		{
 			for(FAnimNotifyEvent Notify : Notifies)
 			{
diff --git a/Engine/Source/Runtime/Engine/Private/Animation/SkeletalMeshActor.cpp b/Engine/Source/Runtime/Engine/Private/Animation/SkeletalMeshActor.cpp
--- a/Engine/Source/Runtime/Engine/Private/Animation/SkeletalMeshActor.cpp
+++ b/Engine/Source/Runtime/Engine/Private/Animation/SkeletalMeshActor.cpp
@@ -285,7 +285,7 @@ void ASkeletalMeshActor::LoadedFromAnotherClass(const FName& OldClassName)
 {
 	Super::LoadedFromAnotherClass(OldClassName);
 
-	if (GetLinkerUE4Version() < VER_UE4_REMOVE_SKELETALPHYSICSACTOR)
+	if (GetLinkerUEVersion() < VER_UE4_REMOVE_SKELETALPHYSICSACTOR)
 	{
 		static FName SkeletalPhysicsActor_NAME(TEXT("SkeletalPhysicsActor"));
 		static FName KAsset_NAME(TEXT("KAsset"));
diff --git a/Engine/Source/Runtime/Engine/Private/BlockingVolume.cpp b/Engine/Source/Runtime/Engine/Private/BlockingVolume.cpp
--- a/Engine/Source/Runtime/Engine/Private/BlockingVolume.cpp
+++ b/Engine/Source/Runtime/Engine/Private/BlockingVolume.cpp
@@ -20,7 +20,7 @@ void ABlockingVolume::LoadedFromAnotherClass(const FName& OldClassName)
 {
 	Super::LoadedFromAnotherClass(OldClassName);
 
-	if(GetLinkerUE4Version() < VER_UE4_REMOVE_DYNAMIC_VOLUME_CLASSES)
+	if(GetLinkerUEVersion() < VER_UE4_REMOVE_DYNAMIC_VOLUME_CLASSES)
 	{
 		static FName DynamicBlockingVolume_NAME(TEXT("DynamicBlockingVolume"));
 
diff --git a/Engine/Source/Runtime/Engine/Private/BlueprintGeneratedClass.cpp b/Engine/Source/Runtime/Engine/Private/BlueprintGeneratedClass.cpp
--- a/Engine/Source/Runtime/Engine/Private/BlueprintGeneratedClass.cpp
+++ b/Engine/Source/Runtime/Engine/Private/BlueprintGeneratedClass.cpp
@@ -113,7 +113,7 @@ void UBlueprintGeneratedClass::PostLoad()
 			});
 		}
 
-		if (GetLinkerUE4Version() < VER_UE4_CLASS_NOTPLACEABLE_ADDED)
+		if (GetLinkerUEVersion() < VER_UE4_CLASS_NOTPLACEABLE_ADDED)
 		{
 			// Make sure the placeable flag is correct for all blueprints
 			UBlueprint* Blueprint = Cast<UBlueprint>(ClassGeneratedBy);
diff --git a/Engine/Source/Runtime/Engine/Private/Camera/CameraActor.cpp b/Engine/Source/Runtime/Engine/Private/Camera/CameraActor.cpp
--- a/Engine/Source/Runtime/Engine/Private/Camera/CameraActor.cpp
+++ b/Engine/Source/Runtime/Engine/Private/Camera/CameraActor.cpp
@@ -60,13 +60,13 @@ void ACameraActor::PostLoadSubobjects(FObjectInstancingGraph* OuterInstanceGraph
 
 	Super::PostLoadSubobjects(OuterInstanceGraph);
 	
-	if (GetLinkerUE4Version() < VER_UE4_CAMERA_ACTOR_USING_CAMERA_COMPONENT)
+	if (GetLinkerUEVersion() < VER_UE4_CAMERA_ACTOR_USING_CAMERA_COMPONENT)
 	{
 		CameraComponent->SetupAttachment(OldAttachParent, OldSocketName);
 		OldRoot->SetupAttachment(nullptr);
 	}
 
-	if (GetLinkerUE4Version() < VER_UE4_CAMERA_COMPONENT_ATTACH_TO_ROOT)
+	if (GetLinkerUEVersion() < VER_UE4_CAMERA_COMPONENT_ATTACH_TO_ROOT)
 	{
 		RootComponent = SceneComponent;
 		if (OldAttachParent != SceneComponent)
diff --git a/Engine/Source/Runtime/Engine/Private/Camera/CameraComponent.cpp b/Engine/Source/Runtime/Engine/Private/Camera/CameraComponent.cpp
--- a/Engine/Source/Runtime/Engine/Private/Camera/CameraComponent.cpp
+++ b/Engine/Source/Runtime/Engine/Private/Camera/CameraComponent.cpp
@@ -149,7 +149,7 @@ void UCameraComponent::PostLoad()
 {
 	Super::PostLoad();
 
-	const int32 LinkerUE4Ver = GetLinkerUE4Version();
+	const int32 LinkerUE4Ver = GetLinkerUEVersion();
 
 	if (LinkerUE4Ver < VER_UE4_RENAME_CAMERA_COMPONENT_VIEW_ROTATION)
 	{
diff --git a/Engine/Source/Runtime/Engine/Private/Components/ActorComponent.cpp b/Engine/Source/Runtime/Engine/Private/Components/ActorComponent.cpp
--- a/Engine/Source/Runtime/Engine/Private/Components/ActorComponent.cpp
+++ b/Engine/Source/Runtime/Engine/Private/Components/ActorComponent.cpp
@@ -370,7 +370,7 @@ void UActorComponent::PostLoad()
 	Super::PostLoad();
 	   
 #if WITH_EDITORONLY_DATA
-	if (GetLinkerUE4Version() < VER_UE4_ACTOR_COMPONENT_CREATION_METHOD)
+	if (GetLinkerUEVersion() < VER_UE4_ACTOR_COMPONENT_CREATION_METHOD)
 	{
 		if (IsTemplate())
 		{
@@ -410,7 +410,7 @@ void UActorComponent::PostLoad()
 
 	if (CreationMethod == EComponentCreationMethod::SimpleConstructionScript)
 	{
-		if ((GetLinkerUE4Version() < VER_UE4_TRACK_UCS_MODIFIED_PROPERTIES) && !HasAnyFlags(RF_ClassDefaultObject))
+		if ((GetLinkerUEVersion() < VER_UE4_TRACK_UCS_MODIFIED_PROPERTIES) && !HasAnyFlags(RF_ClassDefaultObject))
 		{
 			DetermineUCSModifiedProperties();
 		}
diff --git a/Engine/Source/Runtime/Engine/Private/Components/AudioComponent.cpp b/Engine/Source/Runtime/Engine/Private/Components/AudioComponent.cpp
--- a/Engine/Source/Runtime/Engine/Private/Components/AudioComponent.cpp
+++ b/Engine/Source/Runtime/Engine/Private/Components/AudioComponent.cpp
@@ -162,10 +162,10 @@ void UAudioComponent::Serialize(FArchive& Ar)
 void UAudioComponent::PostLoad()
 {
 #if WITH_EDITORONLY_DATA
-	const int32 LinkerUE4Version = GetLinkerUE4Version();
+	const int32 LinkerUEVersion = GetLinkerUEVersion();
 
 	// Translate the old HighFrequencyGainMultiplier value to the new LowPassFilterFrequency value
-	if (LinkerUE4Version < VER_UE4_USE_LOW_PASS_FILTER_FREQ)
+	if (LinkerUEVersion < VER_UE4_USE_LOW_PASS_FILTER_FREQ)
 	{
 		if (HighFrequencyGainMultiplier_DEPRECATED > 0.0f &&  HighFrequencyGainMultiplier_DEPRECATED < 1.0f)
 		{
diff --git a/Engine/Source/Runtime/Engine/Private/Components/BrushComponent.cpp b/Engine/Source/Runtime/Engine/Private/Components/BrushComponent.cpp
--- a/Engine/Source/Runtime/Engine/Private/Components/BrushComponent.cpp
+++ b/Engine/Source/Runtime/Engine/Private/Components/BrushComponent.cpp
@@ -506,7 +506,7 @@ void UBrushComponent::PostLoad()
 	Super::PostLoad();
 
 	// Stop existing BrushComponents from generating mirrored collision mesh
-	if ((GetLinkerUE4Version() < VER_UE4_NO_MIRROR_BRUSH_MODEL_COLLISION) && (BrushBodySetup != NULL))
+	if ((GetLinkerUEVersion() < VER_UE4_NO_MIRROR_BRUSH_MODEL_COLLISION) && (BrushBodySetup != NULL))
 	{
 		BrushBodySetup->bGenerateMirroredCollision = false;
 	}
diff --git a/Engine/Source/Runtime/Engine/Private/Components/CharacterMovementComponent.cpp b/Engine/Source/Runtime/Engine/Private/Components/CharacterMovementComponent.cpp
--- a/Engine/Source/Runtime/Engine/Private/Components/CharacterMovementComponent.cpp
+++ b/Engine/Source/Runtime/Engine/Private/Components/CharacterMovementComponent.cpp
@@ -575,7 +575,7 @@ void UCharacterMovementComponent::PostLoad()
 	Super::PostLoad();
 
 #if WITH_EDITORONLY_DATA
-	const int32 LinkerUE4Ver = GetLinkerUE4Version();
+	const int32 LinkerUE4Ver = GetLinkerUEVersion();
 
 	if (LinkerUE4Ver < VER_UE4_CHARACTER_MOVEMENT_DECELERATION)
 	{
diff --git a/Engine/Source/Runtime/Engine/Private/Components/LightComponent.cpp b/Engine/Source/Runtime/Engine/Private/Components/LightComponent.cpp
--- a/Engine/Source/Runtime/Engine/Private/Components/LightComponent.cpp
+++ b/Engine/Source/Runtime/Engine/Private/Components/LightComponent.cpp
@@ -581,7 +581,7 @@ void ULightComponent::PostLoad()
 	PreviewShadowMapChannel = INDEX_NONE;
 	Intensity = FMath::Max(0.0f, Intensity);
 
-	if (GetLinkerUE4Version() < VER_UE4_LIGHTCOMPONENT_USE_IES_TEXTURE_MULTIPLIER_ON_NON_IES_BRIGHTNESS)
+	if (GetLinkerUEVersion() < VER_UE4_LIGHTCOMPONENT_USE_IES_TEXTURE_MULTIPLIER_ON_NON_IES_BRIGHTNESS)
 	{
 		if(IESTexture)
 		{
diff --git a/Engine/Source/Runtime/Engine/Private/Components/ModelComponent.cpp b/Engine/Source/Runtime/Engine/Private/Components/ModelComponent.cpp
--- a/Engine/Source/Runtime/Engine/Private/Components/ModelComponent.cpp
+++ b/Engine/Source/Runtime/Engine/Private/Components/ModelComponent.cpp
@@ -300,7 +300,7 @@ void UModelComponent::PostLoad()
 	}
 
 	// Stop existing ModelComponents from generating mirrored collision mesh
-	if ((GetLinkerUE4Version() < VER_UE4_NO_MIRROR_BRUSH_MODEL_COLLISION) && (ModelBodySetup != NULL))
+	if ((GetLinkerUEVersion() < VER_UE4_NO_MIRROR_BRUSH_MODEL_COLLISION) && (ModelBodySetup != NULL))
 	{
 		ModelBodySetup->bGenerateMirroredCollision = false;
 	}
diff --git a/Engine/Source/Runtime/Engine/Private/Components/PrimitiveComponent.cpp b/Engine/Source/Runtime/Engine/Private/Components/PrimitiveComponent.cpp
--- a/Engine/Source/Runtime/Engine/Private/Components/PrimitiveComponent.cpp
+++ b/Engine/Source/Runtime/Engine/Private/Components/PrimitiveComponent.cpp
@@ -1207,7 +1207,7 @@ void UPrimitiveComponent::PostLoad()
 {
 	Super::PostLoad();
 	
-	int32 const UE4Version = GetLinkerUE4Version();
+	int32 const UE4Version = GetLinkerUEVersion();
 
 	// as temporary fix for the bug TTP 299926
 	// permanent fix is coming
diff --git a/Engine/Source/Runtime/Engine/Private/Components/ProjectileMovementComponent.cpp b/Engine/Source/Runtime/Engine/Private/Components/ProjectileMovementComponent.cpp
--- a/Engine/Source/Runtime/Engine/Private/Components/ProjectileMovementComponent.cpp
+++ b/Engine/Source/Runtime/Engine/Private/Components/ProjectileMovementComponent.cpp
@@ -57,7 +57,7 @@ void UProjectileMovementComponent::PostLoad()
 {
 	Super::PostLoad();
 
-	const int32 LinkerUE4Ver = GetLinkerUE4Version();
+	const int32 LinkerUE4Ver = GetLinkerUEVersion();
 
 	if (LinkerUE4Ver < VER_UE4_REFACTOR_PROJECTILE_MOVEMENT)
 	{
diff --git a/Engine/Source/Runtime/Engine/Private/Components/SkeletalMeshComponent.cpp b/Engine/Source/Runtime/Engine/Private/Components/SkeletalMeshComponent.cpp
--- a/Engine/Source/Runtime/Engine/Private/Components/SkeletalMeshComponent.cpp
+++ b/Engine/Source/Runtime/Engine/Private/Components/SkeletalMeshComponent.cpp
@@ -1068,7 +1068,7 @@ void USkeletalMeshComponent::LoadedFromAnotherClass(const FName& OldClassName)
 {
 	Super::LoadedFromAnotherClass(OldClassName);
 
-	if(GetLinkerUE4Version() < VER_UE4_REMOVE_SINGLENODEINSTANCE)
+	if(GetLinkerUEVersion() < VER_UE4_REMOVE_SINGLENODEINSTANCE)
 	{
 		static FName SingleAnimSkeletalComponent_NAME(TEXT("SingleAnimSkeletalComponent"));
 
diff --git a/Engine/Source/Runtime/Engine/Private/Components/StaticMeshComponent.cpp b/Engine/Source/Runtime/Engine/Private/Components/StaticMeshComponent.cpp
--- a/Engine/Source/Runtime/Engine/Private/Components/StaticMeshComponent.cpp
+++ b/Engine/Source/Runtime/Engine/Private/Components/StaticMeshComponent.cpp
@@ -1777,7 +1777,7 @@ void UStaticMeshComponent::PostLoad()
 #if WITH_EDITORONLY_DATA
 	// Remap the materials array if the static mesh materials may have been remapped to remove zero triangle sections.
 	// This will cause async static mesh compilation to stall but only if a fixup is actually required.
-	if (GetStaticMesh() && GetLinkerUE4Version() < VER_UE4_REMOVE_ZERO_TRIANGLE_SECTIONS && OverrideMaterials.Num())
+	if (GetStaticMesh() && GetLinkerUEVersion() < VER_UE4_REMOVE_ZERO_TRIANGLE_SECTIONS && OverrideMaterials.Num())
 	{
 		if (GetStaticMesh()->HasValidRenderData()
 			&& GetStaticMesh()->GetRenderData()->MaterialIndexToImportIndex.Num())
diff --git a/Engine/Source/Runtime/Engine/Private/Components/TextRenderComponent.cpp b/Engine/Source/Runtime/Engine/Private/Components/TextRenderComponent.cpp
--- a/Engine/Source/Runtime/Engine/Private/Components/TextRenderComponent.cpp
+++ b/Engine/Source/Runtime/Engine/Private/Components/TextRenderComponent.cpp
@@ -1313,7 +1313,7 @@ void UTextRenderComponent::PostLoad()
 {
 	// Try and fix up assets created before the vertical alignment fix was implemented. Because we didn't flag that
 	// fix with its own version, use the version number closest to that CL
-	if (GetLinkerUE4Version() < VER_UE4_PACKAGE_REQUIRES_LOCALIZATION_GATHER_FLAGGING)
+	if (GetLinkerUEVersion() < VER_UE4_PACKAGE_REQUIRES_LOCALIZATION_GATHER_FLAGGING)
 	{
 		const float Offset = CalculateVerticalAlignmentOffset(*Text.ToString(), Font, XScale, YScale, HorizSpacingAdjust, VertSpacingAdjust, VerticalAlignment);
 		const FTransform RelativeTransform = GetRelativeTransform();
@@ -1324,12 +1324,12 @@ void UTextRenderComponent::PostLoad()
 		SetRelativeTransform(CorrectionLeft * RelativeTransform * CorrectionRight);
 	}
 
-	if (GetLinkerUE4Version() < VER_UE4_ADD_TEXT_COMPONENT_VERTICAL_ALIGNMENT)
+	if (GetLinkerUEVersion() < VER_UE4_ADD_TEXT_COMPONENT_VERTICAL_ALIGNMENT)
 	{
 		VerticalAlignment = EVRTA_QuadTop;
 	}
 
-	if( GetLinkerUE4Version() < VER_UE4_TEXT_RENDER_COMPONENTS_WORLD_SPACE_SIZING )
+	if( GetLinkerUEVersion() < VER_UE4_TEXT_RENDER_COMPONENTS_WORLD_SPACE_SIZING )
 	{
 		if( Font )
 		{
diff --git a/Engine/Source/Runtime/Engine/Private/EdGraph/EdGraphNode.cpp b/Engine/Source/Runtime/Engine/Private/EdGraph/EdGraphNode.cpp
--- a/Engine/Source/Runtime/Engine/Private/EdGraph/EdGraphNode.cpp
+++ b/Engine/Source/Runtime/Engine/Private/EdGraph/EdGraphNode.cpp
@@ -666,7 +666,7 @@ void UEdGraphNode::PostLoad()
 	}
 
 	// Duplicating a Blueprint needs to have a new Node Guid generated, which was not occuring before this version
-	if(GetLinkerUE4Version() < VER_UE4_POST_DUPLICATE_NODE_GUID)
+	if(GetLinkerUEVersion() < VER_UE4_POST_DUPLICATE_NODE_GUID)
 	{
 		UE_LOG(LogBlueprint, Warning, TEXT("Node '%s' missing NodeGuid because of upgrade from old package version, this can cause deterministic cooking issues please resave package."), *GetPathName());
 
@@ -674,7 +674,7 @@ void UEdGraphNode::PostLoad()
 		CreateNewGuid();
 	}
 	// Moving to the new style comments requires conversion to preserve previous state
-	if(GetLinkerUE4Version() < VER_UE4_GRAPH_INTERACTIVE_COMMENTBUBBLES)
+	if(GetLinkerUEVersion() < VER_UE4_GRAPH_INTERACTIVE_COMMENTBUBBLES)
 	{
 		bCommentBubbleVisible = !NodeComment.IsEmpty();
 	}
diff --git a/Engine/Source/Runtime/Engine/Private/HierarchicalInstancedStaticMesh.cpp b/Engine/Source/Runtime/Engine/Private/HierarchicalInstancedStaticMesh.cpp
--- a/Engine/Source/Runtime/Engine/Private/HierarchicalInstancedStaticMesh.cpp
+++ b/Engine/Source/Runtime/Engine/Private/HierarchicalInstancedStaticMesh.cpp
@@ -2921,7 +2921,7 @@ bool UHierarchicalInstancedStaticMeshComponent::BuildTreeIfOutdated(bool Async,
 		|| NumBuiltInstances != PerInstanceSMData.Num() 
 		|| (GetStaticMesh() != nullptr && CacheMeshExtendedBounds != GetStaticMesh()->GetBounds())
 		|| UnbuiltInstanceBoundsList.Num() > 0
-		|| GetLinkerUE4Version() < VER_UE4_REBUILD_HIERARCHICAL_INSTANCE_TREES
+		|| GetLinkerUEVersion() < VER_UE4_REBUILD_HIERARCHICAL_INSTANCE_TREES
 		|| GetLinkerCustomVersion(FReleaseObjectVersion::GUID) < FReleaseObjectVersion::HISMCClusterTreeMigration)
 	{
 		if (!GetStaticMesh())
diff --git a/Engine/Source/Runtime/Engine/Private/LevelStreaming.cpp b/Engine/Source/Runtime/Engine/Private/LevelStreaming.cpp
--- a/Engine/Source/Runtime/Engine/Private/LevelStreaming.cpp
+++ b/Engine/Source/Runtime/Engine/Private/LevelStreaming.cpp
@@ -330,7 +330,7 @@ void ULevelStreaming::PostLoad()
 	}
 
 #if WITH_EDITOR
-	if (GetLinkerUE4Version() < VER_UE4_LEVEL_STREAMING_DRAW_COLOR_TYPE_CHANGE)
+	if (GetLinkerUEVersion() < VER_UE4_LEVEL_STREAMING_DRAW_COLOR_TYPE_CHANGE)
 	{
 		LevelColor = DrawColor_DEPRECATED;
 	}
diff --git a/Engine/Source/Runtime/Engine/Private/Light.cpp b/Engine/Source/Runtime/Engine/Private/Light.cpp
--- a/Engine/Source/Runtime/Engine/Private/Light.cpp
+++ b/Engine/Source/Runtime/Engine/Private/Light.cpp
@@ -192,7 +192,7 @@ void APointLight::LoadedFromAnotherClass(const FName& OldClassName)
 {
 	Super::LoadedFromAnotherClass(OldClassName);
 
-	if(GetLinkerUE4Version() < VER_UE4_REMOVE_LIGHT_MOBILITY_CLASSES)
+	if(GetLinkerUEVersion() < VER_UE4_REMOVE_LIGHT_MOBILITY_CLASSES)
 	{
 		static FName PointLightStatic_NAME(TEXT("PointLightStatic"));
 		static FName PointLightMovable_NAME(TEXT("PointLightMovable"));
@@ -282,7 +282,7 @@ void ADirectionalLight::LoadedFromAnotherClass(const FName& OldClassName)
 {
 	Super::LoadedFromAnotherClass(OldClassName);
 
-	if(GetLinkerUE4Version() < VER_UE4_REMOVE_LIGHT_MOBILITY_CLASSES)
+	if(GetLinkerUEVersion() < VER_UE4_REMOVE_LIGHT_MOBILITY_CLASSES)
 	{
 		static FName DirectionalLightStatic_NAME(TEXT("DirectionalLightStatic"));
 		static FName DirectionalLightMovable_NAME(TEXT("DirectionalLightMovable"));
diff --git a/Engine/Source/Runtime/Engine/Private/Materials/Material.cpp b/Engine/Source/Runtime/Engine/Private/Materials/Material.cpp
--- a/Engine/Source/Runtime/Engine/Private/Materials/Material.cpp
+++ b/Engine/Source/Runtime/Engine/Private/Materials/Material.cpp
@@ -4045,7 +4045,7 @@ void UMaterial::PostLoad()
 	LoadedMaterialResources.Empty();
 
 #if WITH_EDITORONLY_DATA
-	const int32 UE4Ver = GetLinkerUE4Version();
+	const int32 UE4Ver = GetLinkerUEVersion();
 	const int32 RenderObjVer = GetLinkerCustomVersion(FRenderingObjectVersion::GUID);
 
 	DoMaterialAttributeReorder(&DiffuseColor_DEPRECATED, UE4Ver, RenderObjVer);
@@ -4092,7 +4092,7 @@ void UMaterial::PostLoad()
 	}
 
 
-	if( GetLinkerUE4Version() < VER_UE4_REMOVED_MATERIAL_USED_WITH_UI_FLAG && bUsedWithUI_DEPRECATED == true )
+	if( GetLinkerUEVersion() < VER_UE4_REMOVED_MATERIAL_USED_WITH_UI_FLAG && bUsedWithUI_DEPRECATED == true )
 	{
 		MaterialDomain = MD_UI;
 	}
diff --git a/Engine/Source/Runtime/Engine/Private/Materials/MaterialExpressions.cpp b/Engine/Source/Runtime/Engine/Private/Materials/MaterialExpressions.cpp
--- a/Engine/Source/Runtime/Engine/Private/Materials/MaterialExpressions.cpp
+++ b/Engine/Source/Runtime/Engine/Private/Materials/MaterialExpressions.cpp
@@ -9131,7 +9131,7 @@ void UMaterialExpressionDynamicParameter::PostLoad()
 {
 	Super::PostLoad();
 
-	if (GetLinkerUE4Version() < VER_UE4_DYNAMIC_PARAMETER_DEFAULT_VALUE)
+	if (GetLinkerUEVersion() < VER_UE4_DYNAMIC_PARAMETER_DEFAULT_VALUE)
 	{
 		DefaultValue = FLinearColor::Black;//Old data should default to 0.0f;
 	}
@@ -9647,7 +9647,7 @@ void UMaterialExpressionSceneDepth::PostLoad()
 {
 	Super::PostLoad();
 
-	if(GetLinkerUE4Version() < VER_UE4_REFACTOR_MATERIAL_EXPRESSION_SCENECOLOR_AND_SCENEDEPTH_INPUTS)
+	if(GetLinkerUEVersion() < VER_UE4_REFACTOR_MATERIAL_EXPRESSION_SCENECOLOR_AND_SCENEDEPTH_INPUTS)
 	{
 		// Connect deprecated UV input to new expression input
 		InputMode = EMaterialSceneAttributeInputMode::Coordinates;
@@ -9805,7 +9805,7 @@ void UMaterialExpressionSceneColor::PostLoad()
 {
 	Super::PostLoad();
 
-	if(GetLinkerUE4Version() < VER_UE4_REFACTOR_MATERIAL_EXPRESSION_SCENECOLOR_AND_SCENEDEPTH_INPUTS)
+	if(GetLinkerUEVersion() < VER_UE4_REFACTOR_MATERIAL_EXPRESSION_SCENECOLOR_AND_SCENEDEPTH_INPUTS)
 	{
 		// Connect deprecated UV input to new expression input
 		InputMode = EMaterialSceneAttributeInputMode::OffsetFraction;
diff --git a/Engine/Source/Runtime/Engine/Private/Materials/MaterialInstance.cpp b/Engine/Source/Runtime/Engine/Private/Materials/MaterialInstance.cpp
--- a/Engine/Source/Runtime/Engine/Private/Materials/MaterialInstance.cpp
+++ b/Engine/Source/Runtime/Engine/Private/Materials/MaterialInstance.cpp
@@ -1154,7 +1154,7 @@ bool UMaterialInstance::GetRefractionSettings(float& OutBiasValue) const
 	bool bFoundAValue = false;
 
 	FMaterialParameterInfo ParamInfo;
-	if( GetLinkerUE4Version() >= VER_UE4_REFRACTION_BIAS_TO_REFRACTION_DEPTH_BIAS )
+	if( GetLinkerUEVersion() >= VER_UE4_REFRACTION_BIAS_TO_REFRACTION_DEPTH_BIAS )
 	{
 		static FName NAME_RefractionDepthBias(TEXT("RefractionDepthBias"));
 		ParamInfo.Name = NAME_RefractionDepthBias;
diff --git a/Engine/Source/Runtime/Engine/Private/Particles/ParticleModules.cpp b/Engine/Source/Runtime/Engine/Private/Particles/ParticleModules.cpp
--- a/Engine/Source/Runtime/Engine/Private/Particles/ParticleModules.cpp
+++ b/Engine/Source/Runtime/Engine/Private/Particles/ParticleModules.cpp
@@ -4545,7 +4545,7 @@ void UParticleModuleTypeDataGpu::PostLoad()
 	Super::PostLoad();
 	//EmitterInfo.Resources = BeginCreateGPUSpriteResources( ResourceData );
 
-	if (GetLinkerUE4Version() < VER_UE4_OPTIONALLY_CLEAR_GPU_EMITTERS_ON_INIT)
+	if (GetLinkerUEVersion() < VER_UE4_OPTIONALLY_CLEAR_GPU_EMITTERS_ON_INIT)
 	{
 		//Force old emitters to clear their particles on Init() to maintain old behaviour.
 		//New emitters are defaulted to false so they behave like other emitter types.
diff --git a/Engine/Source/Runtime/Engine/Private/Particles/ParticleModules_Location.cpp b/Engine/Source/Runtime/Engine/Private/Particles/ParticleModules_Location.cpp
--- a/Engine/Source/Runtime/Engine/Private/Particles/ParticleModules_Location.cpp
+++ b/Engine/Source/Runtime/Engine/Private/Particles/ParticleModules_Location.cpp
@@ -2205,7 +2205,7 @@ void UParticleModuleLocationSkelVertSurface::PostLoad()
 
 	NormalCheckTolerance = ((1.0f-(NormalCheckToleranceDegrees/180.0f))*2.0f)-1.0f;
 
-	if (GetLinkerUE4Version() < VER_UE4_FIX_SKEL_VERT_ORIENT_MESH_PARTICLES)
+	if (GetLinkerUEVersion() < VER_UE4_FIX_SKEL_VERT_ORIENT_MESH_PARTICLES)
 	{
 		//The code to actually do this hasn't been present ever in UE4 so I'm disabling it for old emitters.
 		//I expect that some users will have this set to true and it will alter the behavior of their emitters under their feet.
diff --git a/Engine/Source/Runtime/Engine/Private/Particles/ParticleModules_Spawn.cpp b/Engine/Source/Runtime/Engine/Private/Particles/ParticleModules_Spawn.cpp
--- a/Engine/Source/Runtime/Engine/Private/Particles/ParticleModules_Spawn.cpp
+++ b/Engine/Source/Runtime/Engine/Private/Particles/ParticleModules_Spawn.cpp
@@ -73,7 +73,7 @@ void UParticleModuleSpawn::PostLoad()
 {
 	Super::PostLoad();
 
-	if (GetLinkerUE4Version() < VER_UE4_GLOBAL_EMITTER_SPAWN_RATE_SCALE && Rate.Distribution)
+	if (GetLinkerUEVersion() < VER_UE4_GLOBAL_EMITTER_SPAWN_RATE_SCALE && Rate.Distribution)
 	{
 		Rate.Distribution->ConditionalPostLoad();//Ensure the distribution is loaded.
 		if (Rate.Distribution->IsA(UDistributionFloatConstant::StaticClass()) && Rate.GetValue() <= 1.0f)
diff --git a/Engine/Source/Runtime/Engine/Private/PhysicsEngine/BodyInstance.cpp b/Engine/Source/Runtime/Engine/Private/PhysicsEngine/BodyInstance.cpp
--- a/Engine/Source/Runtime/Engine/Private/PhysicsEngine/BodyInstance.cpp
+++ b/Engine/Source/Runtime/Engine/Private/PhysicsEngine/BodyInstance.cpp
@@ -3978,7 +3978,7 @@ void FBodyInstance::FixupData(class UObject* Loader)
 {
 	check (Loader);
 
-	int32 const UE4Version = Loader->GetLinkerUE4Version();
+	int32 const UE4Version = Loader->GetLinkerUEVersion();
 
 #if WITH_EDITOR
 	if (UE4Version < VER_UE4_ADD_CUSTOMPROFILENAME_CHANGE)
diff --git a/Engine/Source/Runtime/Engine/Private/PhysicsEngine/BodySetup.cpp b/Engine/Source/Runtime/Engine/Private/PhysicsEngine/BodySetup.cpp
--- a/Engine/Source/Runtime/Engine/Private/PhysicsEngine/BodySetup.cpp
+++ b/Engine/Source/Runtime/Engine/Private/PhysicsEngine/BodySetup.cpp
@@ -447,7 +447,7 @@ void UBodySetup::CreatePhysicsMeshes()
 	
 	// fix up invalid transform to use identity
 	// this can be here because BodySetup isn't blueprintable
-	if ( GetLinkerUE4Version() < VER_UE4_FIXUP_BODYSETUP_INVALID_CONVEX_TRANSFORM )
+	if ( GetLinkerUEVersion() < VER_UE4_FIXUP_BODYSETUP_INVALID_CONVEX_TRANSFORM )
 	{
 		for (int32 i=0; i<AggGeom.ConvexElems.Num(); ++i)
 		{
@@ -1196,7 +1196,7 @@ void UBodySetup::PostLoad()
 	}
 
 #if WITH_EDITORONLY_DATA
-	if ( GetLinkerUE4Version() < VER_UE4_BUILD_SCALE_VECTOR )
+	if ( GetLinkerUEVersion() < VER_UE4_BUILD_SCALE_VECTOR )
 	{
 		BuildScale3D = FVector( BuildScale_DEPRECATED );
 	}
@@ -1204,7 +1204,7 @@ void UBodySetup::PostLoad()
 
 	DefaultInstance.FixupData(this);
 
-	if ( GetLinkerUE4Version() < VER_UE4_REFACTOR_PHYSICS_BLENDING )
+	if ( GetLinkerUEVersion() < VER_UE4_REFACTOR_PHYSICS_BLENDING )
 	{
 		if ( bAlwaysFullAnimWeight_DEPRECATED )
 		{
@@ -1220,7 +1220,7 @@ void UBodySetup::PostLoad()
 		}
 	}
 
-	if ( GetLinkerUE4Version() < VER_UE4_BODYSETUP_COLLISION_CONVERSION )
+	if ( GetLinkerUEVersion() < VER_UE4_BODYSETUP_COLLISION_CONVERSION )
 	{
 		if ( DefaultInstance.GetCollisionEnabled() == ECollisionEnabled::NoCollision )
 		{
diff --git a/Engine/Source/Runtime/Engine/Private/PhysicsEngine/PhysicsConstraintActor.cpp b/Engine/Source/Runtime/Engine/Private/PhysicsEngine/PhysicsConstraintActor.cpp
--- a/Engine/Source/Runtime/Engine/Private/PhysicsEngine/PhysicsConstraintActor.cpp
+++ b/Engine/Source/Runtime/Engine/Private/PhysicsEngine/PhysicsConstraintActor.cpp
@@ -32,7 +32,7 @@ void APhysicsConstraintActor::PostLoad()
 	Super::PostLoad();
 
 	// Copy 'actors to constrain' into component
-	if (GetLinkerUE4Version() < VER_UE4_ALL_PROPS_TO_CONSTRAINTINSTANCE && ConstraintComp != NULL)
+	if (GetLinkerUEVersion() < VER_UE4_ALL_PROPS_TO_CONSTRAINTINSTANCE && ConstraintComp != NULL)
 	{
 		ConstraintComp->ConstraintActor1 = ConstraintActor1_DEPRECATED;
 		ConstraintComp->ConstraintActor2 = ConstraintActor2_DEPRECATED;
diff --git a/Engine/Source/Runtime/Engine/Private/PhysicsEngine/PhysicsConstraintComponent.cpp b/Engine/Source/Runtime/Engine/Private/PhysicsEngine/PhysicsConstraintComponent.cpp
--- a/Engine/Source/Runtime/Engine/Private/PhysicsEngine/PhysicsConstraintComponent.cpp
+++ b/Engine/Source/Runtime/Engine/Private/PhysicsEngine/PhysicsConstraintComponent.cpp
@@ -349,14 +349,14 @@ void UPhysicsConstraintComponent::PostLoad()
 	Super::PostLoad();
 
 	// Fix old content that used a ConstraintSetup
-	if ( GetLinkerUE4Version() < VER_UE4_ALL_PROPS_TO_CONSTRAINTINSTANCE && (ConstraintSetup_DEPRECATED != NULL) )
+	if ( GetLinkerUEVersion() < VER_UE4_ALL_PROPS_TO_CONSTRAINTINSTANCE && (ConstraintSetup_DEPRECATED != NULL) )
 	{
 		// Will have copied from setup into DefaultIntance inside
 		ConstraintInstance.CopyConstraintParamsFrom(&ConstraintSetup_DEPRECATED->DefaultInstance);
 		ConstraintSetup_DEPRECATED = NULL;
 	}
 
-	if (GetLinkerUE4Version() < VER_UE4_SOFT_CONSTRAINTS_USE_MASS)
+	if (GetLinkerUEVersion() < VER_UE4_SOFT_CONSTRAINTS_USE_MASS)
 	{
 		//In previous versions the mass was placed into the spring constant. This is correct because you use different springs for different mass - however, this makes tuning hard
 		//We now multiply mass into the spring constant. To fix old data we use CalculateMass which is not perfect but close (within 0.1kg)
diff --git a/Engine/Source/Runtime/Engine/Private/PhysicsVolume.cpp b/Engine/Source/Runtime/Engine/Private/PhysicsVolume.cpp
--- a/Engine/Source/Runtime/Engine/Private/PhysicsVolume.cpp
+++ b/Engine/Source/Runtime/Engine/Private/PhysicsVolume.cpp
@@ -24,7 +24,7 @@ void APhysicsVolume::LoadedFromAnotherClass(const FName& OldClassName)
 {
 	Super::LoadedFromAnotherClass(OldClassName);
 
-	if(GetLinkerUE4Version() < VER_UE4_REMOVE_DYNAMIC_VOLUME_CLASSES)
+	if(GetLinkerUEVersion() < VER_UE4_REMOVE_DYNAMIC_VOLUME_CLASSES)
 	{
 		static FName DynamicPhysicsVolume_NAME(TEXT("DynamicPhysicsVolume"));
 
diff --git a/Engine/Source/Runtime/Engine/Private/PlayerController.cpp b/Engine/Source/Runtime/Engine/Private/PlayerController.cpp
--- a/Engine/Source/Runtime/Engine/Private/PlayerController.cpp
+++ b/Engine/Source/Runtime/Engine/Private/PlayerController.cpp
@@ -883,7 +883,7 @@ void APlayerController::PostLoad()
 {
 	Super::PostLoad();
 
-	if (GetLinkerUE4Version() < VER_UE4_SPLIT_TOUCH_AND_CLICK_ENABLES)
+	if (GetLinkerUEVersion() < VER_UE4_SPLIT_TOUCH_AND_CLICK_ENABLES)
 	{
 		bEnableTouchEvents = bEnableClickEvents;
 	}
diff --git a/Engine/Source/Runtime/Engine/Private/SimpleConstructionScript.cpp b/Engine/Source/Runtime/Engine/Private/SimpleConstructionScript.cpp
--- a/Engine/Source/Runtime/Engine/Private/SimpleConstructionScript.cpp
+++ b/Engine/Source/Runtime/Engine/Private/SimpleConstructionScript.cpp
@@ -241,7 +241,7 @@ void USimpleConstructionScript::PostLoad()
 	// Reset non-native "root" scene component scale values, prior to the change in which
 	// we began applying custom scale values to root components at construction time. This
 	// way older, existing Blueprint actor instances won't start unexpectedly getting scaled.
-	if(GetLinkerUE4Version() < VER_UE4_BLUEPRINT_USE_SCS_ROOTCOMPONENT_SCALE)
+	if(GetLinkerUEVersion() < VER_UE4_BLUEPRINT_USE_SCS_ROOTCOMPONENT_SCALE)
 	{
 		if(BPGeneratedClass != nullptr)
 		{
@@ -279,7 +279,7 @@ void USimpleConstructionScript::PostLoad()
 		}
 	}
 
-	if (GetLinkerUE4Version() < VER_UE4_SCS_STORES_ALLNODES_ARRAY)
+	if (GetLinkerUEVersion() < VER_UE4_SCS_STORES_ALLNODES_ARRAY)
 	{
 		// Fill out AllNodes if this is an older object
 		if (RootNodes.Num() > 0)
@@ -1497,7 +1497,7 @@ void USimpleConstructionScript::ValidateNodeTemplates(FCompilerResultsLog& Messa
 
 	for (USCS_Node* Node : Nodes)
 	{
-		if (GetLinkerUE4Version() < VER_UE4_REMOVE_INPUT_COMPONENTS_FROM_BLUEPRINTS)
+		if (GetLinkerUEVersion() < VER_UE4_REMOVE_INPUT_COMPONENTS_FROM_BLUEPRINTS)
 		{
 			if (!Node->bIsNative_DEPRECATED && Node->ComponentTemplate && Node->ComponentTemplate->IsA<UInputComponent>())
 			{
diff --git a/Engine/Source/Runtime/Engine/Private/SkeletalMesh.cpp b/Engine/Source/Runtime/Engine/Private/SkeletalMesh.cpp
--- a/Engine/Source/Runtime/Engine/Private/SkeletalMesh.cpp
+++ b/Engine/Source/Runtime/Engine/Private/SkeletalMesh.cpp
@@ -3099,7 +3099,7 @@ void USkeletalMesh::BeginPostLoadInternal(FSkeletalMeshPostLoadContext& Context)
 			GetLODSettings()->SetLODSettingsToMesh(this);
 		}
 
-		if (GetLinkerUE4Version() < VER_UE4_SORT_ACTIVE_BONE_INDICES)
+		if (GetLinkerUEVersion() < VER_UE4_SORT_ACTIVE_BONE_INDICES)
 		{
 			for (int32 LodIndex = 0; LodIndex < LODInfoArray.Num(); LodIndex++)
 			{
diff --git a/Engine/Source/Runtime/Engine/Private/SoundBase.cpp b/Engine/Source/Runtime/Engine/Private/SoundBase.cpp
--- a/Engine/Source/Runtime/Engine/Private/SoundBase.cpp
+++ b/Engine/Source/Runtime/Engine/Private/SoundBase.cpp
@@ -207,9 +207,9 @@ void USoundBase::PostLoad()
 		bOutputToBusOnly_DEPRECATED = false;
 	}
 
-	const int32 LinkerUE4Version = GetLinkerUE4Version();
+	const int32 LinkerUEVersion = GetLinkerUEVersion();
 
-	if (LinkerUE4Version < VER_UE4_SOUND_CONCURRENCY_PACKAGE)
+	if (LinkerUEVersion < VER_UE4_SOUND_CONCURRENCY_PACKAGE)
 	{
 		bOverrideConcurrency = true;
 		ConcurrencyOverrides.bLimitToOwner = false;
diff --git a/Engine/Source/Runtime/Engine/Private/SoundClass.cpp b/Engine/Source/Runtime/Engine/Private/SoundClass.cpp
--- a/Engine/Source/Runtime/Engine/Private/SoundClass.cpp
+++ b/Engine/Source/Runtime/Engine/Private/SoundClass.cpp
@@ -66,7 +66,7 @@ void USoundClass::PostLoad()
 #if WITH_EDITORONLY_DATA
 	for (int32 ChildIndex = ChildClasses.Num()-1; ChildIndex >= 0; ChildIndex--)
 	{
-		if (ChildClasses[ChildIndex] != NULL && ChildClasses[ChildIndex]->GetLinkerUE4Version() < VER_UE4_SOUND_CLASS_GRAPH_EDITOR)
+		if (ChildClasses[ChildIndex] != NULL && ChildClasses[ChildIndex]->GetLinkerUEVersion() < VER_UE4_SOUND_CLASS_GRAPH_EDITOR)
 		{
 			// first come, first served
 			if (ChildClasses[ChildIndex]->ParentClass == nullptr)
diff --git a/Engine/Source/Runtime/Engine/Private/SpotLight.cpp b/Engine/Source/Runtime/Engine/Private/SpotLight.cpp
--- a/Engine/Source/Runtime/Engine/Private/SpotLight.cpp
+++ b/Engine/Source/Runtime/Engine/Private/SpotLight.cpp
@@ -63,7 +63,7 @@ void ASpotLight::LoadedFromAnotherClass(const FName& OldClassName)
 {
 	Super::LoadedFromAnotherClass(OldClassName);
 
-	if(GetLinkerUE4Version() < VER_UE4_REMOVE_LIGHT_MOBILITY_CLASSES)
+	if(GetLinkerUEVersion() < VER_UE4_REMOVE_LIGHT_MOBILITY_CLASSES)
 	{
 		static FName SpotLightStatic_NAME(TEXT("SpotLightStatic"));
 		static FName SpotLightMovable_NAME(TEXT("SpotLightMovable"));
diff --git a/Engine/Source/Runtime/Engine/Private/StaticMesh.cpp b/Engine/Source/Runtime/Engine/Private/StaticMesh.cpp
--- a/Engine/Source/Runtime/Engine/Private/StaticMesh.cpp
+++ b/Engine/Source/Runtime/Engine/Private/StaticMesh.cpp
@@ -5105,7 +5105,7 @@ void UStaticMesh::BeginPostLoadInternal(FStaticMeshPostLoadContext& Context)
 	if (!Context.bIsCookedForEditor)
 	{
 		// Needs to happen before 'CacheDerivedData'
-		if (GetLinkerUE4Version() < VER_UE4_BUILD_SCALE_VECTOR)
+		if (GetLinkerUEVersion() < VER_UE4_BUILD_SCALE_VECTOR)
 		{
 			int32 NumLODs = GetNumSourceModels();
 			for (int32 LODIndex = 0; LODIndex < NumLODs; ++LODIndex)
@@ -5115,7 +5115,7 @@ void UStaticMesh::BeginPostLoadInternal(FStaticMeshPostLoadContext& Context)
 			}
 		}
 
-		if (GetLinkerUE4Version() < VER_UE4_LIGHTMAP_MESH_BUILD_SETTINGS)
+		if (GetLinkerUEVersion() < VER_UE4_LIGHTMAP_MESH_BUILD_SETTINGS)
 		{
 			for (int32 i = 0; i < GetNumSourceModels(); i++)
 			{
@@ -5123,7 +5123,7 @@ void UStaticMesh::BeginPostLoadInternal(FStaticMeshPostLoadContext& Context)
 			}
 		}
 
-		if (GetLinkerUE4Version() < VER_UE4_MIKKTSPACE_IS_DEFAULT)
+		if (GetLinkerUEVersion() < VER_UE4_MIKKTSPACE_IS_DEFAULT)
 		{
 			for (int32 i = 0; i < GetNumSourceModels(); ++i)
 			{
@@ -5131,7 +5131,7 @@ void UStaticMesh::BeginPostLoadInternal(FStaticMeshPostLoadContext& Context)
 			}
 		}
 
-		if (GetLinkerUE4Version() < VER_UE4_BUILD_MESH_ADJ_BUFFER_FLAG_EXPOSED)
+		if (GetLinkerUEVersion() < VER_UE4_BUILD_MESH_ADJ_BUFFER_FLAG_EXPOSED)
 		{
 			FRawMesh TempRawMesh;
 			uint32 TotalIndexCount = 0;
@@ -5162,7 +5162,7 @@ void UStaticMesh::BeginPostLoadInternal(FStaticMeshPostLoadContext& Context)
 		MeshUtilities.FixupMaterialSlotNames(this);
 
 		const int32 CustomVersion = GetLinkerCustomVersion(FReleaseObjectVersion::GUID);
-		if (GetLinkerUE4Version() < VER_UE4_STATIC_MESH_EXTENDED_BOUNDS || CustomVersion < FReleaseObjectVersion::StaticMeshExtendedBoundsFix)
+		if (GetLinkerUEVersion() < VER_UE4_STATIC_MESH_EXTENDED_BOUNDS || CustomVersion < FReleaseObjectVersion::StaticMeshExtendedBoundsFix)
 		{
 			// A stall is almost guaranteed during async build because mesh bounds are used extensively from many different places.
 			Context.bShouldComputeExtendedBounds = true;
diff --git a/Engine/Source/Runtime/Engine/Private/StaticMeshActor.cpp b/Engine/Source/Runtime/Engine/Private/StaticMeshActor.cpp
--- a/Engine/Source/Runtime/Engine/Private/StaticMeshActor.cpp
+++ b/Engine/Source/Runtime/Engine/Private/StaticMeshActor.cpp
@@ -84,7 +84,7 @@ void AStaticMeshActor::LoadedFromAnotherClass(const FName& OldClassName)
 {
 	Super::LoadedFromAnotherClass(OldClassName);
 
-	if(GetLinkerUE4Version() < VER_UE4_REMOVE_STATICMESH_MOBILITY_CLASSES)
+	if(GetLinkerUEVersion() < VER_UE4_REMOVE_STATICMESH_MOBILITY_CLASSES)
 	{
 		static FName InterpActor_NAME(TEXT("InterpActor"));
 		static FName PhysicsActor_NAME(TEXT("PhysicsActor"));
diff --git a/Engine/Source/Runtime/Engine/Private/Texture.cpp b/Engine/Source/Runtime/Engine/Private/Texture.cpp
--- a/Engine/Source/Runtime/Engine/Private/Texture.cpp
+++ b/Engine/Source/Runtime/Engine/Private/Texture.cpp
@@ -445,7 +445,7 @@ void UTexture::Serialize(FArchive& Ar)
 		Source.BulkData.Serialize(Ar, this);
 	}
 
-	if ( GetLinkerUE4Version() < VER_UE4_TEXTURE_LEGACY_GAMMA )
+	if ( GetLinkerUEVersion() < VER_UE4_TEXTURE_LEGACY_GAMMA )
 	{
 		bUseLegacyGamma = true;
 	}
diff --git a/Engine/Source/Runtime/Engine/Private/TriggerVolume.cpp b/Engine/Source/Runtime/Engine/Private/TriggerVolume.cpp
--- a/Engine/Source/Runtime/Engine/Private/TriggerVolume.cpp
+++ b/Engine/Source/Runtime/Engine/Private/TriggerVolume.cpp
@@ -23,7 +23,7 @@ void ATriggerVolume::LoadedFromAnotherClass(const FName& OldClassName)
 {
 	Super::LoadedFromAnotherClass(OldClassName);
 
-	if(GetLinkerUE4Version() < VER_UE4_REMOVE_DYNAMIC_VOLUME_CLASSES)
+	if(GetLinkerUEVersion() < VER_UE4_REMOVE_DYNAMIC_VOLUME_CLASSES)
 	{
 		static FName DynamicTriggerVolume_NAME(TEXT("DynamicTriggerVolume"));
 
diff --git a/Engine/Source/Runtime/Landscape/Private/LandscapeCollision.cpp b/Engine/Source/Runtime/Landscape/Private/LandscapeCollision.cpp
--- a/Engine/Source/Runtime/Landscape/Private/LandscapeCollision.cpp
+++ b/Engine/Source/Runtime/Landscape/Private/LandscapeCollision.cpp
@@ -730,7 +730,7 @@ void ULandscapeHeightfieldCollisionComponent::CreateCollisionObject()
 void ULandscapeHeightfieldCollisionComponent::SpeculativelyLoadAsyncDDCCollsionData()
 {
 #if WITH_PHYSX
-	if (GetLinkerUE4Version() >= VER_UE4_LANDSCAPE_SERIALIZE_PHYSICS_MATERIALS && !GLandscapeCollisionSkipDDC)
+	if (GetLinkerUEVersion() >= VER_UE4_LANDSCAPE_SERIALIZE_PHYSICS_MATERIALS && !GLandscapeCollisionSkipDDC)
 	{
 		UWorld* World = GetWorld();
 		if (World && HeightfieldGuid.IsValid() && CookedPhysicalMaterials.Num() > 0 && GSharedHeightfieldRefs.FindRef(HeightfieldGuid) == nullptr)
@@ -852,7 +852,7 @@ bool ULandscapeHeightfieldCollisionComponent::CookCollisionData(const FName& For
 	if (!GLandscapeCollisionSkipDDC && bCheckDDC && HeightfieldGuid.IsValid())
 	{
 		// Ensure that content was saved with physical materials before using DDC data
-		if (GetLinkerUE4Version() >= VER_UE4_LANDSCAPE_SERIALIZE_PHYSICS_MATERIALS)
+		if (GetLinkerUEVersion() >= VER_UE4_LANDSCAPE_SERIALIZE_PHYSICS_MATERIALS)
 		{
 			FString DDCKey = GetHFDDCKeyString(Format, bUseDefMaterial, HeightfieldGuid, InOutMaterials);
 
@@ -1082,7 +1082,7 @@ bool ULandscapeMeshCollisionComponent::CookCollisionData(const FName& Format, bo
 	if (!GLandscapeCollisionSkipDDC && bCheckDDC)
 	{
 		// Ensure that content was saved with physical materials before using DDC data
-		if (GetLinkerUE4Version() >= VER_UE4_LANDSCAPE_SERIALIZE_PHYSICS_MATERIALS && MeshGuid.IsValid())
+		if (GetLinkerUEVersion() >= VER_UE4_LANDSCAPE_SERIALIZE_PHYSICS_MATERIALS && MeshGuid.IsValid())
 		{
 			FString DDCKey = GetHFDDCKeyString(Format, bUseDefMaterial, MeshGuid, InOutMaterials);
 
diff --git a/Engine/Source/Runtime/Landscape/Private/LandscapeSplines.cpp b/Engine/Source/Runtime/Landscape/Private/LandscapeSplines.cpp
--- a/Engine/Source/Runtime/Landscape/Private/LandscapeSplines.cpp
+++ b/Engine/Source/Runtime/Landscape/Private/LandscapeSplines.cpp
@@ -1621,7 +1621,7 @@ void ULandscapeSplineControlPoint::PostLoad()
 		}
 	}
 
-	if (GetLinkerUE4Version() < VER_UE4_LANDSCAPE_SPLINE_CROSS_LEVEL_MESHES)
+	if (GetLinkerUEVersion() < VER_UE4_LANDSCAPE_SPLINE_CROSS_LEVEL_MESHES)
 	{
 		// Fix collision profile
 		if (LocalMeshComponent != nullptr) // ForeignMeshComponents didn't exist yet
@@ -2434,7 +2434,7 @@ void ULandscapeSplineSegment::PostLoad()
 #if WITH_EDITOR
 	if (GIsEditor)
 	{
-		if (GetLinkerUE4Version() < VER_UE4_ADDED_LANDSCAPE_SPLINE_EDITOR_MESH &&
+		if (GetLinkerUEVersion() < VER_UE4_ADDED_LANDSCAPE_SPLINE_EDITOR_MESH &&
 			LocalMeshComponents.Num() == 0) // ForeignMeshComponents didn't exist yet
 		{
 			UpdateSplinePoints();
@@ -2469,7 +2469,7 @@ void ULandscapeSplineSegment::PostLoad()
 		}
 	}
 
-	if (GetLinkerUE4Version() < VER_UE4_LANDSCAPE_SPLINE_CROSS_LEVEL_MESHES)
+	if (GetLinkerUEVersion() < VER_UE4_LANDSCAPE_SPLINE_CROSS_LEVEL_MESHES)
 	{
 		// Fix collision profile
 		for (auto& LocalMeshComponent : LocalMeshComponents) // ForeignMeshComponents didn't exist yet
diff --git a/Engine/Source/Runtime/Landscape/Private/Materials/MaterialExpressionLandscapeLayerSwitch.cpp b/Engine/Source/Runtime/Landscape/Private/Materials/MaterialExpressionLandscapeLayerSwitch.cpp
--- a/Engine/Source/Runtime/Landscape/Private/Materials/MaterialExpressionLandscapeLayerSwitch.cpp
+++ b/Engine/Source/Runtime/Landscape/Private/Materials/MaterialExpressionLandscapeLayerSwitch.cpp
@@ -124,7 +124,7 @@ void UMaterialExpressionLandscapeLayerSwitch::PostLoad()
 {
 	Super::PostLoad();
 
-	if (GetLinkerUE4Version() < VER_UE4_FIXUP_TERRAIN_LAYER_NODES)
+	if (GetLinkerUEVersion() < VER_UE4_FIXUP_TERRAIN_LAYER_NODES)
 	{
 		UpdateParameterGuid(true, true);
 	}
diff --git a/Engine/Source/Runtime/Landscape/Private/Materials/MaterialExpressionLandscapeLayerWeight.cpp b/Engine/Source/Runtime/Landscape/Private/Materials/MaterialExpressionLandscapeLayerWeight.cpp
--- a/Engine/Source/Runtime/Landscape/Private/Materials/MaterialExpressionLandscapeLayerWeight.cpp
+++ b/Engine/Source/Runtime/Landscape/Private/Materials/MaterialExpressionLandscapeLayerWeight.cpp
@@ -49,7 +49,7 @@ void UMaterialExpressionLandscapeLayerWeight::PostLoad()
 {
 	Super::PostLoad();
 
-	if (GetLinkerUE4Version() < VER_UE4_FIXUP_TERRAIN_LAYER_NODES)
+	if (GetLinkerUEVersion() < VER_UE4_FIXUP_TERRAIN_LAYER_NODES)
 	{
 		UpdateParameterGuid(true, true);
 	}
diff --git a/Engine/Source/Runtime/NavigationSystem/Private/NavigationData.cpp b/Engine/Source/Runtime/NavigationSystem/Private/NavigationData.cpp
--- a/Engine/Source/Runtime/NavigationSystem/Private/NavigationData.cpp
+++ b/Engine/Source/Runtime/NavigationSystem/Private/NavigationData.cpp
@@ -219,7 +219,7 @@ void ANavigationData::PostLoad()
 {
 	Super::PostLoad();
 
-	if ((GetLinkerUE4Version() < VER_UE4_ADD_MODIFIERS_RUNTIME_GENERATION) &&
+	if ((GetLinkerUEVersion() < VER_UE4_ADD_MODIFIERS_RUNTIME_GENERATION) &&
 		(RuntimeGeneration == ERuntimeGenerationType::LegacyGeneration))
 	{
 		RuntimeGeneration = bRebuildAtRuntime_DEPRECATED ? ERuntimeGenerationType::Dynamic : ERuntimeGenerationType::Static;
diff --git a/Engine/Source/Runtime/PhysicsCore/Private/PhysicalMaterial.cpp b/Engine/Source/Runtime/PhysicsCore/Private/PhysicalMaterial.cpp
--- a/Engine/Source/Runtime/PhysicsCore/Private/PhysicalMaterial.cpp
+++ b/Engine/Source/Runtime/PhysicsCore/Private/PhysicalMaterial.cpp
@@ -76,7 +76,7 @@ void UPhysicalMaterial::PostLoad()
 	Super::PostLoad();
 
 	// we're removing physical material property, so convert to Material type
-	if (GetLinkerUE4Version() < VER_UE4_REMOVE_PHYSICALMATERIALPROPERTY)
+	if (GetLinkerUEVersion() < VER_UE4_REMOVE_PHYSICALMATERIALPROPERTY)
 	{
 		if (PhysicalMaterialProperty_DEPRECATED)
 		{
diff --git a/Engine/Source/Runtime/UMG/Private/Components/Border.cpp b/Engine/Source/Runtime/UMG/Private/Components/Border.cpp
--- a/Engine/Source/Runtime/UMG/Private/Components/Border.cpp
+++ b/Engine/Source/Runtime/UMG/Private/Components/Border.cpp
@@ -291,7 +291,7 @@ void UBorder::PostLoad()
 	Super::PostLoad();
 
 #if WITH_EDITORONLY_DATA
-	if ( GetLinkerUE4Version() < VER_UE4_DEPRECATE_UMG_STYLE_ASSETS && Brush_DEPRECATED != nullptr )
+	if ( GetLinkerUEVersion() < VER_UE4_DEPRECATE_UMG_STYLE_ASSETS && Brush_DEPRECATED != nullptr )
 	{
 		Background = Brush_DEPRECATED->Brush;
 		Brush_DEPRECATED = nullptr;
diff --git a/Engine/Source/Runtime/UMG/Private/Components/Button.cpp b/Engine/Source/Runtime/UMG/Private/Components/Button.cpp
--- a/Engine/Source/Runtime/UMG/Private/Components/Button.cpp
+++ b/Engine/Source/Runtime/UMG/Private/Components/Button.cpp
@@ -186,7 +186,7 @@ void UButton::PostLoad()
 		}
 	}
 
-	if( GetLinkerUE4Version() < VER_UE4_DEPRECATE_UMG_STYLE_ASSETS && Style_DEPRECATED != nullptr )
+	if( GetLinkerUEVersion() < VER_UE4_DEPRECATE_UMG_STYLE_ASSETS && Style_DEPRECATED != nullptr )
 	{
 		const FButtonStyle* StylePtr = Style_DEPRECATED->GetStyle<FButtonStyle>();
 		if(StylePtr != nullptr)
diff --git a/Engine/Source/Runtime/UMG/Private/Components/CheckBox.cpp b/Engine/Source/Runtime/UMG/Private/Components/CheckBox.cpp
--- a/Engine/Source/Runtime/UMG/Private/Components/CheckBox.cpp
+++ b/Engine/Source/Runtime/UMG/Private/Components/CheckBox.cpp
@@ -185,7 +185,7 @@ void UCheckBox::PostLoad()
 {
 	Super::PostLoad();
 
-	if ( GetLinkerUE4Version() < VER_UE4_DEPRECATE_UMG_STYLE_ASSETS )
+	if ( GetLinkerUEVersion() < VER_UE4_DEPRECATE_UMG_STYLE_ASSETS )
 	{
 		if ( Style_DEPRECATED != nullptr )
 		{
@@ -253,7 +253,7 @@ void UCheckBox::PostLoad()
 		}
 	}
 
-	if (GetLinkerUE4Version() < VER_UE4_DEPRECATE_UMG_STYLE_OVERRIDES)
+	if (GetLinkerUEVersion() < VER_UE4_DEPRECATE_UMG_STYLE_OVERRIDES)
 	{
 		WidgetStyle.Padding = Padding_DEPRECATED;
 		Padding_DEPRECATED = FMargin(0);
diff --git a/Engine/Source/Runtime/UMG/Private/Components/CircularThrobber.cpp b/Engine/Source/Runtime/UMG/Private/Components/CircularThrobber.cpp
--- a/Engine/Source/Runtime/UMG/Private/Components/CircularThrobber.cpp
+++ b/Engine/Source/Runtime/UMG/Private/Components/CircularThrobber.cpp
@@ -101,7 +101,7 @@ void UCircularThrobber::PostLoad()
 {
 	Super::PostLoad();
 
-	if ( GetLinkerUE4Version() < VER_UE4_DEPRECATE_UMG_STYLE_ASSETS )
+	if ( GetLinkerUEVersion() < VER_UE4_DEPRECATE_UMG_STYLE_ASSETS )
 	{
 		if ( PieceImage_DEPRECATED != nullptr )
 		{
diff --git a/Engine/Source/Runtime/UMG/Private/Components/EditableText.cpp b/Engine/Source/Runtime/UMG/Private/Components/EditableText.cpp
--- a/Engine/Source/Runtime/UMG/Private/Components/EditableText.cpp
+++ b/Engine/Source/Runtime/UMG/Private/Components/EditableText.cpp
@@ -177,7 +177,7 @@ void UEditableText::PostLoad()
 {
 	Super::PostLoad();
 
-	if ( GetLinkerUE4Version() < VER_UE4_DEPRECATE_UMG_STYLE_ASSETS )
+	if ( GetLinkerUEVersion() < VER_UE4_DEPRECATE_UMG_STYLE_ASSETS )
 	{
 		if ( Style_DEPRECATED != nullptr )
 		{
@@ -209,7 +209,7 @@ void UEditableText::PostLoad()
 		}
 	}
 
-	if (GetLinkerUE4Version() < VER_UE4_DEPRECATE_UMG_STYLE_OVERRIDES)
+	if (GetLinkerUEVersion() < VER_UE4_DEPRECATE_UMG_STYLE_OVERRIDES)
 	{
 		if (Font_DEPRECATED.HasValidFont())
 		{
diff --git a/Engine/Source/Runtime/UMG/Private/Components/EditableTextBox.cpp b/Engine/Source/Runtime/UMG/Private/Components/EditableTextBox.cpp
--- a/Engine/Source/Runtime/UMG/Private/Components/EditableTextBox.cpp
+++ b/Engine/Source/Runtime/UMG/Private/Components/EditableTextBox.cpp
@@ -215,7 +215,7 @@ void UEditableTextBox::PostLoad()
 {
 	Super::PostLoad();
 
-	if ( GetLinkerUE4Version() < VER_UE4_DEPRECATE_UMG_STYLE_ASSETS )
+	if ( GetLinkerUEVersion() < VER_UE4_DEPRECATE_UMG_STYLE_ASSETS )
 	{
 		if ( Style_DEPRECATED != nullptr )
 		{
@@ -229,7 +229,7 @@ void UEditableTextBox::PostLoad()
 		}
 	}
 
-	if (GetLinkerUE4Version() < VER_UE4_DEPRECATE_UMG_STYLE_OVERRIDES)
+	if (GetLinkerUEVersion() < VER_UE4_DEPRECATE_UMG_STYLE_OVERRIDES)
 	{
 		if (Font_DEPRECATED.HasValidFont())
 		{
diff --git a/Engine/Source/Runtime/UMG/Private/Components/Image.cpp b/Engine/Source/Runtime/UMG/Private/Components/Image.cpp
--- a/Engine/Source/Runtime/UMG/Private/Components/Image.cpp
+++ b/Engine/Source/Runtime/UMG/Private/Components/Image.cpp
@@ -27,7 +27,7 @@ void UImage::PostLoad()
 {
 	Super::PostLoad();
 
-	if ( GetLinkerUE4Version() < VER_UE4_DEPRECATE_UMG_STYLE_ASSETS && Image_DEPRECATED != nullptr )
+	if ( GetLinkerUEVersion() < VER_UE4_DEPRECATE_UMG_STYLE_ASSETS && Image_DEPRECATED != nullptr )
 	{
 		Brush = Image_DEPRECATED->Brush;
 		Image_DEPRECATED = nullptr;
@@ -106,14 +106,14 @@ void UImage::SetBrush(const FSlateBrush& InBrush)
 
 void UImage::SetBrushSize(FVector2D DesiredSize)
 {
-	if(Brush.ImageSize != DesiredSize)
-	{
-		Brush.ImageSize = DesiredSize;
+	SetDesiredSizeOverride(DesiredSize);
+}
 
-		if (MyImage.IsValid())
-		{
-			MyImage->SetImage(&Brush);
-		}
+void UImage::SetDesiredSizeOverride(FVector2D DesiredSize)
+{
+	if (MyImage.IsValid())
+	{
+		MyImage->SetDesiredSizeOverride(DesiredSize);
 	}
 }
 
diff --git a/Engine/Source/Runtime/UMG/Private/Components/MultiLineEditableText.cpp b/Engine/Source/Runtime/UMG/Private/Components/MultiLineEditableText.cpp
--- a/Engine/Source/Runtime/UMG/Private/Components/MultiLineEditableText.cpp
+++ b/Engine/Source/Runtime/UMG/Private/Components/MultiLineEditableText.cpp
@@ -188,7 +188,7 @@ void UMultiLineEditableText::PostLoad()
 {
 	Super::PostLoad();
 
-	if (GetLinkerUE4Version() < VER_UE4_DEPRECATE_UMG_STYLE_OVERRIDES)
+	if (GetLinkerUEVersion() < VER_UE4_DEPRECATE_UMG_STYLE_OVERRIDES)
 	{
 		if (Font_DEPRECATED.HasValidFont())
 		{
diff --git a/Engine/Source/Runtime/UMG/Private/Components/MultiLineEditableTextBox.cpp b/Engine/Source/Runtime/UMG/Private/Components/MultiLineEditableTextBox.cpp
--- a/Engine/Source/Runtime/UMG/Private/Components/MultiLineEditableTextBox.cpp
+++ b/Engine/Source/Runtime/UMG/Private/Components/MultiLineEditableTextBox.cpp
@@ -210,7 +210,7 @@ void UMultiLineEditableTextBox::PostLoad()
 {
 	Super::PostLoad();
 
-	if ( GetLinkerUE4Version() < VER_UE4_DEPRECATE_UMG_STYLE_ASSETS )
+	if ( GetLinkerUEVersion() < VER_UE4_DEPRECATE_UMG_STYLE_ASSETS )
 	{
 		if ( Style_DEPRECATED != nullptr )
 		{
@@ -224,7 +224,7 @@ void UMultiLineEditableTextBox::PostLoad()
 		}
 	}
 
-	if (GetLinkerUE4Version() < VER_UE4_DEPRECATE_UMG_STYLE_OVERRIDES)
+	if (GetLinkerUEVersion() < VER_UE4_DEPRECATE_UMG_STYLE_OVERRIDES)
 	{
 		if (Font_DEPRECATED.HasValidFont())
 		{
diff --git a/Engine/Source/Runtime/UMG/Private/Components/ProgressBar.cpp b/Engine/Source/Runtime/UMG/Private/Components/ProgressBar.cpp
--- a/Engine/Source/Runtime/UMG/Private/Components/ProgressBar.cpp
+++ b/Engine/Source/Runtime/UMG/Private/Components/ProgressBar.cpp
@@ -92,7 +92,7 @@ void UProgressBar::PostLoad()
 {
 	Super::PostLoad();
 
-	if (GetLinkerUE4Version() < VER_UE4_DEPRECATE_UMG_STYLE_ASSETS)
+	if (GetLinkerUEVersion() < VER_UE4_DEPRECATE_UMG_STYLE_ASSETS)
 	{
 		if (Style_DEPRECATED != nullptr)
 		{
diff --git a/Engine/Source/Runtime/UMG/Private/Components/ScrollBar.cpp b/Engine/Source/Runtime/UMG/Private/Components/ScrollBar.cpp
--- a/Engine/Source/Runtime/UMG/Private/Components/ScrollBar.cpp
+++ b/Engine/Source/Runtime/UMG/Private/Components/ScrollBar.cpp
@@ -96,7 +96,7 @@ void UScrollBar::PostLoad()
 {
 	Super::PostLoad();
 
-	if ( GetLinkerUE4Version() < VER_UE4_DEPRECATE_UMG_STYLE_ASSETS )
+	if ( GetLinkerUEVersion() < VER_UE4_DEPRECATE_UMG_STYLE_ASSETS )
 	{
 		if ( Style_DEPRECATED != nullptr )
 		{
diff --git a/Engine/Source/Runtime/UMG/Private/Components/ScrollBox.cpp b/Engine/Source/Runtime/UMG/Private/Components/ScrollBox.cpp
--- a/Engine/Source/Runtime/UMG/Private/Components/ScrollBox.cpp
+++ b/Engine/Source/Runtime/UMG/Private/Components/ScrollBox.cpp
@@ -232,7 +232,7 @@ void UScrollBox::PostLoad()
 {
 	Super::PostLoad();
 
-	if ( GetLinkerUE4Version() < VER_UE4_DEPRECATE_UMG_STYLE_ASSETS )
+	if ( GetLinkerUEVersion() < VER_UE4_DEPRECATE_UMG_STYLE_ASSETS )
 	{
 		if ( Style_DEPRECATED != nullptr )
 		{
diff --git a/Engine/Source/Runtime/UMG/Private/Components/SpinBox.cpp b/Engine/Source/Runtime/UMG/Private/Components/SpinBox.cpp
--- a/Engine/Source/Runtime/UMG/Private/Components/SpinBox.cpp
+++ b/Engine/Source/Runtime/UMG/Private/Components/SpinBox.cpp
@@ -381,7 +381,7 @@ void USpinBox::PostLoad()
 {
 	Super::PostLoad();
 
-	if ( GetLinkerUE4Version() < VER_UE4_DEPRECATE_UMG_STYLE_ASSETS )
+	if ( GetLinkerUEVersion() < VER_UE4_DEPRECATE_UMG_STYLE_ASSETS )
 	{
 		if ( Style_DEPRECATED != nullptr )
 		{
diff --git a/Engine/Source/Runtime/UMG/Private/Components/Throbber.cpp b/Engine/Source/Runtime/UMG/Private/Components/Throbber.cpp
--- a/Engine/Source/Runtime/UMG/Private/Components/Throbber.cpp
+++ b/Engine/Source/Runtime/UMG/Private/Components/Throbber.cpp
@@ -106,7 +106,7 @@ void UThrobber::PostLoad()
 {
 	Super::PostLoad();
 
-	if ( GetLinkerUE4Version() < VER_UE4_DEPRECATE_UMG_STYLE_ASSETS )
+	if ( GetLinkerUEVersion() < VER_UE4_DEPRECATE_UMG_STYLE_ASSETS )
 	{
 		if ( PieceImage_DEPRECATED != nullptr )
 		{
diff --git a/Engine/Source/Runtime/UMG/Private/WidgetBlueprintGeneratedClass.cpp b/Engine/Source/Runtime/UMG/Private/WidgetBlueprintGeneratedClass.cpp
--- a/Engine/Source/Runtime/UMG/Private/WidgetBlueprintGeneratedClass.cpp
+++ b/Engine/Source/Runtime/UMG/Private/WidgetBlueprintGeneratedClass.cpp
@@ -337,7 +337,7 @@ void UWidgetBlueprintGeneratedClass::PostLoad()
 	}
 
 #if WITH_EDITOR
-	if ( GetLinkerUE4Version() < VER_UE4_RENAME_WIDGET_VISIBILITY )
+	if ( GetLinkerUEVersion() < VER_UE4_RENAME_WIDGET_VISIBILITY )
 	{
 		static const FName Visiblity(TEXT("Visiblity"));
 		static const FName Visibility(TEXT("Visibility"));
