World Partition SoftObjectPath support
- World Partition streaming policy is now persistent and duplicated in PIE, so we can maintain actors to cells mappings.
- Added support for actors subobjects resolving through SoftObjectPath (can now resolve components to actors).
- Added support in SoftObjectPath to resolve through the top-level object (similar to what loading did).
- World Partition subobjects resolving is now available in cooked builds.

diff --git a/Engine/Source/Editor/PropertyEditor/Private/UserInterface/PropertyEditor/SPropertyEditorAsset.cpp b/Engine/Source/Editor/PropertyEditor/Private/UserInterface/PropertyEditor/SPropertyEditorAsset.cpp
--- a/Engine/Source/Editor/PropertyEditor/Private/UserInterface/PropertyEditor/SPropertyEditorAsset.cpp
+++ b/Engine/Source/Editor/PropertyEditor/Private/UserInterface/PropertyEditor/SPropertyEditorAsset.cpp
@@ -709,7 +709,7 @@ SPropertyEditorAsset::EActorReferenceState SPropertyEditorAsset::GetActorReferen
 				if (World && World->GetWorldPartition())
 				{
 					UObject* Object = nullptr;
-					if (World->LoadSubobject(*Value.ObjectPath.GetSubPathString(), Object, /*bOnlyTestExistence*/true))
+					if (World->ResolveSubobject(*Value.ObjectPath.GetSubPathString(), Object, /*bLoadIfExists*/false))
 					{
 						return EActorReferenceState::Exists;
 					}
diff --git a/Engine/Source/Runtime/CoreUObject/Private/UObject/SoftObjectPath.cpp b/Engine/Source/Runtime/CoreUObject/Private/UObject/SoftObjectPath.cpp
--- a/Engine/Source/Runtime/CoreUObject/Private/UObject/SoftObjectPath.cpp
+++ b/Engine/Source/Runtime/CoreUObject/Private/UObject/SoftObjectPath.cpp
@@ -435,14 +435,12 @@ UObject* FSoftObjectPath::TryLoad(FUObjectSerializeContext* InLoadContext) const
 			// This probably loaded the top-level object, so re-resolve ourselves
 			LoadedObject = ResolveObject();
 
-#if WITH_EDITOR
 			// If the the top-level object exists but we can't find the object, defer the loading to the top-level container object in case
 			// it knows how to load that specific object.
 			if (!LoadedObject && TopLevelObject)
 			{
-				TopLevelObject->LoadSubobject(*SubPathString, LoadedObject);
+				TopLevelObject->ResolveSubobject(*SubPathString, LoadedObject, /*bLoadIfExists*/true);
 			}
-#endif
 		}
 		else
 		{
@@ -525,6 +523,20 @@ UObject* FSoftObjectPath::ResolveObjectInternal(const TCHAR* PathString) const
 {
 	UObject* FoundObject = FindObject<UObject>(nullptr, PathString);
 
+	if (!FoundObject && IsSubobject())
+	{
+		// Try to resolve through the top level object
+		FSoftObjectPath TopLevelPath = FSoftObjectPath(AssetPathName, FString());
+		UObject* TopLevelObject = TopLevelPath.ResolveObject();
+
+		// If the the top-level object exists but we can't find the object, defer the resolving to the top-level container object in case
+		// it knows how to load that specific object.
+		if (TopLevelObject)
+		{
+			TopLevelObject->ResolveSubobject(*SubPathString, FoundObject, /*bLoadIfExists*/false);
+		}
+	}
+
 #if WITH_EDITOR
 	// Look at core redirects if we didn't find the object
 	if (!FoundObject)
diff --git a/Engine/Source/Runtime/CoreUObject/Public/UObject/Object.h b/Engine/Source/Runtime/CoreUObject/Public/UObject/Object.h
--- a/Engine/Source/Runtime/CoreUObject/Public/UObject/Object.h
+++ b/Engine/Source/Runtime/CoreUObject/Public/UObject/Object.h
@@ -268,6 +268,7 @@ public:
 	 * Called when the object was loaded from another class via active class redirects.
 	 */
 	virtual void LoadedFromAnotherClass(const FName& OldClassName) {}
+#endif
 
 	/**
 	 * Called to defer loading a subobject to its top-level container object. Usefull with the package override
@@ -276,13 +277,12 @@ public:
 	 * 
  	 * @param	SubObjectPath		object path to handle, relative to the current implementer
 	 * @param	OutObject			receives the loaded object, or existing loaded object if bOnlyTestExistence is true and object is already loaded
-	 * @param	bOnlyTestExistence	if true, don't load the object, just return true if it exists
+	 * @param	bLoadIfExists		if true, load the object if it exists
 	 * 
-	 * @return true if the object loaded or exists, depending on bOnlyTestExistence
+	 * @return true if the object loaded or exists, depending on bLoadIfExists
 * 
 	 */
-	virtual bool LoadSubobject(const TCHAR* SubObjectPath, UObject*& OutObject, bool bOnlyTestExistence = false) { return false; }
-#endif
+	virtual bool ResolveSubobject(const TCHAR* SubObjectPath, UObject*& OutObject, bool bLoadIfExists) { return false; }
 
 	/**
 	 * Called before calling PostLoad() in FAsyncPackage::PostLoadObjects(). This is the safeguard to prevent PostLoad() from stalling the main thread.
diff --git a/Engine/Source/Runtime/Engine/Classes/Engine/Level.h b/Engine/Source/Runtime/Engine/Classes/Engine/Level.h
--- a/Engine/Source/Runtime/Engine/Classes/Engine/Level.h
+++ b/Engine/Source/Runtime/Engine/Classes/Engine/Level.h
@@ -785,8 +785,8 @@ public:
 	virtual void PostEditChangeProperty(FPropertyChangedEvent& PropertyChangedEvent) override;
 	virtual void BeginCacheForCookedPlatformData(const ITargetPlatform *TargetPlatform) override;
 	virtual bool CanEditChange(const FProperty* PropertyThatWillChange) const override;
-	virtual bool LoadSubobject(const TCHAR* SubObjectPath, UObject*& OutObject, bool bOnlyTestExistence) override;
 #endif // WITH_EDITOR
+	virtual bool ResolveSubobject(const TCHAR* SubObjectPath, UObject*& OutObject, bool bLoadIfExists) override;
 	virtual void PostLoad() override;
 	//UE_DEPRECATED(5.0, "Use version that takes FObjectPreSaveContext instead.")
 	virtual void PreSave(const class ITargetPlatform* TargetPlatform) override;
diff --git a/Engine/Source/Runtime/Engine/Classes/Engine/World.h b/Engine/Source/Runtime/Engine/Classes/Engine/World.h
--- a/Engine/Source/Runtime/Engine/Classes/Engine/World.h
+++ b/Engine/Source/Runtime/Engine/Classes/Engine/World.h
@@ -2617,9 +2617,9 @@ public:
 #if WITH_EDITOR
 	virtual bool Rename(const TCHAR* NewName = NULL, UObject* NewOuter = NULL, ERenameFlags Flags = REN_None) override;
 	virtual void GetAssetRegistryTags(TArray<FAssetRegistryTag>& OutTags) const override;
-	virtual bool IsNameStableForNetworking() const override;
-	virtual bool LoadSubobject(const TCHAR* SubObjectPath, UObject*& OutObject, bool bOnlyTestExistence) override;
+	virtual bool IsNameStableForNetworking() const override;	
 #endif
+	virtual bool ResolveSubobject(const TCHAR* SubObjectPath, UObject*& OutObject, bool bLoadIfExists) override;
 	virtual void PostDuplicate(bool bDuplicateForPIE) override;
 	//~ End UObject Interface
 	
diff --git a/Engine/Source/Runtime/Engine/Private/Level.cpp b/Engine/Source/Runtime/Engine/Private/Level.cpp
--- a/Engine/Source/Runtime/Engine/Private/Level.cpp
+++ b/Engine/Source/Runtime/Engine/Private/Level.cpp
@@ -2657,23 +2657,6 @@ bool ULevel::CanEditChange(const FProperty* PropertyThatWillChange) const
 	return Super::CanEditChange(PropertyThatWillChange);
 }
 
-bool ULevel::LoadSubobject(const TCHAR* SubObjectPath, UObject*& OutObject, bool bOnlyTestExistence)
-{
-	OutObject = StaticFindObject(nullptr, this, SubObjectPath);
-
-	if (!OutObject)
-	{
-		if (UWorldPartition* WorldPartition = GetWorldPartition())
-		{
-			return WorldPartition->LoadSubobject(SubObjectPath, OutObject, bOnlyTestExistence);
-		}
-
-		return false;
-	}
-
-	return true;
-}
-
 void ULevel::FixupForPIE(int32 InPIEInstanceID, TFunctionRef<void(int32, FSoftObjectPath&)> InCustomFixupFunction)
 {
 	struct FSoftPathPIEFixupSerializer : public FArchiveUObject
@@ -2709,6 +2692,23 @@ void ULevel::FixupForPIE(int32 InPIEInstanceID, TFunctionRef<void(int32, FSoftOb
 
 #endif	//WITH_EDITOR
 
+bool ULevel::ResolveSubobject(const TCHAR* SubObjectPath, UObject*& OutObject, bool bLoadIfExists)
+{
+	OutObject = StaticFindObject(nullptr, this, SubObjectPath);
+
+	if (!OutObject)
+	{
+		if (UWorldPartition* WorldPartition = GetWorldPartition())
+		{
+			return WorldPartition->ResolveSubobject(SubObjectPath, OutObject, bLoadIfExists);
+		}
+
+		return false;
+	}
+
+	return true;
+}
+
 bool ULevel::IsPersistentLevel() const
 {
 	bool bIsPersistent = false;
diff --git a/Engine/Source/Runtime/Engine/Private/World.cpp b/Engine/Source/Runtime/Engine/Private/World.cpp
--- a/Engine/Source/Runtime/Engine/Private/World.cpp
+++ b/Engine/Source/Runtime/Engine/Private/World.cpp
@@ -8063,8 +8063,9 @@ bool UWorld::IsNameStableForNetworking() const
 {
 	return bIsNameStableForNetworking || Super::IsNameStableForNetworking();
 }
+#endif
 
-bool UWorld::LoadSubobject(const TCHAR* SubObjectPath, UObject*& OutObject, bool bOnlyTestExistence)
+bool UWorld::ResolveSubobject(const TCHAR* SubObjectPath, UObject*& OutObject, bool bLoadIfExists)
 {
 	FString SubObjectName;
 	FString SubObjectContext;	
@@ -8072,14 +8073,13 @@ bool UWorld::LoadSubobject(const TCHAR* SubObjectPath, UObject*& OutObject, bool
 	{
 		if (UObject* SubObject = StaticFindObject(nullptr, this, *SubObjectContext))
 		{
-			return SubObject->LoadSubobject(*SubObjectName, OutObject, bOnlyTestExistence);
+			return SubObject->ResolveSubobject(*SubObjectName, OutObject, bLoadIfExists);
 		}
 	}
 
 	OutObject = nullptr;
 	return false;
 }
-#endif
 
 FPrimaryAssetId UWorld::GetPrimaryAssetId() const
 {
diff --git a/Engine/Source/Runtime/Engine/Private/WorldPartition/WorldPartition.cpp b/Engine/Source/Runtime/Engine/Private/WorldPartition/WorldPartition.cpp
--- a/Engine/Source/Runtime/Engine/Private/WorldPartition/WorldPartition.cpp
+++ b/Engine/Source/Runtime/Engine/Private/WorldPartition/WorldPartition.cpp
@@ -186,7 +186,7 @@ void UWorldPartition::OnEndPlay()
 	FlushStreaming();
 	RuntimeHash->OnEndPlay();
 
-	GetStreamingPolicy()->ClearActorToCellRemapping();
+	StreamingPolicy = nullptr;
 
 	World->PersistentLevel->ActorsModifiedForPIE.Empty();
 }
@@ -229,6 +229,8 @@ void UWorldPartition::Initialize(UWorld* InWorld, const FTransform& InTransform)
 	bool bEditorOnly = !World->IsGameWorld();
 	if (bEditorOnly)
 	{
+		check(!StreamingPolicy);
+
 		if (!EditorHash)
 		{
 			UClass* EditorHashClass = FindObject<UClass>(ANY_PACKAGE, TEXT("WorldPartitionEditorSpatialHash"));
@@ -357,7 +359,7 @@ void UWorldPartition::Initialize(UWorld* InWorld, const FTransform& InTransform)
 		{
 			OnBeginPlay(EWorldPartitionStreamingMode::EditorStandalone);
 		}
-		GetStreamingPolicy()->PrepareActorToCellRemapping();
+
 		// Apply remapping of Persistent Level's SoftObjectPaths
 		FWorldPartitionLevelHelper::RemapLevelSoftObjectPaths(World->PersistentLevel, this);
 	}
@@ -548,7 +550,7 @@ const TArray<FWorldPartitionStreamingSource>& UWorldPartition::GetStreamingSourc
 {
 	if (GetWorld()->IsGameWorld())
 	{
-		return GetStreamingPolicy()->GetStreamingSources();
+		return StreamingPolicy->GetStreamingSources();
 	}
 
 	static TArray<FWorldPartitionStreamingSource> EmptyStreamingSources;
@@ -913,50 +915,49 @@ void UWorldPartition::BeginDestroy()
 #endif
 }
 
-#if WITH_EDITOR
-bool UWorldPartition::LoadSubobject(const TCHAR* SubObjectPath, UObject*& OutObject, bool bOnlyTestExistence)
+bool UWorldPartition::ResolveSubobject(const TCHAR* SubObjectPath, UObject*& OutObject, bool bLoadIfExists)
 {
-	if (GetWorld()->IsGameWorld())
-	{
-		if (UObject* SubObject = StreamingPolicy->GetSubObject(SubObjectPath))
-		{
-			OutObject = bOnlyTestExistence ? nullptr : SubObject;
-			return true;
-		}
-	}
-	else
+	if (GetWorld())
 	{
-		for (UActorDescContainer::TIterator<> ActorDescIterator(this); ActorDescIterator; ++ActorDescIterator)
+		if (GetWorld()->IsGameWorld())
 		{
-			FWorldPartitionActorDesc* ActorDesc = *ActorDescIterator;
-
-			if (FString(ActorDesc->ActorPath.ToString()).EndsWith(SubObjectPath))
+			if (StreamingPolicy)
 			{
-				if (!bOnlyTestExistence)
+				if (UObject* SubObject = StreamingPolicy->GetSubObject(SubObjectPath))
 				{
-					LoadedSubobjects.Emplace(this, ActorDesc->GetGuid());
+					OutObject = SubObject;
+					return true;
 				}
+				else
+				{
+					OutObject = nullptr;
+				}
+			}
+		}
+#if WITH_EDITOR
+		else
+		{
+			for (UActorDescContainer::TIterator<> ActorDescIterator(this); ActorDescIterator; ++ActorDescIterator)
+			{
+				FWorldPartitionActorDesc* ActorDesc = *ActorDescIterator;
 
-				OutObject = ActorDescIterator->GetActor();
+				if (FString(ActorDesc->ActorPath.ToString()).EndsWith(SubObjectPath))
+				{
+					if (bLoadIfExists)
+					{
+						LoadedSubobjects.Emplace(this, ActorDesc->GetGuid());
+					}
 
-				return true;
+					OutObject = ActorDescIterator->GetActor();
+					return true;
+				}
 			}
 		}
+#endif
 	}
 
 	return false;
 }
-#endif
-
-UWorldPartitionStreamingPolicy* UWorldPartition::GetStreamingPolicy() const
-{
-	check(GetWorld());
-	if (!StreamingPolicy && GetWorld())
-	{
-		StreamingPolicy = NewObject<UWorldPartitionStreamingPolicy>(const_cast<UWorldPartition*>(this), GetWorld()->GetWorldSettings()->WorldPartitionStreamingPolicyClass.Get());
-	}
-	return StreamingPolicy;
-}
 
 void UWorldPartition::Tick(float DeltaSeconds)
 {
@@ -980,7 +981,7 @@ void UWorldPartition::UpdateStreamingState()
 {
 	if (GetWorld()->IsGameWorld())
 	{
-		GetStreamingPolicy()->UpdateStreamingState();
+		StreamingPolicy->UpdateStreamingState();
 	}
 }
 
@@ -988,7 +989,7 @@ class ULevel* UWorldPartition::GetPreferredLoadedLevelToAddToWorld() const
 {
 	if (GetWorld()->IsGameWorld())
 	{
-		return GetStreamingPolicy()->GetPreferredLoadedLevelToAddToWorld();
+		return StreamingPolicy->GetPreferredLoadedLevelToAddToWorld();
 	}
 	return nullptr;
 }
@@ -997,7 +998,7 @@ bool UWorldPartition::IsStreamingCompleted(EWorldPartitionRuntimeCellState Query
 {
 	if (GetWorld()->IsGameWorld())
 	{
-		return GetStreamingPolicy()->IsStreamingCompleted(QueryState, QuerySources, bExactState);
+		return StreamingPolicy->IsStreamingCompleted(QueryState, QuerySources, bExactState);
 	}
 
 	return false;
@@ -1011,19 +1012,19 @@ bool UWorldPartition::CanDrawRuntimeHash() const
 FVector2D UWorldPartition::GetDrawRuntimeHash2DDesiredFootprint(const FVector2D& CanvasSize)
 {
 	check(CanDrawRuntimeHash());
-	return GetStreamingPolicy()->GetDrawRuntimeHash2DDesiredFootprint(CanvasSize);
+	return StreamingPolicy->GetDrawRuntimeHash2DDesiredFootprint(CanvasSize);
 }
 
 void UWorldPartition::DrawRuntimeHash2D(class UCanvas* Canvas, const FVector2D& PartitionCanvasOffset, const FVector2D& PartitionCanvasSize)
 {
 	check(CanDrawRuntimeHash());
-	GetStreamingPolicy()->DrawRuntimeHash2D(Canvas, PartitionCanvasOffset, PartitionCanvasSize);
+	StreamingPolicy->DrawRuntimeHash2D(Canvas, PartitionCanvasOffset, PartitionCanvasSize);
 }
 
 void UWorldPartition::DrawRuntimeHash3D()
 {
 	check(CanDrawRuntimeHash());
-	GetStreamingPolicy()->DrawRuntimeHash3D();
+	StreamingPolicy->DrawRuntimeHash3D();
 }
 
 #if WITH_EDITOR
@@ -1034,15 +1035,15 @@ void UWorldPartition::DrawRuntimeHashPreview()
 
 bool UWorldPartition::GenerateStreaming(EWorldPartitionStreamingMode Mode, TArray<FString>* OutPackagesToGenerate)
 {
+	check(!StreamingPolicy);
+	StreamingPolicy = NewObject<UWorldPartitionStreamingPolicy>(const_cast<UWorldPartition*>(this), GetWorld()->GetWorldSettings()->WorldPartitionStreamingPolicyClass.Get());
+
 	check(RuntimeHash);
-	bool Result = RuntimeHash->GenerateRuntimeStreaming(Mode, GetStreamingPolicy(), OutPackagesToGenerate);
-	if (Mode == EWorldPartitionStreamingMode::Cook)
-	{
-		// For Cook, prepare actor to cell remapping
-		GetStreamingPolicy()->PrepareActorToCellRemapping();
-		// Apply remapping of Persistent Level's SoftObjectPaths
-		FWorldPartitionLevelHelper::RemapLevelSoftObjectPaths(World->PersistentLevel, this);
-	}
+	bool Result = RuntimeHash->GenerateRuntimeStreaming(Mode, StreamingPolicy, OutPackagesToGenerate);
+
+	// Prepare actor to cell remapping
+	StreamingPolicy->PrepareActorToCellRemapping();
+
 	return Result;
 }
 
@@ -1158,7 +1159,7 @@ void UWorldPartition::CheckForErrors() const
 
 void UWorldPartition::RemapSoftObjectPath(FSoftObjectPath& ObjectPath)
 {
-	GetStreamingPolicy()->RemapSoftObjectPath(ObjectPath);
+	StreamingPolicy->RemapSoftObjectPath(ObjectPath);
 }
 
 FBox UWorldPartition::GetWorldBounds() const
diff --git a/Engine/Source/Runtime/Engine/Private/WorldPartition/WorldPartitionLevelStreamingPolicy.cpp b/Engine/Source/Runtime/Engine/Private/WorldPartition/WorldPartitionLevelStreamingPolicy.cpp
--- a/Engine/Source/Runtime/Engine/Private/WorldPartition/WorldPartitionLevelStreamingPolicy.cpp
+++ b/Engine/Source/Runtime/Engine/Private/WorldPartition/WorldPartitionLevelStreamingPolicy.cpp
@@ -219,7 +219,7 @@ void UWorldPartitionLevelStreamingPolicy::PrepareActorToCellRemapping()
 
 			int32 DotPos = SubObjectName.Find(TEXT("."), ESearchCase::IgnoreCase, ESearchDir::FromEnd);
 			check(DotPos != INDEX_NONE);
-			ActorToCellRemappingForPIE.Add(*SubObjectName.Right(SubObjectName.Len() - DotPos - 1), StreamingCell->GetFName());
+			SubObjectsToCellRemapping.Add(*SubObjectName.Right(SubObjectName.Len() - DotPos - 1), StreamingCell->GetFName());
 		}
 	}
 }
@@ -281,22 +281,32 @@ void UWorldPartitionLevelStreamingPolicy::RemapSoftObjectPath(FSoftObjectPath& O
 		}
 	}
 }
+#endif
 
 UObject* UWorldPartitionLevelStreamingPolicy::GetSubObject(const TCHAR* SubObjectPath)
 {
-	FString SrcPath = UWorld::RemovePIEPrefix(SubObjectPath);
+	// Support for subobjects such as Actor.Component
+	FString SubObjectName;
+	FString SubObjectContext;	
+	if (!FString(SubObjectPath).Split(TEXT("."), &SubObjectContext, &SubObjectName))
+	{
+		SubObjectName = SubObjectPath;
+	}
 
-	if (FName* CellName = ActorToCellRemappingForPIE.Find(FName(*SrcPath)))
+	const FString SrcPath = UWorld::RemovePIEPrefix(*SubObjectName);
+	if (FName* CellName = SubObjectsToCellRemapping.Find(FName(*SrcPath)))
 	{
 		if (UWorldPartitionRuntimeLevelStreamingCell* Cell = (UWorldPartitionRuntimeLevelStreamingCell*)StaticFindObject(UWorldPartitionRuntimeLevelStreamingCell::StaticClass(), GetOuterUWorldPartition(), *(CellName->ToString())))
 		{
-			if (UWorldPartitionLevelStreamingDynamic* LevelStraming = Cell->GetLevelStreaming())
+			if (UWorldPartitionLevelStreamingDynamic* LevelStreaming = Cell->GetLevelStreaming())
 			{
-				return StaticFindObject(UObject::StaticClass(), LevelStraming->GetRuntimeLevel(), SubObjectPath);
+				if (LevelStreaming->GetLoadedLevel())
+				{
+					return StaticFindObject(UObject::StaticClass(), LevelStreaming->GetLoadedLevel(), SubObjectPath);
+				}
 			}
 		}
 	}
 
 	return nullptr;
-}
-#endif
\ No newline at end of file
+}
\ No newline at end of file
diff --git a/Engine/Source/Runtime/Engine/Public/WorldPartition/WorldPartition.h b/Engine/Source/Runtime/Engine/Public/WorldPartition/WorldPartition.h
--- a/Engine/Source/Runtime/Engine/Public/WorldPartition/WorldPartition.h
+++ b/Engine/Source/Runtime/Engine/Public/WorldPartition/WorldPartition.h
@@ -124,12 +124,12 @@ private:
 	const TArray<FWorldPartitionStreamingSource>& GetStreamingSources() const;
 	// UActorPartitionSubsystem interface-
 
-#if WITH_EDITOR
 public:
 	//~ Begin UObject Interface
-	virtual bool LoadSubobject(const TCHAR* SubObjectPath, UObject*& OutObject, bool bOnlyTestExistence) override;
+	virtual bool ResolveSubobject(const TCHAR* SubObjectPath, UObject*& OutObject, bool bLoadIfExists) override;
 	//~ End UObject Interface
 
+#if WITH_EDITOR
 	FName GetWorldPartitionEditorName();
 
 	void LoadEditorCells(const FBox& Box);
@@ -214,7 +214,7 @@ private:
 	EWorldPartitionInitState InitState;
 	FTransform InstanceTransform;
 
-	UPROPERTY(Transient, DuplicateTransient)
+	UPROPERTY()
 	mutable TObjectPtr<UWorldPartitionStreamingPolicy> StreamingPolicy;
 
 	TArray<IWorldPartitionStreamingSourceProvider*> StreamingSourceProviders;
@@ -239,6 +239,5 @@ private:
 	void ApplyActorTransform(AActor* Actor, const FTransform& InTransform);
 #endif
 
-	UWorldPartitionStreamingPolicy* GetStreamingPolicy() const;
 	friend class UWorldPartitionStreamingPolicy;
 };
\ No newline at end of file
diff --git a/Engine/Source/Runtime/Engine/Public/WorldPartition/WorldPartitionLevelStreamingDynamic.h b/Engine/Source/Runtime/Engine/Public/WorldPartition/WorldPartitionLevelStreamingDynamic.h
--- a/Engine/Source/Runtime/Engine/Public/WorldPartition/WorldPartitionLevelStreamingDynamic.h
+++ b/Engine/Source/Runtime/Engine/Public/WorldPartition/WorldPartitionLevelStreamingDynamic.h
@@ -41,8 +41,6 @@ class ENGINE_API UWorldPartitionLevelStreamingDynamic : public ULevelStreamingDy
 	virtual bool RequestLevel(UWorld* PersistentWorld, bool bAllowLevelLoadRequests, EReqLevelBlock BlockPolicy) override;
 	virtual void BeginDestroy() override;
 
-	ULevel* GetRuntimeLevel() const { return RuntimeLevel; }
-
 private:
 	void CreateRuntimeLevel();
 	bool IssueLoadRequests();
diff --git a/Engine/Source/Runtime/Engine/Public/WorldPartition/WorldPartitionLevelStreamingPolicy.h b/Engine/Source/Runtime/Engine/Public/WorldPartition/WorldPartitionLevelStreamingPolicy.h
--- a/Engine/Source/Runtime/Engine/Public/WorldPartition/WorldPartitionLevelStreamingPolicy.h
+++ b/Engine/Source/Runtime/Engine/Public/WorldPartition/WorldPartitionLevelStreamingPolicy.h
@@ -31,18 +31,22 @@ public:
 	virtual void PrepareActorToCellRemapping() override;
 	virtual void ClearActorToCellRemapping() override;
 	virtual void RemapSoftObjectPath(FSoftObjectPath& ObjectPath) override;
-	virtual UObject* GetSubObject(const TCHAR* SubObjectPath) override;
 	static FString GetCellPackagePath(const FName& InCellName, const UWorld* InWorld);
 #endif
 
+	virtual UObject* GetSubObject(const TCHAR* SubObjectPath) override;
+
 private:
 	void SetCellsStateToLoaded(const TSet<const UWorldPartitionRuntimeCell*>& ToLoadCells);
 	void SetCellsStateToActivated(const TSet<const UWorldPartitionRuntimeCell*>& ToActivateCells);
 	void SetCellsStateToUnloaded(const TSet<const UWorldPartitionRuntimeCell*>& ToUnloadCells);
 	int32 GetCellLoadingCount() const;
 
-#if WITH_EDITOR
+#if WITH_EDITORONLY_DATA
+	UPROPERTY()
 	TMap<FName, FName> ActorToCellRemapping;
-	TMap<FName, FName> ActorToCellRemappingForPIE;
 #endif
+
+	UPROPERTY()
+	TMap<FName, FName> SubObjectsToCellRemapping;
 };
\ No newline at end of file
