Fix bulkdata race-condition while building skeletal mesh by reading for disk

diff --git a/Engine/Source/Runtime/Engine/Private/SkeletalMeshRenderData.cpp b/Engine/Source/Runtime/Engine/Private/SkeletalMeshRenderData.cpp
--- a/Engine/Source/Runtime/Engine/Private/SkeletalMeshRenderData.cpp
+++ b/Engine/Source/Runtime/Engine/Private/SkeletalMeshRenderData.cpp
@@ -220,6 +220,13 @@ void FSkeletalMeshRenderData::Cache(const ITargetPlatform* TargetPlatform, USkel
 		Ar << LODModel->RequiredBones;
 		Ar << LODModel->MeshToImportVertexMap;
 		Ar << LODModel->MaxImportVertex;
+
+		// Unless an async loading query is already in flight, we want to load directly from disk to avoid
+		// accessing the linker which is not thread-safe.
+		if (LODModel->RawPointIndices.IsAsyncLoadingComplete() && !LODModel->RawPointIndices.IsBulkDataLoaded())
+		{
+			LODModel->RawPointIndices.LoadBulkDataWithFileReader();
+		}
 		LODModel->RawPointIndices.Serialize(Ar, Owner);
 	};
 
