Fix for static analysis warning

diff --git a/Engine/Source/Runtime/Engine/Private/Particles/FXSystem.cpp b/Engine/Source/Runtime/Engine/Private/Particles/FXSystem.cpp
--- a/Engine/Source/Runtime/Engine/Private/Particles/FXSystem.cpp
+++ b/Engine/Source/Runtime/Engine/Private/Particles/FXSystem.cpp
@@ -98,7 +98,8 @@ bool FFXSystemInterface::GetReferenceAllowGPUUpdate(TConstArrayView<FViewInfo> V
 const FGlobalDistanceFieldParameterData* FFXSystemInterface::GetReferenceGlobalDistanceFieldData(TConstArrayView<FViewInfo> Views)
 {
 	check(Views.Num() > 0);
-	return &Views[0].GlobalDistanceFieldInfo.ParameterData;
+	const FViewInfo& ReferenceView = Views[0];
+	return &ReferenceView.GlobalDistanceFieldInfo.ParameterData;
 }
 
 /*------------------------------------------------------------------------------
