LevelInstance:
- Build > Packed Level Instances now handles dependencies
- Use the bEditableWhenInherited to make sure Instances of the Packed BPs can't have Component Instance data on the ISMs

diff --git a/Engine/Source/Runtime/Engine/Private/LevelInstance/LevelInstanceSubsystem.cpp b/Engine/Source/Runtime/Engine/Private/LevelInstance/LevelInstanceSubsystem.cpp
--- a/Engine/Source/Runtime/Engine/Private/LevelInstance/LevelInstanceSubsystem.cpp
+++ b/Engine/Source/Runtime/Engine/Private/LevelInstance/LevelInstanceSubsystem.cpp
@@ -390,25 +390,48 @@ void ULevelInstanceSubsystem::PackLevelInstances()
 		return;
 	}
 
-	TArray<APackedLevelInstance*> PackedLevelInstancesToUpdate;
-	TSet<UBlueprint*> BlueprintsToUpdate;
-	for (TObjectIterator<UWorld> It(RF_ClassDefaultObject | RF_ArchetypeObject, true, EInternalObjectFlags::PendingKill); It; ++It)
+	// Add Dependencies first so that we pack the LevelInstances in the proper order (depth first)
+	TFunction<void(APackedLevelInstance*, TArray<UBlueprint*>&, TArray<APackedLevelInstance*>&)> GatherDepencenciesRecursive = [&GatherDepencenciesRecursive](APackedLevelInstance* PackedLevelInstance, TArray<UBlueprint*>& BPsToPack, TArray<APackedLevelInstance*>& ToPack)
 	{
-		UWorld* CurrentWorld = *It;
-		for (TActorIterator<ALevelInstance> LevelInstanceIt(GetWorld()); LevelInstanceIt; ++LevelInstanceIt)
+		// Early out on already processed BPs or non BP Packed LIs.
+		UBlueprint* Blueprint = Cast<UBlueprint>(PackedLevelInstance->GetClass()->ClassGeneratedBy);
+		if ((Blueprint && BPsToPack.Contains(Blueprint)) || ToPack.Contains(PackedLevelInstance))
+		{
+			return;
+		}
+		
+		// Recursive deps
+		for (const TSoftObjectPtr<UBlueprint>& Dependency : PackedLevelInstance->PackedBPDependencies)
 		{
-			if (APackedLevelInstance* PackedLevelInstance = Cast<APackedLevelInstance>(*LevelInstanceIt))
+			if (UBlueprint* LoadedDependency = Dependency.LoadSynchronous())
 			{
-				if (UBlueprint* Blueprint = Cast<UBlueprint>(PackedLevelInstance->GetClass()->ClassGeneratedBy))
+				if (APackedLevelInstance* CDO = Cast<APackedLevelInstance>(LoadedDependency->GeneratedClass ? LoadedDependency->GeneratedClass->GetDefaultObject() : nullptr))
 				{
-					BlueprintsToUpdate.Add(Blueprint);
-				}
-				else
-				{
-					PackedLevelInstancesToUpdate.Add(PackedLevelInstance);
+					GatherDepencenciesRecursive(CDO, BPsToPack, ToPack);
 				}
 			}
 		}
+
+		// Add after dependencies
+		if (Blueprint)
+		{
+			BPsToPack.Add(Blueprint);
+		}
+		else
+		{
+			ToPack.Add(PackedLevelInstance);
+		}
+	};
+
+	TArray<APackedLevelInstance*> PackedLevelInstancesToUpdate;
+	TArray<UBlueprint*> BlueprintsToUpdate;
+	for (TObjectIterator<UWorld> It(RF_ClassDefaultObject | RF_ArchetypeObject, true, EInternalObjectFlags::PendingKill); It; ++It)
+	{
+		UWorld* CurrentWorld = *It;
+		for (TActorIterator<APackedLevelInstance> LevelInstanceIt(GetWorld()); LevelInstanceIt; ++LevelInstanceIt)
+		{
+			GatherDepencenciesRecursive(*LevelInstanceIt, BlueprintsToUpdate, PackedLevelInstancesToUpdate);
+		}
 	}
 
 	int32 Count = BlueprintsToUpdate.Num() + PackedLevelInstancesToUpdate.Num();
@@ -416,7 +439,9 @@ void ULevelInstanceSubsystem::PackLevelInstances()
 	{
 		return;
 	}
-		
+	
+	GEditor->SelectNone(true, true);
+
 	FScopedSlowTask SlowTask(Count, (LOCTEXT("LevelInstance_PackLevelInstances", "Packing Level Instances")));
 	SlowTask.MakeDialog();
 		
@@ -424,16 +449,10 @@ void ULevelInstanceSubsystem::PackLevelInstances()
 	{
 		if (SlowTask.CompletedWork < SlowTask.TotalAmountOfWork)
 		{
-			SlowTask.EnterProgressFrame(1, FText::Format(LOCTEXT("LevelInstance_PackLevelInstancesProgress", "Packing Level Instance {0} of {1})"), FText::AsNumber(SlowTask.CompletedWork), FText::AsNumber(SlowTask.TotalAmountOfWork)));
+			SlowTask.EnterProgressFrame(1, FText::Format(LOCTEXT("LevelInstance_PackLevelInstancesProgress", "Packing Level Instance {0} of {1}"), FText::AsNumber(SlowTask.CompletedWork), FText::AsNumber(SlowTask.TotalAmountOfWork)));
 		}
 	};
 
-	for (APackedLevelInstance* PackedLevelInstance : PackedLevelInstancesToUpdate)
-	{
-		PackedLevelInstance->OnWorldAssetChanged();
-		UpdateProgress();
-	}
-
 	TSharedPtr<FPackedLevelInstanceBuilder> Builder = FPackedLevelInstanceBuilder::CreateDefaultBuilder();
 	const bool bSilentUpdate = true;
 	for (UBlueprint* Blueprint : BlueprintsToUpdate)
@@ -441,6 +460,12 @@ void ULevelInstanceSubsystem::PackLevelInstances()
 		Builder->UpdateBlueprint(Blueprint, bSilentUpdate);
 		UpdateProgress();
 	}
+
+	for (APackedLevelInstance* PackedLevelInstance : PackedLevelInstancesToUpdate)
+	{
+		PackedLevelInstance->OnWorldAssetChanged();
+		UpdateProgress();
+	}
 }
 
 bool ULevelInstanceSubsystem::GetLevelInstanceBounds(const ALevelInstance* LevelInstanceActor, FBox& OutBounds) const
diff --git a/Engine/Source/Runtime/Engine/Private/LevelInstance/Packed/LevelInstanceRecursivePacker.cpp b/Engine/Source/Runtime/Engine/Private/LevelInstance/Packed/LevelInstanceRecursivePacker.cpp
--- a/Engine/Source/Runtime/Engine/Private/LevelInstance/Packed/LevelInstanceRecursivePacker.cpp
+++ b/Engine/Source/Runtime/Engine/Private/LevelInstance/Packed/LevelInstanceRecursivePacker.cpp
@@ -50,9 +50,12 @@ void FLevelInstanceRecursivePacker::PackActors(FPackedLevelInstanceBuilderContex
 	FLevelInstanceRecursivePackerCluster* LevelInstanceCluster = (FLevelInstanceRecursivePackerCluster*)InClusterID.GetData();
 	check(LevelInstanceCluster);
 
-	if (LevelInstanceCluster->LevelInstance->IsLevelInstancePathValid())
+	if (LevelInstanceCluster->LevelInstance->IsA<APackedLevelInstance>())
 	{
-		InPackingActor->PackDependencies.AddUnique(LevelInstanceCluster->LevelInstance->GetWorldAsset());
+		if (UBlueprint* GeneratedBy = Cast<UBlueprint>(LevelInstanceCluster->LevelInstance->GetClass()->ClassGeneratedBy))
+		{
+			InPackingActor->PackedBPDependencies.AddUnique(GeneratedBy);
+		}
 	}
 }
 
diff --git a/Engine/Source/Runtime/Engine/Private/LevelInstance/Packed/PackedLevelInstanceBuilder.cpp b/Engine/Source/Runtime/Engine/Private/LevelInstance/Packed/PackedLevelInstanceBuilder.cpp
--- a/Engine/Source/Runtime/Engine/Private/LevelInstance/Packed/PackedLevelInstanceBuilder.cpp
+++ b/Engine/Source/Runtime/Engine/Private/LevelInstance/Packed/PackedLevelInstanceBuilder.cpp
@@ -316,12 +316,12 @@ ALevelInstance* FPackedLevelInstanceBuilder::CreateTransientLevelInstanceForPack
 bool FPackedLevelInstanceBuilder::PackActor(APackedLevelInstance* InActor, TSoftObjectPtr<UWorld> InWorldAsset)
 {
 	ALevelInstance* TransientLevelInstance = CreateTransientLevelInstanceForPacking(InWorldAsset, InActor->GetActorLocation(), InActor->GetActorRotation());
-
-	bool bResult = PackActor(InActor, TransientLevelInstance);
-
-	TransientLevelInstance->GetWorld()->DestroyActor(TransientLevelInstance);
-
-	return bResult;
+	ON_SCOPE_EXIT
+	{
+		TransientLevelInstance->GetWorld()->DestroyActor(TransientLevelInstance);
+	};
+	
+	return PackActor(InActor, TransientLevelInstance);
 }
 
 void FPackedLevelInstanceBuilder::UpdateBlueprint(UBlueprint* Blueprint, bool bInSilentUpdate)
@@ -425,11 +425,17 @@ bool FPackedLevelInstanceBuilder::CreateOrUpdateBlueprintFromPacked(APackedLevel
 	// Avoid running construction script while dragging an instance of that BP for performance reasons
 	BP->bRunConstructionScriptOnDrag = false;
 	APackedLevelInstance* CDO = CastChecked<APackedLevelInstance>(BP->GeneratedClass->GetDefaultObject());
-	CDO->SetWorldAsset(InActor->GetWorldAsset());
-	
-	// match root component mobility to source actor
-	USceneComponent* Root = CDO->GetRootComponent();
-	Root->SetMobility(InActor->GetRootComponent()->Mobility);
+	auto PropagatePropertiesToActor = [InActor](APackedLevelInstance* TargetActor)
+	{
+		TargetActor->Modify(false);
+		TargetActor->SetWorldAsset(InActor->GetWorldAsset());
+		TargetActor->PackedBPDependencies = InActor->PackedBPDependencies;
+
+		// match root component mobility to source actor
+		USceneComponent* Root = TargetActor->GetRootComponent();
+		Root->SetMobility(InActor->GetRootComponent()->Mobility);
+	};
+	PropagatePropertiesToActor(CDO);
 		
 	// Prep AddComponentsToBlueprintParam
 	FKismetEditorUtilities::FAddComponentsToBlueprintParams AddCompToBPParams;
@@ -439,6 +445,13 @@ bool FPackedLevelInstanceBuilder::CreateOrUpdateBlueprintFromPacked(APackedLevel
 	// Add Components
 	TArray<UActorComponent*> PackedComponents;
 	InActor->GetPackedComponents(PackedComponents);
+	
+	// To avoid any delta serialization happening on those generated components, we make them non editable.
+	for (UActorComponent* PackedComponent : PackedComponents)
+	{
+		PackedComponent->bEditableWhenInherited = false;
+	}
+
 	FKismetEditorUtilities::AddComponentsToBlueprint(BP, PackedComponents, AddCompToBPParams);
 	// If we are packing the actors BP then destroy packed components as they are now part of the BPs construction script
 	if (UBlueprint* GeneratedBy = Cast<UBlueprint>(InActor->GetClass()->ClassGeneratedBy))
@@ -446,6 +459,18 @@ bool FPackedLevelInstanceBuilder::CreateOrUpdateBlueprintFromPacked(APackedLevel
 		InActor->DestroyPackedComponents();
 	}
 
+	// Propagate properties before BP Compilation so that they are considered default (no delta)
+	if (BP->GeneratedClass)
+	{
+		TArray<UObject*> ObjectsOfClass;
+		GetObjectsOfClass(BP->GeneratedClass, ObjectsOfClass, true, RF_ClassDefaultObject | RF_ArchetypeObject, EInternalObjectFlags::PendingKill);
+		for (UObject* ObjectOfClass : ObjectsOfClass)
+		{
+			APackedLevelInstance* PackedLevelInstance = CastChecked<APackedLevelInstance>(ObjectOfClass);
+			PropagatePropertiesToActor(PackedLevelInstance);
+		}
+	}
+			
 	// Synchronous compile
 	FKismetEditorUtilities::CompileBlueprint(BP, EBlueprintCompileOptions::SkipGarbageCollection);
 
diff --git a/Engine/Source/Runtime/Engine/Public/LevelInstance/Packed/PackedLevelInstanceActor.h b/Engine/Source/Runtime/Engine/Public/LevelInstance/Packed/PackedLevelInstanceActor.h
--- a/Engine/Source/Runtime/Engine/Public/LevelInstance/Packed/PackedLevelInstanceActor.h
+++ b/Engine/Source/Runtime/Engine/Public/LevelInstance/Packed/PackedLevelInstanceActor.h
@@ -79,7 +79,7 @@ public:
 	bool bRerunConstructionScripts;
 
 	UPROPERTY()
-	TArray<TSoftObjectPtr<UWorld>> PackDependencies;
+	TArray<TSoftObjectPtr<UBlueprint>> PackedBPDependencies;
 
 private:
 	UPROPERTY(NonTransactional)
