Fixed issue where PathOffsetCurves in root motion sources were mixing time and motion fractions. Now the path is mapped to the motion duration. There was a similar issue last week with time-mapping curves fixed the same way. Added static util functions to reduce duplicate code for these curve evaluations.

diff --git a/Engine/Source/Runtime/Engine/Classes/GameFramework/RootMotionSource.h b/Engine/Source/Runtime/Engine/Classes/GameFramework/RootMotionSource.h
--- a/Engine/Source/Runtime/Engine/Classes/GameFramework/RootMotionSource.h
+++ b/Engine/Source/Runtime/Engine/Classes/GameFramework/RootMotionSource.h
@@ -552,7 +552,7 @@ struct ENGINE_API FRootMotionSource_MoveToForce : public FRootMotionSource
 	UPROPERTY()
 	TObjectPtr<UCurveVector> PathOffsetCurve;
 
-	FVector GetPathOffsetInWorldSpace(float MoveFraction) const;
+	FVector GetPathOffsetInWorldSpace(const float MoveFraction) const;
 
 	virtual FRootMotionSource* Clone() const override;
 
@@ -624,7 +624,7 @@ struct ENGINE_API FRootMotionSource_MoveToDynamicForce : public FRootMotionSourc
 
 	void SetTargetLocation(FVector NewTargetLocation);
 
-	FVector GetPathOffsetInWorldSpace(float MoveFraction) const;
+	FVector GetPathOffsetInWorldSpace(const float MoveFraction) const;
 
 	virtual FRootMotionSource* Clone() const override;
 
diff --git a/Engine/Source/Runtime/Engine/Private/GameFramework/RootMotionSource.cpp b/Engine/Source/Runtime/Engine/Private/GameFramework/RootMotionSource.cpp
--- a/Engine/Source/Runtime/Engine/Private/GameFramework/RootMotionSource.cpp
+++ b/Engine/Source/Runtime/Engine/Private/GameFramework/RootMotionSource.cpp
@@ -68,6 +68,26 @@ void RootMotionSourceDebug::PrintOnScreenServerMsg(const FString& InString)
 
 const float RootMotionSource_InvalidStartTime = -BIG_NUMBER;
 
+
+static float EvaluateFloatCurveAtFraction(const UCurveFloat& Curve, const float Fraction)
+{
+	float MinCurveTime(0.f);
+	float MaxCurveTime(1.f);
+
+	Curve.GetTimeRange(MinCurveTime, MaxCurveTime);
+	return Curve.GetFloatValue(FMath::GetRangeValue(FVector2D(MinCurveTime, MaxCurveTime), Fraction));
+}
+
+static FVector EvaluateVectorCurveAtFraction(const UCurveVector& Curve, const float Fraction)
+{
+	float MinCurveTime(0.f);
+	float MaxCurveTime(1.f);
+
+	Curve.GetTimeRange(MinCurveTime, MaxCurveTime);
+	return Curve.GetVectorValue(FMath::GetRangeValue(FVector2D(MinCurveTime, MaxCurveTime), Fraction));
+}
+
+
 //
 // FRootMotionServerToLocalIDMapping
 //
@@ -693,12 +713,12 @@ void FRootMotionSource_MoveToForce::SetTime(float NewTime)
 	// TODO-RootMotionSource: Check if reached destination?
 }
 
-FVector FRootMotionSource_MoveToForce::GetPathOffsetInWorldSpace(float MoveFraction) const
+FVector FRootMotionSource_MoveToForce::GetPathOffsetInWorldSpace(const float MoveFraction) const
 {
 	if (PathOffsetCurve)
 	{
 		// Calculate path offset
-		const FVector PathOffsetInFacingSpace = PathOffsetCurve->GetVectorValue(MoveFraction);
+		const FVector PathOffsetInFacingSpace = EvaluateVectorCurveAtFraction(*PathOffsetCurve, MoveFraction);
 		FRotator FacingRotation((TargetLocation-StartLocation).Rotation());
 		FacingRotation.Pitch = 0.f; // By default we don't include pitch in the offset, but an option could be added if necessary
 		return FacingRotation.RotateVector(PathOffsetInFacingSpace);
@@ -881,12 +901,12 @@ void FRootMotionSource_MoveToDynamicForce::SetTime(float NewTime)
 	// TODO-RootMotionSource: Check if reached destination?
 }
 
-FVector FRootMotionSource_MoveToDynamicForce::GetPathOffsetInWorldSpace(float MoveFraction) const
+FVector FRootMotionSource_MoveToDynamicForce::GetPathOffsetInWorldSpace(const float MoveFraction) const
 {
 	if (PathOffsetCurve)
 	{
 		// Calculate path offset
-		const FVector PathOffsetInFacingSpace = PathOffsetCurve->GetVectorValue(MoveFraction);
+		const FVector PathOffsetInFacingSpace = EvaluateVectorCurveAtFraction(*PathOffsetCurve, MoveFraction);
 		FRotator FacingRotation((TargetLocation-StartLocation).Rotation());
 		FacingRotation.Pitch = 0.f; // By default we don't include pitch in the offset, but an option could be added if necessary
 		return FacingRotation.RotateVector(PathOffsetInFacingSpace);
@@ -909,14 +929,9 @@ void FRootMotionSource_MoveToDynamicForce::PrepareRootMotion
 	{
 		float MoveFraction = (GetTime() + SimulationTime) / Duration;
 		
-		// only used if TimeMappingCurve is present
-		float MinMappingCurveTime(0.f);
-		float MaxMappingCurveTime(1.f); 
-
 		if (TimeMappingCurve)
 		{
-			TimeMappingCurve->GetTimeRange(MinMappingCurveTime, MaxMappingCurveTime);
-			MoveFraction = TimeMappingCurve->GetFloatValue(FMath::GetRangeValue(FVector2D(MinMappingCurveTime, MaxMappingCurveTime), MoveFraction));
+			MoveFraction = EvaluateFloatCurveAtFraction(*TimeMappingCurve, MoveFraction);
 		}
 
 		FVector CurrentTargetLocation = FMath::Lerp<FVector, float>(StartLocation, TargetLocation, MoveFraction);
@@ -932,7 +947,7 @@ void FRootMotionSource_MoveToDynamicForce::PrepareRootMotion
 			float PreviousMoveFraction = GetTime() / Duration;
 			if (TimeMappingCurve)
 			{
-				PreviousMoveFraction = TimeMappingCurve->GetFloatValue(FMath::GetRangeValue(FVector2D(MinMappingCurveTime, MaxMappingCurveTime), PreviousMoveFraction));
+				PreviousMoveFraction = EvaluateFloatCurveAtFraction(*TimeMappingCurve, PreviousMoveFraction);
 			}
 
 			FVector CurrentExpectedLocation = FMath::Lerp<FVector, float>(StartLocation, TargetLocation, PreviousMoveFraction);
@@ -1094,13 +1109,13 @@ bool FRootMotionSource_JumpForce::UpdateStateFrom(const FRootMotionSource* Sourc
 	return true; // JumpForce has no unique state other than Time which is handled by FRootMotionSource
 }
 
-FVector FRootMotionSource_JumpForce::GetPathOffset(float MoveFraction) const
+FVector FRootMotionSource_JumpForce::GetPathOffset(const float MoveFraction) const
 {
 	FVector PathOffset(FVector::ZeroVector);
 	if (PathOffsetCurve)
 	{
 		// Calculate path offset
-		PathOffset = PathOffsetCurve->GetVectorValue(MoveFraction);
+		PathOffset = EvaluateVectorCurveAtFraction(*PathOffsetCurve, MoveFraction);
 	}
 	else
 	{
@@ -1161,11 +1176,8 @@ void FRootMotionSource_JumpForce::PrepareRootMotion
 
 		if (TimeMappingCurve)
 		{
-			float MinMappingCurveTime, MaxMappingCurveTime;
-			TimeMappingCurve->GetTimeRange(MinMappingCurveTime, MaxMappingCurveTime);
-
-			CurrentMoveFraction = TimeMappingCurve->GetFloatValue(FMath::GetRangeValue(FVector2D(MinMappingCurveTime, MaxMappingCurveTime), CurrentMoveFraction));
-			TargetMoveFraction  = TimeMappingCurve->GetFloatValue(FMath::GetRangeValue(FVector2D(MinMappingCurveTime, MaxMappingCurveTime), TargetMoveFraction));
+			CurrentMoveFraction = EvaluateFloatCurveAtFraction(*TimeMappingCurve, CurrentMoveFraction);
+			TargetMoveFraction  = EvaluateFloatCurveAtFraction(*TimeMappingCurve, TargetMoveFraction);
 		}
 
 		const FVector CurrentRelativeLocation = GetRelativeLocation(CurrentMoveFraction);
