Initial support for new HLSL material translator for materials not using control flow

diff --git a/Engine/Shaders/Public/ShaderVersion.ush b/Engine/Shaders/Public/ShaderVersion.ush
--- a/Engine/Shaders/Public/ShaderVersion.ush
+++ b/Engine/Shaders/Public/ShaderVersion.ush
@@ -2,4 +2,4 @@
 // Each time the console command is executed it generates a new GUID. As this file is included
 // in Platform.ush (which should be included in any shader) it allows to invalidate the shader DDC.
 // 
-// GUID = 4B4392B38CAB46E39C0C10B2CA18FF52
+// GUID = 4B4392B38CAB46E39C0C10B2CA18FF52_A
diff --git a/Engine/Source/Runtime/Engine/Classes/Materials/Material.h b/Engine/Source/Runtime/Engine/Classes/Materials/Material.h
--- a/Engine/Source/Runtime/Engine/Classes/Materials/Material.h
+++ b/Engine/Source/Runtime/Engine/Classes/Materials/Material.h
@@ -17,6 +17,7 @@
 #include "Materials/MaterialFunction.h"
 #include "Materials/MaterialLayersFunctions.h"
 #include "Templates/UniquePtr.h"
+#include "Shader/ShaderTypes.h"
 #if WITH_CHAOS
 #include "Physics/PhysicsInterfaceCore.h"
 #endif
@@ -273,6 +274,15 @@ struct FVector2MaterialInput : public FMaterialInput
 };
 #endif
 
+struct FMaterialInputDescription
+{
+	FExpressionInput* Input = nullptr;
+	UE::Shader::EValueType Type = UE::Shader::EValueType::Void;
+	UE::Shader::FValue ConstantValue;
+	bool bUseConstant = false;
+	bool bHidden = false;
+};
+
 USTRUCT()
 struct FParameterGroupData
 {
@@ -1668,6 +1678,8 @@ public:
 	*									or NULL if an invalid property was requested (some properties have been removed from UI, those return NULL).
 	*/
 	ENGINE_API FExpressionInput* GetExpressionInputForProperty(EMaterialProperty InProperty);
+
+	ENGINE_API bool GetExpressionInputDescription(EMaterialProperty InProperty, FMaterialInputDescription& OutDescription);
 #endif
 
 #if WITH_EDITORONLY_DATA
diff --git a/Engine/Source/Runtime/Engine/Private/HLSLTree/HLSLTreeCommon.cpp b/Engine/Source/Runtime/Engine/Private/HLSLTree/HLSLTreeCommon.cpp
--- a/Engine/Source/Runtime/Engine/Private/HLSLTree/HLSLTreeCommon.cpp
+++ b/Engine/Source/Runtime/Engine/Private/HLSLTree/HLSLTreeCommon.cpp
@@ -273,12 +273,9 @@ void UE::HLSLTree::FExpressionDefaultMaterialAttributes::EmitCode(FEmitContext&
 
 void UE::HLSLTree::FExpressionSetMaterialAttribute::EmitCode(FEmitContext& Context, FExpressionEmitResult& OutResult) const
 {
-	const EMaterialProperty Property = FMaterialAttributeDefinitionMap::GetProperty(AttributeID);
 	const FString PropertyName = FMaterialAttributeDefinitionMap::GetAttributeName(AttributeID);
-	const EMaterialValueType PropertyType = FMaterialAttributeDefinitionMap::GetValueType(AttributeID);
-	const EShaderFrequency Frequency = FMaterialAttributeDefinitionMap::GetShaderFrequency(AttributeID);
 
-	const FEmitValue&AttributresRef = Context.AcquireValue(AttributesExpression);
+	const FEmitValue& AttributresRef = Context.AcquireValue(AttributesExpression);
 	const FEmitValue&ValueRef = Context.AcquireValue(ValueExpression);
 
 	OutResult.EvaluationType = EExpressionEvaluationType::Shader;
diff --git a/Engine/Source/Runtime/Engine/Private/Materials/Material.cpp b/Engine/Source/Runtime/Engine/Private/Materials/Material.cpp
--- a/Engine/Source/Runtime/Engine/Private/Materials/Material.cpp
+++ b/Engine/Source/Runtime/Engine/Private/Materials/Material.cpp
@@ -5640,39 +5640,114 @@ void UMaterial::GetReferencedTexturesAndOverrides(TSet<const UTexture*>& InOutTe
 
 FExpressionInput* UMaterial::GetExpressionInputForProperty(EMaterialProperty InProperty)
 {
-	switch (InProperty)
+	FMaterialInputDescription Description;
+	if (GetExpressionInputDescription(InProperty, Description))
 	{
-		case MP_EmissiveColor:			return &EmissiveColor;
-		case MP_Opacity:				return &Opacity;
-		case MP_OpacityMask:			return &OpacityMask;
-		case MP_BaseColor:				return &BaseColor;
-		case MP_Metallic:				return &Metallic;
-		case MP_Specular:				return &Specular;
-		case MP_Roughness:				return &Roughness;
-		case MP_Anisotropy:				return &Anisotropy;
-		case MP_Normal:					return &Normal;
-		case MP_Tangent:				return &Tangent;
-		case MP_WorldPositionOffset:	return &WorldPositionOffset;
-		case MP_WorldDisplacement:		return &WorldDisplacement;
-		case MP_TessellationMultiplier:	return &TessellationMultiplier;
-		case MP_SubsurfaceColor:		return &SubsurfaceColor;
-		case MP_CustomData0:			return &ClearCoat;
-		case MP_CustomData1:			return &ClearCoatRoughness;
-		case MP_AmbientOcclusion:		return &AmbientOcclusion;
-		case MP_Refraction:				return &Refraction;
-		case MP_MaterialAttributes:		return &MaterialAttributes;
-		case MP_PixelDepthOffset:		return &PixelDepthOffset;
-		case MP_ShadingModel:			return &ShadingModelFromMaterialExpression;
-		case MP_FrontMaterial:			return &FrontMaterial;
-	}
-
-	if (InProperty >= MP_CustomizedUVs0 && InProperty <= MP_CustomizedUVs7)
-	{
-		return &CustomizedUVs[InProperty - MP_CustomizedUVs0];
+		if (!Description.bHidden)
+		{
+			return Description.Input;
+		}
 	}
-
 	return nullptr;
 }
+
+static void SetMaterialInputDescription(FColorMaterialInput& Input, bool bHidden, FMaterialInputDescription& OutDescription)
+{
+	OutDescription.Type = UE::Shader::EValueType::Float3;
+	OutDescription.Input = &Input;
+	OutDescription.bUseConstant = Input.UseConstant;
+	OutDescription.bHidden = bHidden;
+	const FLinearColor ConstantColor(Input.Constant);
+	OutDescription.ConstantValue = UE::Shader::FValue(ConstantColor.R, ConstantColor.G, ConstantColor.B);
+}
+
+static void SetMaterialInputDescription(FVectorMaterialInput& Input, bool bHidden, FMaterialInputDescription& OutDescription)
+{
+	OutDescription.Type = UE::Shader::EValueType::Float3;
+	OutDescription.Input = &Input;
+	OutDescription.bUseConstant = Input.UseConstant;
+	OutDescription.bHidden = bHidden;
+	OutDescription.ConstantValue = Input.Constant;
+}
+
+static void SetMaterialInputDescription(FVector2MaterialInput& Input, bool bHidden, FMaterialInputDescription& OutDescription)
+{
+	OutDescription.Type = UE::Shader::EValueType::Float2;
+	OutDescription.Input = &Input;
+	OutDescription.bUseConstant = Input.UseConstant;
+	OutDescription.bHidden = bHidden;
+	OutDescription.ConstantValue = Input.Constant;
+}
+
+static void SetMaterialInputDescription(FScalarMaterialInput& Input, bool bHidden, FMaterialInputDescription& OutDescription)
+{
+	OutDescription.Type = UE::Shader::EValueType::Float1;
+	OutDescription.Input = &Input;
+	OutDescription.bUseConstant = Input.UseConstant;
+	OutDescription.bHidden = bHidden;
+	OutDescription.ConstantValue = Input.Constant;
+}
+
+static void SetMaterialInputDescription(FShadingModelMaterialInput& Input, bool bHidden, FMaterialInputDescription& OutDescription)
+{
+	OutDescription.Type = UE::Shader::EValueType::Void;
+	OutDescription.Input = &Input;
+	OutDescription.bUseConstant = false;
+	OutDescription.bHidden = bHidden;
+}
+
+static void SetMaterialInputDescription(FMaterialAttributesInput& Input, bool bHidden, FMaterialInputDescription& OutDescription)
+{
+	OutDescription.Type = UE::Shader::EValueType::MaterialAttributes;
+	OutDescription.Input = &Input;
+	OutDescription.bUseConstant = false;
+	OutDescription.bHidden = bHidden;
+}
+
+static void SetMaterialInputDescription(FStrataMaterialInput& Input, bool bHidden, FMaterialInputDescription& OutDescription)
+{
+	OutDescription.Type = UE::Shader::EValueType::Void;
+	OutDescription.Input = &Input;
+	OutDescription.bUseConstant = false;
+	OutDescription.bHidden = bHidden;
+}
+
+bool UMaterial::GetExpressionInputDescription(EMaterialProperty InProperty, FMaterialInputDescription& OutDescription)
+{
+	switch (InProperty)
+	{
+	case MP_EmissiveColor: SetMaterialInputDescription(EmissiveColor, false, OutDescription); return true;
+	case MP_Opacity: SetMaterialInputDescription(Opacity, false, OutDescription); return true;
+	case MP_OpacityMask: SetMaterialInputDescription(OpacityMask, false, OutDescription); return true;
+	case MP_BaseColor: SetMaterialInputDescription(BaseColor, false, OutDescription); return true;
+	case MP_Metallic: SetMaterialInputDescription(Metallic, false, OutDescription); return true;
+	case MP_Specular: SetMaterialInputDescription(Specular, false, OutDescription); return true;
+	case MP_Roughness: SetMaterialInputDescription(Roughness, false, OutDescription); return true;
+	case MP_Anisotropy: SetMaterialInputDescription(Anisotropy, false, OutDescription); return true;
+	case MP_Normal: SetMaterialInputDescription(Normal, false, OutDescription); return true;
+	case MP_Tangent: SetMaterialInputDescription(Tangent, false, OutDescription); return true;
+	case MP_WorldPositionOffset: SetMaterialInputDescription(WorldPositionOffset, false, OutDescription); return true;
+	case MP_WorldDisplacement: SetMaterialInputDescription(WorldDisplacement, false, OutDescription); return true;
+	case MP_TessellationMultiplier: SetMaterialInputDescription(TessellationMultiplier, false, OutDescription); return true;
+	case MP_SubsurfaceColor: SetMaterialInputDescription(SubsurfaceColor, false, OutDescription); return true;
+	case MP_CustomData0: SetMaterialInputDescription(ClearCoat, false, OutDescription); return true;
+	case MP_CustomData1: SetMaterialInputDescription(ClearCoatRoughness, false, OutDescription); return true;
+	case MP_AmbientOcclusion: SetMaterialInputDescription(AmbientOcclusion, false, OutDescription); return true;
+	case MP_Refraction: SetMaterialInputDescription(Refraction, false, OutDescription); return true;
+	case MP_MaterialAttributes: SetMaterialInputDescription(MaterialAttributes, false, OutDescription); return true;
+	case MP_PixelDepthOffset: SetMaterialInputDescription(PixelDepthOffset, false, OutDescription); return true;
+	case MP_ShadingModel: SetMaterialInputDescription(ShadingModelFromMaterialExpression, false, OutDescription); return true;
+	case MP_FrontMaterial: SetMaterialInputDescription(FrontMaterial, false, OutDescription); return true;
+	default:
+		if (InProperty >= MP_CustomizedUVs0 && InProperty <= MP_CustomizedUVs7)
+		{
+			SetMaterialInputDescription(CustomizedUVs[InProperty - MP_CustomizedUVs0], false, OutDescription);
+			return true;
+		}
+		return false;
+	}
+}
+
 #endif // WITH_EDITOR
 
 #if WITH_EDITORONLY_DATA
diff --git a/Engine/Source/Runtime/Engine/Private/Materials/MaterialHLSLTree.cpp b/Engine/Source/Runtime/Engine/Private/Materials/MaterialHLSLTree.cpp
--- a/Engine/Source/Runtime/Engine/Private/Materials/MaterialHLSLTree.cpp
+++ b/Engine/Source/Runtime/Engine/Private/Materials/MaterialHLSLTree.cpp
@@ -16,16 +16,89 @@ FMaterialHLSLTree::FMaterialHLSLTree() : Allocator(0)
 {
 }
 
+static UE::HLSLTree::FExpression* CompileMaterialInput(FMaterialHLSLGenerator& Generator,
+	UE::HLSLTree::FScope& Scope,
+	EMaterialProperty InputProperty,
+	UMaterial* Material,
+	UE::HLSLTree::FExpression* AttributesExpression)
+{
+	UE::HLSLTree::FExpression* Expression = nullptr;
+	if (Material->IsPropertyActive(InputProperty))
+	{
+		FMaterialInputDescription InputDescription;
+		if (Material->GetExpressionInputDescription(InputProperty, InputDescription))
+		{
+			if (InputDescription.bUseConstant)
+			{
+				UE::Shader::FValue DefaultValue = FMaterialAttributeDefinitionMap::GetDefaultValue(InputProperty);
+				DefaultValue.NumComponents = UE::Shader::GetValueTypeDescription(InputDescription.Type).NumComponents;
+				if (InputDescription.ConstantValue != DefaultValue)
+				{
+					Expression = Generator.NewConstant(Scope, InputDescription.ConstantValue);
+				}
+			}
+			else
+			{
+				check(InputDescription.Input);
+				if (InputProperty >= MP_CustomizedUVs0 && InputProperty <= MP_CustomizedUVs7)
+				{
+					const int32 TexCoordIndex = (int32)InputProperty - MP_CustomizedUVs0;
+					if (TexCoordIndex < Material->NumCustomizedUVs)
+					{
+						Expression = InputDescription.Input->AcquireHLSLExpressionWithCast(Generator, Scope, InputDescription.Type);
+					}
+					if (!Expression)
+					{
+						Expression = Generator.NewTexCoord(Scope, TexCoordIndex);
+					}
+				}
+				else
+				{
+					Expression = InputDescription.Input->AcquireHLSLExpressionWithCast(Generator, Scope, InputDescription.Type);
+				}
+			}
+		}
+	}
+
+	if (Expression)
+	{
+		UE::HLSLTree::FExpressionSetMaterialAttribute* SetAttributeExpression = Generator.GetTree().NewExpression<UE::HLSLTree::FExpressionSetMaterialAttribute>(Scope);
+		SetAttributeExpression->AttributeID = FMaterialAttributeDefinitionMap::GetID(InputProperty);
+		SetAttributeExpression->AttributesExpression = AttributesExpression;
+		SetAttributeExpression->ValueExpression = Expression;
+		return SetAttributeExpression;
+	}
+
+	return AttributesExpression;
+}
+
 void FMaterialHLSLTree::InitializeForMaterial(const FMaterialCompileTargetParameters& InCompilerTarget, FMaterial& InOutMaterial)
 {
-	const UMaterialInterface* MaterialInterface = InOutMaterial.GetMaterialInterface();
-	const UMaterial* BaseMaterial = MaterialInterface->GetMaterial();
-	UMaterialExpression* BaseExpression = BaseMaterial->ExpressionExecBegin;
+	UMaterialInterface* MaterialInterface = InOutMaterial.GetMaterialInterface();
+	UMaterial* BaseMaterial = MaterialInterface->GetMaterial();
 
 	HLSLTree = UE::HLSLTree::FTree::Create(Allocator);
-
 	FMaterialHLSLGenerator Generator(InCompilerTarget, *HLSLTree);
-	Generator.AcquireStatement(HLSLTree->GetRootScope(), BaseExpression);
+	UE::HLSLTree::FScope& RootScope = HLSLTree->GetRootScope();
+
+	if (BaseMaterial->IsCompiledWithExecutionFlow())
+	{
+		UMaterialExpression* BaseExpression = BaseMaterial->ExpressionExecBegin;
+		Generator.AcquireStatement(RootScope, BaseExpression);
+	}
+	else
+	{
+		UE::HLSLTree::FExpression* AttributesExpression = Generator.GetTree().NewExpression<UE::HLSLTree::FExpressionDefaultMaterialAttributes>(RootScope);
+		for (int32 PropertyIndex = 0; PropertyIndex < MP_MAX; ++PropertyIndex)
+		{
+			const EMaterialProperty Property = (EMaterialProperty)PropertyIndex;
+			AttributesExpression = CompileMaterialInput(Generator, RootScope, Property, BaseMaterial, AttributesExpression);
+		}
+
+		UE::HLSLTree::FStatementReturn* ReturnStatement = Generator.GetTree().NewStatement<UE::HLSLTree::FStatementReturn>(RootScope);
+		ReturnStatement->Expression = AttributesExpression;
+	}
+
 	Generator.AcquireErrors(InOutMaterial.CompileErrors, InOutMaterial.ErrorExpressions);
 }
 
diff --git a/Engine/Source/Runtime/Engine/Private/Shader/ShaderValue.cpp b/Engine/Source/Runtime/Engine/Private/Shader/ShaderValue.cpp
--- a/Engine/Source/Runtime/Engine/Private/Shader/ShaderValue.cpp
+++ b/Engine/Source/Runtime/Engine/Private/Shader/ShaderValue.cpp
@@ -385,6 +385,22 @@ inline FValue BinaryOp(const Operation& Op, const FValue& Lhs, const FValue& Rhs
 } // namespace Shader
 } // namespace UE
 
+bool UE::Shader::operator==(const FValue& Lhs, const FValue& Rhs)
+{
+	if (Lhs.ComponentType != Rhs.ComponentType || Lhs.NumComponents != Rhs.NumComponents)
+	{
+		return false;
+	}
+	for (int32 i = 0u; i < Lhs.NumComponents; ++i)
+	{
+		if (Lhs.Component[i].Packed != Rhs.Component[i].Packed)
+		{
+			return false;
+		}
+	}
+	return true;
+}
+
 UE::Shader::FValue UE::Shader::Abs(const FValue& Value)
 {
 	return Private::UnaryOp(Private::FOpAbs(), Value);
diff --git a/Engine/Source/Runtime/Engine/Public/Shader/ShaderTypes.h b/Engine/Source/Runtime/Engine/Public/Shader/ShaderTypes.h
--- a/Engine/Source/Runtime/Engine/Public/Shader/ShaderTypes.h
+++ b/Engine/Source/Runtime/Engine/Public/Shader/ShaderTypes.h
@@ -115,6 +115,19 @@ struct FValue
 		Component[3].Float = Value.A;
 	}
 
+	inline FValue(const FVector2D& Value) : ComponentType(EValueComponentType::Float), NumComponents(2)
+	{
+		Component[0].Float = Value.X;
+		Component[1].Float = Value.Y;
+	}
+
+	inline FValue(const FVector& Value) : ComponentType(EValueComponentType::Float), NumComponents(3)
+	{
+		Component[0].Float = Value.X;
+		Component[1].Float = Value.Y;
+		Component[2].Float = Value.Z;
+	}
+
 	inline FValue(const FVector4& Value) : ComponentType(EValueComponentType::Float), NumComponents(4)
 	{
 		Component[0].Float = Value.X;
@@ -144,6 +157,9 @@ struct FValue
 	FValueComponent Component[4];
 };
 
+ENGINE_API bool operator==(const FValue& Lhs, const FValue& Rhs);
+inline bool operator!=(const FValue& Lhs, const FValue& Rhs) { return !operator==(Lhs, Rhs); }
+
 ENGINE_API FValue Abs(const FValue& Value);
 ENGINE_API FValue Saturate(const FValue& Value);
 ENGINE_API FValue Floor(const FValue& Value);
