Make sure material functions get initialized with a valid StateId
Avoid failed check() if used immediately, before PostLoad() has a chance to run

diff --git a/Engine/Source/Runtime/Engine/Classes/Materials/MaterialFunctionInterface.h b/Engine/Source/Runtime/Engine/Classes/Materials/MaterialFunctionInterface.h
--- a/Engine/Source/Runtime/Engine/Classes/Materials/MaterialFunctionInterface.h
+++ b/Engine/Source/Runtime/Engine/Classes/Materials/MaterialFunctionInterface.h
@@ -38,6 +38,7 @@ class UMaterialFunctionInterface : public UObject
 	GENERATED_UCLASS_BODY()
 
 	//~ Begin UObject Interface.
+	virtual void PostInitProperties() override;
 	virtual void PostLoad() override;
 	virtual void GetAssetRegistryTags(TArray<FAssetRegistryTag>& OutTags) const override;
 	//~ End UObject Interface.
diff --git a/Engine/Source/Runtime/Engine/Private/Materials/MaterialExpressions.cpp b/Engine/Source/Runtime/Engine/Private/Materials/MaterialExpressions.cpp
--- a/Engine/Source/Runtime/Engine/Private/Materials/MaterialExpressions.cpp
+++ b/Engine/Source/Runtime/Engine/Private/Materials/MaterialExpressions.cpp
@@ -12022,6 +12022,14 @@ UMaterialFunctionInterface::UMaterialFunctionInterface(const FObjectInitializer&
 {
 }
 
+void UMaterialFunctionInterface::PostInitProperties()
+{
+	Super::PostInitProperties();
+
+	// Initialize StateId to something unique, in case this is a new function
+	StateId = FGuid::NewGuid();
+}
+
 void UMaterialFunctionInterface::PostLoad()
 {
 	Super::PostLoad();
