Console Autocomplete: when populating the list with map assets, only use IsUAsset FAssetDatas, so that the BlueprintGeneratedClass assets inside the maps do not get added to the list.
#rnx

diff --git a/Engine/Source/Runtime/Engine/Private/UserInterface/Console.cpp b/Engine/Source/Runtime/Engine/Private/UserInterface/Console.cpp
--- a/Engine/Source/Runtime/Engine/Private/UserInterface/Console.cpp
+++ b/Engine/Source/Runtime/Engine/Private/UserInterface/Console.cpp
@@ -298,7 +298,7 @@ void UConsole::BuildRuntimeAutoCompleteList(bool bForce)
 
 						for (const FAssetData& Asset : Assets)
 						{
-							if (!!(Asset.PackageFlags & PKG_ContainsMap))
+							if (!!(Asset.PackageFlags & PKG_ContainsMap) && Asset.IsUAsset())
 							{
 								OutPackages.Emplace(Asset.AssetName.ToString());
 							}
