Avoid explicit wait on finish compilation by adding support for reschedule and wait for each compilable asset
Use the asset compiler's thread pool in finish compilation to benefit from its memory management and avoid OOM when under memory pressure.

TESTS
  - Cook P_Construct_WP and ensure tasks are promoted to foreground thread while waiting on them
  - Load P_Construct_WP and make sure behavior is as expected and that the call to FinishCachePlatformData on the game-thread doesn't incur an additional cost that couldn't be performed by the async task
  - Multiple Interchange Import
  - Datasmith Import

diff --git a/Engine/Source/Runtime/Engine/Private/AsyncCompilationHelpers.cpp b/Engine/Source/Runtime/Engine/Private/AsyncCompilationHelpers.cpp
--- a/Engine/Source/Runtime/Engine/Private/AsyncCompilationHelpers.cpp
+++ b/Engine/Source/Runtime/Engine/Private/AsyncCompilationHelpers.cpp
@@ -13,6 +13,7 @@
 #include "Settings/EditorExperimentalSettings.h"
 #include "Misc/QueuedThreadPoolWrapper.h"
 #include "HAL/IConsoleManager.h"
+#include "AssetCompilingManager.h"
 #include "HAL/PlatformStackWalk.h"
 #include "Hash/CityHash.h"
 #include "Misc/ScopeRWLock.h"
@@ -101,41 +102,12 @@ namespace AsyncCompilationHelpers
 			SlowTask->MakeDialogDelayed(1.0f, false /*bShowCancelButton*/, true /*bAllowInPIE*/);
 		}
 
-		struct FFinishTask : public IQueuedWork
-		{
-			ICompilable* Job = nullptr;
-			FEvent* Event;
-			FFinishTask() { Event = FPlatformProcess::GetSynchEventFromPool(true); }
-			~FFinishTask() { FPlatformProcess::ReturnSynchEventToPool(Event); }
-			void DoThreadedWork() override
-			{
-				FOptionalTaskTagScope Scope(ETaskTag::EParallelGameThread);
-				Job->EnsureCompletion();
-				Event->Trigger();
-			}
-
-			void Abandon() override { }
-
-			// Since we're specifically waiting on other tasks here, we can't be run inside a busywait
-			// without risking a deadlock.
-			//
-			// (i.e. T1 -> busywait -> picks T2 that then waits on T1 -> deadlock
-			//
-			// Imagine we're T2 and we're going to wait on T1 here... 
-			// if we get picked up on the same thread inside T1's busy wait, we'll deadlock.
-			//
-			// See FAsyncWorkEnsureCompletionBusyWaitDeadLockTest for a 100% repro case
-			EQueuedWorkFlags GetQueuedWorkFlags() const override { return EQueuedWorkFlags::DoNotRunInsideBusyWait; }
-		};
-
-		// Perform forced compilation on as many thread as possible in high priority since the game-thread is waiting
-		TArray<FFinishTask> PendingTasks;
-		PendingTasks.SetNum(Num);
-
+		// Reschedule everything to be executed at highest priority
+		// Important to keep using the asset compiling manager's thread pool to benefit from it's memory management and avoid OOM.
+		FQueuedThreadPool* AssetThreadPool = FAssetCompilingManager::Get().GetThreadPool();
 		for (int32 Index = 0; Index < Num; ++Index)
 		{
-			PendingTasks[Index].Job = &Getter(Index);
-			GThreadPool->AddQueuedWork(&PendingTasks[Index], EQueuedWorkPriority::Highest);
+			Getter(Index).Reschedule(AssetThreadPool, EQueuedWorkPriority::Highest);
 		}
 
 		auto FormatProgress =
@@ -150,15 +122,15 @@ namespace AsyncCompilationHelpers
 		};
 
 		int32 NumDone = 0;
-		for (FFinishTask& PendingTask : PendingTasks)
+		for (int32 Index = 0; Index < Num; ++Index)
 		{
-			ICompilable* Job = PendingTask.Job;
+			ICompilable& Job = Getter(Index);
 
-			FText Progress = FormatProgress(NumDone++, Num, Job->GetName());
+			FText Progress = FormatProgress(NumDone++, Num, Job.GetName());
 
 			// Be nice with the game thread and tick the progress at 60 fps even when no progress is being made...
 			bool bLogSlowProgress = true;
-			while (!PendingTask.Event->Wait(16))
+			while (!Job.WaitCompletionWithTimeout(0.016))
 			{
 				if (SlowTask.IsSet())
 				{
@@ -176,7 +148,7 @@ namespace AsyncCompilationHelpers
 				SlowTask->EnterProgressFrame(1.0f, Progress);
 			}
 
-			PostCompileSingle(Job);
+			PostCompileSingle(&Job);
 		}
 
 		SaveStallStack(FPlatformTime::Cycles64() - StartTime);
diff --git a/Engine/Source/Runtime/Engine/Private/SkeletalMeshCompiler.cpp b/Engine/Source/Runtime/Engine/Private/SkeletalMeshCompiler.cpp
--- a/Engine/Source/Runtime/Engine/Private/SkeletalMeshCompiler.cpp
+++ b/Engine/Source/Runtime/Engine/Private/SkeletalMeshCompiler.cpp
@@ -270,7 +270,7 @@ void FSkeletalMeshCompilingManager::FinishCompilation(TArrayView<USkeletalMesh*
 
 	if (PendingSkeletalMeshes.Num())
 	{
-		class FCompilableSkeletalMesh : public AsyncCompilationHelpers::ICompilable
+		class FCompilableSkeletalMesh : public AsyncCompilationHelpers::TCompilableAsyncTask<FSkeletalMeshAsyncBuildTask>
 		{
 		public:
 			FCompilableSkeletalMesh(USkeletalMesh* InSkeletalMesh)
@@ -278,8 +278,12 @@ void FSkeletalMeshCompilingManager::FinishCompilation(TArrayView<USkeletalMesh*
 			{
 			}
 
+			FSkeletalMeshAsyncBuildTask* GetAsyncTask() override
+			{
+				return SkeletalMesh->AsyncTask.Get();
+			}
+
 			TStrongObjectPtr<USkeletalMesh> SkeletalMesh;
-			void EnsureCompletion() override { if (SkeletalMesh->AsyncTask) { SkeletalMesh->AsyncTask->EnsureCompletion(); } }
 			FName GetName() override { return SkeletalMesh->GetFName(); }
 		};
 
diff --git a/Engine/Source/Runtime/Engine/Private/StaticMeshCompiler.cpp b/Engine/Source/Runtime/Engine/Private/StaticMeshCompiler.cpp
--- a/Engine/Source/Runtime/Engine/Private/StaticMeshCompiler.cpp
+++ b/Engine/Source/Runtime/Engine/Private/StaticMeshCompiler.cpp
@@ -302,7 +302,7 @@ void FStaticMeshCompilingManager::FinishCompilation(TArrayView<UStaticMesh* cons
 
 	if (PendingStaticMeshes.Num())
 	{
-		class FCompilableStaticMesh : public AsyncCompilationHelpers::ICompilable
+		class FCompilableStaticMesh : public AsyncCompilationHelpers::TCompilableAsyncTask<FStaticMeshAsyncBuildTask>
 		{
 		public:
 			FCompilableStaticMesh(UStaticMesh* InStaticMesh)
@@ -310,8 +310,12 @@ void FStaticMeshCompilingManager::FinishCompilation(TArrayView<UStaticMesh* cons
 			{
 			}
 
+			FStaticMeshAsyncBuildTask* GetAsyncTask() override
+			{
+				return StaticMesh->AsyncTask.Get();
+			}
+
 			TStrongObjectPtr<UStaticMesh> StaticMesh;
-			void EnsureCompletion() override { if (StaticMesh->AsyncTask) { StaticMesh->AsyncTask->EnsureCompletion(); } }
 			FName GetName() override { return StaticMesh->GetOutermost()->GetFName(); }
 		};
 
diff --git a/Engine/Source/Runtime/Engine/Private/TextureCompiler.cpp b/Engine/Source/Runtime/Engine/Private/TextureCompiler.cpp
--- a/Engine/Source/Runtime/Engine/Private/TextureCompiler.cpp
+++ b/Engine/Source/Runtime/Engine/Private/TextureCompiler.cpp
@@ -249,7 +249,7 @@ void FTextureCompilingManager::FinishCompilation(TArrayView<UTexture* const> InT
 
 	if (PendingTextures.Num())
 	{
-		class FCompilableTexture : public AsyncCompilationHelpers::ICompilable
+		class FCompilableTexture : public AsyncCompilationHelpers::TCompilableAsyncTask<FTextureAsyncCacheDerivedDataTask>
 		{
 		public:
 			FCompilableTexture(UTexture* InTexture)
@@ -257,8 +257,20 @@ void FTextureCompilingManager::FinishCompilation(TArrayView<UTexture* const> InT
 			{
 			}
 
+			FTextureAsyncCacheDerivedDataTask* GetAsyncTask() override
+			{
+				if (FTexturePlatformData** PlatformDataPtr = Texture->GetRunningPlatformData())
+				{
+					if (FTexturePlatformData* PlatformData = *PlatformDataPtr)
+					{
+						return PlatformData->AsyncTask;
+					}
+				}
+
+				return nullptr;
+			}
+
 			TStrongObjectPtr<UTexture> Texture;
-			void EnsureCompletion() override { Texture->FinishCachePlatformData(); }
 			FName GetName() override { return Texture->GetOutermost()->GetFName(); }
 		};
 
diff --git a/Engine/Source/Runtime/Engine/Public/AsyncCompilationHelpers.h b/Engine/Source/Runtime/Engine/Public/AsyncCompilationHelpers.h
--- a/Engine/Source/Runtime/Engine/Public/AsyncCompilationHelpers.h
+++ b/Engine/Source/Runtime/Engine/Public/AsyncCompilationHelpers.h
@@ -8,6 +8,7 @@
 #include "HAL/IConsoleManager.h"
 #include "Logging/LogMacros.h"
 #include "Framework/Notifications/NotificationManager.h"
+#include "Misc/QueuedThreadPool.h"
 
 DECLARE_LOG_CATEGORY_EXTERN(LogAsyncCompilation, Log, All);
 
@@ -46,10 +47,42 @@ namespace AsyncCompilationHelpers
 	{
 	public:
 		inline virtual ~ICompilable() {};
-		virtual void EnsureCompletion() = 0;
+
+		/*
+		* Reschedules any async tasks to the given thread pool at the given priority.
+		*/
+		virtual void Reschedule(FQueuedThreadPool* InThreadPool, EQueuedWorkPriority InPriority) = 0;
+
+		/*
+		* Returns true once async tasks are complete, false if timing out.
+		*/
+		virtual bool WaitCompletionWithTimeout(float TimeLimitSeconds) = 0;
 		virtual FName GetName() = 0;
 	};
 
+	template <typename AsyncTaskType>
+	class TCompilableAsyncTask : public ICompilable
+	{
+		virtual AsyncTaskType* GetAsyncTask() = 0;
+
+		void Reschedule(FQueuedThreadPool* InThreadPool, EQueuedWorkPriority InPriority) override
+		{
+			if (AsyncTaskType* AsyncTask = GetAsyncTask())
+			{
+				AsyncTask->Reschedule(InThreadPool, InPriority);
+			}
+		}
+		
+		bool WaitCompletionWithTimeout(float TimeLimitSeconds)
+		{
+			if (AsyncTaskType* AsyncTask = GetAsyncTask())
+			{
+				return AsyncTask->WaitCompletionWithTimeout(TimeLimitSeconds);
+			}
+			return true;
+		}
+	};
+
 	void FinishCompilation(
 		TFunctionRef<ICompilable& (int32 Index)> Getter,
 		int32 Num,
