World Partition's always loaded cell GetCurrentState now always returns EWorldPartitionRuntimeCellState::Activated.

diff --git a/Engine/Source/Runtime/Engine/Private/WorldPartition/WorldPartitionRuntimeLevelStreamingCell.cpp b/Engine/Source/Runtime/Engine/Private/WorldPartition/WorldPartitionRuntimeLevelStreamingCell.cpp
--- a/Engine/Source/Runtime/Engine/Private/WorldPartition/WorldPartitionRuntimeLevelStreamingCell.cpp
+++ b/Engine/Source/Runtime/Engine/Private/WorldPartition/WorldPartitionRuntimeLevelStreamingCell.cpp
@@ -36,7 +36,8 @@ EWorldPartitionRuntimeCellState UWorldPartitionRuntimeLevelStreamingCell::GetCur
 		}
 	}
 	
-	return EWorldPartitionRuntimeCellState::Unloaded;
+	//@todo_ow: Now that actors are moved to the persistent level, remove the AlwaysLoaded cell (it's always empty)
+	return IsAlwaysLoaded() ? EWorldPartitionRuntimeCellState::Activated : EWorldPartitionRuntimeCellState::Unloaded;
 }
 
 UWorldPartitionLevelStreamingDynamic* UWorldPartitionRuntimeLevelStreamingCell::GetLevelStreaming() const
