Rework shader compilation thread allocation on higher core machines (+12 logical cores) to reserve 50% of the cores for non-shader compilation work.

- Added PercentageUnusedShaderCompilingThreads which defaults to 50%.  This means shader compilation will reserve 50% of your available cores so your machine is more responsive.
- Added ShaderCompilerCoreCountThreshold, which defaults to 12, which is a threshold that controls what heuristic is used to determine unused threads.  If you are 12 or less you use the old behavior of reserving 3 threads.  If you are greater than 12 you reserve threads based on the percentage (50% by default).

diff --git a/Engine/Config/BaseEngine.ini b/Engine/Config/BaseEngine.ini
--- a/Engine/Config/BaseEngine.ini
+++ b/Engine/Config/BaseEngine.ini
@@ -1371,6 +1371,12 @@ bAllowAsynchronousShaderCompiling=True
 NumUnusedShaderCompilingThreads=3
 ; Make sure the game has enough cores available to maintain reasonable performance
 NumUnusedShaderCompilingThreadsDuringGame=4
+; Core count threshold.  Below this amount will use NumUnusedShaderCompilingThreads.  Above this threshold will use PercentageUnusedShaderCompilingThreads when determining the number of cores to reserve.
+ShaderCompilerCoreCountThreshold=12
+; Percentage of your available logical cores that will be reserved and NOT used for shader compilation
+; 0 means use all your cores to compile Shaders
+; 100 means use none of your cores to compile shaders (it will still use 1 core).
+PercentageUnusedShaderCompilingThreads=50
 ; Batching multiple jobs to reduce file overhead, but not so many that latency of blocking compiles is hurt
 MaxShaderJobBatchSize=10
 bPromptToRetryFailedShaderCompiles=False
diff --git a/Engine/Source/Runtime/Engine/Private/ShaderCompiler/ShaderCompiler.cpp b/Engine/Source/Runtime/Engine/Private/ShaderCompiler/ShaderCompiler.cpp
--- a/Engine/Source/Runtime/Engine/Private/ShaderCompiler/ShaderCompiler.cpp
+++ b/Engine/Source/Runtime/Engine/Private/ShaderCompiler/ShaderCompiler.cpp
@@ -2761,11 +2761,28 @@ FShaderCompilingManager::FShaderCompilingManager() :
 		bAllowAsynchronousShaderCompiling = false;
 	}
 
+	int32 ShaderCompilerCoreCountThreshold;
+	verify(GConfig->GetInt(TEXT("DevOptions.Shaders"), TEXT("ShaderCompilerCoreCountThreshold"), ShaderCompilerCoreCountThreshold, GEngineIni));
+
+	const int32 NumVirtualCores = FPlatformMisc::NumberOfCoresIncludingHyperthreads();
+
 	int32 NumUnusedShaderCompilingThreads;
-	verify(GConfig->GetInt( TEXT("DevOptions.Shaders"), TEXT("NumUnusedShaderCompilingThreads"), NumUnusedShaderCompilingThreads, GEngineIni ));
+	verify(GConfig->GetInt(TEXT("DevOptions.Shaders"), TEXT("NumUnusedShaderCompilingThreads"), NumUnusedShaderCompilingThreads, GEngineIni));
 
 	int32 NumUnusedShaderCompilingThreadsDuringGame;
-	verify(GConfig->GetInt( TEXT("DevOptions.Shaders"), TEXT("NumUnusedShaderCompilingThreadsDuringGame"), NumUnusedShaderCompilingThreadsDuringGame, GEngineIni ));
+	verify(GConfig->GetInt(TEXT("DevOptions.Shaders"), TEXT("NumUnusedShaderCompilingThreadsDuringGame"), NumUnusedShaderCompilingThreadsDuringGame, GEngineIni));
+
+	if (NumVirtualCores > ShaderCompilerCoreCountThreshold)
+	{
+		float PercentageUnusedShaderCompilingThreads;
+		verify(GConfig->GetFloat(TEXT("DevOptions.Shaders"), TEXT("PercentageUnusedShaderCompilingThreads"), PercentageUnusedShaderCompilingThreads, GEngineIni));
+
+		// ensure we get a valid multiplier.
+		PercentageUnusedShaderCompilingThreads = FMath::Clamp(PercentageUnusedShaderCompilingThreads, 0.0f, 100.0f) / 100.0f;
+
+		NumUnusedShaderCompilingThreads = FMath::CeilToInt(NumVirtualCores * PercentageUnusedShaderCompilingThreads);
+		NumUnusedShaderCompilingThreadsDuringGame = NumUnusedShaderCompilingThreads;
+	}
 
 	// Use all the cores on the build machines.
 	if (GForceAllCoresForShaderCompiling != 0)
@@ -2824,8 +2841,6 @@ FShaderCompilingManager::FShaderCompilingManager() :
 	FPaths::NormalizeDirectoryName(AbsoluteDebugInfoDirectory);
 	AbsoluteShaderDebugInfoDirectory = AbsoluteDebugInfoDirectory;
 
-	const int32 NumVirtualCores = FPlatformMisc::NumberOfCoresIncludingHyperthreads();
-
 	NumShaderCompilingThreads = (bAllowCompilingThroughWorkers && NumVirtualCores > NumUnusedShaderCompilingThreads) ? (NumVirtualCores - NumUnusedShaderCompilingThreads) : 1;
 
 	// Make sure there's at least one worker allowed to be active when compiling during the game
