Improved HLOD dirtiness check
* Avoid relying on BodySetupGuid as it is sometime resetted during loading
* Integrate transform CRC fixes

diff --git a/Engine/Source/Runtime/Engine/Private/HLODProxy.cpp b/Engine/Source/Runtime/Engine/Private/HLODProxy.cpp
--- a/Engine/Source/Runtime/Engine/Private/HLODProxy.cpp
+++ b/Engine/Source/Runtime/Engine/Private/HLODProxy.cpp
@@ -376,8 +376,10 @@ uint32 UHLODProxy::GetCRC(UStaticMesh* InStaticMesh, uint32 InCRC)
 
 	if(InStaticMesh->GetBodySetup())
 	{
-		// Incorporate physics data
-		Ar << InStaticMesh->GetBodySetup()->BodySetupGuid;
+		// Incorporate physics data - Avoid relying on BodySetupGuid as it is sometime resetted during loading
+		FString BodySetupDDCKey;
+		InStaticMesh->GetBodySetup()->GetGeometryDDCKey(BodySetupDDCKey);
+		Ar << BodySetupDDCKey;
 	}
 
 	return Ar.GetCrc();
@@ -389,11 +391,12 @@ static void AppendRoundedTransform(const FRotator& ComponentRotation, const FVec
 	FIntVector Location(FMath::RoundToInt(ComponentLocation.X), FMath::RoundToInt(ComponentLocation.Y), FMath::RoundToInt(ComponentLocation.Z));
 	Ar << Location;
 
-	FVector RotationVector(ComponentRotation.GetNormalized().Vector());
-	FIntVector Rotation(FMath::RoundToInt(RotationVector.X), FMath::RoundToInt(RotationVector.Y), FMath::RoundToInt(RotationVector.Z));
+	FRotator Rotator(ComponentRotation.GetDenormalized());
+	FIntVector Rotation(FMath::RoundToInt(Rotator.Pitch), FMath::RoundToInt(Rotator.Yaw), FMath::RoundToInt(Rotator.Roll));
 	Ar << Rotation;
 
-	FIntVector Scale(FMath::RoundToInt(ComponentScale.X), FMath::RoundToInt(ComponentScale.Y), FMath::RoundToInt(ComponentScale.Z));
+	float SCALE_FACTOR = 100;
+	FIntVector Scale(FMath::RoundToInt(ComponentScale.X * SCALE_FACTOR), FMath::RoundToInt(ComponentScale.Y * SCALE_FACTOR), FMath::RoundToInt(ComponentScale.Z * SCALE_FACTOR));
 	Ar << Scale;
 }
 
