Avoid keeping references on pre-loaded actors in world partition editor cell when running a commandlet (allows commandlet to control manual actor loading/unloading).

diff --git a/Engine/Source/Runtime/Engine/Private/WorldPartition/WorldPartitionEditorCell.cpp b/Engine/Source/Runtime/Engine/Private/WorldPartition/WorldPartitionEditorCell.cpp
--- a/Engine/Source/Runtime/Engine/Private/WorldPartition/WorldPartitionEditorCell.cpp
+++ b/Engine/Source/Runtime/Engine/Private/WorldPartition/WorldPartitionEditorCell.cpp
@@ -42,7 +42,7 @@ void UWorldPartitionEditorCell::AddActor(const FGuid& Source, const FWorldPartit
 
 	if (!bIsAlreadyInSet)
 	{	
-		if (ActorHandle.IsLoaded())
+		if (ActorHandle.IsLoaded() && !IsRunningCommandlet())
 		{
 			LoadedActors.Add(FActorReference(Source, ActorHandle), &bIsAlreadyInSet);
 			check(!bIsAlreadyInSet);
