Changed when initial state of dynamically loaded DataLayer to make sure DataLayers to load/activate will be included by World::BlockTillLevelStreamingCompleted.

diff --git a/Engine/Source/Runtime/Engine/Private/WorldPartition/DataLayer/DataLayerSubsystem.cpp b/Engine/Source/Runtime/Engine/Private/WorldPartition/DataLayer/DataLayerSubsystem.cpp
--- a/Engine/Source/Runtime/Engine/Private/WorldPartition/DataLayer/DataLayerSubsystem.cpp
+++ b/Engine/Source/Runtime/Engine/Private/WorldPartition/DataLayer/DataLayerSubsystem.cpp
@@ -53,6 +53,22 @@ void UDataLayerSubsystem::PostInitialize()
 {
 	Super::PostInitialize();
 
+	if (GetWorld()->IsGameWorld())
+	{
+		// Initialize Dynamically loaded Data Layers state
+		if (const AWorldDataLayers* WorldDataLayers = AWorldDataLayers::Get(GetWorld()))
+		{
+			WorldDataLayers->ForEachDataLayer([this](class UDataLayer* DataLayer)
+			{
+				if (DataLayer && DataLayer->IsDynamicallyLoaded())
+				{
+					SetDataLayerState(DataLayer, DataLayer->GetInitialState());
+				}
+				return true;
+			});
+		}
+	}
+
 	if (GetWorld()->IsGameWorld() && (GetWorld()->GetNetMode() != NM_DedicatedServer))
 	{
 		DrawHandle = UDebugDrawService::Register(TEXT("Game"), FDebugDrawDelegate::CreateUObject(this, &UDataLayerSubsystem::Draw));
diff --git a/Engine/Source/Runtime/Engine/Private/WorldPartition/DataLayer/WorldDataLayers.cpp b/Engine/Source/Runtime/Engine/Private/WorldPartition/DataLayer/WorldDataLayers.cpp
--- a/Engine/Source/Runtime/Engine/Private/WorldPartition/DataLayer/WorldDataLayers.cpp
+++ b/Engine/Source/Runtime/Engine/Private/WorldPartition/DataLayer/WorldDataLayers.cpp
@@ -6,7 +6,6 @@
 
 #include "WorldPartition/DataLayer/WorldDataLayers.h"
 #include "WorldPartition/DataLayer/DataLayer.h"
-#include "WorldPartition/DataLayer/DataLayerSubsystem.h"
 #include "EngineUtils.h"
 #if WITH_EDITOR
 #include "WorldPartition/WorldPartitionEditorPerProjectUserSettings.h"
@@ -250,21 +249,4 @@ void AWorldDataLayers::PostLoad()
 #endif
 }
 
-void AWorldDataLayers::BeginPlay()
-{
-	Super::BeginPlay();
-
-	// Activate initially loaded data layers
-	if (UDataLayerSubsystem* DataLayerSubsystem = GetWorld()->GetSubsystem<UDataLayerSubsystem>())
-	{
-		for (const UDataLayer* DataLayer : WorldDataLayers)
-		{
-			if (DataLayer && DataLayer->IsDynamicallyLoaded())
-			{
-				DataLayerSubsystem->SetDataLayerState(DataLayer, DataLayer->GetInitialState());
-			}
-		}
-	}
-}
-
 #undef LOCTEXT_NAMESPACE
\ No newline at end of file
diff --git a/Engine/Source/Runtime/Engine/Public/WorldPartition/DataLayer/WorldDataLayers.h b/Engine/Source/Runtime/Engine/Public/WorldPartition/DataLayer/WorldDataLayers.h
--- a/Engine/Source/Runtime/Engine/Public/WorldPartition/DataLayer/WorldDataLayers.h
+++ b/Engine/Source/Runtime/Engine/Public/WorldPartition/DataLayer/WorldDataLayers.h
@@ -21,8 +21,6 @@ public:
 	static const AWorldDataLayers* Get(UWorld* World);
 	virtual void PostLoad() override;
 	
-	virtual void BeginPlay() override;
-
 #if WITH_EDITOR
 	static AWorldDataLayers* Get(UWorld* World, bool bCreateIfNotFound);
 	UDataLayer* CreateDataLayer(FName InName = TEXT("DataLayer"), EObjectFlags InObjectFlags = RF_NoFlags);
