Hang during editor booting ShooterGame
#fix pre-load animation data controller module, and replace Load with GetModule in AnimSequenceBase::GetController

diff --git a/Engine/Source/Runtime/Engine/Private/Animation/AnimSequenceBase.cpp b/Engine/Source/Runtime/Engine/Private/Animation/AnimSequenceBase.cpp
--- a/Engine/Source/Runtime/Engine/Private/Animation/AnimSequenceBase.cpp
+++ b/Engine/Source/Runtime/Engine/Private/Animation/AnimSequenceBase.cpp
@@ -1137,7 +1137,7 @@ IAnimationDataController& UAnimSequenceBase::GetController()
 
 	if(Controller == nullptr)
 	{
-		IAnimationDataControllerModule& ControllerModule = FModuleManager::Get().LoadModuleChecked<IAnimationDataControllerModule>("AnimationDataController");
+		IAnimationDataControllerModule& ControllerModule = FModuleManager::Get().GetModuleChecked<IAnimationDataControllerModule>("AnimationDataController");
 		Controller = ControllerModule.GetController();
 		Controller->SetModel(DataModel);
 	}
diff --git a/Engine/Source/Runtime/Launch/Private/LaunchEngineLoop.cpp b/Engine/Source/Runtime/Launch/Private/LaunchEngineLoop.cpp
--- a/Engine/Source/Runtime/Launch/Private/LaunchEngineLoop.cpp
+++ b/Engine/Source/Runtime/Launch/Private/LaunchEngineLoop.cpp
@@ -3954,6 +3954,7 @@ bool FEngineLoop::LoadStartupCoreModules()
 	FModuleManager::Get().LoadModule(TEXT("ClothingSystemRuntimeNv"));
 #if WITH_EDITOR
 	FModuleManager::Get().LoadModule(TEXT("ClothingSystemEditor"));
+	FModuleManager::Get().LoadModule(TEXT("AnimationDataController"));
 #endif
 
 	FModuleManager::Get().LoadModule(TEXT("PacketHandler"));
