SC: make XGE SCWs non-local, lower their priority, and other changes.

- Increased the number of jobs per XGE worker to reduce the number of times we start them. Since XGE SCWs aren't used for latency-critical jobs, this is expected to not cause user-visible delays.
- Added an ability to detect a hung distributed controller (but the build would be currently broken anyway as the jobs aren't reissued).
- Added a check for XGE system service running to avoid attempting to launch XGE on machines without it.
- Removed the code to even attempt to launch XGE XML interface.
- Also added some more logs about what's happening, and reduced the job cache verbosity.

diff --git a/Engine/Plugins/XGEController/Source/Private/XGEControllerModule.cpp b/Engine/Plugins/XGEController/Source/Private/XGEControllerModule.cpp
--- a/Engine/Plugins/XGEController/Source/Private/XGEControllerModule.cpp
+++ b/Engine/Plugins/XGEController/Source/Private/XGEControllerModule.cpp
@@ -216,6 +216,14 @@ bool FXGEControllerModule::IsSupported()
 				XGEControllerVariables::Enabled = 0;
 			}
 
+			// check if build service is running - without this, the build will make no progress
+			const TCHAR* XGEBuildServiceExecutableName = TEXT("BuildService.exe");
+			if (!FPlatformProcess::IsApplicationRunning(XGEBuildServiceExecutableName))
+			{
+				UE_LOG(LogXGEController, Warning, TEXT("XGE's background service (%s) is not running - service is likely disabled on this machine."), XGEBuildServiceExecutableName);
+				XGEControllerVariables::Enabled = 0;
+			}
+
 			if (!PlatformFile.FileExists(*GetControlWorkerExePath()))
 			{
 				UE_LOG(LogXGEController, Warning, TEXT("XGEControlWorker.exe does not exist, XGE may be disabled in your Build Configuration, cannot use XGE."));
@@ -329,14 +337,20 @@ void FXGEControllerModule::WriteOutThreadProc()
 		bRestartWorker = false;
 
 		while (!AreTasksPending() && !bShutdown)
+		{
+			UE_LOG(LogXGEController, Verbose, TEXT("Waiting for new tasks."));
 			WriteOutThreadEvent->Wait();
+		}
 
 		if (bShutdown)
+		{
+			UE_LOG(LogXGEController, Verbose, TEXT("Shutting down communication with the XGE controller."));
 			return;
+		}
 
 		// To handle spaces in the engine path, we just pass the XGEController.exe filename to xgConsole,
 		// and set the working directory of xgConsole.exe to the engine binaries folder below.
-		FString XGConsoleArgs = FString::Printf(TEXT("/VIRTUALIZEDIRECTX /allowremote=\"%s\" /allowintercept=\"%s\" /title=\"Unreal Engine XGE Tasks\" /monitordirs=\"%s\" /command=\"%s -xgecontroller %s\""),
+		FString XGConsoleArgs = FString::Printf(TEXT("/VIRTUALIZEDIRECTX /allowremote=\"%s\" /avoidlocal=ON /allowintercept=\"%s\" /title=\"Unreal Engine XGE Tasks\" /monitordirs=\"%s\" /command=\"%s -xgecontroller %s\""),
 			XGE_INTERCEPT_EXE_NAMES,
 			XGE_CONTROL_WORKER_NAME,
 			*WorkingDirectory,
@@ -349,9 +363,11 @@ void FXGEControllerModule::WriteOutThreadProc()
 			UE_LOG(LogXGEController, Fatal, TEXT("Failed to create the output XGE named pipe."));
 		}
 
+		const int32 PriorityModifier = -1;	// below normal
 		// Start the controller process
 		uint32 XGConsoleProcID = 0;
-		BuildProcessHandle = FPlatformProcess::CreateProc(*XGConsolePath, *XGConsoleArgs, false, false, true, &XGConsoleProcID, 0, *ControlWorkerDirectory, nullptr);
+		UE_LOG(LogXGEController, Verbose, TEXT("Launching xgConsole"));
+		BuildProcessHandle = FPlatformProcess::CreateProc(*XGConsolePath, *XGConsoleArgs, false, false, true, &XGConsoleProcID, PriorityModifier, *ControlWorkerDirectory, nullptr);
 		if (!BuildProcessHandle.IsValid())
 		{
 			UE_LOG(LogXGEController, Fatal, TEXT("Failed to launch the XGE control worker process."));
@@ -360,8 +376,9 @@ void FXGEControllerModule::WriteOutThreadProc()
 		// If the engine crashes, we don't get a chance to kill the build process.
 		// Start up the build monitor process to monitor for engine crashes.
 		uint32 BuildMonitorProcessID;
+		UE_LOG(LogXGEController, Verbose, TEXT("Launching XGEController to fan out tasks"));
 		FString XGMonitorArgs = FString::Printf(TEXT("-xgemonitor %d %d"), FPlatformProcess::GetCurrentProcessId(), XGConsoleProcID);
-		FProcHandle BuildMonitorHandle = FPlatformProcess::CreateProc(*GetControlWorkerExePath(), *XGMonitorArgs, true, false, false, &BuildMonitorProcessID, 0, nullptr, nullptr);
+		FProcHandle BuildMonitorHandle = FPlatformProcess::CreateProc(*GetControlWorkerExePath(), *XGMonitorArgs, true, false, false, &BuildMonitorProcessID, PriorityModifier, nullptr, nullptr);
 		FPlatformProcess::CloseProc(BuildMonitorHandle);
 
 		// Wait for the controller to connect to the output pipe
diff --git a/Engine/Plugins/XGEController/Source/Public/XGEControllerModule.h b/Engine/Plugins/XGEController/Source/Public/XGEControllerModule.h
--- a/Engine/Plugins/XGEController/Source/Public/XGEControllerModule.h
+++ b/Engine/Plugins/XGEController/Source/Public/XGEControllerModule.h
@@ -7,8 +7,11 @@
 
 class FXGEControllerModule : public IDistributedBuildController
 {
+	/** Whether this system can use XGE shader compilation. */
 	bool bSupported;
+	/** Whether the module was initialized. */
 	bool bModuleInitialized;
+	/** Whether the controller was successfully initialized. This is only possible when bSupported is true. */
 	bool bControllerInitialized;
 
 	FThreadSafeCounter NextFileID;
@@ -36,6 +39,7 @@ class FXGEControllerModule : public IDistributedBuildController
 	TFuture<void> WriteOutThreadFuture;
 	TFuture<void> ReadBackThreadFuture;
 
+	/** event for the thread that communicates with the XGE controller */
 	FEvent* WriteOutThreadEvent;
 
 	// We need two pipes, as the named pipe API does not support simultaneous read/write on two threads.
diff --git a/Engine/Source/Programs/ShaderCompileWorker/Private/XGEControlWorker.cpp b/Engine/Source/Programs/ShaderCompileWorker/Private/XGEControlWorker.cpp
--- a/Engine/Source/Programs/ShaderCompileWorker/Private/XGEControlWorker.cpp
+++ b/Engine/Source/Programs/ShaderCompileWorker/Private/XGEControlWorker.cpp
@@ -173,8 +173,8 @@ void FXGEControlWorker::InputThreadProc()
 		Reader << Task->Executable;
 		Reader << Task->Arguments;
 
-		// Launch the process
-		Task->Handle = FPlatformProcess::CreateProc(*Task->Executable, *Task->Arguments, true, false, false, nullptr, 0, nullptr, nullptr);
+		// Launch the process (with a lower than normal priority)
+		Task->Handle = FPlatformProcess::CreateProc(*Task->Executable, *Task->Arguments, true, false, false, nullptr, -1, nullptr, nullptr);
 
 		FScopeLock Lock(CS);
 		CurrentTasks.Add(Task);
diff --git a/Engine/Source/Runtime/Engine/Private/ShaderCompiler/ShaderCompiler.cpp b/Engine/Source/Runtime/Engine/Private/ShaderCompiler/ShaderCompiler.cpp
--- a/Engine/Source/Runtime/Engine/Private/ShaderCompiler/ShaderCompiler.cpp
+++ b/Engine/Source/Runtime/Engine/Private/ShaderCompiler/ShaderCompiler.cpp
@@ -73,7 +73,7 @@
 DEFINE_LOG_CATEGORY(LogShaderCompilers);
 
 // Switch to Verbose after initial testing
-#define UE_SHADERCACHE_LOG_LEVEL		Verbose
+#define UE_SHADERCACHE_LOG_LEVEL		VeryVerbose
 
 int32 GShaderCompilerJobCache = 1;
 static FAutoConsoleVariableRef CVarShaderCompilerJobCache(
@@ -2922,11 +2922,6 @@ FShaderCompilingManager::FShaderCompilingManager() :
 		UE_LOG(LogShaderCompilers, Display, TEXT("Using %s for Shader Compilation."), *BuildDistributionController->GetName());
 		RemoteCompileThread = MakeUnique<FShaderCompileDistributedThreadRunnable_Interface>(this, *BuildDistributionController);
 	}
-	else if (bCanUseRemoteCompiling && FShaderCompileXGEThreadRunnable_XmlInterface::IsSupported())
-	{
-		UE_LOG(LogShaderCompilers, Display, TEXT("Using XGE Shader Compiler (XML Interface)."));
-		RemoteCompileThread = MakeUnique<FShaderCompileXGEThreadRunnable_XmlInterface>(this);
-	}
 	else
 #endif // PLATFORM_WINDOWS
 #if PLATFORM_DESKTOP
@@ -2943,7 +2938,7 @@ FShaderCompilingManager::FShaderCompilingManager() :
 	if (RemoteCompileThread)
 	{
 		// Keep high priority jobs on the local machine, to avoid XGE latency
-		RemoteCompileThread->SetPriorityRange(EShaderCompileJobPriority::Low, EShaderCompileJobPriority::High);
+		RemoteCompileThread->SetPriorityRange(EShaderCompileJobPriority::Low, EShaderCompileJobPriority::Low);
 		LocalThread->SetPriorityRange(EShaderCompileJobPriority::Normal, EShaderCompileJobPriority::ForceLocal);
 		Threads.Add(MoveTemp(RemoteCompileThread));
 	}
diff --git a/Engine/Source/Runtime/Engine/Private/ShaderCompiler/ShaderCompilerDistributed.cpp b/Engine/Source/Runtime/Engine/Private/ShaderCompiler/ShaderCompilerDistributed.cpp
--- a/Engine/Source/Runtime/Engine/Private/ShaderCompiler/ShaderCompilerDistributed.cpp
+++ b/Engine/Source/Runtime/Engine/Private/ShaderCompiler/ShaderCompilerDistributed.cpp
@@ -14,6 +14,14 @@ namespace DistributedShaderCompilerVariables
         TEXT("Minimum number of shaders to compile with XGE.\n")
         TEXT("Smaller number of shaders will compile locally."),
         ECVF_Default);
+
+	static int32 GDistributedControllerTimeout = 5 * 60;	// 
+	static FAutoConsoleVariableRef CVarDistributedControllerTimeout(
+		TEXT("r.ShaderCompiler.DistributedControllerTimeout"),
+		GDistributedControllerTimeout,
+		TEXT("Maximum number of seconds we expect to pass between getting distributed controller complete a task (this is used to detect problems with the distribution controllers).")
+	);
+
 }
 
 bool FShaderCompileDistributedThreadRunnable_Interface::IsSupported()
@@ -43,6 +51,8 @@ FShaderCompileDistributedThreadRunnable_Interface::FShaderCompileDistributedThre
 	: FShaderCompileThreadRunnableBase(InManager)
 	, NumDispatchedJobs(0)
 	, CachedController(InController)
+	, LastTimeTaskCompleted(FPlatformTime::Seconds())
+	, bIsHung(false)
 {
 }
 
@@ -157,28 +167,37 @@ TArray<FString> FShaderCompileDistributedThreadRunnable_Interface::GetDependency
 int32 FShaderCompileDistributedThreadRunnable_Interface::CompilingLoop()
 {
 	TArray<FShaderCommonCompileJobPtr> PendingJobs;
-	for (int32 PriorityIndex = MaxPriorityIndex; PriorityIndex >= MinPriorityIndex; --PriorityIndex)
+	if (LIKELY(!bIsHung))	// stop accepting jobs if we're hung
 	{
-		// Grab as many jobs from the job queue as we can
-		const EShaderCompileJobPriority Priority = (EShaderCompileJobPriority)PriorityIndex;
-		const int32 MinBatchSize = (Priority == EShaderCompileJobPriority::Low) ? 1 : DistributedShaderCompilerVariables::MinBatchSize;
-		const int32 NumJobs = Manager->AllJobs.GetPendingJobs(EShaderCompilerWorkerType::XGE, Priority, MinBatchSize, INT32_MAX, PendingJobs);
-		if (NumJobs > 0)
-		{
-			UE_LOG(LogShaderCompilers, Display, TEXT("Started %d 'XGE' shader compile jobs with '%s' priority"),
-				NumJobs,
-				ShaderCompileJobPriorityToString((EShaderCompileJobPriority)PriorityIndex));
-		}
-		if (PendingJobs.Num() >= DistributedShaderCompilerVariables::MinBatchSize)
+		for (int32 PriorityIndex = MaxPriorityIndex; PriorityIndex >= MinPriorityIndex; --PriorityIndex)
 		{
-			break;
+			// Grab as many jobs from the job queue as we can
+			const EShaderCompileJobPriority Priority = (EShaderCompileJobPriority)PriorityIndex;
+			const int32 MinBatchSize = (Priority == EShaderCompileJobPriority::Low) ? 1 : DistributedShaderCompilerVariables::MinBatchSize;
+			const int32 NumJobs = Manager->AllJobs.GetPendingJobs(EShaderCompilerWorkerType::XGE, Priority, MinBatchSize, INT32_MAX, PendingJobs);
+			if (NumJobs > 0)
+			{
+				UE_LOG(LogShaderCompilers, Display, TEXT("Started %d 'XGE' shader compile jobs with '%s' priority"),
+					NumJobs,
+					ShaderCompileJobPriorityToString((EShaderCompileJobPriority)PriorityIndex));
+			}
+			if (PendingJobs.Num() >= DistributedShaderCompilerVariables::MinBatchSize)
+			{
+				break;
+			}
 		}
 	}
 
 	if (PendingJobs.Num() > 0)
 	{
 		// Increase the batch size when more jobs are queued/in flight.
-		const uint32 JobsPerBatch = FMath::Max(1, FMath::FloorToInt(FMath::LogX(2.f, PendingJobs.Num() + NumDispatchedJobs)));
+
+		// Just to provide typical numbers: the number of total jobs is usually in tens of thousands at most, oftentimes in low thousands. Thus JobsPerBatch when calculated as a log2 rarely reaches the value of 16.
+		// Given that 
+		//   a) starting each XGE shader compile worker has an associated cost (it scans directories, loads multiple libraries and reads configs)
+		//   b) XGE is currently used for low priority, non time-critical jobs
+		// JobsPerBatch is now increased by a factor of 3 compared to original formula, to produce numbers in 3-48 range most of the time
+		const uint32 JobsPerBatch = FMath::Max(1, FMath::FloorToInt(3 * FMath::LogX(2.f, PendingJobs.Num() + NumDispatchedJobs)));
 		UE_LOG(LogShaderCompilers, Verbose, TEXT("Current jobs: %d, Batch size: %d, Num Already Dispatched: %d"), PendingJobs.Num(), JobsPerBatch, NumDispatchedJobs);
 
 
@@ -285,6 +304,7 @@ int32 FShaderCompileDistributedThreadRunnable_Interface::CompilingLoop()
 
 		FDistributedBuildTaskResult Result = Task->Future.Get();
 		NumDispatchedJobs -= Task->ShaderJobs.Num();
+		LastTimeTaskCompleted = FPlatformTime::Seconds();
 
 		if (Result.ReturnCode != 0)
 		{
@@ -352,6 +372,18 @@ int32 FShaderCompileDistributedThreadRunnable_Interface::CompilingLoop()
 	// Yield for a short while to stop this thread continuously polling the disk.
 	FPlatformProcess::Sleep(0.01f);
 
+	// normally we expect to have at least one task in 5 minutes, although there could be edge cases
+	double TimeSinceLastCompletedTask = FPlatformTime::Seconds() - LastTimeTaskCompleted;
+	if (TimeSinceLastCompletedTask > DistributedShaderCompilerVariables::GDistributedControllerTimeout)
+	{
+		if (!bIsHung)
+		{
+			UE_LOG(LogShaderCompilers, Error, TEXT("Distributed compilation controller didn't receive a completed task in %f seconds!"), TimeSinceLastCompletedTask);
+			bIsHung = true;
+			// TODO: resubmit the hung jobs
+		}
+	}
+
 	// Return true if there is more work to be done.
 	return Manager->AllJobs.GetNumOutstandingJobs() > 0;
 }
diff --git a/Engine/Source/Runtime/Engine/Public/ShaderCompiler.h b/Engine/Source/Runtime/Engine/Public/ShaderCompiler.h
--- a/Engine/Source/Runtime/Engine/Public/ShaderCompiler.h
+++ b/Engine/Source/Runtime/Engine/Public/ShaderCompiler.h
@@ -749,6 +749,12 @@ class FShaderCompileDistributedThreadRunnable_Interface : public FShaderCompileT
 
 	TSparseArray<class FDistributedShaderCompilerTask*> DispatchedTasks;
 
+	/** Last time we received a task back. */
+	double LastTimeTaskCompleted;
+
+	/** Whether we consider this controller hung / out of order. */
+	bool bIsHung;
+
 public:
 	/** Initialization constructor. */
 	FShaderCompileDistributedThreadRunnable_Interface(class FShaderCompilingManager* InManager, class IDistributedBuildController& InController);
