Revert back to old NumUnusedThread logic for shader compilation if we are in a commandlet.  This way cooks, COTF, and other commandlets that compile shader still try to do so aggressively.

diff --git a/Engine/Source/Runtime/Engine/Private/ShaderCompiler/ShaderCompiler.cpp b/Engine/Source/Runtime/Engine/Private/ShaderCompiler/ShaderCompiler.cpp
--- a/Engine/Source/Runtime/Engine/Private/ShaderCompiler/ShaderCompiler.cpp
+++ b/Engine/Source/Runtime/Engine/Private/ShaderCompiler/ShaderCompiler.cpp
@@ -2772,8 +2772,11 @@ FShaderCompilingManager::FShaderCompilingManager() :
 	int32 NumUnusedShaderCompilingThreadsDuringGame;
 	verify(GConfig->GetInt(TEXT("DevOptions.Shaders"), TEXT("NumUnusedShaderCompilingThreadsDuringGame"), NumUnusedShaderCompilingThreadsDuringGame, GEngineIni));
 
-	if (NumVirtualCores > ShaderCompilerCoreCountThreshold)
+	// Don't reserve threads based on a percentage if we are in a commandlet or on a low core machine.
+	// In these scenarios we should try to use as many threads as possible.
+	if (!IsRunningCommandlet() && NumVirtualCores > ShaderCompilerCoreCountThreshold)
 	{
+		// Reserve a percentage of the threads for general background work.
 		float PercentageUnusedShaderCompilingThreads;
 		verify(GConfig->GetFloat(TEXT("DevOptions.Shaders"), TEXT("PercentageUnusedShaderCompilingThreads"), PercentageUnusedShaderCompilingThreads, GEngineIni));
 
