WP HLOD - Added more (verbose) logging

diff --git a/Engine/Source/Runtime/Engine/Private/WorldPartition/HLOD/HLODSubsystem.cpp b/Engine/Source/Runtime/Engine/Private/WorldPartition/HLOD/HLODSubsystem.cpp
--- a/Engine/Source/Runtime/Engine/Private/WorldPartition/HLOD/HLODSubsystem.cpp
+++ b/Engine/Source/Runtime/Engine/Private/WorldPartition/HLOD/HLODSubsystem.cpp
@@ -82,6 +82,9 @@ void UHLODSubsystem::RegisterHLODActor(AWorldPartitionHLOD* InWorldPartitionHLOD
 
 	FName CellName = InWorldPartitionHLOD->GetCellName();
 	FCellHLODMapping* CellHLODs = CellsHLODMapping.Find(CellName);
+
+	UE_LOG(LogHLODSubsystem, Verbose, TEXT("Registering HLOD %s (%s) for cell %s"), *InWorldPartitionHLOD->GetActorLabel(), *InWorldPartitionHLOD->GetActorGuid().ToString(), *CellName.ToString());
+
 	if (CellHLODs)
 	{
 		CellHLODs->LoadedHLODs.Add(InWorldPartitionHLOD);
@@ -98,6 +101,8 @@ void UHLODSubsystem::UnregisterHLODActor(AWorldPartitionHLOD* InWorldPartitionHL
 {
 	TRACE_CPUPROFILER_EVENT_SCOPE(UHLODSubsystem::UnregisterHLODActor);
 
+	UE_LOG(LogHLODSubsystem, Verbose, TEXT("Unregistering HLOD %s (%s) for cell %s"), *InWorldPartitionHLOD->GetActorLabel(), *InWorldPartitionHLOD->GetActorGuid().ToString(), *InWorldPartitionHLOD->GetCellName().ToString());
+
 	FName CellName = InWorldPartitionHLOD->GetCellName();
 	FCellHLODMapping* CellHLODs = CellsHLODMapping.Find(CellName);
 	if (CellHLODs)
@@ -111,8 +116,12 @@ void UHLODSubsystem::OnCellShown(const UWorldPartitionRuntimeCell* InCell)
 {
 	FCellHLODMapping& CellHLODs = CellsHLODMapping.FindChecked(InCell->GetFName());
 	CellHLODs.bIsCellVisible = true;
+
+	UE_LOG(LogHLODSubsystem, Verbose, TEXT("Cell shown - %s - hiding %d HLOD actors"), *InCell->GetName(), CellHLODs.LoadedHLODs.Num());
+
 	for (AWorldPartitionHLOD* HLODActor : CellHLODs.LoadedHLODs)
 	{
+		UE_LOG(LogHLODSubsystem, Verbose, TEXT("\t\t%s - %s"), *HLODActor->GetActorLabel(), *HLODActor->GetActorGuid().ToString());
 		HLODActor->SetVisibility(false);
 	}
 }
@@ -121,8 +130,12 @@ void UHLODSubsystem::OnCellHidden(const UWorldPartitionRuntimeCell* InCell)
 {
 	FCellHLODMapping& CellHLODs = CellsHLODMapping.FindChecked(InCell->GetFName());
 	CellHLODs.bIsCellVisible = false;
+
+	UE_LOG(LogHLODSubsystem, Verbose, TEXT("Cell hidden - %s - showing %d HLOD actors"), *InCell->GetName(), CellHLODs.LoadedHLODs.Num());
+
 	for (AWorldPartitionHLOD* HLODActor : CellHLODs.LoadedHLODs)
 	{
+		UE_LOG(LogHLODSubsystem, Verbose, TEXT("\t\t%s - %s"), *HLODActor->GetActorLabel(), *HLODActor->GetActorGuid().ToString());
 		HLODActor->SetVisibility(true);
 	}
 }
diff --git a/Engine/Source/Runtime/Engine/Private/WorldPartition/RuntimeSpatialHash/RuntimeSpatialHashGridHelper.cpp b/Engine/Source/Runtime/Engine/Private/WorldPartition/RuntimeSpatialHash/RuntimeSpatialHashGridHelper.cpp
--- a/Engine/Source/Runtime/Engine/Private/WorldPartition/RuntimeSpatialHash/RuntimeSpatialHashGridHelper.cpp
+++ b/Engine/Source/Runtime/Engine/Private/WorldPartition/RuntimeSpatialHash/RuntimeSpatialHashGridHelper.cpp
@@ -219,11 +219,9 @@ FSquare2DGridHelper GetPartitionedActors(const UWorldPartition* WorldPartition,
 		{
 			if (ActorCluster->Actors.Num() > 1)
 			{
-				static UEnum* ActorGridPlacementEnum = FindObject<UEnum>(ANY_PACKAGE, TEXT("EActorGridPlacement"));
-
 				UE_LOG(LogWorldPartition, Verbose, TEXT("Clustered %d actors (%s%s%s), generated shared BV of [%d x %d] (meters)"),
 					ActorCluster->Actors.Num(),
-					*ActorGridPlacementEnum->GetNameStringByValue((int64)GridPlacement),
+					*StaticEnum<EActorGridPlacement>()->GetNameStringByValue((int64)GridPlacement),
 					bAlwaysLoadedPromotedCluster ? TEXT(":PromotedCluster") : TEXT(""),
 					bAlwaysLoadedPromotedOutOfGrid ? TEXT(":PromotedOutOfGrid") : TEXT(""),
 					(int)(0.01f * ClusterInstance->Bounds.GetSize().X),
diff --git a/Engine/Source/Runtime/Engine/Private/WorldPartition/RuntimeSpatialHash/RuntimeSpatialHashHLOD.cpp b/Engine/Source/Runtime/Engine/Private/WorldPartition/RuntimeSpatialHash/RuntimeSpatialHashHLOD.cpp
--- a/Engine/Source/Runtime/Engine/Private/WorldPartition/RuntimeSpatialHash/RuntimeSpatialHashHLOD.cpp
+++ b/Engine/Source/Runtime/Engine/Private/WorldPartition/RuntimeSpatialHash/RuntimeSpatialHashHLOD.cpp
@@ -153,6 +153,8 @@ static TArray<FGuid> GenerateHLODsForGrid(UWorldPartition* WorldPartition, const
 		}
 	});
 
+	UE_LOG(LogWorldPartitionRuntimeSpatialHashHLOD, Verbose, TEXT("Building HLODs for grid %s..."), *RuntimeGrid.GridName.ToString());
+
 	FScopedSlowTask SlowTask(NbCellsToProcess, FText::FromString(FString::Printf(TEXT("Building HLODs for grid %s..."), *RuntimeGrid.GridName.ToString())));
 	SlowTask.MakeDialog();
 
@@ -168,16 +170,19 @@ static TArray<FGuid> GenerateHLODsForGrid(UWorldPartition* WorldPartition, const
 		PartitionedActors.GetCellBounds(CellCoord, CellBounds2D);
 		FBox CellBounds = FBox(FVector(CellBounds2D.Min, WorldBounds.Min.Z), FVector(CellBounds2D.Max, WorldBounds.Max.Z));
 
+		FIntVector CellGlobalCoord;
+		verify(PartitionedActors.GetCellGlobalCoords(CellCoord, CellGlobalCoord));
+
 		for (const FSquare2DGridHelper::FGridLevel::FGridCellDataChunk& GridCellDataChunk : GridCell.GetDataChunks())
 		{
 			const bool bShouldGenerateHLODs = ShouldGenerateHLODs(GridCell, GridCellDataChunk);
 			if (bShouldGenerateHLODs)
 			{
 				SlowTask.EnterProgressFrame(1);
-				FIntVector CellGlobalCoord;
-				verify(PartitionedActors.GetCellGlobalCoords(CellCoord, CellGlobalCoord));
 				FName CellName = UWorldPartitionRuntimeSpatialHash::GetCellName(WorldPartition, RuntimeGrid.GridName, CellGlobalCoord, GridCellDataChunk.GetDataLayersID());
 
+				UE_LOG(LogWorldPartitionRuntimeSpatialHashHLOD, Verbose, TEXT("Creating HLOD for cell %s at %s..."), *CellName.ToString(), *CellCoord.ToString());
+
 				TRACE_CPUPROFILER_EVENT_SCOPE_TEXT(*CellName.ToString());
 
 				UE_LOG(LogWorldPartitionRuntimeSpatialHashHLOD, Display, TEXT("[%d / %d] Creating/updating HLOD actor(s) for cell %s"), (int32)SlowTask.CompletedWork, (int32)SlowTask.TotalAmountOfWork, *CellName.ToString());
@@ -233,6 +238,12 @@ static TArray<FGuid> GenerateHLODsForGrid(UWorldPartition* WorldPartition, const
 					for (AWorldPartitionHLOD* CellHLODActor : CellHLODActors)
 					{
 						check(WorldPartition->GetActorDesc(CellHLODActor->GetActorGuid()));
+
+						UE_LOG(LogWorldPartitionRuntimeSpatialHashHLOD, Verbose, TEXT("Created HLOD actor %s - %s, for cell %s, represented actors:"), *CellHLODActor->GetName(), *CellHLODActor->GetActorGuid().ToString(), *CellName.ToString());
+						for (const auto& SubActorGuid : CellHLODActor->GetSubActors())
+						{
+							UE_LOG(LogWorldPartitionRuntimeSpatialHashHLOD, Verbose, TEXT("\t\t%s - %s"), *WorldPartition->GetActorDesc(SubActorGuid)->GetActorLabel().ToString(), *SubActorGuid.ToString());
+						}
 					}
 
 					// Update newly created HLOD actors
diff --git a/Engine/Source/Runtime/Engine/Private/WorldPartition/WorldPartitionRuntimeHash.cpp b/Engine/Source/Runtime/Engine/Private/WorldPartition/WorldPartitionRuntimeHash.cpp
--- a/Engine/Source/Runtime/Engine/Private/WorldPartition/WorldPartitionRuntimeHash.cpp
+++ b/Engine/Source/Runtime/Engine/Private/WorldPartition/WorldPartitionRuntimeHash.cpp
@@ -145,7 +145,7 @@ void UWorldPartitionRuntimeHash::ChangeActorDescViewGridPlacement(FWorldPartitio
 	if (ActorDescView.EffectiveGridPlacement != GridPlacement)
 	{
 		ActorDescView.EffectiveGridPlacement = GridPlacement;
-		UE_LOG(LogWorldPartition, Verbose, TEXT("Actor '%s' grid placement changed to %d"), *ActorDescView.GetActorLabel().ToString(), (int32)GridPlacement);
+		UE_LOG(LogWorldPartition, Verbose, TEXT("Actor '%s' grid placement changed to %s"), *ActorDescView.GetActorLabel().ToString(), *StaticEnum<EActorGridPlacement>()->GetNameStringByValue((int64)GridPlacement));
 	}
 }
 
