!editor build fix
#buildfix

diff --git a/Engine/Source/Runtime/Engine/Private/WorldPartition/HLOD/HLODSubsystem.cpp b/Engine/Source/Runtime/Engine/Private/WorldPartition/HLOD/HLODSubsystem.cpp
--- a/Engine/Source/Runtime/Engine/Private/WorldPartition/HLOD/HLODSubsystem.cpp
+++ b/Engine/Source/Runtime/Engine/Private/WorldPartition/HLOD/HLODSubsystem.cpp
@@ -83,7 +83,9 @@ void UHLODSubsystem::RegisterHLODActor(AWorldPartitionHLOD* InWorldPartitionHLOD
 	FName CellName = InWorldPartitionHLOD->GetCellName();
 	FCellHLODMapping* CellHLODs = CellsHLODMapping.Find(CellName);
 
+#if WITH_EDITOR
 	UE_LOG(LogHLODSubsystem, Verbose, TEXT("Registering HLOD %s (%s) for cell %s"), *InWorldPartitionHLOD->GetActorLabel(), *InWorldPartitionHLOD->GetActorGuid().ToString(), *CellName.ToString());
+#endif
 
 	if (CellHLODs)
 	{
@@ -101,7 +103,9 @@ void UHLODSubsystem::UnregisterHLODActor(AWorldPartitionHLOD* InWorldPartitionHL
 {
 	TRACE_CPUPROFILER_EVENT_SCOPE(UHLODSubsystem::UnregisterHLODActor);
 
+#if WITH_EDITOR
 	UE_LOG(LogHLODSubsystem, Verbose, TEXT("Unregistering HLOD %s (%s) for cell %s"), *InWorldPartitionHLOD->GetActorLabel(), *InWorldPartitionHLOD->GetActorGuid().ToString(), *InWorldPartitionHLOD->GetCellName().ToString());
+#endif
 
 	FName CellName = InWorldPartitionHLOD->GetCellName();
 	FCellHLODMapping* CellHLODs = CellsHLODMapping.Find(CellName);
@@ -117,11 +121,15 @@ void UHLODSubsystem::OnCellShown(const UWorldPartitionRuntimeCell* InCell)
 	FCellHLODMapping& CellHLODs = CellsHLODMapping.FindChecked(InCell->GetFName());
 	CellHLODs.bIsCellVisible = true;
 
+#if WITH_EDITOR
 	UE_LOG(LogHLODSubsystem, Verbose, TEXT("Cell shown - %s - hiding %d HLOD actors"), *InCell->GetName(), CellHLODs.LoadedHLODs.Num());
+#endif
 
 	for (AWorldPartitionHLOD* HLODActor : CellHLODs.LoadedHLODs)
 	{
+#if WITH_EDITOR
 		UE_LOG(LogHLODSubsystem, Verbose, TEXT("\t\t%s - %s"), *HLODActor->GetActorLabel(), *HLODActor->GetActorGuid().ToString());
+#endif
 		HLODActor->SetVisibility(false);
 	}
 }
@@ -131,11 +139,15 @@ void UHLODSubsystem::OnCellHidden(const UWorldPartitionRuntimeCell* InCell)
 	FCellHLODMapping& CellHLODs = CellsHLODMapping.FindChecked(InCell->GetFName());
 	CellHLODs.bIsCellVisible = false;
 
+#if WITH_EDITOR
 	UE_LOG(LogHLODSubsystem, Verbose, TEXT("Cell hidden - %s - showing %d HLOD actors"), *InCell->GetName(), CellHLODs.LoadedHLODs.Num());
+#endif
 
 	for (AWorldPartitionHLOD* HLODActor : CellHLODs.LoadedHLODs)
 	{
+#if WITH_EDITOR
 		UE_LOG(LogHLODSubsystem, Verbose, TEXT("\t\t%s - %s"), *HLODActor->GetActorLabel(), *HLODActor->GetActorGuid().ToString());
+#endif
 		HLODActor->SetVisibility(true);
 	}
 }
