Removed a MaterialInstance assert that trips erroneously when doing `memreport -full`

diff --git a/Engine/Source/Runtime/Engine/Private/Materials/MaterialInstance.cpp b/Engine/Source/Runtime/Engine/Private/Materials/MaterialInstance.cpp
--- a/Engine/Source/Runtime/Engine/Private/Materials/MaterialInstance.cpp
+++ b/Engine/Source/Runtime/Engine/Private/Materials/MaterialInstance.cpp
@@ -3379,7 +3379,7 @@ void UMaterialInstance::Serialize(FArchive& Ar)
 #endif // WITH_EDITOR
 
 #if !WITH_EDITORONLY_DATA
-	checkf(bSavedCachedData, TEXT("MaterialInstance %s must have saved cached data, if editor-only data is not present"), *GetName());
+	//checkf(bSavedCachedData, TEXT("MaterialInstance %s must have saved cached data, if editor-only data is not present"), *GetName());
 #endif
 	if (bSavedCachedData)
 	{
