Rework how unsaved actors are duplicated in PIE to avoid having a time period where actor templates are in the PIE persistent level, making them "visible" to unfixed soft object paths, by moving them to the duplication container at post duplication.

diff --git a/Engine/Source/Runtime/Engine/Classes/Engine/Level.h b/Engine/Source/Runtime/Engine/Classes/Engine/Level.h
--- a/Engine/Source/Runtime/Engine/Classes/Engine/Level.h
+++ b/Engine/Source/Runtime/Engine/Classes/Engine/Level.h
@@ -76,6 +76,17 @@ private:
 	ILevelPartitionInterface* InterfacePtr = nullptr;
 	ULevel* Level = nullptr;
 };
+
+// Actor container class used to duplicate actors during cells streaming in PIE
+UCLASS()
+class UActorContainer : public UObject
+{
+	GENERATED_BODY()
+
+public:
+	UPROPERTY()
+	TMap<FName, AActor*> Actors;
+};
 #endif
 
 /**
@@ -440,10 +451,6 @@ public:
 	TArray<AActor*> ActorsForGC;
 
 #if WITH_EDITORONLY_DATA
-	/** List of modified, unsaved actors that needs to be duplicated for PIE */
-	UPROPERTY(NonPIEDuplicateTransient)
-	TMap<FName, AActor*> ActorsModifiedForPIE;
-
 	/** Use external actors, new actor spawned in this level will be external and existing external actors will be loaded on load. */
 	UPROPERTY(EditInstanceOnly, Category=World)
 	bool bUseExternalActors;
diff --git a/Engine/Source/Runtime/Engine/Private/WorldPartition/WorldPartition.cpp b/Engine/Source/Runtime/Engine/Private/WorldPartition/WorldPartition.cpp
--- a/Engine/Source/Runtime/Engine/Private/WorldPartition/WorldPartition.cpp
+++ b/Engine/Source/Runtime/Engine/Private/WorldPartition/WorldPartition.cpp
@@ -153,15 +153,6 @@ void UWorldPartition::OnPreBeginPIE(bool bStartSimulate)
 	check(!bIsPIE);
 	bIsPIE = true;
 
-	// Gather all modified/newly created actors for PIE so they can be duplicated when streaming their cell
-	for (AActor* Actor : World->PersistentLevel->Actors)
-	{
-		if (Actor && Actor->IsPackageExternal() && !Actor->IsPendingKill() && Actor->GetPackage()->IsDirty())
-		{
-			World->PersistentLevel->ActorsModifiedForPIE.Add(Actor->GetFName(), Actor);
-		}
-	}
-
 	check(IsMainWorldPartition());
 
 	OnBeginPlay(EWorldPartitionStreamingMode::PIE);
@@ -187,8 +178,6 @@ void UWorldPartition::OnEndPlay()
 	RuntimeHash->OnEndPlay();
 
 	StreamingPolicy = nullptr;
-
-	World->PersistentLevel->ActorsModifiedForPIE.Empty();
 }
 
 FName UWorldPartition::GetWorldPartitionEditorName()
diff --git a/Engine/Source/Runtime/Engine/Private/WorldPartition/WorldPartitionLevelStreamingDynamic.cpp b/Engine/Source/Runtime/Engine/Private/WorldPartition/WorldPartitionLevelStreamingDynamic.cpp
--- a/Engine/Source/Runtime/Engine/Private/WorldPartition/WorldPartitionLevelStreamingDynamic.cpp
+++ b/Engine/Source/Runtime/Engine/Private/WorldPartition/WorldPartitionLevelStreamingDynamic.cpp
@@ -51,15 +51,15 @@ void UWorldPartitionLevelStreamingDynamic::Initialize(const UWorldPartitionRunti
 	check(ChildPackages.Num() == 0);
 	check(!WorldAsset.IsNull());
 
-	bIsModifiedForPIE = InCell.IsModifiedForPIE();
 	bShouldBeAlwaysLoaded = InCell.IsAlwaysLoaded();
 	StreamingPriority = InCell.GetStreamingPriority();
 	ChildPackages = InCell.GetPackages();
+	ActorContainer = InCell.ActorContainer;
 
 	UWorld* OuterWorld = InCell.GetOuterUWorldPartition()->GetTypedOuter<UWorld>();
 	OriginalLevelPackageName = OuterWorld->GetPackage()->GetLoadedPath().GetPackageFName();
 	PackageNameToLoad = GetWorldAssetPackageFName();
-	OuterWorldPartition = OuterWorld->GetWorldPartition();
+	OuterWorldPartition = OuterWorld->GetWorldPartition();	
 }
 
 /**
@@ -225,9 +225,6 @@ bool UWorldPartitionLevelStreamingDynamic::IssueLoadRequests()
 	UPackage* RuntimePackage = RuntimeLevel->GetPackage();
 	InstancingContext.AddMapping(OriginalLevelPackageName, RuntimePackage->GetFName());
 
-	TArray<AActor*> ActorsToDuplicate;
-	ActorsToDuplicate.Reserve(ChildPackages.Num());
-
 	TArray<FWorldPartitionRuntimeCellObjectMapping> ChildPackagesToLoad;
 	ChildPackagesToLoad.Reserve(ChildPackages.Num());
 
@@ -237,15 +234,14 @@ bool UWorldPartitionLevelStreamingDynamic::IssueLoadRequests()
 		if (ChildPackage.ContainerID == 0)
 		{
 			bool bNeedDup = false;
-			if (bIsModifiedForPIE)
+			if (ActorContainer)
 			{
 				FString SubObjectName;
 				FString SubObjectContext;
 				if (ChildPackage.LoadedPath.ToString().Split(TEXT("."), &SubObjectContext, &SubObjectName, ESearchCase::IgnoreCase, ESearchDir::FromEnd))
 				{
-					if (AActor* ActorModifiedForPIE = World->PersistentLevel->ActorsModifiedForPIE.FindRef(*SubObjectName))
+					if (AActor* ActorModifiedForPIE = ActorContainer->Actors.FindRef(*SubObjectName))
 					{
-						ActorsToDuplicate.Add(ActorModifiedForPIE);
 						bNeedDup = true;
 					}
 				}
@@ -259,17 +255,8 @@ bool UWorldPartitionLevelStreamingDynamic::IssueLoadRequests()
 	}
 
 	// Duplicate unsaved actors
-	if (ActorsToDuplicate.Num())
+	if (ActorContainer)
 	{
-		// Create an actor container to make sure duplicated actors will share an outer to properly remap inter-actors references
-		UActorContainer* ActorContainer = NewObject<UActorContainer>(World->PersistentLevel);
-
-		for (AActor* ActorToDuplicate : ActorsToDuplicate)
-		{
-			ActorToDuplicate->UObject::Rename(nullptr, ActorContainer);
-			ActorContainer->Actors.Add(ActorToDuplicate);
-		}
-
 		FObjectDuplicationParameters Parameters(ActorContainer, RuntimeLevel);
 		Parameters.DestClass = ActorContainer->GetClass();
 		Parameters.FlagMask = RF_AllFlags & ~(RF_MarkAsRootSet | RF_MarkAsNative | RF_HasExternalPackage);
@@ -281,18 +268,11 @@ bool UWorldPartitionLevelStreamingDynamic::IssueLoadRequests()
 		UActorContainer* ActorContainerDup = (UActorContainer*)StaticDuplicateObjectEx(Parameters);
 
 		// Add the duplicated actors to the corresponding cell level
-		for (AActor* Actor : ActorContainerDup->Actors)
+		for (auto& ActorPair : ActorContainerDup->Actors)
 		{
-			Actor->Rename(nullptr, RuntimeLevel);
+			ActorPair.Value->Rename(nullptr, RuntimeLevel);
 		}
 
-		// Bring back actors to the persistent level
-		for (AActor* ActorToDuplicate : ActorsToDuplicate)
-		{
-			ActorToDuplicate->UObject::Rename(nullptr, World->PersistentLevel);
-		}
-
-		ActorContainer->MarkPendingKill();
 		ActorContainerDup->MarkPendingKill();
 	}
 
@@ -486,5 +466,4 @@ UWorld* UWorldPartitionLevelStreamingDynamic::GetOuterWorld() const
 	return OuterWorldPartition->GetTypedOuter<UWorld>();
 }
 
-
 #undef LOCTEXT_NAMESPACE
\ No newline at end of file
diff --git a/Engine/Source/Runtime/Engine/Private/WorldPartition/WorldPartitionRuntimeCell.cpp b/Engine/Source/Runtime/Engine/Private/WorldPartition/WorldPartitionRuntimeCell.cpp
--- a/Engine/Source/Runtime/Engine/Private/WorldPartition/WorldPartitionRuntimeCell.cpp
+++ b/Engine/Source/Runtime/Engine/Private/WorldPartition/WorldPartitionRuntimeCell.cpp
@@ -5,11 +5,7 @@
 UWorldPartitionRuntimeCell::UWorldPartitionRuntimeCell(const FObjectInitializer& ObjectInitializer)
 : Super(ObjectInitializer)
 , bIsAlwaysLoaded(false)
-#if WITH_EDITORONLY_DATA
-, bIsModifiedForPIE(false)
-#endif
-{
-}
+{}
 
 #if WITH_EDITOR
 
@@ -38,5 +34,4 @@ const UWorldPartitionRuntimeCellData* UWorldPartitionRuntimeCell::GetCellData(co
 
 UWorldPartitionRuntimeCellData::UWorldPartitionRuntimeCellData(const FObjectInitializer& ObjectInitializer)
 : Super(ObjectInitializer)
-{
-}
+{}
diff --git a/Engine/Source/Runtime/Engine/Private/WorldPartition/WorldPartitionRuntimeHash.cpp b/Engine/Source/Runtime/Engine/Private/WorldPartition/WorldPartitionRuntimeHash.cpp
--- a/Engine/Source/Runtime/Engine/Private/WorldPartition/WorldPartitionRuntimeHash.cpp
+++ b/Engine/Source/Runtime/Engine/Private/WorldPartition/WorldPartitionRuntimeHash.cpp
@@ -114,11 +114,15 @@ void UWorldPartitionRuntimeHash::CreateActorDescViewMap(const UActorDescContaine
 		// Append new unsaved actors for the persistent level
 		if (Container->GetContainerPackage() == GetWorld()->PersistentLevel->GetPackage()->GetLoadedPath().GetPackageFName())
 		{
-			for (auto& ActorPair : GetWorld()->PersistentLevel->ActorsModifiedForPIE)
+			for (AActor* Actor : GetWorld()->PersistentLevel->Actors)
 			{
-				if (ActorPair.Value->GetPackage()->HasAnyPackageFlags(PKG_NewlyCreated))
+				if (Actor && 
+					Actor->IsPackageExternal() && 
+					!Actor->IsPendingKill() && 
+					Actor->GetPackage()->IsDirty() &&
+					Actor->GetPackage()->HasAnyPackageFlags(PKG_NewlyCreated))
 				{
-					FWorldPartitionActorDesc* ActorDesc = ModifiedActorDescListForPIE.AddActor(ActorPair.Value);
+					FWorldPartitionActorDesc* ActorDesc = ModifiedActorDescListForPIE.AddActor(Actor);
 					OutActorDescViewMap.Emplace(ActorDesc->GetGuid(), ActorDesc);
 				}
 			}
diff --git a/Engine/Source/Runtime/Engine/Private/WorldPartition/WorldPartitionRuntimeSpatialHash.cpp b/Engine/Source/Runtime/Engine/Private/WorldPartition/WorldPartitionRuntimeSpatialHash.cpp
--- a/Engine/Source/Runtime/Engine/Private/WorldPartition/WorldPartitionRuntimeSpatialHash.cpp
+++ b/Engine/Source/Runtime/Engine/Private/WorldPartition/WorldPartitionRuntimeSpatialHash.cpp
@@ -727,7 +727,6 @@ bool UWorldPartitionRuntimeSpatialHash::CreateStreamingGrid(const FSpatialHashRu
 			const uint32 CellIndex = TempCellMapping.Key;
 			const int32 CellCoordX = CellIndex % TempLevel.GridSize;
 			const int32 CellCoordY = CellIndex / TempLevel.GridSize;
-			bool bCellModifiedForPIE = false;
 
 			const FSquare2DGridHelper::FGridLevel::FGridCell& TempCell = TempLevel.Cells[TempCellMapping.Value];
 
@@ -769,17 +768,6 @@ bool UWorldPartitionRuntimeSpatialHash::CreateStreamingGrid(const FSpatialHashRu
 									AlwaysLoadedActorsForPIE.Emplace(Reference, AlwaysLoadedActor);
 									continue;
 								}
-							} // InPIE, If a Cell as at least one modified actor we need to flag all actors of this cell so they get added to the ActorsModifiedForPIE so they are all duplicated
-							else if(!bCellModifiedForPIE && Mode == EWorldPartitionStreamingMode::PIE && World->PersistentLevel->ActorsModifiedForPIE.Num() > 0)
-							{
-								FString SubObjectName;
-								FString SubObjectContext;
-								
-								ActorDescView.GetActorPath().ToString().Split(TEXT("."), &SubObjectContext, &SubObjectName, ESearchCase::IgnoreCase, ESearchDir::FromEnd);
-								if (World->PersistentLevel->ActorsModifiedForPIE.FindRef(*SubObjectName))
-								{
-									bCellModifiedForPIE = true;
-								}
 							}
 						}
 
@@ -812,7 +800,6 @@ bool UWorldPartitionRuntimeSpatialHash::CreateStreamingGrid(const FSpatialHashRu
 
 				GridLevel.LayerCells[LayerCellIndex].GridCells.Add(StreamingCell);
 				StreamingCell->SetIsAlwaysLoaded(bIsCellAlwaysLoaded);
-				StreamingCell->SetIsModifiedForPIE(bCellModifiedForPIE);
 				StreamingCell->SetDataLayers(GridCellDataChunk.GetDataLayers());
 				StreamingCell->Level = Level;
 				FBox2D Bounds;
@@ -821,6 +808,21 @@ bool UWorldPartitionRuntimeSpatialHash::CreateStreamingGrid(const FSpatialHashRu
 
 				UE_LOG(LogWorldPartition, Verbose, TEXT("Cell%s %s Actors = %d Bounds (%s)"), bIsCellAlwaysLoaded ? TEXT(" (AlwaysLoaded)") : TEXT(""), *StreamingCell->GetName(), FilteredActors.Num(), *Bounds.ToString());
 
+				check(!StreamingCell->ActorContainer);
+				for (const FActorInstance& ActorInstance : FilteredActors)
+				{
+					const FWorldPartitionActorDescView& ActorDescView = ActorInstance.GetActorDescView();
+					if (AActor* Actor = FindObject<AActor>(nullptr, *ActorDescView.GetActorPath().ToString()))
+					{
+						if (Actor->GetPackage()->IsDirty())
+						{
+							// Create an actor container to make sure duplicated actors will share an outer to properly remap inter-actors references
+							StreamingCell->ActorContainer = NewObject<UActorContainer>(StreamingCell);
+							break;
+						}
+					}
+				}
+
 				// Keep track of all AWorldPartitionHLOD actors referenced by this cell
 				int32 NumHLODActors = 0;
 
@@ -829,12 +831,11 @@ bool UWorldPartitionRuntimeSpatialHash::CreateStreamingGrid(const FSpatialHashRu
 					const FWorldPartitionActorDescView& ActorDescView = ActorInstance.GetActorDescView();
 					StreamingCell->AddActorToCell(ActorDescView, ActorInstance.ContainerInstance->ID, ActorInstance.ContainerInstance->Transform, ActorInstance.ContainerInstance->Container);
 
-					// Add all actors of this cell to the ActorsModifiedForPIE (for Duplication)
-					if (StreamingCell->IsModifiedForPIE())
+					if (StreamingCell->ActorContainer)
 					{
 						if (AActor* Actor = FindObject<AActor>(nullptr, *ActorDescView.GetActorPath().ToString()))
 						{
-							World->PersistentLevel->ActorsModifiedForPIE.Add(Actor->GetFName(), Actor);
+							StreamingCell->ActorContainer->Actors.Add(Actor->GetFName(), Actor);
 						}
 					}
 
@@ -1158,4 +1159,4 @@ void UWorldPartitionRuntimeSpatialHash::Draw3D(const TArray<FWorldPartitionStrea
 			StreamingGrid.Draw3D(World, Sources, Transform);
 		}
 	}
-}
\ No newline at end of file
+}
diff --git a/Engine/Source/Runtime/Engine/Private/WorldPartition/WorldPartitionRuntimeSpatialHashCell.cpp b/Engine/Source/Runtime/Engine/Private/WorldPartition/WorldPartitionRuntimeSpatialHashCell.cpp
--- a/Engine/Source/Runtime/Engine/Private/WorldPartition/WorldPartitionRuntimeSpatialHashCell.cpp
+++ b/Engine/Source/Runtime/Engine/Private/WorldPartition/WorldPartitionRuntimeSpatialHashCell.cpp
@@ -5,3 +5,20 @@ UWorldPartitionRuntimeSpatialHashCell::UWorldPartitionRuntimeSpatialHashCell(con
 : Super(ObjectInitializer)
 , Level(0)
 {}
+
+#if WITH_EDITOR
+void UWorldPartitionRuntimeSpatialHashCell::PostDuplicate(bool bDuplicateForPIE)
+{
+	Super::PostDuplicate(bDuplicateForPIE);
+
+	if (ActorContainer)
+	{
+		for (auto& ActorPair : ActorContainer->Actors)
+		{
+			// Don't use AActor::Rename here since the actor is not par of the world, it's only a duplication template.
+			ActorPair.Value->UObject::Rename(nullptr, ActorContainer);
+		}
+	}
+}
+#endif
+
diff --git a/Engine/Source/Runtime/Engine/Public/WorldPartition/WorldPartitionLevelStreamingDynamic.h b/Engine/Source/Runtime/Engine/Public/WorldPartition/WorldPartitionLevelStreamingDynamic.h
--- a/Engine/Source/Runtime/Engine/Public/WorldPartition/WorldPartitionLevelStreamingDynamic.h
+++ b/Engine/Source/Runtime/Engine/Public/WorldPartition/WorldPartitionLevelStreamingDynamic.h
@@ -65,20 +65,9 @@ private:
 
 #if WITH_EDITORONLY_DATA
 	UPROPERTY()
-	bool bIsModifiedForPIE;
+	UActorContainer* ActorContainer;
 #endif
 
 	UPROPERTY()
 	TWeakObjectPtr<UWorldPartition> OuterWorldPartition;
-};
-
-// Actor container class used to duplicate actors during cells streaming
-UCLASS(Transient)
-class UActorContainer : public UObject
-{
-	GENERATED_BODY()
-
-public:
-	UPROPERTY()
-	TArray<AActor*> Actors;
-};
+};
\ No newline at end of file
diff --git a/Engine/Source/Runtime/Engine/Public/WorldPartition/WorldPartitionRuntimeCell.h b/Engine/Source/Runtime/Engine/Public/WorldPartition/WorldPartitionRuntimeCell.h
--- a/Engine/Source/Runtime/Engine/Public/WorldPartition/WorldPartitionRuntimeCell.h
+++ b/Engine/Source/Runtime/Engine/Public/WorldPartition/WorldPartitionRuntimeCell.h
@@ -117,9 +117,6 @@ class UWorldPartitionRuntimeCell : public UObject
 	const TArray<FName>& GetDataLayers() const { return DataLayers; }
 
 #if WITH_EDITOR
-	bool IsModifiedForPIE() const { return bIsModifiedForPIE; };
-	void SetIsModifiedForPIE(bool bInIsModifiedForPIE) { bIsModifiedForPIE = bInIsModifiedForPIE; }
-	
 	void SetDataLayers(const TArray<const UDataLayer*> InDataLayers);
 	void AddCellData(const UWorldPartitionRuntimeCellData* InCellData);
 	virtual void AddActorToCell(const FWorldPartitionActorDescView& ActorDescView, uint32 InContainerID, const FTransform& InContainerTransform, const UActorDescContainer* InContainer) PURE_VIRTUAL(UWorldPartitionRuntimeCell::AddActorToCell,);
@@ -140,11 +137,6 @@ protected:
 	UPROPERTY()
 	bool bIsAlwaysLoaded;
 
-#if WITH_EDITORONLY_DATA
-	UPROPERTY()
-	bool bIsModifiedForPIE;
-#endif
-
 private:
 	UPROPERTY()
 	TMap<const TSubclassOf<UWorldPartitionRuntimeCellData>, TObjectPtr<const UWorldPartitionRuntimeCellData>> CellDataMap;
diff --git a/Engine/Source/Runtime/Engine/Public/WorldPartition/WorldPartitionRuntimeSpatialHashCell.h b/Engine/Source/Runtime/Engine/Public/WorldPartition/WorldPartitionRuntimeSpatialHashCell.h
--- a/Engine/Source/Runtime/Engine/Public/WorldPartition/WorldPartitionRuntimeSpatialHashCell.h
+++ b/Engine/Source/Runtime/Engine/Public/WorldPartition/WorldPartitionRuntimeSpatialHashCell.h
@@ -10,9 +10,18 @@ class UWorldPartitionRuntimeSpatialHashCell : public UWorldPartitionRuntimeCell
 {
 	GENERATED_UCLASS_BODY()
 
+#if WITH_EDITOR
+	virtual void PostDuplicate(bool bDuplicateForPIE) override;
+#endif
+
 	UPROPERTY()
 	FVector Position;
 
 	UPROPERTY()
 	int32 Level;
+
+#if WITH_EDITORONLY_DATA
+	UPROPERTY()
+	UActorContainer* ActorContainer;
+#endif
 };
\ No newline at end of file
