Fix spurious timeouts in distributed shader compilation controller.

- Periods of inactivity shouldn't count towards the delay between received jobs.
- Also made buildmachines use local XGE SCWs as before.
- Also downgraded the error to warning and ceased from refusing new jobs, so if we misdetect we will just work as usual.

diff --git a/Engine/Plugins/XGEController/Source/Private/XGEControllerModule.cpp b/Engine/Plugins/XGEController/Source/Private/XGEControllerModule.cpp
--- a/Engine/Plugins/XGEController/Source/Private/XGEControllerModule.cpp
+++ b/Engine/Plugins/XGEController/Source/Private/XGEControllerModule.cpp
@@ -350,8 +350,9 @@ void FXGEControllerModule::WriteOutThreadProc()
 
 		// To handle spaces in the engine path, we just pass the XGEController.exe filename to xgConsole,
 		// and set the working directory of xgConsole.exe to the engine binaries folder below.
-		FString XGConsoleArgs = FString::Printf(TEXT("/VIRTUALIZEDIRECTX /allowremote=\"%s\" /avoidlocal=ON /allowintercept=\"%s\" /title=\"Unreal Engine XGE Tasks\" /monitordirs=\"%s\" /command=\"%s -xgecontroller %s\""),
+		FString XGConsoleArgs = FString::Printf(TEXT("/VIRTUALIZEDIRECTX /allowremote=\"%s\" %s /allowintercept=\"%s\" /title=\"Unreal Engine XGE Tasks\" /monitordirs=\"%s\" /command=\"%s -xgecontroller %s\""),
 			XGE_INTERCEPT_EXE_NAMES,
+			GIsBuildMachine ? TEXT("") : TEXT("/avoidlocal=ON")
 			XGE_CONTROL_WORKER_NAME,
 			*WorkingDirectory,
 			XGE_CONTROL_WORKER_FILENAME,
diff --git a/Engine/Source/Runtime/Engine/Private/ShaderCompiler/ShaderCompilerDistributed.cpp b/Engine/Source/Runtime/Engine/Private/ShaderCompiler/ShaderCompilerDistributed.cpp
--- a/Engine/Source/Runtime/Engine/Private/ShaderCompiler/ShaderCompilerDistributed.cpp
+++ b/Engine/Source/Runtime/Engine/Private/ShaderCompiler/ShaderCompilerDistributed.cpp
@@ -167,7 +167,7 @@ TArray<FString> FShaderCompileDistributedThreadRunnable_Interface::GetDependency
 int32 FShaderCompileDistributedThreadRunnable_Interface::CompilingLoop()
 {
 	TArray<FShaderCommonCompileJobPtr> PendingJobs;
-	if (LIKELY(!bIsHung))	// stop accepting jobs if we're hung
+	//if (LIKELY(!bIsHung))	// stop accepting jobs if we're hung - TODO: re-enable this after lockup detection logic is proved reliable and/or we have job resubmission in place
 	{
 		for (int32 PriorityIndex = MaxPriorityIndex; PriorityIndex >= MinPriorityIndex; --PriorityIndex)
 		{
@@ -188,6 +188,12 @@ int32 FShaderCompileDistributedThreadRunnable_Interface::CompilingLoop()
 		}
 	}
 
+	// if we don't have any dispatched jobs, reset the completion timer
+	if (DispatchedTasks.Num() == 0)
+	{
+		LastTimeTaskCompleted = FPlatformTime::Seconds();
+	}
+
 	if (PendingJobs.Num() > 0)
 	{
 		// Increase the batch size when more jobs are queued/in flight.
@@ -374,11 +380,11 @@ int32 FShaderCompileDistributedThreadRunnable_Interface::CompilingLoop()
 
 	// normally we expect to have at least one task in 5 minutes, although there could be edge cases
 	double TimeSinceLastCompletedTask = FPlatformTime::Seconds() - LastTimeTaskCompleted;
-	if (TimeSinceLastCompletedTask > DistributedShaderCompilerVariables::GDistributedControllerTimeout)
+	if (TimeSinceLastCompletedTask >= DistributedShaderCompilerVariables::GDistributedControllerTimeout)
 	{
 		if (!bIsHung)
 		{
-			UE_LOG(LogShaderCompilers, Error, TEXT("Distributed compilation controller didn't receive a completed task in %f seconds!"), TimeSinceLastCompletedTask);
+			UE_LOG(LogShaderCompilers, Warning, TEXT("Distributed compilation controller didn't receive a completed task in %f seconds!"), TimeSinceLastCompletedTask);
 			bIsHung = true;
 			// TODO: resubmit the hung jobs
 		}
