Fixed a bug that was causing textures to be saved with both PNG and zlib compression enabled.

* FTextureSource::Compress
- The original intent was to attempt to store the texture data in a compressed png format and only compress the bulkdata object with zlib compression if it was unable to do so.
- If the texture is saved a second time and the data is already stored in compressed PNG format then ::CanPNGCompress would return false which would cause the code to enable zlib compressed on the bulkdata object.
- Now we set the compression flag based on if the data is stored as a compressed PNG or not, the result of ::CanPNGCompress is only used to determine if we should try converting the data to compressed PNG or not.

diff --git a/Engine/Source/Runtime/Engine/Private/Texture.cpp b/Engine/Source/Runtime/Engine/Private/Texture.cpp
--- a/Engine/Source/Runtime/Engine/Private/Texture.cpp
+++ b/Engine/Source/Runtime/Engine/Private/Texture.cpp
@@ -1131,7 +1131,7 @@ void FTextureSource::Compress()
 	FWriteScopeLock BulkDataExclusiveScope(BulkDataLock.Get());
 #endif
 
-	if (CanPNGCompress())
+	if (CanPNGCompress()) // Note that this will return false if the data is already a compressed PNG
 	{
 		uint8* BulkDataPtr = (uint8*)BulkData.Lock(LOCK_READ_WRITE);
 		IImageWrapperModule& ImageWrapperModule = FModuleManager::LoadModuleChecked<IImageWrapperModule>( FName("ImageWrapper") );
@@ -1147,16 +1147,13 @@ void FTextureSource::Compress()
 				FMemory::Memcpy(BulkDataPtr, CompressedData.GetData(), CompressedData.Num());
 				BulkData.Unlock();
 				bPNGCompressed = true;
-
-				BulkData.StoreCompressedOnDisk(NAME_None);
 			}
 		}
 	}
-	else
-	{
-		//Can't PNG compress so just zlib compress the lot when its serialized out to disk. 
-		BulkData.StoreCompressedOnDisk(NAME_Zlib);
-	}
+	
+	// If the data is already compressed as a png then we don't need to compress it
+	// again when serializing to disk.
+	BulkData.StoreCompressedOnDisk(bPNGCompressed ? NAME_None : NAME_Zlib);
 }
 
 uint8* FTextureSource::LockMip(int32 BlockIndex, int32 LayerIndex, int32 MipIndex)
