Don't rerun construction scripts on HLOD actors.

diff --git a/Engine/Source/Runtime/Engine/Private/WorldPartition/HLOD/HLODActor.cpp b/Engine/Source/Runtime/Engine/Private/WorldPartition/HLOD/HLODActor.cpp
--- a/Engine/Source/Runtime/Engine/Private/WorldPartition/HLOD/HLODActor.cpp
+++ b/Engine/Source/Runtime/Engine/Private/WorldPartition/HLOD/HLODActor.cpp
@@ -61,6 +61,10 @@ void AWorldPartitionHLOD::PostLoad()
 #endif
 }
 
+
+void AWorldPartitionHLOD::RerunConstructionScripts()
+{}
+
 #if WITH_EDITOR
 
 void AWorldPartitionHLOD::PostRegisterAllComponents()
diff --git a/Engine/Source/Runtime/Engine/Public/WorldPartition/HLOD/HLODActor.h b/Engine/Source/Runtime/Engine/Public/WorldPartition/HLOD/HLODActor.h
--- a/Engine/Source/Runtime/Engine/Public/WorldPartition/HLOD/HLODActor.h
+++ b/Engine/Source/Runtime/Engine/Public/WorldPartition/HLOD/HLODActor.h
@@ -60,6 +60,7 @@ public:
 protected:
 	//~ Begin UObject Interface.
 	virtual void PostLoad() override;
+	virtual void RerunConstructionScripts() override;
 	//~ End UObject Interface.
 
 	//~ Begin AActor Interface.
