Remove code that used to skip rerun construction scripts in packed level instances.

diff --git a/Engine/Source/Runtime/Engine/Private/LevelInstance/Packed/PackedLevelInstanceActor.cpp b/Engine/Source/Runtime/Engine/Private/LevelInstance/Packed/PackedLevelInstanceActor.cpp
--- a/Engine/Source/Runtime/Engine/Private/LevelInstance/Packed/PackedLevelInstanceActor.cpp
+++ b/Engine/Source/Runtime/Engine/Private/LevelInstance/Packed/PackedLevelInstanceActor.cpp
@@ -172,12 +172,4 @@ void APackedLevelInstance::DestroyPackedComponents()
 		PackedComponent->DestroyComponent();
 	}
 }
-
-void APackedLevelInstance::RerunConstructionScripts()
-{
-	//if (bRerunConstructionScripts)
-	{
-		Super::RerunConstructionScripts();
-	}
-}
 #endif
\ No newline at end of file
diff --git a/Engine/Source/Runtime/Engine/Public/LevelInstance/Packed/PackedLevelInstanceActor.h b/Engine/Source/Runtime/Engine/Public/LevelInstance/Packed/PackedLevelInstanceActor.h
--- a/Engine/Source/Runtime/Engine/Public/LevelInstance/Packed/PackedLevelInstanceActor.h
+++ b/Engine/Source/Runtime/Engine/Public/LevelInstance/Packed/PackedLevelInstanceActor.h
@@ -64,8 +64,6 @@ public:
 		NewComponent->ComponentTags.Add(GetPackedComponentTag());
 		return NewComponent;
 	}
-
-	virtual void RerunConstructionScripts() override;
 #endif
 
 #if WITH_EDITORONLY_DATA
