Prevent repeated shader compile jobs submission from the same material (UE-111362)

- The priority tracking variable introduced in CL 15324895 is supposed to cause FMaterial to submit the jobs only once per priority. However, once some of the submitted jobs complete, it will get a new (partial) shadermap (in hopes that it can be now rendered) and this would reset the priority tracking.
- Delaying priority tracking until the received shadermap is complete (i.e. has all the shaders needed).

diff --git a/Engine/Source/Runtime/Engine/Private/Materials/MaterialShared.cpp b/Engine/Source/Runtime/Engine/Private/Materials/MaterialShared.cpp
--- a/Engine/Source/Runtime/Engine/Private/Materials/MaterialShared.cpp
+++ b/Engine/Source/Runtime/Engine/Private/Materials/MaterialShared.cpp
@@ -1035,7 +1035,13 @@ void FMaterial::SetRenderingThreadShaderMap(TRefCountPtr<FMaterialShaderMap>& In
 	check(IsInRenderingThread());
 	RenderingThreadShaderMap = MoveTemp(InMaterialShaderMap);
 	bRenderingThreadShaderMapIsComplete = RenderingThreadShaderMap ? RenderingThreadShaderMap->IsComplete(this, true) : false;
-	RenderingThreadShaderMapSubmittedPriority = -1;
+	// if SM isn't complete, it is perhaps a partial update incorporating results from the already submitted compile jobs.
+	// Only reset the priority if the SM is complete, as otherwise we risk resubmitting the same jobs over and over again 
+	// as FMaterialRenderProxy::GetMaterialWithFallback will queue job submissions any time it sees an incomplete SM.
+	if (bRenderingThreadShaderMapIsComplete)
+	{
+		RenderingThreadShaderMapSubmittedPriority = -1;
+	}
 }
 
 void FMaterial::AddReferencedObjects(FReferenceCollector& Collector)
