LevelInstance: Change "_Packed_BP" suffix to "BPP_" prefix for default asset naming

diff --git a/Engine/Source/Runtime/Engine/Private/LevelInstance/LevelInstanceSubsystem.cpp b/Engine/Source/Runtime/Engine/Private/LevelInstance/LevelInstanceSubsystem.cpp
--- a/Engine/Source/Runtime/Engine/Private/LevelInstance/LevelInstanceSubsystem.cpp
+++ b/Engine/Source/Runtime/Engine/Private/LevelInstance/LevelInstanceSubsystem.cpp
@@ -785,7 +785,7 @@ ALevelInstance* ULevelInstanceSubsystem::CreateLevelInstanceFrom(const TArray<AA
 		LongPackageName.FindLastChar('/', LastSlashIndex);
 
 		FString PackagePath = LongPackageName.Mid(0, LastSlashIndex == INDEX_NONE ? MAX_int32 : LastSlashIndex);
-		FString AssetName = WorldPtr.GetAssetName() + FPackedLevelInstanceBuilder::GetPackedBPSuffix();
+		FString AssetName = FPackedLevelInstanceBuilder::GetPackedBPPrefix() + WorldPtr.GetAssetName();
 		const bool bCompile = true;
 		if (UBlueprint* NewBP = FPackedLevelInstanceBuilder::CreatePackedLevelInstanceBlueprintWithDialog(AssetName, PackagePath, bCompile))
 		{
diff --git a/Engine/Source/Runtime/Engine/Private/LevelInstance/Packed/PackedLevelInstanceBuilder.cpp b/Engine/Source/Runtime/Engine/Private/LevelInstance/Packed/PackedLevelInstanceBuilder.cpp
--- a/Engine/Source/Runtime/Engine/Private/LevelInstance/Packed/PackedLevelInstanceBuilder.cpp
+++ b/Engine/Source/Runtime/Engine/Private/LevelInstance/Packed/PackedLevelInstanceBuilder.cpp
@@ -68,10 +68,10 @@ FPackedLevelInstanceBuilder::FPackedLevelInstanceBuilder()
 {
 }
 
-const FString& FPackedLevelInstanceBuilder::GetPackedBPSuffix()
+const FString& FPackedLevelInstanceBuilder::GetPackedBPPrefix()
 {
-	static FString BPSuffix = "_BP_Packed";
-	return BPSuffix;
+	static FString BPPrefix = "BPP_";
+	return BPPrefix;
 }
 
 UBlueprint* FPackedLevelInstanceBuilder::CreatePackedLevelInstanceBlueprintWithDialog(const FString& InAssetName, const FString& InPackagePath, bool bInCompile)
@@ -402,7 +402,7 @@ bool FPackedLevelInstanceBuilder::CreateOrUpdateBlueprintFromPacked(APackedLevel
 		LongPackageName.FindLastChar('/', LastSlashIndex);
 
 		FString PackagePath = LongPackageName.Mid(0, LastSlashIndex == INDEX_NONE ? MAX_int32 : LastSlashIndex);
-		FString AssetName = InActor->GetWorldAsset().GetAssetName() + GetPackedBPSuffix();
+		FString AssetName = GetPackedBPPrefix() + InActor->GetWorldAsset().GetAssetName();
 		const bool bCompile = false;
 
 		BP = CreatePackedLevelInstanceBlueprintWithDialog(AssetName, PackagePath, bCompile);
diff --git a/Engine/Source/Runtime/Engine/Public/LevelInstance/Packed/PackedLevelInstanceBuilder.h b/Engine/Source/Runtime/Engine/Public/LevelInstance/Packed/PackedLevelInstanceBuilder.h
--- a/Engine/Source/Runtime/Engine/Public/LevelInstance/Packed/PackedLevelInstanceBuilder.h
+++ b/Engine/Source/Runtime/Engine/Public/LevelInstance/Packed/PackedLevelInstanceBuilder.h
@@ -42,7 +42,7 @@ public:
 	/* Update existing Blueprint */
 	void UpdateBlueprint(UBlueprint* Blueprint, bool bSilentUpdate = false);
 
-	static const FString& GetPackedBPSuffix();
+	static const FString& GetPackedBPPrefix();
 	/* Creates a new APackedLevelInstance Blueprint using InPackagePath/InAssetName as hint for path. Prompts the user to input the final asset name. */
 	static UBlueprint* CreatePackedLevelInstanceBlueprintWithDialog(const FString& InAssetName, const FString& InPackagePath, bool bInCompile);
 	/* Creates a new APackedLevelInstance Blueprint using InPackagePath/InAssetName */
