World Partition HLOD - Disable collision on components created for HLODs

diff --git a/Engine/Source/Runtime/Engine/Private/WorldPartition/HLOD/HLODBuilder.cpp b/Engine/Source/Runtime/Engine/Private/WorldPartition/HLOD/HLODBuilder.cpp
--- a/Engine/Source/Runtime/Engine/Private/WorldPartition/HLOD/HLODBuilder.cpp
+++ b/Engine/Source/Runtime/Engine/Private/WorldPartition/HLOD/HLODBuilder.cpp
@@ -232,6 +232,16 @@ public:
 			
 		return Ar.GetCrc();
 	}
+
+	void DisableCollisions(UPrimitiveComponent* Component)
+	{
+		Component->SetCollisionProfileName(UCollisionProfile::NoCollision_ProfileName);
+		Component->SetGenerateOverlapEvents(false);
+		Component->SetCanEverAffectNavigation(false);
+		Component->CanCharacterStepUpOn = ECanBeCharacterBase::ECB_No;
+		Component->SetCanEverAffectNavigation(false);
+		Component->SetCollisionEnabled(ECollisionEnabled::NoCollision);
+	}
 };
 
 
@@ -319,6 +329,8 @@ class FHLODBuilder_Instancing : public FHLODBuilder
 			EntryInstancingKey.ApplyTo(Component);
 			Component->SetForcedLodModel(Component->GetStaticMesh()->GetNumLODs());
 
+			DisableCollisions(Component);
+
 			// Add all instances
 			for (UPrimitiveComponent* SMC : EntryComponents)
 			{
@@ -362,7 +374,7 @@ class FHLODBuilder_MeshMerge : public FHLODBuilder
 		MeshMergeUtilities.MergeComponentsToStaticMesh(InSubComponents, InHLODActor->GetWorld(), InHLODLayer->GetMeshMergeSettings(), InHLODLayer->GetHLODMaterial().LoadSynchronous(), InHLODActor->GetPackage(), InHLODActor->GetActorLabel(), Assets, MergedActorLocation, 0.25f, false);
 
 		UStaticMeshComponent* Component = nullptr;
-		Algo::ForEach(Assets, [InHLODActor, &Component, &MergedActorLocation](UObject* Asset)
+		Algo::ForEach(Assets, [this, InHLODActor, &Component, &MergedActorLocation](UObject* Asset)
 		{
 			Asset->ClearFlags(RF_Public | RF_Standalone);
 			Asset->Rename(nullptr, InHLODActor);
@@ -372,6 +384,7 @@ class FHLODBuilder_MeshMerge : public FHLODBuilder
 				Component = NewObject<UStaticMeshComponent>(InHLODActor);
 				Component->SetStaticMesh(static_cast<UStaticMesh*>(Asset));
 				Component->SetWorldLocation(MergedActorLocation);
+				DisableCollisions(Component);
 			}
 		});
 
@@ -399,7 +412,7 @@ class FHLODBuilder_MeshSimplify : public FHLODBuilder
 		MeshMergeUtilities.CreateProxyMesh(StaticMeshComponents, InHLODLayer->GetMeshSimplifySettings(), InHLODLayer->GetHLODMaterial().LoadSynchronous(), InHLODActor->GetPackage(), InHLODActor->GetActorLabel(), FGuid::NewGuid(), ProxyDelegate, true);
 
 		UStaticMeshComponent* Component = nullptr;
-		Algo::ForEach(Assets, [InHLODActor, &Component](UObject* Asset)
+		Algo::ForEach(Assets, [this, InHLODActor, &Component](UObject* Asset)
 		{
 			Asset->ClearFlags(RF_Public | RF_Standalone);
 			Asset->Rename(nullptr, InHLODActor, REN_ForceNoResetLoaders | REN_NonTransactional | REN_DoNotDirty);
@@ -408,6 +421,7 @@ class FHLODBuilder_MeshSimplify : public FHLODBuilder
 			{
 				Component = NewObject<UStaticMeshComponent>(InHLODActor);
 				Component->SetStaticMesh(static_cast<UStaticMesh*>(Asset));
+				DisableCollisions(Component);
 			}
 		});
 
