World Partition HLOD - Sort streaming cells with main grid first, then HLOD levels in order to minimize the impact of HLOD grids streaming on the main grid.
Note: will be effecive in the next HLODs generation.

diff --git a/Engine/Source/Runtime/Engine/Private/WorldPartition/RuntimeSpatialHash/RuntimeSpatialHashHLOD.cpp b/Engine/Source/Runtime/Engine/Private/WorldPartition/RuntimeSpatialHash/RuntimeSpatialHashHLOD.cpp
--- a/Engine/Source/Runtime/Engine/Private/WorldPartition/RuntimeSpatialHash/RuntimeSpatialHashHLOD.cpp
+++ b/Engine/Source/Runtime/Engine/Private/WorldPartition/RuntimeSpatialHash/RuntimeSpatialHashHLOD.cpp
@@ -400,7 +400,8 @@ static void UpdateHLODGridsActors(UWorld* World, const TMap<FName, FSpatialHashR
 		ASpatialHashRuntimeGridInfo* GridActor = *ItRuntimeGridActor;
 		if (GridActor->ActorHasTag(HLODGridTag))
 		{
-			if (HLODGrids.Contains(GridActor->GridSettings.GridName))
+			const FSpatialHashRuntimeGrid* HLODGrid = HLODGrids.Find(GridActor->GridSettings.GridName);
+			if (HLODGrid && HLODGrid->Priority)
 			{
 				ExistingGridActors.Emplace(GridActor->GridSettings.GridName, GridActor);
 			}
@@ -435,6 +436,9 @@ static void UpdateHLODGridsActors(UWorld* World, const TMap<FName, FSpatialHashR
 			{
 				HLODLevel = LastLODColorationColorIdx;
 			}
+
+			GridActor->GridSettings.Priority = 100 + HLODLevel;
+
 			HLODLevel = FMath::Clamp(HLODLevel + 2, 0U, LastLODColorationColorIdx);
 			GridActor->GridSettings.DebugColor = GEngine->HLODColorationColors[HLODLevel];
 
diff --git a/Engine/Source/Runtime/Engine/Private/WorldPartition/WorldPartitionRuntimeSpatialHash.cpp b/Engine/Source/Runtime/Engine/Private/WorldPartition/WorldPartitionRuntimeSpatialHash.cpp
--- a/Engine/Source/Runtime/Engine/Private/WorldPartition/WorldPartitionRuntimeSpatialHash.cpp
+++ b/Engine/Source/Runtime/Engine/Private/WorldPartition/WorldPartitionRuntimeSpatialHash.cpp
@@ -802,6 +802,7 @@ bool UWorldPartitionRuntimeSpatialHash::CreateStreamingGrid(const FSpatialHashRu
 				StreamingCell->SetIsAlwaysLoaded(bIsCellAlwaysLoaded);
 				StreamingCell->SetDataLayers(GridCellDataChunk.GetDataLayers());
 				StreamingCell->Level = Level;
+				StreamingCell->Priority = RuntimeGrid.Priority;
 				FBox2D Bounds;
 				verify(TempLevel.GetCellBounds(FIntVector2(CellCoordX, CellCoordY), Bounds));
 				StreamingCell->Position = FVector(Bounds.GetCenter(), 0.f);
@@ -1054,11 +1055,15 @@ void UWorldPartitionRuntimeSpatialHash::SortStreamingCellsByImportance(const TSe
 
 	Algo::Sort(SortedCells, [](const FCellShortestDist& A, const FCellShortestDist& B)
 	{
-		if (A.Cell->Level == B.Cell->Level)
+		if (A.Cell->Priority == B.Cell->Priority)
 		{
-			return A.SourceMinDistance < B.SourceMinDistance;
+			if (A.Cell->Level == B.Cell->Level)
+			{
+				return A.SourceMinDistance < B.SourceMinDistance;
+			}
+			return A.Cell->Level > B.Cell->Level;
 		}
-		return A.Cell->Level > B.Cell->Level;
+		return A.Cell->Priority < B.Cell->Priority;
 	});
 
 	OutSortedCells.Reserve(InCells.Num());
diff --git a/Engine/Source/Runtime/Engine/Public/WorldPartition/WorldPartitionRuntimeSpatialHash.h b/Engine/Source/Runtime/Engine/Public/WorldPartition/WorldPartitionRuntimeSpatialHash.h
--- a/Engine/Source/Runtime/Engine/Public/WorldPartition/WorldPartitionRuntimeSpatialHash.h
+++ b/Engine/Source/Runtime/Engine/Public/WorldPartition/WorldPartitionRuntimeSpatialHash.h
@@ -150,6 +150,7 @@ struct FSpatialHashRuntimeGrid
 #if WITH_EDITORONLY_DATA
 		: CellSize(0)
 		, LoadingRange(1000)
+		, Priority(0)
 		, DebugColor(FLinearColor::MakeRandomColor())
 #endif
 	{}
@@ -164,6 +165,9 @@ struct FSpatialHashRuntimeGrid
 	UPROPERTY(EditAnywhere, Category=Settings)
 	float LoadingRange;
 
+	UPROPERTY(EditAnywhere, Category=Settings)
+	int32 Priority;
+
 	UPROPERTY(EditAnywhere, Category=Settings, meta = (IgnoreForMemberInitializationTest))
 	FLinearColor DebugColor;
 #endif
diff --git a/Engine/Source/Runtime/Engine/Public/WorldPartition/WorldPartitionRuntimeSpatialHashCell.h b/Engine/Source/Runtime/Engine/Public/WorldPartition/WorldPartitionRuntimeSpatialHashCell.h
--- a/Engine/Source/Runtime/Engine/Public/WorldPartition/WorldPartitionRuntimeSpatialHashCell.h
+++ b/Engine/Source/Runtime/Engine/Public/WorldPartition/WorldPartitionRuntimeSpatialHashCell.h
@@ -20,6 +20,9 @@ class UWorldPartitionRuntimeSpatialHashCell : public UWorldPartitionRuntimeCell
 	UPROPERTY()
 	int32 Level;
 
+	UPROPERTY()
+	int32 Priority;
+
 #if WITH_EDITORONLY_DATA
 	UPROPERTY()
 	UActorContainer* ActorContainer;
