Fix Non-Unity Build

diff --git a/Engine/Source/Runtime/Engine/Private/LevelInstance/Packed/LevelInstanceRecursivePacker.cpp b/Engine/Source/Runtime/Engine/Private/LevelInstance/Packed/LevelInstanceRecursivePacker.cpp
--- a/Engine/Source/Runtime/Engine/Private/LevelInstance/Packed/LevelInstanceRecursivePacker.cpp
+++ b/Engine/Source/Runtime/Engine/Private/LevelInstance/Packed/LevelInstanceRecursivePacker.cpp
@@ -10,6 +10,8 @@
 #include "LevelInstance/LevelInstanceActor.h"
 #include "LevelInstance/LevelInstanceSubsystem.h"
 
+#include "Engine/Blueprint.h"
+
 #include "Misc/Crc.h"
 
 FLevelInstancePackerID FLevelInstanceRecursivePacker::PackerID = 'RECP';
