LevelInstance: Fix Details Panel Transform not updating Actor

diff --git a/Engine/Source/Runtime/Engine/Private/LevelInstance/LevelInstanceActor.cpp b/Engine/Source/Runtime/Engine/Private/LevelInstance/LevelInstanceActor.cpp
--- a/Engine/Source/Runtime/Engine/Private/LevelInstance/LevelInstanceActor.cpp
+++ b/Engine/Source/Runtime/Engine/Private/LevelInstance/LevelInstanceActor.cpp
@@ -86,6 +86,12 @@ void ALevelInstance::PostRegisterAllComponents()
 		{
 			LoadLevelInstance();
 		}
+
+#if WITH_EDITOR
+		// Make sure transformation is up to date after registration as its possible LevelInstance actor can get unregistered when editing properties
+		// through Details panel. In this case the ULevelInstanceComponent might not be able to update the ALevelInstanceEditorInstanceActor transform.
+		Cast<ULevelInstanceComponent>(GetRootComponent())->UpdateEditorInstanceActor();
+#endif
 	}
 }
 
