LevelInstance:
- FNewLevelInstanceParams now allows to specify an AssetPath for the level to avoid dialog
- FNewLevelInstanceParams also defaults to bPromptForSave = false
- Add Delay to UpdateStreamingState dialog
- Fix log error because WorldAsset wasn't set prior to first BP Compile. (failed MapCheck)

diff --git a/Engine/Source/Runtime/Engine/Private/LevelInstance/LevelInstanceSubsystem.cpp b/Engine/Source/Runtime/Engine/Private/LevelInstance/LevelInstanceSubsystem.cpp
--- a/Engine/Source/Runtime/Engine/Private/LevelInstance/LevelInstanceSubsystem.cpp
+++ b/Engine/Source/Runtime/Engine/Private/LevelInstance/LevelInstanceSubsystem.cpp
@@ -18,6 +18,7 @@
 #include "LevelInstance/LevelInstanceEditorLevelStreaming.h"
 #include "Misc/ScopedSlowTask.h"
 #include "Misc/ITransaction.h"
+#include "Misc/Paths.h"
 #include "AssetRegistryModule.h"
 #include "AssetData.h"
 #include "FileHelpers.h"
@@ -178,7 +179,7 @@ void ULevelInstanceSubsystem::UpdateStreamingState()
 	}
 
 	FScopedSlowTask SlowTask(LevelInstancesToUnload.Num() + LevelInstancesToLoadOrUpdate.Num() * 2, LOCTEXT("UpdatingLevelInstances", "Updating Level Instances..."), !GetWorld()->IsGameWorld());
-	SlowTask.MakeDialog();
+	SlowTask.MakeDialogDelayed(1.0f);
 
 	check(!LevelsToRemoveScope);
 	LevelsToRemoveScope.Reset(new FLevelsToRemoveScope());
@@ -454,10 +455,10 @@ void ULevelInstanceSubsystem::PackLevelInstances()
 	};
 
 	TSharedPtr<FPackedLevelInstanceBuilder> Builder = FPackedLevelInstanceBuilder::CreateDefaultBuilder();
-	const bool bSilentUpdate = true;
+	const bool bCheckoutAndSave = false;
 	for (UBlueprint* Blueprint : BlueprintsToUpdate)
 	{
-		Builder->UpdateBlueprint(Blueprint, bSilentUpdate);
+		Builder->UpdateBlueprint(Blueprint, bCheckoutAndSave);
 		UpdateProgress();
 	}
 
@@ -728,7 +729,7 @@ ALevelInstance* ULevelInstanceSubsystem::CreateLevelInstanceFrom(const TArray<AA
 		}
 	}
 		
-	ULevelStreamingLevelInstanceEditor* LevelStreaming = StaticCast<ULevelStreamingLevelInstanceEditor*>(EditorLevelUtils::CreateNewStreamingLevelForWorld(*GetWorld(), ULevelStreamingLevelInstanceEditor::StaticClass(), TEXT(""), false, CreationParams.TemplateWorld));
+	ULevelStreamingLevelInstanceEditor* LevelStreaming = StaticCast<ULevelStreamingLevelInstanceEditor*>(EditorLevelUtils::CreateNewStreamingLevelForWorld(*GetWorld(), ULevelStreamingLevelInstanceEditor::StaticClass(), CreationParams.AssetPath, false, CreationParams.TemplateWorld));
 	if (!LevelStreaming)
 	{
 		UE_LOG(LogLevelInstance, Warning, TEXT("Failed to create new Level Instance level"));
@@ -780,18 +781,24 @@ ALevelInstance* ULevelInstanceSubsystem::CreateLevelInstanceFrom(const TArray<AA
 	}
 	else if (CreationParams.Type == ELevelInstanceCreationType::PackedLevelInstanceBlueprint)
 	{
-		int32 LastSlashIndex = 0;
-		FString LongPackageName = WorldPtr.GetLongPackageName();
-		LongPackageName.FindLastChar('/', LastSlashIndex);
-
-		FString PackagePath = LongPackageName.Mid(0, LastSlashIndex == INDEX_NONE ? MAX_int32 : LastSlashIndex);
+		FString PackageDir = FPaths::GetPath(WorldPtr.GetLongPackageName());
 		FString AssetName = FPackedLevelInstanceBuilder::GetPackedBPPrefix() + WorldPtr.GetAssetName();
+		FString BPAssetPath = FString::Format(TEXT("{0}/{1}.{1}"), { PackageDir , AssetName });
 		const bool bCompile = true;
-		if (UBlueprint* NewBP = FPackedLevelInstanceBuilder::CreatePackedLevelInstanceBlueprintWithDialog(AssetName, PackagePath, bCompile))
+
+		UBlueprint* NewBP = nullptr;
+		if (CreationParams.AssetPath.IsEmpty())
+		{
+			NewBP = FPackedLevelInstanceBuilder::CreatePackedLevelInstanceBlueprintWithDialog(TSoftObjectPtr<UBlueprint>(BPAssetPath), WorldPtr, bCompile);
+		}
+		else
+		{
+			NewBP = FPackedLevelInstanceBuilder::CreatePackedLevelInstanceBlueprint(TSoftObjectPtr<UBlueprint>(BPAssetPath), WorldPtr, bCompile);
+		}
+				
+		if (NewBP)
 		{
 			NewLevelInstanceActor = GetWorld()->SpawnActor<APackedLevelInstance>(NewBP->GeneratedClass, SpawnParams);
-			APackedLevelInstance* CDO = Cast<APackedLevelInstance>(NewBP->GeneratedClass->GetDefaultObject());
-			CDO->SetWorldAsset(WorldPtr);
 		}
 
 		if (!NewLevelInstanceActor)
@@ -831,7 +838,7 @@ ALevelInstance* ULevelInstanceSubsystem::CreateLevelInstanceFrom(const TArray<AA
 	}
 				
 	// Use LevelStreaming->GetLevelInstanceActor() because OnWorldAssetSaved could've reinstanced the LevelInstanceActor
-	return CommitLevelInstance(LevelStreaming->GetLevelInstanceActor(), /*bDiscardEdits*/false, &DirtyPackages);
+	return CommitLevelInstance(LevelStreaming->GetLevelInstanceActor(), /*bDiscardEdits*/false, CreationParams.bPromptForSave, &DirtyPackages);
 }
 
 bool ULevelInstanceSubsystem::BreakLevelInstance(ALevelInstance* LevelInstanceActor, uint32 Levels /* = 1 */, TArray<AActor*>* OutMovedActors /* = nullptr */)
@@ -1503,7 +1510,7 @@ void ULevelInstanceSubsystem::CommitChildrenLevelInstances(ALevelInstance* Level
 	});
 }
 
-ALevelInstance* ULevelInstanceSubsystem::CommitLevelInstance(ALevelInstance* LevelInstanceActor, bool bDiscardEdits, TSet<UPackage*>* DirtyPackages)
+ALevelInstance* ULevelInstanceSubsystem::CommitLevelInstance(ALevelInstance* LevelInstanceActor, bool bDiscardEdits, bool bPromptForSave, TSet<UPackage*>* DirtyPackages)
 {
 	check(CanCommitLevelInstance(LevelInstanceActor));
 
@@ -1515,7 +1522,7 @@ ALevelInstance* ULevelInstanceSubsystem::CommitLevelInstance(ALevelInstance* Lev
 	bool bChangesCommitted = false;
 	if (IsLevelInstanceEditDirty(LevelInstanceEdit) && !bDiscardEdits)
 	{
-		const bool bPromptUserToSave = true;
+		const bool bPromptUserToSave = bPromptForSave;
 		const bool bSaveMapPackages = true;
 		const bool bSaveContentPackages = true;
 		const bool bFastSave = false;
@@ -1544,7 +1551,7 @@ ALevelInstance* ULevelInstanceSubsystem::CommitLevelInstance(ALevelInstance* Lev
 			IAssetRegistry& AssetRegistry = FModuleManager::LoadModuleChecked<FAssetRegistryModule>(TEXT("AssetRegistry")).Get();
 			AssetRegistry.ScanPathsSynchronous({ LevelInstanceActor->GetWorldAssetPackage() }, true);
 			// Notify 
-			LevelInstanceActor->OnWorldAssetSaved();
+			LevelInstanceActor->OnWorldAssetSaved(bPromptUserToSave);
 
 			// Update pointer since BP Compilation might have invalidated LevelInstanceActor
 			LevelInstanceActor = LevelInstanceEdit->LevelStreaming->GetLevelInstanceActor();
diff --git a/Engine/Source/Runtime/Engine/Private/LevelInstance/Packed/PackedLevelInstanceActor.cpp b/Engine/Source/Runtime/Engine/Private/LevelInstance/Packed/PackedLevelInstanceActor.cpp
--- a/Engine/Source/Runtime/Engine/Private/LevelInstance/Packed/PackedLevelInstanceActor.cpp
+++ b/Engine/Source/Runtime/Engine/Private/LevelInstance/Packed/PackedLevelInstanceActor.cpp
@@ -50,14 +50,15 @@ void APackedLevelInstance::OnWorldAssetChanged()
 	}
 }
 
-void APackedLevelInstance::OnWorldAssetSaved()
+void APackedLevelInstance::OnWorldAssetSaved(bool bPromptForSave)
 {
 	TSharedPtr<FPackedLevelInstanceBuilder> Builder = FPackedLevelInstanceBuilder::CreateDefaultBuilder();
 	
 	if (UBlueprint* GeneratedBy = Cast<UBlueprint>(GetClass()->ClassGeneratedBy))
 	{
 		check(GeneratedBy == BlueprintAsset.Get());
-		Builder->UpdateBlueprint(GeneratedBy);
+		const bool bCheckoutAndSave = true;
+		Builder->UpdateBlueprint(GeneratedBy, bCheckoutAndSave, bPromptForSave);
 	}
 	else
 	{
diff --git a/Engine/Source/Runtime/Engine/Private/LevelInstance/Packed/PackedLevelInstanceBuilder.cpp b/Engine/Source/Runtime/Engine/Private/LevelInstance/Packed/PackedLevelInstanceBuilder.cpp
--- a/Engine/Source/Runtime/Engine/Private/LevelInstance/Packed/PackedLevelInstanceBuilder.cpp
+++ b/Engine/Source/Runtime/Engine/Private/LevelInstance/Packed/PackedLevelInstanceBuilder.cpp
@@ -17,6 +17,7 @@
 #include "Engine/SCS_Node.h"
 
 #include "Misc/ScopeExit.h"
+#include "Misc/Paths.h"
 #include "Logging/LogMacros.h"
 #include "Logging/MessageLog.h"
 
@@ -74,12 +75,12 @@ const FString& FPackedLevelInstanceBuilder::GetPackedBPPrefix()
 	return BPPrefix;
 }
 
-UBlueprint* FPackedLevelInstanceBuilder::CreatePackedLevelInstanceBlueprintWithDialog(const FString& InAssetName, const FString& InPackagePath, bool bInCompile)
+UBlueprint* FPackedLevelInstanceBuilder::CreatePackedLevelInstanceBlueprintWithDialog(TSoftObjectPtr<UBlueprint> InBlueprintAsset, TSoftObjectPtr<UWorld> InWorldAsset, bool bInCompile)
 {
 	FSaveAssetDialogConfig SaveAssetDialogConfig;
 	SaveAssetDialogConfig.DialogTitleOverride = LOCTEXT("SaveAssetDialogTitle", "Save Asset As");
-	SaveAssetDialogConfig.DefaultPath = InPackagePath;
-	SaveAssetDialogConfig.DefaultAssetName = InAssetName;
+	SaveAssetDialogConfig.DefaultPath = FPaths::GetPath(InBlueprintAsset.GetLongPackageName());
+	SaveAssetDialogConfig.DefaultAssetName = InBlueprintAsset.GetAssetName();
 	SaveAssetDialogConfig.ExistingAssetPolicy = ESaveAssetDialogExistingAssetPolicy::AllowButWarn;
 
 	FContentBrowserModule& ContentBrowserModule = FModuleManager::LoadModuleChecked<FContentBrowserModule>("ContentBrowser");
@@ -93,18 +94,14 @@ UBlueprint* FPackedLevelInstanceBuilder::CreatePackedLevelInstanceBlueprintWithD
 		{
 			return BP;
 		}
-		
-		const FString SavePackageName = FPackageName::ObjectPathToPackageName(SaveObjectPath);
-		const FString SavePackagePath = FPaths::GetPath(SavePackageName);
-		const FString SaveAssetName = FPaths::GetBaseFilename(SavePackageName);
-
-		return CreatePackedLevelInstanceBlueprint(SaveAssetName, SavePackagePath, bInCompile);
+				
+		return CreatePackedLevelInstanceBlueprint(ExistingBPAsset, InWorldAsset, bInCompile);
 	}
 
 	return nullptr;
 }
 
-UBlueprint* FPackedLevelInstanceBuilder::CreatePackedLevelInstanceBlueprint(const FString& InSaveAssetName, const FString& InPackagePath, bool bInCompile)
+UBlueprint* FPackedLevelInstanceBuilder::CreatePackedLevelInstanceBlueprint(TSoftObjectPtr<UBlueprint> InBlueprintAsset, TSoftObjectPtr<UWorld> InWorldAsset, bool bInCompile)
 {
 	IAssetTools& AssetTools = FAssetToolsModule::GetModule().Get();
 
@@ -115,12 +112,14 @@ UBlueprint* FPackedLevelInstanceBuilder::CreatePackedLevelInstanceBlueprint(cons
 	FEditorDelegates::OnConfigureNewAssetProperties.Broadcast(BlueprintFactory);
 	if (BlueprintFactory->ConfigureProperties())
 	{
-		FEditorDirectories::Get().SetLastDirectory(ELastDirectory::NEW_ASSET, InPackagePath);
+		FString PackageDir = FPaths::GetPath(InBlueprintAsset.GetLongPackageName());
+		FEditorDirectories::Get().SetLastDirectory(ELastDirectory::NEW_ASSET, PackageDir);
 
-		if (UBlueprint* NewBP = Cast<UBlueprint>(AssetTools.CreateAsset(InSaveAssetName, InPackagePath, UBlueprint::StaticClass(), BlueprintFactory, FName("Create LevelInstance Blueprint"))))
+		if (UBlueprint* NewBP = Cast<UBlueprint>(AssetTools.CreateAsset(InBlueprintAsset.GetAssetName(), PackageDir, UBlueprint::StaticClass(), BlueprintFactory, FName("Create LevelInstance Blueprint"))))
 		{
 			APackedLevelInstance* CDO = CastChecked<APackedLevelInstance>(NewBP->GeneratedClass->GetDefaultObject());
 			CDO->BlueprintAsset = NewBP;
+			CDO->SetWorldAsset(InWorldAsset);
 
 			if (bInCompile)
 			{
@@ -324,38 +323,38 @@ bool FPackedLevelInstanceBuilder::PackActor(APackedLevelInstance* InActor, TSoft
 	return PackActor(InActor, TransientLevelInstance);
 }
 
-void FPackedLevelInstanceBuilder::UpdateBlueprint(UBlueprint* Blueprint, bool bInSilentUpdate)
+void FPackedLevelInstanceBuilder::UpdateBlueprint(UBlueprint* Blueprint, bool bCheckoutAndSave, bool bPromptForSave)
 {
 	APackedLevelInstance* CDO = CastChecked<APackedLevelInstance>(Blueprint->GeneratedClass->GetDefaultObject());
 	check(CDO);
 
-	CreateOrUpdateBlueprint(CDO->GetWorldAsset(), Blueprint, bInSilentUpdate);
+	CreateOrUpdateBlueprint(CDO->GetWorldAsset(), Blueprint, bCheckoutAndSave, bPromptForSave);
 }
 
-bool FPackedLevelInstanceBuilder::CreateOrUpdateBlueprint(TSoftObjectPtr<UWorld> InWorldAsset, TSoftObjectPtr<UBlueprint> InBlueprintAsset, bool bInSilentUpdate)
+bool FPackedLevelInstanceBuilder::CreateOrUpdateBlueprint(TSoftObjectPtr<UWorld> InWorldAsset, TSoftObjectPtr<UBlueprint> InBlueprintAsset, bool bCheckoutAndSave, bool bPromptForSave)
 {
 	bool bResult = true;
 	
 	ALevelInstance* TransientLevelInstance = CreateTransientLevelInstanceForPacking(InWorldAsset, FVector::ZeroVector, FRotator::ZeroRotator);
 	
-	bResult = CreateOrUpdateBlueprintFromUnpacked(TransientLevelInstance, InBlueprintAsset, bInSilentUpdate);
+	bResult = CreateOrUpdateBlueprintFromUnpacked(TransientLevelInstance, InBlueprintAsset, bCheckoutAndSave, bPromptForSave);
 
 	TransientLevelInstance->GetWorld()->DestroyActor(TransientLevelInstance);
 
 	return bResult;
 }
 
-bool FPackedLevelInstanceBuilder::CreateOrUpdateBlueprint(ALevelInstance* InLevelInstance, TSoftObjectPtr<UBlueprint> InBlueprintAsset, bool bInSilentUpdate)
+bool FPackedLevelInstanceBuilder::CreateOrUpdateBlueprint(ALevelInstance* InLevelInstance, TSoftObjectPtr<UBlueprint> InBlueprintAsset, bool bCheckoutAndSave, bool bPromptForSave)
 {
 	if (APackedLevelInstance* PackedLevelInstance = Cast<APackedLevelInstance>(InLevelInstance))
 	{
-		return CreateOrUpdateBlueprintFromPacked(PackedLevelInstance, InBlueprintAsset, bInSilentUpdate);
+		return CreateOrUpdateBlueprintFromPacked(PackedLevelInstance, InBlueprintAsset, bCheckoutAndSave, bPromptForSave);
 	}
 	
-	return CreateOrUpdateBlueprintFromUnpacked(InLevelInstance, InBlueprintAsset, bInSilentUpdate);
+	return CreateOrUpdateBlueprintFromUnpacked(InLevelInstance, InBlueprintAsset, bCheckoutAndSave, bPromptForSave);
 }
 
-bool FPackedLevelInstanceBuilder::CreateOrUpdateBlueprintFromUnpacked(ALevelInstance* InActor, TSoftObjectPtr<UBlueprint> InBlueprintAsset, bool bInSilentUpdate)
+bool FPackedLevelInstanceBuilder::CreateOrUpdateBlueprintFromUnpacked(ALevelInstance* InActor, TSoftObjectPtr<UBlueprint> InBlueprintAsset, bool bCheckoutAndSave, bool bPromptForSave)
 {
 	bool bResult = true;
 	
@@ -382,12 +381,12 @@ bool FPackedLevelInstanceBuilder::CreateOrUpdateBlueprintFromUnpacked(ALevelInst
 	}
 
 	PackedLevelInstance->BlueprintAsset = InBlueprintAsset;
-	bResult &= CreateOrUpdateBlueprintFromPacked(PackedLevelInstance, PackedLevelInstance->BlueprintAsset, bInSilentUpdate);
+	bResult &= CreateOrUpdateBlueprintFromPacked(PackedLevelInstance, PackedLevelInstance->BlueprintAsset, bCheckoutAndSave, bPromptForSave);
 
 	return bResult;
 }
 
-bool FPackedLevelInstanceBuilder::CreateOrUpdateBlueprintFromPacked(APackedLevelInstance* InActor, TSoftObjectPtr<UBlueprint> InBlueprintAsset, bool bInSilentUpdate)
+bool FPackedLevelInstanceBuilder::CreateOrUpdateBlueprintFromPacked(APackedLevelInstance* InActor, TSoftObjectPtr<UBlueprint> InBlueprintAsset, bool bCheckoutAndSave, bool bPromptToSave)
 {
 	UBlueprint* BP = nullptr;
 	if (!InBlueprintAsset.IsNull())
@@ -405,7 +404,8 @@ bool FPackedLevelInstanceBuilder::CreateOrUpdateBlueprintFromPacked(APackedLevel
 		FString AssetName = GetPackedBPPrefix() + InActor->GetWorldAsset().GetAssetName();
 		const bool bCompile = false;
 
-		BP = CreatePackedLevelInstanceBlueprintWithDialog(AssetName, PackagePath, bCompile);
+		FString AssetPath = PackagePath + AssetName + "." + AssetName;
+		BP = CreatePackedLevelInstanceBlueprintWithDialog(TSoftObjectPtr<UBlueprint>(AssetPath), InActor->GetWorldAsset(), bCompile);
 	}
 
 	if (BP && BP->SimpleConstructionScript)
@@ -474,10 +474,9 @@ bool FPackedLevelInstanceBuilder::CreateOrUpdateBlueprintFromPacked(APackedLevel
 	// Synchronous compile
 	FKismetEditorUtilities::CompileBlueprint(BP, EBlueprintCompileOptions::SkipGarbageCollection);
 
-	if (!bInSilentUpdate)
+	if (bCheckoutAndSave)
 	{
 		const bool bCheckDirty = false;
-		const bool bPromptToSave = true;
 		TArray<UPackage*> OutFailedPackages;
 		FEditorFileUtils::PromptForCheckoutAndSave({ BP->GetPackage() }, bCheckDirty, bPromptToSave, &OutFailedPackages);
 
diff --git a/Engine/Source/Runtime/Engine/Public/LevelInstance/LevelInstanceActor.h b/Engine/Source/Runtime/Engine/Public/LevelInstance/LevelInstanceActor.h
--- a/Engine/Source/Runtime/Engine/Public/LevelInstance/LevelInstanceActor.h
+++ b/Engine/Source/Runtime/Engine/Public/LevelInstance/LevelInstanceActor.h
@@ -98,7 +98,7 @@ public:
 	AActor* FindEditorInstanceActor() const;
 
 	virtual void OnWorldAssetChanged() { UpdateLevelInstance(); }
-	virtual void OnWorldAssetSaved() {}
+	virtual void OnWorldAssetSaved(bool bPromptForSave) {}
 	virtual void OnEdit();
 	virtual void OnEditChild() {}
 	virtual void OnCommit();
diff --git a/Engine/Source/Runtime/Engine/Public/LevelInstance/LevelInstanceSubsystem.h b/Engine/Source/Runtime/Engine/Public/LevelInstance/LevelInstanceSubsystem.h
--- a/Engine/Source/Runtime/Engine/Public/LevelInstance/LevelInstanceSubsystem.h
+++ b/Engine/Source/Runtime/Engine/Public/LevelInstance/LevelInstanceSubsystem.h
@@ -54,7 +54,7 @@ public:
 	bool CanEditLevelInstance(const ALevelInstance* LevelInstanceActor, FText* OutReason = nullptr) const;
 	bool CanCommitLevelInstance(const ALevelInstance* LevelInstanceActor, FText* OutReason = nullptr) const;
 	void EditLevelInstance(ALevelInstance* LevelInstanceActor, TWeakObjectPtr<AActor> ContextActorPtr = nullptr);
-	ALevelInstance* CommitLevelInstance(ALevelInstance* LevelInstanceActor, bool bDiscardEdits = false, TSet<UPackage*>* DirtyPackages = nullptr);
+	ALevelInstance* CommitLevelInstance(ALevelInstance* LevelInstanceActor, bool bDiscardEdits = false, bool bPromptForSave = true, TSet<UPackage*>* DirtyPackages = nullptr);
 	void SaveLevelInstanceAs(ALevelInstance* LevelInstanceActor);
 	bool IsEditingLevelInstanceDirty(const ALevelInstance* LevelInstanceActor) const;
 	bool IsEditingLevelInstance(const ALevelInstance* LevelInstanceActor) const { return GetLevelInstanceEdit(LevelInstanceActor) != nullptr; }
diff --git a/Engine/Source/Runtime/Engine/Public/LevelInstance/LevelInstanceTypes.h b/Engine/Source/Runtime/Engine/Public/LevelInstance/LevelInstanceTypes.h
--- a/Engine/Source/Runtime/Engine/Public/LevelInstance/LevelInstanceTypes.h
+++ b/Engine/Source/Runtime/Engine/Public/LevelInstance/LevelInstanceTypes.h
@@ -46,4 +46,10 @@ struct FNewLevelInstanceParams
 
 	UPROPERTY()
 	TObjectPtr<UWorld> TemplateWorld = nullptr;
+		
+	UPROPERTY()
+	FString AssetPath = TEXT("");
+
+	UPROPERTY()
+	bool bPromptForSave = false;
 };
\ No newline at end of file
diff --git a/Engine/Source/Runtime/Engine/Public/LevelInstance/Packed/PackedLevelInstanceActor.h b/Engine/Source/Runtime/Engine/Public/LevelInstance/Packed/PackedLevelInstanceActor.h
--- a/Engine/Source/Runtime/Engine/Public/LevelInstance/Packed/PackedLevelInstanceActor.h
+++ b/Engine/Source/Runtime/Engine/Public/LevelInstance/Packed/PackedLevelInstanceActor.h
@@ -36,7 +36,7 @@ public:
 	static FName GetPackedComponentTag();
 
 	virtual void OnWorldAssetChanged() override;
-	virtual void OnWorldAssetSaved() override;
+	virtual void OnWorldAssetSaved(bool bPromptForSave) override;
 	virtual void OnCommit() override;
 	virtual void OnCommitChild(bool bChanged) override;
 	virtual void OnEdit() override;
diff --git a/Engine/Source/Runtime/Engine/Public/LevelInstance/Packed/PackedLevelInstanceBuilder.h b/Engine/Source/Runtime/Engine/Public/LevelInstance/Packed/PackedLevelInstanceBuilder.h
--- a/Engine/Source/Runtime/Engine/Public/LevelInstance/Packed/PackedLevelInstanceBuilder.h
+++ b/Engine/Source/Runtime/Engine/Public/LevelInstance/Packed/PackedLevelInstanceBuilder.h
@@ -36,23 +36,23 @@ public:
 	bool PackActor(APackedLevelInstance* InPackedLevelInstance, TSoftObjectPtr<UWorld> InWorldAsset);
 
 	/* Creates/Updates a APackedLevelInstance Blueprint from InLevelInstance (will overwrite existing asset or show a dialog if InBlueprintAsset is null) */
-	bool CreateOrUpdateBlueprint(ALevelInstance* InLevelInstance, TSoftObjectPtr<UBlueprint> InBlueprintAsset, bool bSilentUpdate = false);
+	bool CreateOrUpdateBlueprint(ALevelInstance* InLevelInstance, TSoftObjectPtr<UBlueprint> InBlueprintAsset, bool bCheckoutAndSave = true, bool bPromptForSave = true);
 	/* Creates/Updates a APackedLeveInstance Blueprint from InWorldAsset (will overwrite existing asset or show a dialog if InBlueprintAsset is null) */
-	bool CreateOrUpdateBlueprint(TSoftObjectPtr<UWorld> InWorldAsset, TSoftObjectPtr<UBlueprint> InBlueprintAsset, bool bSilentUpdate = false);
+	bool CreateOrUpdateBlueprint(TSoftObjectPtr<UWorld> InWorldAsset, TSoftObjectPtr<UBlueprint> InBlueprintAsset, bool bCheckoutAndSave = true, bool bPromptForSave = true);
 	/* Update existing Blueprint */
-	void UpdateBlueprint(UBlueprint* Blueprint, bool bSilentUpdate = false);
+	void UpdateBlueprint(UBlueprint* Blueprint, bool bCheckoutAndSave = true, bool bPromptForSave = true);
 
 	static const FString& GetPackedBPPrefix();
 	/* Creates a new APackedLevelInstance Blueprint using InPackagePath/InAssetName as hint for path. Prompts the user to input the final asset name. */
-	static UBlueprint* CreatePackedLevelInstanceBlueprintWithDialog(const FString& InAssetName, const FString& InPackagePath, bool bInCompile);
+	static UBlueprint* CreatePackedLevelInstanceBlueprintWithDialog(TSoftObjectPtr<UBlueprint> InBlueprintAsset, TSoftObjectPtr<UWorld> InWorldAsset, bool bInCompile);
 	/* Creates a new APackedLevelInstance Blueprint using InPackagePath/InAssetName */
-	static UBlueprint* CreatePackedLevelInstanceBlueprint(const FString& InSaveAssetName, const FString& InPackagePath, bool bInCompile);
+	static UBlueprint* CreatePackedLevelInstanceBlueprint(TSoftObjectPtr<UBlueprint> InBlueprintAsset, TSoftObjectPtr<UWorld> InWorldAsset, bool bInCompile);
 	
 private:
 	/* Create/Updates a APackedLevelInstance Blueprint from InPackedActor (will overwrite existing asset or show a dialog if InBlueprintAsset is null) */
-	bool CreateOrUpdateBlueprintFromPacked(APackedLevelInstance* InPackedActor, TSoftObjectPtr<UBlueprint> InBlueprintAsset, bool bInSilentUpdate);
+	bool CreateOrUpdateBlueprintFromPacked(APackedLevelInstance* InPackedActor, TSoftObjectPtr<UBlueprint> InBlueprintAsset, bool bCheckoutAndSave, bool bPromptForSave);
 	/* Creates/Updates a APackedLevelInstance Blueprint from InLevelInstance (will overwrite existing asset or show a dialog if InBlueprintAsset is null) */
-	bool CreateOrUpdateBlueprintFromUnpacked(ALevelInstance* InPackedActor, TSoftObjectPtr<UBlueprint> InBlueprintAsset, bool bInSilentUpdate);
+	bool CreateOrUpdateBlueprintFromUnpacked(ALevelInstance* InPackedActor, TSoftObjectPtr<UBlueprint> InBlueprintAsset, bool bCheckoutAndSave, bool bPromptForSave);
 	/* Creates and loads a ALevelInstance so it can be used for packing */
 	ALevelInstance* CreateTransientLevelInstanceForPacking(TSoftObjectPtr<UWorld> InWorldAsset, const FVector& InLocation, const FRotator& InRotator);
 
