Revert parts of the recent changes to XGE shader compilation.

- Disallowing XGE compiler to work on Normal jobs excluded global shaders from going XGE-wide. This change again makes XGE take Normal jobs (together with local SCWs).
- Batch size change, which was intended to compensate for the higher cost of starting remote-only XGE SCWs (now that we avoid running them on the local machine), ended up bringing a net slowdown and is also being reverted.

diff --git a/Engine/Source/Runtime/Engine/Private/ShaderCompiler/ShaderCompiler.cpp b/Engine/Source/Runtime/Engine/Private/ShaderCompiler/ShaderCompiler.cpp
--- a/Engine/Source/Runtime/Engine/Private/ShaderCompiler/ShaderCompiler.cpp
+++ b/Engine/Source/Runtime/Engine/Private/ShaderCompiler/ShaderCompiler.cpp
@@ -2941,7 +2941,7 @@ FShaderCompilingManager::FShaderCompilingManager() :
 	if (RemoteCompileThread)
 	{
 		// Keep high priority jobs on the local machine, to avoid XGE latency
-		RemoteCompileThread->SetPriorityRange(EShaderCompileJobPriority::Low, EShaderCompileJobPriority::Low);
+		RemoteCompileThread->SetPriorityRange(EShaderCompileJobPriority::Low, EShaderCompileJobPriority::Normal);
 		LocalThread->SetPriorityRange(EShaderCompileJobPriority::Normal, EShaderCompileJobPriority::ForceLocal);
 		Threads.Add(MoveTemp(RemoteCompileThread));
 	}
diff --git a/Engine/Source/Runtime/Engine/Private/ShaderCompiler/ShaderCompilerDistributed.cpp b/Engine/Source/Runtime/Engine/Private/ShaderCompiler/ShaderCompilerDistributed.cpp
--- a/Engine/Source/Runtime/Engine/Private/ShaderCompiler/ShaderCompilerDistributed.cpp
+++ b/Engine/Source/Runtime/Engine/Private/ShaderCompiler/ShaderCompilerDistributed.cpp
@@ -198,12 +198,9 @@ int32 FShaderCompileDistributedThreadRunnable_Interface::CompilingLoop()
 	{
 		// Increase the batch size when more jobs are queued/in flight.
 
-		// Just to provide typical numbers: the number of total jobs is usually in tens of thousands at most, oftentimes in low thousands. Thus JobsPerBatch when calculated as a log2 rarely reaches the value of 16.
-		// Given that 
-		//   a) starting each XGE shader compile worker has an associated cost (it scans directories, loads multiple libraries and reads configs)
-		//   b) XGE is currently used for low priority, non time-critical jobs
-		// JobsPerBatch is now increased by a factor of 3 compared to original formula, to produce numbers in 3-48 range most of the time
-		const uint32 JobsPerBatch = FMath::Max(1, FMath::FloorToInt(3 * FMath::LogX(2.f, PendingJobs.Num() + NumDispatchedJobs)));
+		// Just to provide typical numbers: the number of total jobs is usually in tens of thousands at most, oftentimes in low thousands. Thus JobsPerBatch when calculated as a log2 rarely reaches the value of 16,
+		// and that seems to be a sweet spot: lowering it does not result in faster completion, while increasing the number of jobs per batch slows it down .
+		const uint32 JobsPerBatch = FMath::Max(1, FMath::FloorToInt(FMath::LogX(2.f, PendingJobs.Num() + NumDispatchedJobs)));
 		UE_LOG(LogShaderCompilers, Verbose, TEXT("Current jobs: %d, Batch size: %d, Num Already Dispatched: %d"), PendingJobs.Num(), JobsPerBatch, NumDispatchedJobs);
 
 
