Actor Descriptors: added acceleration structures to query actor descs by guid and by name.

diff --git a/Engine/Source/Runtime/Engine/Private/WorldPartition/ActorDescContainer.cpp b/Engine/Source/Runtime/Engine/Private/WorldPartition/ActorDescContainer.cpp
--- a/Engine/Source/Runtime/Engine/Private/WorldPartition/ActorDescContainer.cpp
+++ b/Engine/Source/Runtime/Engine/Private/WorldPartition/ActorDescContainer.cpp
@@ -43,12 +43,56 @@ void UActorDescContainer::Initialize(UWorld* InWorld, FName InPackageName, bool
 		AssetRegistry.GetAssets(Filter, Assets);
 	}
 
+	auto GetActorDescriptor = [this](const FAssetData& InAssetData) -> TUniquePtr<FWorldPartitionActorDesc>
+	{
+		FString ActorClassName;
+		static FName NAME_ActorMetaDataClass(TEXT("ActorMetaDataClass"));
+		if (InAssetData.GetTagValue(NAME_ActorMetaDataClass, ActorClassName))
+		{
+			FString ActorMetaDataStr;
+			static FName NAME_ActorMetaData(TEXT("ActorMetaData"));
+			if (InAssetData.GetTagValue(NAME_ActorMetaData, ActorMetaDataStr))
+			{
+				bool bIsValidClass = true;
+				UClass* ActorClass = FindObject<UClass>(ANY_PACKAGE, *ActorClassName, true);
+
+				if (!ActorClass)
+				{
+					ActorClass = AActor::StaticClass();
+					bIsValidClass = false;
+				}
+
+				FWorldPartitionActorDescInitData ActorDescInitData;
+				ActorDescInitData.NativeClass = ActorClass;
+				ActorDescInitData.PackageName = InAssetData.PackageName;
+				ActorDescInitData.ActorPath = InAssetData.ObjectPath;
+				FBase64::Decode(ActorMetaDataStr, ActorDescInitData.SerializedData);
+
+				TUniquePtr<FWorldPartitionActorDesc> NewActorDesc(AActor::CreateClassActorDesc(ActorDescInitData.NativeClass));
+
+				NewActorDesc->Init(this, ActorDescInitData);
+			
+				//check(bIsValidClass || (NewActorDesc->GetActorIsEditorOnly() && IsRunningGame()));
+
+				if (!bIsValidClass)
+				{
+					return nullptr;
+				}
+
+				return NewActorDesc;
+			}
+		}
+
+		return nullptr;
+	};
+
 	for (const FAssetData& Asset : Assets)
 	{
-		TUniquePtr<FWorldPartitionActorDesc>* NewActorDesc = new(ActorDescList) TUniquePtr<FWorldPartitionActorDesc>(GetActorDescriptor(Asset));
-		if (NewActorDesc->IsValid())
+		TUniquePtr<FWorldPartitionActorDesc> ActorDesc = GetActorDescriptor(Asset);
+
+		if (ActorDesc.IsValid())
 		{
-			Actors.Add((*NewActorDesc)->GetGuid(), NewActorDesc);
+			AddActorDescriptor(ActorDesc.Release());
 		}
 	}
 
@@ -95,49 +139,6 @@ void UActorDescContainer::BeginDestroy()
 }
 
 #if WITH_EDITOR
-TUniquePtr<FWorldPartitionActorDesc> UActorDescContainer::GetActorDescriptor(const FAssetData& InAssetData)
-{
-	FString ActorClassName;
-	static FName NAME_ActorMetaDataClass(TEXT("ActorMetaDataClass"));
-	if (InAssetData.GetTagValue(NAME_ActorMetaDataClass, ActorClassName))
-	{
-		FString ActorMetaDataStr;
-		static FName NAME_ActorMetaData(TEXT("ActorMetaData"));
-		if (InAssetData.GetTagValue(NAME_ActorMetaData, ActorMetaDataStr))
-		{
-			bool bIsValidClass = true;
-			UClass* ActorClass = FindObject<UClass>(ANY_PACKAGE, *ActorClassName, true);
-
-			if (!ActorClass)
-			{
-				ActorClass = AActor::StaticClass();
-				bIsValidClass = false;
-			}
-
-			FWorldPartitionActorDescInitData ActorDescInitData;
-			ActorDescInitData.NativeClass = ActorClass;
-			ActorDescInitData.PackageName = InAssetData.PackageName;
-			ActorDescInitData.ActorPath = InAssetData.ObjectPath;
-			FBase64::Decode(ActorMetaDataStr, ActorDescInitData.SerializedData);
-
-			TUniquePtr<FWorldPartitionActorDesc> NewActorDesc(AActor::CreateClassActorDesc(ActorDescInitData.NativeClass));
-
-			NewActorDesc->Init(this, ActorDescInitData);
-			
-			//check(bIsValidClass || (NewActorDesc->GetActorIsEditorOnly() && IsRunningGame()));
-
-			if (!bIsValidClass)
-			{
-				return nullptr;
-			}
-
-			return NewActorDesc;
-		}
-	}
-
-	return nullptr;
-}
-
 bool UActorDescContainer::ShouldHandleActorEvent(const AActor* Actor)
 {
 	if (Actor && Actor->IsPackageExternal() && (Actor->GetLevel() == World->PersistentLevel) && Actor->IsMainPackageActor())
@@ -168,7 +169,7 @@ void UActorDescContainer::OnObjectPreSave(UObject* Object)
 		if (ShouldHandleActorEvent(Actor))
 		{
 			check(!Actor->IsPendingKill());
-			if (TUniquePtr<FWorldPartitionActorDesc>* ExistingActorDesc = Actors.FindRef(Actor->GetActorGuid()))
+			if (TUniquePtr<FWorldPartitionActorDesc>* ExistingActorDesc = GetActorDescriptor(Actor->GetActorGuid()))
 			{
 				// Pin the actor handle on the actor to prevent unloading it when unhashing
 				FWorldPartitionHandle ExistingActorHandle(ExistingActorDesc);
@@ -176,25 +177,19 @@ void UActorDescContainer::OnObjectPreSave(UObject* Object)
 
 				TUniquePtr<FWorldPartitionActorDesc> NewActorDesc(Actor->CreateActorDesc());
 
-				OnActorDescUpdating(*ExistingActorDesc);
+				OnActorDescUpdating(ExistingActorDesc->Get());
 
 				// Transfer any reference count from external sources
 				NewActorDesc->TransferRefCounts(ExistingActorDesc->Get());
 
 				*ExistingActorDesc = MoveTemp(NewActorDesc);
 
-				OnActorDescUpdated(*ExistingActorDesc);
+				OnActorDescUpdated(ExistingActorDesc->Get());
 			}
 			// New actor
 			else
 			{
-				TUniquePtr<FWorldPartitionActorDesc>* NewActorDesc = new(ActorDescList) TUniquePtr<FWorldPartitionActorDesc>(Actor->CreateActorDesc());
-				check(NewActorDesc->IsValid());
-
-				check(!Actors.Contains((*NewActorDesc)->GetGuid()));
-				Actors.Add((*NewActorDesc)->GetGuid(), NewActorDesc);
-
-				OnActorDescAdded(*NewActorDesc);
+				OnActorDescAdded(AddActor(Actor));
 			}
 		}
 	}
@@ -206,24 +201,24 @@ void UActorDescContainer::OnPackageDeleted(UPackage* Package)
 
 	if (ShouldHandleActorEvent(Actor))
 	{
-		if (TUniquePtr<FWorldPartitionActorDesc>* ExistingActorDesc = Actors.FindRef(Actor->GetActorGuid()))
+		if (TUniquePtr<FWorldPartitionActorDesc>* ExistingActorDesc = GetActorDescriptor(Actor->GetActorGuid()))
 		{
-			OnActorDescRemoved(*ExistingActorDesc);
-			Actors.Remove((*ExistingActorDesc)->GetGuid());
-			ExistingActorDesc->Release();
+			OnActorDescRemoved(ExistingActorDesc->Get());
+			RemoveActorDescriptor(ExistingActorDesc->Get());
+			ExistingActorDesc->Reset();
 		}
 	}
 }
 
 void UActorDescContainer::RemoveActor(const FGuid& ActorGuid)
 {
-	if (TUniquePtr<FWorldPartitionActorDesc>* ExistingActorDesc = Actors.FindRef(ActorGuid))
+	if (TUniquePtr<FWorldPartitionActorDesc>* ExistingActorDesc = GetActorDescriptor(ActorGuid))
 	{
 		// This should never be called on already loaded actors
 		check(!(*ExistingActorDesc)->GetActor());
-		OnActorDescRemoved(*ExistingActorDesc);
-		Actors.Remove((*ExistingActorDesc)->GetGuid());
-		ExistingActorDesc->Release();
+		OnActorDescRemoved(ExistingActorDesc->Get());
+		RemoveActorDescriptor(ExistingActorDesc->Get());
+		ExistingActorDesc->Reset();
 	}
 }
 
diff --git a/Engine/Source/Runtime/Engine/Private/WorldPartition/ActorDescList.cpp b/Engine/Source/Runtime/Engine/Private/WorldPartition/ActorDescList.cpp
--- a/Engine/Source/Runtime/Engine/Private/WorldPartition/ActorDescList.cpp
+++ b/Engine/Source/Runtime/Engine/Private/WorldPartition/ActorDescList.cpp
@@ -5,48 +5,50 @@
 #if WITH_EDITOR
 FWorldPartitionActorDesc* FActorDescList::AddActor(const AActor* InActor)
 {
-	TUniquePtr<FWorldPartitionActorDesc>* NewActorDesc = new(ActorDescList) TUniquePtr<FWorldPartitionActorDesc>(InActor->CreateActorDesc());
-	check(NewActorDesc->IsValid());
+	FWorldPartitionActorDesc* NewActorDesc = InActor->CreateActorDesc().Release();
+	check(NewActorDesc);
 
-	check(!Actors.Contains((*NewActorDesc)->GetGuid()));
-	Actors.Add((*NewActorDesc)->GetGuid(), NewActorDesc);
+	AddActorDescriptor(NewActorDesc);
 
-	return NewActorDesc->Get();
+	return NewActorDesc;
 }
 
 const FWorldPartitionActorDesc* FActorDescList::GetActorDesc(const FGuid& Guid) const
 {
-	const TUniquePtr<FWorldPartitionActorDesc>* const * ActorDesc = Actors.Find(Guid);
+	const TUniquePtr<FWorldPartitionActorDesc>* const * ActorDesc = ActorsByGuid.Find(Guid);
 	return ActorDesc ? (*ActorDesc)->Get() : nullptr;
 }
 
 FWorldPartitionActorDesc* FActorDescList::GetActorDesc(const FGuid& Guid)
 {
-	TUniquePtr<FWorldPartitionActorDesc>** ActorDesc = Actors.Find(Guid);
+	TUniquePtr<FWorldPartitionActorDesc>** ActorDesc = ActorsByGuid.Find(Guid);
 	return ActorDesc ? (*ActorDesc)->Get() : nullptr;
 }
 
 const FWorldPartitionActorDesc& FActorDescList::GetActorDescChecked(const FGuid& Guid) const
 {
-	const TUniquePtr<FWorldPartitionActorDesc>* const ActorDesc = Actors.FindChecked(Guid);
+	const TUniquePtr<FWorldPartitionActorDesc>* const ActorDesc = ActorsByGuid.FindChecked(Guid);
 	return *ActorDesc->Get();
 }
 
 FWorldPartitionActorDesc& FActorDescList::GetActorDescChecked(const FGuid& Guid)
 {
-	TUniquePtr<FWorldPartitionActorDesc>* ActorDesc = Actors.FindChecked(Guid);
+	TUniquePtr<FWorldPartitionActorDesc>* ActorDesc = ActorsByGuid.FindChecked(Guid);
 	return *ActorDesc->Get();
 }
 
 const FWorldPartitionActorDesc* FActorDescList::GetActorDesc(const FString& PackageName) const
 {
-	const FName PackageFName(*PackageName);
-	for (const TUniquePtr<FWorldPartitionActorDesc>& ActorDescPtr : ActorDescList)
+	FString ActorName;
+	FString ActorContext;
+	if (!PackageName.Split(TEXT("."), &ActorContext, &ActorName, ESearchCase::CaseSensitive, ESearchDir::FromEnd))
 	{
-		if (ActorDescPtr && ActorDescPtr->GetActorPackage() == PackageFName)
-		{
-			return ActorDescPtr.Get();
-		}
+		ActorName = PackageName;
+	}
+
+	if (const TUniquePtr<FWorldPartitionActorDesc>* const* ActorDesc = ActorsByName.Find(*ActorName))
+	{
+		return (*ActorDesc)->Get();
 	}
 
 	return nullptr;
@@ -68,7 +70,28 @@ const FWorldPartitionActorDesc* FActorDescList::GetActorDesc(const FSoftObjectPa
 
 void FActorDescList::Empty()
 {
-	Actors.Empty();
+	ActorsByGuid.Empty();
+	ActorsByName.Empty();
 	ActorDescList.Empty();
 }
+
+void FActorDescList::AddActorDescriptor(FWorldPartitionActorDesc* ActorDesc)
+{
+	check(ActorDesc);
+	TUniquePtr<FWorldPartitionActorDesc>* NewActorDesc = new(ActorDescList) TUniquePtr<FWorldPartitionActorDesc>(ActorDesc);
+	ActorsByGuid.Add(ActorDesc->GetGuid(), NewActorDesc);
+	ActorsByName.Add(*ActorDesc->GetActorName().ToString(), NewActorDesc);
+}
+
+void FActorDescList::RemoveActorDescriptor(FWorldPartitionActorDesc* ActorDesc)
+{
+	check(ActorDesc);
+	verify(ActorsByGuid.Remove(ActorDesc->GetGuid()));
+	verify(ActorsByName.Remove(*ActorDesc->GetActorName().ToString()));
+}
+
+TUniquePtr<FWorldPartitionActorDesc>* FActorDescList::GetActorDescriptor(const FGuid& ActorGuid)
+{
+	return ActorsByGuid.FindRef(ActorGuid);
+}
 #endif
\ No newline at end of file
diff --git a/Engine/Source/Runtime/Engine/Private/WorldPartition/WorldPartition.cpp b/Engine/Source/Runtime/Engine/Private/WorldPartition/WorldPartition.cpp
--- a/Engine/Source/Runtime/Engine/Private/WorldPartition/WorldPartition.cpp
+++ b/Engine/Source/Runtime/Engine/Private/WorldPartition/WorldPartition.cpp
@@ -811,9 +811,9 @@ void UWorldPartition::UpdateLoadingEditorCell(UWorldPartitionEditorCell* Cell, b
 	}
 }
 
-void UWorldPartition::OnActorDescAdded(const TUniquePtr<FWorldPartitionActorDesc>& NewActorDesc)
+void UWorldPartition::OnActorDescAdded(FWorldPartitionActorDesc* NewActorDesc)
 {
-	HashActorDesc(NewActorDesc.Get());
+	HashActorDesc(NewActorDesc);
 
 	NewActorDesc->OnRegister();
 
@@ -823,9 +823,9 @@ void UWorldPartition::OnActorDescAdded(const TUniquePtr<FWorldPartitionActorDesc
 	}
 }
 
-void UWorldPartition::OnActorDescRemoved(const TUniquePtr<FWorldPartitionActorDesc>& ActorDesc)
+void UWorldPartition::OnActorDescRemoved(FWorldPartitionActorDesc* ActorDesc)
 {
-	UnhashActorDesc(ActorDesc.Get());
+	UnhashActorDesc(ActorDesc);
 	ActorDesc->OnUnregister();
 		
 	if (WorldPartitionEditor)
@@ -834,14 +834,14 @@ void UWorldPartition::OnActorDescRemoved(const TUniquePtr<FWorldPartitionActorDe
 	}
 }
 
-void UWorldPartition::OnActorDescUpdating(const TUniquePtr<FWorldPartitionActorDesc>& ActorDesc)
+void UWorldPartition::OnActorDescUpdating(FWorldPartitionActorDesc* ActorDesc)
 {
-	UnhashActorDesc(ActorDesc.Get());
+	UnhashActorDesc(ActorDesc);
 }
 
-void UWorldPartition::OnActorDescUpdated(const TUniquePtr<FWorldPartitionActorDesc>& ActorDesc)
+void UWorldPartition::OnActorDescUpdated(FWorldPartitionActorDesc* ActorDesc)
 {
-	HashActorDesc(ActorDesc.Get());
+	HashActorDesc(ActorDesc);
 
 	if (WorldPartitionEditor)
 	{
diff --git a/Engine/Source/Runtime/Engine/Private/WorldPartition/WorldPartitionActorDesc.cpp b/Engine/Source/Runtime/Engine/Private/WorldPartition/WorldPartitionActorDesc.cpp
--- a/Engine/Source/Runtime/Engine/Private/WorldPartition/WorldPartitionActorDesc.cpp
+++ b/Engine/Source/Runtime/Engine/Private/WorldPartition/WorldPartitionActorDesc.cpp
@@ -211,6 +211,14 @@ FBox FWorldPartitionActorDesc::GetBounds() const
 	return FBox(BoundsLocation - BoundsExtent, BoundsLocation + BoundsExtent);
 }
 
+FName FWorldPartitionActorDesc::GetActorName() const
+{
+	FString ActorName;
+	FString ActorContext;
+	verify(GetActorPath().ToString().Split(TEXT("."), &ActorContext, &ActorName, ESearchCase::CaseSensitive, ESearchDir::FromEnd));
+	return *ActorName;
+}
+
 UHLODLayer* FWorldPartitionActorDesc::GetHLODLayer() const
 {
 	return HLODLayer.IsNone() ? nullptr : Cast<UHLODLayer>(FSoftObjectPath(HLODLayer).TryLoad());
diff --git a/Engine/Source/Runtime/Engine/Private/WorldPartition/WorldPartitionHandle.cpp b/Engine/Source/Runtime/Engine/Private/WorldPartition/WorldPartitionHandle.cpp
--- a/Engine/Source/Runtime/Engine/Private/WorldPartition/WorldPartitionHandle.cpp
+++ b/Engine/Source/Runtime/Engine/Private/WorldPartition/WorldPartitionHandle.cpp
@@ -7,7 +7,7 @@
 #if WITH_EDITOR
 TUniquePtr<FWorldPartitionActorDesc>* FWorldPartitionHandleUtils::GetActorDesc(UActorDescContainer* Container, const FGuid& ActorGuid)
 {
-	if (TUniquePtr<FWorldPartitionActorDesc>** ActorDescPtr = Container->Actors.Find(ActorGuid))
+	if (TUniquePtr<FWorldPartitionActorDesc>** ActorDescPtr = Container->ActorsByGuid.Find(ActorGuid))
 	{
 		return *ActorDescPtr;
 	}
diff --git a/Engine/Source/Runtime/Engine/Public/WorldPartition/ActorDescContainer.h b/Engine/Source/Runtime/Engine/Public/WorldPartition/ActorDescContainer.h
--- a/Engine/Source/Runtime/Engine/Public/WorldPartition/ActorDescContainer.h
+++ b/Engine/Source/Runtime/Engine/Public/WorldPartition/ActorDescContainer.h
@@ -42,15 +42,13 @@ protected:
 	//~ End UObject Interface
 
 #if WITH_EDITOR
-	TUniquePtr<FWorldPartitionActorDesc> GetActorDescriptor(const FAssetData& InAssetData);
-	
 	virtual void RegisterDelegates();
 	virtual void UnregisterDelegates();
 
-	virtual void OnActorDescAdded(const TUniquePtr<FWorldPartitionActorDesc>& NewActorDesc) {}
-	virtual void OnActorDescRemoved(const TUniquePtr<FWorldPartitionActorDesc>& ActorDesc) {}
-	virtual void OnActorDescUpdating(const TUniquePtr<FWorldPartitionActorDesc>& ActorDesc) {}
-	virtual void OnActorDescUpdated(const TUniquePtr<FWorldPartitionActorDesc>& ActorDesc) {}
+	virtual void OnActorDescAdded(FWorldPartitionActorDesc* NewActorDesc) {}
+	virtual void OnActorDescRemoved(FWorldPartitionActorDesc* ActorDesc) {}
+	virtual void OnActorDescUpdating(FWorldPartitionActorDesc* ActorDesc) {}
+	virtual void OnActorDescUpdated(FWorldPartitionActorDesc* ActorDesc) {}
 
 	virtual void OnActorDescRegistered(const FWorldPartitionActorDesc&) {}
 	virtual void OnActorDescUnregistered(const FWorldPartitionActorDesc&) {}
diff --git a/Engine/Source/Runtime/Engine/Public/WorldPartition/ActorDescList.h b/Engine/Source/Runtime/Engine/Public/WorldPartition/ActorDescList.h
--- a/Engine/Source/Runtime/Engine/Public/WorldPartition/ActorDescList.h
+++ b/Engine/Source/Runtime/Engine/Public/WorldPartition/ActorDescList.h
@@ -9,6 +9,8 @@
 class ENGINE_API FActorDescList
 {
 #if WITH_EDITOR
+	friend struct FWorldPartitionHandleUtils;
+
 public:
 	FActorDescList() {}
 
@@ -27,7 +29,7 @@ public:
 	const FWorldPartitionActorDesc* GetActorDesc(const FString& PackageName) const;
 	const FWorldPartitionActorDesc* GetActorDesc(const FSoftObjectPath& InActorPath) const;
 
-	int32 GetActorDescCount() const { return Actors.Num(); }
+	int32 GetActorDescCount() const { return ActorsByGuid.Num(); }
 
 	void Empty();
 
@@ -45,7 +47,7 @@ public:
 
 	public:
 		TBaseIterator(ListType InActorDescList, UClass* InActorClass)
-			: ActorsIterator(InActorDescList->Actors)
+			: ActorsIterator(InActorDescList->ActorsByGuid)
 			, ActorClass(InActorClass)
 		{
 			check(ActorClass->IsChildOf(ActorType::StaticClass()));
@@ -139,7 +141,14 @@ public:
 	};
 
 protected:
+	void AddActorDescriptor(FWorldPartitionActorDesc* ActorDesc);
+	void RemoveActorDescriptor(FWorldPartitionActorDesc* ActorDesc);
+	TUniquePtr<FWorldPartitionActorDesc>* GetActorDescriptor(const FGuid& ActorGuid);
+
 	TChunkedArray<TUniquePtr<FWorldPartitionActorDesc>> ActorDescList;
-	TMap<FGuid, TUniquePtr<FWorldPartitionActorDesc>*> Actors;
+
+private:
+	TMap<FGuid, TUniquePtr<FWorldPartitionActorDesc>*> ActorsByGuid;
+	TMap<FName, TUniquePtr<FWorldPartitionActorDesc>*> ActorsByName;
 #endif
 };
\ No newline at end of file
diff --git a/Engine/Source/Runtime/Engine/Public/WorldPartition/WorldPartition.h b/Engine/Source/Runtime/Engine/Public/WorldPartition/WorldPartition.h
--- a/Engine/Source/Runtime/Engine/Public/WorldPartition/WorldPartition.h
+++ b/Engine/Source/Runtime/Engine/Public/WorldPartition/WorldPartition.h
@@ -100,10 +100,10 @@ private:
 	void OnEndPlay();
 
 	// UActorDescContainer events
-	virtual void OnActorDescAdded(const TUniquePtr<FWorldPartitionActorDesc>& NewActorDesc) override;
-	virtual void OnActorDescRemoved(const TUniquePtr<FWorldPartitionActorDesc>& ActorDesc) override;
-	virtual void OnActorDescUpdating(const TUniquePtr<FWorldPartitionActorDesc>& ActorDesc) override;
-	virtual void OnActorDescUpdated(const TUniquePtr<FWorldPartitionActorDesc>& ActorDesc) override;
+	virtual void OnActorDescAdded(FWorldPartitionActorDesc* NewActorDesc) override;
+	virtual void OnActorDescRemoved(FWorldPartitionActorDesc* ActorDesc) override;
+	virtual void OnActorDescUpdating(FWorldPartitionActorDesc* ActorDesc) override;
+	virtual void OnActorDescUpdated(FWorldPartitionActorDesc* ActorDesc) override;
 
 	virtual void OnActorDescRegistered(const FWorldPartitionActorDesc&) override;
 	virtual void OnActorDescUnregistered(const FWorldPartitionActorDesc&) override;
diff --git a/Engine/Source/Runtime/Engine/Public/WorldPartition/WorldPartitionActorDesc.h b/Engine/Source/Runtime/Engine/Public/WorldPartition/WorldPartitionActorDesc.h
--- a/Engine/Source/Runtime/Engine/Public/WorldPartition/WorldPartitionActorDesc.h
+++ b/Engine/Source/Runtime/Engine/Public/WorldPartition/WorldPartitionActorDesc.h
@@ -73,6 +73,8 @@ public:
 	inline FName GetFolderPath() const { return FolderPath; }
 	FBox GetBounds() const;
 
+	FName GetActorName() const;
+
 	virtual bool GetContainerInstance(const UActorDescContainer*& OutLevelContainer, FTransform& OutLevelTransform, EContainerClusterMode& OutClusterMode) const { return false; }
 
 	bool operator==(const FWorldPartitionActorDesc& Other) const
