World Partition HLOD - Fixed new condition added for grid priority

diff --git a/Engine/Source/Runtime/Engine/Private/WorldPartition/RuntimeSpatialHash/RuntimeSpatialHashHLOD.cpp b/Engine/Source/Runtime/Engine/Private/WorldPartition/RuntimeSpatialHash/RuntimeSpatialHashHLOD.cpp
--- a/Engine/Source/Runtime/Engine/Private/WorldPartition/RuntimeSpatialHash/RuntimeSpatialHashHLOD.cpp
+++ b/Engine/Source/Runtime/Engine/Private/WorldPartition/RuntimeSpatialHash/RuntimeSpatialHashHLOD.cpp
@@ -401,7 +401,7 @@ static void UpdateHLODGridsActors(UWorld* World, const TMap<FName, FSpatialHashR
 		if (GridActor->ActorHasTag(HLODGridTag))
 		{
 			const FSpatialHashRuntimeGrid* HLODGrid = HLODGrids.Find(GridActor->GridSettings.GridName);
-			if (HLODGrid && HLODGrid->Priority)
+			if (HLODGrid && GridActor->GridSettings.Priority)
 			{
 				ExistingGridActors.Emplace(GridActor->GridSettings.GridName, GridActor);
 			}
