Fix for computing cloth skinning using multiple weights when there are multiple cloth data components on a skinned mesh and skin caching is not enabled.

#rnx

diff --git a/Engine/Shaders/Private/GpuSkinVertexFactory.ush b/Engine/Shaders/Private/GpuSkinVertexFactory.ush
--- a/Engine/Shaders/Private/GpuSkinVertexFactory.ush
+++ b/Engine/Shaders/Private/GpuSkinVertexFactory.ush
@@ -154,11 +154,13 @@ Buffer<float4> GPUSkinApexCloth;
 
 void SetupClothVertex(uint VertexIndex, out FClothVertex ClothVertexInfluences[NUM_INFLUENCES_PER_VERTEX])
 {
-    VertexIndex = VertexIndex + GPUSkinApexClothStartIndexOffset.y - GPUSkinApexClothStartIndexOffset.x;
+	const uint MinVertexIndex = GPUSkinApexClothStartIndexOffset.x;
+	const uint ClothIndexOffset = GPUSkinApexClothStartIndexOffset.y;
+    VertexIndex = NUM_INFLUENCES_PER_VERTEX * (VertexIndex - MinVertexIndex) + ClothIndexOffset;
 
     for (int i = 0; i < NUM_INFLUENCES_PER_VERTEX; ++i )
     {
-        uint Offset = (NUM_INFLUENCES_PER_VERTEX * VertexIndex + i) * NUM_FLOAT4S_PER_VERTEX_INFLUENCE;
+        const uint Offset = (VertexIndex + i) * NUM_FLOAT4S_PER_VERTEX_INFLUENCE;
 
         ClothVertexInfluences[i].BaryCoordPos = GPUSkinApexCloth[Offset];
         ClothVertexInfluences[i].BaryCoordNormal = GPUSkinApexCloth[Offset + 1];
diff --git a/Engine/Source/Runtime/Engine/Private/SkeletalRenderGPUSkin.cpp b/Engine/Source/Runtime/Engine/Private/SkeletalRenderGPUSkin.cpp
--- a/Engine/Source/Runtime/Engine/Private/SkeletalRenderGPUSkin.cpp
+++ b/Engine/Source/Runtime/Engine/Private/SkeletalRenderGPUSkin.cpp
@@ -1624,8 +1624,6 @@ void FSkeletalMeshObjectGPUSkin::FVertexFactoryData::InitAPEXClothVertexFactorie
 	const TArray<FSkelMeshRenderSection>& Sections,
 	ERHIFeatureLevel::Type InFeatureLevel)
 {
-
-	bool bUseMultipleInfluences = (VertexBuffers.APEXClothVertexBuffer->GetNumVertices() > VertexBuffers.StaticVertexBuffers->PositionVertexBuffer.GetNumVertices());
 	bool bSupportsManualFetch = InFeatureLevel >= ERHIFeatureLevel::SM5 && RHISupportsManualVertexFetch(GetFeatureLevelShaderPlatform(InFeatureLevel));
 
 	// clear existing factories (resources assumed to have been released already)
@@ -1634,6 +1632,10 @@ void FSkeletalMeshObjectGPUSkin::FVertexFactoryData::InitAPEXClothVertexFactorie
 	{
 		if (Sections[FactoryIdx].HasClothingData() && bSupportsManualFetch)
 		{
+			const uint32 NumClothWeights = Sections[FactoryIdx].ClothMappingData.Num();
+			const uint32 NumPositionVertices = Sections[FactoryIdx].NumVertices;
+			const bool bUseMultipleInfluences = (NumClothWeights > NumPositionVertices);
+
 			GPUSkinBoneInfluenceType BoneInfluenceType = VertexBuffers.SkinWeightVertexBuffer->GetBoneInfluenceType();
 			if (BoneInfluenceType == GPUSkinBoneInfluenceType::DefaultBoneInfluence)
 			{
