Fix for stale flags on sound wave proxy attempting to parse zeroth chunk with invalid CompressedAudioInfo

diff --git a/Engine/Source/Runtime/Engine/Classes/Sound/SoundWave.h b/Engine/Source/Runtime/Engine/Classes/Sound/SoundWave.h
--- a/Engine/Source/Runtime/Engine/Classes/Sound/SoundWave.h
+++ b/Engine/Source/Runtime/Engine/Classes/Sound/SoundWave.h
@@ -342,10 +342,10 @@ public:
 
 	void ReleaseCompressedAudio();
 
-	bool IsStreaming() const { return bIsStreaming; }
-	bool IsSeekableStreaming() const { return bIsStreaming && bSeekableStreaming; }
+	bool IsStreaming() const { return *bIsStreamingPtr; }
+	bool IsSeekableStreaming() const { return *bIsStreamingPtr && *bSeekableStreamingPtr; }
 	bool IsRetainingAudio() const { return FirstChunk.IsValid(); }
-	bool WasLoadingBehaviorOverridden() const { return bLoadingBehaviorOverridden; }
+	bool WasLoadingBehaviorOverridden() const { return *bLoadingBehaviorOverriddenPtr; }
 	ESoundWaveLoadingBehavior GetLoadingBehavior() const { return LoadingBehavior; }
 
 	static TArrayView<const uint8> GetZerothChunk(const FSoundWaveProxyPtr& SoundWaveProxy, bool bForImmediatePlayback = false);
@@ -385,10 +385,10 @@ private:
 	FAudioChunkHandle FirstChunk;
 	ESoundWaveLoadingBehavior LoadingBehavior;
 
-	bool bLoadingBehaviorOverridden;
-	bool bIsStreaming;
-	bool bSeekableStreaming;
-	bool bShouldUseStreamCaching;
+	TSharedPtr<FThreadSafeBool, ESPMode::ThreadSafe> bLoadingBehaviorOverriddenPtr;
+	TSharedPtr<FThreadSafeBool, ESPMode::ThreadSafe> bIsStreamingPtr;
+	TSharedPtr<FThreadSafeBool, ESPMode::ThreadSafe> bSeekableStreamingPtr;
+	TSharedPtr<FThreadSafeBool, ESPMode::ThreadSafe> bShouldUseStreamCachingPtr;
 
 #if WITH_EDITOR
 	TSharedPtr<FThreadSafeCounter> CurrentChunkRevision{ nullptr };
@@ -434,6 +434,10 @@ public:
 	/** Whether this sound supports seeking. This requires recooking with a codec which supports seekability and streaming. */
 	UPROPERTY(EditAnywhere, Category = "Playback|Streaming", meta = (DisplayName = "Seekable", EditCondition = "bStreaming"))
 	uint8 bSeekableStreaming : 1;
+private:
+	// updated to reflect bSeekableStreaming in PostEditChangeProperty()
+	TSharedPtr<FThreadSafeBool, ESPMode::ThreadSafe> bIsSeekableStreamingProxyFlag{ MakeShared<FThreadSafeBool, ESPMode::ThreadSafe>() };
+public:
 
 	// Loading behavior members are lazily initialized in const getters
 	/** Specifies how and when compressed audio data is loaded for asset if stream caching is enabled. */
@@ -441,7 +445,7 @@ public:
 	mutable ESoundWaveLoadingBehavior LoadingBehavior;
 
 	/** Set to true if LoadingBehavior was inherited from a SoundCue. This is useful for debugging/logging */
-	mutable uint8 bLoadingBehaviorOverridden : 1;
+	mutable TSharedPtr<FThreadSafeBool, ESPMode::ThreadSafe> bLoadingBehaviorOverriddenPtr { MakeShared<FThreadSafeBool, ESPMode::ThreadSafe>() };
 
 	/** Set to true for programmatically generated audio. */
 	uint8 bProcedural : 1;
@@ -1019,7 +1023,10 @@ public:
 	static FName GetCurvePropertyName() { return GET_MEMBER_NAME_CHECKED(USoundWave, Curves); }
 #endif // WITH_EDITOR
 
-	/** Checks whether sound has been categorised as streaming. */
+	/** Checks whether sound has been categorized as streaming. */
+private:
+	TSharedPtr<FThreadSafeBool, ESPMode::ThreadSafe> bIsStreamingProxyFlag{ MakeShared<FThreadSafeBool, ESPMode::ThreadSafe>() }; // shared w/ proxies. The below functions should update this bool
+public:
 	bool IsStreaming(const TCHAR* PlatformName = nullptr) const;
 	bool IsStreaming(const FPlatformAudioCookOverrides& Overrides) const;
 
@@ -1028,6 +1035,10 @@ public:
 	/**
 	 * Checks whether we should use the load on demand cache.
 	 */
+private:
+	// ShouldUseStreamCaching() updates this value with what it returns
+	TSharedPtr<FThreadSafeBool, ESPMode::ThreadSafe> bShouldUseStreamCachingProxyFlag{ MakeShared<FThreadSafeBool, ESPMode::ThreadSafe>() };
+public:
 	bool ShouldUseStreamCaching() const;
 
 	/**
diff --git a/Engine/Source/Runtime/Engine/Private/SoundWave.cpp b/Engine/Source/Runtime/Engine/Private/SoundWave.cpp
--- a/Engine/Source/Runtime/Engine/Private/SoundWave.cpp
+++ b/Engine/Source/Runtime/Engine/Private/SoundWave.cpp
@@ -463,7 +463,7 @@ void USoundWave::Serialize( FArchive& Ar )
 				const ITargetPlatform* CookingTarget = Ar.CookingTarget();
 				if (!CookingTarget->IsServerOnly())
 				{
-					// for now we only support one format per wav
+					// for now we only support one format per wave
 					FName Format = CookingTarget->GetWaveFormat(this);
 					const FPlatformAudioCookOverrides* CompressionOverrides = FPlatformCompressionUtilities::GetCookOverrides(*CookingTarget->IniPlatformName());
 
@@ -558,7 +558,7 @@ void USoundWave::Serialize( FArchive& Ar )
 			{
 				// Prime first chunk of audio.
 				CachedSoundWaveLoadingBehavior = ESoundWaveLoadingBehavior::PrimeOnLoad;
-				bLoadingBehaviorOverridden = true;
+				*bLoadingBehaviorOverriddenPtr = true;
 
 				if (!InternalProxy.IsValid())
 				{
@@ -1804,6 +1804,11 @@ void USoundWave::PostEditChangeProperty(FPropertyChangedEvent& PropertyChangedEv
 	static const FName StreamingFName = GET_MEMBER_NAME_CHECKED(USoundWave, bStreaming);
 	static const FName SeekableStreamingFName = GET_MEMBER_NAME_CHECKED(USoundWave, bSeekableStreaming);
 
+	// force proxy flags to be up to date
+	*bIsSeekableStreamingProxyFlag = bSeekableStreaming;
+	*bIsStreamingProxyFlag = IsStreaming(nullptr);
+	*bShouldUseStreamCachingProxyFlag = ShouldUseStreamCaching();
+
 	// Prevent constant re-compression of SoundWave while properties are being changed interactively
 	if (PropertyChangedEvent.ChangeType != EPropertyChangeType::Interactive)
 	{
@@ -2343,12 +2348,19 @@ float USoundWave::GetDuration()
 
 bool USoundWave::IsStreaming(const TCHAR* PlatformName/* = nullptr */) const
 {
+	bool Result = false;
 	if (GIsEditor && ForceNonStreamingInEditorCVar != 0)
 	{
-		return false;
+		Result = false;
+	}
+	else
+	{
+		Result = IsStreaming(*FPlatformCompressionUtilities::GetCookOverrides(PlatformName));
 	}
 
-	return IsStreaming(*FPlatformCompressionUtilities::GetCookOverrides(PlatformName));
+	*bIsStreamingProxyFlag = Result; // updates proxis
+	
+	return Result;
 }
 
 bool USoundWave::IsStreaming(const FPlatformAudioCookOverrides& Overrides) const
@@ -2357,14 +2369,17 @@ bool USoundWave::IsStreaming(const FPlatformAudioCookOverrides& Overrides) const
 	// our cook overrides, or the AutoStreamingThreshold was set and this sound is longer than the auto streaming threshold.
 	if (bStreaming)
 	{
+		*bIsStreamingProxyFlag = true;
 		return true;
 	}
 	else if (LoadingBehavior == ESoundWaveLoadingBehavior::ForceInline) // TODO: Support setting ESoundWaveLoadingBehavior to ForceInline on a USoundClass.
 	{
+		*bIsStreamingProxyFlag = false;
 		return false;
 	}
 	else if (bProcedural)
 	{
+		*bIsStreamingProxyFlag = false;
 		return false;
 	}
 
@@ -2372,14 +2387,21 @@ bool USoundWave::IsStreaming(const FPlatformAudioCookOverrides& Overrides) const
 	const bool bUsesStreamCache = Overrides.bUseStreamCaching;
 	const bool bOverAutoStreamingThreshold = (Overrides.AutoStreamingThreshold > SMALL_NUMBER  && Duration > Overrides.AutoStreamingThreshold);
 
-	return bUsesStreamCache || bOverAutoStreamingThreshold;
+	const bool Result = bUsesStreamCache || bOverAutoStreamingThreshold;
+
+	*bIsStreamingProxyFlag = Result;
+	return Result;
 }
 
 bool USoundWave::ShouldUseStreamCaching() const
 {
 	const bool bPlatformUsingStreamCaching = FPlatformCompressionUtilities::IsCurrentPlatformUsingStreamCaching();
 	const bool bIsStreaming = IsStreaming(nullptr);
-	return  bPlatformUsingStreamCaching && bIsStreaming;
+	const bool Result = bPlatformUsingStreamCaching && bIsStreaming;
+
+	*bShouldUseStreamCachingProxyFlag = Result;
+
+	return Result;
 }
 
 TArrayView<const uint8> USoundWave::GetZerothChunk(bool bForImmediatePlayback)
@@ -2453,6 +2475,10 @@ FSoundWaveProxyPtr USoundWave::CreateSoundWaveProxy()
 	EnsureZerothChunkIsLoaded();
 #endif // #if WITH_EDITORONLY_DATA
 
+	*bIsSeekableStreamingProxyFlag = bSeekableStreaming;
+	*bIsStreamingProxyFlag = IsStreaming(nullptr);
+	*bShouldUseStreamCachingProxyFlag = ShouldUseStreamCaching();
+
 	return MakeShared<FSoundWaveProxy, ESPMode::ThreadSafe>(this);
 }
 
@@ -2462,6 +2488,10 @@ TUniquePtr<Audio::IProxyData> USoundWave::CreateNewProxyData(const Audio::FProxy
 	EnsureZerothChunkIsLoaded();
 #endif // #if WITH_EDITORONLY_DATA
 
+	*bIsSeekableStreamingProxyFlag = bSeekableStreaming;
+	*bIsStreamingProxyFlag = IsStreaming(nullptr);
+	*bShouldUseStreamCachingProxyFlag = ShouldUseStreamCaching();
+
 	return MakeUnique<FSoundWaveProxy>(this);
 }
 
@@ -2957,7 +2987,7 @@ void USoundWave::OverrideLoadingBehavior(ESoundWaveLoadingBehavior InLoadingBeha
 	// (if this soundwave isn't loaded yet, 
 	// CachedSoundWaveLoadingBehavior will take precedence when it does load)
 	CachedSoundWaveLoadingBehavior = InLoadingBehavior;
-	bLoadingBehaviorOverridden = true;
+	*bLoadingBehaviorOverriddenPtr = true;
 
 	// If we're loading for the cook commandlet, we don't have streamed audio chunks to load.
 	const bool bHasBuiltStreamedAudio = !GetOutermost()->HasAnyPackageFlags(PKG_ReloadingForCooker) && FApp::CanEverRenderAudio();
@@ -2995,14 +3025,14 @@ void USoundWave::CacheInheritedLoadingBehavior() const
 			CachedSoundWaveLoadingBehavior = LoadingBehavior;
 		}
 	}
- 	else if (bLoadingBehaviorOverridden)
+ 	else if (*bLoadingBehaviorOverriddenPtr)
  	{
  		ensureMsgf(CachedSoundWaveLoadingBehavior != ESoundWaveLoadingBehavior::Inherited, TEXT("SoundCue set loading behavior to Inherited on SoudWave: %s"), *GetFullName());
  	}
 	else
 	{
 		// if this is true then the behavior should not be Inherited here
-		check(!bLoadingBehaviorOverridden);
+		check(!*bLoadingBehaviorOverriddenPtr);
 
 		USoundClass* CurrentSoundClass = SoundClassObject;
 		ESoundWaveLoadingBehavior SoundClassLoadingBehavior = ESoundWaveLoadingBehavior::Inherited;
@@ -3023,7 +3053,7 @@ void USoundWave::CacheInheritedLoadingBehavior() const
 
 			// override this loading behavior w/ our default
 			SoundClassLoadingBehavior = DefaultLoadingBehavior;
-			bLoadingBehaviorOverridden = true;
+			*bLoadingBehaviorOverriddenPtr = true;
 		}
 
 		CachedSoundWaveLoadingBehavior = SoundClassLoadingBehavior;
@@ -3038,12 +3068,12 @@ ESoundWaveLoadingBehavior USoundWave::GetLoadingBehavior(bool bCheckSoundClasses
 
 	if (!bCheckSoundClasses)
 	{
-		if ((LoadingBehavior != ESoundWaveLoadingBehavior::Inherited && !bLoadingBehaviorOverridden))
+		if ((LoadingBehavior != ESoundWaveLoadingBehavior::Inherited && !*bLoadingBehaviorOverriddenPtr))
 		{
 			// If this sound wave specifies it's own loading behavior, use that.
 			return LoadingBehavior;
 		}
-		else if (bLoadingBehaviorOverridden)
+		else if (*bLoadingBehaviorOverriddenPtr)
 		{
 			// If this sound wave has already had it's loading behavior cached from soundclasses or soundcues, use that.
 			return CachedSoundWaveLoadingBehavior;
@@ -3077,10 +3107,10 @@ FSoundWaveProxy::FSoundWaveProxy(USoundWave* InWave)
 	, NumFrames((int32)(Duration * (float)InWave->SampleRate))
 	, FirstChunk(InWave->FirstChunk)
 	, LoadingBehavior(InWave->CachedSoundWaveLoadingBehavior)
-	, bLoadingBehaviorOverridden(InWave->bLoadingBehaviorOverridden)
-	, bIsStreaming(InWave->IsStreaming(nullptr))
-	, bSeekableStreaming(InWave->bSeekableStreaming)
-	, bShouldUseStreamCaching(InWave->ShouldUseStreamCaching())
+	, bLoadingBehaviorOverriddenPtr(InWave->bLoadingBehaviorOverriddenPtr)
+	, bIsStreamingPtr(InWave->bIsStreamingProxyFlag)
+	, bSeekableStreamingPtr(InWave->bIsSeekableStreamingProxyFlag)
+	, bShouldUseStreamCachingPtr(InWave->bShouldUseStreamCachingProxyFlag)
 #if WITH_EDITOR
 	, CurrentChunkRevision(InWave->CurrentChunkRevision)
 #endif // #if WITH_EDITOR
@@ -3101,6 +3131,11 @@ FSoundWaveProxy::FSoundWaveProxy(USoundWave* InWave)
 		RunningPlatformData = InWave->RunningPlatformData;
 		NumChunks = InWave->GetNumChunks();
 	}
+
+	// force updates of flags shared w/ the sound wave
+	// (assignment is to avoid these calls to const functions w/ mutable side effects being from being optimized out)
+	*bIsStreamingPtr = InWave->IsStreaming(nullptr);
+	*bShouldUseStreamCachingPtr = InWave->ShouldUseStreamCaching();
 }
 
 void FSoundWaveProxy::ReleaseCompressedAudio()
@@ -3118,7 +3153,7 @@ TArrayView<const uint8> FSoundWaveProxy::GetZerothChunk(const FSoundWaveProxyPtr
 {
 	if (ensure(SoundWaveProxy.IsValid()) && SoundWaveProxy->IsZerothChunkDataLoaded())
 	{
-		if (SoundWaveProxy->bShouldUseStreamCaching)
+		if (*SoundWaveProxy->bShouldUseStreamCachingPtr)
 		{
 			if (SoundWaveProxy->GetNumChunks() > 1)
 			{
@@ -3179,7 +3214,7 @@ const TArrayView<uint8> FSoundWaveProxy::GetZerothChunkDataView() const
 void FSoundWaveProxy::EnsureZerothChunkIsLoaded()
 {
 	// If the zeroth chunk is already loaded, early exit.
-	if (ZerothChunkData->GetView().Num() > 0 || !bShouldUseStreamCaching)
+	if (ZerothChunkData->GetView().Num() > 0 || !*bShouldUseStreamCachingPtr)
 	{
 		return;
 	}
