Don't create actor descriptors when checking map for errors after loading.
This is the real fix for 15887579.

diff --git a/Engine/Source/Runtime/Engine/Private/WorldPartition/WorldPartitionActorDescViewProxy.cpp b/Engine/Source/Runtime/Engine/Private/WorldPartition/WorldPartitionActorDescViewProxy.cpp
--- a/Engine/Source/Runtime/Engine/Private/WorldPartition/WorldPartitionActorDescViewProxy.cpp
+++ b/Engine/Source/Runtime/Engine/Private/WorldPartition/WorldPartitionActorDescViewProxy.cpp
@@ -10,8 +10,11 @@ FWorldPartitionActorViewProxy::FWorldPartitionActorViewProxy(const FWorldPartiti
 {
 	if (AActor* Actor = ActorDesc->GetActor())
 	{
-		CachedActorDesc = Actor->CreateActorDesc();
-		ActorDesc = CachedActorDesc.Get();
+		if (Actor->GetPackage()->IsDirty())
+		{
+			CachedActorDesc = Actor->CreateActorDesc();
+			ActorDesc = CachedActorDesc.Get();
+		}
 	}
 }
 #endif
\ No newline at end of file
diff --git a/Engine/Source/Runtime/Engine/Public/WorldPartition/WorldPartitionActorDescViewProxy.h b/Engine/Source/Runtime/Engine/Public/WorldPartition/WorldPartitionActorDescViewProxy.h
--- a/Engine/Source/Runtime/Engine/Public/WorldPartition/WorldPartitionActorDescViewProxy.h
+++ b/Engine/Source/Runtime/Engine/Public/WorldPartition/WorldPartitionActorDescViewProxy.h
@@ -5,7 +5,7 @@
 
 #if WITH_EDITOR
 /**
- * A view on top of an actor desc, which reverts to the actor if its loaded.
+ * A view on top of an actor desc, which reverts to the actor if its loaded and dirty.
  */
 class FWorldPartitionActorViewProxy : public FWorldPartitionActorDescView
 {
