infrastructure code for codecs other than adpcm to do seekable streaming

diff --git a/Engine/Source/Runtime/Engine/Private/AudioDecompress.cpp b/Engine/Source/Runtime/Engine/Private/AudioDecompress.cpp
--- a/Engine/Source/Runtime/Engine/Private/AudioDecompress.cpp
+++ b/Engine/Source/Runtime/Engine/Private/AudioDecompress.cpp
@@ -27,6 +27,8 @@ IStreamedCompressedInfo::IStreamedCompressedInfo()
 	, CurrentChunkIndex(0)
 	, bPrintChunkFailMessage(true)
 	, SrcBufferPadding(0)
+	, StreamSeekBlockIndex(INDEX_NONE)
+	, StreamSeekBlockOffset(0)
 {
 }
 
@@ -189,6 +191,25 @@ bool IStreamedCompressedInfo::StreamCompressedData(uint8* Destination, bool bLoo
 	// Write out any PCM data that was decoded during the last request
 	uint32 RawPCMOffset = WriteFromDecodedPCM(Destination, BufferSize);
 
+	// If we have a pending next chunk from seeking, move to it now.
+	if (StreamSeekBlockIndex != INDEX_NONE)
+	{
+		uint32 ChunkSize = 0;
+		SrcBufferData = GetLoadedChunk(StreamingSoundWave, StreamSeekBlockIndex, ChunkSize);
+		UE_LOG(LogAudio, Log, TEXT("Seek block request: %d / %d (%s)"), StreamSeekBlockIndex.load(), StreamSeekBlockOffset, SrcBufferData == nullptr ? TEXT("present") : TEXT("missing"));
+		if (SrcBufferData == nullptr)
+		{
+			// After a seek we're likely to need to wait a bit for the chunk to get in to memory.
+			ZeroBuffer(Destination + RawPCMOffset, BufferSize - RawPCMOffset);
+			return false;
+		}
+
+		CurrentChunkIndex = StreamSeekBlockIndex;
+		SrcBufferDataSize = ChunkSize;
+		SrcBufferOffset = StreamSeekBlockOffset;
+		StreamSeekBlockIndex = INDEX_NONE;
+	}
+
 	// If next chunk wasn't loaded when last one finished reading, try to get it again now
 	if (SrcBufferData == NULL)
 	{
diff --git a/Engine/Source/Runtime/Engine/Private/AudioStreaming.cpp b/Engine/Source/Runtime/Engine/Private/AudioStreaming.cpp
--- a/Engine/Source/Runtime/Engine/Private/AudioStreaming.cpp
+++ b/Engine/Source/Runtime/Engine/Private/AudioStreaming.cpp
@@ -606,6 +606,25 @@ void FLegacyAudioStreamingManager::UpdateResourceStreaming(float DeltaTime, bool
 			// Request the chunk the source is using and the one after that
 			FWaveRequest& WaveRequest = GetWaveRequest(SoundWave->GetFObjectKey());
 			int32 SourceChunk = Decoder->GetCurrentChunkIndex();
+
+			// If there's a seek request, use that as our current chunk index.				
+			if (Decoder->IsStreamedCompressedInfo())
+			{
+				// This is awk because the decoder thread catches the seek to
+				// advance, so we have to read this atomically. However, if it changes,
+				// that means that the chunk is already in memory (because StreamCompressedData
+				// only updates it when the chunk is present).
+				// ... and also, GetCurrentChunkIndex has the same problem. So it's layers
+				// on layers.
+				IStreamedCompressedInfo* Streamed = (IStreamedCompressedInfo*)Decoder;
+				int32 StreamSeekToBlockIndex = Streamed->GetStreamSeekBlockIndex();
+				if (StreamSeekToBlockIndex != INDEX_NONE)
+				{
+					UE_LOG(LogAudio, Log, TEXT("Diverting streaming chunk due to seek -- from %d to %d"), SourceChunk, StreamSeekToBlockIndex);
+					SourceChunk = StreamSeekToBlockIndex;
+				}
+			}
+
 			if (SourceChunk >= 0 && SourceChunk < (int32)SoundWave->GetNumChunks())
 			{
 				WaveRequest.RequiredIndices.AddUnique(SourceChunk);
diff --git a/Engine/Source/Runtime/Engine/Public/AudioDecompress.h b/Engine/Source/Runtime/Engine/Public/AudioDecompress.h
--- a/Engine/Source/Runtime/Engine/Public/AudioDecompress.h
+++ b/Engine/Source/Runtime/Engine/Public/AudioDecompress.h
@@ -89,6 +89,12 @@ public:
 	*/
 	virtual int GetStreamBufferSize() const = 0;
 
+	/**
+	* Returns whether this instance can be cast to a IStreamedCompressedInfo. Surprisingly
+	* SupportsStreaming doesn't work for this.
+	*/
+	virtual bool IsStreamedCompressedInfo() const { return false; }
+
 	////////////////////////////////////////////////////////////////
 	// Following functions are optional if streaming is supported //
 	////////////////////////////////////////////////////////////////
@@ -201,6 +207,7 @@ public:
 	virtual bool StreamCompressedData(uint8* Destination, bool bLooping, uint32 BufferSize, int32& OutNumBytesStreamed) override;
 	virtual int32 GetCurrentChunkIndex() const override { return CurrentChunkIndex; }
 	virtual int32 GetCurrentChunkOffset() const override { return SrcBufferOffset; }
+	virtual bool IsStreamedCompressedInfo() const override { return true; }
 	//~ End ICompressedInfo Interface
 
 	/** Parse the header information from the input source buffer data. This is dependent on compression format. */
@@ -221,6 +228,8 @@ public:
 	/** The size of the decode PCM buffer size. */
 	virtual uint32 GetMaxFrameSizeSamples() const = 0;
 
+	int32 GetStreamSeekBlockIndex() const { return StreamSeekBlockIndex; }
+	int32 GetStreamSeekBlockOffset() const { return StreamSeekBlockOffset; }
 protected:
 
 	/** Reads from the internal source audio buffer stream of the given data size. */
@@ -309,6 +318,18 @@ protected:
 	uint32 SrcBufferPadding;
 	/** Chunk Handle to ensure that this chunk of streamed audio is not deleted while we are using it. */
 	FAudioChunkHandle CurCompressedChunkHandle;
+
+	/** 
+		When a streaming seek request comes down, this is the block we are going to. INDEX_NONE means no seek pending.
+		When using the legacy streaming system this is read on a thread other than the decompression thread
+		to prime the correct chunk, so it needs an atomic. It's only ever _set_ on one thread, and the read
+		has no timing restrictions, so no lock is necessary. Also the legacy streamer isn't used anymore anyway.
+
+		This and StreamSeekBlockOffset are expected to be set by the codec's SeekToTime() function.
+	*/
+	std::atomic<int32> StreamSeekBlockIndex;
+	/** When a streaming seek request comes down, this is the offset in to the block we want to start decoding from. */
+	int32 StreamSeekBlockOffset;
 };
 
 
