Fixed crashes in the sequencer due to capture components.

diff --git a/Engine/Source/Runtime/Engine/Private/Components/ReflectionCaptureComponent.cpp b/Engine/Source/Runtime/Engine/Private/Components/ReflectionCaptureComponent.cpp
--- a/Engine/Source/Runtime/Engine/Private/Components/ReflectionCaptureComponent.cpp
+++ b/Engine/Source/Runtime/Engine/Private/Components/ReflectionCaptureComponent.cpp
@@ -1050,7 +1050,7 @@ void UReflectionCaptureComponent::MarkDirtyForRecapture()
 
 void UReflectionCaptureComponent::UpdateReflectionCaptureContents(UWorld* WorldToUpdate, const TCHAR* CaptureReason, bool bVerifyOnlyCapturing, bool bCapturingForMobile)
 {
-	if (WorldToUpdate->Scene
+	if (WorldToUpdate && WorldToUpdate->Scene
 		// Don't capture and read back capture contents if we are currently doing async shader compiling
 		// This will keep the update requests in the queue until compiling finishes
 		// Note: this will also prevent uploads of cubemaps from DDC, which is unintentional
diff --git a/Engine/Source/Runtime/Engine/Private/Components/SkyLightComponent.cpp b/Engine/Source/Runtime/Engine/Private/Components/SkyLightComponent.cpp
--- a/Engine/Source/Runtime/Engine/Private/Components/SkyLightComponent.cpp
+++ b/Engine/Source/Runtime/Engine/Private/Components/SkyLightComponent.cpp
@@ -797,7 +797,7 @@ void USkyLightComponent::UpdateSkyCaptureContentsArray(UWorld* WorldToUpdate, TA
 
 void USkyLightComponent::UpdateSkyCaptureContents(UWorld* WorldToUpdate)
 {
-	if (WorldToUpdate->Scene)
+	if (WorldToUpdate && WorldToUpdate->Scene)
 	{
 		QUICK_SCOPE_CYCLE_COUNTER(STAT_SkylightCaptures);
 
