WorldPartitionLevelStreamingPolicy won't limit number of streaming cells to load while world hasn't started

diff --git a/Engine/Source/Runtime/Engine/Private/WorldPartition/WorldPartitionLevelStreamingPolicy.cpp b/Engine/Source/Runtime/Engine/Private/WorldPartition/WorldPartitionLevelStreamingPolicy.cpp
--- a/Engine/Source/Runtime/Engine/Private/WorldPartition/WorldPartitionLevelStreamingPolicy.cpp
+++ b/Engine/Source/Runtime/Engine/Private/WorldPartition/WorldPartitionLevelStreamingPolicy.cpp
@@ -45,6 +45,13 @@ int32 UWorldPartitionLevelStreamingPolicy::GetCellLoadingCount() const
 	return CellLoadingCount;
 }
 
+int32 UWorldPartitionLevelStreamingPolicy::GetMaxCellsToLoad() const
+{
+	// This policy limits the number of concurrent loading streaming cells, except if match hasn't started
+	UWorld* World = WorldPartition->GetWorld();
+	return World->bMatchStarted ? (GMaxLoadingLevelStreamingCells - GetCellLoadingCount()) : MAX_int32;
+}
+
 void UWorldPartitionLevelStreamingPolicy::SetTargetStateForCells(EWorldPartitionRuntimeCellState TargetState, const TSet<const UWorldPartitionRuntimeCell*>& Cells)
 {
 	switch (TargetState)
@@ -65,9 +72,7 @@ void UWorldPartitionLevelStreamingPolicy::SetTargetStateForCells(EWorldPartition
 
 void UWorldPartitionLevelStreamingPolicy::SetCellsStateToLoaded(const TSet<const UWorldPartitionRuntimeCell*>& ToLoadCells)
 {
-	// This policy limits the number of concurrent loading streaming cells
-	int32 CellLoadingCount = GetCellLoadingCount();
-	int32 MaxCellsToLoad = GMaxLoadingLevelStreamingCells - CellLoadingCount;
+	int32 MaxCellsToLoad = GetMaxCellsToLoad();
 	
 	// Sort cells based on importance
 	TArray<const UWorldPartitionRuntimeCell*, TInlineAllocator<256>> SortedCells;
@@ -98,13 +103,7 @@ void UWorldPartitionLevelStreamingPolicy::SetCellsStateToLoaded(const TSet<const
 
 void UWorldPartitionLevelStreamingPolicy::SetCellsStateToActivated(const TSet<const UWorldPartitionRuntimeCell*>& ToActivateCells)
 {
-	// This policy limits the number of concurrent loading streaming cells
-	int32 CellLoadingCount = GetCellLoadingCount();
-	int32 MaxCellsToLoad = GMaxLoadingLevelStreamingCells - CellLoadingCount;
-	if (MaxCellsToLoad <= 0)
-	{
-		return;
-	}
+	int32 MaxCellsToLoad = GetMaxCellsToLoad();
 
 	// Sort cells based on importance
 	TArray<const UWorldPartitionRuntimeCell*, TInlineAllocator<256>> SortedCells;
diff --git a/Engine/Source/Runtime/Engine/Public/WorldPartition/WorldPartitionLevelStreamingPolicy.h b/Engine/Source/Runtime/Engine/Public/WorldPartition/WorldPartitionLevelStreamingPolicy.h
--- a/Engine/Source/Runtime/Engine/Public/WorldPartition/WorldPartitionLevelStreamingPolicy.h
+++ b/Engine/Source/Runtime/Engine/Public/WorldPartition/WorldPartitionLevelStreamingPolicy.h
@@ -41,6 +41,7 @@ private:
 	void SetCellsStateToActivated(const TSet<const UWorldPartitionRuntimeCell*>& ToActivateCells);
 	void SetCellsStateToUnloaded(const TSet<const UWorldPartitionRuntimeCell*>& ToUnloadCells);
 	int32 GetCellLoadingCount() const;
+	int32 GetMaxCellsToLoad() const;
 
 #if WITH_EDITORONLY_DATA
 	UPROPERTY()
