World Partition: don't handle actor save events from autosave.

diff --git a/Engine/Source/Editor/UnrealEd/Private/FileHelpers.cpp b/Engine/Source/Editor/UnrealEd/Private/FileHelpers.cpp
--- a/Engine/Source/Editor/UnrealEd/Private/FileHelpers.cpp
+++ b/Engine/Source/Editor/UnrealEd/Private/FileHelpers.cpp
@@ -2856,6 +2856,7 @@ EAutosaveContentPackagesResult::Type FEditorFileUtils::AutosaveMapEx(const FStri
 			// Now gather the world external packages and save them if needed
 			if (World->PersistentLevel)
 			{
+				World->PersistentLevel->SetIsAutoSaveExternalActorPackages(true);
 				for (UPackage* ExternalPackage : World->PersistentLevel->GetLoadedExternalActorPackages())
 				{
 					if (ExternalPackage->IsDirty() && (bForceIfNotInList || DirtyPackagesForAutoSave.Contains(ExternalPackage))
@@ -2890,6 +2891,7 @@ EAutosaveContentPackagesResult::Type FEditorFileUtils::AutosaveMapEx(const FStri
 						}
 					}
 				}
+				World->PersistentLevel->SetIsAutoSaveExternalActorPackages(false);
 			}
 		}
 
diff --git a/Engine/Source/Runtime/Engine/Classes/Engine/Level.h b/Engine/Source/Runtime/Engine/Classes/Engine/Level.h
--- a/Engine/Source/Runtime/Engine/Classes/Engine/Level.h
+++ b/Engine/Source/Runtime/Engine/Classes/Engine/Level.h
@@ -609,6 +609,9 @@ public:
 
 	/** Whether the level has been saved after introducing actor GUIDs */
 	uint8										bContainsStableActorGUIDs:1;
+
+	/** Whether we are autosaviong external actors packages or not */
+	uint8										bIsAutoSaveExternalActorPackages:1;
 #endif
 	
 	/** The below variables are used temporarily while making a level visible.				*/
@@ -1044,7 +1047,17 @@ public:
 	 * @return Array of packages associated with this level
 	 */
 	ENGINE_API TArray<UPackage*> GetLoadedExternalActorPackages() const;
-				
+
+	/**
+	 * Tells the level that we're autosaving external actors or not
+	 */
+	ENGINE_API void SetIsAutoSaveExternalActorPackages(bool bInIsAutoSaveExternalActorPackages);
+
+	/**
+	 * Returns if we are autosaving external actor packages or not
+	 */
+	ENGINE_API bool GetIsAutoSaveExternalActorPackages() const;
+
 	/**
 	 * Create an package for this actor
 	 * @param InActorPath the fully qualified actor path, in the format: 'Outermost.Outer.Name'
diff --git a/Engine/Source/Runtime/Engine/Private/Level.cpp b/Engine/Source/Runtime/Engine/Private/Level.cpp
--- a/Engine/Source/Runtime/Engine/Private/Level.cpp
+++ b/Engine/Source/Runtime/Engine/Private/Level.cpp
@@ -330,6 +330,7 @@ ULevel::ULevel( const FObjectInitializer& ObjectInitializer )
 	FixupOverrideVertexColorsCount = 0;
 	bUseExternalActors = false;
 	bContainsStableActorGUIDs = true;
+	bIsAutoSaveExternalActorPackages = false;
 #endif	
 	bActorClusterCreated = false;
 	bIsPartitioned = false;
@@ -2469,6 +2470,17 @@ TArray<UPackage*> ULevel::GetLoadedExternalActorPackages() const
 	return ActorPackages.Array();
 }
 
+void ULevel::SetIsAutoSaveExternalActorPackages(bool bInIsAutoSaveExternalActorPackages)
+{
+	check(bIsAutoSaveExternalActorPackages != bInIsAutoSaveExternalActorPackages);
+	bIsAutoSaveExternalActorPackages = bInIsAutoSaveExternalActorPackages;
+}
+
+bool ULevel::GetIsAutoSaveExternalActorPackages() const
+{
+	return bIsAutoSaveExternalActorPackages;
+}
+
 UPackage* ULevel::CreateActorPackage(UPackage* InLevelPackage, const FString& InActorPath)
 {
 	UPackage* ActorPackage = CreatePackage(*GetActorPackageName(InLevelPackage, InActorPath));
diff --git a/Engine/Source/Runtime/Engine/Private/WorldPartition/ActorDescContainer.cpp b/Engine/Source/Runtime/Engine/Private/WorldPartition/ActorDescContainer.cpp
--- a/Engine/Source/Runtime/Engine/Private/WorldPartition/ActorDescContainer.cpp
+++ b/Engine/Source/Runtime/Engine/Private/WorldPartition/ActorDescContainer.cpp
@@ -141,7 +141,7 @@ void UActorDescContainer::BeginDestroy()
 #if WITH_EDITOR
 bool UActorDescContainer::ShouldHandleActorEvent(const AActor* Actor)
 {
-	if (Actor && Actor->IsPackageExternal() && (Actor->GetLevel() == World->PersistentLevel) && Actor->IsMainPackageActor())
+	if (Actor && Actor->IsPackageExternal() && (Actor->GetLevel() == World->PersistentLevel) && !World->PersistentLevel->GetIsAutoSaveExternalActorPackages() && Actor->IsMainPackageActor())
 	{
 		// Only handle assets that belongs to our level
 		auto RemoveAfterFirstDot = [](const FString& InValue)
