Cleanup get actor desc with soft object path after refactoring from JF Dube.

diff --git a/Engine/Source/Runtime/Engine/Private/WorldPartition/ActorDescList.cpp b/Engine/Source/Runtime/Engine/Private/WorldPartition/ActorDescList.cpp
--- a/Engine/Source/Runtime/Engine/Private/WorldPartition/ActorDescList.cpp
+++ b/Engine/Source/Runtime/Engine/Private/WorldPartition/ActorDescList.cpp
@@ -37,13 +37,13 @@ FWorldPartitionActorDesc& FActorDescList::GetActorDescChecked(const FGuid& Guid)
 	return *ActorDesc->Get();
 }
 
-const FWorldPartitionActorDesc* FActorDescList::GetActorDesc(const FString& PackageName) const
+const FWorldPartitionActorDesc* FActorDescList::GetActorDesc(const FString& ActorPath) const
 {
 	FString ActorName;
 	FString ActorContext;
-	if (!PackageName.Split(TEXT("."), &ActorContext, &ActorName, ESearchCase::CaseSensitive, ESearchDir::FromEnd))
+	if (!ActorPath.Split(TEXT("."), &ActorContext, &ActorName, ESearchCase::CaseSensitive, ESearchDir::FromEnd))
 	{
-		ActorName = PackageName;
+		ActorName = ActorPath;
 	}
 
 	if (const TUniquePtr<FWorldPartitionActorDesc>* const* ActorDesc = ActorsByName.Find(*ActorName))
@@ -54,18 +54,9 @@ const FWorldPartitionActorDesc* FActorDescList::GetActorDesc(const FString& Pack
 	return nullptr;
 }
 
-const FWorldPartitionActorDesc* FActorDescList::GetActorDesc(const FSoftObjectPath& InActorPath) const
+const FWorldPartitionActorDesc* FActorDescList::GetActorDesc(const FSoftObjectPath& ActorPath) const
 {
-	const FName ActorPath(*InActorPath.ToString());
-	for (const TUniquePtr<FWorldPartitionActorDesc>& ActorDescPtr : ActorDescList)
-	{
-		if (ActorDescPtr && ActorDescPtr->GetActorPath() == ActorPath)
-		{
-			return ActorDescPtr.Get();
-		}
-	}
-
-	return nullptr;
+	return GetActorDesc(ActorPath.ToString());
 }
 
 void FActorDescList::Empty()
@@ -80,14 +71,14 @@ void FActorDescList::AddActorDescriptor(FWorldPartitionActorDesc* ActorDesc)
 	check(ActorDesc);
 	TUniquePtr<FWorldPartitionActorDesc>* NewActorDesc = new(ActorDescList) TUniquePtr<FWorldPartitionActorDesc>(ActorDesc);
 	ActorsByGuid.Add(ActorDesc->GetGuid(), NewActorDesc);
-	ActorsByName.Add(*ActorDesc->GetActorName().ToString(), NewActorDesc);
+	ActorsByName.Add(ActorDesc->GetActorName(), NewActorDesc);
 }
 
 void FActorDescList::RemoveActorDescriptor(FWorldPartitionActorDesc* ActorDesc)
 {
 	check(ActorDesc);
 	verify(ActorsByGuid.Remove(ActorDesc->GetGuid()));
-	verify(ActorsByName.Remove(*ActorDesc->GetActorName().ToString()));
+	verify(ActorsByName.Remove(ActorDesc->GetActorName()));
 }
 
 TUniquePtr<FWorldPartitionActorDesc>* FActorDescList::GetActorDescriptor(const FGuid& ActorGuid)
