Optimize UWorldPartition::ResolveSubobject to go through the newly added actor descriptors map from 15894197.
Also added proper support for soft references to actors subobjects (such as components).

diff --git a/Engine/Source/Runtime/Engine/Private/WorldPartition/WorldPartition.cpp b/Engine/Source/Runtime/Engine/Private/WorldPartition/WorldPartition.cpp
--- a/Engine/Source/Runtime/Engine/Private/WorldPartition/WorldPartition.cpp
+++ b/Engine/Source/Runtime/Engine/Private/WorldPartition/WorldPartition.cpp
@@ -926,20 +926,23 @@ bool UWorldPartition::ResolveSubobject(const TCHAR* SubObjectPath, UObject*& Out
 #if WITH_EDITOR
 		else
 		{
-			for (UActorDescContainer::TIterator<> ActorDescIterator(this); ActorDescIterator; ++ActorDescIterator)
+			// Support for subobjects such as Actor.Component
+			FString SubObjectName;
+			FString SubObjectContext;	
+			if (!FString(SubObjectPath).Split(TEXT("."), &SubObjectContext, &SubObjectName))
 			{
-				FWorldPartitionActorDesc* ActorDesc = *ActorDescIterator;
+				SubObjectName = SubObjectPath;
+			}
 
-				if (FString(ActorDesc->ActorPath.ToString()).EndsWith(SubObjectPath))
+			if (const FWorldPartitionActorDesc* ActorDesc = GetActorDesc(SubObjectName))
+			{
+				if (bLoadIfExists)
 				{
-					if (bLoadIfExists)
-					{
-						LoadedSubobjects.Emplace(this, ActorDesc->GetGuid());
-					}
-
-					OutObject = ActorDescIterator->GetActor();
-					return true;
+					LoadedSubobjects.Emplace(this, ActorDesc->GetGuid());
 				}
+
+				OutObject = StaticFindObject(UObject::StaticClass(), GetWorld()->PersistentLevel, SubObjectPath);
+				return true;
 			}
 		}
 #endif
@@ -1180,4 +1183,4 @@ FBox UWorldPartition::GetEditorWorldBounds() const
 }
 #endif
 
-#undef LOCTEXT_NAMESPACE
\ No newline at end of file
+#undef LOCTEXT_NAMESPACE
