Fix for garbage deref in Wave Cue parsing when the asset has no cues.

diff --git a/Engine/Source/Runtime/Engine/Private/Audio.cpp b/Engine/Source/Runtime/Engine/Private/Audio.cpp
--- a/Engine/Source/Runtime/Engine/Private/Audio.cpp
+++ b/Engine/Source/Runtime/Engine/Private/Audio.cpp
@@ -1388,7 +1388,7 @@ bool FWaveModInfo::ReadWaveInfo( const uint8* WaveData, int32 WaveDataSize, FStr
 	RiffChunk = RiffChunkStart;
 
 	// Look for the 'data' chunk.
-	while ((((uint8*)RiffChunk + 8 ) <= WaveDataEnd ) && (INTEL_ORDER32(RiffChunk->ChunkID) != UE_mmioFOURCC('d','a','t','a')))
+	while ((((uint8*)RiffChunk + 8) < WaveDataEnd) && (INTEL_ORDER32(RiffChunk->ChunkID) != UE_mmioFOURCC('d','a','t','a')))
 	{
 		RiffChunk = (FRiffChunkOld*)((uint8*)RiffChunk + Pad16Bit(INTEL_ORDER32(RiffChunk->ChunkLen)) + 8);
 	}
@@ -1496,7 +1496,7 @@ bool FWaveModInfo::ReadWaveInfo( const uint8* WaveData, int32 WaveDataSize, FStr
 	RiffChunk = RiffChunkStart;
 
 	// Look for the 'cue ' chunk.
-	while ((((uint8*)RiffChunk + 8) <= WaveDataEnd) && (INTEL_ORDER32(RiffChunk->ChunkID) != UE_mmioFOURCC('c', 'u', 'e', ' ')))
+	while ((((uint8*)RiffChunk + 8) < WaveDataEnd) && (INTEL_ORDER32(RiffChunk->ChunkID) != UE_mmioFOURCC('c', 'u', 'e', ' ')))
 	{
 		RiffChunk = (FRiffChunkOld*)((uint8*)RiffChunk + Pad16Bit(INTEL_ORDER32(RiffChunk->ChunkLen)) + 8);
 	}
@@ -1523,7 +1523,7 @@ bool FWaveModInfo::ReadWaveInfo( const uint8* WaveData, int32 WaveDataSize, FStr
 		RiffChunk = RiffChunkStart;
 
 		// Look for the 'list' chunk.
-		while ((((uint8*)RiffChunk + 8) <= WaveDataEnd) && (INTEL_ORDER32(RiffChunk->ChunkID) != UE_mmioFOURCC('L', 'I', 'S', 'T')))
+		while ((((uint8*)RiffChunk + 8) < WaveDataEnd) && (INTEL_ORDER32(RiffChunk->ChunkID) != UE_mmioFOURCC('L', 'I', 'S', 'T')))
 		{
 			RiffChunk = (FRiffChunkOld*)((uint8*)RiffChunk + Pad16Bit(INTEL_ORDER32(RiffChunk->ChunkLen)) + 8);
 		}
