Integrates changes to HLODActor.cpp from EA

diff --git a/Engine/Source/Runtime/Engine/Private/WorldPartition/HLOD/HLODActor.cpp b/Engine/Source/Runtime/Engine/Private/WorldPartition/HLOD/HLODActor.cpp
--- a/Engine/Source/Runtime/Engine/Private/WorldPartition/HLOD/HLODActor.cpp
+++ b/Engine/Source/Runtime/Engine/Private/WorldPartition/HLOD/HLODActor.cpp
@@ -78,7 +78,7 @@ void AWorldPartitionHLOD::PostRegisterAllComponents()
 
 EActorGridPlacement AWorldPartitionHLOD::GetDefaultGridPlacement() const
 {
-	return EActorGridPlacement::Location;
+	return EActorGridPlacement::Bounds;
 }
 
 TUniquePtr<FWorldPartitionActorDesc> AWorldPartitionHLOD::CreateClassActorDesc() const
@@ -147,11 +147,7 @@ void AWorldPartitionHLOD::SetHLODBounds(const FBox& InBounds)
 
 void AWorldPartitionHLOD::GetActorBounds(bool bOnlyCollidingComponents, FVector& Origin, FVector& BoxExtent, bool bIncludeFromChildActors) const
 {
-	Super::GetActorBounds(bOnlyCollidingComponents, Origin, BoxExtent, bIncludeFromChildActors);
-
-	FBox Bounds = FBox(Origin - BoxExtent, Origin + BoxExtent);
-	Bounds += HLODBounds;
-	Bounds.GetCenterAndExtents(Origin, BoxExtent);
+	HLODBounds.GetCenterAndExtents(Origin, BoxExtent);
 }
 
 void AWorldPartitionHLOD::GetActorLocationBounds(bool bOnlyCollidingComponents, FVector& Origin, FVector& BoxExtent, bool bIncludeFromChildActors) const
