Temp fix: disable async staticmesh compilation when running the WrangleContent commandlet, as it was crashing 100% while collecting garbage during my tests.

diff --git a/Engine/Source/Runtime/Engine/Private/StaticMeshCompiler.cpp b/Engine/Source/Runtime/Engine/Private/StaticMeshCompiler.cpp
--- a/Engine/Source/Runtime/Engine/Private/StaticMeshCompiler.cpp
+++ b/Engine/Source/Runtime/Engine/Private/StaticMeshCompiler.cpp
@@ -157,6 +157,11 @@ bool FStaticMeshCompilingManager::IsAsyncStaticMeshCompilationEnabled() const
 		return false;
 	}
 
+	if (GetRunningCommandletClass() && (GetRunningCommandletClass()->GetName() == TEXT("WrangleContentCommandlet")))
+	{
+		return false;
+	}
+
 	StaticMeshCompilingManagerImpl::EnsureInitializedCVars();
 
 	return CVarAsyncStaticMeshStandard.AsyncCompilation.GetValueOnAnyThread() != 0;
