Added explicitly named GPU event scopes for GPUSkinCache TLAS builds and updates.

diff --git a/Engine/Source/Runtime/Engine/Private/GPUSkinCache.cpp b/Engine/Source/Runtime/Engine/Private/GPUSkinCache.cpp
--- a/Engine/Source/Runtime/Engine/Private/GPUSkinCache.cpp
+++ b/Engine/Source/Runtime/Engine/Private/GPUSkinCache.cpp
@@ -807,8 +807,6 @@ void FGPUSkinCache::CommitRayTracingGeometryUpdates(FRHICommandListImmediate& RH
 {
 	TRACE_CPUPROFILER_EVENT_SCOPE(FGPUSkinCache::CommitRayTracingGeometryUpdates);
 
-	SCOPED_DRAW_EVENT(RHICmdList, CommitSkeletalRayTracingGeometryUpdates);	
-
 	if (RayTracingGeometriesToUpdate.Num())
 	{
 		// If we have more deferred deleted data than set limit then force flush to make sure all pending releases have actually been freed
@@ -822,7 +820,29 @@ void FGPUSkinCache::CommitRayTracingGeometryUpdates(FRHICommandListImmediate& RH
 		// Track the amount of primitives which need to be build/updated in a single batch
 		uint64 PrimitivesToUpdates = 0;
 		TArray<FRayTracingGeometryBuildParams> BatchedBuildParams;
+		TArray<FRayTracingGeometryBuildParams> BatchedUpdateParams;
+
 		BatchedBuildParams.Reserve(RayTracingGeometriesToUpdate.Num());
+		BatchedUpdateParams.Reserve(RayTracingGeometriesToUpdate.Num());
+
+		auto KickBatch = [&RHICmdList, &BatchedBuildParams, &BatchedUpdateParams, &PrimitivesToUpdates]()
+		{
+			if (BatchedBuildParams.Num())
+			{
+				SCOPED_DRAW_EVENT(RHICmdList, BuildGPUSkinCacheBLAS);
+				RHICmdList.BuildAccelerationStructures(BatchedBuildParams);
+				BatchedBuildParams.Empty(BatchedBuildParams.Max());
+			}
+
+			if (BatchedUpdateParams.Num())
+			{
+				SCOPED_DRAW_EVENT(RHICmdList, UpdateGPUSkinCacheBLAS);
+				RHICmdList.BuildAccelerationStructures(BatchedUpdateParams);
+				BatchedUpdateParams.Empty(BatchedUpdateParams.Max());
+			}
+
+			PrimitivesToUpdates = 0;
+		};
 
 		// Iterate all the geometries which need an update
 		for (TMap<FRayTracingGeometry*, EAccelerationStructureBuildMode>::TRangedForIterator Iter = RayTracingGeometriesToUpdate.begin(); Iter != RayTracingGeometriesToUpdate.end(); ++Iter)
@@ -838,25 +858,25 @@ void FGPUSkinCache::CommitRayTracingGeometryUpdates(FRHICommandListImmediate& RH
 			if (BuildParams.BuildMode == EAccelerationStructureBuildMode::Build)
 			{
 				PrimitiveCount *= 10;
+				BatchedBuildParams.Add(BuildParams);
+			}
+			else
+			{
+				BatchedUpdateParams.Add(BuildParams);
 			}
-			PrimitivesToUpdates += PrimitiveCount;
 
-			// Add to batch
-			BatchedBuildParams.Add(BuildParams);
+			PrimitivesToUpdates += PrimitiveCount;
 
 			// Flush batch when limit is reached
 			if (GMaxRayTracingPrimitivesPerCmdList > 0 && PrimitivesToUpdates >= GMaxRayTracingPrimitivesPerCmdList)
 			{
-				RHICmdList.BuildAccelerationStructures(BatchedBuildParams);
+				KickBatch();
 				RHICmdList.SubmitCommandsHint();
-
-				BatchedBuildParams.Empty(BatchedBuildParams.Max());
-				PrimitivesToUpdates = 0;
 			}
 		}
 
 		// Enqueue the last batch
-		RHICmdList.BuildAccelerationStructures(BatchedBuildParams);
+		KickBatch();
 
 		// Clear working data
 		RayTracingGeometriesToUpdate.Empty(RayTracingGeometriesToUpdate.Num());
