Removing StereoBleed

UE-112707

diff --git a/Engine/Source/Runtime/Engine/Private/Audio.cpp b/Engine/Source/Runtime/Engine/Private/Audio.cpp
--- a/Engine/Source/Runtime/Engine/Private/Audio.cpp
+++ b/Engine/Source/Runtime/Engine/Private/Audio.cpp
@@ -351,11 +351,6 @@ bool FSoundSource::SetReverbApplied(bool bHardwareAvailable)
 	return(bReverbApplied);
 }
 
-float FSoundSource::SetStereoBleed()
-{
-	return 0.f;
-}
-
 float FSoundSource::SetLFEBleed()
 {
 	LFEBleed = WaveInstance->LFEBleed;
diff --git a/Engine/Source/Runtime/Engine/Public/Audio.h b/Engine/Source/Runtime/Engine/Public/Audio.h
--- a/Engine/Source/Runtime/Engine/Public/Audio.h
+++ b/Engine/Source/Runtime/Engine/Public/Audio.h
@@ -228,10 +228,6 @@ public:
 	/** The volume at which the radio filter kicks in */
 	float RadioFilterVolumeThreshold;
 
-	/** The amount of stereo sounds to bleed to the rear speakers */
-	UE_DEPRECATED(4.25, "Stereo Bleed is no longer supported.")
-	float StereoBleed;
-
 	/** The amount of a sound to bleed to the LFE channel */
 	float LFEBleed;
 
@@ -642,10 +638,6 @@ public:
 	/** Set the bReverbApplied variable. */
 	ENGINE_API bool SetReverbApplied(bool bHardwareAvailable);
 
-	/** Set the StereoBleed variable. */
-	UE_DEPRECATED(4.25, "Stereo Bleed is no longer supported.")
-	ENGINE_API float SetStereoBleed();
-
 	/** Updates and sets the LFEBleed variable. */
 	ENGINE_API float SetLFEBleed();
 
