Packed Level Instance: Support new 'bConsiderForActorPlacementWhenHidden' in the ISMComponentDescriptor

diff --git a/Engine/Source/Runtime/Engine/Private/ISMPartition/ISMComponentDescriptor.cpp b/Engine/Source/Runtime/Engine/Private/ISMPartition/ISMComponentDescriptor.cpp
--- a/Engine/Source/Runtime/Engine/Private/ISMPartition/ISMComponentDescriptor.cpp
+++ b/Engine/Source/Runtime/Engine/Private/ISMPartition/ISMComponentDescriptor.cpp
@@ -65,6 +65,7 @@ void FISMComponentDescriptor::InitFrom(const UStaticMeshComponent* Template, boo
 	bIsEditorOnly = Template->bIsEditorOnly;
 	bVisible = Template->GetVisibleFlag();
 	bVisibleInRayTracing = Template->bVisibleInRayTracing;
+	bConsiderForActorPlacementWhenHidden = Template->bConsiderForActorPlacementWhenHidden;
 
 	if (const UInstancedStaticMeshComponent* ISMTemplate = Cast<UInstancedStaticMeshComponent>(Template))
 	{
@@ -127,6 +128,7 @@ bool FISMComponentDescriptor::operator==(const FISMComponentDescriptor& Other) c
 	bIsEditorOnly == Other.bIsEditorOnly &&
 	bVisible == Other.bVisible &&
 	bVisibleInRayTracing == Other.bVisibleInRayTracing &&
+	bConsiderForActorPlacementWhenHidden == Other.bConsiderForActorPlacementWhenHidden &&
 	BodyInstance.GetCollisionEnabled() == Other.BodyInstance.GetCollisionEnabled() && 
 	BodyInstance.GetCollisionResponse() == Other.BodyInstance.GetCollisionResponse() &&
 	BodyInstance.DoesUseCollisionProfile() == Other.BodyInstance.DoesUseCollisionProfile() &&
@@ -189,6 +191,7 @@ void FISMComponentDescriptor::InitComponent(UInstancedStaticMeshComponent* ISMCo
 	ISMComponent->bIsEditorOnly = bIsEditorOnly;
 	ISMComponent->SetVisibleFlag(bVisible);
 	ISMComponent->bVisibleInRayTracing = bVisibleInRayTracing;
+	ISMComponent->bConsiderForActorPlacementWhenHidden = bConsiderForActorPlacementWhenHidden;
 
 	// HISM Specific
 	if (UHierarchicalInstancedStaticMeshComponent* HISMComponent = Cast<UHierarchicalInstancedStaticMeshComponent>(ISMComponent))
diff --git a/Engine/Source/Runtime/Engine/Public/ISMPartition/ISMComponentDescriptor.h b/Engine/Source/Runtime/Engine/Public/ISMPartition/ISMComponentDescriptor.h
--- a/Engine/Source/Runtime/Engine/Public/ISMPartition/ISMComponentDescriptor.h
+++ b/Engine/Source/Runtime/Engine/Public/ISMPartition/ISMComponentDescriptor.h
@@ -162,5 +162,8 @@ public:
 
 	UPROPERTY()
 	uint8 bVisible : 1;
+
+	UPROPERTY()
+	uint8 bConsiderForActorPlacementWhenHidden : 1;
 #endif
 };
