WorldPartition: Add MapCheck for LevelScriptBlueprint references to Actors that are part of one or more DataLayer

diff --git a/Engine/Source/Runtime/Engine/Private/WorldPartition/WorldPartitionRuntimeHash.cpp b/Engine/Source/Runtime/Engine/Private/WorldPartition/WorldPartitionRuntimeHash.cpp
--- a/Engine/Source/Runtime/Engine/Private/WorldPartition/WorldPartitionRuntimeHash.cpp
+++ b/Engine/Source/Runtime/Engine/Private/WorldPartition/WorldPartitionRuntimeHash.cpp
@@ -236,6 +236,28 @@ void UWorldPartitionRuntimeHash::CheckForErrors() const
 			}
 		}
 	}
+
+	// Check Level Script Blueprint
+	if (ULevelScriptBlueprint* LevelScriptBlueprint = GetOuterUWorldPartition()->GetWorld()->PersistentLevel->GetLevelScriptBlueprint(true))
+	{
+		TArray<AActor*> LevelScriptExternalActorReferences = ActorsReferencesUtils::GetExternalActorReferences(LevelScriptBlueprint);
+
+		for (AActor* Actor : LevelScriptExternalActorReferences)
+		{
+			if (FWorldPartitionActorDescView* ActorDescView = ActorDescList.Find(Actor->GetActorGuid()))
+			{
+				TArray<FName> ActorDescLayers = ActorDescView->GetDataLayers();
+				if (ActorDescLayers.Num())
+				{
+					TSharedRef<FTokenizedMessage> Error = FMessageLog("MapCheck").Error()
+						->AddToken(FTextToken::Create(LOCTEXT("MapCheck_WorldPartition_LevelScriptBlueprintActorReference", "Level Script Blueprint references actor")))
+						->AddToken(FAssetNameToken::Create(GetActorLabel(*ActorDescView)))
+						->AddToken(FTextToken::Create(LOCTEXT("MapCheck_WorldPartition_LevelScriptBlueprintDataLayerReference", "with a non empty set of data layers")))
+						->AddToken(FMapErrorToken::Create(FName(TEXT("WorldPartition_LevelScriptBlueprintRefefenceDataLayer_CheckForErrors"))));
+				}
+			}
+		}
+	}
 }
 #endif
 
