Level Instance: Fix invalid filename and rename AssetPath to LevelPackageName

diff --git a/Engine/Source/Runtime/Engine/Private/LevelInstance/LevelInstanceSubsystem.cpp b/Engine/Source/Runtime/Engine/Private/LevelInstance/LevelInstanceSubsystem.cpp
--- a/Engine/Source/Runtime/Engine/Private/LevelInstance/LevelInstanceSubsystem.cpp
+++ b/Engine/Source/Runtime/Engine/Private/LevelInstance/LevelInstanceSubsystem.cpp
@@ -729,7 +729,13 @@ ALevelInstance* ULevelInstanceSubsystem::CreateLevelInstanceFrom(const TArray<AA
 		}
 	}
 		
-	ULevelStreamingLevelInstanceEditor* LevelStreaming = StaticCast<ULevelStreamingLevelInstanceEditor*>(EditorLevelUtils::CreateNewStreamingLevelForWorld(*GetWorld(), ULevelStreamingLevelInstanceEditor::StaticClass(), CreationParams.AssetPath, false, CreationParams.TemplateWorld));
+	FString LevelFilename;
+	if (!CreationParams.LevelPackageName.IsEmpty())
+	{
+		LevelFilename = FPackageName::LongPackageNameToFilename(CreationParams.LevelPackageName, FPackageName::GetMapPackageExtension());
+	}
+
+	ULevelStreamingLevelInstanceEditor* LevelStreaming = StaticCast<ULevelStreamingLevelInstanceEditor*>(EditorLevelUtils::CreateNewStreamingLevelForWorld(*GetWorld(), ULevelStreamingLevelInstanceEditor::StaticClass(), LevelFilename, false, CreationParams.TemplateWorld));
 	if (!LevelStreaming)
 	{
 		UE_LOG(LogLevelInstance, Warning, TEXT("Failed to create new Level Instance level"));
@@ -787,7 +793,7 @@ ALevelInstance* ULevelInstanceSubsystem::CreateLevelInstanceFrom(const TArray<AA
 		const bool bCompile = true;
 
 		UBlueprint* NewBP = nullptr;
-		if (CreationParams.AssetPath.IsEmpty())
+		if (CreationParams.LevelPackageName.IsEmpty())
 		{
 			NewBP = FPackedLevelInstanceBuilder::CreatePackedLevelInstanceBlueprintWithDialog(TSoftObjectPtr<UBlueprint>(BPAssetPath), WorldPtr, bCompile);
 		}
diff --git a/Engine/Source/Runtime/Engine/Public/LevelInstance/LevelInstanceTypes.h b/Engine/Source/Runtime/Engine/Public/LevelInstance/LevelInstanceTypes.h
--- a/Engine/Source/Runtime/Engine/Public/LevelInstance/LevelInstanceTypes.h
+++ b/Engine/Source/Runtime/Engine/Public/LevelInstance/LevelInstanceTypes.h
@@ -48,7 +48,7 @@ struct FNewLevelInstanceParams
 	TObjectPtr<UWorld> TemplateWorld = nullptr;
 		
 	UPROPERTY()
-	FString AssetPath = TEXT("");
+	FString LevelPackageName = TEXT("");
 
 	UPROPERTY()
 	bool bPromptForSave = false;
