HLODProxy: Fix infinite loop in PreSave when an ALODActor has a null ProxyDesc.
#rnx

diff --git a/Engine/Source/Runtime/Engine/Private/HLOD/HLODProxyDesc.cpp b/Engine/Source/Runtime/Engine/Private/HLOD/HLODProxyDesc.cpp
--- a/Engine/Source/Runtime/Engine/Private/HLOD/HLODProxyDesc.cpp
+++ b/Engine/Source/Runtime/Engine/Private/HLOD/HLODProxyDesc.cpp
@@ -314,17 +314,20 @@ ALODActor* UHLODProxyDesc::SpawnLODActor(ULevel* InLevel) const
 	{
 		if (ALODActor* SubLODActor = Cast<ALODActor>(Actor))
 		{
-			if (SubHLODDescs.Contains(SubLODActor->ProxyDesc))
+			if (SubLODActor->ProxyDesc && SubHLODDescs.Contains(SubLODActor->ProxyDesc))
 			{
+				check(SubLODActor != LODActor);
 				SubActorsToAdd.Add(SubLODActor);
 			}
 		}
 	}
 
 	// Find all subactors from the level
-	Algo::Transform(SubActors, SubActorsToAdd, [InLevel](const FName& ActorName)
+	Algo::Transform(SubActors, SubActorsToAdd, [InLevel, LODActor](const FName& ActorName)
 	{
-		return FindObjectFast<AActor>(InLevel, ActorName);
+		AActor* Actor = FindObjectFast<AActor>(InLevel, ActorName);
+		check(Actor != LODActor);
+		return Actor;
 	});
 
 	// Remove null entries
diff --git a/Engine/Source/Runtime/Engine/Private/LODActor.cpp b/Engine/Source/Runtime/Engine/Private/LODActor.cpp
--- a/Engine/Source/Runtime/Engine/Private/LODActor.cpp
+++ b/Engine/Source/Runtime/Engine/Private/LODActor.cpp
@@ -818,6 +818,7 @@ void ALODActor::AddSubActors(const TArray<AActor*>& InActors)
 
 	for(AActor* Actor : InActors)
 	{
+		check(Actor != this);
 		UStaticMeshComponent* LODComponent = GetOrCreateLODComponentForActor(Actor);
 		Actor->SetLODParent(LODComponent, LODDrawDistanceWithOverride);
 
