Crash when saving HLOD data on unsaved level

Integrated fixes from 4.26.2

diff --git a/Engine/Source/Runtime/Engine/Private/World.cpp b/Engine/Source/Runtime/Engine/Private/World.cpp
--- a/Engine/Source/Runtime/Engine/Private/World.cpp
+++ b/Engine/Source/Runtime/Engine/Private/World.cpp
@@ -574,8 +574,10 @@ bool UWorld::Rename(const TCHAR* InName, UObject* NewOuter, ERenameFlags Flags)
 		{
 			if (ALODActor* LODActor = Cast<ALODActor>(Actor))
 			{
-				UHLODProxy* HLODProxy = LODActor->GetProxy();
-				OldHLODPackages[LODActor->LODLevel - 1] = HLODProxy->GetPackage();
+				if (UHLODProxy* HLODProxy = LODActor->GetProxy())
+				{
+					OldHLODPackages[LODActor->LODLevel - 1] = HLODProxy->GetPackage();
+				}
 			}
 		}
 	}
