Make newly imported asset quality 100 and seekable streaming

diff --git a/Engine/Source/Editor/AudioEditor/Private/Factories/SoundFactory.cpp b/Engine/Source/Editor/AudioEditor/Private/Factories/SoundFactory.cpp
--- a/Engine/Source/Editor/AudioEditor/Private/Factories/SoundFactory.cpp
+++ b/Engine/Source/Editor/AudioEditor/Private/Factories/SoundFactory.cpp
@@ -352,6 +352,10 @@ UObject* USoundFactory::CreateObject
 			Sound->ConcurrencySet = TemplateSoundWave->ConcurrencySet;
 		}
 
+		Sound->CompressionQuality = 100;
+		Sound->bStreaming = true;
+		Sound->bSeekableStreaming = true;
+
 		// If we're a multi-channel file, we're going to spoof the behavior of the SoundSurroundFactory
 		int32 ChannelCount = (int32)*WaveInfo.pChannels;
 		check(ChannelCount >0);
diff --git a/Engine/Source/Runtime/Engine/Private/ADPCMAudioInfo.cpp b/Engine/Source/Runtime/Engine/Private/ADPCMAudioInfo.cpp
--- a/Engine/Source/Runtime/Engine/Private/ADPCMAudioInfo.cpp
+++ b/Engine/Source/Runtime/Engine/Private/ADPCMAudioInfo.cpp
@@ -203,24 +203,19 @@ void FADPCMAudioInfo::SeekToTimeInternal(const float InSeekTime)
 		{
 			const int32 ChannelBlockSize = sizeof(int16) * NumChannels;
 
-			// 1. Find total offset
-			CurrentChunkBufferOffset = HeaderOffset + (TotalSamplesStreamed * ChannelBlockSize);
-			
-			// 2. Calculate index and remove from offset
-			while (CurrentChunkBufferOffset >= StreamingSoundWave->GetSizeOfChunk(CurrentChunkIndex))
+			uint32 TotalByteSizeToSeek = TotalSamplesStreamed * ChannelBlockSize;		
+			uint32 SizeOfCurrentChunk = StreamingSoundWave->GetSizeOfChunk(CurrentChunkIndex);
+			uint32 CurrentBytes = 0;
+			while (CurrentBytes + SizeOfCurrentChunk < TotalByteSizeToSeek)
 			{
-				CurrentChunkBufferOffset -= StreamingSoundWave->GetSizeOfChunk(CurrentChunkIndex);
 				CurrentChunkIndex++;
-				
-				if (CurrentChunkIndex >= TotalStreamingChunks)
-				{
-					CurrentChunkIndex = FirstChunkSampleDataIndex;
-					CurrentChunkBufferOffset = FirstChunkSampleDataOffset;
-					break;
-				}
+				CurrentBytes += SizeOfCurrentChunk;
+				CurrentChunkBufferOffset -= SizeOfCurrentChunk;
+				SizeOfCurrentChunk = StreamingSoundWave->GetSizeOfChunk(CurrentChunkIndex);
 			}
 
-			// 3. Trim remainder of block size, effectively aligning the block to a channel pair boundary
+			check(CurrentBytes <= TotalByteSizeToSeek);
+			CurrentChunkBufferOffset = TotalByteSizeToSeek - CurrentBytes;
 			CurrentChunkBufferOffset -= CurrentChunkBufferOffset % ChannelBlockSize;
 		}
 		else
