Add UI for controlling the new shadow system in EA
Enable non-nanite VSMs by default in new projects in UE5

Submitted on behalf of Ola.Olsson.

diff --git a/Engine/Source/Editor/GameProjectGeneration/Private/GameProjectUtils.cpp b/Engine/Source/Editor/GameProjectGeneration/Private/GameProjectUtils.cpp
--- a/Engine/Source/Editor/GameProjectGeneration/Private/GameProjectUtils.cpp
+++ b/Engine/Source/Editor/GameProjectGeneration/Private/GameProjectUtils.cpp
@@ -225,6 +225,23 @@ namespace
 			true /* ShouldReplaceExistingValue */);
 	}
 
+	void AddNewProjectDefaultShadowConfigValues(const FProjectInformation& InProjectInfo, TArray<FTemplateConfigValue>& ConfigValues)
+	{
+		// Enable support for virtual shadow maps by default for new projects
+		ConfigValues.Emplace(TEXT("DefaultEngine.ini"),
+			TEXT("/Script/Engine.RendererSettings"),
+			TEXT("r.Shadow.Virtual.Enable"),
+			TEXT("1"),
+			true /* ShouldReplaceExistingValue */);
+
+		// and also support for non-nanite VSM, this has an impact on all shadow rendering, which is why this is off for existing projects by default.
+		ConfigValues.Emplace(TEXT("DefaultEngine.ini"),
+			TEXT("/Script/Engine.RendererSettings"),
+			TEXT("r.Shadow.Virtual.NonNaniteVSM"),
+			TEXT("True"),
+			true /* ShouldReplaceExistingValue */);
+	}
+
 	/** Get the configuration values for raytracing if enabled. */
 	void AddRaytracingConfigValues(const FProjectInformation& InProjectInfo, TArray<FTemplateConfigValue>& ConfigValues)
 	{
@@ -1816,7 +1833,7 @@ TOptional<FGuid> GameProjectUtils::CreateProjectFromTemplate(const FProjectInfor
 
 	AddLumenConfigValues(InProjectInfo, ConfigValuesToSet);
 	AddRaytracingConfigValues(InProjectInfo, ConfigValuesToSet);
-
+	AddNewProjectDefaultShadowConfigValues(InProjectInfo, ConfigValuesToSet);
 	AddDefaultMapConfigValues(InProjectInfo, ConfigValuesToSet);
 	
 	TemplateDefs->AddConfigValues(ConfigValuesToSet, TemplateName, ProjectName, InProjectInfo.bShouldGenerateCode);
@@ -2153,6 +2170,7 @@ bool GameProjectUtils::GenerateConfigFiles(const FProjectInformation& InProjectI
 		TArray<FTemplateConfigValue> ConfigValuesToSet;
 		AddHardwareConfigValues(InProjectInfo, ConfigValuesToSet);
 		AddLumenConfigValues(InProjectInfo, ConfigValuesToSet);
+		AddNewProjectDefaultShadowConfigValues(InProjectInfo, ConfigValuesToSet);
 		AddRaytracingConfigValues(InProjectInfo, ConfigValuesToSet);
 
 		if (!SaveConfigValues(InProjectInfo, ConfigValuesToSet, OutFailReason))
diff --git a/Engine/Source/Runtime/Engine/Classes/Engine/EngineTypes.h b/Engine/Source/Runtime/Engine/Classes/Engine/EngineTypes.h
--- a/Engine/Source/Runtime/Engine/Classes/Engine/EngineTypes.h
+++ b/Engine/Source/Runtime/Engine/Classes/Engine/EngineTypes.h
@@ -405,6 +405,23 @@ namespace EReflectionMethod
 	};
 }
 
+
+
+// Note: Maps to r.Shadow.Virtual.Enable AND r.Shadow.Virtual.NonNaniteVSM, where Standard == 0 == off, this enum is used in URendererSettings for readability
+UENUM()
+namespace EDynamicShadowingMethod
+{
+	enum Type
+	{
+		/** UE4 shadowing system only. Required to enable stationary baked shadows (but which is incompatible with Nanite geometry) */
+		Standard,
+
+		/** High-quality shadows for next-gen projects with simplified setup. This system is in development and thus has a number of performance pitfalls. */
+		VirtualShadowMaps UMETA(DisplayName = "Virtual Shadow Maps (Beta)"),
+	};
+}
+
+
 /** Specifies which component of the scene rendering should be output to the final render target. */
 UENUM()
 enum ESceneCaptureSource 
diff --git a/Engine/Source/Runtime/Engine/Classes/Engine/RendererSettings.h b/Engine/Source/Runtime/Engine/Classes/Engine/RendererSettings.h
--- a/Engine/Source/Runtime/Engine/Classes/Engine/RendererSettings.h
+++ b/Engine/Source/Runtime/Engine/Classes/Engine/RendererSettings.h
@@ -362,9 +362,23 @@ class ENGINE_API URendererSettings : public UDeveloperSettings
 		ToolTip="Controls which tracing method Lumen uses when using Software Ray Tracing."))
 	TEnumAsByte<ELumenSoftwareTracingMode::Type> LumenSoftwareTracingMode;
 
+	UPROPERTY(config, EditAnywhere, Category = Shadows, meta = (
+		ConsoleVariable = "r.Shadow.Virtual.Enable", DisplayName = "Dynamic Shadowing Method",
+		ToolTip = "Dynamic Shadowing Method",
+		ConfigRestartRequired = true))
+	TEnumAsByte<EDynamicShadowingMethod::Type> DynamicShadowingMethod;
+
+	// Note: this is slaved to the DynamicShadowingMethod, which is why it is not editable, the property is only here to let us propagate to the ini file.
+	UPROPERTY(config, EditAnywhere, Category = Shadows, meta = (
+		EditCondition = "false",
+		ReadOnly,
+		ConsoleVariable="r.Shadow.Virtual.NonNaniteVSM",
+		ConfigRestartRequired = true))
+	uint32 bEnableNonNaniteVSM:1;
+
 	/**
-	"Ray Tracing settings."
-	*/
+	 * "Ray Tracing settings."
+	 */
 	UPROPERTY(config, EditAnywhere, Category = HardwareRayTracing, meta = (
 		ConsoleVariable = "r.RayTracing", DisplayName = "Support Hardware Ray Tracing",
 		ToolTip = "Support Hardware Ray Tracing features.  Requires 'Support Compute Skincache' before project is allowed to set this.",
diff --git a/Engine/Source/Runtime/Engine/Private/Components/StaticMeshComponent.cpp b/Engine/Source/Runtime/Engine/Private/Components/StaticMeshComponent.cpp
--- a/Engine/Source/Runtime/Engine/Private/Components/StaticMeshComponent.cpp
+++ b/Engine/Source/Runtime/Engine/Private/Components/StaticMeshComponent.cpp
@@ -411,6 +411,19 @@ void UStaticMeshComponent::CheckForErrors()
 		OwnerName = Owner->GetName();
 	}
 
+	if (GetStaticMesh() != nullptr && GetStaticMesh()->NaniteSettings.bEnabled != 0)
+	{
+		static const auto CVar = IConsoleManager::Get().FindConsoleVariable(TEXT("r.Shadow.Virtual.Enable"));
+		if (CVar->GetInt() == 0)
+		{
+			FFormatNamedArguments Arguments;
+			Arguments.Add(TEXT("MeshName"), FText::FromString(GetStaticMesh()->GetName()));
+			FMessageLog("MapCheck").Warning()
+				->AddToken(FUObjectToken::Create(Owner))
+				->AddToken(FTextToken::Create(FText::Format(LOCTEXT("MapCheck_Message_NaniteNoVSM", "Static mesh '{MeshName}' uses Nanite but Virtual Shadow Maps are not enabled in the project settings. Nanite geometry does not support stationary light shadows, and may yield poor visual quality and reduced performance. Nanite geometry works best with virtual shadow maps enabled. See release notes."), Arguments)));
+		}
+	}
+
 	// Make sure any simplified meshes can still find their high res source mesh
 	if( GetStaticMesh() != NULL && GetStaticMesh()->GetRenderData())
 	{
diff --git a/Engine/Source/Runtime/Engine/Private/RendererSettings.cpp b/Engine/Source/Runtime/Engine/Private/RendererSettings.cpp
--- a/Engine/Source/Runtime/Engine/Private/RendererSettings.cpp
+++ b/Engine/Source/Runtime/Engine/Private/RendererSettings.cpp
@@ -77,6 +77,8 @@ URendererSettings::URendererSettings(const FObjectInitializer& ObjectInitializer
 	bEnableRayTracingShadows = 0;
 	bEnableRayTracingTextureLOD = 0;
 	MaxSkinBones = FGPUBaseSkinVertexFactory::GHardwareMaxGPUSkinBones;
+
+	bEnableNonNaniteVSM = 0;
 }
 
 void URendererSettings::PostInitProperties()
@@ -193,6 +195,31 @@ void URendererSettings::PostEditChangeProperty(FPropertyChangedEvent& PropertyCh
 			}
 		}
 
+		if (PropertyChangedEvent.Property->GetFName() == GET_MEMBER_NAME_CHECKED(URendererSettings, DynamicShadowingMethod))
+		{
+			// Propagate the setting to the 'r.Shadow.Virtual.NonNaniteVSM' CVar
+			int32 Value = static_cast<int32>(DynamicShadowingMethod);
+			bEnableNonNaniteVSM = Value;
+
+			// Don't think we need this as the below should trigger a reload of INI?
+			//IConsoleVariable* CVar = IConsoleManager::Get().FindConsoleVariable(TEXT("r.Shadow.Virtual.NonNaniteVSM"));
+			//CVar->Set(Value, ECVF_SetByProjectSetting);
+			for (TFieldIterator<FProperty> PropIt(GetClass()); PropIt; ++PropIt)
+			{
+				FProperty* Property = *PropIt;
+				if (Property->GetFName() == GET_MEMBER_NAME_CHECKED(URendererSettings, bEnableNonNaniteVSM))
+				{
+					UpdateSinglePropertyInConfigFile(Property, GetDefaultConfigFilename());
+				}
+			}
+
+			if (Value != 0)
+			{
+				FMessageDialog::Open(EAppMsgType::Ok, LOCTEXT("Virtual Shadow Maps (Beta) Enabled", "Note that the virtual shadow map system is in beta and may have a performance impact, even if the run-time toggle (r.Shadow.Virtual.Enable) is used to turn them off."));
+				// TODO: put this in release notes? Note: When enabling the virtual shadow map system (beta), both 'r.Shadow.Virtual.Enable' and 'r.Shadow.Virtual.NonNaniteVSM' are turned on for the project. 'r.Shadow.Virtual.Enable' can be toggled at run-time, but there may still be some overhead as long as 'r.Shadow.Virtual.NonNaniteVSM' is turned on (requires restart)
+			}
+		}
+
 		if (PropertyChangedEvent.Property->GetFName() == GET_MEMBER_NAME_CHECKED(URendererSettings, VirtualTextureTileSize))
 		{
 			VirtualTextureTileSize = FMath::RoundUpToPowerOfTwo(VirtualTextureTileSize);
