Clean up some of the comments around possession and mark things as authority only
This is important for game feature plugin support

diff --git a/Engine/Source/Runtime/Engine/Classes/GameFramework/Pawn.h b/Engine/Source/Runtime/Engine/Classes/GameFramework/Pawn.h
--- a/Engine/Source/Runtime/Engine/Classes/GameFramework/Pawn.h
+++ b/Engine/Source/Runtime/Engine/Classes/GameFramework/Pawn.h
@@ -27,6 +27,8 @@ class UPrimitiveComponent;
 
 ENGINE_API DECLARE_LOG_CATEGORY_EXTERN(LogDamage, Warning, All);
 
+DECLARE_DYNAMIC_MULTICAST_SPARSE_DELEGATE_OneParam(FPawnRestartedSignature, APawn, ReceiveRestartedDelegate, APawn*, Pawn);
+
 /** 
  * Pawn is the base class of all actors that can be possessed by players or AI.
  * They are the physical representations of players and creatures in a level.
@@ -167,9 +169,6 @@ public:
 	/** Freeze pawn - stop sounds, animations, physics, weapon firing */
 	virtual void TurnOff();
 
-	/** Called when the Pawn is being restarted (usually by being possessed by a Controller). */
-	virtual void Restart();
-
 	/** Handle StartFire() passed from PlayerController */
 	virtual void PawnStartFire(uint8 FireModeNum = 0);
 
@@ -201,9 +200,9 @@ public:
 	UFUNCTION(BlueprintCallable, Category=Pawn)
 	bool IsControlled() const;
 
-	/** Check if this actor is currently being controlled at all (the actor has a valid Controller) */
+	/** Check if this actor is currently being controlled at all (the actor has a valid Controller, which will be false for remote clients) */
 	UFUNCTION(BlueprintCallable, Category = Pawn)
-	bool IsPawnControlled() const;
+	virtual bool IsPawnControlled() const;
 
 	/** Returns controller for this actor. */
 	UFUNCTION(BlueprintCallable, Category=Pawn)
@@ -309,19 +308,27 @@ public:
 	 */
 	virtual void PossessedBy(AController* NewController);
 
-	/** Event called when the Pawn is possessed by a Controller (normally only occurs on the server/standalone). */
-	UFUNCTION(BlueprintImplementableEvent, meta=(DisplayName= "Possessed"))
+	/** Event called when the Pawn is possessed by a Controller. Only called on the server (or in standalone) */
+	UFUNCTION(BlueprintImplementableEvent, BlueprintAuthorityOnly, meta=(DisplayName= "Possessed"))
 	void ReceivePossessed(AController* NewController);
 
-	/** Event called when the Pawn is no longer possessed by a Controller. */
-	UFUNCTION(BlueprintImplementableEvent, meta=(DisplayName= "Unpossessed"))
+	/** Event called when the Pawn is no longer possessed by a Controller. Only called on the server (or in standalone) */
+	UFUNCTION(BlueprintImplementableEvent, BlueprintAuthorityOnly, meta=(DisplayName= "Unpossessed"))
 	void ReceiveUnpossessed(AController* OldController);
 
+	/** Event called when a pawn has been restarted, usually by a possession change. This is called on the server for all pawns and the owning client for player pawns */
+	UFUNCTION(BlueprintImplementableEvent)
+	void ReceiveRestarted();
+
+	/** Event called when a pawn has been restarted, usually by a possession change. This is called on the server for all pawns and the owning client for player pawns */
+	UPROPERTY(BlueprintAssignable, Category = Pawn)
+	FPawnRestartedSignature ReceiveRestartedDelegate;
+
 	/** Returns true if controlled by a local (not network) Controller.	 */
 	UFUNCTION(BlueprintPure, Category=Pawn)
 	virtual bool IsLocallyControlled() const;
   
-	/** Returns true if controlled by a human player (possessed by a PlayerController).	 */
+	/** Returns true if controlled by a human player (possessed by a PlayerController).	This returns true for players controlled by remote clients */
 	UFUNCTION(BlueprintPure, Category=Pawn)
 	virtual bool IsPlayerControlled() const;
 
@@ -349,9 +356,21 @@ public:
 	/** Returns true if player is viewing this Pawn in FreeCam */
 	virtual bool InFreeCam() const;
 
-	/** Tell client that the Pawn is begin restarted. Calls Restart(). */
+	/** Called when the Pawn is being restarted (usually by being possessed by a Controller). This is called on the server for all pawns and the owning client for player pawns  */
+	virtual void Restart();
+
+	/** Called on the owning client of a player-controlled Pawn when it is restarted, this calls Restart() */
 	virtual void PawnClientRestart();
 
+	/** 
+	 * Notifies other systems that a pawn has been restarted. By default this is called on the server for all pawns and the owning client for player pawns.
+	 * This can be overridden by subclasses to delay the notification of restart until data has loaded/replicated
+	 */
+	virtual void NotifyRestarted();
+
+	/** Wrapper function to call correct restart functions, enable bCallClientRestart if this is a locally owned player pawn or equivalent */
+	void DispatchRestart(bool bCallClientRestart);
+
 	/** Updates Pawn's rotation to the given rotation, assumed to be the Controller's ControlRotation. Respects the bUseControllerRotation* settings. */
 	virtual void FaceRotation(FRotator NewControlRotation, float DeltaTime = 0.f);
 
diff --git a/Engine/Source/Runtime/Engine/Private/Controller.cpp b/Engine/Source/Runtime/Engine/Private/Controller.cpp
--- a/Engine/Source/Runtime/Engine/Private/Controller.cpp
+++ b/Engine/Source/Runtime/Engine/Private/Controller.cpp
@@ -362,7 +362,7 @@ void AController::OnPossess(APawn* InPawn)
 	// update rotation to match possessed pawn's rotation
 	SetControlRotation(Pawn->GetActorRotation());
 
-	Pawn->Restart();
+	Pawn->DispatchRestart(false);
 }
 
 void AController::UnPossess()
diff --git a/Engine/Source/Runtime/Engine/Private/DebugCameraController.cpp b/Engine/Source/Runtime/Engine/Private/DebugCameraController.cpp
--- a/Engine/Source/Runtime/Engine/Private/DebugCameraController.cpp
+++ b/Engine/Source/Runtime/Engine/Private/DebugCameraController.cpp
@@ -340,7 +340,7 @@ ASpectatorPawn* ADebugCameraController::SpawnSpectatorPawn()
 			if (SpawnedSpectator)
 			{
 				SpawnedSpectator->PossessedBy(this);
-				SpawnedSpectator->PawnClientRestart();
+				SpawnedSpectator->DispatchRestart(true);
 				if (SpawnedSpectator->PrimaryActorTick.bStartWithTickEnabled)
 				{
 					SpawnedSpectator->SetActorTickEnabled(true);
diff --git a/Engine/Source/Runtime/Engine/Private/Pawn.cpp b/Engine/Source/Runtime/Engine/Private/Pawn.cpp
--- a/Engine/Source/Runtime/Engine/Private/Pawn.cpp
+++ b/Engine/Source/Runtime/Engine/Private/Pawn.cpp
@@ -393,6 +393,27 @@ void APawn::PawnClientRestart()
 	}
 }
 
+void APawn::NotifyRestarted()
+{
+	ReceiveRestarted();
+	ReceiveRestartedDelegate.Broadcast(this);
+}
+
+void APawn::DispatchRestart(bool bCallClientRestart)
+{
+	if (bCallClientRestart)
+	{
+		// This calls Restart()
+		PawnClientRestart();
+	}
+	else
+	{
+		Restart();
+	}
+
+	NotifyRestarted();
+}
+
 void APawn::Destroyed()
 {
 	DetachFromControllerPendingDestroy();
diff --git a/Engine/Source/Runtime/Engine/Private/PlayerController.cpp b/Engine/Source/Runtime/Engine/Private/PlayerController.cpp
--- a/Engine/Source/Runtime/Engine/Private/PlayerController.cpp
+++ b/Engine/Source/Runtime/Engine/Private/PlayerController.cpp
@@ -772,7 +772,7 @@ void APlayerController::ClientRestart_Implementation(APawn* NewPawn)
 	AcknowledgePossession(GetPawn());
 
 	GetPawn()->Controller = this;
-	GetPawn()->PawnClientRestart();
+	GetPawn()->DispatchRestart(true);
 	
 	if (GetLocalRole() < ROLE_Authority)
 	{
@@ -829,7 +829,7 @@ void APlayerController::OnPossess(APawn* PawnToPossess)
 		// We're really just trying to avoid calling Restart() multiple times.
 		if (!IsLocalPlayerController())
 		{
-			GetPawn()->Restart();
+			GetPawn()->DispatchRestart(false);
 		}
 
 		ClientRestart(GetPawn());
@@ -4931,7 +4931,7 @@ ASpectatorPawn* APlayerController::SpawnSpectatorPawn()
 				{
 					SpawnedSpectator->SetReplicates(false); // Client-side only
 					SpawnedSpectator->PossessedBy(this);
-					SpawnedSpectator->PawnClientRestart();
+					SpawnedSpectator->DispatchRestart(true);
 					if (SpawnedSpectator->PrimaryActorTick.bStartWithTickEnabled)
 					{
 						SpawnedSpectator->SetActorTickEnabled(true);
