Added OnTickPose delegate invoked at the beginning of TickPose. #jira UE-111916 #rb [REDACTED_HANDLE]

diff --git a/Engine/Source/Runtime/Engine/Classes/Components/SkinnedMeshComponent.h b/Engine/Source/Runtime/Engine/Classes/Components/SkinnedMeshComponent.h
--- a/Engine/Source/Runtime/Engine/Classes/Components/SkinnedMeshComponent.h
+++ b/Engine/Source/Runtime/Engine/Classes/Components/SkinnedMeshComponent.h
@@ -27,6 +27,7 @@ struct FSkelMeshRenderSection;
 class FPositionVertexBuffer;
 
 DECLARE_DELEGATE_OneParam(FOnAnimUpdateRateParamsCreated, FAnimUpdateRateParameters*)
+DECLARE_MULTICAST_DELEGATE_ThreeParams(FOnTickPose, USkinnedMeshComponent* /*SkinnedMeshComponent*/, float /*DeltaTime*/, bool /*bNeedsValidRootMotion*/)
 
 //
 // Bone Visibility.
@@ -1111,7 +1112,12 @@ public:
 	 * @return	Return true if anything modified. Return false otherwise
 	 * @param bNeedsValidRootMotion - Networked games care more about this, but if false we can do less calculations
 	 */
-	virtual void TickPose(float DeltaTime, bool bNeedsValidRootMotion) { TickUpdateRate(DeltaTime, bNeedsValidRootMotion); }
+	virtual void TickPose(float DeltaTime, bool bNeedsValidRootMotion);
+
+	/**
+	 * Invoked at the beginning of TickPose before doing the bulk of the tick work
+	 */
+	FOnTickPose OnTickPose;
 
 	/** 
 	 * Update Slave Component. This gets called when MasterPoseComponent!=NULL
diff --git a/Engine/Source/Runtime/Engine/Private/Components/SkinnedMeshComponent.cpp b/Engine/Source/Runtime/Engine/Private/Components/SkinnedMeshComponent.cpp
--- a/Engine/Source/Runtime/Engine/Private/Components/SkinnedMeshComponent.cpp
+++ b/Engine/Source/Runtime/Engine/Private/Components/SkinnedMeshComponent.cpp
@@ -788,6 +788,12 @@ bool USkinnedMeshComponent::ShouldUseUpdateRateOptimizations() const
 	return bEnableUpdateRateOptimizations && CVarEnableAnimRateOptimization.GetValueOnAnyThread() > 0;
 }
 
+void USkinnedMeshComponent::TickPose(float DeltaTime, bool bNeedsValidRootMotion)
+{
+	OnTickPose.Broadcast(this, DeltaTime, bNeedsValidRootMotion);
+	TickUpdateRate(DeltaTime, bNeedsValidRootMotion);
+}
+
 void USkinnedMeshComponent::TickUpdateRate(float DeltaTime, bool bNeedsValidRootMotion)
 {
 	SCOPE_CYCLE_COUNTER(STAT_TickUpdateRate);
