Made StaticMesh.bHasNavigationData no longer "advanced display" and made StaticMesh.NavCollision's edition gated by bHasNavigationData

diff --git a/Engine/Source/Runtime/Engine/Classes/Engine/StaticMesh.h b/Engine/Source/Runtime/Engine/Classes/Engine/StaticMesh.h
--- a/Engine/Source/Runtime/Engine/Classes/Engine/StaticMesh.h
+++ b/Engine/Source/Runtime/Engine/Classes/Engine/StaticMesh.h
@@ -941,7 +941,7 @@ public:
 
 	/** If true, mesh will have NavCollision property with additional data for navmesh generation and usage.
 	    Set to false for distant meshes (always outside navigation bounds) to save memory on collision data. */
-	UPROPERTY(EditAnywhere, AdvancedDisplay, Category=Navigation)
+	UPROPERTY(EditAnywhere, Category=Navigation)
 	uint8 bHasNavigationData:1;
 
 	/**	
@@ -1387,7 +1387,7 @@ public:
 	/** Pre-build navigation collision */
 private:
 	UE_DEPRECATED(5.00, "This must be protected for async build, always use the accessors even internally.")
-	UPROPERTY(VisibleAnywhere, transient, duplicatetransient, Instanced, Category = Navigation)
+	UPROPERTY(VisibleAnywhere, transient, duplicatetransient, Instanced, Category = Navigation, meta = (EditCondition = "bHasNavigationData"))
 	TObjectPtr<UNavCollisionBase> NavCollision;
 
 public:
