Various ODSC fixes for when running in cooked builds.

- These fixes avoid doing uncessary work when running with ODSC which should speed up iteration.
- Only perform FMaterialupdateContext and caching material expressions if there were actually materials to load.  This work was being done unnecessarily when iterating on global shaders.
- Implement `recompileshaders changed` which should only recompile changed shaders.  Previously this command would fall through to the else clause and the cooking server would load and attempt to compile all material shaders.

diff --git a/Engine/Source/Runtime/Engine/Private/ShaderCompiler/ShaderCompiler.cpp b/Engine/Source/Runtime/Engine/Private/ShaderCompiler/ShaderCompiler.cpp
--- a/Engine/Source/Runtime/Engine/Private/ShaderCompiler/ShaderCompiler.cpp
+++ b/Engine/Source/Runtime/Engine/Private/ShaderCompiler/ShaderCompiler.cpp
@@ -5089,7 +5089,7 @@ namespace
 	{
 		FString CmdName = FParse::Token(CmdString, 0);
 
-		bool bCompileChangedShaders = true;
+		bool bCompileChangedShaders = false;
 		OutMaterialsToLoad.Empty();
 
 		if( !CmdName.IsEmpty() && FCString::Stricmp(*CmdName,TEXT("Material"))==0 )
@@ -5110,6 +5110,10 @@ namespace
 				}
 			}
 		}
+		else if (!CmdName.IsEmpty() && FCString::Stricmp(*CmdName, TEXT("Changed")) == 0)
+		{
+			bCompileChangedShaders = true;
+		}
 		else
 		{
 			// tell other side all the materials to load, by pathname
@@ -5141,9 +5145,6 @@ void ProcessCookOnTheFlyShaders(bool bReloadGlobalShaders, const TArray<uint8>&
 	// load all the mesh material shaders if any were sent back
 	if (MeshMaterialMaps.Num() > 0)
 	{
-		// this will stop the rendering thread, and reattach components, in the destructor
-		FMaterialUpdateContext UpdateContext(0);
-
 		// parse the shaders
 		FMemoryReader MemoryReader(MeshMaterialMaps, true);
 		FNameAsStringProxyArchive Ar(MemoryReader);
@@ -5151,11 +5152,18 @@ void ProcessCookOnTheFlyShaders(bool bReloadGlobalShaders, const TArray<uint8>&
 		TArray<UMaterialInterface*> LoadedMaterials;
 		FMaterialShaderMap::LoadForRemoteRecompile(Ar, GMaxRHIShaderPlatform, LoadedMaterials);
 
-		// gather the shader maps to reattach
-		for (UMaterialInterface* Material : LoadedMaterials)
+		// Only update materials if we need to.
+		if (LoadedMaterials.Num())
 		{
-			Material->RecacheUniformExpressions(true);
-			UpdateContext.AddMaterialInterface(Material);
+			// this will stop the rendering thread, and reattach components, in the destructor
+			FMaterialUpdateContext UpdateContext(0);
+
+			// gather the shader maps to reattach
+			for (UMaterialInterface* Material : LoadedMaterials)
+			{
+				Material->RecacheUniformExpressions(true);
+				UpdateContext.AddMaterialInterface(Material);
+			}
 		}
 	}
 }
