Validate PixelNormal output before using it.
Fixes crash when connecting Normal SceneTexture to Normal material output.

diff --git a/Engine/Source/Runtime/Engine/Private/Materials/HLSLMaterialTranslator.cpp b/Engine/Source/Runtime/Engine/Private/Materials/HLSLMaterialTranslator.cpp
--- a/Engine/Source/Runtime/Engine/Private/Materials/HLSLMaterialTranslator.cpp
+++ b/Engine/Source/Runtime/Engine/Private/Materials/HLSLMaterialTranslator.cpp
@@ -5953,6 +5953,12 @@ int32 FHLSLMaterialTranslator::SceneTextureLookup(int32 ViewportUV, uint32 InSce
 			BufferUV = AddInlinedCodeChunk(MCT_Float2, TEXT("GetDefaultSceneTextureUV(Parameters, %d)"), (int)SceneTextureId);
 		}
 
+		int32 PixelNormal = PixelNormalWS();
+		if (PixelNormal == INDEX_NONE)
+		{
+			return INDEX_NONE;
+		}
+
 		if (FeatureLevel >= ERHIFeatureLevel::SM5)
 		{
 			return AddCodeChunk(MCT_Float4,
