Added bInNetTick to NetDriver. If trying to destroy the net driver while the net driver is ticking with a Browse operation, this operation will be canceled.

diff --git a/Engine/Source/Runtime/Engine/Classes/Engine/NetDriver.h b/Engine/Source/Runtime/Engine/Classes/Engine/NetDriver.h
--- a/Engine/Source/Runtime/Engine/Classes/Engine/NetDriver.h
+++ b/Engine/Source/Runtime/Engine/Classes/Engine/NetDriver.h
@@ -893,9 +893,14 @@ public:
 	double GetElapsedTime() const { return ElapsedTime; }
 	void ResetElapsedTime() { ElapsedTime = 0.0; }
 
+	bool IsInTick() const { return bInTick; }
+
 private:
 	double						ElapsedTime;
 
+	/** Whether or not the NetDriver is ticking */
+	bool bInTick;
+
 public:
 	/** Last realtime a tick dispatch occurred. Used currently to try and diagnose timeout issues */
 	double						LastTickDispatchRealtime;
diff --git a/Engine/Source/Runtime/Engine/Private/NetDriver.cpp b/Engine/Source/Runtime/Engine/Private/NetDriver.cpp
--- a/Engine/Source/Runtime/Engine/Private/NetDriver.cpp
+++ b/Engine/Source/Runtime/Engine/Private/NetDriver.cpp
@@ -356,6 +356,7 @@ UNetDriver::UNetDriver(const FObjectInitializer& ObjectInitializer)
 ,   Notify(nullptr)
 ,	Time( 0.f )
 ,	ElapsedTime( 0.0 )
+,	bInTick(false)
 ,	LastTickDispatchRealtime( 0.f )
 ,   bIsPeer(false)
 ,	bSkipLocalStats(false)
@@ -664,6 +665,8 @@ bool ShouldEnableScopeSecondsTimers()
 
 void UNetDriver::TickFlush(float DeltaSeconds)
 {
+	TGuardValue<bool> GuardInNetTick(bInTick, true);
+
 	CSV_SCOPED_TIMING_STAT_EXCLUSIVE(NetworkOutgoing);
 	SCOPE_CYCLE_COUNTER(STAT_NetTickFlush);
 	bool bEnableTimer = (NetDriverName == NAME_GameNetDriver) && ShouldEnableScopeSecondsTimers();
@@ -1832,6 +1835,8 @@ struct FRPCCSVTracker
 
 void UNetDriver::TickDispatch( float DeltaTime )
 {
+	TGuardValue<bool> GuardInNetTick(bInTick, true);
+
 	SendCycles=0;
 
 	const double CurrentRealtime = FPlatformTime::Seconds();
diff --git a/Engine/Source/Runtime/Engine/Private/UnrealEngine.cpp b/Engine/Source/Runtime/Engine/Private/UnrealEngine.cpp
--- a/Engine/Source/Runtime/Engine/Private/UnrealEngine.cpp
+++ b/Engine/Source/Runtime/Engine/Private/UnrealEngine.cpp
@@ -12362,6 +12362,16 @@ EBrowseReturnVal::Type UEngine::Browse( FWorldContext& WorldContext, FURL URL, F
 	Error = TEXT("");
 	WorldContext.TravelURL = TEXT("");
 
+	if (WorldContext.World() && WorldContext.World()->GetNetDriver())
+	{
+		const TCHAR* InTickErrorString = TEXT("Attempting to call UEngine::Browse and destroy the net driver while the net driver is ticking. Instead try using UWorld::ServerTravel or APlayerController::ClientTravel.");
+		if (!ensureMsgf(!WorldContext.World()->GetNetDriver()->IsInTick(), TEXT("%s"), InTickErrorString))
+		{
+			Error = InTickErrorString;
+			return EBrowseReturnVal::Failure;
+		}
+	}
+
 	// Convert .unreal link files.
 	const TCHAR* LinkStr = TEXT(".unreal");//!!
 	if( FCString::Strstr(*URL.Map,LinkStr)-*URL.Map==FCString::Strlen(*URL.Map)-FCString::Strlen(LinkStr) )
