Remove UE4 references from networking code

diff --git a/Engine/Source/Runtime/Engine/Classes/Engine/NetDriver.h b/Engine/Source/Runtime/Engine/Classes/Engine/NetDriver.h
--- a/Engine/Source/Runtime/Engine/Classes/Engine/NetDriver.h
+++ b/Engine/Source/Runtime/Engine/Classes/Engine/NetDriver.h
@@ -194,7 +194,7 @@
  * Those bunches are then passed along to individual Channels to be processed further.
  *
  * A Packet may contain no bunches, a single bunch, or multiple bunches.
- * Because size limits for bunches may be larger than the size limits of a single packet, UE4 supports
+ * Because size limits for bunches may be larger than the size limits of a single packet, UE supports
  * the notion of partial bunches.
  *
  * When a bunch is too large, before transmission we will slice it into a number of smaller bunches.
@@ -220,7 +220,7 @@
  *****************************************************************************************
  *
  *
- * UE4 Networking typically assumes reliability isn't guaranteed by the underlying network protocol.
+ * UE Networking typically assumes reliability isn't guaranteed by the underlying network protocol.
  * Instead, it implements its own reliability and retransmission of both packets and bunches.
  *
  * When a NetConnection is established, it will establish a Sequence Number for its packets and bunches.
diff --git a/Engine/Source/Runtime/Engine/Private/DataReplication.cpp b/Engine/Source/Runtime/Engine/Private/DataReplication.cpp
--- a/Engine/Source/Runtime/Engine/Private/DataReplication.cpp
+++ b/Engine/Source/Runtime/Engine/Private/DataReplication.cpp
@@ -2221,7 +2221,7 @@ void FObjectReplicator::QueuePropertyRepNotify(
 		//@note: AddUniqueItem() here for static arrays since RepNotify() currently doesn't indicate index,
 		//			so reporting the same property multiple times is not useful and wastes CPU
 		//			were that changed, this should go back to AddItem() for efficiency
-		// @todo UE4 - not checking if replicated value is changed from old.  Either fix or document, as may get multiple repnotifies of unacked properties.
+		// @todo UE - not checking if replicated value is changed from old.  Either fix or document, as may get multiple repnotifies of unacked properties.
 		ReceivingRepState->RepNotifies.AddUnique(Property);
 
 		UFunction* RepNotifyFunc = Object->FindFunctionChecked(Property->RepNotifyFunc);
diff --git a/Engine/Source/Runtime/Engine/Private/Interpolation.cpp b/Engine/Source/Runtime/Engine/Private/Interpolation.cpp
--- a/Engine/Source/Runtime/Engine/Private/Interpolation.cpp
+++ b/Engine/Source/Runtime/Engine/Private/Interpolation.cpp
@@ -3729,7 +3729,7 @@ void UInterpTrackMove::UpdateKeyframe(int32 KeyIndex, UInterpTrackInst* TrInst)
 	UInterpTrackInstMove* MoveTrackInst = CastChecked<UInterpTrackInstMove>(TrInst);
 
 
-	/* @todo UE4 do we still need this code? 
+	/* @todo UE do we still need this code? 
 	if ( APawn * Pawn = Actor->GetAPawn() )
 	{
 		if ( Pawn->CapsuleComponent )
@@ -5754,7 +5754,7 @@ void UInterpTrackInstFloatProp::RestoreActorState(UInterpTrack* Track)
 
 	// We update components, so things like draw scale take effect.
 	// Don't force update all components unless we're in the editor.
-	Actor->ReregisterAllComponents(); // @todo UE4 insist on a property update callback, and use here?
+	Actor->ReregisterAllComponents(); // @todo UE insist on a property update callback, and use here?
 }
 
 void UInterpTrackInstFloatProp::InitTrackInst(UInterpTrack* Track)
@@ -5896,7 +5896,7 @@ void UInterpTrackInstVectorProp::RestoreActorState(UInterpTrack* Track)
 
 	// We update components, so things like draw scale take effect.
 	// Don't force update all components unless we're in the editor.
-	Actor->ReregisterAllComponents(); // @todo UE4 insist on a property update callback, and use here?
+	Actor->ReregisterAllComponents(); // @todo UE insist on a property update callback, and use here?
 }
 
 void UInterpTrackInstVectorProp::InitTrackInst(UInterpTrack* Track)
@@ -6098,7 +6098,7 @@ void UInterpTrackInstBoolProp::RestoreActorState( UInterpTrack* Track )
 
 	// We update components, so things like draw scale take effect.
 	// Don't force update all components unless we're in the editor.
-	Actor->ReregisterAllComponents(); // @todo UE4 insist on a property update callback, and use here?
+	Actor->ReregisterAllComponents(); // @todo UE insist on a property update callback, and use here?
 }
 
 /*-----------------------------------------------------------------------------
@@ -6240,7 +6240,7 @@ void UInterpTrackInstColorProp::RestoreActorState(UInterpTrack* Track)
 
 	// We update components, so things like draw scale take effect.
 	// Don't force update all components unless we're in the editor.
-	Actor->ReregisterAllComponents(); // @todo UE4 insist on a property update callback, and use here?
+	Actor->ReregisterAllComponents(); // @todo UE insist on a property update callback, and use here?
 }
 
 /** Initialize this Track instance. Called in-game before doing any interpolation. */
@@ -6389,7 +6389,7 @@ void UInterpTrackInstLinearColorProp::RestoreActorState(UInterpTrack* Track)
 
 	// We update components, so things like draw scale take effect.
 	// Don't force update all components unless we're in the editor.
-	Actor->ReregisterAllComponents(); // @todo UE4 insist on a property update callback, and use here?
+	Actor->ReregisterAllComponents(); // @todo UE insist on a property update callback, and use here?
 }
 
 void UInterpTrackInstLinearColorProp::InitTrackInst(UInterpTrack* Track)
diff --git a/Engine/Source/Runtime/Engine/Private/Net/NetPushModelHelpers.cpp b/Engine/Source/Runtime/Engine/Private/Net/NetPushModelHelpers.cpp
--- a/Engine/Source/Runtime/Engine/Private/Net/NetPushModelHelpers.cpp
+++ b/Engine/Source/Runtime/Engine/Private/Net/NetPushModelHelpers.cpp
@@ -35,7 +35,7 @@ void UNetPushModelHelpers::MarkPropertyDirty(UObject* Object, FName PropertyName
 }
 
 #if WITH_PUSH_MODEL
-namespace UE4PushModelPrivate
+namespace UEPushModelPrivate
 {
 #if WITH_PUSH_VALIDATION_SUPPORT
 	static bool bCheckPushBPRepIndexAgainstName = false;
@@ -65,7 +65,7 @@ void UNetPushModelHelpers::MarkPropertyDirtyFromRepIndex(UObject* Object, int32
 			else
 			{
 #if WITH_PUSH_VALIDATION_SUPPORT
-				checkf(!UE4PushModelPrivate::bCheckPushBPRepIndexAgainstName || Class->ClassReps[RepIndex].Property->GetFName() == PropertyName,
+				checkf(!UEPushModelPrivate::bCheckPushBPRepIndexAgainstName || Class->ClassReps[RepIndex].Property->GetFName() == PropertyName,
 					TEXT("Property and RepIndex don't match! Object=%s, RepIndex=%d, InPropertyName=%s, FoundPropertyName=%s"),
 						*Object->GetPathName(), RepIndex, *PropertyName.ToString(), *(Class->ClassReps[RepIndex].Property->GetName()));
 #endif
diff --git a/Engine/Source/Runtime/Engine/Private/NetConnection.cpp b/Engine/Source/Runtime/Engine/Private/NetConnection.cpp
--- a/Engine/Source/Runtime/Engine/Private/NetConnection.cpp
+++ b/Engine/Source/Runtime/Engine/Private/NetConnection.cpp
@@ -95,7 +95,7 @@ DECLARE_CYCLE_STAT(TEXT("NetConnection Tick"), Stat_NetConnectionTick, STATGROUP
 DECLARE_CYCLE_STAT(TEXT("NetConnection ReceivedNak"), Stat_NetConnectionReceivedNak, STATGROUP_Net);
 DECLARE_CYCLE_STAT(TEXT("NetConnection NetConnectionReceivedAcks"), Stat_NetConnectionReceivedAcks, STATGROUP_Net);
 
-namespace UE4_NetConnectionPrivate
+namespace UE_NetConnectionPrivate
 {
 	struct FValidateLevelVisibilityResult
 	{
@@ -1077,7 +1077,7 @@ void UNetConnection::UpdateLevelVisibility(const FUpdateLevelVisibilityLevelInfo
 
 void UNetConnection::UpdateLevelVisibilityInternal(const FUpdateLevelVisibilityLevelInfo& LevelVisibility)
 {
-	using namespace UE4_NetConnectionPrivate;
+	using namespace UE_NetConnectionPrivate;
 
 	GNumClientUpdateLevelVisibility++;
 
@@ -1898,20 +1898,20 @@ void UNetConnection::WriteFinalPacketInfo(FBitWriter& Writer, const double Packe
 		uint32 ClockTimeMilliseconds = 0;
 
 		// If the delta is over our max precision, we send MAX value and jitter will be ignored by the receiver.
-		if (DeltaSendTimeInMS >= UE4_NetConnectionPrivate::MaxJitterPrecisionInMS)
+		if (DeltaSendTimeInMS >= UE_NetConnectionPrivate::MaxJitterPrecisionInMS)
 		{
-			ClockTimeMilliseconds = UE4_NetConnectionPrivate::MaxJitterClockTimeValue;
+			ClockTimeMilliseconds = UE_NetConnectionPrivate::MaxJitterClockTimeValue;
 		}
 		else
 		{
 			// Get the fractional part (milliseconds) of the clock time
-			ClockTimeMilliseconds = UE4_NetConnectionPrivate::GetClockTimeMilliseconds(PacketSentTimeInS);
+			ClockTimeMilliseconds = UE_NetConnectionPrivate::GetClockTimeMilliseconds(PacketSentTimeInS);
 
 			// Ensure we don't overflow
-			ClockTimeMilliseconds &= UE4_NetConnectionPrivate::MaxJitterClockTimeValue;
+			ClockTimeMilliseconds &= UE_NetConnectionPrivate::MaxJitterClockTimeValue;
 		}
 
-		Writer.SerializeInt(ClockTimeMilliseconds, UE4_NetConnectionPrivate::MaxJitterClockTimeValue + 1);
+		Writer.SerializeInt(ClockTimeMilliseconds, UE_NetConnectionPrivate::MaxJitterClockTimeValue + 1);
 
 #if !UE_BUILD_SHIPPING
 		checkf(Writer.GetNumBits() - BitsWrittenPreJitterClock == NetConnectionHelper::NumBitsForJitterClockTimeInHeader, TEXT("WriteFinalPacketInfo did not write the expected nb of bits: Wrote %d, Expected %d"),
@@ -2150,7 +2150,7 @@ void UNetConnection::ReceivedPacket( FBitReader& Reader, bool bIsReinjectedPacke
 #endif
 				// Read jitter clock time from the packet header
 				uint32 PacketJitterClockTimeMS = 0;
-				Reader.SerializeInt(PacketJitterClockTimeMS, UE4_NetConnectionPrivate::MaxJitterClockTimeValue + 1);
+				Reader.SerializeInt(PacketJitterClockTimeMS, UE_NetConnectionPrivate::MaxJitterClockTimeValue + 1);
 
 #if !UE_BUILD_SHIPPING
 				checkf(Reader.GetPosBits() - BitsReadPreJitterClock == NetConnectionHelper::NumBitsForJitterClockTimeInHeader, TEXT("JitterClockTime did not read the expected nb of bits. Read %d, Expected %d"),
@@ -4280,13 +4280,13 @@ void UNetConnection::TrackReplicationForAnalytics(const bool bWasSaturated)
 
 void UNetConnection::ProcessJitter(uint32 PacketJitterClockTimeMS)
 {
-	if (PacketJitterClockTimeMS >= UE4_NetConnectionPrivate::MaxJitterClockTimeValue)
+	if (PacketJitterClockTimeMS >= UE_NetConnectionPrivate::MaxJitterClockTimeValue)
 	{
 		// Ignore this packet for jitter calculations
 		return;
 	}
 
-	const int32 CurrentClockTimeMilliseconds = UE4_NetConnectionPrivate::GetClockTimeMilliseconds(LastReceiveRealtime);
+	const int32 CurrentClockTimeMilliseconds = UE_NetConnectionPrivate::GetClockTimeMilliseconds(LastReceiveRealtime);
 
 	// Get the delta between the sent and receive clock time.
 	int32 CurrentJitterTimeDelta = CurrentClockTimeMilliseconds - (int32)PacketJitterClockTimeMS;
@@ -4294,7 +4294,7 @@ void UNetConnection::ProcessJitter(uint32 PacketJitterClockTimeMS)
 	if (CurrentJitterTimeDelta < 0)
 	{
 		// Account for wrap-around
-		CurrentJitterTimeDelta += UE4_NetConnectionPrivate::MaxJitterPrecisionInMS;
+		CurrentJitterTimeDelta += UE_NetConnectionPrivate::MaxJitterPrecisionInMS;
 	}
 
 	// Get the difference between the last two deltas
@@ -4303,7 +4303,7 @@ void UNetConnection::ProcessJitter(uint32 PacketJitterClockTimeMS)
 
 	// Add the value to the average and smooth it out to reduce divergences
 	const float PreviousJitterInMS = AverageJitterInMS;
-	AverageJitterInMS = PreviousJitterInMS + ((CurrentJitter - PreviousJitterInMS) / UE4_NetConnectionPrivate::JitterNoiseReduction);
+	AverageJitterInMS = PreviousJitterInMS + ((CurrentJitter - PreviousJitterInMS) / UE_NetConnectionPrivate::JitterNoiseReduction);
 
 	//UE_LOG(LogNetTraffic, Verbose, TEXT("Avg Jitter: %f | CurrentJitter: %f | CurrentTimeDelta: %d | PreviousTimeDelta: %d | ReceivedClockMilli %d | LocalClockMilli %d"), AverageJitterInMS, CurrentJitter, CurrentJitterTimeDelta, PreviousJitterTimeDelta, PacketJitterClockTimeMS, CurrentClockTimeMilliseconds);
 
diff --git a/Engine/Source/Runtime/Engine/Private/RepLayout.cpp b/Engine/Source/Runtime/Engine/Private/RepLayout.cpp
--- a/Engine/Source/Runtime/Engine/Private/RepLayout.cpp
+++ b/Engine/Source/Runtime/Engine/Private/RepLayout.cpp
@@ -99,12 +99,12 @@ extern TAutoConsoleVariable<int32> CVarNetEnableDetailedScopeCounters;
 */
 inline uint32 GetBitStreamPositionForNetTrace(const FBitArchive& Stream) { return (uint32(Stream.IsError()) - 1U) & (Stream.IsSaving() ? (uint32)(static_cast<const FBitWriter*>(&Stream))->GetNumBits() : (uint32)(static_cast<const FBitReader*>(&Stream))->GetPosBits()); }
 
-namespace UE4PushModelPrivate
+namespace UEPushModelPrivate
 {
 	class FPushModelPerNetDriverState;
 }
 
-namespace UE4_RepLayout_Private
+namespace UE_RepLayout_Private
 {
 	template<typename OutputType, typename CommandType, typename BufferType>
 	static typename TCopyQualifiersFromTo<BufferType, OutputType>::Type*
@@ -142,7 +142,7 @@ namespace UE4_RepLayout_Private
 		//@note: AddUniqueItem() here for static arrays since RepNotify() currently doesn't indicate index,
 		//			so reporting the same property multiple times is not useful and wastes CPU
 		//			were that changed, this should go back to AddItem() for efficiency
-		// @todo UE4 - not checking if replicated value is changed from old.  Either fix or document, as may get multiple repnotifies of unacked properties.
+		// @todo UE - not checking if replicated value is changed from old.  Either fix or document, as may get multiple repnotifies of unacked properties.
 		ReceivingRepState->RepNotifies.AddUnique(Property);
 
 		UFunction* RepNotifyFunc = Params.Object->FindFunctionChecked(Property->RepNotifyFunc);
@@ -192,35 +192,35 @@ namespace UE4_RepLayout_Private
 	}
 
 #if WITH_PUSH_MODEL
-	const UE4PushModelPrivate::FPushModelPerNetDriverHandle ConditionallyAddPushModelObject(const UObject* const Object, const TSharedRef<const FRepLayout>& InRepLayout)
+	const UEPushModelPrivate::FPushModelPerNetDriverHandle ConditionallyAddPushModelObject(const UObject* const Object, const TSharedRef<const FRepLayout>& InRepLayout)
 	{
 		const int32 NumReplicatedProperties = InRepLayout->GetNumParents();
-		if (UE4PushModelPrivate::IsPushModelEnabled() &&
+		if (UEPushModelPrivate::IsPushModelEnabled() &&
 			NumReplicatedProperties > 0 &&
 			EnumHasAnyFlags(InRepLayout->GetFlags(), ERepLayoutFlags::FullPushSupport | ERepLayoutFlags::PartialPushSupport))
 		{
 			const FObjectKey ObjectKey(Object);
-			return UE4PushModelPrivate::AddPushModelObject(ObjectKey, NumReplicatedProperties);
+			return UEPushModelPrivate::AddPushModelObject(ObjectKey, NumReplicatedProperties);
 		}
 
-		return UE4PushModelPrivate::FPushModelPerNetDriverHandle::MakeInvalidHandle();
+		return UEPushModelPrivate::FPushModelPerNetDriverHandle::MakeInvalidHandle();
 	}
 
-	void ConditionallyRemovePushModelObject(const UE4PushModelPrivate::FPushModelPerNetDriverHandle Handle)
+	void ConditionallyRemovePushModelObject(const UEPushModelPrivate::FPushModelPerNetDriverHandle Handle)
 	{
 		if (Handle.IsValid())
 		{
-			UE4PushModelPrivate::RemovePushModelObject(Handle);
+			UEPushModelPrivate::RemovePushModelObject(Handle);
 		}
 	}
 
-	UE4PushModelPrivate::FPushModelPerNetDriverState* GetPerNetDriverState(const FRepChangelistState* ChangelistState)
+	UEPushModelPrivate::FPushModelPerNetDriverState* GetPerNetDriverState(const FRepChangelistState* ChangelistState)
 	{
-		const UE4PushModelPrivate::FPushModelPerNetDriverHandle Handle = ChangelistState->GetPushModelObjectHandle();
-		return Handle.IsValid() ? UE4PushModelPrivate::GetPerNetDriverState(Handle) : nullptr;
+		const UEPushModelPrivate::FPushModelPerNetDriverHandle Handle = ChangelistState->GetPushModelObjectHandle();
+		return Handle.IsValid() ? UEPushModelPrivate::GetPerNetDriverState(Handle) : nullptr;
 	}
 #else
-	UE4PushModelPrivate::FPushModelPerNetDriverState* GetPerNetDriverState(const FRepChangelistState* ChangelistState)
+	UEPushModelPrivate::FPushModelPerNetDriverState* GetPerNetDriverState(const FRepChangelistState* ChangelistState)
 	{
 		return nullptr;
 	}
@@ -905,19 +905,19 @@ static uint32 GetRepLayoutCmdCompatibleChecksum(
 #if (WITH_PUSH_MODEL)
 struct FNetPrivatePushIdHelper
 {
-	static void SetNetPushID(UObject* InObject, const UE4PushModelPrivate::FNetPushObjectId ObjectId)
+	static void SetNetPushID(UObject* InObject, const UEPushModelPrivate::FNetPushObjectId ObjectId)
 	{
-		const UE4PushModelPrivate::FNetPushObjectId CurrentId = InObject->GetNetPushIdDynamic();
+		const UEPushModelPrivate::FNetPushObjectId CurrentId = InObject->GetNetPushIdDynamic();
 		if (CurrentId != ObjectId)
 		{
 			if (CurrentId != INDEX_NONE)
 			{
 				UE_LOG(LogRep, Error, TEXT("SetNetPushID: %s already has a push id. Existing ID = %s, New ID = %s"),
 					*InObject->GetPathName(),
-					*UE4PushModelPrivate::ToString(CurrentId),
-					*UE4PushModelPrivate::ToString(ObjectId));
+					*UEPushModelPrivate::ToString(CurrentId),
+					*UEPushModelPrivate::ToString(ObjectId));
 
-				if (!UE4PushModelPrivate::ValidateObjectIdReassignment(CurrentId, ObjectId))
+				if (!UEPushModelPrivate::ValidateObjectIdReassignment(CurrentId, ObjectId))
 				{
 					return;
 				}
@@ -1112,7 +1112,7 @@ FRepChangelistState::FRepChangelistState(
 	, StaticBuffer(InRepLayout->CreateShadowBuffer(InSource))
 
 #if WITH_PUSH_MODEL
-	, PushModelObjectHandle(UE4_RepLayout_Private::ConditionallyAddPushModelObject(InRepresenting, InRepLayout))
+	, PushModelObjectHandle(UE_RepLayout_Private::ConditionallyAddPushModelObject(InRepresenting, InRepLayout))
 #endif // WITH_PUSH_MODEL
 
 {
@@ -1127,7 +1127,7 @@ FRepChangelistState::FRepChangelistState(
 FRepChangelistState::~FRepChangelistState()
 {
 #if WITH_PUSH_MODEL
-	UE4_RepLayout_Private::ConditionallyRemovePushModelObject(PushModelObjectHandle);
+	UE_RepLayout_Private::ConditionallyRemovePushModelObject(PushModelObjectHandle);
 #endif // WITH_PUSH_MODEL
 }
 
@@ -1254,7 +1254,7 @@ struct FComparePropertiesSharedParams
 	FRepChangelistState* const RepChangelistState;
 	FRepChangedPropertyTracker* const RepChangedPropertyTracker;
 	const TMap<FRepLayoutCmd*, TArray<FRepLayoutCmd>>& NetSerializeLayouts;
-	UE4PushModelPrivate::FPushModelPerNetDriverState* const PushModelState = nullptr;
+	UEPushModelPrivate::FPushModelPerNetDriverState* const PushModelState = nullptr;
 	const TBitArray<>* const PushModelProperties = nullptr;
 	const bool bValidateProperties = false;
 	const bool bIsNetworkProfilerActive = false;
@@ -1368,7 +1368,7 @@ static bool CompareParentProperty(
 	return !!(StackParams.Changed.Num() - NumChanges);
 }
 
-namespace UE4_RepLayout_Private
+namespace UE_RepLayout_Private
 {
 	static bool CompareParentPropertyHelper(
 		const int32 ParentIndex,
@@ -1433,7 +1433,7 @@ static void CompareParentProperties(
 		{
 			for (int32 ParentIndex = 0; ParentIndex < SharedParams.Parents.Num(); ++ParentIndex)
 			{
-				UE4_RepLayout_Private::CompareParentPropertyHelper(ParentIndex, SharedParams, StackParams);
+				UE_RepLayout_Private::CompareParentPropertyHelper(ParentIndex, SharedParams, StackParams);
 			}
 		}
 
@@ -1443,8 +1443,8 @@ static void CompareParentProperties(
 		{
 			for (int32 ParentIndex = 0; ParentIndex < SharedParams.Parents.Num(); ++ParentIndex)
 			{
-				const bool bIsPropertyDirty = UE4_RepLayout_Private::IsPropertyDirty(ParentIndex, bRecentlyCollectedGarbage, SharedParams, StackParams);
-				const bool bDidPropertyChange = UE4_RepLayout_Private::CompareParentPropertyHelper(ParentIndex, SharedParams, StackParams);
+				const bool bIsPropertyDirty = UE_RepLayout_Private::IsPropertyDirty(ParentIndex, bRecentlyCollectedGarbage, SharedParams, StackParams);
+				const bool bDidPropertyChange = UE_RepLayout_Private::CompareParentPropertyHelper(ParentIndex, SharedParams, StackParams);
 
 				ensureAlwaysMsgf(!bDidPropertyChange || bIsPropertyDirty, TEXT("Push Model Property changed value, but was not marked dirty! Property=%s"), *SharedParams.Parents[ParentIndex].Property->GetPathName());
 			}	
@@ -1456,16 +1456,16 @@ static void CompareParentProperties(
 		{
 			for (TConstSetBitIterator<> It = SharedParams.PushModelState->GetDirtyProperties(); It; ++It)
 			{
-				UE4_RepLayout_Private::CompareParentPropertyHelper(It.GetIndex(), SharedParams, StackParams);
+				UE_RepLayout_Private::CompareParentPropertyHelper(It.GetIndex(), SharedParams, StackParams);
 			}
 		}
 		else
 		{
 			for (int32 ParentIndex = 0; ParentIndex < SharedParams.Parents.Num(); ++ParentIndex)
 			{
-				if (UE4_RepLayout_Private::IsPropertyDirty(ParentIndex, bRecentlyCollectedGarbage, SharedParams, StackParams))
+				if (UE_RepLayout_Private::IsPropertyDirty(ParentIndex, bRecentlyCollectedGarbage, SharedParams, StackParams))
 				{
-					UE4_RepLayout_Private::CompareParentPropertyHelper(ParentIndex, SharedParams, StackParams);
+					UE_RepLayout_Private::CompareParentPropertyHelper(ParentIndex, SharedParams, StackParams);
 				}
 			}
 		}
@@ -1477,7 +1477,7 @@ static void CompareParentProperties(
 
 	for (int32 ParentIndex = 0; ParentIndex < SharedParams.Parents.Num(); ++ParentIndex)
 	{
-		UE4_RepLayout_Private::CompareParentPropertyHelper(ParentIndex, SharedParams, StackParams);
+		UE_RepLayout_Private::CompareParentPropertyHelper(ParentIndex, SharedParams, StackParams);
 	}
 }
 
@@ -1542,7 +1542,7 @@ static void CompareProperties_Array_r(
 	const int32 ArrayNum = Array->Num();
 	const int32 ShadowArrayNum = ShadowArray->Num();
 
-	if (!UE4_RepLayout_Private::ValidateArraySize(ArrayNum, Cmd.Property))
+	if (!UE_RepLayout_Private::ValidateArraySize(ArrayNum, Cmd.Property))
 	{
 		StackParams.Result = ERepLayoutResult::FatalError;
 		return;
@@ -1584,7 +1584,7 @@ static void CompareProperties_Array_r(
 	{
 		const int32 NumChangedEntries = ChangedLocal.Num();
 		
-		if (!UE4_RepLayout_Private::ValidateArraySize(NumChangedEntries, Cmd.Property))
+		if (!UE_RepLayout_Private::ValidateArraySize(NumChangedEntries, Cmd.Property))
 		{
 			StackParams.Result = ERepLayoutResult::FatalError;
 			return;
@@ -1648,10 +1648,10 @@ ERepLayoutResult FRepLayout::CompareProperties(
 		RepChangelistState,
 		(RepState ? RepState->RepChangedPropertyTracker.Get() : nullptr),
 		NetSerializeLayouts,
-		/*PushModelState=*/UE4_RepLayout_Private::GetPerNetDriverState(RepChangelistState),
+		/*PushModelState=*/UE_RepLayout_Private::GetPerNetDriverState(RepChangelistState),
 		/*PushModelProperties=*/ LocalPushModelProperties,	
 		/*bValidateProperties=*/GbPushModelValidateProperties,
-		/*bIsNetworkProfilerActive=*/UE4_RepLayout_Private::IsNetworkProfilerComparisonTrackingEnabled(),
+		/*bIsNetworkProfilerActive=*/UE_RepLayout_Private::IsNetworkProfilerComparisonTrackingEnabled(),
 		/*bChangedNetOwner=*/ RepState && RepState->RepFlags.bNetOwner != RepFlags.bNetOwner
 	};
 
@@ -4281,7 +4281,7 @@ void FRepLayout::UpdateUnmappedObjects(
 
 				if (TempParams.bOutSomeObjectsWereMapped && INDEX_NONE != Parent.RepNotifyNumParams)
 				{
-					UE4_RepLayout_Private::QueueRepNotifyForCustomDeltaProperty(RepState, Params, StructProperty, Parent.ArrayIndex);
+					UE_RepLayout_Private::QueueRepNotifyForCustomDeltaProperty(RepState, Params, StructProperty, Parent.ArrayIndex);
 				}
 
 				Params.bOutSomeObjectsWereMapped |= TempParams.bOutSomeObjectsWereMapped;
@@ -4399,7 +4399,7 @@ bool FRepLayout::ReceiveCustomDeltaProperty(
 		// Successfully received it.
 		if (INDEX_NONE != Parent.RepNotifyNumParams)
 		{
-			UE4_RepLayout_Private::QueueRepNotifyForCustomDeltaProperty(ReceivingRepState, Params, Property, StaticArrayIndex);
+			UE_RepLayout_Private::QueueRepNotifyForCustomDeltaProperty(ReceivingRepState, Params, Property, StaticArrayIndex);
 		}
 
 		return true;
@@ -6257,7 +6257,7 @@ void FRepLayout::SerializeProperties_DynamicArray_r(
 	const int32 ArrayNum = Ar.IsLoading() ? (int32)OutArrayNum : Array->Num();
 
 	// Validate the maximum number of elements.
-	if (!UE4_RepLayout_Private::ValidateArraySize(ArrayNum, Cmd.Property))
+	if (!UE_RepLayout_Private::ValidateArraySize(ArrayNum, Cmd.Property))
 	{
 		Ar.SetError();
 	}
@@ -6536,7 +6536,7 @@ void FRepLayout::BuildSharedSerializationForRPC_DynamicArray_r(
 	FScriptArray* Array = (FScriptArray *)Data.Data;	
 	const int32 ArrayNum = Array->Num();
 
-	if (!UE4_RepLayout_Private::ValidateArraySize(ArrayNum, Cmd.Property))
+	if (!UE_RepLayout_Private::ValidateArraySize(ArrayNum, Cmd.Property))
 	{
 		return;
 	}
@@ -7092,7 +7092,7 @@ void FRepLayout::PreSendCustomDeltaProperties(
 	FReplicationChangelistMgr& ChangelistMgr,
 	TArray<TSharedPtr<INetDeltaBaseState>>& CustomDeltaStates) const
 {
-	using namespace UE4_RepLayout_Private;
+	using namespace UE_RepLayout_Private;
 
 	if (!Connection->IsInternalAck())
 	{
@@ -7187,7 +7187,7 @@ void FRepLayout::PostSendCustomDeltaProperties(
 
 ERepLayoutResult FRepLayout::DeltaSerializeFastArrayProperty(FFastArrayDeltaSerializeParams& Params, FReplicationChangelistMgr* ChangelistMgr) const
 {
-	using namespace UE4_RepLayout_Private;
+	using namespace UE_RepLayout_Private;
 
 	CONDITIONAL_SCOPE_CYCLE_COUNTER(STAT_RepLayout_DeltaSerializeFastArray, CVarNetEnableDetailedScopeCounters.GetValueOnAnyThread() > 0);
 
@@ -7815,7 +7815,7 @@ ERepLayoutResult FRepLayout::DeltaSerializeFastArrayProperty(FFastArrayDeltaSeri
 
 void FRepLayout::GatherGuidReferencesForFastArray(FFastArrayDeltaSerializeParams& Params) const
 {
-	using namespace UE4_RepLayout_Private;
+	using namespace UE_RepLayout_Private;
 
 	const FConstRepObjectDataBuffer ObjectData(Params.DeltaSerializeInfo.Object);
 	const FLifetimeCustomDeltaProperty& CustomDeltaProperty = LifetimeCustomPropertyState->GetCustomDeltaProperty(Params.DeltaSerializeInfo.CustomDeltaIndex);
@@ -7838,7 +7838,7 @@ void FRepLayout::GatherGuidReferencesForFastArray(FFastArrayDeltaSerializeParams
 
 bool FRepLayout::MoveMappedObjectToUnmappedForFastArray(FFastArrayDeltaSerializeParams& Params) const
 {
-	using namespace UE4_RepLayout_Private;
+	using namespace UE_RepLayout_Private;
 
 	const FRepObjectDataBuffer ObjectData(Params.DeltaSerializeInfo.Object);
 	const FLifetimeCustomDeltaProperty& CustomDeltaProperty = LifetimeCustomPropertyState->GetCustomDeltaProperty(Params.DeltaSerializeInfo.CustomDeltaIndex);
@@ -7857,7 +7857,7 @@ bool FRepLayout::MoveMappedObjectToUnmappedForFastArray(FFastArrayDeltaSerialize
 
 void FRepLayout::UpdateUnmappedGuidsForFastArray(FFastArrayDeltaSerializeParams& Params) const
  {
-	using namespace UE4_RepLayout_Private;
+	using namespace UE_RepLayout_Private;
 
 	check(LifetimeCustomPropertyState);
 
diff --git a/Engine/Source/Runtime/Engine/Public/Net/RepLayout.h b/Engine/Source/Runtime/Engine/Public/Net/RepLayout.h
--- a/Engine/Source/Runtime/Engine/Public/Net/RepLayout.h
+++ b/Engine/Source/Runtime/Engine/Public/Net/RepLayout.h
@@ -55,7 +55,7 @@ enum class ERepDataBufferType
 	ShadowBuffer	//! Indicates this buffer is a packed shadow buffer.
 };
 
-namespace UE4_RepLayout_Private
+namespace UE_RepLayout_Private
 {
 	/**
 	 * TRepDataBuffer and TConstRepDataBuffer act as wrapper around internal data
@@ -105,8 +105,8 @@ namespace UE4_RepLayout_Private
 	}
 }
 
-template<ERepDataBufferType DataType> using TRepDataBuffer = UE4_RepLayout_Private::TRepDataBufferBase<DataType, uint8>;
-template<ERepDataBufferType DataType> using TConstRepDataBuffer = UE4_RepLayout_Private::TRepDataBufferBase<DataType, const uint8>;
+template<ERepDataBufferType DataType> using TRepDataBuffer = UE_RepLayout_Private::TRepDataBufferBase<DataType, uint8>;
+template<ERepDataBufferType DataType> using TConstRepDataBuffer = UE_RepLayout_Private::TRepDataBufferBase<DataType, const uint8>;
 
 typedef TRepDataBuffer<ERepDataBufferType::ObjectBuffer> FRepObjectDataBuffer;
 typedef TRepDataBuffer<ERepDataBufferType::ShadowBuffer> FRepShadowDataBuffer;
@@ -509,13 +509,13 @@ public:
 
 #if WITH_PUSH_MODEL
 
-	const UE4PushModelPrivate::FPushModelPerNetDriverHandle& GetPushModelObjectHandle() const
+	const UEPushModelPrivate::FPushModelPerNetDriverHandle& GetPushModelObjectHandle() const
 	{
 		return PushModelObjectHandle;
 	}
 
 private:
-	const UE4PushModelPrivate::FPushModelPerNetDriverHandle PushModelObjectHandle;
+	const UEPushModelPrivate::FPushModelPerNetDriverHandle PushModelObjectHandle;
 #endif
 };
 
diff --git a/Engine/Source/Runtime/Net/Core/Private/Net/Core/PushModel/PushModel.cpp b/Engine/Source/Runtime/Net/Core/Private/Net/Core/PushModel/PushModel.cpp
--- a/Engine/Source/Runtime/Net/Core/Private/Net/Core/PushModel/PushModel.cpp
+++ b/Engine/Source/Runtime/Net/Core/Private/Net/Core/PushModel/PushModel.cpp
@@ -17,7 +17,7 @@
 DECLARE_CYCLE_STAT(TEXT("PushModel PostGarbageCollect"), STAT_PushModel_PostGarbageCollect, STATGROUP_Net);
 DEFINE_LOG_CATEGORY_STATIC(LogPushModel, All, All);
 
-namespace UE4PushModelPrivate
+namespace UEPushModelPrivate
 {
 	//! Originally, multiple implementations of FPushModelObjectManagers were tested.
 	//! The way this worked was by implementing the managers here, using Macros or
diff --git a/Engine/Source/Runtime/Net/Core/Private/Net/Core/PushModel/Types/PushModelPerNetDriverState.h b/Engine/Source/Runtime/Net/Core/Private/Net/Core/PushModel/Types/PushModelPerNetDriverState.h
--- a/Engine/Source/Runtime/Net/Core/Private/Net/Core/PushModel/Types/PushModelPerNetDriverState.h
+++ b/Engine/Source/Runtime/Net/Core/Private/Net/Core/PushModel/Types/PushModelPerNetDriverState.h
@@ -10,7 +10,7 @@
 #include "Containers/BitArray.h"
 #include "PushModelUtils.h"
 
-namespace UE4PushModelPrivate
+namespace UEPushModelPrivate
 {
 	class FPushModelPerNetDriverState
 	{
diff --git a/Engine/Source/Runtime/Net/Core/Private/Net/Core/PushModel/Types/PushModelPerObjectState.h b/Engine/Source/Runtime/Net/Core/Private/Net/Core/PushModel/Types/PushModelPerObjectState.h
--- a/Engine/Source/Runtime/Net/Core/Private/Net/Core/PushModel/Types/PushModelPerObjectState.h
+++ b/Engine/Source/Runtime/Net/Core/Private/Net/Core/PushModel/Types/PushModelPerObjectState.h
@@ -13,7 +13,7 @@
 #include "Containers/SparseArray.h"
 #include "UObject/ObjectKey.h"
 
-namespace UE4PushModelPrivate
+namespace UEPushModelPrivate
 {
 	/**
 	 * This is a "state" for a given Object that is being tracked by a Push Model Object Manager.
diff --git a/Engine/Source/Runtime/Net/Core/Private/Net/Core/PushModel/Types/PushModelUtils.h b/Engine/Source/Runtime/Net/Core/Private/Net/Core/PushModel/Types/PushModelUtils.h
--- a/Engine/Source/Runtime/Net/Core/Private/Net/Core/PushModel/Types/PushModelUtils.h
+++ b/Engine/Source/Runtime/Net/Core/Private/Net/Core/PushModel/Types/PushModelUtils.h
@@ -9,7 +9,7 @@
 #include "CoreMinimal.h"
 #include "Containers/BitArray.h"
 
-namespace UE4PushModelPrivate
+namespace UEPushModelPrivate
 {
 	static const int32 GetNumWordsInBitArray(const TBitArray<>& BitArray)
 	{
diff --git a/Engine/Source/Runtime/Net/Core/Public/Net/Core/PushModel/PushModel.h b/Engine/Source/Runtime/Net/Core/Public/Net/Core/PushModel/PushModel.h
--- a/Engine/Source/Runtime/Net/Core/Public/Net/Core/PushModel/PushModel.h
+++ b/Engine/Source/Runtime/Net/Core/Public/Net/Core/PushModel/PushModel.h
@@ -20,10 +20,10 @@
  * Rationale
  ****************************************************************
  *
- * The intention of Push Model is to provide a fundamental trade off in the way UE4 property replication works, without changing
+ * The intention of Push Model is to provide a fundamental trade off in the way UE property replication works, without changing
  * the underlying replication machinery.
  *
- * Traditionally, UE4 Networking works by having coders and designers mark properties as replicated, and then having the network
+ * Traditionally, UE Networking works by having coders and designers mark properties as replicated, and then having the network
  * system periodically compare the values of those properties to their last sent state. When we detect a change in the state,
  * then we will send that data to clients.
  *
@@ -36,7 +36,7 @@
  * on every object. This sort of overhead isn't a big deal while games are relatively small, but as they start to scale up, this
  * can become a bottleneck.
  *
- * To abate this, there's a bunch of machinery UE4 provides, things like:
+ * To abate this, there's a bunch of machinery UE provides, things like:
  *
  *	Net Frequency
  *		Gives devs the ability to tune how often actors are considered for replication.
@@ -274,7 +274,7 @@ private:
 // DO NOT USE METHODS IN THIS NAMESPACE DIRECTLY
 // Use the Macros instead, as they respect conditional compilation.
 // See PushModelMarcos.h
-namespace UE4PushModelPrivate
+namespace UEPushModelPrivate
 {
 	//~ Using int32 isn't very forward looking, but for now GUObjectArray also uses int32
 	//~ so we're probably safe.
@@ -380,9 +380,9 @@ namespace UE4PushModelPrivate
 }
 
 
-#define CONDITIONAL_ON_PUSH_MODEL(Work) if (UE4PushModelPrivate::IsPushModelEnabled()) { Work; }
-#define IS_PUSH_MODEL_ENABLED() UE4PushModelPrivate::IsPushModelEnabled()
-#define PUSH_MAKE_BP_PROPERTIES_PUSH_MODEL() (UE4PushModelPrivate::IsPushModelEnabled() && UE4PushModelPrivate::MakeBpPropertiesPushModel())
+#define CONDITIONAL_ON_PUSH_MODEL(Work) if (UEPushModelPrivate::IsPushModelEnabled()) { Work; }
+#define IS_PUSH_MODEL_ENABLED() UEPushModelPrivate::IsPushModelEnabled()
+#define PUSH_MAKE_BP_PROPERTIES_PUSH_MODEL() (UEPushModelPrivate::IsPushModelEnabled() && UEPushModelPrivate::MakeBpPropertiesPushModel())
 
 #define GET_PROPERTY_REP_INDEX(ClassName, PropertyName) (int32)ClassName::ENetFields_Private::PropertyName
 #define GET_PROPERTY_REP_INDEX_STATIC_ARRAY_START(ClassName, PropertyName) ((int32)ClassName::ENetFields_Private::PropertyName ## _STATIC_ARRAY)
@@ -391,34 +391,34 @@ namespace UE4PushModelPrivate
 
 #define IS_PROPERTY_REPLICATED(Property) (0 != (EPropertyFlags::CPF_Net & Property->PropertyFlags))
 
-#define CONDITIONAL_ON_OBJECT_NET_ID(Object, Work) { const UE4PushModelPrivate::FNetPushObjectId PrivatePushId = Object->GetNetPushId(); Work; }
-#define CONDITIONAL_ON_OBJECT_NET_ID_DYNAMIC(Object, Work) { const UE4PushModelPrivate::FNetPushObjectId PrivatePushId = Object->GetNetPushIdDynamic(); Work; }
-#define CONDITIONAL_ON_REP_INDEX_AND_OBJECT_NET_ID(Object, Property, Work) if (IS_PROPERTY_REPLICATED(Property)) { const UE4PushModelPrivate::FNetPushObjectId PrivatePushId = Object->GetNetPushIdDynamic();  Work; }
+#define CONDITIONAL_ON_OBJECT_NET_ID(Object, Work) { const UEPushModelPrivate::FNetPushObjectId PrivatePushId = Object->GetNetPushId(); Work; }
+#define CONDITIONAL_ON_OBJECT_NET_ID_DYNAMIC(Object, Work) { const UEPushModelPrivate::FNetPushObjectId PrivatePushId = Object->GetNetPushIdDynamic(); Work; }
+#define CONDITIONAL_ON_REP_INDEX_AND_OBJECT_NET_ID(Object, Property, Work) if (IS_PROPERTY_REPLICATED(Property)) { const UEPushModelPrivate::FNetPushObjectId PrivatePushId = Object->GetNetPushIdDynamic();  Work; }
 
 //~ For these macros, we won't bother checking if Push Model is enabled. Instead, we'll just check to see whether or not the Custom ID is valid.
 
 // Marks a property dirty by RepIndex without doing additional rep index validation.
-#define MARK_PROPERTY_DIRTY_UNSAFE(Object, RepIndex) CONDITIONAL_ON_OBJECT_NET_ID_DYNAMIC(Object, UE4PushModelPrivate::MarkPropertyDirty(PrivatePushId, RepIndex))
+#define MARK_PROPERTY_DIRTY_UNSAFE(Object, RepIndex) CONDITIONAL_ON_OBJECT_NET_ID_DYNAMIC(Object, UEPushModelPrivate::MarkPropertyDirty(PrivatePushId, RepIndex))
 
 // Marks a property dirty by UProperty*, validating that it's actually a replicated property.
-#define MARK_PROPERTY_DIRTY(Object, Property) CONDITIONAL_ON_REP_INDEX_AND_OBJECT_NET_ID(Object, Property, UE4PushModelPrivate::MarkPropertyDirty(PrivatePushId, Property->RepIndex))
+#define MARK_PROPERTY_DIRTY(Object, Property) CONDITIONAL_ON_REP_INDEX_AND_OBJECT_NET_ID(Object, Property, UEPushModelPrivate::MarkPropertyDirty(PrivatePushId, Property->RepIndex))
 
 
 // Marks a static array property dirty given, the Object, UProperty*, and Index.
-#define MARK_PROPERTY_DIRTY_STATIC_ARRAY_INDEX(Object, Property, ArrayIndex) CONDITIONAL_ON_REP_INDEX_AND_OBJECT_NET_ID(Object, Property, UE4PushModelPrivate::MarkPropertyDirty(PrivatePushId, Property->RepIndex + ArrayIndex))
+#define MARK_PROPERTY_DIRTY_STATIC_ARRAY_INDEX(Object, Property, ArrayIndex) CONDITIONAL_ON_REP_INDEX_AND_OBJECT_NET_ID(Object, Property, UEPushModelPrivate::MarkPropertyDirty(PrivatePushId, Property->RepIndex + ArrayIndex))
 
 // Marks all elements of a static array property dirty, given the Object and UProperty*
-#define MARK_PROPERTY_DIRTY_STATIC_ARRAY(Object, Property) CONDITIONAL_ON_REP_INDEX_AND_OBJECT_NET_ID(Object, Property, UE4PushModelPrivate::MarkPropertyDirty(PrivatePushId, Property->RepIndex, Property->RepIndex + Property->ArrayDim - 1))
+#define MARK_PROPERTY_DIRTY_STATIC_ARRAY(Object, Property) CONDITIONAL_ON_REP_INDEX_AND_OBJECT_NET_ID(Object, Property, UEPushModelPrivate::MarkPropertyDirty(PrivatePushId, Property->RepIndex, Property->RepIndex + Property->ArrayDim - 1))
 
 
 // Marks a property dirty, given the Class Name, Property Name, and Object. This will fail to compile if the Property or Class aren't valid.
-#define MARK_PROPERTY_DIRTY_FROM_NAME(ClassName, PropertyName, Object) CONDITIONAL_ON_OBJECT_NET_ID(Object, UE4PushModelPrivate::MarkPropertyDirty(PrivatePushId, GET_PROPERTY_REP_INDEX(ClassName, PropertyName)))
+#define MARK_PROPERTY_DIRTY_FROM_NAME(ClassName, PropertyName, Object) CONDITIONAL_ON_OBJECT_NET_ID(Object, UEPushModelPrivate::MarkPropertyDirty(PrivatePushId, GET_PROPERTY_REP_INDEX(ClassName, PropertyName)))
 
 // Marks a static array property dirty, given the Class Name, Property Name, Index, and Object. This will fail to compile if the Property and Class aren't valid. Callers are responsible for validating the index.
-#define MARK_PROPERTY_DIRTY_FROM_NAME_STATIC_ARRAY_INDEX(ClassName, PropertyName, ArrayIndex, Object) CONDITIONAL_ON_OBJECT_NET_ID(Object, UE4PushModelPrivate::MarkPropertyDirty(PrivatePushId, GET_PROPERTY_REP_INDEX_STATIC_ARRAY_INDEX(ClassName, PropertyName, ArrayIndex)))
+#define MARK_PROPERTY_DIRTY_FROM_NAME_STATIC_ARRAY_INDEX(ClassName, PropertyName, ArrayIndex, Object) CONDITIONAL_ON_OBJECT_NET_ID(Object, UEPushModelPrivate::MarkPropertyDirty(PrivatePushId, GET_PROPERTY_REP_INDEX_STATIC_ARRAY_INDEX(ClassName, PropertyName, ArrayIndex)))
 
 // Marks an entire static array property dirty, given the Class Name, Property Name, and Object. This will fail to compile if the Property or Class aren't valid.
-#define MARK_PROPERTY_DIRTY_FROM_NAME_STATIC_ARRAY(ClassName, PropertyName, Object) CONDITIONAL_ON_OBJECT_NET_ID(Object, UE4PushModelPrivate::MarkPropertyDirty(PrivatePushId, GET_PROPERTY_REP_INDEX_STATIC_ARRAY_START(ClassName, PropertyName), GET_PROPERTY_REP_INDEX_STATIC_ARRAY_END(ClassName, PropertyName))
+#define MARK_PROPERTY_DIRTY_FROM_NAME_STATIC_ARRAY(ClassName, PropertyName, Object) CONDITIONAL_ON_OBJECT_NET_ID(Object, UEPushModelPrivate::MarkPropertyDirty(PrivatePushId, GET_PROPERTY_REP_INDEX_STATIC_ARRAY_START(ClassName, PropertyName), GET_PROPERTY_REP_INDEX_STATIC_ARRAY_END(ClassName, PropertyName))
 
 
 #else // WITH_PUSH_MODEL
