'Shadow Map Method' project setting, which can be changed at runtime for easy comparison
Exposed 'Support non-Nanite geometry in Virtual Shadow Maps' project setting, which is readonly
Setting 'Shadow Map Method' to 'Virtual Shadow Maps' automatically enables it's dependency 'Support non-Nanite geometry in Virtual Shadow Maps'

diff --git a/Engine/Source/Runtime/Engine/Classes/Engine/EngineTypes.h b/Engine/Source/Runtime/Engine/Classes/Engine/EngineTypes.h
--- a/Engine/Source/Runtime/Engine/Classes/Engine/EngineTypes.h
+++ b/Engine/Source/Runtime/Engine/Classes/Engine/EngineTypes.h
@@ -405,23 +405,20 @@ namespace EReflectionMethod
 	};
 }
 
-
-
-// Note: Maps to r.Shadow.Virtual.Enable AND r.Shadow.Virtual.NonNaniteVSM, where Standard == 0 == off, this enum is used in URendererSettings for readability
+// Note: Must match r.Shadow.Virtual.Enable, this is used in URendererSettings
 UENUM()
-namespace EDynamicShadowingMethod
+namespace EShadowMapMethod
 {
 	enum Type
 	{
-		/** UE4 shadowing system only. Required to enable stationary baked shadows (but which is incompatible with Nanite geometry) */
-		Standard,
+		/** Render geometry into shadow depth maps for shadowing.  Requires manual setup of shadowing distances and only culls per-component, causing poor performance with high poly scenes.  Required to enable stationary baked shadows (but which is incompatible with Nanite geometry). */
+		ShadowMaps UMETA(DisplayName = "Shadow Maps"),
 
-		/** High-quality shadows for next-gen projects with simplified setup. This system is in development and thus has a number of performance pitfalls. */
-		VirtualShadowMaps UMETA(DisplayName = "Virtual Shadow Maps (Beta)"),
+		/** Render geometry into virtualized shadow depth maps for shadowing.  Provides high-quality shadows for next-gen projects with simplified setup.  High efficiency culling when used with Nanite. This system is in development and thus has a number of performance pitfalls. */
+		VirtualShadowMaps UMETA(DisplayName = "Virtual Shadow Maps (Beta)")
 	};
 }
 
-
 /** Specifies which component of the scene rendering should be output to the final render target. */
 UENUM()
 enum ESceneCaptureSource 
diff --git a/Engine/Source/Runtime/Engine/Classes/Engine/RendererSettings.h b/Engine/Source/Runtime/Engine/Classes/Engine/RendererSettings.h
--- a/Engine/Source/Runtime/Engine/Classes/Engine/RendererSettings.h
+++ b/Engine/Source/Runtime/Engine/Classes/Engine/RendererSettings.h
@@ -363,16 +363,13 @@ class ENGINE_API URendererSettings : public UDeveloperSettings
 	TEnumAsByte<ELumenSoftwareTracingMode::Type> LumenSoftwareTracingMode;
 
 	UPROPERTY(config, EditAnywhere, Category = Shadows, meta = (
-		ConsoleVariable = "r.Shadow.Virtual.Enable", DisplayName = "Dynamic Shadowing Method",
-		ToolTip = "Dynamic Shadowing Method",
-		ConfigRestartRequired = true))
-	TEnumAsByte<EDynamicShadowingMethod::Type> DynamicShadowingMethod;
+		ConsoleVariable = "r.Shadow.Virtual.Enable", DisplayName = "Shadow Map Method",
+		ToolTip = "Shadow Map Method"))
+	TEnumAsByte<EShadowMapMethod::Type> ShadowMapMethod;
 
-	// Note: this is slaved to the DynamicShadowingMethod, which is why it is not editable, the property is only here to let us propagate to the ini file.
 	UPROPERTY(config, EditAnywhere, Category = Shadows, meta = (
-		EditCondition = "false",
-		ReadOnly,
-		ConsoleVariable="r.Shadow.Virtual.NonNaniteVSM",
+		ConsoleVariable = "r.Shadow.Virtual.NonNaniteVSM", DisplayName = "Support non-Nanite geometry in Virtual Shadow Maps",
+		ToolTip = "Whether to support non-Nanite geometry in Virtual Shadow Maps.  Enabling this incurs a small overhead if Virtual Shadow Maps is not the active Dynamic Shadowing Method.",
 		ConfigRestartRequired = true))
 	uint32 bEnableNonNaniteVSM:1;
 
diff --git a/Engine/Source/Runtime/Engine/Private/RendererSettings.cpp b/Engine/Source/Runtime/Engine/Private/RendererSettings.cpp
--- a/Engine/Source/Runtime/Engine/Private/RendererSettings.cpp
+++ b/Engine/Source/Runtime/Engine/Private/RendererSettings.cpp
@@ -77,8 +77,6 @@ URendererSettings::URendererSettings(const FObjectInitializer& ObjectInitializer
 	bEnableRayTracingShadows = 0;
 	bEnableRayTracingTextureLOD = 0;
 	MaxSkinBones = FGPUBaseSkinVertexFactory::GHardwareMaxGPUSkinBones;
-
-	bEnableNonNaniteVSM = 0;
 }
 
 void URendererSettings::PostInitProperties()
@@ -195,29 +193,24 @@ void URendererSettings::PostEditChangeProperty(FPropertyChangedEvent& PropertyCh
 			}
 		}
 
-		if (PropertyChangedEvent.Property->GetFName() == GET_MEMBER_NAME_CHECKED(URendererSettings, DynamicShadowingMethod))
+		if (PropertyChangedEvent.Property->GetFName() == GET_MEMBER_NAME_CHECKED(URendererSettings, ShadowMapMethod)
+			&& ShadowMapMethod == EShadowMapMethod::VirtualShadowMaps)
 		{
-			// Propagate the setting to the 'r.Shadow.Virtual.NonNaniteVSM' CVar
-			int32 Value = static_cast<int32>(DynamicShadowingMethod);
-			bEnableNonNaniteVSM = Value;
-
-			// Don't think we need this as the below should trigger a reload of INI?
-			//IConsoleVariable* CVar = IConsoleManager::Get().FindConsoleVariable(TEXT("r.Shadow.Virtual.NonNaniteVSM"));
-			//CVar->Set(Value, ECVF_SetByProjectSetting);
-			for (TFieldIterator<FProperty> PropIt(GetClass()); PropIt; ++PropIt)
+			if (!bEnableNonNaniteVSM)
 			{
-				FProperty* Property = *PropIt;
-				if (Property->GetFName() == GET_MEMBER_NAME_CHECKED(URendererSettings, bEnableNonNaniteVSM))
+				FMessageDialog::Open(EAppMsgType::Ok, LOCTEXT("Setting automatically enabled, needs restart.", "Virtual Shadow Maps are most efficient when handling both geometry paths, 'Support non-Nanite geometry in Virtual Shadow Maps' has been automatically enabled.  This requires a restart."));
+
+				bEnableNonNaniteVSM = true;
+
+				for (TFieldIterator<FProperty> PropIt(GetClass()); PropIt; ++PropIt)
 				{
-					UpdateSinglePropertyInConfigFile(Property, GetDefaultConfigFilename());
+					FProperty* Property = *PropIt;
+					if (Property->GetFName() == GET_MEMBER_NAME_CHECKED(URendererSettings, bEnableNonNaniteVSM))
+					{
+						UpdateSinglePropertyInConfigFile(Property, GetDefaultConfigFilename());
+					}
 				}
 			}
-
-			if (Value != 0)
-			{
-				FMessageDialog::Open(EAppMsgType::Ok, LOCTEXT("Virtual Shadow Maps (Beta) Enabled", "Note that the virtual shadow map system is in beta and may have a performance impact, even if the run-time toggle (r.Shadow.Virtual.Enable) is used to turn them off."));
-				// TODO: put this in release notes? Note: When enabling the virtual shadow map system (beta), both 'r.Shadow.Virtual.Enable' and 'r.Shadow.Virtual.NonNaniteVSM' are turned on for the project. 'r.Shadow.Virtual.Enable' can be toggled at run-time, but there may still be some overhead as long as 'r.Shadow.Virtual.NonNaniteVSM' is turned on (requires restart)
-			}
 		}
 
 		if (PropertyChangedEvent.Property->GetFName() == GET_MEMBER_NAME_CHECKED(URendererSettings, VirtualTextureTileSize))
@@ -287,6 +280,11 @@ bool URendererSettings::CanEditChange(const FProperty* InProperty) const
 		return ParentVal && bSupportSkyAtmosphere;
 	}
 
+	if ((InProperty->GetFName() == GET_MEMBER_NAME_CHECKED(URendererSettings, bEnableNonNaniteVSM)))
+	{
+		return ParentVal && ShadowMapMethod == EShadowMapMethod::VirtualShadowMaps;
+	}
+
 	return ParentVal;
 }
 #endif // #if WITH_EDITOR
