Add a pawn controller changed callback, similar to the one on game instance. This can be used to detect unpossess
Store the last controller on pawn, this is the only way to reliably handle possession changes on the client without refactoring all of the replication logic

diff --git a/Engine/Source/Runtime/Engine/Classes/GameFramework/Pawn.h b/Engine/Source/Runtime/Engine/Classes/GameFramework/Pawn.h
--- a/Engine/Source/Runtime/Engine/Classes/GameFramework/Pawn.h
+++ b/Engine/Source/Runtime/Engine/Classes/GameFramework/Pawn.h
@@ -28,6 +28,7 @@ class UPrimitiveComponent;
 ENGINE_API DECLARE_LOG_CATEGORY_EXTERN(LogDamage, Warning, All);
 
 DECLARE_DYNAMIC_MULTICAST_SPARSE_DELEGATE_OneParam(FPawnRestartedSignature, APawn, ReceiveRestartedDelegate, APawn*, Pawn);
+DECLARE_DYNAMIC_MULTICAST_SPARSE_DELEGATE_ThreeParams(FPawnControllerChangedSignature, APawn, ReceiveControllerChangedDelegate, APawn*, Pawn, AController*, OldController, AController*, NewController);
 
 /** 
  * Pawn is the base class of all actors that can be possessed by players or AI.
@@ -163,6 +164,10 @@ public:
 	UPROPERTY(replicatedUsing=OnRep_Controller)
 	TObjectPtr<AController> Controller;
 
+	/** Previous controller that was controlling this pawn since the last controller change notification */
+	UPROPERTY(transient)
+	TObjectPtr<AController> LastController;
+
 	/** Max difference between pawn's Rotation.Yaw and GetDesiredRotation().Yaw for pawn to be considered as having reached its desired rotation */
 	float AllowedYawError;
 
@@ -179,9 +184,6 @@ public:
 	 */
 	void SetRemoteViewPitch(float NewRemoteViewPitch);
 
-	/** Called when our Controller no longer possesses us. */
-	virtual void UnPossessed();
-
 	/** Return Physics Volume for this Pawn */
 	virtual APhysicsVolume* GetPawnPhysicsVolume() const;
 
@@ -302,6 +304,7 @@ public:
 	/** Whether this Pawn's input handling is enabled.  Pawn must still be possessed to get input even if this is true */
 	bool InputEnabled() const { return bInputEnabled; }
 
+
 	/** 
 	 * Called when this Pawn is possessed. Only called on the server (or in standalone).
 	 * @param NewController The controller possessing this pawn
@@ -312,18 +315,50 @@ public:
 	UFUNCTION(BlueprintImplementableEvent, BlueprintAuthorityOnly, meta=(DisplayName= "Possessed"))
 	void ReceivePossessed(AController* NewController);
 
+	/** Called when our Controller no longer possesses us. Only called on the server (or in standalone). */
+	virtual void UnPossessed();
+
 	/** Event called when the Pawn is no longer possessed by a Controller. Only called on the server (or in standalone) */
 	UFUNCTION(BlueprintImplementableEvent, BlueprintAuthorityOnly, meta=(DisplayName= "Unpossessed"))
 	void ReceiveUnpossessed(AController* OldController);
 
-	/** Event called when a pawn has been restarted, usually by a possession change. This is called on the server for all pawns and the owning client for player pawns */
+
+	/** Event called after a pawn's controller has changed, on the server and owning client. This will happen at the same time as the delegate on GameInstance */
+	UFUNCTION(BlueprintImplementableEvent)
+	void ReceiveControllerChanged(AController* OldController, AController* NewController);
+
+	/** Event called after a pawn's controller has changed, on the server and owning client. This will happen at the same time as the delegate on GameInstance */
+	UPROPERTY(BlueprintAssignable, Category = Pawn)
+	FPawnControllerChangedSignature ReceiveControllerChangedDelegate;
+
+	/** Call to notify about a change in controller, on both the server and owning client. This calls the above event and delegate */
+	virtual void NotifyControllerChanged();
+
+
+	/** Event called after a pawn has been restarted, usually by a possession change. This is called on the server for all pawns and the owning client for player pawns */
 	UFUNCTION(BlueprintImplementableEvent)
 	void ReceiveRestarted();
 
-	/** Event called when a pawn has been restarted, usually by a possession change. This is called on the server for all pawns and the owning client for player pawns */
+	/** Event called after a pawn has been restarted, usually by a possession change. This is called on the server for all pawns and the owning client for player pawns */
 	UPROPERTY(BlueprintAssignable, Category = Pawn)
 	FPawnRestartedSignature ReceiveRestartedDelegate;
 
+	/**
+	 * Notifies other systems that a pawn has been restarted. By default this is called on the server for all pawns and the owning client for player pawns.
+	 * This can be overridden by subclasses to delay the notification of restart until data has loaded/replicated
+	 */
+	virtual void NotifyRestarted();
+
+	/** Called when the Pawn is being restarted (usually by being possessed by a Controller). This is called on the server for all pawns and the owning client for player pawns  */
+	virtual void Restart();
+
+	/** Called on the owning client of a player-controlled Pawn when it is restarted, this calls Restart() */
+	virtual void PawnClientRestart();
+
+	/** Wrapper function to call correct restart functions, enable bCallClientRestart if this is a locally owned player pawn or equivalent */
+	void DispatchRestart(bool bCallClientRestart);
+
+
 	/** Returns true if controlled by a local (not network) Controller.	 */
 	UFUNCTION(BlueprintPure, Category=Pawn)
 	virtual bool IsLocallyControlled() const;
@@ -356,21 +391,6 @@ public:
 	/** Returns true if player is viewing this Pawn in FreeCam */
 	virtual bool InFreeCam() const;
 
-	/** Called when the Pawn is being restarted (usually by being possessed by a Controller). This is called on the server for all pawns and the owning client for player pawns  */
-	virtual void Restart();
-
-	/** Called on the owning client of a player-controlled Pawn when it is restarted, this calls Restart() */
-	virtual void PawnClientRestart();
-
-	/** 
-	 * Notifies other systems that a pawn has been restarted. By default this is called on the server for all pawns and the owning client for player pawns.
-	 * This can be overridden by subclasses to delay the notification of restart until data has loaded/replicated
-	 */
-	virtual void NotifyRestarted();
-
-	/** Wrapper function to call correct restart functions, enable bCallClientRestart if this is a locally owned player pawn or equivalent */
-	void DispatchRestart(bool bCallClientRestart);
-
 	/** Updates Pawn's rotation to the given rotation, assumed to be the Controller's ControlRotation. Respects the bUseControllerRotation* settings. */
 	virtual void FaceRotation(FRotator NewControlRotation, float DeltaTime = 0.f);
 
diff --git a/Engine/Source/Runtime/Engine/Private/Pawn.cpp b/Engine/Source/Runtime/Engine/Private/Pawn.cpp
--- a/Engine/Source/Runtime/Engine/Private/Pawn.cpp
+++ b/Engine/Source/Runtime/Engine/Private/Pawn.cpp
@@ -504,10 +504,7 @@ void APawn::OnRep_Controller()
 
 	if (bNotifyControllerChange)
 	{
-		if (UGameInstance* GameInstance = GetGameInstance())
-		{
-			GameInstance->GetOnPawnControllerChanged().Broadcast(this, Controller);
-		}
+		NotifyControllerChanged();
 	}
 }
 
@@ -561,10 +558,7 @@ void APawn::PossessedBy(AController* NewController)
 	{
 		ReceivePossessed(Controller);
 	
-		if (UGameInstance* GameInstance = GetGameInstance())
-		{
-			GameInstance->GetOnPawnControllerChanged().Broadcast(this, Controller);
-		}
+		NotifyControllerChanged();
 	}
 }
 
@@ -587,15 +581,24 @@ void APawn::UnPossessed()
 		ReceiveUnpossessed(OldController);
 	}
 
+	NotifyControllerChanged();
+
+	ConsumeMovementInputVector();
+}
+
+void APawn::NotifyControllerChanged()
+{
+	ReceiveControllerChanged(LastController, Controller);
+	ReceiveControllerChangedDelegate.Broadcast(this, LastController, Controller);
 	if (UGameInstance* GameInstance = GetGameInstance())
 	{
-		GameInstance->GetOnPawnControllerChanged().Broadcast(this, nullptr);
+		GameInstance->GetOnPawnControllerChanged().Broadcast(this, Controller);
 	}
 
-	ConsumeMovementInputVector();
+	// Update the cached controller
+	LastController = Controller;
 }
 
-
 class UNetConnection* APawn::GetNetConnection() const
 {
 	// if have a controller, it has the net connection
