Adding a direct rotator->quat conversion node for Blueprints.

diff --git a/Engine/Source/Runtime/Engine/Classes/Kismet/KismetMathLibrary.h b/Engine/Source/Runtime/Engine/Classes/Kismet/KismetMathLibrary.h
--- a/Engine/Source/Runtime/Engine/Classes/Kismet/KismetMathLibrary.h
+++ b/Engine/Source/Runtime/Engine/Classes/Kismet/KismetMathLibrary.h
@@ -2847,6 +2847,10 @@ class ENGINE_API UKismetMathLibrary : public UBlueprintFunctionLibrary
 	UFUNCTION(BlueprintPure, meta = (DisplayName = "ToRotator (Quat)", CompactNodeTitle = "->", ScriptMethod = "Rotator", Keywords = "cast convert", BlueprintAutocast), Category = "Math|Conversions")
 	static FRotator Quat_Rotator(const FQuat& Q);
 
+	/** Convert to Quaternion representation of this Rotator. */
+	UFUNCTION(BlueprintPure, meta=(DisplayName = "ToQuaternion (Rotator)", CompactNodeTitle = "->", ScriptMethod = "Quaternion", Keywords="cast convert", BlueprintAutocast), Category="Math|Conversions")
+	static FQuat Conv_RotatorToQuaternion(FRotator InRot);
+
 	/**
 	 * Get the length of the quaternion.
 	 *
diff --git a/Engine/Source/Runtime/Engine/Classes/Kismet/KismetMathLibrary.inl b/Engine/Source/Runtime/Engine/Classes/Kismet/KismetMathLibrary.inl
--- a/Engine/Source/Runtime/Engine/Classes/Kismet/KismetMathLibrary.inl
+++ b/Engine/Source/Runtime/Engine/Classes/Kismet/KismetMathLibrary.inl
@@ -3350,6 +3350,12 @@ FRotator UKismetMathLibrary::Quat_Rotator(const FQuat& Q)
 	return Q.Rotator();
 }
 
+KISMET_MATH_FORCEINLINE
+FQuat UKismetMathLibrary::Conv_RotatorToQuaternion(FRotator InRot)
+{
+	return InRot.Quaternion();
+}
+
 KISMET_MATH_FORCEINLINE
 float UKismetMathLibrary::Quat_Size(const FQuat& Q)
 {
