Add a log message to track round trip time for ODSC requests with the cooking server.

diff --git a/Engine/Source/Runtime/Engine/Private/ODSC/ODSCThread.cpp b/Engine/Source/Runtime/Engine/Private/ODSC/ODSCThread.cpp
--- a/Engine/Source/Runtime/Engine/Private/ODSC/ODSCThread.cpp
+++ b/Engine/Source/Runtime/Engine/Private/ODSC/ODSCThread.cpp
@@ -29,6 +29,9 @@ FODSCMessageHandler::FODSCMessageHandler(const TArray<FString>& InMaterials, ESh
 
 void FODSCMessageHandler::FillPayload(FArchive& Payload)
 {
+	// When did we start this request?
+	RequestStartTime = FPlatformTime::Seconds();
+
 	Payload << MaterialsToLoad;
 	uint32 ConvertedShaderPlatform = (uint32)ShaderPlatform;
 	Payload << ConvertedShaderPlatform;
@@ -38,6 +41,8 @@ void FODSCMessageHandler::FillPayload(FArchive& Payload)
 
 void FODSCMessageHandler::ProcessResponse(FArchive& Response)
 {
+	UE_LOG(LogODSC, Display, TEXT("Received response in %lf seconds."), FPlatformTime::Seconds() - RequestStartTime);
+
 	// pull back the compiled mesh material data (if any)
 	Response << OutMeshMaterialMaps;
 }
diff --git a/Engine/Source/Runtime/Engine/Private/ODSC/ODSCThread.h b/Engine/Source/Runtime/Engine/Private/ODSC/ODSCThread.h
--- a/Engine/Source/Runtime/Engine/Private/ODSC/ODSCThread.h
+++ b/Engine/Source/Runtime/Engine/Private/ODSC/ODSCThread.h
@@ -31,6 +31,9 @@ public:
 	bool ReloadGlobalShaders() const;
 
 private:
+	/** The time when this command was issued.  This isn't serialized to the cooking server. */
+	double RequestStartTime = 0.0;
+
 	/** The materials we send over the network and expect maps for on the return */
 	TArray<FString> MaterialsToLoad;
 
