Add error message when mipmaps are built and put into DDC but then can not be read back from DDC due to a read-only DDC.
#rnx

diff --git a/Engine/Source/Runtime/Engine/Private/TextureDerivedDataTask.cpp b/Engine/Source/Runtime/Engine/Private/TextureDerivedDataTask.cpp
--- a/Engine/Source/Runtime/Engine/Private/TextureDerivedDataTask.cpp
+++ b/Engine/Source/Runtime/Engine/Private/TextureDerivedDataTask.cpp
@@ -299,6 +299,10 @@ void FTextureCacheDerivedDataWorker::BuildTexture(bool bReplaceExistingDDC)
 			{
 				const bool bInlineMips = (CacheFlags & ETextureCacheFlags::InlineMips) != 0;
 				bSucceeded = !bInlineMips || DerivedData->TryInlineMipData(BuildSettingsPerLayer[0].LODBiasWithCinematicMips, &Texture);
+				if (!bSucceeded)
+				{
+					UE_LOG(LogTexture, Error, TEXT("Failed to put and then read back mipmap data from DDC for %s"), *Texture.GetPathName());
+				}
 			}
 			else
 			{
@@ -385,6 +389,10 @@ void FTextureCacheDerivedDataWorker::BuildTexture(bool bReplaceExistingDDC)
 		{
 			const bool bInlineMips = (CacheFlags & ETextureCacheFlags::InlineMips) != 0;
 			bSucceeded = !bInlineMips || DerivedData->TryInlineMipData(BuildSettingsPerLayer[0].LODBiasWithCinematicMips, &Texture);
+			if (!bSucceeded)
+			{
+				UE_LOG(LogTexture, Error, TEXT("Failed to put and then read back mipmap data from DDC for %s"), *Texture.GetPathName());
+			}
 		}
 		else
 		{
