Prevent non thread-safe bulkdata from being detached during async static mesh cancellation
  - This is expected to be removed once we have a thread-safe bulkdata implementation (i.e. virtualbulkdata).
  - Remove hack from StaticMeshCompiler

TESTS
  - Successfully execute -run=WrangleContent on special project

diff --git a/Engine/Source/Runtime/Engine/Classes/Engine/StaticMesh.h b/Engine/Source/Runtime/Engine/Classes/Engine/StaticMesh.h
--- a/Engine/Source/Runtime/Engine/Classes/Engine/StaticMesh.h
+++ b/Engine/Source/Runtime/Engine/Classes/Engine/StaticMesh.h
@@ -6,8 +6,10 @@
 #include "UObject/ObjectMacros.h"
 #include "UObject/UObjectGlobals.h"
 #include "UObject/Object.h"
+#include "UObject/StrongObjectPtr.h"
 #include "Misc/Guid.h"
 #include "Templates/SubclassOf.h"
+#include "Templates/UnrealTemplate.h"
 #include "Engine/EngineTypes.h"
 #include "UObject/ScriptMacros.h"
 #include "Interfaces/Interface_AssetUserData.h"
@@ -32,6 +34,7 @@ class UNavCollisionBase;
 class UStaticMeshComponent;
 class UStaticMeshDescription;
 class FStaticMeshUpdate;
+class UPackage;
 struct FMeshDescription;
 struct FStaticMeshLODResources;
 
@@ -1825,18 +1828,37 @@ private:
 	void FinishPostLoadInternal(FStaticMeshPostLoadContext& Context);
 };
 
-class FStaticMeshPostLoadContext
+class FStaticMeshCompilationContext
 {
 public:
+	FStaticMeshCompilationContext() = default;
+	// Non-copyable
+	FStaticMeshCompilationContext(const FStaticMeshCompilationContext&) = delete;
+	FStaticMeshCompilationContext& operator=(const FStaticMeshCompilationContext&) = delete;
+	// Movable
+	FStaticMeshCompilationContext(FStaticMeshCompilationContext&&) = default;
+	FStaticMeshCompilationContext& operator=(FStaticMeshCompilationContext&&) = default;
+
 	bool bShouldComputeExtendedBounds = false;
+
+	// Prevent the garbage collection of the static mesh's package while the async compilation is in progress.
+	// This may happen because we allow garbage collection of UStaticMesh to be initiated even in the presence of
+	// an async operation in order to cancel it if possible. In that case, we need to wait until the async operation
+	// is fully cancelled before trying to destroy the UPackage otherwise the bulk data of the static mesh could
+	// become inaccessible while the async operation is still using it.
+	TStrongObjectPtr<UPackage> StaticMeshOwner;
+};
+
+class FStaticMeshPostLoadContext : public FStaticMeshCompilationContext
+{
+public:
 	bool bNeedsMeshUVDensityFix = false;
 	bool bNeedsMaterialFixup = false;
 	bool bIsCookedForEditor = false;
 };
 
-class FStaticMeshBuildContext
+class FStaticMeshBuildContext : public FStaticMeshCompilationContext
 {
 public:
 	bool bHasRenderDataChanged = false;
-	bool bShouldComputeExtendedBounds = true;
 };
\ No newline at end of file
diff --git a/Engine/Source/Runtime/Engine/Private/StaticMesh.cpp b/Engine/Source/Runtime/Engine/Private/StaticMesh.cpp
--- a/Engine/Source/Runtime/Engine/Private/StaticMesh.cpp
+++ b/Engine/Source/Runtime/Engine/Private/StaticMesh.cpp
@@ -5049,7 +5049,8 @@ void UStaticMesh::PostLoad()
 			FStaticMeshCompilingManager::Get().FinishCompilation({ ComplexCollisionMesh });
 		}
 
-		AsyncTask = MakeUnique<FStaticMeshAsyncBuildTask>(this, MakeUnique<FStaticMeshPostLoadContext>(Context));
+		Context.StaticMeshOwner.Reset(GetPackage());
+		AsyncTask = MakeUnique<FStaticMeshAsyncBuildTask>(this, MakeUnique<FStaticMeshPostLoadContext>(MoveTemp(Context)));
 		AsyncTask->StartBackgroundTask(StaticMeshThreadPool, BasePriority);
 		FStaticMeshCompilingManager::Get().AddStaticMeshes({this});
 	}
diff --git a/Engine/Source/Runtime/Engine/Private/StaticMeshBuild.cpp b/Engine/Source/Runtime/Engine/Private/StaticMeshBuild.cpp
--- a/Engine/Source/Runtime/Engine/Private/StaticMeshBuild.cpp
+++ b/Engine/Source/Runtime/Engine/Private/StaticMeshBuild.cpp
@@ -247,6 +247,7 @@ void UStaticMesh::BatchBuild(const TArray<UStaticMesh*>& InStaticMeshes, bool bI
 				if (OutErrors == nullptr && FStaticMeshCompilingManager::Get().IsAsyncCompilationAllowed(StaticMesh))
 				{
 					TUniquePtr<FStaticMeshBuildContext> Context = MakeUnique<FStaticMeshBuildContext>();
+					Context->StaticMeshOwner.Reset(StaticMesh->GetPackage());
 					StaticMesh->BeginBuildInternal(Context.Get());
 
 					FQueuedThreadPool* StaticMeshThreadPool = FStaticMeshCompilingManager::Get().GetThreadPool();
diff --git a/Engine/Source/Runtime/Engine/Private/StaticMeshCompiler.cpp b/Engine/Source/Runtime/Engine/Private/StaticMeshCompiler.cpp
--- a/Engine/Source/Runtime/Engine/Private/StaticMeshCompiler.cpp
+++ b/Engine/Source/Runtime/Engine/Private/StaticMeshCompiler.cpp
@@ -157,11 +157,6 @@ bool FStaticMeshCompilingManager::IsAsyncStaticMeshCompilationEnabled() const
 		return false;
 	}
 
-	if (GetRunningCommandletClass() && (GetRunningCommandletClass()->GetName() == TEXT("WrangleContentCommandlet")))
-	{
-		return false;
-	}
-
 	StaticMeshCompilingManagerImpl::EnsureInitializedCVars();
 
 	return CVarAsyncStaticMeshStandard.AsyncCompilation.GetValueOnAnyThread() != 0;
