Actor Descriptors Debugging
- Added the wp.Editor.DumpActorDescs <file> to dump all actor descriptors to a CSV file (easier for searching multiple actors, etc).
- Added the wp.Editor.DumpActorDesc <actorpath> to dump a specific actor descriptor to the log output from it's full path.

diff --git a/Engine/Source/Runtime/Engine/Private/WorldPartition/WorldPartition.cpp b/Engine/Source/Runtime/Engine/Private/WorldPartition/WorldPartition.cpp
--- a/Engine/Source/Runtime/Engine/Private/WorldPartition/WorldPartition.cpp
+++ b/Engine/Source/Runtime/Engine/Private/WorldPartition/WorldPartition.cpp
@@ -43,6 +43,43 @@ DEFINE_LOG_CATEGORY(LogWorldPartition);
 #define LOCTEXT_NAMESPACE "WorldPartition"
 
 #if WITH_EDITOR
+static FAutoConsoleCommand DumpActorDesc(
+	TEXT("wp.Editor.DumpActorDesc"),
+	TEXT("Dump a specific actor descriptor on the console."),
+	FConsoleCommandWithArgsDelegate::CreateLambda([](const TArray<FString>& Args)
+	{
+		if (Args.Num() > 0)
+		{
+			if (UWorld* World = GEditor->GetEditorWorldContext().World())
+			{
+				if (!World->IsGameWorld())
+				{
+					if (UWorldPartition* WorldPartition = World->GetWorldPartition())
+					{
+						if (const FWorldPartitionActorDesc* ActorDesc = WorldPartition->GetActorDesc(Args[0]))
+						{
+							UE_LOG(LogWorldPartition, Log, TEXT("Guid, Class, Name, Package, BVCenterX, BVCenterY, BVCenterZ, BVExtentX, BVExtentY, BVExtentZ"));
+							UE_LOG(LogWorldPartition, Log, TEXT("%s, %s, %s, %s, %s, %.2f, %.2f, %.2f, %.2f, %.2f, %.2f"),
+								*ActorDesc->GetGuid().ToString(), 
+								*ActorDesc->GetClass().ToString(), 
+								*ActorDesc->GetActorName().ToString(), 
+								*ActorDesc->GetActorPackage().ToString(), 
+								*ActorDesc->GetActorLabel().ToString(),
+								ActorDesc->GetBounds().GetCenter().X,
+								ActorDesc->GetBounds().GetCenter().Y,
+								ActorDesc->GetBounds().GetCenter().Z,
+								ActorDesc->GetBounds().GetExtent().X,
+								ActorDesc->GetBounds().GetExtent().Y,
+								ActorDesc->GetBounds().GetExtent().Z
+							);
+						}
+					}
+				}
+			}
+		}
+	})
+);
+
 static FAutoConsoleCommand DumpActorDescs(
 	TEXT("wp.Editor.DumpActorDescs"),
 	TEXT("Dump the list of actor descriptors in a CSV file."),
@@ -1106,7 +1143,7 @@ void UWorldPartition::DumpActorDescs(const FString& Path)
 {
 	if (FArchive* LogFile = IFileManager::Get().CreateFileWriter(*Path))
 	{
-		FString LineEntry = TEXT("Guid, Class, Name, BVCenterX, BVCenterY, BVCenterZ, BVExtentX, BVExtentY, BVExtentZ") LINE_TERMINATOR;
+		FString LineEntry = TEXT("Guid, Class, Name, Package, BVCenterX, BVCenterY, BVCenterZ, BVExtentX, BVExtentY, BVExtentZ") LINE_TERMINATOR;
 		LogFile->Serialize(TCHAR_TO_ANSI(*LineEntry), LineEntry.Len());
 
 		TArray<const FWorldPartitionActorDesc*> ActorDescs;
@@ -1121,11 +1158,12 @@ void UWorldPartition::DumpActorDescs(const FString& Path)
 		for (const FWorldPartitionActorDesc* ActorDescIterator : ActorDescs)
 		{
 			LineEntry = FString::Printf(
-				TEXT("%s, %s, %s, %s, %.2f, %.2f, %.2f, %.2f, %.2f, %.2f") LINE_TERMINATOR, 
+				TEXT("%s, %s, %s, %s, %s, %.2f, %.2f, %.2f, %.2f, %.2f, %.2f") LINE_TERMINATOR, 
 				*ActorDescIterator->GetGuid().ToString(), 
 				*ActorDescIterator->GetClass().ToString(), 
-				*FPaths::GetExtension(ActorDescIterator->GetActorPath().ToString()), 
-				ActorDescIterator->GetActor() ? *ActorDescIterator->GetActor()->GetActorLabel(false) : TEXT("None"),
+				*ActorDescIterator->GetActorName().ToString(), 
+				*ActorDescIterator->GetActorPackage().ToString(), 
+				*ActorDescIterator->GetActorLabel().ToString(),
 				ActorDescIterator->GetBounds().GetCenter().X,
 				ActorDescIterator->GetBounds().GetCenter().Y,
 				ActorDescIterator->GetBounds().GetCenter().Z,
