Improve ODSC logging information to make it more obvious when a shader request is happening and what shaders have been requested.

- Making some logging verbose.
- Minor formatting to output.

diff --git a/Engine/Source/Editor/UnrealEd/Private/CookOnTheFlyServer.cpp b/Engine/Source/Editor/UnrealEd/Private/CookOnTheFlyServer.cpp
--- a/Engine/Source/Editor/UnrealEd/Private/CookOnTheFlyServer.cpp
+++ b/Engine/Source/Editor/UnrealEd/Private/CookOnTheFlyServer.cpp
@@ -7995,7 +7995,7 @@ void UCookOnTheFlyServer::HandleNetworkFileServerRecompileShaders(const FShaderR
 	// if we aren't in the game thread, we need to push this over to the game thread and wait for it to finish
 	if (!IsInGameThread())
 	{
-		UE_LOG(LogCook, Display, TEXT("Got a recompile request on non-game thread"));
+		UE_LOG(LogCook, Verbose, TEXT("Got a recompile request on non-game thread"));
 
 		// make a new request
 		UE::Cook::FRecompileRequest* Request = new UE::Cook::FRecompileRequest;
@@ -8011,7 +8011,7 @@ void UCookOnTheFlyServer::HandleNetworkFileServerRecompileShaders(const FShaderR
 			FPlatformProcess::Sleep(0);
 		}
 		delete Request;
-		UE_LOG(LogCook, Display, TEXT("Completed recompile..."));
+		UE_LOG(LogCook, Verbose, TEXT("Completed recompile..."));
 
 		// at this point, we are done on the game thread, and ModifiedFiles will have been filled out
 		return;
diff --git a/Engine/Source/Runtime/Engine/Private/ShaderCompiler/ShaderCompiler.cpp b/Engine/Source/Runtime/Engine/Private/ShaderCompiler/ShaderCompiler.cpp
--- a/Engine/Source/Runtime/Engine/Private/ShaderCompiler/ShaderCompiler.cpp
+++ b/Engine/Source/Runtime/Engine/Private/ShaderCompiler/ShaderCompiler.cpp
@@ -6159,23 +6159,27 @@ void RecompileShadersForRemote(
 		return;
 	}
 
+	UE_LOG(LogShaders, Display, TEXT(""));
+	UE_LOG(LogShaders, Display, TEXT("********************************"));
+	UE_LOG(LogShaders, Display, TEXT("Received compile shader request."));
+
 	const bool bPreviousState = GShaderCompilingManager->IsShaderCompilationSkipped();
 	GShaderCompilingManager->SkipShaderCompilation(false);
 
 	TArray<FName> DesiredShaderFormats;
 	TargetPlatform->GetAllTargetedShaderFormats(DesiredShaderFormats);
 
-	UE_LOG(LogShaders, Display, TEXT("Loading %d materials..."), MaterialsToLoad.Num());
+	UE_LOG(LogShaders, Verbose, TEXT("Loading %d materials..."), MaterialsToLoad.Num());
 	// make sure all materials the client has loaded will be processed
 	TArray<UMaterialInterface*> MaterialsToCompile;
 
 	for (int32 Index = 0; Index < MaterialsToLoad.Num(); Index++)
 	{
-		UE_LOG(LogShaders, Display, TEXT("   --> %s"), *MaterialsToLoad[Index]);
+		UE_LOG(LogShaders, Verbose, TEXT("   --> %s"), *MaterialsToLoad[Index]);
 		MaterialsToCompile.Add(LoadObject<UMaterialInterface>(NULL, *MaterialsToLoad[Index]));
 	}
 
-	UE_LOG(LogShaders, Display, TEXT("  Done!"))
+	UE_LOG(LogShaders, Verbose, TEXT("  Done!"))
 
 	// figure out which shaders are out of date
 	TArray<const FShaderType*> OutdatedShaderTypes;
@@ -6198,9 +6202,9 @@ void RecompileShadersForRemote(
 
 	for (const FODSCRequestPayload& payload: ShadersToRecompile)
 	{
-		UE_LOG(LogShaders, Display, TEXT("Material: %s "), *payload.MaterialName);
-		UE_LOG(LogShaders, Display, TEXT("VFType: %s "), *payload.VertexFactoryName);
-		UE_LOG(LogShaders, Display, TEXT("Pipeline: %s "), *payload.PipelineName);
+		UE_LOG(LogShaders, Display, TEXT(""));
+		UE_LOG(LogShaders, Display, TEXT("\tMaterial:    %s "), *payload.MaterialName);
+		UE_LOG(LogShaders, Display, TEXT("\tVF Type:     %s "), *payload.VertexFactoryName);
 
 		MaterialsToCompile.Add(LoadObject<UMaterialInterface>(NULL, *payload.MaterialName));
 
@@ -6218,7 +6222,7 @@ void RecompileShadersForRemote(
 
 		for (const FString& ShaderTypeName : payload.ShaderTypeNames)
 		{
-			UE_LOG(LogShaders, Display, TEXT("\tShaderType: %s"), *ShaderTypeName);
+			UE_LOG(LogShaders, Display, TEXT("\tShader Type: %s"), *ShaderTypeName);
 
 			const FShaderType* ShaderType = FShaderType::GetShaderTypeByName(*ShaderTypeName);
 			if (ShaderType)
@@ -6278,6 +6282,9 @@ void RecompileShadersForRemote(
 		}
 	}
 
+	UE_LOG(LogShaders, Display, TEXT(""));
+	UE_LOG(LogShaders, Display, TEXT("Finished shader compile request."));
+
 	// Restore compilation state.
 	GShaderCompilingManager->SkipShaderCompilation(bPreviousState);
 }
