Lazy-generate AnimDataModel transient data used by deprecated API, rather than explicitly during PostLoad

diff --git a/Engine/Source/Runtime/Engine/Classes/Animation/AnimData/AnimDataModel.h b/Engine/Source/Runtime/Engine/Classes/Animation/AnimData/AnimDataModel.h
--- a/Engine/Source/Runtime/Engine/Classes/Animation/AnimData/AnimDataModel.h
+++ b/Engine/Source/Runtime/Engine/Classes/Animation/AnimData/AnimDataModel.h
@@ -380,7 +380,8 @@ private:
 	}
 
 private:
-	void GenerateTransientData();
+	void GenerateTransientData() const;
+	void CheckTransientData() const;
 
 	UPROPERTY(Transient)
 	int32 BracketCounter = 0;
@@ -423,16 +424,18 @@ private:
 private:
 	/** Transient data, kept around for backward-compatibility */
 	UPROPERTY(Transient, VisibleAnywhere, Category = "Animation Data Model")
-	TArray<FRawAnimSequenceTrack> RawAnimationTracks;
+	mutable TArray<FRawAnimSequenceTrack> RawAnimationTracks;
 
 	UPROPERTY(Transient, VisibleAnywhere, Category = "Animation Data Model")
-	TArray<FName> RawAnimationTrackNames;
+	mutable TArray<FName> RawAnimationTrackNames;
 
 	UPROPERTY(Transient, VisibleAnywhere, Category = "Animation Data Model")
-	TArray<FTrackToSkeletonMap> RawAnimationTrackSkeletonMappings;
+	mutable TArray<FTrackToSkeletonMap> RawAnimationTrackSkeletonMappings;
+
+	mutable bool bHasTransientDataBeenGenerated = false;
 
 	UPROPERTY(Transient, VisibleAnywhere, Category = "Animation Data Model")
-	FRawCurveTracks RawCurveTracks;
+	mutable FRawCurveTracks RawCurveTracks;
 
 	friend class UAnimDataController;
 	friend class FAnimDataControllerTestBase;
diff --git a/Engine/Source/Runtime/Engine/Private/Animation/AnimDataModel.cpp b/Engine/Source/Runtime/Engine/Private/Animation/AnimDataModel.cpp
--- a/Engine/Source/Runtime/Engine/Private/Animation/AnimDataModel.cpp
+++ b/Engine/Source/Runtime/Engine/Private/Animation/AnimDataModel.cpp
@@ -13,8 +13,6 @@
 void UAnimDataModel::PostLoad()
 {
 	UObject::PostLoad();
-
-	GenerateTransientData();
 }
 
 void UAnimDataModel::PostDuplicate(bool bDuplicateForPIE)
@@ -346,36 +344,42 @@ FGuid UAnimDataModel::GenerateGuid() const
 
 const TArray<FRawAnimSequenceTrack>& UAnimDataModel::GetTransientRawAnimationTracks() const
 {
+	CheckTransientData();
 	return RawAnimationTracks;
 }
 
 const TArray<FName>& UAnimDataModel::GetTransientRawAnimationTrackNames() const
 {
+	CheckTransientData();
 	return RawAnimationTrackNames;
 }
 
 const TArray<FTrackToSkeletonMap>& UAnimDataModel::GetTransientRawAnimationTrackSkeletonMappings() const
 {
+	CheckTransientData();
 	return RawAnimationTrackSkeletonMappings;
 }
 
 FRawAnimSequenceTrack& UAnimDataModel::GetNonConstRawAnimationTrackByIndex(int32 TrackIndex)
 {
+	CheckTransientData();
 	checkf(BoneAnimationTracks.IsValidIndex(TrackIndex), TEXT("Invalid track index"));
 	return BoneAnimationTracks[TrackIndex].InternalTrackData;
 }
 
 const FRawCurveTracks& UAnimDataModel::GetTransientRawCurveTracks() const
 {
+	CheckTransientData();
 	return RawCurveTracks;
 }
 
 FAnimationCurveData& UAnimDataModel::GetNonConstCurveData()
 {
+	CheckTransientData();	
 	return CurveData;	
 }
 
-void UAnimDataModel::GenerateTransientData()
+void UAnimDataModel::GenerateTransientData() const
 {
 	RawAnimationTracks.Empty(BoneAnimationTracks.Num());
 	RawAnimationTrackNames.Empty(BoneAnimationTracks.Num());
@@ -392,6 +396,16 @@ void UAnimDataModel::GenerateTransientData()
 #if WITH_EDITOR
 	RawCurveTracks.TransformCurves = CurveData.TransformCurves;
 #endif
+	
+	bHasTransientDataBeenGenerated = true;
+}
+
+void UAnimDataModel::CheckTransientData() const
+{
+	if(!bHasTransientDataBeenGenerated)
+	{
+		GenerateTransientData();
+	}
 }
 
 FRichCurve* UAnimDataModel::GetMutableRichCurve(const FAnimationCurveIdentifier& CurveIdentifier)
