Frosty cooking fails due to critical error: LastKeyHandleIndex != 0 in KeyHandle.cpp
#fix implement better handling of CurveRenaming model notifies

diff --git a/Engine/Source/Runtime/Engine/Private/Animation/AnimSequenceBase.cpp b/Engine/Source/Runtime/Engine/Private/Animation/AnimSequenceBase.cpp
--- a/Engine/Source/Runtime/Engine/Private/Animation/AnimSequenceBase.cpp
+++ b/Engine/Source/Runtime/Engine/Private/Animation/AnimSequenceBase.cpp
@@ -1075,7 +1075,6 @@ void UAnimSequenceBase::OnModelModified(const EAnimDataModelNotifyType& NotifyTy
 		case EAnimDataModelNotifyType::CurveChanged:
 		case EAnimDataModelNotifyType::CurveRemoved:
 		case EAnimDataModelNotifyType::CurveFlagsChanged:
-		case EAnimDataModelNotifyType::CurveRenamed:
 		case EAnimDataModelNotifyType::CurveScaled:
 		{
 			if (NotifyCollector.IsNotWithinBracket())
@@ -1084,6 +1083,40 @@ void UAnimSequenceBase::OnModelModified(const EAnimDataModelNotifyType& NotifyTy
 			}
 			break;
 		}
+		
+		case EAnimDataModelNotifyType::CurveRenamed:
+		{
+			const FCurveRenamedPayload& TypedPayload = Payload.GetPayload<FCurveRenamedPayload>();
+			FAnimCurveBase* CurvePtr = [this, Identifier=TypedPayload.Identifier]() -> FAnimCurveBase*
+			{
+				if (Identifier.CurveType == ERawCurveTrackTypes::RCT_Float)
+				{
+					PRAGMA_DISABLE_DEPRECATION_WARNINGS
+					return RawCurveData.FloatCurves.FindByPredicate([SmartName=Identifier.InternalName](FFloatCurve& Curve)
+	                {
+	                    return Curve.Name == SmartName;
+	                });
+					PRAGMA_ENABLE_DEPRECATION_WARNINGS
+				}
+				else if (Identifier.CurveType == ERawCurveTrackTypes::RCT_Transform)
+				{
+					PRAGMA_DISABLE_DEPRECATION_WARNINGS
+					return RawCurveData.TransformCurves.FindByPredicate([SmartName=Identifier.InternalName](FTransformCurve& Curve)
+                    {
+                        return Curve.Name == SmartName;
+                    });
+					PRAGMA_ENABLE_DEPRECATION_WARNINGS
+				}
+
+				return nullptr;
+			}();
+				
+			if (CurvePtr)
+			{
+				CurvePtr->Name = TypedPayload.NewIdentifier.InternalName;
+			}
+			break;
+		}
 
 		case EAnimDataModelNotifyType::Populated:
 		{
@@ -1105,7 +1138,7 @@ void UAnimSequenceBase::OnModelModified(const EAnimDataModelNotifyType& NotifyTy
 		{
 			if (NotifyCollector.IsNotWithinBracket())
 			{
-				const auto CurveCopyNotifies = { EAnimDataModelNotifyType::CurveAdded, EAnimDataModelNotifyType::CurveChanged, EAnimDataModelNotifyType::CurveRemoved, EAnimDataModelNotifyType::CurveFlagsChanged, EAnimDataModelNotifyType::CurveRenamed, EAnimDataModelNotifyType::CurveScaled, EAnimDataModelNotifyType::Populated, EAnimDataModelNotifyType::Reset };
+				const auto CurveCopyNotifies = { EAnimDataModelNotifyType::CurveAdded, EAnimDataModelNotifyType::CurveChanged, EAnimDataModelNotifyType::CurveRemoved, EAnimDataModelNotifyType::CurveFlagsChanged, EAnimDataModelNotifyType::CurveScaled, EAnimDataModelNotifyType::Populated, EAnimDataModelNotifyType::Reset };
 
 				const auto LengthChangingNotifies = { EAnimDataModelNotifyType::SequenceLengthChanged, EAnimDataModelNotifyType::FrameRateChanged, EAnimDataModelNotifyType::Reset, EAnimDataModelNotifyType::Populated };
 
