Temp fix for CIS to work around a project-specific compile error.

#rnx

diff --git a/Engine/Source/Runtime/Engine/Classes/GameFramework/Pawn.h b/Engine/Source/Runtime/Engine/Classes/GameFramework/Pawn.h
--- a/Engine/Source/Runtime/Engine/Classes/GameFramework/Pawn.h
+++ b/Engine/Source/Runtime/Engine/Classes/GameFramework/Pawn.h
@@ -166,7 +166,7 @@ public:
 
 	/** Previous controller that was controlling this pawn since the last controller change notification */
 	UPROPERTY(transient)
-	TObjectPtr<AController> LastController;
+	TObjectPtr<AController> PreviousController;
 
 	/** Max difference between pawn's Rotation.Yaw and GetDesiredRotation().Yaw for pawn to be considered as having reached its desired rotation */
 	float AllowedYawError;
diff --git a/Engine/Source/Runtime/Engine/Private/Pawn.cpp b/Engine/Source/Runtime/Engine/Private/Pawn.cpp
--- a/Engine/Source/Runtime/Engine/Private/Pawn.cpp
+++ b/Engine/Source/Runtime/Engine/Private/Pawn.cpp
@@ -588,15 +588,15 @@ void APawn::UnPossessed()
 
 void APawn::NotifyControllerChanged()
 {
-	ReceiveControllerChanged(LastController, Controller);
-	ReceiveControllerChangedDelegate.Broadcast(this, LastController, Controller);
+	ReceiveControllerChanged(PreviousController, Controller);
+	ReceiveControllerChangedDelegate.Broadcast(this, PreviousController, Controller);
 	if (UGameInstance* GameInstance = GetGameInstance())
 	{
 		GameInstance->GetOnPawnControllerChanged().Broadcast(this, Controller);
 	}
 
 	// Update the cached controller
-	LastController = Controller;
+	PreviousController = Controller;
 }
 
 class UNetConnection* APawn::GetNetConnection() const
