Update some save game references to UE4 to be UE instead

diff --git a/Engine/Source/Runtime/Engine/Private/GameplayStatics.cpp b/Engine/Source/Runtime/Engine/Private/GameplayStatics.cpp
--- a/Engine/Source/Runtime/Engine/Private/GameplayStatics.cpp
+++ b/Engine/Source/Runtime/Engine/Private/GameplayStatics.cpp
@@ -89,7 +89,7 @@ struct FSaveGameHeader
 
 	int32 FileTypeTag;
 	int32 SaveGameFileVersion;
-	int32 PackageFileUE4Version;
+	int32 PackageFileUEVersion;
 	FEngineVersion SavedEngineVersion;
 	int32 CustomVersionFormat;
 	FCustomVersionContainer CustomVersions;
@@ -99,14 +99,14 @@ struct FSaveGameHeader
 FSaveGameHeader::FSaveGameHeader()
 	: FileTypeTag(0)
 	, SaveGameFileVersion(0)
-	, PackageFileUE4Version(0)
+	, PackageFileUEVersion(0)
 	, CustomVersionFormat(static_cast<int32>(ECustomVersionSerializationFormat::Unknown))
 {}
 
 FSaveGameHeader::FSaveGameHeader(TSubclassOf<USaveGame> ObjectType)
 	: FileTypeTag(UE_SAVEGAME_FILE_TYPE_TAG)
 	, SaveGameFileVersion(FSaveGameFileVersion::LatestVersion)
-	, PackageFileUE4Version(GPackageFileUEVersion)
+	, PackageFileUEVersion(GPackageFileUEVersion)
 	, SavedEngineVersion(FEngineVersion::Current())
 	, CustomVersionFormat(static_cast<int32>(ECustomVersionSerializationFormat::Latest))
 	, CustomVersions(FCurrentCustomVersions::GetAll())
@@ -117,7 +117,7 @@ void FSaveGameHeader::Empty()
 {
 	FileTypeTag = 0;
 	SaveGameFileVersion = 0;
-	PackageFileUE4Version = 0;
+	PackageFileUEVersion = 0;
 	SavedEngineVersion.Empty();
 	CustomVersionFormat = (int32)ECustomVersionSerializationFormat::Unknown;
 	CustomVersions.Empty();
@@ -144,7 +144,7 @@ void FSaveGameHeader::Read(FMemoryReader& MemoryReader)
 		// Note for 4.8 and beyond: if you get a crash loading a pre-4.8 version of your savegame file and 
 		// you don't want to delete it, try uncommenting these lines and changing them to use the version 
 		// information from your previous build. Then load and resave your savegame file.
-		//MemoryReader.SetUE4Ver(MyPreviousUE4Version);				// @see GPackageFileUEVersion
+		//MemoryReader.SetUE4Ver(MyPreviousUEVersion);				// @see GPackageFileUEVersion
 		//MemoryReader.SetEngineVer(MyPreviousEngineVersion);		// @see FEngineVersion::Current()
 	}
 	else
@@ -152,12 +152,12 @@ void FSaveGameHeader::Read(FMemoryReader& MemoryReader)
 		// Read version for this file format
 		MemoryReader << SaveGameFileVersion;
 
-		// Read engine and UE4 version information
-		MemoryReader << PackageFileUE4Version;
+		// Read engine and UE version information
+		MemoryReader << PackageFileUEVersion;
 
 		MemoryReader << SavedEngineVersion;
 
-		MemoryReader.SetUE4Ver(PackageFileUE4Version);
+		MemoryReader.SetUE4Ver(PackageFileUEVersion);
 		MemoryReader.SetEngineVer(SavedEngineVersion);
 
 		if (SaveGameFileVersion >= FSaveGameFileVersion::AddedCustomVersions)
@@ -182,8 +182,8 @@ void FSaveGameHeader::Write(FMemoryWriter& MemoryWriter)
 	// Write version for this file format
 	MemoryWriter << SaveGameFileVersion;
 
-	// Write out engine and UE4 version information
-	MemoryWriter << PackageFileUE4Version;
+	// Write out engine and UE version information
+	MemoryWriter << PackageFileUEVersion;
 	MemoryWriter << SavedEngineVersion;
 
 	// Write out custom version data
