Control Rig: Default Animating Rig option for SkelMeshComponent

diff --git a/Engine/Plugins/Experimental/ControlRig/Source/ControlRigEditor/Private/Sequencer/ControlRigParameterTrackEditor.cpp b/Engine/Plugins/Experimental/ControlRig/Source/ControlRigEditor/Private/Sequencer/ControlRigParameterTrackEditor.cpp
--- a/Engine/Plugins/Experimental/ControlRig/Source/ControlRigEditor/Private/Sequencer/ControlRigParameterTrackEditor.cpp
+++ b/Engine/Plugins/Experimental/ControlRig/Source/ControlRigEditor/Private/Sequencer/ControlRigParameterTrackEditor.cpp
@@ -1414,9 +1414,9 @@ void FControlRigParameterTrackEditor::AddTrackForComponent(USceneComponent* InCo
 {
 	if (USkeletalMeshComponent* SkelMeshComp = Cast<USkeletalMeshComponent>(InComponent))
 	{
-		if (SkelMeshComp->SkeletalMesh && !SkelMeshComp->SkeletalMesh->GetDefaultAnimatingRig().IsNull())
+		if (!SkelMeshComp->GetDefaultAnimatingRig().IsNull())
 		{
-			UObject* Object = SkelMeshComp->SkeletalMesh->GetDefaultAnimatingRig().LoadSynchronous();
+			UObject* Object = SkelMeshComp->GetDefaultAnimatingRig().LoadSynchronous();
 			if (Object != nullptr && (Object->IsA<UControlRigBlueprint>() || Object->IsA<UControlRigComponent>()))
 			{
 				FGuid Binding = GetSequencer()->GetHandleToObject(InComponent, true /*bCreateHandle*/);
diff --git a/Engine/Source/Runtime/Engine/Classes/Components/SkeletalMeshComponent.h b/Engine/Source/Runtime/Engine/Classes/Components/SkeletalMeshComponent.h
--- a/Engine/Source/Runtime/Engine/Classes/Components/SkeletalMeshComponent.h
+++ b/Engine/Source/Runtime/Engine/Classes/Components/SkeletalMeshComponent.h
@@ -770,6 +770,11 @@ protected:
 	/** If true, will play cloth in editor */
 	UPROPERTY(AdvancedDisplay, EditInstanceOnly, transient, Category = SkeletalMesh)
 	uint8 bUpdateClothInEditor : 1;
+
+	/** If true, DefaultAnimatingRigOverride will be used. If false, use the DefaultAnimatingRig in the SkeletalMesh */
+	UPROPERTY(AdvancedDisplay,EditAnywhere, Category = SkeletalMesh, meta = (InlineEditConditionToggle))
+	uint8 bOverrideDefaultAnimatingRig : 1;
+
 #endif
 
 	/** If true, OnSyncComponentToRBPhysics() notify will be called */
@@ -2485,6 +2490,30 @@ public:
 	FDelegateHandle RegisterOnTeleportDelegate(const FOnSkelMeshTeleported& Delegate);
 	void UnregisterOnTeleportDelegate(const FDelegateHandle& DelegateHandle);
 
+private:
+
+#if WITH_EDITORONLY_DATA
+
+	/** Default Animating Rig to Use if bOverrideDefaultAnimatingRig is true */
+	UPROPERTY(EditAnywhere, Category = AnimationRig, BlueprintGetter = GetDefaultAnimatingRigOverride, BlueprintSetter = SetDefaultAnimatingRigOverride, meta = (AllowedClasses = "ControlRigBlueprint"), meta = (editcondition = "bOverrideDefaultAnimatingRig"))
+	TSoftObjectPtr<UObject> DefaultAnimatingRigOverride;
+
+#endif
+
+#if WITH_EDITOR
+public:
+
+	UFUNCTION(BlueprintSetter)
+	void SetDefaultAnimatingRigOverride(TSoftObjectPtr<UObject> InAnimatingRig);
+
+	UFUNCTION(BlueprintGetter)
+	TSoftObjectPtr<UObject> GetDefaultAnimatingRigOverride() const;
+
+	UFUNCTION(BlueprintCallable, Category = "Components|SkeletalMesh")
+	TSoftObjectPtr<UObject> GetDefaultAnimatingRig() const;
+
+
+#endif
 private:
 
 	/** Multicaster fired when this component creates physics state (in case external objects rely on physics state)*/
diff --git a/Engine/Source/Runtime/Engine/Private/Components/SkeletalMeshComponent.cpp b/Engine/Source/Runtime/Engine/Private/Components/SkeletalMeshComponent.cpp
--- a/Engine/Source/Runtime/Engine/Private/Components/SkeletalMeshComponent.cpp
+++ b/Engine/Source/Runtime/Engine/Private/Components/SkeletalMeshComponent.cpp
@@ -241,6 +241,7 @@ USkeletalMeshComponent::USkeletalMeshComponent(const FObjectInitializer& ObjectI
 #if WITH_EDITORONLY_DATA
 	DefaultPlayRate_DEPRECATED = 1.0f;
 	bDefaultPlaying_DEPRECATED = true;
+	bOverrideDefaultAnimatingRig = false;
 #endif
 	bEnablePhysicsOnDedicatedServer = UPhysicsSettings::Get()->bSimulateSkeletalMeshOnDedicatedServer;
 	bEnableUpdateRateOptimizations = false;
@@ -1093,8 +1094,35 @@ void USkeletalMeshComponent::LoadedFromAnotherClass(const FName& OldClassName)
 		}
 	}
 }
+
+TSoftObjectPtr<UObject> USkeletalMeshComponent::GetDefaultAnimatingRig() const
+{
+	if (bOverrideDefaultAnimatingRig)
+	{
+		return DefaultAnimatingRigOverride;
+	}
+	if (SkeletalMesh)
+	{
+		return SkeletalMesh->GetDefaultAnimatingRig();
+	}
+	return nullptr;
+}
+
+void USkeletalMeshComponent::SetDefaultAnimatingRigOverride(TSoftObjectPtr<UObject> InAnimatingRig) 
+{
+	DefaultAnimatingRigOverride = InAnimatingRig;
+}
+
+TSoftObjectPtr<UObject> USkeletalMeshComponent::GetDefaultAnimatingRigOverride() const
+{
+	return DefaultAnimatingRigOverride;
+}
+
 #endif // WITH_EDITOR
 
+
+
+
 bool USkeletalMeshComponent::ShouldOnlyTickMontages(const float DeltaTime) const
 {
 	// Ignore DeltaSeconds == 0.f, as that is used when we want to force an update followed by RefreshBoneTransforms.
