Avoid excessive access tracking overhead during cook caused by TObjectPtr accesses in the material code.

diff --git a/Engine/Source/Developer/AssetTools/Private/AssetVtConversion.cpp b/Engine/Source/Developer/AssetTools/Private/AssetVtConversion.cpp
--- a/Engine/Source/Developer/AssetTools/Private/AssetVtConversion.cpp
+++ b/Engine/Source/Developer/AssetTools/Private/AssetVtConversion.cpp
@@ -687,12 +687,12 @@ void FVTConversionWorker::DoConvert()
 			MaterialTask.EnterProgressFrame();
 
 			bool FuncModified = false;
-			const TArray<UMaterialExpression*> *Expressions = Func->GetFunctionExpressions();
-			for (UMaterialExpression *Expr : *Expressions)
+			const TArray<TObjectPtr<UMaterialExpression>> *Expressions = Func->GetFunctionExpressions();
+			for (TObjectPtr<UMaterialExpression> Expr : *Expressions)
 			{
-				UMaterialExpressionTextureBase *TexExpr = Cast<UMaterialExpressionTextureBase>(Expr);
-				if (TexExpr)
+				if (Expr && Expr.IsA<UMaterialExpressionTextureBase>())
 				{
+					UMaterialExpressionTextureBase* TexExpr = (UMaterialExpressionTextureBase*)Expr.Get();
 					if (Textures.Contains(TexExpr->Texture))
 					{
 						UE_LOG(LogVirtualTextureConversion, Display, TEXT("Adjusting sampler %s."), *TexExpr->GetName());
diff --git a/Engine/Source/Editor/MaterialEditor/Private/MaterialInstanceEditor.cpp b/Engine/Source/Editor/MaterialEditor/Private/MaterialInstanceEditor.cpp
--- a/Engine/Source/Editor/MaterialEditor/Private/MaterialInstanceEditor.cpp
+++ b/Engine/Source/Editor/MaterialEditor/Private/MaterialInstanceEditor.cpp
@@ -309,8 +309,8 @@ void FMaterialInstanceEditor::InitEditorForMaterialFunction(UMaterialFunctionIns
 	{
 		BaseFunction = Instance->GetBaseFunction();
 	}
-	const TArray<UMaterialExpression*>* FunctionExpressions = BaseFunction ? BaseFunction->GetFunctionExpressions() : nullptr;
-	FunctionMaterialProxy->Expressions = FunctionExpressions ? *FunctionExpressions : TArray<UMaterialExpression*>();
+	const TArray<TObjectPtr<UMaterialExpression>>* FunctionExpressions = BaseFunction ? BaseFunction->GetFunctionExpressions() : nullptr;
+	FunctionMaterialProxy->Expressions = FunctionExpressions ? *FunctionExpressions : TArray<TObjectPtr<UMaterialExpression>>();
 
 	// Set expressions to be used with preview material
 	bool bSetPreviewExpression = false;
diff --git a/Engine/Source/Editor/VirtualTexturingEditor/Private/RuntimeVirtualTextureAssetTypeActions.cpp b/Engine/Source/Editor/VirtualTexturingEditor/Private/RuntimeVirtualTextureAssetTypeActions.cpp
--- a/Engine/Source/Editor/VirtualTexturingEditor/Private/RuntimeVirtualTextureAssetTypeActions.cpp
+++ b/Engine/Source/Editor/VirtualTexturingEditor/Private/RuntimeVirtualTextureAssetTypeActions.cpp
@@ -153,12 +153,12 @@ namespace
 				Task.EnterProgressFrame();
 
 				bool bFunctionModified = false;
-				const TArray<UMaterialExpression*> *Expressions = Function->GetFunctionExpressions();
-				for (UMaterialExpression *Expression : *Expressions)
+				const TArray<TObjectPtr<UMaterialExpression>> *Expressions = Function->GetFunctionExpressions();
+				for (TObjectPtr<UMaterialExpression> Expression : *Expressions)
 				{
-					UMaterialExpressionRuntimeVirtualTextureSample* RVTSampleExpression = Cast<UMaterialExpressionRuntimeVirtualTextureSample>(Expression);
-					if (RVTSampleExpression)
+					if (Expression && Expression.IsA<UMaterialExpressionRuntimeVirtualTextureSample>())
 					{
+						UMaterialExpressionRuntimeVirtualTextureSample* RVTSampleExpression = (UMaterialExpressionRuntimeVirtualTextureSample*)Expression.Get();
 						if (RuntimeVirtualTexture == RVTSampleExpression->VirtualTexture)
 						{
 							if (RVTSampleExpression->InitVirtualTextureDependentSettings())
diff --git a/Engine/Source/Runtime/Engine/Classes/Materials/Material.h b/Engine/Source/Runtime/Engine/Classes/Materials/Material.h
--- a/Engine/Source/Runtime/Engine/Classes/Materials/Material.h
+++ b/Engine/Source/Runtime/Engine/Classes/Materials/Material.h
@@ -1838,7 +1838,7 @@ private:
 #if WITH_EDITORONLY_DATA
 	/* Helper function to help finding expression GUID taking into account UMaterialExpressionMaterialFunctionCall */
 	template<typename ExpressionType>
-	ExpressionType* FindExpressionByGUIDRecursive(const FGuid &InGUID, const TArray<UMaterialExpression*>& InMaterialExpression)
+	ExpressionType* FindExpressionByGUIDRecursive(const FGuid &InGUID, const TArray<TObjectPtr<UMaterialExpression>>& InMaterialExpression)
 	{
 		for (int32 ExpressionIndex = 0; ExpressionIndex < InMaterialExpression.Num(); ExpressionIndex++)
 		{
@@ -1853,7 +1853,7 @@ private:
 			}
 			else if (MaterialFunctionCall && MaterialFunctionCall->MaterialFunction)
 			{
-				if (const TArray<UMaterialExpression*>* FunctionExpressions = MaterialFunctionCall->MaterialFunction->GetFunctionExpressions())
+				if (const TArray<TObjectPtr<UMaterialExpression>>* FunctionExpressions = MaterialFunctionCall->MaterialFunction->GetFunctionExpressions())
 				{
 					if (ExpressionType* Expression = FindExpressionByGUIDRecursive<ExpressionType>(InGUID, *FunctionExpressions))
 					{
@@ -1870,7 +1870,7 @@ private:
 				{
 					if (Layer)
 					{
-						if (const TArray<UMaterialExpression*>* FunctionExpressions = Layer->GetFunctionExpressions())
+						if (const TArray<TObjectPtr<UMaterialExpression>>* FunctionExpressions = Layer->GetFunctionExpressions())
 						{
 							if (ExpressionType* Expression = FindExpressionByGUIDRecursive<ExpressionType>(InGUID, *FunctionExpressions))
 							{
@@ -1884,7 +1884,7 @@ private:
 				{
 					if (Blend)
 					{
-						if (const TArray<UMaterialExpression*>* FunctionExpressions = Blend->GetFunctionExpressions())
+						if (const TArray<TObjectPtr<UMaterialExpression>>* FunctionExpressions = Blend->GetFunctionExpressions())
 						{
 							if (ExpressionType* Expression = FindExpressionByGUIDRecursive<ExpressionType>(InGUID, *FunctionExpressions))
 							{
@@ -1922,9 +1922,9 @@ private:
 			return false;
 		};
 
-		for (UMaterialExpression* Expression : Expressions)
+		for (TObjectPtr<UMaterialExpression>& Expression : Expressions)
 		{
-			if (TrySetParameterValue(Cast<ParameterType>(Expression)))
+			if (TrySetParameterValue((Expression && Expression.IsA<ParameterType>()) ? (ParameterType*)Expression.Get() : nullptr))
 			{
 				return true;
 			}
@@ -1938,12 +1938,12 @@ private:
 
 					for (UMaterialFunctionInterface* Function : Functions)
 					{
-						const TArray<UMaterialExpression*>* ExpressionPtr = Function->GetFunctionExpressions();
+						const TArray<TObjectPtr<UMaterialExpression>>* ExpressionPtr = Function->GetFunctionExpressions();
 						if (ExpressionPtr)
 						{
-							for (UMaterialExpression* FunctionExpression : *ExpressionPtr)
+							for (TObjectPtr<UMaterialExpression> FunctionExpression : *ExpressionPtr)
 							{
-								if (TrySetParameterValue(Cast<ParameterType>(FunctionExpression)))
+								if (TrySetParameterValue((FunctionExpression && FunctionExpression.IsA<ParameterType>()) ? (ParameterType*)FunctionExpression.Get() : nullptr))
 								{
 									return true;
 								}
diff --git a/Engine/Source/Runtime/Engine/Classes/Materials/MaterialFunction.h b/Engine/Source/Runtime/Engine/Classes/Materials/MaterialFunction.h
--- a/Engine/Source/Runtime/Engine/Classes/Materials/MaterialFunction.h
+++ b/Engine/Source/Runtime/Engine/Classes/Materials/MaterialFunction.h
@@ -147,7 +147,7 @@ public:
 	virtual UMaterialFunctionInterface* GetBaseFunction() override { return this; }
 	virtual const UMaterialFunctionInterface* GetBaseFunction() const override { return this; }
 #if WITH_EDITORONLY_DATA
-	virtual const TArray<UMaterialExpression*>* GetFunctionExpressions() const override { return &static_cast<const TArray<UMaterialExpression*>&>(FunctionExpressions); }
+	virtual const TArray<TObjectPtr<UMaterialExpression>>* GetFunctionExpressions() const override { return &FunctionExpressions; }
 #endif
 	virtual const FString* GetDescription() const override { return &Description; }
 
diff --git a/Engine/Source/Runtime/Engine/Classes/Materials/MaterialFunctionInstance.h b/Engine/Source/Runtime/Engine/Classes/Materials/MaterialFunctionInstance.h
--- a/Engine/Source/Runtime/Engine/Classes/Materials/MaterialFunctionInstance.h
+++ b/Engine/Source/Runtime/Engine/Classes/Materials/MaterialFunctionInstance.h
@@ -139,7 +139,7 @@ class UMaterialFunctionInstance : public UMaterialFunctionInterface
 	}
 
 #if WITH_EDITORONLY_DATA
-	virtual const TArray<UMaterialExpression*>* GetFunctionExpressions() const override
+	virtual const TArray<TObjectPtr<UMaterialExpression>>* GetFunctionExpressions() const override
 	{
 		const UMaterialFunctionInterface* BaseFunction = GetBaseFunction();
 		return BaseFunction ? BaseFunction->GetFunctionExpressions() : nullptr;
diff --git a/Engine/Source/Runtime/Engine/Classes/Materials/MaterialFunctionInterface.h b/Engine/Source/Runtime/Engine/Classes/Materials/MaterialFunctionInterface.h
--- a/Engine/Source/Runtime/Engine/Classes/Materials/MaterialFunctionInterface.h
+++ b/Engine/Source/Runtime/Engine/Classes/Materials/MaterialFunctionInterface.h
@@ -138,7 +138,7 @@ public:
 		PURE_VIRTUAL(UMaterialFunctionInterface::GetBaseFunction,return nullptr;);
 
 #if WITH_EDITORONLY_DATA
-	virtual const TArray<UMaterialExpression*>* GetFunctionExpressions() const
+	virtual const TArray<TObjectPtr<UMaterialExpression>>* GetFunctionExpressions() const
 		PURE_VIRTUAL(UMaterialFunctionInterface::GetFunctionExpressions,return nullptr;);
 #endif // WITH_EDITORONLY_DATA
 
@@ -162,18 +162,18 @@ public:
 		if (const UMaterialFunctionInterface* ParameterFunction = GetBaseFunction())
 		{
 			const UClass* TargetClass = UMaterialExpressionMaterialFunctionCall::StaticClass();
-			for (UMaterialExpression* Expression : *ParameterFunction->GetFunctionExpressions())
+			for (TObjectPtr<UMaterialExpression> Expression : *ParameterFunction->GetFunctionExpressions())
 			{
-				if (const UMaterialExpressionMaterialFunctionCall* FunctionExpression = (Expression && Expression->IsA(TargetClass)) ? (const UMaterialExpressionMaterialFunctionCall*)Expression : nullptr)
+				if (const UMaterialExpressionMaterialFunctionCall* FunctionExpression = (Expression && Expression.IsA(TargetClass)) ? (const UMaterialExpressionMaterialFunctionCall*)Expression.Get() : nullptr)
 				{
 					if (FunctionExpression->MaterialFunction)
 					{
 						FunctionExpression->MaterialFunction->GetAllParameterInfo<const ExpressionType>(OutParameterInfo, OutParameterIds, InBaseParameterInfo);
 					}
 				}
-				else if (const ExpressionType* ParameterExpression = Cast<const ExpressionType>(Expression))
+				else if (Expression && Expression.IsA<const ExpressionType>())
 				{
-					ParameterExpression->GetAllParameterInfo(OutParameterInfo, OutParameterIds, InBaseParameterInfo);
+					((const ExpressionType*)Expression.Get())->GetAllParameterInfo(OutParameterInfo, OutParameterIds, InBaseParameterInfo);
 				}
 			}
 
diff --git a/Engine/Source/Runtime/Engine/Classes/Materials/MaterialInstance.h b/Engine/Source/Runtime/Engine/Classes/Materials/MaterialInstance.h
--- a/Engine/Source/Runtime/Engine/Classes/Materials/MaterialInstance.h
+++ b/Engine/Source/Runtime/Engine/Classes/Materials/MaterialInstance.h
@@ -805,7 +805,7 @@ namespace MaterialInstance_Private
 {
 	/** Workaround - Similar to base call but evaluates all expressions found, not just the first */
 	template<typename ExpressionType>
-	void FindClosestExpressionByGUIDRecursive(const FName& InName, const FGuid& InGUID, const TArray<UMaterialExpression*>& InMaterialExpression, ExpressionType*& OutExpression)
+	void FindClosestExpressionByGUIDRecursive(const FName& InName, const FGuid& InGUID, const TArray<TObjectPtr<UMaterialExpression>>& InMaterialExpression, ExpressionType*& OutExpression)
 	{
 		for (int32 ExpressionIndex = 0; ExpressionIndex < InMaterialExpression.Num(); ExpressionIndex++)
 		{
@@ -835,7 +835,7 @@ namespace MaterialInstance_Private
 			}
 			else if (MaterialFunctionCall && MaterialFunctionCall->MaterialFunction)
 			{
-				if (const TArray<UMaterialExpression*>* FunctionExpressions = MaterialFunctionCall->MaterialFunction->GetFunctionExpressions())
+				if (const TArray<TObjectPtr<UMaterialExpression>>* FunctionExpressions = MaterialFunctionCall->MaterialFunction->GetFunctionExpressions())
 				{
 					FindClosestExpressionByGUIDRecursive<ExpressionType>(InName, InGUID, *FunctionExpressions, OutExpression);
 				}
@@ -849,7 +849,7 @@ namespace MaterialInstance_Private
 				{
 					if (Layer)
 					{
-						if (const TArray<UMaterialExpression*>* FunctionExpressions = Layer->GetFunctionExpressions())
+						if (const TArray<TObjectPtr<UMaterialExpression>>* FunctionExpressions = Layer->GetFunctionExpressions())
 						{
 							FindClosestExpressionByGUIDRecursive<ExpressionType>(InName, InGUID, *FunctionExpressions, OutExpression);
 						}
@@ -860,7 +860,7 @@ namespace MaterialInstance_Private
 				{
 					if (Blend)
 					{
-						if (const TArray<UMaterialExpression*>* FunctionExpressions = Blend->GetFunctionExpressions())
+						if (const TArray<TObjectPtr<UMaterialExpression>>* FunctionExpressions = Blend->GetFunctionExpressions())
 						{
 							FindClosestExpressionByGUIDRecursive<ExpressionType>(InName, InGUID, *FunctionExpressions, OutExpression);
 						}
diff --git a/Engine/Source/Runtime/Engine/Private/Materials/HLSLMaterialTranslator.cpp b/Engine/Source/Runtime/Engine/Private/Materials/HLSLMaterialTranslator.cpp
--- a/Engine/Source/Runtime/Engine/Private/Materials/HLSLMaterialTranslator.cpp
+++ b/Engine/Source/Runtime/Engine/Private/Materials/HLSLMaterialTranslator.cpp
@@ -297,7 +297,7 @@ void FHLSLMaterialTranslator::GatherCustomVertexInterpolators(TArray<UMaterialEx
 				FunctionCall->LinkFunctionIntoCaller(this);
 				PushFunction(&LocalState);
 
-				if (const TArray<UMaterialExpression*>* FunctionExpressions = FunctionCall->MaterialFunction->GetFunctionExpressions())
+				if (const TArray<TObjectPtr<UMaterialExpression>>* FunctionExpressions = FunctionCall->MaterialFunction->GetFunctionExpressions())
 				{
 					GatherCustomVertexInterpolators(*FunctionExpressions);
 				}
@@ -325,7 +325,7 @@ void FHLSLMaterialTranslator::GatherCustomVertexInterpolators(TArray<UMaterialEx
 						Layer->LinkFunctionIntoCaller(this);
 						PushFunction(&LocalState);
 
-						if (const TArray<UMaterialExpression*>* FunctionExpressions = Layer->MaterialFunction->GetFunctionExpressions())
+						if (const TArray<TObjectPtr<UMaterialExpression>>* FunctionExpressions = Layer->MaterialFunction->GetFunctionExpressions())
 						{
 							GatherCustomVertexInterpolators(*FunctionExpressions);
 						}
@@ -344,7 +344,7 @@ void FHLSLMaterialTranslator::GatherCustomVertexInterpolators(TArray<UMaterialEx
 						Blend->LinkFunctionIntoCaller(this);
 						PushFunction(&LocalState);
 
-						if (const TArray<UMaterialExpression*>* FunctionExpressions = Blend->MaterialFunction->GetFunctionExpressions())
+						if (const TArray<TObjectPtr<UMaterialExpression>>* FunctionExpressions = Blend->MaterialFunction->GetFunctionExpressions())
 						{
 							GatherCustomVertexInterpolators(*FunctionExpressions);
 						}
@@ -428,7 +428,7 @@ EMaterialExpressionVisitResult FHLSLMaterialTranslator::VisitExpressionsRecursiv
 				FunctionCall->LinkFunctionIntoCaller(this);
 				PushFunction(&LocalState);
 
-				if (const TArray<UMaterialExpression*>* FunctionExpressions = FunctionCall->MaterialFunction->GetFunctionExpressions())
+				if (const TArray<TObjectPtr<UMaterialExpression>>* FunctionExpressions = FunctionCall->MaterialFunction->GetFunctionExpressions())
 				{
 					VisitResult = VisitExpressionsRecursive(*FunctionExpressions, InVisitor);
 				}
@@ -461,7 +461,7 @@ EMaterialExpressionVisitResult FHLSLMaterialTranslator::VisitExpressionsRecursiv
 						Layer->LinkFunctionIntoCaller(this);
 						PushFunction(&LocalState);
 
-						if (const TArray<UMaterialExpression*>* FunctionExpressions = Layer->MaterialFunction->GetFunctionExpressions())
+						if (const TArray<TObjectPtr<UMaterialExpression>>* FunctionExpressions = Layer->MaterialFunction->GetFunctionExpressions())
 						{
 							VisitResult = VisitExpressionsRecursive(*FunctionExpressions, InVisitor);
 						}
@@ -485,7 +485,7 @@ EMaterialExpressionVisitResult FHLSLMaterialTranslator::VisitExpressionsRecursiv
 						Blend->LinkFunctionIntoCaller(this);
 						PushFunction(&LocalState);
 
-						if (const TArray<UMaterialExpression*>* FunctionExpressions = Blend->MaterialFunction->GetFunctionExpressions())
+						if (const TArray<TObjectPtr<UMaterialExpression>>* FunctionExpressions = Blend->MaterialFunction->GetFunctionExpressions())
 						{
 							VisitResult = VisitExpressionsRecursive(*FunctionExpressions, InVisitor);
 						}
diff --git a/Engine/Source/Runtime/Engine/Private/Materials/MaterialCachedData.cpp b/Engine/Source/Runtime/Engine/Private/Materials/MaterialCachedData.cpp
--- a/Engine/Source/Runtime/Engine/Private/Materials/MaterialCachedData.cpp
+++ b/Engine/Source/Runtime/Engine/Private/Materials/MaterialCachedData.cpp
@@ -229,7 +229,7 @@ bool FMaterialCachedExpressionData::UpdateForFunction(const FMaterialCachedExpre
 	FMaterialCachedExpressionData* Self = this;
 	auto ProcessFunction = [Self, &LocalContext, Association, ParameterIndex, &bResult](UMaterialFunctionInterface* InFunction) -> bool
 	{
-		const TArray<UMaterialExpression*>* FunctionExpressions = InFunction->GetFunctionExpressions();
+		const TArray<TObjectPtr<UMaterialExpression>>* FunctionExpressions = InFunction->GetFunctionExpressions();
 		if (FunctionExpressions)
 		{
 			if (!Self->UpdateForExpressions(LocalContext, *FunctionExpressions, Association, ParameterIndex))
@@ -402,7 +402,7 @@ void FMaterialCachedParameters_UpdateForLayerParameters(FMaterialCachedParameter
 	}
 }
 
-bool FMaterialCachedExpressionData::UpdateForExpressions(const FMaterialCachedExpressionContext& Context, const TArray<UMaterialExpression*>& Expressions, EMaterialParameterAssociation Association, int32 ParameterIndex)
+bool FMaterialCachedExpressionData::UpdateForExpressions(const FMaterialCachedExpressionContext& Context, const TArray<TObjectPtr<UMaterialExpression>>& Expressions, EMaterialParameterAssociation Association, int32 ParameterIndex)
 {
 	bool bResult = true;
 	for (UMaterialExpression* Expression : Expressions)
diff --git a/Engine/Source/Runtime/Engine/Private/Materials/MaterialExpressions.cpp b/Engine/Source/Runtime/Engine/Private/Materials/MaterialExpressions.cpp
--- a/Engine/Source/Runtime/Engine/Private/Materials/MaterialExpressions.cpp
+++ b/Engine/Source/Runtime/Engine/Private/Materials/MaterialExpressions.cpp
@@ -12728,13 +12728,14 @@ bool UMaterialFunction::SetVectorParameterValueEditorOnly(FName ParameterName, F
 
 				for (UMaterialFunctionInterface* Function : Functions)
 				{
-					const TArray<UMaterialExpression*>* ExpressionPtr = Function->GetFunctionExpressions();
+					const TArray<TObjectPtr<UMaterialExpression>>* ExpressionPtr = Function->GetFunctionExpressions();
 					if (ExpressionPtr)
 					{
-						for (UMaterialExpression* FunctionExpression : *ExpressionPtr)
+						for (TObjectPtr<UMaterialExpression> FunctionExpression : *ExpressionPtr)
 						{
-							if (UMaterialExpressionVectorParameter* FunctionExpressionParameter = Cast<UMaterialExpressionVectorParameter>(FunctionExpression))
+							if (FunctionExpression && FunctionExpression.IsA<UMaterialExpressionVectorParameter>())
 							{
+								UMaterialExpressionVectorParameter* FunctionExpressionParameter = (UMaterialExpressionVectorParameter*)FunctionExpression.Get();
 								if (FunctionExpressionParameter->SetParameterValue(ParameterName, InValue))
 								{
 									return true;
@@ -12771,13 +12772,14 @@ bool UMaterialFunction::SetScalarParameterValueEditorOnly(FName ParameterName, f
 
 				for (UMaterialFunctionInterface* Function : Functions)
 				{
-					const TArray<UMaterialExpression*>* ExpressionPtr = Function->GetFunctionExpressions();
+					const TArray<TObjectPtr<UMaterialExpression>>* ExpressionPtr = Function->GetFunctionExpressions();
 					if (ExpressionPtr)
 					{
-						for (UMaterialExpression* FunctionExpression : *ExpressionPtr)
+						for (TObjectPtr<UMaterialExpression> FunctionExpression : *ExpressionPtr)
 						{
-							if (UMaterialExpressionScalarParameter* FunctionExpressionParameter = Cast<UMaterialExpressionScalarParameter>(FunctionExpression))
+							if (FunctionExpression && FunctionExpression.IsA<UMaterialExpressionScalarParameter>())
 							{
+								UMaterialExpressionScalarParameter* FunctionExpressionParameter = (UMaterialExpressionScalarParameter*)FunctionExpression.Get();
 								if (FunctionExpressionParameter->SetParameterValue(ParameterName, InValue))
 								{
 									return true;
@@ -12814,13 +12816,14 @@ bool UMaterialFunction::SetTextureParameterValueEditorOnly(FName ParameterName,
 
 				for (UMaterialFunctionInterface* Function : Functions)
 				{
-					const TArray<UMaterialExpression*>* ExpressionPtr = Function->GetFunctionExpressions();
+					const TArray<TObjectPtr<UMaterialExpression>>* ExpressionPtr = Function->GetFunctionExpressions();
 					if (ExpressionPtr)
 					{
-						for (UMaterialExpression* FunctionExpression : *ExpressionPtr)
+						for (TObjectPtr<UMaterialExpression> FunctionExpression : *ExpressionPtr)
 						{
-							if (UMaterialExpressionTextureSampleParameter* FunctionExpressionParameter = Cast<UMaterialExpressionTextureSampleParameter>(FunctionExpression))
+							if (FunctionExpression && FunctionExpression.IsA<UMaterialExpressionTextureSampleParameter>())
 							{
+								UMaterialExpressionTextureSampleParameter* FunctionExpressionParameter = (UMaterialExpressionTextureSampleParameter*)FunctionExpression.Get();
 								if (FunctionExpressionParameter->SetParameterValue(ParameterName, InValue))
 								{
 									return true;
@@ -12857,13 +12860,14 @@ bool UMaterialFunction::SetRuntimeVirtualTextureParameterValueEditorOnly(FName P
 
 				for (UMaterialFunctionInterface* Function : Functions)
 				{
-					const TArray<UMaterialExpression*>* ExpressionPtr = Function->GetFunctionExpressions();
+					const TArray<TObjectPtr<UMaterialExpression>>* ExpressionPtr = Function->GetFunctionExpressions();
 					if (ExpressionPtr)
 					{
-						for (UMaterialExpression* FunctionExpression : *ExpressionPtr)
+						for (TObjectPtr<UMaterialExpression> FunctionExpression : *ExpressionPtr)
 						{
-							if (UMaterialExpressionRuntimeVirtualTextureSampleParameter* FunctionExpressionParameter = Cast<UMaterialExpressionRuntimeVirtualTextureSampleParameter>(FunctionExpression))
+							if (FunctionExpression && FunctionExpression.IsA<UMaterialExpressionRuntimeVirtualTextureSampleParameter>())
 							{
+								UMaterialExpressionRuntimeVirtualTextureSampleParameter* FunctionExpressionParameter = (UMaterialExpressionRuntimeVirtualTextureSampleParameter*)FunctionExpression.Get();
 								if (FunctionExpressionParameter->SetParameterValue(ParameterName, InValue))
 								{
 									return true;
@@ -12900,13 +12904,14 @@ bool UMaterialFunction::SetFontParameterValueEditorOnly(FName ParameterName, cla
 
 				for (UMaterialFunctionInterface* Function : Functions)
 				{
-					const TArray<UMaterialExpression*>* ExpressionPtr = Function->GetFunctionExpressions();
+					const TArray<TObjectPtr<UMaterialExpression>>* ExpressionPtr = Function->GetFunctionExpressions();
 					if (ExpressionPtr)
 					{
-						for (UMaterialExpression* FunctionExpression : *ExpressionPtr)
+						for (TObjectPtr<UMaterialExpression> FunctionExpression : *ExpressionPtr)
 						{
-							if (UMaterialExpressionFontSampleParameter* FunctionExpressionParameter = Cast<UMaterialExpressionFontSampleParameter>(FunctionExpression))
+							if (FunctionExpression && FunctionExpression.IsA<UMaterialExpressionFontSampleParameter>())
 							{
+								UMaterialExpressionFontSampleParameter* FunctionExpressionParameter = (UMaterialExpressionFontSampleParameter*)FunctionExpression.Get();
 								if (FunctionExpressionParameter->SetParameterValue(ParameterName, InFontValue, InFontPage))
 								{
 									return true;
@@ -12943,13 +12948,14 @@ bool UMaterialFunction::SetStaticSwitchParameterValueEditorOnly(FName ParameterN
 
 				for (UMaterialFunctionInterface* Function : Functions)
 				{
-					const TArray<UMaterialExpression*>* ExpressionPtr = Function->GetFunctionExpressions();
+					const TArray<TObjectPtr<UMaterialExpression>>* ExpressionPtr = Function->GetFunctionExpressions();
 					if (ExpressionPtr)
 					{
-						for (UMaterialExpression* FunctionExpression : *ExpressionPtr)
+						for (TObjectPtr<UMaterialExpression> FunctionExpression : *ExpressionPtr)
 						{
-							if (UMaterialExpressionStaticSwitchParameter* FunctionExpressionParameter = Cast<UMaterialExpressionStaticSwitchParameter>(FunctionExpression))
+							if (FunctionExpression && FunctionExpression.IsA<UMaterialExpressionStaticSwitchParameter>())
 							{
+								UMaterialExpressionStaticSwitchParameter* FunctionExpressionParameter = (UMaterialExpressionStaticSwitchParameter*)FunctionExpression.Get();
 								if (FunctionExpressionParameter->SetParameterValue(ParameterName, OutValue, OutExpressionGuid))
 								{
 									return true;
@@ -12986,13 +12992,14 @@ bool UMaterialFunction::SetStaticComponentMaskParameterValueEditorOnly(FName Par
 
 				for (UMaterialFunctionInterface* Function : Functions)
 				{
-					const TArray<UMaterialExpression*>* ExpressionPtr = Function->GetFunctionExpressions();
+					const TArray<TObjectPtr<UMaterialExpression>>* ExpressionPtr = Function->GetFunctionExpressions();
 					if (ExpressionPtr)
 					{
-						for (UMaterialExpression* FunctionExpression : *ExpressionPtr)
+						for (TObjectPtr<UMaterialExpression> FunctionExpression : *ExpressionPtr)
 						{
-							if (UMaterialExpressionStaticComponentMaskParameter* FunctionExpressionParameter = Cast<UMaterialExpressionStaticComponentMaskParameter>(FunctionExpression))
+							if (FunctionExpression && FunctionExpression.IsA<UMaterialExpressionStaticComponentMaskParameter>())
 							{
+								UMaterialExpressionStaticComponentMaskParameter* FunctionExpressionParameter = (UMaterialExpressionStaticComponentMaskParameter*)FunctionExpression.Get();
 								if (FunctionExpressionParameter->SetParameterValue(ParameterName, R, G, B, A, OutExpressionGuid))
 								{
 									return true;
diff --git a/Engine/Source/Runtime/Engine/Private/Materials/MaterialInstance.cpp b/Engine/Source/Runtime/Engine/Private/Materials/MaterialInstance.cpp
--- a/Engine/Source/Runtime/Engine/Private/Materials/MaterialInstance.cpp
+++ b/Engine/Source/Runtime/Engine/Private/Materials/MaterialInstance.cpp
@@ -5075,32 +5075,35 @@ void UMaterialInstance::CopyMaterialUniformParametersInternal(UMaterialInterface
 
 #if WITH_EDITOR
 
-void FindCollectionExpressionRecursive(TArray<FGuid>& OutGuidList, const TArray<UMaterialExpression*>& InMaterialExpression)
+void FindCollectionExpressionRecursive(TArray<FGuid>& OutGuidList, const TArray<TObjectPtr<UMaterialExpression>>& InMaterialExpression)
 {
 	for (int32 ExpressionIndex = 0; ExpressionIndex < InMaterialExpression.Num(); ExpressionIndex++)
 	{
-		const UMaterialExpression* ExpressionPtr = InMaterialExpression[ExpressionIndex];
-		const UMaterialExpressionCollectionParameter* CollectionPtr = Cast<UMaterialExpressionCollectionParameter>(ExpressionPtr);
-		const UMaterialExpressionMaterialFunctionCall* MaterialFunctionCall = Cast<UMaterialExpressionMaterialFunctionCall>(ExpressionPtr);
-		const UMaterialExpressionMaterialAttributeLayers* MaterialLayers = Cast<UMaterialExpressionMaterialAttributeLayers>(ExpressionPtr);
+		if (!InMaterialExpression[ExpressionIndex])
+		{
+			continue;
+		}
 
-		if (CollectionPtr)
+		if (InMaterialExpression[ExpressionIndex].IsA<UMaterialExpressionCollectionParameter>())
 		{
+			const UMaterialExpressionCollectionParameter* CollectionPtr = (UMaterialExpressionCollectionParameter*)InMaterialExpression[ExpressionIndex].Get();
 			if (CollectionPtr->Collection)
 			{
 				OutGuidList.Add(CollectionPtr->Collection->StateId);
 			}
 			return;
 		}
-		else if (MaterialFunctionCall && MaterialFunctionCall->MaterialFunction)
+		else if (InMaterialExpression[ExpressionIndex].IsA<UMaterialExpressionMaterialFunctionCall>())
 		{
-			if (const TArray<UMaterialExpression*>* FunctionExpressions = MaterialFunctionCall->MaterialFunction->GetFunctionExpressions())
+			const UMaterialExpressionMaterialFunctionCall* MaterialFunctionCall = (UMaterialExpressionMaterialFunctionCall*)InMaterialExpression[ExpressionIndex].Get();
+			if (const TArray<TObjectPtr<UMaterialExpression>>* FunctionExpressions = MaterialFunctionCall->MaterialFunction->GetFunctionExpressions())
 			{
 				FindCollectionExpressionRecursive(OutGuidList, *FunctionExpressions);
 			}
 		}
-		else if (MaterialLayers)
+		else if (InMaterialExpression[ExpressionIndex].IsA<UMaterialExpressionMaterialAttributeLayers>())
 		{
+			const UMaterialExpressionMaterialAttributeLayers* MaterialLayers = (UMaterialExpressionMaterialAttributeLayers*)InMaterialExpression[ExpressionIndex].Get();
 			const TArray<UMaterialFunctionInterface*>& Layers = MaterialLayers->GetLayers();
 			const TArray<UMaterialFunctionInterface*>& Blends = MaterialLayers->GetBlends();
 
@@ -5108,7 +5111,7 @@ void FindCollectionExpressionRecursive(TArray<FGuid>& OutGuidList, const TArray<
 			{
 				if (Layer)
 				{
-					if (const TArray<UMaterialExpression*>* FunctionExpressions = Layer->GetFunctionExpressions())
+					if (const TArray<TObjectPtr<UMaterialExpression>>* FunctionExpressions = Layer->GetFunctionExpressions())
 					{
 						FindCollectionExpressionRecursive(OutGuidList, *FunctionExpressions);
 					}
@@ -5119,7 +5122,7 @@ void FindCollectionExpressionRecursive(TArray<FGuid>& OutGuidList, const TArray<
 			{
 				if (Blend)
 				{
-					if (const TArray<UMaterialExpression*>* FunctionExpressions = Blend->GetFunctionExpressions())
+					if (const TArray<TObjectPtr<UMaterialExpression>>* FunctionExpressions = Blend->GetFunctionExpressions())
 					{
 						FindCollectionExpressionRecursive(OutGuidList, *FunctionExpressions);
 					}
@@ -5143,7 +5146,7 @@ void UMaterialInstance::AppendReferencedParameterCollectionIdsTo(TArray<FGuid>&
 		{
 			if (Layer)
 			{
-				if (const TArray<UMaterialExpression*>* FunctionExpressions = Layer->GetFunctionExpressions())
+				if (const TArray<TObjectPtr<UMaterialExpression>>* FunctionExpressions = Layer->GetFunctionExpressions())
 				{
 					FindCollectionExpressionRecursive(OutIds, *FunctionExpressions);
 				}
@@ -5153,7 +5156,7 @@ void UMaterialInstance::AppendReferencedParameterCollectionIdsTo(TArray<FGuid>&
 		{
 			if (Blend)
 			{
-				if (const TArray<UMaterialExpression*>* FunctionExpressions = Blend->GetFunctionExpressions())
+				if (const TArray<TObjectPtr<UMaterialExpression>>* FunctionExpressions = Blend->GetFunctionExpressions())
 				{
 					FindCollectionExpressionRecursive(OutIds, *FunctionExpressions);
 				}
diff --git a/Engine/Source/Runtime/Engine/Private/Materials/MaterialStaticParameterValueResolver.h b/Engine/Source/Runtime/Engine/Private/Materials/MaterialStaticParameterValueResolver.h
--- a/Engine/Source/Runtime/Engine/Private/Materials/MaterialStaticParameterValueResolver.h
+++ b/Engine/Source/Runtime/Engine/Private/Materials/MaterialStaticParameterValueResolver.h
@@ -72,9 +72,9 @@ public:
 				auto GetExpressionParameterByNamePredicate =
 					[this, &PendingParametersForFunction, &ParameterIndicesForFunction, &ParametersAndOwners, TargetClass](UMaterialFunctionInterface* InFunction) -> bool
 				{
-					for (UMaterialExpression* FunctionExpression : *InFunction->GetFunctionExpressions())
+					for (TObjectPtr<UMaterialExpression> FunctionExpression : *InFunction->GetFunctionExpressions())
 					{
-						if (ParameterType* ExpressionParameter = (FunctionExpression && FunctionExpression->IsA(TargetClass)) ? (ParameterType*)FunctionExpression : nullptr)
+						if (ParameterType* ExpressionParameter = (FunctionExpression && FunctionExpression.IsA(TargetClass)) ? (ParameterType*)FunctionExpression.Get() : nullptr)
 						{
 							for (int32 FunctionParameterIdx = 0; FunctionParameterIdx < ParameterIndicesForFunction.Num(); ++FunctionParameterIdx)
 							{
diff --git a/Engine/Source/Runtime/Engine/Public/MaterialCachedData.h b/Engine/Source/Runtime/Engine/Public/MaterialCachedData.h
--- a/Engine/Source/Runtime/Engine/Public/MaterialCachedData.h
+++ b/Engine/Source/Runtime/Engine/Public/MaterialCachedData.h
@@ -232,7 +232,7 @@ struct FMaterialCachedExpressionData
 
 #if WITH_EDITOR
 	/** Returns 'false' if update is incomplete, due to missing expression data (stripped from non-editor build) */
-	bool UpdateForExpressions(const FMaterialCachedExpressionContext& Context, const TArray<UMaterialExpression*>& Expressions, EMaterialParameterAssociation Association, int32 ParameterIndex);
+	bool UpdateForExpressions(const FMaterialCachedExpressionContext& Context, const TArray<TObjectPtr<UMaterialExpression>>& Expressions, EMaterialParameterAssociation Association, int32 ParameterIndex);
 	bool UpdateForFunction(const FMaterialCachedExpressionContext& Context, UMaterialFunctionInterface* Function, EMaterialParameterAssociation Association, int32 ParameterIndex);
 	bool UpdateForLayerFunctions(const FMaterialCachedExpressionContext& Context, const FMaterialLayersFunctions& LayerFunctions);
 #endif // WITH_EDITOR
