World Partition HLOD - Fix HLOD assets being serialized by the PIE fixup archive

* The core issue is that by being outered to the actors, the HLOD assets would be "IsInA(UWorld)" causing the PIE fixup archive to serialize them
* Now outering the assets to the actor package rather than the actor itself
* Modified FWorldPartitionLevelHelper::MoveExternalActorsToLevel() to move all objects found in the actors packages to the destination level package

Tested in PIE, game & cooked

diff --git a/Engine/Source/Runtime/Engine/Private/WorldPartition/HLOD/HLODBuilder.cpp b/Engine/Source/Runtime/Engine/Private/WorldPartition/HLOD/HLODBuilder.cpp
--- a/Engine/Source/Runtime/Engine/Private/WorldPartition/HLOD/HLODBuilder.cpp
+++ b/Engine/Source/Runtime/Engine/Private/WorldPartition/HLOD/HLODBuilder.cpp
@@ -377,7 +377,6 @@ class FHLODBuilder_MeshMerge : public FHLODBuilder
 		Algo::ForEach(Assets, [this, InHLODActor, &Component, &MergedActorLocation](UObject* Asset)
 		{
 			Asset->ClearFlags(RF_Public | RF_Standalone);
-			Asset->Rename(nullptr, InHLODActor);
 
 			if (Cast<UStaticMesh>(Asset))
 			{
@@ -415,7 +414,6 @@ class FHLODBuilder_MeshSimplify : public FHLODBuilder
 		Algo::ForEach(Assets, [this, InHLODActor, &Component](UObject* Asset)
 		{
 			Asset->ClearFlags(RF_Public | RF_Standalone);
-			Asset->Rename(nullptr, InHLODActor, REN_ForceNoResetLoaders | REN_NonTransactional | REN_DoNotDirty);
 
 			if (Cast<UStaticMesh>(Asset))
 			{
diff --git a/Engine/Source/Runtime/Engine/Private/WorldPartition/WorldPartitionLevelHelper.cpp b/Engine/Source/Runtime/Engine/Private/WorldPartition/WorldPartitionLevelHelper.cpp
--- a/Engine/Source/Runtime/Engine/Private/WorldPartition/WorldPartitionLevelHelper.cpp
+++ b/Engine/Source/Runtime/Engine/Private/WorldPartition/WorldPartitionLevelHelper.cpp
@@ -58,6 +58,8 @@ void FWorldPartitionLevelHelper::MoveExternalActorsToLevel(const TArray<FWorldPa
 		AActor* Actor = FindObject<AActor>(nullptr, *PackageObjectMapping.LoadedPath.ToString());
 		if (ensure(Actor))
 		{
+			UPackage* ActorPackage = Actor->GetPackage();
+
 			const bool bSameOuter = (InLevel == Actor->GetOuter());
 			Actor->SetPackageExternal(false, false);
 			Actor->Rename(nullptr, InLevel, REN_ForceNoResetLoaders);
@@ -66,6 +68,17 @@ void FWorldPartitionLevelHelper::MoveExternalActorsToLevel(const TArray<FWorldPa
 			{
 				InLevel->AddLoadedActor(Actor);
 			}
+
+			// Include objects found in the source actor package in the destination level package
+			TArray<UObject*> Objects;
+			GetObjectsWithOuter(ActorPackage, Objects);
+			for (UObject* Object : Objects)
+			{
+				if (Object->GetFName() != NAME_PackageMetaData)
+				{
+					Object->Rename(nullptr, LevelPackage, REN_ForceNoResetLoaders);
+				}
+			}
 		}
 		else
 		{
