Don't output warnings about missing world context for worlds that are preloading; their world context is about to be set later on in LoadMap.

diff --git a/Engine/Source/Runtime/Engine/Private/LevelActor.cpp b/Engine/Source/Runtime/Engine/Private/LevelActor.cpp
--- a/Engine/Source/Runtime/Engine/Private/LevelActor.cpp
+++ b/Engine/Source/Runtime/Engine/Private/LevelActor.cpp
@@ -940,9 +940,13 @@ bool UWorld::DestroyActor( AActor* ThisActor, bool bNetForce, bool bShouldModify
 	}
 	else if (WorldType != EWorldType::Inactive && !IsRunningCommandlet())
 	{
-		// Inactive worlds do not have a world context, otherwise only worlds in the middle of seamless travel should have no context,
-		// and in that case, we shouldn't be destroying actors on them until they have become the current world (i.e. CopyWorldData has been called)
-		UE_LOG(LogSpawn, Warning, TEXT("UWorld::DestroyActor: World has no context! World: %s, Actor: %s"), *GetName(), *ThisActor->GetPathName());
+		// If we are preloading this world, it's normal that we don't have a valid context yet.
+		if (!UWorld::WorldTypePreLoadMap.Find(GetOuter()->GetFName()))
+		{
+			// Inactive worlds do not have a world context, otherwise only worlds in the middle of seamless travel should have no context,
+			// and in that case, we shouldn't be destroying actors on them until they have become the current world (i.e. CopyWorldData has been called)
+			UE_LOG(LogSpawn, Warning, TEXT("UWorld::DestroyActor: World has no context! World: %s, Actor: %s"), *GetName(), *ThisActor->GetPathName());
+		}
 	}
 
 	// Remove the actor from the actor list.
