Fixing anim notifies with metasounds

diff --git a/Engine/Source/Runtime/Engine/Private/Animation/AnimNotify_PlaySound.cpp b/Engine/Source/Runtime/Engine/Private/Animation/AnimNotify_PlaySound.cpp
--- a/Engine/Source/Runtime/Engine/Private/Animation/AnimNotify_PlaySound.cpp
+++ b/Engine/Source/Runtime/Engine/Private/Animation/AnimNotify_PlaySound.cpp
@@ -4,6 +4,7 @@
 #include "Components/SkeletalMeshComponent.h"
 #include "Kismet/GameplayStatics.h"
 #include "Sound/SoundBase.h"
+#include "Sound/SoundWaveProcedural.h"
 #include "Animation/AnimSequenceBase.h"
 
 #if WITH_EDITOR
@@ -31,7 +32,7 @@ void UAnimNotify_PlaySound::Notify(class USkeletalMeshComponent* MeshComp, class
 	// Don't call super to avoid call back in to blueprints
 	if (Sound && MeshComp)
 	{
-		if (Sound->IsLooping())
+		if (Sound->IsLooping() && !Sound->IsA<USoundWaveProcedural>())
 		{
 			UE_LOG(LogAudio, Warning, TEXT("PlaySound notify: Anim %s tried to spawn infinitely looping sound asset %s. Spawning suppressed."), *GetNameSafe(Animation), *GetNameSafe(Sound));
 			return;
