Fix async compilation progress disappearing too early

diff --git a/Engine/Source/Runtime/Engine/Private/AsyncCompilationHelpers.cpp b/Engine/Source/Runtime/Engine/Private/AsyncCompilationHelpers.cpp
--- a/Engine/Source/Runtime/Engine/Private/AsyncCompilationHelpers.cpp
+++ b/Engine/Source/Runtime/Engine/Private/AsyncCompilationHelpers.cpp
@@ -57,13 +57,17 @@ void FAsyncCompilationNotification::Update(int32 NumJobs)
 	}
 	else 
 	{
-		if(!NotificationHandle.IsValid())
+		if (!NotificationHandle.IsValid())
 		{
 			StartNumJobs = NumJobs;
 			NotificationHandle = FSlateNotificationManager::Get().StartProgressNotification(ProgressMessage, StartNumJobs);
 		}
 		else
 		{
+			if (NumJobs > StartNumJobs)
+			{
+				StartNumJobs = NumJobs;
+			}
 			FSlateNotificationManager::Get().UpdateProgressNotification(NotificationHandle, StartNumJobs - NumJobs, StartNumJobs, ProgressMessage);
 		}
 	}
