Fix character falling through the floor in PIE or -game when async static meshes are still building

diff --git a/Engine/Source/Runtime/Engine/Private/StaticMeshCompiler.cpp b/Engine/Source/Runtime/Engine/Private/StaticMeshCompiler.cpp
--- a/Engine/Source/Runtime/Engine/Private/StaticMeshCompiler.cpp
+++ b/Engine/Source/Runtime/Engine/Private/StaticMeshCompiler.cpp
@@ -398,7 +398,7 @@ void FStaticMeshCompilingManager::FinishCompilationsForGame()
 					PIEWorlds.Contains(Component->GetWorld()) &&
 					Component->GetStaticMesh() != nullptr &&
 					RegisteredStaticMesh.Contains(Component->GetStaticMesh()) &&
-					(PlayInEditorMode == 0 || Component->GetCollisionEnabled() != ECollisionEnabled::NoCollision || Component->IsNavigationRelevant()))
+					(PlayInEditorMode == 0 || Component->GetCollisionEnabled() != ECollisionEnabled::NoCollision || Component->IsNavigationRelevant() || Component->bAlwaysCreatePhysicsState || Component->CanCharacterStepUpOn != ECB_No))
 				{
 					const FBoxSphereBounds ComponentBounds = Component->Bounds.GetBox();
 					const UWorld* ComponentWorld = Component->GetWorld();
