Fix for spamming warnings about orphaned sounds when using metasounds without audio components

UE-113076

diff --git a/Engine/Source/Runtime/Engine/Private/SoundWave.cpp b/Engine/Source/Runtime/Engine/Private/SoundWave.cpp
--- a/Engine/Source/Runtime/Engine/Private/SoundWave.cpp
+++ b/Engine/Source/Runtime/Engine/Private/SoundWave.cpp
@@ -2060,13 +2060,20 @@ void USoundWave::Parse(FAudioDevice* AudioDevice, const UPTRINT NodeWaveInstance
 	if (bLooping || ParseParams.bLooping)
 	{
 		WaveInstance->bIsFinished = false;
-#if !(NO_LOGGING || UE_BUILD_SHIPPING || UE_BUILD_TEST)
-		if (!ActiveSound.bWarnedAboutOrphanedLooping && ActiveSound.GetAudioComponentID() == 0 && ActiveSound.FadeOut == FActiveSound::EFadeOut::None)
+
+		const USoundBase* Sound = ActiveSound.GetSound();
+
+		// Don't need to worry about logging orphaned sounds for procedural sounds
+		if (!Sound->IsA<USoundWaveProcedural>())
 		{
-			UE_LOG(LogAudio, Warning, TEXT("Detected orphaned looping sound '%s'."), *ActiveSound.GetSound()->GetName());
-			ActiveSound.bWarnedAboutOrphanedLooping = true;
-		}
+#if !(NO_LOGGING || UE_BUILD_SHIPPING || UE_BUILD_TEST)
+			if (!ActiveSound.bWarnedAboutOrphanedLooping && ActiveSound.GetAudioComponentID() == 0 && ActiveSound.FadeOut == FActiveSound::EFadeOut::None)
+			{
+				UE_LOG(LogAudio, Warning, TEXT("Detected orphaned looping sound '%s'."), *ActiveSound.GetSound()->GetName());
+				ActiveSound.bWarnedAboutOrphanedLooping = true;
+			}
 #endif
+		}
 	}
 
 	// Early out if finished.
