Fixed incorrectly matched FAudioThread::Suspend()/Resume()

diff --git a/Engine/Source/Runtime/Engine/Private/AudioDeviceManager.cpp b/Engine/Source/Runtime/Engine/Private/AudioDeviceManager.cpp
--- a/Engine/Source/Runtime/Engine/Private/AudioDeviceManager.cpp
+++ b/Engine/Source/Runtime/Engine/Private/AudioDeviceManager.cpp
@@ -309,7 +309,7 @@ bool FAudioDeviceManager::ToggleAudioMixer()
 		}
 
 		// Suspend the audio thread
-		FAudioThread::SuspendAudioThread();
+		FAudioThreadSuspendContext AudioThreadSuspendScope;
 
 		// If using audio mixer, we need to toggle back to non-audio mixer
 		FString ModuleToUnload;
@@ -433,9 +433,6 @@ bool FAudioDeviceManager::ToggleAudioMixer()
 
 			// Unload the previous audio device module
 			FModuleManager::Get().UnloadModule(*ModuleToUnload);
-
-			// Resume the audio thread
-			FAudioThread::ResumeAudioThread();
 		}
 	}
 
