Removed deprecated FAudioThread::IsAudioThreadRunning()

diff --git a/Engine/Source/Runtime/Engine/Private/AudioThread.cpp b/Engine/Source/Runtime/Engine/Private/AudioThread.cpp
--- a/Engine/Source/Runtime/Engine/Private/AudioThread.cpp
+++ b/Engine/Source/Runtime/Engine/Private/AudioThread.cpp
@@ -81,7 +81,6 @@ struct FAudioThreadInteractor
 
 static FAutoConsoleVariableSink CVarUseAudioThreadSink(FConsoleCommandDelegate::CreateStatic(&FAudioThreadInteractor::UseAudioThreadCVarSinkFunction));
 
-TAtomic<bool> FAudioThread::bIsAudioThreadRunning(false);
 bool FAudioThread::bUseThreadedAudio = false;
 FRunnable* FAudioThread::AudioThreadRunnable = nullptr;
 FCriticalSection FAudioThread::CurrentAudioThreadStatIdCS;
@@ -132,10 +131,9 @@ static int32 AudioThreadSuspendCount = 0;
 
 void FAudioThread::SuspendAudioThread()
 {
-PRAGMA_DISABLE_DEPRECATION_WARNINGS
 	check(FPlatformTLS::GetCurrentThreadId() == GGameThreadId);
 	check(!GIsAudioThreadSuspended.Load() || CVarSuspendAudioThread.GetValueOnGameThread() != 0);
-	if (bIsAudioThreadRunning.Load())
+	if (IsAudioThreadRunning())
 	{
 		// Make GC wait on the audio thread finishing processing
 		FAudioCommandFence AudioFence;
@@ -144,11 +142,8 @@ PRAGMA_DISABLE_DEPRECATION_WARNINGS
 
 		GIsAudioThreadSuspended = true;
 		FPlatformMisc::MemoryBarrier();
-		bIsAudioThreadRunning = false;
 	}
-	check(!bIsAudioThreadRunning.Load());
-
-PRAGMA_ENABLE_DEPRECATION_WARNINGS
+	check(!IsAudioThreadRunning());
 }
 
 void FAudioThread::ResumeAudioThread()
@@ -160,7 +155,6 @@ PRAGMA_DISABLE_DEPRECATION_WARNINGS
 	{
 		GIsAudioThreadSuspended = false;
 		FPlatformMisc::MemoryBarrier();
-		bIsAudioThreadRunning = true;
 	}
 	ProcessAllCommands();
 
@@ -214,7 +208,7 @@ uint32 FAudioThread::Run()
 
 void FAudioThread::SetUseThreadedAudio(const bool bInUseThreadedAudio)
 {
-	if (bIsAudioThreadRunning.Load() && !bInUseThreadedAudio)
+	if (IsAudioThreadRunning() && !bInUseThreadedAudio)
 	{
 		UE_LOG(LogAudio, Error, TEXT("You cannot disable using threaded audio once the thread has already begun running."));
 	}
@@ -229,11 +223,6 @@ bool FAudioThread::IsUsingThreadedAudio()
 	return bUseThreadedAudio;
 }
 
-bool FAudioThread::IsAudioThreadRunning()
-{ 
-	return bIsAudioThreadRunning; 
-}
-
 struct FAudioAsyncBatcher
 {
 	FGraphEventArray DispatchEvent;
@@ -289,7 +278,7 @@ static FAudioAsyncBatcher GAudioAsyncBatcher;
 void FAudioThread::RunCommandOnAudioThread(TUniqueFunction<void()> InFunction, const TStatId InStatId)
 {
 	check(FPlatformTLS::GetCurrentThreadId() == GGameThreadId);
-	if (bIsAudioThreadRunning)
+	if (IsAudioThreadRunning())
 	{
 		if (GCVarEnableAudioCommandLogging == 1)
 		{
@@ -375,7 +364,7 @@ void FAudioThread::GetLongestTaskInfo(FString& OutLongestTask, double& OutLonges
 
 void FAudioThread::ProcessAllCommands()
 {
-	if (bIsAudioThreadRunning)
+	if (IsAudioThreadRunning())
 	{
 		GAudioAsyncBatcher.Flush();
 	}
@@ -387,7 +376,7 @@ void FAudioThread::ProcessAllCommands()
 
 void FAudioThread::RunCommandOnGameThread(TUniqueFunction<void()> InFunction, const TStatId InStatId)
 {
-	if (bIsAudioThreadRunning)
+	if (IsAudioThreadRunning())
 	{
 		check(IsInAudioThread());
 		FFunctionGraphTask::CreateAndDispatchWhenReady(MoveTemp(InFunction), InStatId, nullptr, ENamedThreads::GameThread);
@@ -404,7 +393,7 @@ void FAudioThread::StartAudioThread()
 {
 	check(FPlatformTLS::GetCurrentThreadId() == GGameThreadId);
 
-	check(!bIsAudioThreadRunning.Load());
+	check(!IsAudioThreadRunning());
 	check(!GIsAudioThreadSuspended.Load());
 	if (bUseThreadedAudio)
 	{
@@ -415,8 +404,6 @@ PRAGMA_ENABLE_DEPRECATION_WARNINGS
 		static uint32 ThreadCount = 0;
 		check(!ThreadCount); // we should not stop and restart the audio thread; it is complexity we don't need.
 
-		bIsAudioThreadRunning = true;
-
 		// Create the audio thread.
 		AudioThreadRunnable = new FAudioThread();
 
@@ -447,7 +434,7 @@ void FAudioThread::StopAudioThread()
 	check(FPlatformTLS::GetCurrentThreadId() == GGameThreadId);
 	check(!GIsAudioThreadSuspended.Load() || CVarSuspendAudioThread.GetValueOnGameThread() != 0);
 
-	if (!bIsAudioThreadRunning.Load())
+	if (!IsAudioThreadRunning())
 	{
 		return;
 	}
@@ -471,8 +458,6 @@ PRAGMA_DISABLE_DEPRECATION_WARNINGS
 	GAudioThread = nullptr;
 PRAGMA_ENABLE_DEPRECATION_WARNINGS
 
-	bIsAudioThreadRunning = false;
-
 	delete AudioThreadRunnable;
 	AudioThreadRunnable = nullptr;
 }
@@ -495,9 +480,7 @@ FAudioCommandFence::~FAudioCommandFence()
 
 void FAudioCommandFence::BeginFence()
 {
-PRAGMA_DISABLE_DEPRECATION_WARNINGS
-	if (FAudioThread::IsAudioThreadRunning())
-PRAGMA_ENABLE_DEPRECATION_WARNINGS
+	if (IsAudioThreadRunning())
 	{
 		DECLARE_CYCLE_STAT(TEXT("FNullGraphTask.FenceAudioCommand"),
 			STAT_FNullGraphTask_FenceAudioCommand,
@@ -536,9 +519,7 @@ bool FAudioCommandFence::IsFenceComplete() const
 		return true;
 	}
 
-PRAGMA_DISABLE_DEPRECATION_WARNINGS
-	check(FAudioThread::IsAudioThreadRunning());
-PRAGMA_ENABLE_DEPRECATION_WARNINGS
+	check(IsAudioThreadRunning());
 
 	return FenceDoneEvent->Wait(0);
 }
diff --git a/Engine/Source/Runtime/Engine/Public/AudioThread.h b/Engine/Source/Runtime/Engine/Public/AudioThread.h
--- a/Engine/Source/Runtime/Engine/Public/AudioThread.h
+++ b/Engine/Source/Runtime/Engine/Public/AudioThread.h
@@ -82,9 +82,6 @@ public:
 	static ENGINE_API void SetUseThreadedAudio(bool bInUseThreadedAudio);
 	static ENGINE_API bool IsUsingThreadedAudio();
 
-	UE_DEPRECATED(4.26, "Please use global `IsAudioThreadRunning()`")
-	static ENGINE_API bool IsAudioThreadRunning();
-
 	static ENGINE_API void SuspendAudioThread();
 	static ENGINE_API void ResumeAudioThread();
 
