[DebugRenderSceneProxy] Preserve and use the value of the alpha when displaying debug meshes

diff --git a/Engine/Source/Runtime/Engine/Private/DebugRenderSceneProxy.cpp b/Engine/Source/Runtime/Engine/Private/DebugRenderSceneProxy.cpp
--- a/Engine/Source/Runtime/Engine/Private/DebugRenderSceneProxy.cpp
+++ b/Engine/Source/Runtime/Engine/Private/DebugRenderSceneProxy.cpp
@@ -271,7 +271,7 @@ void FDebugRenderSceneProxy::GetDynamicMeshElements(const TArray<const FSceneVie
 				MeshBuilder.AddVertices(Mesh.Vertices);
 				MeshBuilder.AddTriangles(Mesh.Indices);
 
-				MeshBuilder.GetMesh(FMatrix::Identity, MaterialCache[Mesh.Color.A == 255 ? 1 : 0][Mesh.Color.WithAlpha(DrawAlpha)], SDPG_World, false, false, ViewIndex, Collector);
+				MeshBuilder.GetMesh(FMatrix::Identity, MaterialCache[Mesh.Color.A == 255 ? 1 : 0][Mesh.Color], SDPG_World, false, false, ViewIndex, Collector);
 			}
 
 		}
