ActorDescContainer: merged identical functions and fix delegate broadcast when removing actors.

diff --git a/Engine/Source/Runtime/Engine/Private/WorldPartition/ActorDescContainer.cpp b/Engine/Source/Runtime/Engine/Private/WorldPartition/ActorDescContainer.cpp
--- a/Engine/Source/Runtime/Engine/Private/WorldPartition/ActorDescContainer.cpp
+++ b/Engine/Source/Runtime/Engine/Private/WorldPartition/ActorDescContainer.cpp
@@ -203,14 +203,7 @@ void UActorDescContainer::OnPackageDeleted(UPackage* Package)
 
 	if (ShouldHandleActorEvent(Actor))
 	{
-		if (TUniquePtr<FWorldPartitionActorDesc>* ExistingActorDesc = GetActorDescriptor(Actor->GetActorGuid()))
-		{
-			FWorldPartitionActorDesc* const RemovedActorDesc = ExistingActorDesc->Get();
-			OnActorDescRemovedEvent.Broadcast(RemovedActorDesc);
-			OnActorDescRemoved(RemovedActorDesc);
-			RemoveActorDescriptor(RemovedActorDesc);
-			ExistingActorDesc->Reset();
-		}
+		RemoveActor(Actor->GetActorGuid());
 	}
 }
 
@@ -218,8 +211,7 @@ void UActorDescContainer::RemoveActor(const FGuid& ActorGuid)
 {
 	if (TUniquePtr<FWorldPartitionActorDesc>* ExistingActorDesc = GetActorDescriptor(ActorGuid))
 	{
-		// This should never be called on already loaded actors
-		check(!(*ExistingActorDesc)->GetActor());
+		OnActorDescRemovedEvent.Broadcast(ExistingActorDesc->Get());
 		OnActorDescRemoved(ExistingActorDesc->Get());
 		RemoveActorDescriptor(ExistingActorDesc->Get());
 		ExistingActorDesc->Reset();
@@ -253,4 +245,4 @@ void UActorDescContainer::OnActorDescRemoved(FWorldPartitionActorDesc* ActorDesc
 {
 	OnActorDescRemovedEvent.Broadcast(ActorDesc);
 }
-#endif
+#endif
\ No newline at end of file
