Remove DLLEXPORT from header enums.
They are not necessary and clang generates errors.

diff --git a/Engine/Plugins/Media/MediaIOFramework/Source/MediaIOCore/Public/MediaIOCoreEncodeTime.h b/Engine/Plugins/Media/MediaIOFramework/Source/MediaIOCore/Public/MediaIOCoreEncodeTime.h
--- a/Engine/Plugins/Media/MediaIOFramework/Source/MediaIOCore/Public/MediaIOCoreEncodeTime.h
+++ b/Engine/Plugins/Media/MediaIOFramework/Source/MediaIOCore/Public/MediaIOCoreEncodeTime.h
@@ -5,7 +5,7 @@
 #include "CoreMinimal.h"
 #include "IMediaTextureSample.h"
 
-enum class MEDIAIOCORE_API EMediaIOCoreEncodePixelFormat
+enum class EMediaIOCoreEncodePixelFormat
 {
 	A2B10G10R10,
 	CharBGRA,
diff --git a/Engine/Source/Runtime/Engine/Classes/Engine/NetDriver.h b/Engine/Source/Runtime/Engine/Classes/Engine/NetDriver.h
--- a/Engine/Source/Runtime/Engine/Classes/Engine/NetDriver.h
+++ b/Engine/Source/Runtime/Engine/Classes/Engine/NetDriver.h
@@ -1131,7 +1131,7 @@ public:
 	void UpdateStandbyCheatStatus(void);
 
 	/** Sets the analytics provider */
-	virtual void ENGINE_API SetAnalyticsProvider(TSharedPtr<IAnalyticsProvider> InProvider);
+	ENGINE_API virtual void SetAnalyticsProvider(TSharedPtr<IAnalyticsProvider> InProvider);
 
 #if DO_ENABLE_NET_TEST
 	FPacketSimulationSettings	PacketSimulationSettings;
@@ -1283,7 +1283,7 @@ public:
 	ENGINE_API virtual void ProcessRemoteFunction(class AActor* Actor, class UFunction* Function, void* Parameters, struct FOutParmRec* OutParms, struct FFrame* Stack, class UObject* SubObject = nullptr );
 
 
-	enum ENGINE_API ERemoteFunctionSendPolicy
+	enum ERemoteFunctionSendPolicy
 	{		
 		/** Unreliable multicast are queued. Everything else is send immediately */
 		Default, 
@@ -1475,7 +1475,7 @@ public:
 	/**
 	 * Get the socket subsytem appropriate for this net driver
 	 */
-	virtual class ISocketSubsystem* GetSocketSubsystem() PURE_VIRTUAL(UNetDriver::GetSocketSubsystem, return NULL;);
+	ENGINE_API virtual class ISocketSubsystem* GetSocketSubsystem() PURE_VIRTUAL(UNetDriver::GetSocketSubsystem, return NULL;);
 
 	/**
 	 * Associate a world with this net driver. 
@@ -1488,7 +1488,7 @@ public:
 	/**
 	 * Get the world associated with this net driver
 	 */
-	virtual class UWorld* GetWorld() const override final { return World; }
+	ENGINE_API virtual class UWorld* GetWorld() const override final { return World; }
 
 	class UPackage* GetWorldPackage() const { return WorldPackage; }
 
@@ -1496,7 +1496,7 @@ public:
 	ENGINE_API virtual void ResetGameWorldState();
 
 	/** @return true if the net resource is valid or false if it should not be used */
-	virtual bool IsNetResourceValid(void) PURE_VIRTUAL(UNetDriver::IsNetResourceValid, return false;);
+	ENGINE_API virtual bool IsNetResourceValid(void) PURE_VIRTUAL(UNetDriver::IsNetResourceValid, return false;);
 
 	bool NetObjectIsDynamic(const UObject *Object) const;
 
@@ -1510,25 +1510,25 @@ public:
 	TSharedPtr<FRepChangedPropertyTracker> FindOrCreateRepChangedPropertyTracker(UObject *Obj);
 
 	/** Returns true if the client should destroy immediately any actor that becomes torn-off */
-	virtual bool ShouldClientDestroyTearOffActors() const { return false; }
+	ENGINE_API virtual bool ShouldClientDestroyTearOffActors() const { return false; }
 
 	/** Returns whether or not properties that are replicating using this driver should not call RepNotify functions. */
 	ENGINE_API virtual bool ShouldSkipRepNotifies() const;
 
 	/** Returns true if actor channels with InGUID should queue up bunches, even if they wouldn't otherwise be queued. */
-	virtual bool ShouldQueueBunchesForActorGUID(FNetworkGUID InGUID) const { return false; }
+	ENGINE_API virtual bool ShouldQueueBunchesForActorGUID(FNetworkGUID InGUID) const { return false; }
 
 	/** Returns whether or not RPCs processed by this driver should be ignored. */
-	virtual bool ShouldIgnoreRPCs() const { return false; }
+	ENGINE_API virtual bool ShouldIgnoreRPCs() const { return false; }
 
 	/** Returns the existing FNetworkGUID of InActor, if it has one. */
-	virtual FNetworkGUID GetGUIDForActor(const AActor* InActor) const { return FNetworkGUID(); }
+	ENGINE_API virtual FNetworkGUID GetGUIDForActor(const AActor* InActor) const { return FNetworkGUID(); }
 
 	/** Returns the actor that corresponds to InGUID, if one can be found. */
-	virtual AActor* GetActorForGUID(FNetworkGUID InGUID) const { return nullptr; }
+	ENGINE_API virtual AActor* GetActorForGUID(FNetworkGUID InGUID) const { return nullptr; }
 
 	/** Returns true if RepNotifies should be checked and generated when receiving properties for the given object. */
-	virtual bool ShouldReceiveRepNotifiesForObject(UObject* Object) const { return true; }
+	ENGINE_API virtual bool ShouldReceiveRepNotifiesForObject(UObject* Object) const { return true; }
 
 	/** Returns the object that manages the list of replicated UObjects. */
 	ENGINE_API FNetworkObjectList& GetNetworkObjectList() { return *NetworkObjects; }
diff --git a/Engine/Source/Runtime/Experimental/Chaos/Public/Chaos/CollisionResolutionTypes.h b/Engine/Source/Runtime/Experimental/Chaos/Public/Chaos/CollisionResolutionTypes.h
--- a/Engine/Source/Runtime/Experimental/Chaos/Public/Chaos/CollisionResolutionTypes.h
+++ b/Engine/Source/Runtime/Experimental/Chaos/Public/Chaos/CollisionResolutionTypes.h
@@ -9,14 +9,14 @@ namespace Chaos
 {
 
 	/** Specifies the type of work we should do*/
-	enum class CHAOS_API ECollisionUpdateType
+	enum class ECollisionUpdateType
 	{
 		Any,	//stop if we have at least one deep penetration. Does not compute location or normal
 		Deepest	//find the deepest penetration. Compute location and normal
 	};
 
 	/** Return value of the collision modification callback */
-	enum class CHAOS_API ECollisionModifierResult
+	enum class ECollisionModifierResult
 	{
 		Unchanged,	/** No change to the collision */
 		Modified,	/** Modified the collision, but want it to remain enabled */
@@ -24,7 +24,7 @@ namespace Chaos
 	};
 
 	/** The shape types involved in a contact constraint. Used to look up the collision detection function */
-	enum class CHAOS_API EContactShapesType
+	enum class EContactShapesType
 	{
 		Unknown,
 		SphereSphere,
@@ -53,7 +53,7 @@ namespace Chaos
 		NumShapesTypes
 	};
 
-	enum class CHAOS_API EContactManifoldType
+	enum class EContactManifoldType
 	{
 		None,			// No manifold - run collision detection whenever we need latest contact
 		OneShot,		// A manifold is created once and reused. The manifold consists of a plane attached to one shape and a set of points on the other
diff --git a/Engine/Source/Runtime/Experimental/Chaos/Public/Field/FieldSystem.h b/Engine/Source/Runtime/Experimental/Chaos/Public/Field/FieldSystem.h
--- a/Engine/Source/Runtime/Experimental/Chaos/Public/Field/FieldSystem.h
+++ b/Engine/Source/Runtime/Experimental/Chaos/Public/Field/FieldSystem.h
@@ -55,7 +55,7 @@ struct CHAOS_API FFieldContextIndex
 class CHAOS_API FFieldSystemMetaData {
 public:
 
-	enum CHAOS_API EMetaType
+	enum EMetaType
 	{
 		ECommandData_None = 0,
 		ECommandData_ProcessingResolution,
diff --git a/Engine/Source/Runtime/InstallBundleManager/Public/InstallBundleUtils.h b/Engine/Source/Runtime/InstallBundleManager/Public/InstallBundleUtils.h
--- a/Engine/Source/Runtime/InstallBundleManager/Public/InstallBundleUtils.h
+++ b/Engine/Source/Runtime/InstallBundleManager/Public/InstallBundleUtils.h
@@ -183,7 +183,7 @@ namespace InstallBundleUtil
 		//_FG stats account for time while the app is active
 		//_BG stats are an estimate of time while the app is inactive (backgrounded, closed, etc.)
 		//		*_BG stats are an estimate as we can not account for unexpected crashes, etc.
-		enum class INSTALLBUNDLEMANAGER_API ETimingStatNames : uint8
+		enum class ETimingStatNames : uint8
 		{
 			  TotalTime_Real
 			, TotalTime_FG
@@ -201,7 +201,7 @@ namespace InstallBundleUtil
 		};
 
 		//Stats used to store count information in an int format
-		enum class INSTALLBUNDLEMANAGER_API ECountStatNames : uint8
+		enum class ECountStatNames : uint8
 		{
 			  NumResumedFromBackground
 			, NumResumedFromLaunch
diff --git a/Engine/Source/Runtime/SignalProcessing/Public/DSP/VariablePoleFilter.h b/Engine/Source/Runtime/SignalProcessing/Public/DSP/VariablePoleFilter.h
--- a/Engine/Source/Runtime/SignalProcessing/Public/DSP/VariablePoleFilter.h
+++ b/Engine/Source/Runtime/SignalProcessing/Public/DSP/VariablePoleFilter.h
@@ -7,7 +7,7 @@
 
 namespace Audio
 {
-	enum class SIGNALPROCESSING_API EFilterOrder : uint32
+	enum class EFilterOrder : uint32
 	{
 		TwoPole = 1,
 		FourPole,
diff --git a/Engine/Source/Runtime/SignalProcessing/Public/DSP/WaveShaper.h b/Engine/Source/Runtime/SignalProcessing/Public/DSP/WaveShaper.h
--- a/Engine/Source/Runtime/SignalProcessing/Public/DSP/WaveShaper.h
+++ b/Engine/Source/Runtime/SignalProcessing/Public/DSP/WaveShaper.h
@@ -6,7 +6,7 @@
 
 namespace Audio
 {
-	enum class SIGNALPROCESSING_API EWaveShaperType : int32
+	enum class EWaveShaperType : int32
 	{
 		Sin,
 		ATan,
