Merge additional files from UE5/RES @ 15958325 to UE5/Main
This represents UE4/Main @ 15913390 and Dev-PerfTest @ 15913304

diff --git a/Engine/Source/Runtime/Engine/Private/Components/SkeletalMeshComponent.cpp b/Engine/Source/Runtime/Engine/Private/Components/SkeletalMeshComponent.cpp
--- a/Engine/Source/Runtime/Engine/Private/Components/SkeletalMeshComponent.cpp
+++ b/Engine/Source/Runtime/Engine/Private/Components/SkeletalMeshComponent.cpp
@@ -2138,6 +2138,7 @@ void USkeletalMeshComponent::HandleExistingParallelClothSimulation()
 
 	if(IsValidRef(ParallelClothTask))
 	{
+		QUICK_SCOPE_CYCLE_COUNTER(STAT_EndParallelClothTask);
 		CSV_SCOPED_SET_WAIT_STAT(Cloth);
 
 		// There's a simulation in flight
diff --git a/Engine/Source/Runtime/Engine/Private/GameplayStatics.cpp b/Engine/Source/Runtime/Engine/Private/GameplayStatics.cpp
--- a/Engine/Source/Runtime/Engine/Private/GameplayStatics.cpp
+++ b/Engine/Source/Runtime/Engine/Private/GameplayStatics.cpp
@@ -1513,7 +1513,7 @@ UAudioComponent* UGameplayStatics::SpawnSoundAtLocation(const UObject* WorldCont
 		AudioComponent->SetWorldLocationAndRotation(Location, Rotation);
 		AudioComponent->SetVolumeMultiplier(VolumeMultiplier);
 		AudioComponent->SetPitchMultiplier(PitchMultiplier);
-		AudioComponent->bAllowSpatialization	= bIsInGameWorld;
+		AudioComponent->bAllowSpatialization	= Params.ShouldUseAttenuation();
 		AudioComponent->bIsUISound				= !bIsInGameWorld;
 		AudioComponent->bAutoDestroy			= bAutoDestroy;
 		AudioComponent->SubtitlePriority		= Sound->GetSubtitlePriority();
@@ -1598,7 +1598,7 @@ UAudioComponent* UGameplayStatics::SpawnSoundAttached(USoundBase* Sound, USceneC
 
 			AudioComponent->SetVolumeMultiplier(VolumeMultiplier);
 			AudioComponent->SetPitchMultiplier(PitchMultiplier);
-			AudioComponent->bAllowSpatialization = bIsInGameWorld;
+			AudioComponent->bAllowSpatialization = Params.ShouldUseAttenuation();
 			AudioComponent->bIsUISound = !bIsInGameWorld;
 			AudioComponent->bAutoDestroy = bAutoDestroy;
 			AudioComponent->SubtitlePriority = Sound->GetSubtitlePriority();
diff --git a/Engine/Source/Runtime/Engine/Private/World.cpp b/Engine/Source/Runtime/Engine/Private/World.cpp
--- a/Engine/Source/Runtime/Engine/Private/World.cpp
+++ b/Engine/Source/Runtime/Engine/Private/World.cpp
@@ -134,6 +134,7 @@
 #include "Engine/HLODProxy.h"
 #include "ProfilingDebugging/CsvProfiler.h"
 #include "ObjectTrace.h"
+#include "ReplaySubsystem.h"
 
 #if INCLUDE_CHAOS
 #include "ChaosSolversModule.h"
@@ -3360,7 +3361,9 @@ UWorld* UWorld::DuplicateWorldForPIE(const FString& PackageName, UWorld* OwningW
 	PIELevelPackage->SetPackageFlags(PKG_PlayInEditor);
 	PIELevelPackage->PIEInstanceID = PIEInstanceID;
 	PIELevelPackage->SetLoadedPath(FPackagePath::FromPackageNameChecked(PackageFName));
+	PRAGMA_DISABLE_DEPRECATION_WARNINGS
 	PIELevelPackage->SetGuid( EditorLevelPackage->GetGuid() );
+	PRAGMA_ENABLE_DEPRECATION_WARNINGS
 	PIELevelPackage->MarkAsFullyLoaded();
 
 	ULevel::StreamedLevelsOwningWorld.Add(PIELevelPackage->GetFName(), OwningWorld);
@@ -5962,8 +5965,18 @@ bool UWorld::IsPlayingReplay() const
 
 bool UWorld::IsRecordingReplay() const
 {
-	// Using IsServer() because it also calls IsRecording() internally
-	return (DemoNetDriver && DemoNetDriver->IsServer());
+	const UGameInstance* GameInst = GetGameInstance();
+	const UReplaySubsystem* ReplaySubsystem = GameInst ? GameInst->GetSubsystem<UReplaySubsystem>() : nullptr;
+
+	if (ReplaySubsystem)
+	{
+		return ReplaySubsystem->IsRecording();
+	}
+	else
+	{
+		// Using IsServer() because it also calls IsRecording() internally
+		return (DemoNetDriver && DemoNetDriver->IsServer());
+	}
 }
 PRAGMA_ENABLE_DEPRECATION_WARNINGS
 
@@ -6997,7 +7010,9 @@ void UWorld::SeamlessTravel(const FString& SeamlessTravelURL, bool bAbsolute, FG
 		}
 		// tell the handler to start the transition
 		FSeamlessTravelHandler &SeamlessTravelHandler = GEngine->SeamlessTravelHandlerForWorld( this );
+		PRAGMA_DISABLE_DEPRECATION_WARNINGS
 		if (!SeamlessTravelHandler.StartTravel(this, NewURL, MapPackageGuid) && !SeamlessTravelHandler.IsInTransition())
+		PRAGMA_ENABLE_DEPRECATION_WARNINGS
 		{
 			const FString Error = FText::Format( NSLOCTEXT("Engine", "InvalidUrl", "Invalid URL: {0}"), FText::FromString( SeamlessTravelURL ) ).ToString();
 			GEngine->BroadcastTravelFailure(this, ETravelFailure::InvalidURL, Error);
@@ -7842,7 +7857,9 @@ static ULevel* DuplicateLevelWithPrefix(ULevel* InLevel, int32 InstanceID )
 	NewPackage->SetPackageFlags( PKG_PlayInEditor );
 	NewPackage->PIEInstanceID = InstanceID;
 	NewPackage->SetLoadedPath(OriginalPackage->GetLoadedPath());
+	PRAGMA_DISABLE_DEPRECATION_WARNINGS
 	NewPackage->SetGuid( OriginalPackage->GetGuid() );
+	PRAGMA_ENABLE_DEPRECATION_WARNINGS
 	NewPackage->MarkAsFullyLoaded();
 
 	FSoftObjectPath::AddPIEPackageName(NewPackage->GetFName());
