Enables Gen5 TAA and Temporal Upsampling by default engine wide

Disabled on EngineTest to avoid invalidating screenshot test's gold images
Disabled on QAGame to avoid causing unexpected surprises with short notice until QA have tried it out.
Has already been enabled on ShooterGame and all other internal projects

diff --git a/Engine/Source/Runtime/Engine/Classes/Engine/RendererSettings.h b/Engine/Source/Runtime/Engine/Classes/Engine/RendererSettings.h
--- a/Engine/Source/Runtime/Engine/Classes/Engine/RendererSettings.h
+++ b/Engine/Source/Runtime/Engine/Classes/Engine/RendererSettings.h
@@ -529,9 +529,15 @@ class ENGINE_API URendererSettings : public UDeveloperSettings
 	UPROPERTY(config, EditAnywhere, Category = DefaultSettings, meta = (
 		EditCondition = "DefaultFeatureAntiAliasing == EAntiAliasingMethod::AAM_TemporalAA",
 		ConsoleVariable = "r.TemporalAA.Upsampling", DisplayName = "Temporal Upsampling",
-		ToolTip = "Whether to do primary screen percentage with temporal AA or not."))
+		ToolTip = "Whether to do primary screen percentage upscale with Temporal AA pass or not."))
 	uint32 bTemporalUpsampling : 1;
 
+	UPROPERTY(config, EditAnywhere, Category = DefaultSettings, meta = (
+		EditCondition = "DefaultFeatureAntiAliasing == EAntiAliasingMethod::AAM_TemporalAA",
+		ConsoleVariable = "r.TemporalAA.Algorithm", DisplayName = "Temporal Anti-Aliasing for 5th generation of consoles",
+		ToolTip = "Whether to enable UE5's new Temporal Anti-Aliasing method. This new technique is mainly built for handling Nanite amount of details on screen and higher upscale quality with Temporal Upsampling and ScreenPercentage/Dynamic Resolution on desktops and 5th generation of consoles."))
+	uint32 bUseGen5TAA : 1;
+
 	UPROPERTY(config, EditAnywhere, Category = DefaultSettings, meta = (
 		ConsoleVariable = "r.DefaultFeature.AntiAliasing", DisplayName = "Anti-Aliasing Method",
 		ToolTip = "Which anti-aliasing mode is used by default"))
diff --git a/Engine/Source/Runtime/Engine/Private/SceneView.cpp b/Engine/Source/Runtime/Engine/Private/SceneView.cpp
--- a/Engine/Source/Runtime/Engine/Private/SceneView.cpp
+++ b/Engine/Source/Runtime/Engine/Private/SceneView.cpp
@@ -322,10 +322,10 @@ static TAutoConsoleVariable<int32> CVarAllowTranslucencyAfterDOF(
 
 static TAutoConsoleVariable<int32> CVarEnableTemporalUpsample(
 	TEXT("r.TemporalAA.Upsampling"),
-	0,
+	1,
 	TEXT("Whether to do primary screen percentage with temporal AA or not.\n")
-	TEXT(" 0: use spatial upscale pass independently of TAA (default);\n")
-	TEXT(" 1: TemporalAA performs spatial and temporal upscale as screen percentage method."),
+	TEXT(" 0: use spatial upscale pass independently of TAA;\n")
+	TEXT(" 1: TemporalAA performs spatial and temporal upscale as screen percentage method (default)."),
 	ECVF_Default);
 
 // Conversion factor used when "r.DefaultFeature.AutoExposure.ExtendDefaultLuminanceRange" 
diff --git a/Engine/Source/Runtime/Renderer/Private/PostProcess/TemporalAA.cpp b/Engine/Source/Runtime/Renderer/Private/PostProcess/TemporalAA.cpp
--- a/Engine/Source/Runtime/Renderer/Private/PostProcess/TemporalAA.cpp
+++ b/Engine/Source/Runtime/Renderer/Private/PostProcess/TemporalAA.cpp
@@ -20,7 +20,7 @@ const int32 GTemporalAATileSizeX = 8;
 const int32 GTemporalAATileSizeY = 8;
 
 TAutoConsoleVariable<int32> CVarTAAAlgorithm(
-	TEXT("r.TemporalAA.Algorithm"), 0,
+	TEXT("r.TemporalAA.Algorithm"), 1,
 	TEXT("Algorithm to use for Temporal AA\n")
 	TEXT(" 0: Gen 4 TAAU (default)\n")
 	TEXT(" 1: Gen 5 TAAU (experimental)"),
@@ -113,6 +113,11 @@ inline bool DoesPlatformSupportTemporalHistoryUpscale(EShaderPlatform Platform)
 
 inline bool DoesPlatformSupportGen5TAA(EShaderPlatform Platform)
 {
+	// TODO(Gen5TAA): alpha channel is not supported yet
+	if (IsPostProcessingWithAlphaChannelSupported())
+	{
+		return false;
+	}
 	return FDataDrivenShaderPlatformInfo::GetSupportsGen5TemporalAA(Platform);
 }
 
