Anim Runtime: Resize Play Play Length Spams Undo History when playing back some Level Sequences
#fix set sequence length directly for dynamic montages, rather than going through Controller

diff --git a/Engine/Source/Runtime/Engine/Private/Animation/AnimMontage.cpp b/Engine/Source/Runtime/Engine/Private/Animation/AnimMontage.cpp
--- a/Engine/Source/Runtime/Engine/Private/Animation/AnimMontage.cpp
+++ b/Engine/Source/Runtime/Engine/Private/Animation/AnimMontage.cpp
@@ -3049,7 +3049,9 @@ UAnimMontage* UAnimMontage::CreateSlotAnimationAsDynamicMontage(UAnimSequenceBas
 	NewSegment.AnimPlayRate = 1.f;
 	NewSegment.StartPos = 0.f;
 	NewSegment.LoopingCount = LoopCount;
-	NewMontage->SetCompositeLength(NewSegment.GetLength());
+	PRAGMA_DISABLE_DEPRECATION_WARNINGS
+    NewMontage->SequenceLength = NewSegment.GetLength();
+	PRAGMA_ENABLE_DEPRECATION_WARNINGS
 	NewTrack.AnimTrack.AnimSegments.Add(NewSegment);
 
 	FCompositeSection NewSection;
