Generate debug names for skin cache BVHs based on the asset name.

diff --git a/Engine/Source/Runtime/Engine/Private/GPUSkinCache.cpp b/Engine/Source/Runtime/Engine/Private/GPUSkinCache.cpp
--- a/Engine/Source/Runtime/Engine/Private/GPUSkinCache.cpp
+++ b/Engine/Source/Runtime/Engine/Private/GPUSkinCache.cpp
@@ -1530,9 +1530,21 @@ void FGPUSkinCache::ProcessRayTracingGeometryToUpdate(
 			}
 
 			FRayTracingGeometryInitializer Initializer;
-			static const FName DebugName("FSkeletalMeshObjectGPUSkin");
-			static int32 DebugNumber = 0;
-			Initializer.DebugName = FName(DebugName, DebugNumber++);
+			
+		#if !UE_BUILD_SHIPPING
+			FName OwnerName = SkinCacheEntry->GPUSkin->DebugName;
+			if (OwnerName.IsValid())
+			{
+				Initializer.DebugName = OwnerName;
+			}
+			else
+		#endif
+			{
+				static const FName DefaultDebugName("FSkeletalMeshObjectGPUSkin");
+				static int32 DebugNumber = 0;
+				Initializer.DebugName = FName(DefaultDebugName, DebugNumber++);
+			}
+
 			Initializer.IndexBuffer = IndexBufferRHI;
 			Initializer.TotalPrimitiveCount = TrianglesCount;
 			Initializer.GeometryType = RTGT_Triangles;
diff --git a/Engine/Source/Runtime/Engine/Private/SkeletalRender.cpp b/Engine/Source/Runtime/Engine/Private/SkeletalRender.cpp
--- a/Engine/Source/Runtime/Engine/Private/SkeletalRender.cpp
+++ b/Engine/Source/Runtime/Engine/Private/SkeletalRender.cpp
@@ -32,6 +32,9 @@ FSkeletalMeshObject::FSkeletalMeshObject(USkinnedMeshComponent* InMeshComponent,
 #if RHI_RAYTRACING
 , bRequireRecreatingRayTracingGeometry(false)
 #endif
+#if !UE_BUILD_SHIPPING
+, DebugName(InMeshComponent->SkeletalMesh->GetFName())
+#endif // !UE_BUILD_SHIPPING
 #if WITH_EDITORONLY_DATA
 ,   SectionIndexPreview(InMeshComponent->GetSectionPreview())
 ,   MaterialIndexPreview(InMeshComponent->GetMaterialPreview())
diff --git a/Engine/Source/Runtime/Engine/Public/SkeletalRenderPublic.h b/Engine/Source/Runtime/Engine/Public/SkeletalRenderPublic.h
--- a/Engine/Source/Runtime/Engine/Public/SkeletalRenderPublic.h
+++ b/Engine/Source/Runtime/Engine/Public/SkeletalRenderPublic.h
@@ -246,6 +246,10 @@ public:
 	bool bSupportRayTracing;
 #endif
 
+#if !UE_BUILD_SHIPPING
+	FName DebugName;
+#endif // !UE_BUILD_SHIPPING
+
 #if WITH_EDITORONLY_DATA
 	/** Index of the section to preview... If set to -1, all section will be rendered */
 	int32 SectionIndexPreview;
