Relax ensure() when building VT thumbnail, allow high priority requests to return Invalid status (can occur with sparse UDIM)

diff --git a/Engine/Source/Runtime/Engine/Private/Texture2D.cpp b/Engine/Source/Runtime/Engine/Private/Texture2D.cpp
--- a/Engine/Source/Runtime/Engine/Private/Texture2D.cpp
+++ b/Engine/Source/Runtime/Engine/Private/Texture2D.cpp
@@ -1339,8 +1339,13 @@ void FVirtualTexture2DResource::InitializeEditorResources(IVirtualTexture* InVir
 			{
 				const uint32 vAddress = FMath::MortonCode2(TileX) | (FMath::MortonCode2(TileY) << 1);
 				const FVTRequestPageResult RequestResult = InVirtualTexture->RequestPageData(ProducerHandle, LayerMask, MipLevel, vAddress, EVTRequestPagePriority::High);
-				// High priority request should always generate data
-				if (ensure(VTRequestPageStatus_HasData(RequestResult.Status)))
+				
+				// High priority request should never be Saturated
+				// It's possible for status to be Invalid, if requesting data from a mip level that doesn't exist for the given producer (when using sparse UDIMs)
+				// Technically could try to handle this, by check LocalMipBias, grabbing lower resolution tile, and resizing...but that would make this code much more complex for very little gain
+				ensure(RequestResult.Status != EVTRequestPageStatus::Saturated);
+
+				if (VTRequestPageStatus_HasData(RequestResult.Status))
 				{
 					PagesToProduce.Add({ RequestResult.Handle, TileX, TileY });
 				}
