Fix RemovePrefixFromName handling of names that contain both platform and custom prefixes.

Also fixes format name passed to Crunch, which is expected to be prefix-free.

diff --git a/Engine/Source/Runtime/Engine/Private/VT/VirtualTextureDataBuilder.cpp b/Engine/Source/Runtime/Engine/Private/VT/VirtualTextureDataBuilder.cpp
--- a/Engine/Source/Runtime/Engine/Private/VT/VirtualTextureDataBuilder.cpp
+++ b/Engine/Source/Runtime/Engine/Private/VT/VirtualTextureDataBuilder.cpp
@@ -511,6 +511,18 @@ void FVirtualTextureDataBuilder::BuildPagesMacroBlocks(bool bAllowAsync)
 	}
 }
 
+static const TCHAR* GetSafePixelFormatName(EPixelFormat Format)
+{
+	if (Format >= PF_MAX)
+	{
+		return TEXT("INVALID");
+	}
+	else
+	{
+		return GPixelFormats[Format].Name;
+	}
+}
+
 void FVirtualTextureDataBuilder::BuildTiles(const TArray<FVTSourceTileEntry>& TileList, uint32 LayerIndex, FLayerData& GeneratedData, bool bAllowAsync)
 {
 	const FTextureBuildSettings& BuildSettingsLayer0 = SettingsPerLayer[0];
@@ -675,7 +687,11 @@ void FVirtualTextureDataBuilder::BuildTiles(const TArray<FVTSourceTileEntry>& Ti
 			}
 
 			check(CompressedMip.Num() == 1);
-			check(CompressedFormat == PF_Unknown || CompressedFormat == CompressedMip[0].PixelFormat);
+			checkf(CompressedFormat == PF_Unknown || CompressedFormat == CompressedMip[0].PixelFormat, 
+				TEXT("CompressedFormat: %s (%d), CompressedMip[0].PixelFormat: %s (%d)"),
+				GetSafePixelFormatName(CompressedFormat), (int32)CompressedFormat, 
+				GetSafePixelFormatName((EPixelFormat)CompressedMip[0].PixelFormat), (int32)CompressedMip[0].PixelFormat);
+
 			CompressedFormat = (EPixelFormat)CompressedMip[0].PixelFormat;
 
 			const uint32 SizeRaw = CompressedMip[0].RawData.Num() * CompressedMip[0].RawData.GetTypeSize();
@@ -807,18 +823,31 @@ int32 FVirtualTextureDataBuilder::FindSourceBlockIndex(int32 MipIndex, int32 Blo
 	return INDEX_NONE;
 }
 
+// Removes platform and other custom prefixes from the name.
+// Returns plain format name and the non-platform prefix (with trailing underscore).
+// i.e. PLAT_BLAH_AutoDXT returns AutoDXT and writes BLAH_ to OutPrefix.
 static const FName RemovePrefixFromName(FName const& InName, FName& OutPrefix)
 {
 	FString NameString = InName.ToString();
 
+	// Format names may have one of the following forms:
+	// - PLATFORM_PREFIX_FORMAT
+	// - PLATFORM_FORMAT
+	// - PREFIX_FORMAT
+	// - FORMAT
+	// We have to remove the platform prefix first, if it exists.
+	// Then we detect a non-platform prefix (such as codec name)
+	// and split the result into  explicit FORMAT and PREFIX parts.
+
 	for (FName PlatformName : FDataDrivenPlatformInfoRegistry::GetSortedPlatformNames())
 	{
 		FString PlatformTextureFormatPrefix = PlatformName.ToString();
 		PlatformTextureFormatPrefix += TEXT('_');
 		if (NameString.StartsWith(PlatformTextureFormatPrefix, ESearchCase::IgnoreCase))
 		{
-			OutPrefix = FName(); // don't keep platform prefixes
-			return *NameString.RightChop(PlatformTextureFormatPrefix.Len());
+			// Remove platform prefix and proceed with non-platform prefix detection.
+			NameString = NameString.RightChop(PlatformTextureFormatPrefix.Len());
+			break;
 		}
 	}
 
@@ -830,7 +859,7 @@ static const FName RemovePrefixFromName(FName const& InName, FName& OutPrefix)
 		return *NameString.RightChop(UnderscoreIndex + 1);
 	}
 
-	return InName;
+	return *NameString;
 }
 
 // This builds an uncompressed version of the texture containing all other build settings baked in
@@ -1167,7 +1196,9 @@ void FVirtualTextureDataBuilder::BuildSourcePixels(const FTextureSourceData& Sou
 				CrunchCompression::IsValidFormat(TextureFormatName);
 		}
 #endif // WITH_CRUNCH_COMPRESSION
-		if (TextureFormatPrefix.IsNone())
+
+		if (bUseCrunch // NOTE: Crunch expects a format with no prefixes. See GetCrnFormat().
+			|| TextureFormatPrefix.IsNone())
 		{
 			LayerData.TextureFormatName = TextureFormatName;
 		}
