Fixed fatal error on startup in -game mode due to missing special material.

diff --git a/Engine/Source/Runtime/Engine/Private/UnrealEngine.cpp b/Engine/Source/Runtime/Engine/Private/UnrealEngine.cpp
--- a/Engine/Source/Runtime/Engine/Private/UnrealEngine.cpp
+++ b/Engine/Source/Runtime/Engine/Private/UnrealEngine.cpp
@@ -2699,6 +2699,9 @@ void UEngine::InitializeObjectReferences()
 
 		ClothPaintMaterialInstance = UMaterialInstanceDynamic::Create(ClothPaintMaterial, nullptr);
 		ClothPaintMaterialWireframeInstance = UMaterialInstanceDynamic::Create(ClothPaintMaterialWireframe, nullptr);
+
+		LoadSpecialMaterial(TEXT("DefaultFlattenMaterialName"), DefaultFlattenMaterialName.ToString(), DefaultFlattenMaterial, false);
+		LoadSpecialMaterial(TEXT("DefaultHLODFlattenMaterialName"), DefaultHLODFlattenMaterialName.ToString(), DefaultHLODFlattenMaterial, false);
 #endif
 		FString ValidPreviewShadowsIndicatorMaterialName = PreviewShadowsIndicatorMaterialName.ToString();
 		if (ValidPreviewShadowsIndicatorMaterialName.IsEmpty())
@@ -2715,11 +2718,6 @@ void UEngine::InitializeObjectReferences()
 		}
 	}
 
-#if WITH_EDITORONLY_DATA
-	LoadSpecialMaterial(TEXT("DefaultFlattenMaterialName"), DefaultFlattenMaterialName.ToString(), DefaultFlattenMaterial, false);
-	LoadSpecialMaterial(TEXT("DefaultHLODFlattenMaterialName"), DefaultHLODFlattenMaterialName.ToString(), DefaultHLODFlattenMaterial, false);
-#endif
-
 	LoadEngineTexture(DefaultTexture, *DefaultTextureName.ToString());
 	LoadEngineTexture(DefaultDiffuseTexture, *DefaultDiffuseTextureName.ToString());
 	LoadEngineTexture(HighFrequencyNoiseTexture, *HighFrequencyNoiseTextureName.ToString());
