HLOD Builder - When deleting HLOD packages, remove their matching ActorDesc from the WorldPartition
#partial cherrypick of 15808954, 15954195

diff --git a/Engine/Source/Runtime/Engine/Private/WorldPartition/RuntimeSpatialHash/RuntimeSpatialHashHLOD.cpp b/Engine/Source/Runtime/Engine/Private/WorldPartition/RuntimeSpatialHash/RuntimeSpatialHashHLOD.cpp
--- a/Engine/Source/Runtime/Engine/Private/WorldPartition/RuntimeSpatialHash/RuntimeSpatialHashHLOD.cpp
+++ b/Engine/Source/Runtime/Engine/Private/WorldPartition/RuntimeSpatialHash/RuntimeSpatialHashHLOD.cpp
@@ -80,15 +80,17 @@ static void DeletePackage(UPackage* Package, ISourceControlHelper* SourceControl
 	}
 }
 
-static void DeletePackage(FWorldPartitionActorDesc* ActorDesc, ISourceControlHelper* SourceControlHelper)
+static void DeletePackage(UWorldPartition* WorldPartition, FWorldPartitionActorDesc* ActorDesc, ISourceControlHelper* SourceControlHelper)
 {
 	if (ActorDesc->IsLoaded())
 	{
 		DeletePackage(ActorDesc->GetActor()->GetPackage(), SourceControlHelper);
+		WorldPartition->OnPackageDeleted(ActorDesc->GetActor()->GetPackage());
 	}
 	else
 	{
 		DeletePackage(ActorDesc->GetActorPackage().ToString(), SourceControlHelper);
+		WorldPartition->RemoveActor(ActorDesc->GetGuid());
 	}
 }
 
@@ -587,12 +589,12 @@ bool UWorldPartitionRuntimeSpatialHash::GenerateHLOD(ISourceControlHelper* Sourc
 		GenerateHLODs(HLODGrids[HLODGridName], GridsDepth[HLODGridName] + 1, HLODActorClusterInstancePtrs);
 	}
 
-	auto DeleteHLODActor = [&SourceControlHelper](FWorldPartitionHandle ActorHandle)
+	auto DeleteHLODActor = [&SourceControlHelper, WorldPartition](FWorldPartitionHandle ActorHandle)
 	{
 		FWorldPartitionActorDesc* HLODActorDesc = ActorHandle.Get();
 		check(HLODActorDesc);
 
-		DeletePackage(HLODActorDesc, SourceControlHelper);
+		DeletePackage(WorldPartition, HLODActorDesc, SourceControlHelper);
 	};
 
 	// Destroy all unreferenced HLOD actors
