Do not write the global shadermap to DDC until it's complete (UE-108245).

diff --git a/Engine/Source/Editor/UnrealEd/Private/Commandlets/DerivedDataCacheCommandlet.cpp b/Engine/Source/Editor/UnrealEd/Private/Commandlets/DerivedDataCacheCommandlet.cpp
--- a/Engine/Source/Editor/UnrealEd/Private/Commandlets/DerivedDataCacheCommandlet.cpp
+++ b/Engine/Source/Editor/UnrealEd/Private/Commandlets/DerivedDataCacheCommandlet.cpp
@@ -275,9 +275,9 @@ int32 UDerivedDataCacheCommandlet::Main( const FString& Params )
 
 			for (int32 FormatIndex = 0; FormatIndex < DesiredShaderFormats.Num(); FormatIndex++)
 			{
-				const EShaderPlatform TargetPlatform = ShaderFormatToLegacyShaderPlatform(DesiredShaderFormats[FormatIndex]);
-				// Kick off global shader compiles for each target platform
-				CompileGlobalShaderMap(TargetPlatform);
+				const EShaderPlatform ShaderPlatform = ShaderFormatToLegacyShaderPlatform(DesiredShaderFormats[FormatIndex]);
+				// Kick off global shader compiles for each target platform. Note that shader platform alone is not sufficient to distinguish between WindowsEditor and WindowsClient, which after UE 4.25 have different DDC
+				CompileGlobalShaderMap(ShaderPlatform, Platforms[Index], false);
 			}
 		}
 
diff --git a/Engine/Source/Runtime/Engine/Private/ShaderCompiler/ShaderCompiler.cpp b/Engine/Source/Runtime/Engine/Private/ShaderCompiler/ShaderCompiler.cpp
--- a/Engine/Source/Runtime/Engine/Private/ShaderCompiler/ShaderCompiler.cpp
+++ b/Engine/Source/Runtime/Engine/Private/ShaderCompiler/ShaderCompiler.cpp
@@ -115,6 +115,20 @@ static FAutoConsoleVariableRef CVarShaderCompilerCacheStatsPrintoutInterval(
 	ECVF_Default
 );
 
+/** Helper functions for logging more debug info */
+namespace ShaderCompiler
+{
+	FString GetTargetPlatformName(const ITargetPlatform* TargetPlatform)
+	{
+		if (TargetPlatform)
+		{
+			return TargetPlatform->PlatformName();
+		}
+
+		return TEXT("(current)");
+	}
+}
+
 /** Storage for the global shadar map(s) that have been replaced by new one(s), which aren't yet compiled.
  * 
  *	Sometimes a mesh drawing command references a pointer to global SM's memory. To nix these MDCs when we're replacing a global SM, we would just recreate the render state for all the components, but
@@ -5386,7 +5400,7 @@ static void PrepareGlobalShaderCompileJob(EShaderPlatform Platform,
 
 	FShaderCompilerEnvironment& ShaderEnvironment = NewJob->Input.Environment;
 
-	UE_LOG(LogShaders, Verbose, TEXT("	%s"), ShaderType->GetName());
+	UE_LOG(LogShaders, Verbose, TEXT("	%s (permutation %d)"), ShaderType->GetName(), Key.PermutationId);
 	COOK_STAT(GlobalShaderCookStats::ShadersCompiled++);
 
 	// Allow the shader type to modify the compile environment.
@@ -5519,7 +5533,7 @@ void VerifyGlobalShaders(EShaderPlatform Platform, const ITargetPlatform* Target
 	LayoutParams.InitializeForPlatform(TargetPlatform);
 	const EShaderPermutationFlags PermutationFlags = GetShaderPermutationFlags(LayoutParams);
 
-	UE_LOG(LogMaterial, Verbose, TEXT("Verifying Global Shaders for %s"), *LegacyShaderPlatformToShaderFormat(Platform).ToString());
+	UE_LOG(LogMaterial, Verbose, TEXT("Verifying Global Shaders for %s (%s)"), *LegacyShaderPlatformToShaderFormat(Platform).ToString(), *ShaderCompiler::GetTargetPlatformName(TargetPlatform));
 
 	// Ensure that the global shader map contains all global shader types.
 	FGlobalShaderMap* GlobalShaderMap = GetGlobalShaderMap(Platform);
@@ -5715,19 +5729,24 @@ static void SaveGlobalShaderMapToDerivedDataCache(EShaderPlatform Platform)
 	TArray<uint8> SaveData;
 
 	FGlobalShaderMapId ShaderMapId(Platform, TargetPlatform);
-	for (auto const& ShaderFilenameDependencies : ShaderMapId.GetShaderFilenameToDependeciesMap())
+	// avoid saving incomplete shadermaps
+	FGlobalShaderMap* GlobalSM = GetGlobalShaderMap(Platform);
+	if (GlobalSM->IsComplete(TargetPlatform))
 	{
-		FGlobalShaderMapSection* Section = GGlobalShaderMap[Platform]->FindSection(ShaderFilenameDependencies.Key);
-		if (Section)
+		for (auto const& ShaderFilenameDependencies : ShaderMapId.GetShaderFilenameToDependeciesMap())
 		{
-			Section->FinalizeContent();
+			FGlobalShaderMapSection* Section = GlobalSM->FindSection(ShaderFilenameDependencies.Key);
+			if (Section)
+			{
+				Section->FinalizeContent();
 
-			SaveData.Reset();
-			FMemoryWriter Ar(SaveData, true);
-			Section->Serialize(Ar);
+				SaveData.Reset();
+				FMemoryWriter Ar(SaveData, true);
+				Section->Serialize(Ar);
 
-			GetDerivedDataCacheRef().Put(*GetGlobalShaderMapKeyString(ShaderMapId, Platform, TargetPlatform, ShaderFilenameDependencies.Value), SaveData, TEXT("GlobalShaderMap"_SV));
-			COOK_STAT(Timer.AddMiss(SaveData.Num()));
+				GetDerivedDataCacheRef().Put(*GetGlobalShaderMapKeyString(ShaderMapId, Platform, TargetPlatform, ShaderFilenameDependencies.Value), SaveData, TEXT("GlobalShaderMap"_SV));
+				COOK_STAT(Timer.AddMiss(SaveData.Num()));
+			}
 		}
 	}
 }
@@ -6374,7 +6393,7 @@ static inline FShader* ProcessCompiledJob(FShaderCompileJob* SingleJob, const FS
 
 void ProcessCompiledGlobalShaders(const TArray<FShaderCommonCompileJobPtr>& CompilationResults)
 {
-	UE_LOG(LogShaders, Warning, TEXT("Compiled %u global shaders"), CompilationResults.Num());
+	UE_LOG(LogShaders, Verbose, TEXT("Compiled %u global shaders"), CompilationResults.Num());
 
 	TArray<EShaderPlatform> ShaderPlatformsProcessed;
 	TArray<const FShaderPipelineType*> SharedPipelines;
@@ -6412,7 +6431,7 @@ void ProcessCompiledGlobalShaders(const TArray<FShaderCommonCompileJobPtr>& Comp
 		{
 			// Process the shader pipelines that share shaders
 			EShaderPlatform Platform = ShaderPlatformsProcessed[PlatformIndex];
-			auto* GlobalShaderMap = GGlobalShaderMap[Platform];
+			FGlobalShaderMap* GlobalShaderMap = GGlobalShaderMap[Platform];
 			const ITargetPlatform* TargetPlatform = GGlobalShaderTargetPlatform[Platform];
 
 			FPlatformTypeLayoutParameters LayoutParams;
diff --git a/Engine/Source/Runtime/Engine/Private/ShaderDerivedDataVersion.h b/Engine/Source/Runtime/Engine/Private/ShaderDerivedDataVersion.h
--- a/Engine/Source/Runtime/Engine/Private/ShaderDerivedDataVersion.h
+++ b/Engine/Source/Runtime/Engine/Private/ShaderDerivedDataVersion.h
@@ -10,5 +10,5 @@
 
 // In case of merge conflicts with DDC versions, you MUST generate a new GUID and set this new
 // guid as version
-#define GLOBALSHADERMAP_DERIVEDDATA_VER			TEXT("CE4A26A5071E4C9596A658A9AC69B29B")
+#define GLOBALSHADERMAP_DERIVEDDATA_VER			TEXT("31E4CCE9AC91470E8CAB008CABE852D4")
 #define MATERIALSHADERMAP_DERIVEDDATA_VER		TEXT("FCCBAA5CC66144528A3A2E4B4A29881B}")
\ No newline at end of file
diff --git a/Engine/Source/Runtime/RenderCore/Private/GlobalShader.cpp b/Engine/Source/Runtime/RenderCore/Private/GlobalShader.cpp
--- a/Engine/Source/Runtime/RenderCore/Private/GlobalShader.cpp
+++ b/Engine/Source/Runtime/RenderCore/Private/GlobalShader.cpp
@@ -343,6 +343,49 @@ bool FGlobalShaderMap::IsEmpty() const
 	return true;
 }
 
+bool FGlobalShaderMap::IsComplete(const ITargetPlatform* TargetPlatform) const
+{
+	// TODO: store these in the shadermap before it's start to be compiled?
+	FPlatformTypeLayoutParameters LayoutParams;
+	LayoutParams.InitializeForPlatform(TargetPlatform);
+	const EShaderPermutationFlags PermutationFlags = GetShaderPermutationFlags(LayoutParams);
+
+	FGlobalShaderMapId ShaderMapId(Platform, TargetPlatform);
+
+	// traverse all global shader types
+	for (TLinkedList<FShaderType*>::TIterator ShaderTypeIt(FShaderType::GetTypeList()); ShaderTypeIt; ShaderTypeIt.Next())
+	{
+		FGlobalShaderType* GlobalShaderType = ShaderTypeIt->GetGlobalShaderType();
+		if (!GlobalShaderType)
+		{
+			continue;
+		}
+
+		int32 PermutationCountToCompile = 0;
+		for (int32 PermutationId = 0; PermutationId < GlobalShaderType->GetPermutationCount(); PermutationId++)
+		{
+			if (GlobalShaderType->ShouldCompilePermutation(Platform, PermutationId, PermutationFlags)
+				&& !HasShader(GlobalShaderType, PermutationId))
+			{
+				return false;
+			}
+		}
+	}
+
+	// traverse all pipelines. Note that there's no ShouldCompile call for them. Materials instead test individual stages, but it leads to another problems
+	// like including the standalone types even if they are not going to be used. This code follows VerifyGlobalShaders() logic that includes all global pipelines unconditionally.
+	for (TLinkedList<FShaderPipelineType*>::TIterator ShaderPipelineIt(FShaderPipelineType::GetTypeList()); ShaderPipelineIt; ShaderPipelineIt.Next())
+	{
+		const FShaderPipelineType* Pipeline = *ShaderPipelineIt;
+		if (Pipeline->IsGlobalTypePipeline() && !HasShaderPipeline(Pipeline))
+		{
+			return false;
+		}
+	}
+
+	return true;
+}
+
 void FGlobalShaderMap::Empty()
 {
 	for (const auto& It : SectionMap)
diff --git a/Engine/Source/Runtime/RenderCore/Public/GlobalShader.h b/Engine/Source/Runtime/RenderCore/Public/GlobalShader.h
--- a/Engine/Source/Runtime/RenderCore/Public/GlobalShader.h
+++ b/Engine/Source/Runtime/RenderCore/Public/GlobalShader.h
@@ -195,6 +195,9 @@ public:
 
 	bool IsEmpty() const;
 
+	/** Whether the shadermap has all the shader types it needs (i.e. ones that returned ShouldCompilePermutation) */
+	bool IsComplete(const ITargetPlatform* TargetPlatform) const;
+
 	void Empty();
 	void ReleaseAllSections();
 
