Added GPU stat scopes for GPUSkinCache work done during EndSendEndOfFrameUpdatesDrawEvent. Makes skin cache costs (actual skin cached dispatches and any ray tracing acceleration structure work) visible in `stat gpu` and Unreal Insights.

diff --git a/Engine/Source/Runtime/Engine/Private/GPUSkinCache.cpp b/Engine/Source/Runtime/Engine/Private/GPUSkinCache.cpp
--- a/Engine/Source/Runtime/Engine/Private/GPUSkinCache.cpp
+++ b/Engine/Source/Runtime/Engine/Private/GPUSkinCache.cpp
@@ -803,6 +803,8 @@ void FGPUSkinCache::AddRayTracingGeometryToUpdate(FRayTracingGeometry* InRayTrac
 	}
 }
 
+DECLARE_GPU_STAT(GPUSkinCacheBuildBLAS);
+DECLARE_GPU_STAT(GPUSkinCacheUpdateBLAS);
 void FGPUSkinCache::CommitRayTracingGeometryUpdates(FRHICommandListImmediate& RHICmdList)
 {
 	TRACE_CPUPROFILER_EVENT_SCOPE(FGPUSkinCache::CommitRayTracingGeometryUpdates);
@@ -829,14 +831,16 @@ void FGPUSkinCache::CommitRayTracingGeometryUpdates(FRHICommandListImmediate& RH
 		{
 			if (BatchedBuildParams.Num())
 			{
-				SCOPED_DRAW_EVENT(RHICmdList, BuildGPUSkinCacheBLAS);
+				SCOPED_GPU_STAT(RHICmdList, GPUSkinCacheBuildBLAS);
+				SCOPED_DRAW_EVENT(RHICmdList, GPUSkinCacheBuildBLAS);
 				RHICmdList.BuildAccelerationStructures(BatchedBuildParams);
 				BatchedBuildParams.Empty(BatchedBuildParams.Max());
 			}
 
 			if (BatchedUpdateParams.Num())
 			{
-				SCOPED_DRAW_EVENT(RHICmdList, UpdateGPUSkinCacheBLAS);
+				SCOPED_GPU_STAT(RHICmdList, GPUSkinCacheUpdateBLAS);
+				SCOPED_DRAW_EVENT(RHICmdList, GPUSkinCacheUpdateBLAS);
 				RHICmdList.BuildAccelerationStructures(BatchedUpdateParams);
 				BatchedUpdateParams.Empty(BatchedUpdateParams.Max());
 			}
diff --git a/Engine/Source/Runtime/Engine/Private/LevelTick.cpp b/Engine/Source/Runtime/Engine/Private/LevelTick.cpp
--- a/Engine/Source/Runtime/Engine/Private/LevelTick.cpp
+++ b/Engine/Source/Runtime/Engine/Private/LevelTick.cpp
@@ -61,6 +61,7 @@
 #include "InGamePerformanceTracker.h"
 #include "Streaming/TextureStreamingHelpers.h"
 #include "ProfilingDebugging/CsvProfiler.h"
+#include "ProfilingDebugging/RealtimeGPUProfiler.h"
 #include "GPUSkinCache.h"
 #include "ComputeFramework/ComputeFramework.h"
 
@@ -994,27 +995,41 @@ FSendAllEndOfFrameUpdates* BeginSendEndOfFrameUpdatesDrawEvent(
 	return SendAllEndOfFrameUpdates;
 }
 
+DECLARE_GPU_STAT(EndOfFrameUpdates);
+DECLARE_GPU_STAT(GPUSkinCache);
+DECLARE_GPU_STAT(GPUSkinCacheRayTracingGeometry);
+DECLARE_GPU_STAT(ComputeFrameworkExecuteBatches);
 void EndSendEndOfFrameUpdatesDrawEvent(FSendAllEndOfFrameUpdates* SendAllEndOfFrameUpdates)
 {
 	ENQUEUE_RENDER_COMMAND(EndDrawEventCommand)(
 		[SendAllEndOfFrameUpdates](FRHICommandListImmediate& RHICmdList)
 	{
+		SCOPED_GPU_STAT(RHICmdList, EndOfFrameUpdates);
+
 		if (SendAllEndOfFrameUpdates->GPUSkinCache)
 		{
-			// Once all the individual components have received their DoDeferredRenderUpdates_Concurrent()
-			// allow the GPU Skin Cache system to update.
-			SendAllEndOfFrameUpdates->GPUSkinCache->EndBatchDispatch(RHICmdList);
+			{
+				SCOPED_GPU_STAT(RHICmdList, GPUSkinCache);
+
+				// Once all the individual components have received their DoDeferredRenderUpdates_Concurrent()
+				// allow the GPU Skin Cache system to update.
+				SendAllEndOfFrameUpdates->GPUSkinCache->EndBatchDispatch(RHICmdList);
 
-			// Flush any remaining pending resource barriers.
-			SendAllEndOfFrameUpdates->GPUSkinCache->TransitionAllToReadable(RHICmdList);
+				// Flush any remaining pending resource barriers.
+				SendAllEndOfFrameUpdates->GPUSkinCache->TransitionAllToReadable(RHICmdList);
+			}
 
 		#if RHI_RAYTRACING
-			SendAllEndOfFrameUpdates->GPUSkinCache->CommitRayTracingGeometryUpdates(RHICmdList);
+			{
+				SCOPED_GPU_STAT(RHICmdList, GPUSkinCacheRayTracingGeometry);
+				SendAllEndOfFrameUpdates->GPUSkinCache->CommitRayTracingGeometryUpdates(RHICmdList);
+			}
 		#endif // RHI_RAYTRACING
 		}
 
 		if (SendAllEndOfFrameUpdates->ComputeFramework)
 		{
+			SCOPED_GPU_STAT(RHICmdList, ComputeFrameworkExecuteBatches);
 			SendAllEndOfFrameUpdates->ComputeFramework->ExecuteBatches(RHICmdList, SendAllEndOfFrameUpdates->FeatureLevel);
 		}
 
