Packed Level Instance: Add Version GUID to determine if RerunConstructionScript needs to happen or not in GameWorld (PIE)

diff --git a/Engine/Source/Runtime/Core/Public/UObject/UE5ReleaseStreamObjectVersion.h b/Engine/Source/Runtime/Core/Public/UObject/UE5ReleaseStreamObjectVersion.h
--- a/Engine/Source/Runtime/Core/Public/UObject/UE5ReleaseStreamObjectVersion.h
+++ b/Engine/Source/Runtime/Core/Public/UObject/UE5ReleaseStreamObjectVersion.h
@@ -36,6 +36,9 @@ struct CORE_API FUE5ReleaseStreamObjectVersion
 		// Added max linear and angular speed to Chaos bodies
 		AddChaosMaxLinearAngularSpeed,
 
+		// PackedLevelInstance version
+		PackedLevelInstanceVersion,
+
 		// -----<new versions can be added above this line>-------------------------------------------------
 		VersionPlusOne,
 		LatestVersion = VersionPlusOne - 1
diff --git a/Engine/Source/Runtime/Engine/Private/LevelInstance/Packed/PackedLevelInstanceActor.cpp b/Engine/Source/Runtime/Engine/Private/LevelInstance/Packed/PackedLevelInstanceActor.cpp
--- a/Engine/Source/Runtime/Engine/Private/LevelInstance/Packed/PackedLevelInstanceActor.cpp
+++ b/Engine/Source/Runtime/Engine/Private/LevelInstance/Packed/PackedLevelInstanceActor.cpp
@@ -4,12 +4,13 @@
 #include "LevelInstance/Packed/PackedLevelInstanceBuilder.h"
 #include "LevelInstance/LevelInstanceSubsystem.h"
 #include "Components/InstancedStaticMeshComponent.h"
+#include "LevelInstance/LevelInstancePrivate.h"
+#include "UObject/UE5ReleaseStreamObjectVersion.h"
 
 APackedLevelInstance::APackedLevelInstance(const FObjectInitializer& ObjectInitializer)
 	: Super(ObjectInitializer)
 #if WITH_EDITORONLY_DATA
 	, ISMComponentClass(UInstancedStaticMeshComponent::StaticClass())
-	, bRerunConstructionScripts(false)
 	, bEditing(false)
 	, ChildEditing(0)
 	, bChildChanged(false)
@@ -21,6 +22,20 @@ APackedLevelInstance::APackedLevelInstance(const FObjectInitializer& ObjectIniti
 #endif
 }
 
+void APackedLevelInstance::Serialize(FArchive& Ar)
+{
+	Ar.UsingCustomVersion(FUE5ReleaseStreamObjectVersion::GUID);
+	Super::Serialize(Ar);
+
+#if WITH_EDITORONLY_DATA
+	// We want to make sure we serialize that property so we can compare to the CDO
+	if (Ar.CustomVer(FUE5ReleaseStreamObjectVersion::GUID) >= FUE5ReleaseStreamObjectVersion::PackedLevelInstanceVersion)
+	{
+		Ar << PackedVersion;
+	}
+#endif
+}
+
 bool APackedLevelInstance::SupportsLoading() const
 {
 #if WITH_EDITOR
@@ -31,6 +46,46 @@ bool APackedLevelInstance::SupportsLoading() const
 }
 
 #if WITH_EDITOR
+
+void APackedLevelInstance::PostLoad()
+{
+	Super::PostLoad();
+
+	// Non CDO: Set the Guid to something different then the default value so that we actually run the construction script on actors that haven't been resaved against their latest BP
+	if (!HasAnyFlags(RF_ClassDefaultObject | RF_ArchetypeObject) && GetLinkerCustomVersion(FUE5ReleaseStreamObjectVersion::GUID) < FUE5ReleaseStreamObjectVersion::PackedLevelInstanceVersion)
+	{
+		const static FGuid NoVersionGUID(0x50817615, 0x74A547A3, 0x9295D655, 0x8A852C0F);
+		PackedVersion = NoVersionGUID;
+	}
+}
+
+void APackedLevelInstance::RerunConstructionScripts()
+{
+	bool bShouldRerunConstructionScript = true;
+
+	if (GetWorld() && GetWorld()->IsGameWorld())
+	{
+		bShouldRerunConstructionScript = false;
+
+		// Only rerun if version mismatchs
+		if (PackedVersion != GetClass()->GetDefaultObject<APackedLevelInstance>()->PackedVersion)
+		{
+			bShouldRerunConstructionScript = true;
+			UE_LOG(LogLevelInstance, Warning, TEXT("RerunConstructionScript was executed on %s (%s) because its version (%s) doesn't match latest version (%s). Resaving this actor will fix this"),
+				*GetPathName(),
+				*GetPackage()->GetPathName(),
+				*PackedVersion.ToString(),
+				*GetClass()->GetDefaultObject<APackedLevelInstance>()->PackedVersion.ToString());
+		}
+	}
+	
+	if(bShouldRerunConstructionScript)
+	{
+		Super::RerunConstructionScripts();
+		PackedVersion = GetClass()->GetDefaultObject<APackedLevelInstance>()->PackedVersion;
+	}
+}
+
 FName APackedLevelInstance::GetPackedComponentTag()
 {
 	static FName PackedComponentTag("PackedComponent");
diff --git a/Engine/Source/Runtime/Engine/Private/LevelInstance/Packed/PackedLevelInstanceBuilder.cpp b/Engine/Source/Runtime/Engine/Private/LevelInstance/Packed/PackedLevelInstanceBuilder.cpp
--- a/Engine/Source/Runtime/Engine/Private/LevelInstance/Packed/PackedLevelInstanceBuilder.cpp
+++ b/Engine/Source/Runtime/Engine/Private/LevelInstance/Packed/PackedLevelInstanceBuilder.cpp
@@ -424,13 +424,14 @@ bool FPackedLevelInstanceBuilder::CreateOrUpdateBlueprintFromPacked(APackedLevel
 	BP->Modify();
 	// Avoid running construction script while dragging an instance of that BP for performance reasons
 	BP->bRunConstructionScriptOnDrag = false;
+	FGuid NewVersion = FGuid::NewGuid();
 	APackedLevelInstance* CDO = CastChecked<APackedLevelInstance>(BP->GeneratedClass->GetDefaultObject());
-	auto PropagatePropertiesToActor = [InActor](APackedLevelInstance* TargetActor)
+	auto PropagatePropertiesToActor = [InActor, NewVersion](APackedLevelInstance* TargetActor)
 	{
 		TargetActor->Modify(false);
 		TargetActor->SetWorldAsset(InActor->GetWorldAsset());
 		TargetActor->PackedBPDependencies = InActor->PackedBPDependencies;
-
+		TargetActor->SetPackedVersion(NewVersion);
 		// match root component mobility to source actor
 		USceneComponent* Root = TargetActor->GetRootComponent();
 		Root->SetMobility(InActor->GetRootComponent()->Mobility);
diff --git a/Engine/Source/Runtime/Engine/Public/LevelInstance/Packed/PackedLevelInstanceActor.h b/Engine/Source/Runtime/Engine/Public/LevelInstance/Packed/PackedLevelInstanceActor.h
--- a/Engine/Source/Runtime/Engine/Public/LevelInstance/Packed/PackedLevelInstanceActor.h
+++ b/Engine/Source/Runtime/Engine/Public/LevelInstance/Packed/PackedLevelInstanceActor.h
@@ -32,9 +32,11 @@ class ENGINE_API APackedLevelInstance : public ALevelInstance
 public:
 	virtual bool SupportsLoading() const override;
 
+	virtual void Serialize(FArchive& Ar) override;
 #if WITH_EDITOR
 	static FName GetPackedComponentTag();
 
+	virtual void PostLoad() override;
 	virtual void OnWorldAssetChanged() override;
 	virtual void OnWorldAssetSaved(bool bPromptForSave) override;
 	virtual void OnCommit() override;
@@ -43,6 +45,8 @@ public:
 	virtual void OnEditChild() override;
 
 	virtual bool CanEditChange(const FProperty* InProperty) const override;
+		
+	void SetPackedVersion(const FGuid& Version) { PackedVersion = Version; }
 
 	// AActor overrides
 	virtual EActorGridPlacement GetDefaultGridPlacement() const override { return EActorGridPlacement::None; }
@@ -54,6 +58,8 @@ public:
 
 	virtual ELevelInstanceRuntimeBehavior GetDefaultRuntimeBehavior() const override { return ELevelInstanceRuntimeBehavior::None; }
 
+	virtual void RerunConstructionScripts() override;
+
 	template<class T>
 	T* AddPackedComponent()
 	{
@@ -73,9 +79,6 @@ public:
 	UPROPERTY(VisibleAnywhere, Category = Pack)
 	TSoftObjectPtr<UBlueprint> BlueprintAsset;
 
-	UPROPERTY(VisibleAnywhere, Category = Pack)
-	bool bRerunConstructionScripts;
-
 	UPROPERTY()
 	TArray<TSoftObjectPtr<UBlueprint>> PackedBPDependencies;
 
@@ -88,5 +91,8 @@ private:
 
 	UPROPERTY(NonTransactional)
 	bool bChildChanged;
+
+	UPROPERTY()
+	FGuid PackedVersion;
 #endif
 };
\ No newline at end of file
