Removing code and comment references to 'UE4'

diff --git a/Engine/Source/Runtime/Engine/Private/Level.cpp b/Engine/Source/Runtime/Engine/Private/Level.cpp
--- a/Engine/Source/Runtime/Engine/Private/Level.cpp
+++ b/Engine/Source/Runtime/Engine/Private/Level.cpp
@@ -1530,7 +1530,7 @@ void ULevel::CreateModelComponents()
 			if (NodeGroup->Nodes.Num())
 			{
 				// get one of the surfaces/components from the NodeGroup
-				// @todo UE4: Remove need for GetSurfaceLightMapResolution to take a surfaceindex, or a ModelComponent :)
+				// @todo: Remove need for GetSurfaceLightMapResolution to take a surfaceindex, or a ModelComponent :)
 				UModelComponent* SomeModelComponent = ModelComponents[Model->Nodes[NodeGroup->Nodes[0]].ComponentIndex];
 				int32 SurfaceIndex = Model->Nodes[NodeGroup->Nodes[0]].iSurf;
 
diff --git a/Engine/Source/Runtime/Engine/Private/Materials/Material.cpp b/Engine/Source/Runtime/Engine/Private/Materials/Material.cpp
--- a/Engine/Source/Runtime/Engine/Private/Materials/Material.cpp
+++ b/Engine/Source/Runtime/Engine/Private/Materials/Material.cpp
@@ -4045,40 +4045,40 @@ void UMaterial::PostLoad()
 	LoadedMaterialResources.Empty();
 
 #if WITH_EDITORONLY_DATA
-	const int32 UE4Ver = GetLinkerUEVersion();
+	const int32 UEVer = GetLinkerUEVersion();
 	const int32 RenderObjVer = GetLinkerCustomVersion(FRenderingObjectVersion::GUID);
 
-	DoMaterialAttributeReorder(&DiffuseColor_DEPRECATED, UE4Ver, RenderObjVer);
-	DoMaterialAttributeReorder(&SpecularColor_DEPRECATED, UE4Ver, RenderObjVer);
-	DoMaterialAttributeReorder(&BaseColor, UE4Ver, RenderObjVer);
-	DoMaterialAttributeReorder(&Metallic, UE4Ver, RenderObjVer);
-	DoMaterialAttributeReorder(&Specular, UE4Ver, RenderObjVer);
-	DoMaterialAttributeReorder(&Roughness, UE4Ver, RenderObjVer);
-	DoMaterialAttributeReorder(&Anisotropy, UE4Ver, RenderObjVer);
-	DoMaterialAttributeReorder(&Normal, UE4Ver, RenderObjVer);
-	DoMaterialAttributeReorder(&Tangent, UE4Ver, RenderObjVer);
-	DoMaterialAttributeReorder(&EmissiveColor, UE4Ver, RenderObjVer);
-	DoMaterialAttributeReorder(&Opacity, UE4Ver, RenderObjVer);
-	DoMaterialAttributeReorder(&OpacityMask, UE4Ver, RenderObjVer);
-	DoMaterialAttributeReorder(&WorldPositionOffset, UE4Ver, RenderObjVer);
-	DoMaterialAttributeReorder(&WorldDisplacement, UE4Ver, RenderObjVer);
-	DoMaterialAttributeReorder(&TessellationMultiplier, UE4Ver, RenderObjVer);
-	DoMaterialAttributeReorder(&SubsurfaceColor, UE4Ver, RenderObjVer);
-	DoMaterialAttributeReorder(&ClearCoat, UE4Ver, RenderObjVer);
-	DoMaterialAttributeReorder(&ClearCoatRoughness, UE4Ver, RenderObjVer);
-	DoMaterialAttributeReorder(&AmbientOcclusion, UE4Ver, RenderObjVer);
-	DoMaterialAttributeReorder(&Refraction, UE4Ver, RenderObjVer);
-	DoMaterialAttributeReorder(&CustomizedUVs[0], UE4Ver, RenderObjVer);
-	DoMaterialAttributeReorder(&CustomizedUVs[1], UE4Ver, RenderObjVer);
-	DoMaterialAttributeReorder(&CustomizedUVs[2], UE4Ver, RenderObjVer);
-	DoMaterialAttributeReorder(&CustomizedUVs[3], UE4Ver, RenderObjVer);
-	DoMaterialAttributeReorder(&CustomizedUVs[4], UE4Ver, RenderObjVer);
-	DoMaterialAttributeReorder(&CustomizedUVs[5], UE4Ver, RenderObjVer);
-	DoMaterialAttributeReorder(&CustomizedUVs[6], UE4Ver, RenderObjVer);
-	DoMaterialAttributeReorder(&CustomizedUVs[7], UE4Ver, RenderObjVer);
-	DoMaterialAttributeReorder(&PixelDepthOffset, UE4Ver, RenderObjVer);
-	DoMaterialAttributeReorder(&ShadingModelFromMaterialExpression, UE4Ver, RenderObjVer);
-	DoMaterialAttributeReorder(&FrontMaterial, UE4Ver, RenderObjVer);
+	DoMaterialAttributeReorder(&DiffuseColor_DEPRECATED, UEVer, RenderObjVer);
+	DoMaterialAttributeReorder(&SpecularColor_DEPRECATED, UEVer, RenderObjVer);
+	DoMaterialAttributeReorder(&BaseColor, UEVer, RenderObjVer);
+	DoMaterialAttributeReorder(&Metallic, UEVer, RenderObjVer);
+	DoMaterialAttributeReorder(&Specular, UEVer, RenderObjVer);
+	DoMaterialAttributeReorder(&Roughness, UEVer, RenderObjVer);
+	DoMaterialAttributeReorder(&Anisotropy, UEVer, RenderObjVer);
+	DoMaterialAttributeReorder(&Normal, UEVer, RenderObjVer);
+	DoMaterialAttributeReorder(&Tangent, UEVer, RenderObjVer);
+	DoMaterialAttributeReorder(&EmissiveColor, UEVer, RenderObjVer);
+	DoMaterialAttributeReorder(&Opacity, UEVer, RenderObjVer);
+	DoMaterialAttributeReorder(&OpacityMask, UEVer, RenderObjVer);
+	DoMaterialAttributeReorder(&WorldPositionOffset, UEVer, RenderObjVer);
+	DoMaterialAttributeReorder(&WorldDisplacement, UEVer, RenderObjVer);
+	DoMaterialAttributeReorder(&TessellationMultiplier, UEVer, RenderObjVer);
+	DoMaterialAttributeReorder(&SubsurfaceColor, UEVer, RenderObjVer);
+	DoMaterialAttributeReorder(&ClearCoat, UEVer, RenderObjVer);
+	DoMaterialAttributeReorder(&ClearCoatRoughness, UEVer, RenderObjVer);
+	DoMaterialAttributeReorder(&AmbientOcclusion, UEVer, RenderObjVer);
+	DoMaterialAttributeReorder(&Refraction, UEVer, RenderObjVer);
+	DoMaterialAttributeReorder(&CustomizedUVs[0], UEVer, RenderObjVer);
+	DoMaterialAttributeReorder(&CustomizedUVs[1], UEVer, RenderObjVer);
+	DoMaterialAttributeReorder(&CustomizedUVs[2], UEVer, RenderObjVer);
+	DoMaterialAttributeReorder(&CustomizedUVs[3], UEVer, RenderObjVer);
+	DoMaterialAttributeReorder(&CustomizedUVs[4], UEVer, RenderObjVer);
+	DoMaterialAttributeReorder(&CustomizedUVs[5], UEVer, RenderObjVer);
+	DoMaterialAttributeReorder(&CustomizedUVs[6], UEVer, RenderObjVer);
+	DoMaterialAttributeReorder(&CustomizedUVs[7], UEVer, RenderObjVer);
+	DoMaterialAttributeReorder(&PixelDepthOffset, UEVer, RenderObjVer);
+	DoMaterialAttributeReorder(&ShadingModelFromMaterialExpression, UEVer, RenderObjVer);
+	DoMaterialAttributeReorder(&FrontMaterial, UEVer, RenderObjVer);
 #endif // WITH_EDITORONLY_DATA
 
 	if (!IsDefaultMaterial())
diff --git a/Engine/Source/Runtime/Engine/Private/Materials/MaterialExpressions.cpp b/Engine/Source/Runtime/Engine/Private/Materials/MaterialExpressions.cpp
--- a/Engine/Source/Runtime/Engine/Private/Materials/MaterialExpressions.cpp
+++ b/Engine/Source/Runtime/Engine/Private/Materials/MaterialExpressions.cpp
@@ -13893,7 +13893,7 @@ bool FMaterialLayersFunctions::ResolveParent(const FMaterialLayersFunctions& Par
 
 		if (bHasUninitializedLinks)
 		{
-			// If we had any unitialized links, this means we're loading data saved by a previous version of UE4
+			// If we had any unitialized links, this means we're loading data saved by a previous version of UE
 			// In this case, we have no way of determining if this layer was added to parent (and should therefore be kept),
 			// or if this layer was explicitly deleted from child (and should therefore remain deleted).
 			// In order to avoid needlessly changing legacy materials, we assume the layer was explicitly deleted, so we keep it deleted here
diff --git a/Engine/Source/Runtime/Engine/Private/Materials/MaterialShared.cpp b/Engine/Source/Runtime/Engine/Private/Materials/MaterialShared.cpp
--- a/Engine/Source/Runtime/Engine/Private/Materials/MaterialShared.cpp
+++ b/Engine/Source/Runtime/Engine/Private/Materials/MaterialShared.cpp
@@ -4086,11 +4086,11 @@ bool UMaterialInterface::IsTextureReferencedByProperty(EMaterialProperty InPrope
 //Reorder the output index for any FExpressionInput connected to a UMaterialExpressionBreakMaterialAttributes.
 //If the order of pins in the material results or the make/break attributes nodes changes 
 //then the OutputIndex stored in any FExpressionInput coming from UMaterialExpressionBreakMaterialAttributes will be wrong and needs reordering.
-void DoMaterialAttributeReorder(FExpressionInput* Input, int32 UE4Ver, int32 RenderObjVer)
+void DoMaterialAttributeReorder(FExpressionInput* Input, int32 UEVer, int32 RenderObjVer)
 {
 	if( Input && Input->Expression && Input->Expression->IsA(UMaterialExpressionBreakMaterialAttributes::StaticClass()) )
 	{
-		if( UE4Ver < VER_UE4_MATERIAL_ATTRIBUTES_REORDERING )
+		if(UEVer < VER_UE4_MATERIAL_ATTRIBUTES_REORDERING )
 		{
 			switch(Input->OutputIndex)
 			{
@@ -4101,7 +4101,7 @@ void DoMaterialAttributeReorder(FExpressionInput* Input, int32 UE4Ver, int32 Ren
 			}
 		}
 		
-		if( UE4Ver < VER_UE4_FIX_REFRACTION_INPUT_MASKING && Input->OutputIndex == 13 )
+		if(UEVer < VER_UE4_FIX_REFRACTION_INPUT_MASKING && Input->OutputIndex == 13 )
 		{
 			Input->Mask = 1;
 			Input->MaskR = 1;
@@ -4111,7 +4111,7 @@ void DoMaterialAttributeReorder(FExpressionInput* Input, int32 UE4Ver, int32 Ren
 		}
 
 		// closest version to the clear coat change
-		if( UE4Ver < VER_UE4_ADD_ROOTCOMPONENT_TO_FOLIAGEACTOR && Input->OutputIndex >= 12 )
+		if(UEVer < VER_UE4_ADD_ROOTCOMPONENT_TO_FOLIAGEACTOR && Input->OutputIndex >= 12 )
 		{
 			Input->OutputIndex += 2;
 		}
diff --git a/Engine/Source/Runtime/Engine/Private/SceneView.cpp b/Engine/Source/Runtime/Engine/Private/SceneView.cpp
--- a/Engine/Source/Runtime/Engine/Private/SceneView.cpp
+++ b/Engine/Source/Runtime/Engine/Private/SceneView.cpp
@@ -281,7 +281,7 @@ static TAutoConsoleVariable<float> CVarTessellationAdaptivePixelsPerTriangle(
 	TEXT("Global tessellation factor multiplier"),
 	ECVF_RenderThreadSafe);
 
-// should be changed to BaseColor and Metallic, since some time now UE4 is not using DiffuseColor and SpecularColor any more
+// should be changed to BaseColor and Metallic, since some time now UE is not using DiffuseColor and SpecularColor any more
 static TAutoConsoleVariable<float> CVarDiffuseColorMin(
 	TEXT("r.DiffuseColor.Min"),
 	0.0f,
diff --git a/Engine/Source/Runtime/Engine/Private/SkeletalRenderGPUSkin.cpp b/Engine/Source/Runtime/Engine/Private/SkeletalRenderGPUSkin.cpp
--- a/Engine/Source/Runtime/Engine/Private/SkeletalRenderGPUSkin.cpp
+++ b/Engine/Source/Runtime/Engine/Private/SkeletalRenderGPUSkin.cpp
@@ -648,7 +648,7 @@ void FSkeletalMeshObjectGPUSkin::ProcessUpdatedDynamicData(FGPUSkinCache* GPUSki
 			{
 				// This takes the cloth positions from cloth space into world space
 				FMatrix ClothLocalToWorld = bClothFactory ? VertexFactoryData.ClothVertexFactories[SectionIdx]->GetClothShaderData().GetClothLocalToWorldForWriting(FrameNumberToPrepare) : FMatrix::Identity;
-				// Matrices are transposed in ue4 meaning matrix multiples need to happen in reverse ((AB)x = b becomes xTBTAT = b).
+				// Matrices are transposed in UE meaning matrix multiples need to happen in reverse ((AB)x = b becomes xTBTAT = b).
 				FMatrix LocalToCloth = DynamicData->ClothObjectLocalToWorld * ClothLocalToWorld.Inverse();
 
 				GPUSkinCache->ProcessEntry(
diff --git a/Engine/Source/Runtime/Engine/Private/SpeedTreeWind.cpp b/Engine/Source/Runtime/Engine/Private/SpeedTreeWind.cpp
--- a/Engine/Source/Runtime/Engine/Private/SpeedTreeWind.cpp
+++ b/Engine/Source/Runtime/Engine/Private/SpeedTreeWind.cpp
@@ -14,7 +14,7 @@
 //      IDV, Inc.
 //      http://www.idvinc.com
 
-//  SpeedTree v6.2.2 wind class rewritten for use inside UE4 with no other dependencies
+//  SpeedTree v6.2.2 wind class rewritten for use inside UE with no other dependencies
 
 
 ///////////////////////////////////////////////////////////////////////  
diff --git a/Engine/Source/Runtime/Engine/Private/VT/VirtualTextureDataBuilder.cpp b/Engine/Source/Runtime/Engine/Private/VT/VirtualTextureDataBuilder.cpp
--- a/Engine/Source/Runtime/Engine/Private/VT/VirtualTextureDataBuilder.cpp
+++ b/Engine/Source/Runtime/Engine/Private/VT/VirtualTextureDataBuilder.cpp
@@ -561,7 +561,7 @@ void FVirtualTextureDataBuilder::BuildTiles(const TArray<FVTSourceTileEntry>& Ti
 
 		// We can't split crunch compression into multiple tasks/threads, since all tiles need to compress together in order to generate the codec payload
 		// Instead we rely on internal Crunch threading to make this efficient
-		// Might be worth modifying Crunch to expose threading callbacks, so this can use UE4 task graph instead of Crunch's internal threadpool
+		// Might be worth modifying Crunch to expose threading callbacks, so this can use UE task graph instead of Crunch's internal threadpool
 		if (bAllowAsync && FApp::ShouldUseThreadingForPerformance())
 		{
 			CrunchParameters.NumWorkerThreads = FTaskGraphInterface::Get().GetNumWorkerThreads();
@@ -917,7 +917,7 @@ void FVirtualTextureDataBuilder::BuildSourcePixels(const FTextureSourceData& Sou
 
 		FTextureSourceBlockData& BlockData = SourceBlocks[BlockIndex];
 		BlockData.BlockX = SourceBlockData.BlockX;
-		// UE4 applies a (1-y) transform to imported UVs, so apply a similar transform to UDIM block locations here
+		// UE applies a (1-y) transform to imported UVs, so apply a similar transform to UDIM block locations here
 		// This ensures that UDIM tiles will appear in the correct location when sampled with transformed UVs
 		BlockData.BlockY = (SizeInBlocksY - SourceBlockData.BlockY) % SizeInBlocksY;
 		BlockData.NumMips = SourceBlockData.NumMips;
diff --git a/Engine/Source/Runtime/Engine/Public/MaterialShared.h b/Engine/Source/Runtime/Engine/Public/MaterialShared.h
--- a/Engine/Source/Runtime/Engine/Public/MaterialShared.h
+++ b/Engine/Source/Runtime/Engine/Public/MaterialShared.h
@@ -2716,7 +2716,7 @@ ENGINE_API bool DoesMaterialUseTexture(const UMaterialInterface* Material,const
 
 #if WITH_EDITORONLY_DATA
 /** TODO - This can be removed whenever VER_UE4_MATERIAL_ATTRIBUTES_REORDERING is no longer relevant. */
-ENGINE_API void DoMaterialAttributeReorder(FExpressionInput* Input, int32 UE4Ver, int32 RenderObjVer);
+ENGINE_API void DoMaterialAttributeReorder(FExpressionInput* Input, int32 UEVer, int32 RenderObjVer);
 #endif // WITH_EDITORONLY_DATA
 
 /**
diff --git a/Engine/Source/Runtime/Engine/Public/SceneView.h b/Engine/Source/Runtime/Engine/Public/SceneView.h
--- a/Engine/Source/Runtime/Engine/Public/SceneView.h
+++ b/Engine/Source/Runtime/Engine/Public/SceneView.h
@@ -45,7 +45,7 @@ struct FSceneViewProjectionData
 	/** Rotation matrix transforming from world space to view space. */
 	FMatrix ViewRotationMatrix;
 
-	/** UE4 projection matrix projects such that clip space Z=1 is the near plane, and Z=0 is the infinite far plane. */
+	/** UE projection matrix projects such that clip space Z=1 is the near plane, and Z=0 is the infinite far plane. */
 	FMatrix ProjectionMatrix;
 
 protected:
@@ -250,19 +250,19 @@ private:
 
 	void Init(const FMinimalInitializer& Initializer);
 
-	/** ViewToClip : UE4 projection matrix projects such that clip space Z=1 is the near plane, and Z=0 is the infinite far plane. */
+	/** ViewToClip : UE projection matrix projects such that clip space Z=1 is the near plane, and Z=0 is the infinite far plane. */
 	FMatrix		ProjectionMatrix;
-	/** ViewToClipNoAA : UE4 projection matrix projects such that clip space Z=1 is the near plane, and Z=0 is the infinite far plane. Don't apply any AA jitter */
+	/** ViewToClipNoAA : UE projection matrix projects such that clip space Z=1 is the near plane, and Z=0 is the infinite far plane. Don't apply any AA jitter */
 	FMatrix		ProjectionNoAAMatrix;
-	/** ClipToView : UE4 projection matrix projects such that clip space Z=1 is the near plane, and Z=0 is the infinite far plane. */
+	/** ClipToView : UE projection matrix projects such that clip space Z=1 is the near plane, and Z=0 is the infinite far plane. */
 	FMatrix		InvProjectionMatrix;
 	// WorldToView..
 	FMatrix		ViewMatrix;
 	// ViewToWorld..
 	FMatrix		InvViewMatrix;
-	// WorldToClip : UE4 projection matrix projects such that clip space Z=1 is the near plane, and Z=0 is the infinite far plane. */
+	// WorldToClip : UE projection matrix projects such that clip space Z=1 is the near plane, and Z=0 is the infinite far plane. */
 	FMatrix		ViewProjectionMatrix;
-	// ClipToWorld : UE4 projection matrix projects such that clip space Z=1 is the near plane, and Z=0 is the infinite far plane. */
+	// ClipToWorld : UE projection matrix projects such that clip space Z=1 is the near plane, and Z=0 is the infinite far plane. */
 	FMatrix		InvViewProjectionMatrix;
 	// HMD WorldToView with roll removed
 	FMatrix		HMDViewMatrixNoRoll;
diff --git a/Engine/Source/Runtime/Engine/Public/SceneViewExtension.h b/Engine/Source/Runtime/Engine/Public/SceneViewExtension.h
--- a/Engine/Source/Runtime/Engine/Public/SceneViewExtension.h
+++ b/Engine/Source/Runtime/Engine/Public/SceneViewExtension.h
@@ -15,7 +15,7 @@
  *  SCENE VIEW EXTENSIONS
  *  -----------------------------------------------------------------------------------------------
  *
- *  This system lets you hook various aspects of UE4 rendering.
+ *  This system lets you hook various aspects of UE rendering.
  *  To create a view extension, it is advisable to inherit
  *  from FSceneViewExtensionBase, which implements the
  *  ISceneViewExtension interface.
diff --git a/Engine/Source/Runtime/Engine/Public/ShaderCompiler.h b/Engine/Source/Runtime/Engine/Public/ShaderCompiler.h
--- a/Engine/Source/Runtime/Engine/Public/ShaderCompiler.h
+++ b/Engine/Source/Runtime/Engine/Public/ShaderCompiler.h
@@ -580,7 +580,7 @@ public:
 
 /** 
  * Shader compiling thread
- * This runs in the background while UE4 is running, launches shader compile worker processes when necessary, and feeds them inputs and reads back the outputs.
+ * This runs in the background while UE is running, launches shader compile worker processes when necessary, and feeds them inputs and reads back the outputs.
  */
 class FShaderCompileThreadRunnable : public FShaderCompileThreadRunnableBase
 {
@@ -975,7 +975,7 @@ private:
 	int32 MaxShaderJobBatchSize;
 	/** Number of runs through single-threaded compiling before we can retry to compile through workers. -1 if not used. */
 	int32 NumSingleThreadedRunsBeforeRetry;
-	/** Process Id of UE4. */
+	/** Process Id of UE. */
 	uint32 ProcessId;
 	/** Whether to allow compiling shaders through the worker application, which allows multiple cores to be used. */
 	bool bAllowCompilingThroughWorkers;
diff --git a/Engine/Source/Runtime/Engine/Public/Slate/SceneViewport.h b/Engine/Source/Runtime/Engine/Public/Slate/SceneViewport.h
--- a/Engine/Source/Runtime/Engine/Public/Slate/SceneViewport.h
+++ b/Engine/Source/Runtime/Engine/Public/Slate/SceneViewport.h
@@ -290,7 +290,7 @@ private:
 	virtual void InitDynamicRHI() override;
 	virtual void ReleaseDynamicRHI() override;
 
-	// @todo UE4 DLL: Without these functions we get unresolved linker errors with FRenderResource
+	// @todo DLL: Without these functions we get unresolved linker errors with FRenderResource
 	virtual void InitRHI() override {}
 	virtual void ReleaseRHI() override {}
 	virtual void InitResource() override { FViewport::InitResource(); }
diff --git a/Engine/Source/Runtime/Engine/Public/SpeedTreeWind.h b/Engine/Source/Runtime/Engine/Public/SpeedTreeWind.h
--- a/Engine/Source/Runtime/Engine/Public/SpeedTreeWind.h
+++ b/Engine/Source/Runtime/Engine/Public/SpeedTreeWind.h
@@ -14,7 +14,7 @@
 //		IDV, Inc.
 //		Web: http://www.idvinc.com
 
-//  SpeedTree v6.2.2 wind class rewritten for use inside UE4 with no other dependencies
+//  SpeedTree v6.2.2 wind class rewritten for use inside UE with no other dependencies
 
 #pragma once
 
diff --git a/Engine/Source/Runtime/Engine/Public/UnrealClient.h b/Engine/Source/Runtime/Engine/Public/UnrealClient.h
--- a/Engine/Source/Runtime/Engine/Public/UnrealClient.h
+++ b/Engine/Source/Runtime/Engine/Public/UnrealClient.h
@@ -1225,7 +1225,7 @@ public:
 		RHICreateTargetableShaderResource2D( SizeX, SizeY, PF_A2B10G10R10, 1, TexCreate_None, TexCreate_RenderTargetable, false, CreateInfo, RenderTargetTextureRHI, ShaderResourceTextureRHI );
 	}
 
-	// @todo UE4 DLL: Without these functions we get unresolved linker errors with FRenderResource
+	// @todo DLL: Without these functions we get unresolved linker errors with FRenderResource
 	virtual void InitRHI() override{}
 	virtual void ReleaseRHI() override{}
 	virtual void InitResource() override{ FViewport::InitResource(); }
