Fixed launching with -game

diff --git a/Engine/Source/Runtime/Engine/Private/UnrealEngine.cpp b/Engine/Source/Runtime/Engine/Private/UnrealEngine.cpp
--- a/Engine/Source/Runtime/Engine/Private/UnrealEngine.cpp
+++ b/Engine/Source/Runtime/Engine/Private/UnrealEngine.cpp
@@ -2699,9 +2699,6 @@ void UEngine::InitializeObjectReferences()
 
 		ClothPaintMaterialInstance = UMaterialInstanceDynamic::Create(ClothPaintMaterial, nullptr);
 		ClothPaintMaterialWireframeInstance = UMaterialInstanceDynamic::Create(ClothPaintMaterialWireframe, nullptr);
-
-		LoadSpecialMaterial(TEXT("DefaultFlattenMaterialName"), DefaultFlattenMaterialName.ToString(), DefaultFlattenMaterial, false);
-		LoadSpecialMaterial(TEXT("DefaultHLODFlattenMaterialName"), DefaultHLODFlattenMaterialName.ToString(), DefaultHLODFlattenMaterial, false);
 #endif
 		FString ValidPreviewShadowsIndicatorMaterialName = PreviewShadowsIndicatorMaterialName.ToString();
 		if (ValidPreviewShadowsIndicatorMaterialName.IsEmpty())
@@ -2718,6 +2715,14 @@ void UEngine::InitializeObjectReferences()
 		}
 	}
 
+#if WITH_EDITORONLY_DATA
+	if (GIsEditor)
+	{
+		LoadSpecialMaterial(TEXT("DefaultFlattenMaterialName"), DefaultFlattenMaterialName.ToString(), DefaultFlattenMaterial, false);
+		LoadSpecialMaterial(TEXT("DefaultHLODFlattenMaterialName"), DefaultHLODFlattenMaterialName.ToString(), DefaultHLODFlattenMaterial, false);
+	}
+#endif
+
 	LoadEngineTexture(DefaultTexture, *DefaultTextureName.ToString());
 	LoadEngineTexture(DefaultDiffuseTexture, *DefaultDiffuseTextureName.ToString());
 	LoadEngineTexture(HighFrequencyNoiseTexture, *HighFrequencyNoiseTextureName.ToString());
