Fix for Static Analysis

diff --git a/Engine/Source/Runtime/Engine/Private/DataChannel.cpp b/Engine/Source/Runtime/Engine/Private/DataChannel.cpp
--- a/Engine/Source/Runtime/Engine/Private/DataChannel.cpp
+++ b/Engine/Source/Runtime/Engine/Private/DataChannel.cpp
@@ -752,7 +752,7 @@ bool UChannel::ReceivedNextBunch( FInBunch & Bunch, bool & bOutSkipAck )
 			}
 		}
 
-		if (IsBunchTooLarge(Connection, InPartialBunch))
+		if (InPartialBunch && IsBunchTooLarge(Connection, InPartialBunch))
 		{
 			UE_LOG(LogNetPartialBunch, Error, TEXT("Received a partial bunch exceeding max allowed size. BunchSize=%d, MaximumSize=%d"), InPartialBunch->GetNumBytes(), NetMaxConstructedPartialBunchSizeBytes);
 			Bunch.SetError();
