Ignore HLOD cells in UWorldPartitionStreamingPolicy::IsStreamingCompleted() if wp.Runtime.HLOD is set to 0

diff --git a/Engine/Source/Runtime/Engine/Private/WorldPartition/WorldPartitionStreamingPolicy.cpp b/Engine/Source/Runtime/Engine/Private/WorldPartition/WorldPartitionStreamingPolicy.cpp
--- a/Engine/Source/Runtime/Engine/Private/WorldPartition/WorldPartitionStreamingPolicy.cpp
+++ b/Engine/Source/Runtime/Engine/Private/WorldPartition/WorldPartitionStreamingPolicy.cpp
@@ -229,8 +229,15 @@ bool UWorldPartitionStreamingPolicy::IsStreamingCompleted(EWorldPartitionRuntime
 			if (CellState != QueryState)
 			{
 				bool bSkipCell = false;
+
+				// Don't consider HLOD cells if HLODs are disabled.
+				if (!GWorldPartitionShowHLODs)
+				{
+					bSkipCell = Cell->HasCellData<UWorldPartitionRuntimeHLODCellData>();
+				}
+
 				// If we are querying for Unloaded/Loaded but a Cell is part of a data layer outside of the query that is activated do not consider it
-				if (QueryState < CellState)
+				if (!bSkipCell && QueryState < CellState)
 				{
 					for (const FName& CellDataLayer : Cell->GetDataLayers())
 					{
