Fix Cook: Filter serialize of editor only data

diff --git a/Engine/Source/Runtime/Engine/Private/LevelInstance/Packed/PackedLevelInstanceActor.cpp b/Engine/Source/Runtime/Engine/Private/LevelInstance/Packed/PackedLevelInstanceActor.cpp
--- a/Engine/Source/Runtime/Engine/Private/LevelInstance/Packed/PackedLevelInstanceActor.cpp
+++ b/Engine/Source/Runtime/Engine/Private/LevelInstance/Packed/PackedLevelInstanceActor.cpp
@@ -29,7 +29,7 @@ void APackedLevelInstance::Serialize(FArchive& Ar)
 
 #if WITH_EDITORONLY_DATA
 	// We want to make sure we serialize that property so we can compare to the CDO
-	if (Ar.CustomVer(FUE5ReleaseStreamObjectVersion::GUID) >= FUE5ReleaseStreamObjectVersion::PackedLevelInstanceVersion)
+	if (!Ar.IsFilterEditorOnly() && Ar.CustomVer(FUE5ReleaseStreamObjectVersion::GUID) >= FUE5ReleaseStreamObjectVersion::PackedLevelInstanceVersion)
 	{
 		Ar << PackedVersion;
 	}
