Enable the GPU physics field by default

diff --git a/Engine/Source/Runtime/Engine/Private/PhysicsField/PhysicsFieldComponent.cpp b/Engine/Source/Runtime/Engine/Private/PhysicsField/PhysicsFieldComponent.cpp
--- a/Engine/Source/Runtime/Engine/Private/PhysicsField/PhysicsFieldComponent.cpp
+++ b/Engine/Source/Runtime/Engine/Private/PhysicsField/PhysicsFieldComponent.cpp
@@ -38,7 +38,7 @@ DEFINE_LOG_CATEGORY_STATIC(LogGlobalField, Log, All);
 /** Clipmap enable/disable */
 static TAutoConsoleVariable<int32> CVarPhysicsFieldEnableClipmap(
 	TEXT("r.PhysicsField.EnableField"),
-	0,
+	1,
 	TEXT("Enable/Disable the Physics field clipmap"),
 	ECVF_RenderThreadSafe);
 
