Fix aggressive skipping of adding geometry for ray tracing.

- We want to skip adding geometry that will run through FRayTracingDynamicGeometryConverterCS and does not have a supported vertex factory type.
- I've moved checking the vertex factory type from the beginning of the function to happen closer to where we add primitives to the DynamicRayTracingGeometriesToUpdate list (which are the primitives that route through FRayTracingDynamicGeometryConverterCS)
- This is a follow up on CL# 15920326

diff --git a/Engine/Source/Runtime/Engine/Private/InstancedStaticMesh.cpp b/Engine/Source/Runtime/Engine/Private/InstancedStaticMesh.cpp
--- a/Engine/Source/Runtime/Engine/Private/InstancedStaticMesh.cpp
+++ b/Engine/Source/Runtime/Engine/Private/InstancedStaticMesh.cpp
@@ -1472,11 +1472,6 @@ void FInstancedStaticMeshSceneProxy::GetDynamicRayTracingInstances(struct FRayTr
 		return;
 	}
 
-	if (!InstancedRenderData.VertexFactories[LOD].GetType()->SupportsRayTracingDynamicGeometry())
-	{
-		return;
-	}
-
 	//setup a 'template' for the instance first, so we aren't duplicating work
 	//#dxr_todo: when multiple LODs are used, template needs to be an array of templates, probably best initialized on-demand via a lamda
 	FRayTracingInstance RayTracingInstanceTemplate;
@@ -1615,7 +1610,7 @@ void FInstancedStaticMeshSceneProxy::GetDynamicRayTracingInstances(struct FRayTr
 					FMatrix InstanceTransform = PerInstanceTransforms[InstanceIndex] * GetLocalToWorld();
 					const int32 DynamicInstanceIdx = InstanceIndex % SimulatedInstances;
 
-					if (bHasWorldPositionOffset)
+					if (bHasWorldPositionOffset && InstancedRenderData.VertexFactories[LOD].GetType()->SupportsRayTracingDynamicGeometry())
 					{
 						FRayTracingInstance *DynamicInstance = nullptr;
 
@@ -1681,7 +1676,7 @@ void FInstancedStaticMeshSceneProxy::GetDynamicRayTracingInstances(struct FRayTr
 		{
 			FMatrix InstanceTransform = PerInstanceTransforms[InstanceIdx] * GetLocalToWorld();
 
-			if (bHasWorldPositionOffset)
+			if (bHasWorldPositionOffset && InstancedRenderData.VertexFactories[LOD].GetType()->SupportsRayTracingDynamicGeometry())
 			{
 				FRayTracingInstance *DynamicInstance = nullptr;
 
diff --git a/Engine/Source/Runtime/Engine/Private/SkeletalMesh.cpp b/Engine/Source/Runtime/Engine/Private/SkeletalMesh.cpp
--- a/Engine/Source/Runtime/Engine/Private/SkeletalMesh.cpp
+++ b/Engine/Source/Runtime/Engine/Private/SkeletalMesh.cpp
@@ -6220,30 +6220,18 @@ void FSkeletalMeshSceneProxy::GetDynamicRayTracingInstances(FRayTracingMaterialG
 				MeshObject->MaterialIndexPreview = INDEX_NONE;
 			#endif
 
-				bool bSupportedVFType = true;
 				for (FSkeletalMeshSectionIter Iter(LODIndex, *MeshObject, LODData, LODSection); Iter; ++Iter)
 				{
 					const FSkelMeshRenderSection& Section = Iter.GetSection();
 					const int32 SectionIndex = Iter.GetSectionElementIndex();
 					const FSectionElementInfo& SectionElementInfo = Iter.GetSectionElementInfo();
 
-					if (!MeshObject->GetSkinVertexFactory(Context.ReferenceView, LODIndex, SectionIndex)->GetType()->SupportsRayTracingDynamicGeometry())
-					{
-						bSupportedVFType = false;
-						break;
-					}
-
 					FMeshBatch MeshBatch;
 					CreateBaseMeshBatch(Context.ReferenceView, LODData, LODIndex, SectionIndex, SectionElementInfo, MeshBatch);
 
 					RayTracingInstance.Materials.Add(MeshBatch);
 				}
 
-				if (!bSupportedVFType)
-				{
-					return;
-				}
-
 			#if WITH_EDITORONLY_DATA
 				MeshObject->SectionIndexPreview = SectionIndexPreview;
 				MeshObject->MaterialIndexPreview = MaterialIndexPreview;
@@ -6257,7 +6245,7 @@ void FSkeletalMeshSceneProxy::GetDynamicRayTracingInstances(FRayTracingMaterialG
 					RayTracingInstance.InstanceTransforms.Add(GetLocalToWorld());
 				}
 
-				if (bAnySegmentUsesWorldPositionOffset)
+				if (bAnySegmentUsesWorldPositionOffset && MeshObject->GetSkinVertexFactory(Context.ReferenceView, LODIndex, 0)->GetType()->SupportsRayTracingDynamicGeometry())
 				{
 					TArray<FRayTracingGeometrySegment> GeometrySections;
 					GeometrySections.Reserve(LODData.RenderSections.Num());
diff --git a/Engine/Source/Runtime/Engine/Private/StaticMeshRender.cpp b/Engine/Source/Runtime/Engine/Private/StaticMeshRender.cpp
--- a/Engine/Source/Runtime/Engine/Private/StaticMeshRender.cpp
+++ b/Engine/Source/Runtime/Engine/Private/StaticMeshRender.cpp
@@ -1697,11 +1697,6 @@ void FStaticMeshSceneProxy::GetDynamicRayTracingInstances(FRayTracingMaterialGat
 	const uint32 LODIndex = FMath::Max(GetLOD(Context.ReferenceView), (int32)GetCurrentFirstLODIdx_RenderThread());
 	const FStaticMeshLODResources& LODModel = RenderData->LODResources[LODIndex];
 
-	if (!RenderData->LODVertexFactories[LODIndex].VertexFactory.GetType()->SupportsRayTracingDynamicGeometry())
-	{
-		return;
-	}
-
 	if (LODModel.GetNumVertices() <= 0)
 	{
 		return;
@@ -1758,7 +1753,7 @@ void FStaticMeshSceneProxy::GetDynamicRayTracingInstances(FRayTracingMaterialGat
 
 		RayTracingInstance.Geometry = &Geometry;
 
-		if (bEvaluateWPO)
+		if (bEvaluateWPO && RenderData->LODVertexFactories[LODIndex].VertexFactory.GetType()->SupportsRayTracingDynamicGeometry())
 		{
 			RayTracingInstance.InstanceTransforms.Add(FMatrix::Identity);
 	
diff --git a/Engine/Source/Runtime/Landscape/Private/LandscapeRender.cpp b/Engine/Source/Runtime/Landscape/Private/LandscapeRender.cpp
--- a/Engine/Source/Runtime/Landscape/Private/LandscapeRender.cpp
+++ b/Engine/Source/Runtime/Landscape/Private/LandscapeRender.cpp
@@ -2472,11 +2472,6 @@ void FLandscapeComponentSceneProxy::GetDynamicRayTracingInstances(FRayTracingMat
 		return;
 	}
 
-	if (!VertexFactory->GetType()->SupportsRayTracingDynamicGeometry())
-	{
-		return;
-	}
-
 	int32 ForcedLODLevel = ForcedLOD;
 
 	int32 ViewLodOveride = GetViewLodOverride(*Context.ReferenceView);
@@ -2612,7 +2607,7 @@ void FLandscapeComponentSceneProxy::GetDynamicRayTracingInstances(FRayTracingMat
 			RayTracingInstance.BuildInstanceMaskAndFlags();
 			OutRayTracingInstances.Add(RayTracingInstance);
 
-			if (bNeedsRayTracingGeometryUpdate)
+			if (bNeedsRayTracingGeometryUpdate && VertexFactory->GetType()->SupportsRayTracingDynamicGeometry())
 			{
 				// Use the internal managed vertex buffer because landscape dynamic RT geometries are not updated every frame
 				// which is a requirement for the shared vertex buffer usage
