Add incremental update of AActor::RerunConstructionScripts.
Addresses long stalls while streaming a large number of packed level instances.

diff --git a/Engine/Source/Runtime/Engine/Classes/Engine/Level.h b/Engine/Source/Runtime/Engine/Classes/Engine/Level.h
--- a/Engine/Source/Runtime/Engine/Classes/Engine/Level.h
+++ b/Engine/Source/Runtime/Engine/Classes/Engine/Level.h
@@ -653,10 +653,21 @@ public:
 	/** Whether the level is partitioned or not. */
     UPROPERTY()
 	uint8										bIsPartitioned : 1;
+
+	enum class EIncrementalComponentState : uint8
+	{
+		Init,
+		RegisterInitialComponents,
+		RunConstructionScripts,
+		Finalize
+	};
+
 	/** Whether the actor referenced by CurrentActorIndexForUpdateComponents has called PreRegisterAllComponents */
-	uint8										bHasCurrentActorCalledPreRegister;
+	uint8										bHasCurrentActorCalledPreRegister:1;
+	/** The current stage for incrementally updating actor components in the level*/
+	EIncrementalComponentState					IncrementalComponentState;
 	/** Current index into actors array for updating components.							*/
-	int32										CurrentActorIndexForUpdateComponents;
+	int32										CurrentActorIndexForIncrementalUpdate;
 	/** Current index into actors array for updating components.							*/
 	int32										CurrentActorIndexForUnregisterComponents;
 
@@ -1151,6 +1162,9 @@ public:
 
 #endif
 
+private:
+	bool IncrementalRegisterComponents(bool bPreRegisterComponents, int32 NumComponentsToUpdate, FRegisterComponentContext* Context);
+	bool IncrementalRunConstructionScripts(bool bProcessAllActors);
 };
 
 
diff --git a/Engine/Source/Runtime/Engine/Private/Level.cpp b/Engine/Source/Runtime/Engine/Private/Level.cpp
--- a/Engine/Source/Runtime/Engine/Private/Level.cpp
+++ b/Engine/Source/Runtime/Engine/Private/Level.cpp
@@ -335,6 +335,7 @@ ULevel::ULevel( const FObjectInitializer& ObjectInitializer )
 	bActorClusterCreated = false;
 	bIsPartitioned = false;
 	bStaticComponentsRegisteredInStreamingManager = false;
+	IncrementalComponentState = EIncrementalComponentState::Init;
 }
 
 void ULevel::Initialize(const FURL& InURL)
@@ -1187,27 +1188,74 @@ void ULevel::IncrementalUpdateComponents(int32 NumComponentsToUpdate, bool bReru
 		checkf(OwningWorld->IsGameWorld(), TEXT("Cannot call IncrementalUpdateComponents with non 0 argument in the Editor/ commandlets."));
 	}
 
-	// Do BSP on the first pass.
-	if (CurrentActorIndexForUpdateComponents == 0)
+	bool bFullyUpdateComponents = NumComponentsToUpdate == 0;
+
+	// The editor is never allowed to incrementally update components.  Make sure to pass in a value of zero for NumActorsToUpdate.
+	check(bFullyUpdateComponents || OwningWorld->IsGameWorld());
+
+	do
 	{
-		UpdateModelComponents();
-		// Sort actors to ensure that parent actors will be registered before child actors
-		SortActorsHierarchy(Actors, this);
+		switch (IncrementalComponentState)
+		{
+		case EIncrementalComponentState::Init:
+			// Do BSP on the first pass.
+			UpdateModelComponents();
+			// Sort actors to ensure that parent actors will be registered before child actors
+			SortActorsHierarchy(Actors, this);
+			IncrementalComponentState = EIncrementalComponentState::RegisterInitialComponents;
+
+			//pass through
+		case EIncrementalComponentState::RegisterInitialComponents:
+			if (IncrementalRegisterComponents(true, NumComponentsToUpdate, Context))
+			{
+				bool ShouldRunConstructionScripts = !bHasRerunConstructionScripts && bRerunConstructionScripts && !IsTemplate() && !GIsUCCMakeStandaloneHeaderGenerator;
+				IncrementalComponentState = ShouldRunConstructionScripts ? EIncrementalComponentState::RunConstructionScripts : EIncrementalComponentState::Finalize;
+			}
+			break;
+
+		case EIncrementalComponentState::RunConstructionScripts:
+			if (IncrementalRunConstructionScripts(bFullyUpdateComponents))
+			{
+				IncrementalComponentState = EIncrementalComponentState::Finalize;
+			}
+			break;
+
+		case EIncrementalComponentState::Finalize:
+			IncrementalComponentState = EIncrementalComponentState::Init;
+			CurrentActorIndexForIncrementalUpdate = 0;
+			bHasCurrentActorCalledPreRegister = false;
+			bAreComponentsCurrentlyRegistered = true;
+			CreateCluster();
+			break;
+		}
+	}while(bFullyUpdateComponents && !bAreComponentsCurrentlyRegistered);
+	
+
+	{
+		SCOPE_CYCLE_COUNTER(STAT_DeferredUpdateBodies);
+#if WITH_CHAOS
+		FPhysScene* PhysScene = OwningWorld->GetPhysicsScene();
+		if (PhysScene)
+		{
+			PhysScene->ProcessDeferredCreatePhysicsState();
+		}
+#endif
 	}
+}
 
-	int32 PreviousIndex = CurrentActorIndexForUpdateComponents;
+bool ULevel::IncrementalRegisterComponents(bool bPreRegisterComponents, int32 NumComponentsToUpdate, FRegisterComponentContext* Context)
+{
 	// Find next valid actor to process components registration
-
-	while (CurrentActorIndexForUpdateComponents < Actors.Num())
+	while (CurrentActorIndexForIncrementalUpdate < Actors.Num())
 	{
-		AActor* Actor = Actors[CurrentActorIndexForUpdateComponents];
+		AActor* Actor = Actors[CurrentActorIndexForIncrementalUpdate];
 		bool bAllComponentsRegistered = true;
 		if (Actor && !Actor->IsPendingKill())
 		{
 #if PERF_TRACK_DETAILED_ASYNC_STATS
 			FScopeCycleCounterUObject ContextScope(Actor);
 #endif
-			if (!bHasCurrentActorCalledPreRegister)
+			if (bPreRegisterComponents && !bHasCurrentActorCalledPreRegister)
 			{
 				Actor->PreRegisterAllComponents();
 				bHasCurrentActorCalledPreRegister = true;
@@ -1216,9 +1264,9 @@ void ULevel::IncrementalUpdateComponents(int32 NumComponentsToUpdate, bool bReru
 		}
 
 		if (bAllComponentsRegistered)
-		{	
+		{
 			// All components have been registered fro this actor, move to a next one
-			CurrentActorIndexForUpdateComponents++;
+			CurrentActorIndexForIncrementalUpdate++;
 			bHasCurrentActorCalledPreRegister = false;
 		}
 
@@ -1230,68 +1278,65 @@ void ULevel::IncrementalUpdateComponents(int32 NumComponentsToUpdate, bool bReru
 		}
 	}
 
-	// See whether we are done.
-	if (CurrentActorIndexForUpdateComponents >= Actors.Num())
+	if (CurrentActorIndexForIncrementalUpdate >= Actors.Num())
 	{
-		CurrentActorIndexForUpdateComponents	= 0;
-		bHasCurrentActorCalledPreRegister		= false;
-		bAreComponentsCurrentlyRegistered		= true;
+		// We need to process pending adds prior to rerunning the construction scripts, which may internally
+		// perform removals / adds themselves.
+		if (Context)
+		{
+			Context->Process();
+		}
+		CurrentActorIndexForIncrementalUpdate = 0;
+		return true;
+	}
 
-#if PERF_TRACK_DETAILED_ASYNC_STATS
-		QUICK_SCOPE_CYCLE_COUNTER(STAT_ULevel_IncrementalUpdateComponents_RerunConstructionScripts);
-#endif
-		if (bRerunConstructionScripts && !IsTemplate() && !GIsUCCMakeStandaloneHeaderGenerator)
+	return false;
+}
+
+bool ULevel::IncrementalRunConstructionScripts(bool bProcessAllActors)
+{
+	// Find next valid actor to process components registration
+	while (CurrentActorIndexForIncrementalUpdate < Actors.Num())
+	{
+		AActor* Actor = Actors[CurrentActorIndexForIncrementalUpdate];
+		CurrentActorIndexForIncrementalUpdate++;
+
+		if (Actor == nullptr || Actor->IsPendingKill() || Actor->IsChildActor())
 		{
-			// We need to process pending adds prior to rerunning the construction scripts, which may internally
-			// perform removals / adds themselves.
-			if (Context)
-			{
-				Context->Process();
-			}
+			continue;
+		}
+
+		// Child actors have already been built and initialized up by their parent and they should not be reconstructed again
+		if (Actor->IsChildActor())
+		{
+			continue;
+		}
 
-			// Don't rerun construction scripts until after all actors' components have been registered.  This
-			// is necessary because child attachment lists are populated during registration, and running construction
-			// scripts requires that the attachments are correctly initialized.
-			// Don't use ranged for as construction scripts can manipulate the actor array
-			for (int32 ActorIndex = 0; ActorIndex < Actors.Num(); ++ActorIndex)
-			{
-				if (AActor* Actor = Actors[ActorIndex])
-				{
-					// Child actors have already been built and initialized up by their parent and they should not be reconstructed again
-					if (!Actor->IsChildActor())
-					{
 #if PERF_TRACK_DETAILED_ASYNC_STATS
-						FScopeCycleCounterUObject ContextScope(Actor);
+		FScopeCycleCounterUObject ContextScope(Actor);
 #endif
-						Actor->RerunConstructionScripts();
-					}
-				}
-			}
-			bHasRerunConstructionScripts = true;
-		}
+		
+		// TODO - Ideally we could pull out component registration from this call, allowing for even better distribution of the 
+		// expensive RegisterComponentWithWorld calls, but there are a few systems that would need to be re-factored in order to do this
+		// safely.
+		Actor->RerunConstructionScripts();
 
-		CreateCluster();
-	}
-	// Only the game can use incremental update functionality.
-	else
-	{
-		// The editor is never allowed to incrementally updated components.  Make sure to pass in a value of zero for NumActorsToUpdate.
-		check(OwningWorld->IsGameWorld());
+		// If we do an incremental registration return to outer loop after each processed actor 
+		// so outer loop can decide whether we want to continue processing this frame
+		if (!bProcessAllActors)
+		{
+			break;
+		}
 	}
 
+	if (CurrentActorIndexForIncrementalUpdate >= Actors.Num())
 	{
-		SCOPE_CYCLE_COUNTER(STAT_DeferredUpdateBodies);
-#if WITH_CHAOS
-		FPhysScene* PhysScene = OwningWorld->GetPhysicsScene();
-		if (PhysScene)
-		{
-			PhysScene->ProcessDeferredCreatePhysicsState();
-		}
-#endif
+		CurrentActorIndexForIncrementalUpdate = 0;
+		return true;
 	}
+	return false;
 }
 
-
 bool ULevel::IncrementalUnregisterComponents(int32 NumComponentsToUnregister)
 {
 	// A value of 0 means that we want to unregister all components.
