HLODLayer - Added a "Always Loaded" property
HLODActor - Avoid using SetIsTemporarilyHiddenInEditor() as it's causing isues with some workflow (Ctrl+H)

diff --git a/Engine/Source/Runtime/Engine/Private/WorldPartition/HLOD/HLODActor.cpp b/Engine/Source/Runtime/Engine/Private/WorldPartition/HLOD/HLODActor.cpp
--- a/Engine/Source/Runtime/Engine/Private/WorldPartition/HLOD/HLODActor.cpp
+++ b/Engine/Source/Runtime/Engine/Private/WorldPartition/HLOD/HLODActor.cpp
@@ -9,6 +9,7 @@
 #include "WorldPartition/WorldPartitionActorDesc.h"
 #include "WorldPartition/HLOD/HLODActorDesc.h"
 #include "WorldPartition/HLOD/HLODBuilder.h"
+#include "WorldPartition/HLOD/HLODLayer.h"
 #endif
 
 AWorldPartitionHLOD::AWorldPartitionHLOD(const FObjectInitializer& ObjectInitializer)
@@ -67,18 +68,14 @@ void AWorldPartitionHLOD::RerunConstructionScripts()
 
 #if WITH_EDITOR
 
-void AWorldPartitionHLOD::PostRegisterAllComponents()
+bool AWorldPartitionHLOD::IsHiddenEd() const
 {
-	Super::PostRegisterAllComponents();
-
-	SetIsTemporarilyHiddenInEditor(true);
-	bListedInSceneOutliner = true;
-	SetFolderPath(*FString::Printf(TEXT("HLOD/HLOD%d"), GetLODLevel()));
+	return true;
 }
 
-EActorGridPlacement AWorldPartitionHLOD::GetDefaultGridPlacement() const
+EActorGridPlacement AWorldPartitionHLOD::GetGridPlacement() const
 {
-	return EActorGridPlacement::Bounds;
+	return SubActorsHLODLayer && SubActorsHLODLayer->IsAlwaysLoaded() ? EActorGridPlacement::AlwaysLoaded : EActorGridPlacement::Bounds;
 }
 
 TUniquePtr<FWorldPartitionActorDesc> AWorldPartitionHLOD::CreateClassActorDesc() const
diff --git a/Engine/Source/Runtime/Engine/Private/WorldPartition/HLOD/HLODLayer.cpp b/Engine/Source/Runtime/Engine/Private/WorldPartition/HLOD/HLODLayer.cpp
--- a/Engine/Source/Runtime/Engine/Private/WorldPartition/HLOD/HLODLayer.cpp
+++ b/Engine/Source/Runtime/Engine/Private/WorldPartition/HLOD/HLODLayer.cpp
@@ -156,7 +156,13 @@ FName UHLODLayer::GetRuntimeGridName(uint32 InLODLevel, int32 InCellSize, float
 
 FName UHLODLayer::GetRuntimeGrid(uint32 InHLODLevel) const
 {
-	return GetRuntimeGridName(InHLODLevel, CellSize, LoadingRange);
+	return IsAlwaysLoaded() ? NAME_None : GetRuntimeGridName(InHLODLevel, CellSize, LoadingRange);
+}
+
+const TSoftObjectPtr<UHLODLayer>& UHLODLayer::GetParentLayer() const
+{
+	static const TSoftObjectPtr<UHLODLayer> NullLayer;
+	return IsAlwaysLoaded() ? NullLayer : ParentLayer;
 }
 
 #endif // WITH_EDITOR
diff --git a/Engine/Source/Runtime/Engine/Private/WorldPartition/WorldPartitionRuntimeSpatialHash.cpp b/Engine/Source/Runtime/Engine/Private/WorldPartition/WorldPartitionRuntimeSpatialHash.cpp
--- a/Engine/Source/Runtime/Engine/Private/WorldPartition/WorldPartitionRuntimeSpatialHash.cpp
+++ b/Engine/Source/Runtime/Engine/Private/WorldPartition/WorldPartitionRuntimeSpatialHash.cpp
@@ -848,11 +848,8 @@ bool UWorldPartitionRuntimeSpatialHash::CreateStreamingGrid(const FSpatialHashRu
 					UE_LOG(LogWorldPartition, Verbose, TEXT("  Actor : %s (%s) (Container %08x) Origin(%s)"), *(ActorDescView.GetActorPath().ToString()), *ActorDescView.GetGuid().ToString(EGuidFormats::UniqueObjectGuid), ActorInstance.ContainerInstance->ID, *FVector2D(ActorInstance.GetOrigin()).ToString());
 				}
 
-				if (NumHLODActors > 0)
+				if (NumHLODActors == FilteredActors.Num())
 				{
-					// HLOD cells should contain only HLODs
-					check(NumHLODActors == FilteredActors.Num());
-
 					// Mark this cell as containing HLOD data
 					UWorldPartitionRuntimeHLODCellData* HLODCellData = NewObject<UWorldPartitionRuntimeHLODCellData>(StreamingCell);
 					StreamingCell->AddCellData(HLODCellData);
diff --git a/Engine/Source/Runtime/Engine/Public/WorldPartition/HLOD/HLODActor.h b/Engine/Source/Runtime/Engine/Public/WorldPartition/HLOD/HLODActor.h
--- a/Engine/Source/Runtime/Engine/Public/WorldPartition/HLOD/HLODActor.h
+++ b/Engine/Source/Runtime/Engine/Public/WorldPartition/HLOD/HLODActor.h
@@ -67,8 +67,8 @@ protected:
 	virtual void BeginPlay() override;
 	virtual void EndPlay(const EEndPlayReason::Type EndPlayReason) override;
 #if WITH_EDITOR
-	virtual void PostRegisterAllComponents() override;
-	virtual EActorGridPlacement GetDefaultGridPlacement() const override;
+	virtual bool IsHiddenEd() const override;
+	virtual EActorGridPlacement GetGridPlacement() const override;
 	virtual TUniquePtr<class FWorldPartitionActorDesc> CreateClassActorDesc() const override;
 
 	virtual void GetActorBounds(bool bOnlyCollidingComponents, FVector& Origin, FVector& BoxExtent, bool bIncludeFromChildActors) const override;
diff --git a/Engine/Source/Runtime/Engine/Public/WorldPartition/HLOD/HLODLayer.h b/Engine/Source/Runtime/Engine/Public/WorldPartition/HLOD/HLODLayer.h
--- a/Engine/Source/Runtime/Engine/Public/WorldPartition/HLOD/HLODLayer.h
+++ b/Engine/Source/Runtime/Engine/Public/WorldPartition/HLOD/HLODLayer.h
@@ -43,9 +43,10 @@ public:
 	const FMeshProxySettings& GetMeshSimplifySettings() const { return MeshSimplifySettings; }
 	const TSoftObjectPtr<UMaterial>& GetHLODMaterial() const { return HLODMaterial; }
 	FName GetRuntimeGrid(uint32 InHLODLevel) const;
-	int32 GetCellSize() const { return CellSize; }
-	float GetLoadingRange() const { return LoadingRange; }
-	const TSoftObjectPtr<UHLODLayer>& GetParentLayer() const { return ParentLayer; }
+	int32 GetCellSize() const { return bAlwaysLoaded ? 0 : CellSize; }
+	float GetLoadingRange() const { return bAlwaysLoaded ? WORLD_MAX : LoadingRange; }
+	const TSoftObjectPtr<UHLODLayer>& GetParentLayer() const;
+	bool IsAlwaysLoaded() const { return bAlwaysLoaded; }
 
 	static FName GetRuntimeGridName(uint32 InLODLevel, int32 InCellSize, float InLoadingRange);
 #endif
@@ -57,26 +58,31 @@ private:
 	EHLODLayerType LayerType;
 
 	/** Merged mesh generation settings - Used when this layer is of type EHLODLayerType::MeshMerge */
-	UPROPERTY(EditAnywhere, Config, Category=HLOD, meta = (editcondition = "LayerType == EHLODLayerType::MeshMerge"))
+	UPROPERTY(EditAnywhere, Config, Category=HLOD, meta = (EditConditionHides, EditCondition = "LayerType == EHLODLayerType::MeshMerge"))
 	FMeshMergingSettings MeshMergeSettings;
 
 	/** Simplified mesh generation settings - Used when this layer is of type EHLODLayerType::MeshSimplify */
-	UPROPERTY(EditAnywhere, Config, Category=HLOD, meta = (editcondition = "LayerType == EHLODLayerType::MeshSimplify"))
+	UPROPERTY(EditAnywhere, Config, Category=HLOD, meta = (EditConditionHides, EditCondition = "LayerType == EHLODLayerType::MeshSimplify"))
 	FMeshProxySettings MeshSimplifySettings;
 
 	/** Material that will be used by the generated HLOD static mesh */
-	UPROPERTY(EditAnywhere, Config, AdvancedDisplay, Category=HLOD, meta = (editcondition = "LayerType == EHLODLayerType::MeshMerge || LayerType == EHLODLayerType::MeshSimplify"))
+	UPROPERTY(EditAnywhere, Config, AdvancedDisplay, Category=HLOD, meta = (EditConditionHides, EditCondition = "LayerType == EHLODLayerType::MeshMerge || LayerType == EHLODLayerType::MeshSimplify"))
 	TSoftObjectPtr<UMaterial> HLODMaterial;
 
+	/** Whether HLOD actors generated for this layer will be always loaded */
 	UPROPERTY(EditAnywhere, Config, Category=HLOD)
+	uint32 bAlwaysLoaded : 1;
+
+	/** Cell size of the runtime grid created to encompass HLOD actors generated for this HLOD Layer */
+	UPROPERTY(EditAnywhere, Config, Category=HLOD, meta = (EditConditionHides, EditCondition = "!bAlwaysLoaded"))
 	int32 CellSize;
 
-	/** Loading range of the RuntimeGrid */
-	UPROPERTY(EditAnywhere, Config, Category=HLOD)
+	/** Loading range of the runtime grid created to encompass HLOD actors generated for this HLOD Layer */
+	UPROPERTY(EditAnywhere, Config, Category=HLOD, meta = (EditConditionHides, EditCondition = "!bAlwaysLoaded"))
 	float LoadingRange;
 
-	/** HLODLayer to assign to the generated HLOD actors */
-	UPROPERTY(EditAnywhere, Config, Category=HLOD)
+	/** HLOD Layer to assign to the generated HLOD actors */
+	UPROPERTY(EditAnywhere, Config, Category=HLOD, meta = (EditConditionHides, EditCondition = "!bAlwaysLoaded"))
 	TSoftObjectPtr<UHLODLayer> ParentLayer;
 #endif
 };
