World Partition: changed how we promote actors into the spatial hash grid levels: always store an actor into cells that has over 66% coverage instead of going to higher levels.
Also added check for errors for derived runtime hashes, and implemented empty actor bounds (2D) checks for spatial hash class.

diff --git a/Engine/Source/Runtime/Engine/Private/WorldPartition/RuntimeSpatialHash/RuntimeSpatialHashGridHelper.cpp b/Engine/Source/Runtime/Engine/Private/WorldPartition/RuntimeSpatialHash/RuntimeSpatialHashGridHelper.cpp
--- a/Engine/Source/Runtime/Engine/Private/WorldPartition/RuntimeSpatialHash/RuntimeSpatialHashGridHelper.cpp
+++ b/Engine/Source/Runtime/Engine/Private/WorldPartition/RuntimeSpatialHash/RuntimeSpatialHashGridHelper.cpp
@@ -186,18 +186,33 @@ FSquare2DGridHelper GetPartitionedActors(const UWorldPartition* WorldPartition,
 			bool bFoundCell = false;
 			for (FSquare2DGridHelper::FGridLevel& GridLevel : PartitionedActors.Levels)
 			{
-				int32 IntersectingCellCount = 0;
-				GridLevel.ForEachIntersectingCellsBreakable(ClusterInstance->Bounds, [&IntersectingCellCount](const FIntVector2& Coords) { return ++IntersectingCellCount <= 1; });
-				if (IntersectingCellCount == 1)
+				const FBox2D ActorBounds(FVector2D(ClusterInstance->Bounds.Min), FVector2D(ClusterInstance->Bounds.Max));
+				const float ActorBoundsArea = ActorBounds.GetArea();
+
+				GridLevel.ForEachIntersectingCellsBreakable(ClusterInstance->Bounds, [&GridLevel, ActorCluster, ClusterInstance, &ActorBounds, &ActorBoundsArea, &bFoundCell](const FIntVector2& Coords)
+				{
+					FBox2D CellBounds;
+					GridLevel.GetCellBounds(Coords, CellBounds);
+
+					const FBox2D Intersection(FVector2D::Max(CellBounds.Min, ActorBounds.Min), FVector2D::Min(CellBounds.Max, ActorBounds.Max));
+					const float CoveredArea = Intersection.GetArea() / ActorBoundsArea;
+
+					if (CoveredArea > 0.66f)
+					{
+						GridLevel.GetCell(Coords).AddActors(ActorCluster->Actors, ClusterInstance->ContainerInstance, ClusterInstance->DataLayers);
+						bFoundCell = true;
+						return false;
+					}
+
+					return true;
+				});
+
+				if (bFoundCell)
 				{
-					GridLevel.ForEachIntersectingCells(ClusterInstance->Bounds, [&GridLevel, &ActorCluster, &ClusterInstance](const FIntVector2& Coords)
-						{
-							GridLevel.GetCell(Coords).AddActors(ActorCluster->Actors, ClusterInstance->ContainerInstance, ClusterInstance->DataLayers);
-						});
-					bFoundCell = true;
 					break;
 				}
 			}
+
 			if (!(ensure(bFoundCell)))
 			{
 				GridPlacement = EActorGridPlacement::AlwaysLoaded;
diff --git a/Engine/Source/Runtime/Engine/Private/WorldPartition/WorldPartitionRuntimeHash.cpp b/Engine/Source/Runtime/Engine/Private/WorldPartition/WorldPartitionRuntimeHash.cpp
--- a/Engine/Source/Runtime/Engine/Private/WorldPartition/WorldPartitionRuntimeHash.cpp
+++ b/Engine/Source/Runtime/Engine/Private/WorldPartition/WorldPartitionRuntimeHash.cpp
@@ -9,7 +9,6 @@
 
 #if WITH_EDITOR
 #include "WorldPartition/WorldPartitionHandle.h"
-#include "WorldPartition/WorldPartitionActorDescViewProxy.h"
 #include "Logging/MessageLog.h"
 #include "Misc/UObjectToken.h"
 #include "Misc/MapErrors.h"
@@ -154,6 +153,17 @@ void UWorldPartitionRuntimeHash::ChangeActorDescViewGridPlacement(FWorldPartitio
 }
 
 void UWorldPartitionRuntimeHash::CheckForErrors() const
+{
+	TMap<FGuid, FWorldPartitionActorViewProxy> ActorDescList;
+	for (UActorDescContainer::TConstIterator<> ActorDescIt(GetOuterUWorldPartition()); ActorDescIt; ++ActorDescIt)
+	{
+		ActorDescList.Emplace(ActorDescIt->GetGuid(), *ActorDescIt);
+	}
+
+	CheckForErrorsInternal(ActorDescList);
+}
+
+void UWorldPartitionRuntimeHash::CheckForErrorsInternal(const TMap<FGuid, FWorldPartitionActorViewProxy>& ActorDescList) const
 {
 	auto GetActorLabel = [](const FWorldPartitionActorDescView& ActorDescView) -> FString
 	{
@@ -175,12 +185,6 @@ void UWorldPartitionRuntimeHash::CheckForErrors() const
 		return ActorPath;
 	};
 
-	TMap<FGuid, FWorldPartitionActorViewProxy> ActorDescList;
-	for (UActorDescContainer::TConstIterator<> ActorDescIt(GetOuterUWorldPartition()); ActorDescIt; ++ActorDescIt)
-	{
-		ActorDescList.Emplace(ActorDescIt->GetGuid(), *ActorDescIt);
-	}
-
 	for (auto& ActorDescListPair: ActorDescList)
 	{
 		const FWorldPartitionActorViewProxy& ActorDescView = ActorDescListPair.Value;
@@ -244,7 +248,7 @@ void UWorldPartitionRuntimeHash::CheckForErrors() const
 
 		for (AActor* Actor : LevelScriptExternalActorReferences)
 		{
-			if (FWorldPartitionActorDescView* ActorDescView = ActorDescList.Find(Actor->GetActorGuid()))
+			if (const FWorldPartitionActorDescView* ActorDescView = ActorDescList.Find(Actor->GetActorGuid()))
 			{
 				TArray<FName> ActorDescLayers = ActorDescView->GetDataLayers();
 				if (ActorDescLayers.Num())
diff --git a/Engine/Source/Runtime/Engine/Private/WorldPartition/WorldPartitionRuntimeSpatialHash.cpp b/Engine/Source/Runtime/Engine/Private/WorldPartition/WorldPartitionRuntimeSpatialHash.cpp
--- a/Engine/Source/Runtime/Engine/Private/WorldPartition/WorldPartitionRuntimeSpatialHash.cpp
+++ b/Engine/Source/Runtime/Engine/Private/WorldPartition/WorldPartitionRuntimeSpatialHash.cpp
@@ -42,9 +42,15 @@
 #include "Engine/WorldComposition.h"
 #include "LevelUtils.h"
 
+#include "Logging/MessageLog.h"
+#include "Misc/UObjectToken.h"
+#include "Misc/MapErrors.h"
+
 extern UNREALED_API class UEditorEngine* GEditor;
 #endif //WITH_EDITOR
 
+#define LOCTEXT_NAMESPACE "WorldPartition"
+
 static int32 GShowRuntimeSpatialHashGridLevel = 0;
 static FAutoConsoleVariableRef CVarShowRuntimeSpatialHashGridLevel(
 	TEXT("wp.Runtime.ShowRuntimeSpatialHashGridLevel"),
@@ -484,6 +490,51 @@ FName UWorldPartitionRuntimeSpatialHash::GetActorRuntimeGrid(const AActor* Actor
 	return Super::GetActorRuntimeGrid(Actor);
 }
 
+void UWorldPartitionRuntimeSpatialHash::CheckForErrorsInternal(const TMap<FGuid, FWorldPartitionActorViewProxy>& ActorDescList) const
+{
+	Super::CheckForErrorsInternal(ActorDescList);
+
+	auto GetActorLabel = [](const FWorldPartitionActorDescView& ActorDescView) -> FString
+	{
+		const FName ActorLabel = ActorDescView.GetActorLabel();
+		if (!ActorLabel.IsNone())
+		{
+			return ActorLabel.ToString();
+		}
+
+		const FString ActorPath = ActorDescView.GetActorPath().ToString();
+
+		FString SubObjectName;
+		FString SubObjectContext;
+		if (FString(ActorPath).Split(TEXT("."), &SubObjectContext, &SubObjectName))
+		{
+			return SubObjectName;
+		}
+
+		return ActorPath;
+	};
+
+	for (auto& ActorDescListPair: ActorDescList)
+	{
+		const FWorldPartitionActorViewProxy& ActorDescView = ActorDescListPair.Value;
+
+		if (!ActorDescView.GetActorIsEditorOnly())
+		{
+			const FBox2D ActorBounds2D = FBox2D(FVector2D(ActorDescView.GetBounds().Min), FVector2D(ActorDescView.GetBounds().Max));
+			const float ActorBoundsArea2D = ActorBounds2D.GetArea();
+
+			if (ActorBoundsArea2D < KINDA_SMALL_NUMBER)
+			{
+				TSharedRef<FTokenizedMessage> Error = FMessageLog("MapCheck").Warning()
+					->AddToken(FTextToken::Create(LOCTEXT("MapCheck_WorldPartition_Actor", "Actor")))
+					->AddToken(FAssetNameToken::Create(GetActorLabel(ActorDescView)))
+					->AddToken(FTextToken::Create(LOCTEXT("MapCheck_WorldPartition_HasEmptyBounds", "has empty bounds")))
+					->AddToken(FMapErrorToken::Create(FName(TEXT("WorldPartition_ActorHasEmptyBounds_CheckForErrors"))));
+			}
+		}
+	}
+}
+
 void UWorldPartitionRuntimeSpatialHash::SetDefaultValues()
 {
 	FSpatialHashRuntimeGrid& MainGrid = Grids.AddDefaulted_GetRef();
@@ -681,7 +732,11 @@ void UWorldPartitionRuntimeSpatialHash::CreateActorDescViewMap(const UActorDescC
 					const FBox2D ActorBounds2D = FBox2D(FVector2D(ActorDescView.GetBounds().Min), FVector2D(ActorDescView.GetBounds().Max));
 					const float ActorBoundsArea = ActorBounds2D.GetArea();
 
-					if (ActorBoundsArea > WorldBoundsArea / 2.0f)
+					if (ActorBoundsArea < 1.0f)
+					{
+						ChangeActorDescViewGridPlacement(ActorDescView, EActorGridPlacement::Location);
+					}
+					else if (ActorBoundsArea > WorldBoundsArea / 2.0f)
 					{
 						ChangeActorDescViewGridPlacement(ActorDescView, EActorGridPlacement::AlwaysLoaded);
 					}
@@ -1162,3 +1217,5 @@ void UWorldPartitionRuntimeSpatialHash::Draw3D(const TArray<FWorldPartitionStrea
 		}
 	}
 }
+
+#undef LOCTEXT_NAMESPACE
\ No newline at end of file
diff --git a/Engine/Source/Runtime/Engine/Public/WorldPartition/WorldPartitionRuntimeHash.h b/Engine/Source/Runtime/Engine/Public/WorldPartition/WorldPartitionRuntimeHash.h
--- a/Engine/Source/Runtime/Engine/Public/WorldPartition/WorldPartitionRuntimeHash.h
+++ b/Engine/Source/Runtime/Engine/Public/WorldPartition/WorldPartitionRuntimeHash.h
@@ -9,6 +9,7 @@
 #include "WorldPartitionActorDescView.h"
 #include "WorldPartition/ActorDescList.h"
 #include "WorldPartition/WorldPartitionHandle.h"
+#include "WorldPartition/WorldPartitionActorDescViewProxy.h"
 #include "WorldPartitionRuntimeHash.generated.h"
 
 class UWorldPartitionRuntimeCell;
@@ -31,9 +32,10 @@ class ENGINE_API UWorldPartitionRuntimeHash : public UObject
 	virtual bool GenerateHLOD(ISourceControlHelper* SourceControlHelper, bool bCreateActorsOnly) { return false; }
 	virtual bool GenerateNavigationData() { return false; }
 	virtual FName GetActorRuntimeGrid(const AActor* Actor) const { return NAME_None; }
-	virtual void CheckForErrors() const;
 	virtual void DrawPreview() const {}
 
+	void CheckForErrors() const;
+
 	// PIE/Game methods
 	void OnBeginPlay(EWorldPartitionStreamingMode Mode);
 	void OnEndPlay();
@@ -57,9 +59,11 @@ class ENGINE_API UWorldPartitionRuntimeHash : public UObject
 
 protected:
 #if WITH_EDITOR
+	virtual void CheckForErrorsInternal(const TMap<FGuid, FWorldPartitionActorViewProxy>& ActorDescList) const;
+
 	virtual bool GenerateStreaming(EWorldPartitionStreamingMode Mode, class UWorldPartitionStreamingPolicy* Policy, TArray<FString>* OutPackagesToGenerate) { return false; }
 	virtual void CreateActorDescViewMap(const UActorDescContainer* Container, TMap<FGuid, FWorldPartitionActorDescView>& OutActorDescViewMap) const;
-	void ChangeActorDescViewGridPlacement(FWorldPartitionActorDescView& ActorDescView, EActorGridPlacement GridPlacement) const;
+	void ChangeActorDescViewGridPlacement(FWorldPartitionActorDescView& ActorDescView, EActorGridPlacement GridPlacement) const;	
 #endif
 
 private:
diff --git a/Engine/Source/Runtime/Engine/Public/WorldPartition/WorldPartitionRuntimeSpatialHash.h b/Engine/Source/Runtime/Engine/Public/WorldPartition/WorldPartitionRuntimeSpatialHash.h
--- a/Engine/Source/Runtime/Engine/Public/WorldPartition/WorldPartitionRuntimeSpatialHash.h
+++ b/Engine/Source/Runtime/Engine/Public/WorldPartition/WorldPartitionRuntimeSpatialHash.h
@@ -207,6 +207,7 @@ public:
 	virtual bool GenerateHLOD(ISourceControlHelper* SourceControlHelper, bool bCreateActorsOnly) override;
 	virtual bool GenerateNavigationData() override;
 	virtual FName GetActorRuntimeGrid(const AActor* Actor) const override;
+	virtual void CheckForErrorsInternal(const TMap<FGuid, FWorldPartitionActorViewProxy>& ActorDescList) const override;
 	virtual void DrawPreview() const override;
 
 	FName GetCellName(FName InGridName, const FIntVector& InCellGlobalCoord, const FDataLayersID& InDataLayerID) const;
