HLOD - Readded GetDefaultGridPlacement() that was removed in previous commit
* We must override GetDefaultGridPlacement() as the default AActor::GetDefaultGridPlacement() will mark all actors that are not placeable as AlwaysLoaded (& AWorldPartitionHLOD isn't placeable)...

diff --git a/Engine/Source/Runtime/Engine/Private/WorldPartition/HLOD/HLODActor.cpp b/Engine/Source/Runtime/Engine/Private/WorldPartition/HLOD/HLODActor.cpp
--- a/Engine/Source/Runtime/Engine/Private/WorldPartition/HLOD/HLODActor.cpp
+++ b/Engine/Source/Runtime/Engine/Private/WorldPartition/HLOD/HLODActor.cpp
@@ -78,6 +78,12 @@ EActorGridPlacement AWorldPartitionHLOD::GetGridPlacement() const
 	return SubActorsHLODLayer && SubActorsHLODLayer->IsAlwaysLoaded() ? EActorGridPlacement::AlwaysLoaded : EActorGridPlacement::Bounds;
 }
 
+EActorGridPlacement AWorldPartitionHLOD::GetDefaultGridPlacement() const
+{
+	// Overriden as AActor::GetDefaultGridPlacement() will mark all actors that are not placeable as AlwaysLoaded...
+	return EActorGridPlacement::None;
+}
+
 TUniquePtr<FWorldPartitionActorDesc> AWorldPartitionHLOD::CreateClassActorDesc() const
 {
 	return TUniquePtr<FWorldPartitionActorDesc>(new FHLODActorDesc());
diff --git a/Engine/Source/Runtime/Engine/Public/WorldPartition/HLOD/HLODActor.h b/Engine/Source/Runtime/Engine/Public/WorldPartition/HLOD/HLODActor.h
--- a/Engine/Source/Runtime/Engine/Public/WorldPartition/HLOD/HLODActor.h
+++ b/Engine/Source/Runtime/Engine/Public/WorldPartition/HLOD/HLODActor.h
@@ -69,6 +69,7 @@ protected:
 #if WITH_EDITOR
 	virtual bool IsHiddenEd() const override;
 	virtual EActorGridPlacement GetGridPlacement() const override;
+	virtual EActorGridPlacement GetDefaultGridPlacement() const override;
 	virtual TUniquePtr<class FWorldPartitionActorDesc> CreateClassActorDesc() const override;
 
 	virtual void GetActorBounds(bool bOnlyCollidingComponents, FVector& Origin, FVector& BoxExtent, bool bIncludeFromChildActors) const override;
