Re-order property access batched external copy to after event graph work

This allows properties that are transformed and accessed via BP functions to be correctly transformed

diff --git a/Engine/Source/Runtime/Engine/Private/Animation/AnimInstance.cpp b/Engine/Source/Runtime/Engine/Private/Animation/AnimInstance.cpp
--- a/Engine/Source/Runtime/Engine/Private/Animation/AnimInstance.cpp
+++ b/Engine/Source/Runtime/Engine/Private/Animation/AnimInstance.cpp
@@ -487,19 +487,13 @@ void UAnimInstance::UpdateAnimation(float DeltaSeconds, bool bNeedsValidRootMoti
 		UpdateMontageEvaluationData();
 	}
 
-	if (IAnimClassInterface* AnimBlueprintClass = IAnimClassInterface::GetFromClass(GetClass()))
-	{
-		// Process internal batched property copies
-		PropertyAccess::ProcessCopies(this, AnimBlueprintClass->GetPropertyAccessLibrary(), EPropertyAccessCopyBatch::ExternalBatched);
-	}
-
 	{
 		SCOPE_CYCLE_COUNTER(STAT_NativeUpdateAnimation);
 		CSV_SCOPED_TIMING_STAT(Animation, NativeUpdate);
 		NativeUpdateAnimation(DeltaSeconds);
 
 PRAGMA_DISABLE_DEPRECATION_WARNINGS
-		// @todo: remove once deprecated - called for backwards-compatibility
+		// @todo: remove once deprecated - called for backwards-compatibility¬
 		NativeUpdateAnimation_WorkerThread(DeltaSeconds);
 PRAGMA_ENABLE_DEPRECATION_WARNINGS
 	}
@@ -508,6 +502,14 @@ PRAGMA_ENABLE_DEPRECATION_WARNINGS
 		CSV_SCOPED_TIMING_STAT(Animation, BlueprintUpdate);
 		BlueprintUpdateAnimation(DeltaSeconds);
 	}
+
+	// Perform property access copies after the event graph work - this allows properties that are transformed and
+	// accessed via BP functions to be correctly transformed.
+	if (IAnimClassInterface* AnimBlueprintClass = IAnimClassInterface::GetFromClass(GetClass()))
+	{
+		// Process internal batched property copies
+		PropertyAccess::ProcessCopies(this, AnimBlueprintClass->GetPropertyAccessLibrary(), EPropertyAccessCopyBatch::ExternalBatched);
+	}
 	
 	// Determine whether or not the animation should be immediately updated according to current state
 	const bool bWantsImmediateUpdate = bNeedsValidRootMotion || NeedsImmediateUpdate(DeltaSeconds);
