Fixing log spam on anim notifies with metasounds

UE-113192

diff --git a/Engine/Source/Runtime/Engine/Private/Animation/AnimNotify_PlaySound.cpp b/Engine/Source/Runtime/Engine/Private/Animation/AnimNotify_PlaySound.cpp
--- a/Engine/Source/Runtime/Engine/Private/Animation/AnimNotify_PlaySound.cpp
+++ b/Engine/Source/Runtime/Engine/Private/Animation/AnimNotify_PlaySound.cpp
@@ -76,7 +76,7 @@ void UAnimNotify_PlaySound::ValidateAssociatedAssets()
 {
 	static const FName NAME_AssetCheck("AssetCheck");
 
-	if ((Sound != nullptr) && (Sound->IsLooping()))
+	if ((Sound != nullptr) && (Sound->IsLooping()) && (!Sound->IsA<USoundWaveProcedural>()))
 	{
 		UObject* ContainingAsset = GetContainingAsset();
 
