WorldPartition: Fixed missing remap of SoftObjectPath on Persistent Level at cook.

diff --git a/Engine/Source/Runtime/Engine/Private/WorldPartition/WorldPartition.cpp b/Engine/Source/Runtime/Engine/Private/WorldPartition/WorldPartition.cpp
--- a/Engine/Source/Runtime/Engine/Private/WorldPartition/WorldPartition.cpp
+++ b/Engine/Source/Runtime/Engine/Private/WorldPartition/WorldPartition.cpp
@@ -1086,6 +1086,9 @@ void UWorldPartition::FinalizeGeneratedPackageForCook()
 {
 	check(RuntimeHash);
 	RuntimeHash->FinalizeGeneratedPackageForCook();
+
+	// Apply remapping of Persistent Level's SoftObjectPaths
+	FWorldPartitionLevelHelper::RemapLevelSoftObjectPaths(World->PersistentLevel, this);
 }
 
 void UWorldPartition::SavePerUserSettings()
