Fixed check happening next time PIE is started after last one was cancelled due to an error.

diff --git a/Engine/Source/Editor/UnrealEd/Private/Editor.cpp b/Engine/Source/Editor/UnrealEd/Private/Editor.cpp
--- a/Engine/Source/Editor/UnrealEd/Private/Editor.cpp
+++ b/Engine/Source/Editor/UnrealEd/Private/Editor.cpp
@@ -109,6 +109,7 @@ FEditorDelegates::FOnPIEEvent							FEditorDelegates::ResumePIE;
 FEditorDelegates::FOnPIEEvent							FEditorDelegates::SingleStepPIE;
 FEditorDelegates::FOnPIEEvent							FEditorDelegates::OnPreSwitchBeginPIEAndSIE;
 FEditorDelegates::FOnPIEEvent							FEditorDelegates::OnSwitchBeginPIEAndSIE;
+FSimpleMulticastDelegate								FEditorDelegates::CancelPIE;
 FEditorDelegates::FOnStandaloneLocalPlayEvent			FEditorDelegates::BeginStandaloneLocalPlay;
 FSimpleMulticastDelegate								FEditorDelegates::PropertySelectionChange;
 FSimpleMulticastDelegate								FEditorDelegates::PostLandscapeLayerUpdated;
diff --git a/Engine/Source/Editor/UnrealEd/Private/PlayLevel.cpp b/Engine/Source/Editor/UnrealEd/Private/PlayLevel.cpp
--- a/Engine/Source/Editor/UnrealEd/Private/PlayLevel.cpp
+++ b/Engine/Source/Editor/UnrealEd/Private/PlayLevel.cpp
@@ -944,6 +944,7 @@ void UEditorEngine::RequestPlaySession(const FString& DeviceId, const FString& D
 
 void UEditorEngine::CancelRequestPlaySession()
 {
+	FEditorDelegates::CancelPIE.Broadcast();
 	PlaySessionRequest.Reset();
 	PlayInEditorSessionInfo.Reset();
 }
diff --git a/Engine/Source/Editor/UnrealEd/Public/Editor.h b/Engine/Source/Editor/UnrealEd/Public/Editor.h
--- a/Engine/Source/Editor/UnrealEd/Public/Editor.h
+++ b/Engine/Source/Editor/UnrealEd/Public/Editor.h
@@ -209,6 +209,8 @@ struct UNREALED_API FEditorDelegates
 	static FOnPIEEvent OnPreSwitchBeginPIEAndSIE;
 	/** Sent after the user switches between from PIE to SIE, or vice-versa.  Passes in whether we are currently in SIE */
 	static FOnPIEEvent OnSwitchBeginPIEAndSIE;
+	/** Sent when a PIE session is cancelled */
+	static FSimpleMulticastDelegate CancelPIE;
 	/** Sent when PC local play session is starting */
 	static FOnStandaloneLocalPlayEvent BeginStandaloneLocalPlay;
 	/** Within a property window, the currently selected item was changed.*/
diff --git a/Engine/Source/Runtime/Engine/Private/WorldPartition/WorldPartition.cpp b/Engine/Source/Runtime/Engine/Private/WorldPartition/WorldPartition.cpp
--- a/Engine/Source/Runtime/Engine/Private/WorldPartition/WorldPartition.cpp
+++ b/Engine/Source/Runtime/Engine/Private/WorldPartition/WorldPartition.cpp
@@ -207,6 +207,14 @@ void UWorldPartition::OnBeginPlay(EWorldPartitionStreamingMode Mode)
 	RuntimeHash->OnBeginPlay(Mode);
 }
 
+void UWorldPartition::OnCancelPIE()
+{
+	// No check here since CancelPIE can be called after PrePIEEnded
+	bIsPIE = false;
+	// Call OnEndPlay here since EndPlayMapDelegate is not called when cancelling PIE
+	OnEndPlay();
+}
+
 void UWorldPartition::OnEndPlay()
 {
 	check(IsMainWorldPartition());
@@ -534,6 +542,7 @@ void UWorldPartition::RegisterDelegates()
 	{
 		FEditorDelegates::PreBeginPIE.AddUObject(this, &UWorldPartition::OnPreBeginPIE);
 		FEditorDelegates::PrePIEEnded.AddUObject(this, &UWorldPartition::OnPrePIEEnded);
+		FEditorDelegates::CancelPIE.AddUObject(this, &UWorldPartition::OnCancelPIE);
 		FGameDelegates::Get().GetEndPlayMapDelegate().AddUObject(this, &UWorldPartition::OnEndPlay);
 	}
 }
@@ -545,6 +554,7 @@ void UWorldPartition::UnregisterDelegates()
 	{
 		FEditorDelegates::PreBeginPIE.RemoveAll(this);
 		FEditorDelegates::PrePIEEnded.RemoveAll(this);
+		FEditorDelegates::CancelPIE.RemoveAll(this);
 		FGameDelegates::Get().GetEndPlayMapDelegate().RemoveAll(this);
 	}
 }
diff --git a/Engine/Source/Runtime/Engine/Public/WorldPartition/WorldPartition.h b/Engine/Source/Runtime/Engine/Public/WorldPartition/WorldPartition.h
--- a/Engine/Source/Runtime/Engine/Public/WorldPartition/WorldPartition.h
+++ b/Engine/Source/Runtime/Engine/Public/WorldPartition/WorldPartition.h
@@ -96,6 +96,7 @@ private:
 	// PIE/Game Methods
 	void OnPreBeginPIE(bool bStartSimulate);
 	void OnPrePIEEnded(bool bWasSimulatingInEditor);
+	void OnCancelPIE();
 	void OnBeginPlay(EWorldPartitionStreamingMode Mode);
 	void OnEndPlay();
 
