Allow wp.Runtime.OverrideRuntimeSpatialHashLoadingRange command to override loading range with a value of 0 meter. This allows to completely disable runtime streaming of spatial grid cells.

diff --git a/Engine/Source/Runtime/Engine/Private/WorldPartition/WorldPartitionRuntimeSpatialHash.cpp b/Engine/Source/Runtime/Engine/Private/WorldPartition/WorldPartitionRuntimeSpatialHash.cpp
--- a/Engine/Source/Runtime/Engine/Private/WorldPartition/WorldPartitionRuntimeSpatialHash.cpp
+++ b/Engine/Source/Runtime/Engine/Private/WorldPartition/WorldPartitionRuntimeSpatialHash.cpp
@@ -82,7 +82,7 @@ FSpatialHashStreamingGrid::FSpatialHashStreamingGrid()
 	, LoadingRange(0.0f)
 	, DebugColor(ForceInitToZero)
 	, WorldBounds(ForceInitToZero)
-	, OverrideLoadingRange(0.f)
+	, OverrideLoadingRange(-1.f)
 	, GridHelper(nullptr)
 {
 }
@@ -984,25 +984,22 @@ FAutoConsoleCommand UWorldPartitionRuntimeSpatialHash::OverrideLoadingRangeComma
 	{
 		FString ArgString = FString::Join(Args, TEXT(" "));
 		int32 GridIndex = 0;
-		float OverrideLoadingRange = 0.f;
+		float OverrideLoadingRange = -1.f;
 		FParse::Value(*ArgString, TEXT("grid="), GridIndex);
 		FParse::Value(*ArgString, TEXT("range="), OverrideLoadingRange);
 
-		if (OverrideLoadingRange > 0.f)
+		for (const FWorldContext& Context : GEngine->GetWorldContexts())
 		{
-			for (const FWorldContext& Context : GEngine->GetWorldContexts())
+			UWorld* World = Context.World();
+			if (World && World->IsGameWorld())
 			{
-				UWorld* World = Context.World();
-				if (World && World->IsGameWorld())
+				if (UWorldPartition* WorldPartition = World->GetWorldPartition())
 				{
-					if (UWorldPartition* WorldPartition = World->GetWorldPartition())
+					if (UWorldPartitionRuntimeSpatialHash* RuntimeSpatialHash = Cast<UWorldPartitionRuntimeSpatialHash>(WorldPartition->RuntimeHash))
 					{
-						if (UWorldPartitionRuntimeSpatialHash* RuntimeSpatialHash = Cast<UWorldPartitionRuntimeSpatialHash>(WorldPartition->RuntimeHash))
+						if (RuntimeSpatialHash->StreamingGrids.IsValidIndex(GridIndex))
 						{
-							if (RuntimeSpatialHash->StreamingGrids.IsValidIndex(GridIndex))
-							{
-								RuntimeSpatialHash->StreamingGrids[GridIndex].OverrideLoadingRange = OverrideLoadingRange;
-							}
+							RuntimeSpatialHash->StreamingGrids[GridIndex].OverrideLoadingRange = OverrideLoadingRange;
 						}
 					}
 				}
diff --git a/Engine/Source/Runtime/Engine/Public/WorldPartition/WorldPartitionRuntimeSpatialHash.h b/Engine/Source/Runtime/Engine/Public/WorldPartition/WorldPartitionRuntimeSpatialHash.h
--- a/Engine/Source/Runtime/Engine/Public/WorldPartition/WorldPartitionRuntimeSpatialHash.h
+++ b/Engine/Source/Runtime/Engine/Public/WorldPartition/WorldPartitionRuntimeSpatialHash.h
@@ -116,7 +116,7 @@ struct FSpatialHashStreamingGrid
 
 	float GetLoadingRange() const
 	{
-		if (OverrideLoadingRange > 0.f)
+		if (OverrideLoadingRange >= 0.f)
 		{
 			return OverrideLoadingRange;
 		}
