WorldPartition: Avoid calling AActor::Rename when Outer doesn't change as this as the side effect of renaming the actor and break softobjectrefs. (Was happening when cooking AlwaysLoaded actors to PersistentLevel)

diff --git a/Engine/Source/Runtime/Engine/Private/WorldPartition/WorldPartitionLevelHelper.cpp b/Engine/Source/Runtime/Engine/Private/WorldPartition/WorldPartitionLevelHelper.cpp
--- a/Engine/Source/Runtime/Engine/Private/WorldPartition/WorldPartitionLevelHelper.cpp
+++ b/Engine/Source/Runtime/Engine/Private/WorldPartition/WorldPartitionLevelHelper.cpp
@@ -62,7 +62,14 @@ void FWorldPartitionLevelHelper::MoveExternalActorsToLevel(const TArray<FWorldPa
 
 			const bool bSameOuter = (InLevel == Actor->GetOuter());
 			Actor->SetPackageExternal(false, false);
-			Actor->Rename(nullptr, InLevel, REN_ForceNoResetLoaders);
+						
+			// Avoid calling Rename on the actor if it's already outered to InLevel as this will cause it's name to be changed. 
+			// (UObject::Rename doesn't check if Rename is being called with existing outer and assigns new name)
+			if (!bSameOuter)
+			{
+				Actor->Rename(nullptr, InLevel, REN_ForceNoResetLoaders);
+			}
+			
 			check(Actor->GetPackage() == LevelPackage);
 			if (bSameOuter && !InLevel->Actors.Contains(Actor))
 			{
