Fix logspam for audio streamcaching cache misses. (made verbose and made log more clear) and increased default cache size.

diff --git a/Engine/Source/Runtime/AudioPlatformConfiguration/Public/AudioCompressionSettings.h b/Engine/Source/Runtime/AudioPlatformConfiguration/Public/AudioCompressionSettings.h
--- a/Engine/Source/Runtime/AudioPlatformConfiguration/Public/AudioCompressionSettings.h
+++ b/Engine/Source/Runtime/AudioPlatformConfiguration/Public/AudioCompressionSettings.h
@@ -22,7 +22,7 @@ enum class ESoundwaveSampleRateSettings : uint8
 /************************************************************************/
 struct FAudioStreamCachingSettings
 {
-	static constexpr int32 DefaultCacheSize = 32 * 1024;
+	static constexpr int32 DefaultCacheSize = 64 * 1024;
 
 	// Target memory usage, in kilobytes.
 	// In the future settings for the cache can be more complex, but for now
diff --git a/Engine/Source/Runtime/Engine/Private/AudioCompressionSettingsUtils.cpp b/Engine/Source/Runtime/Engine/Private/AudioCompressionSettingsUtils.cpp
--- a/Engine/Source/Runtime/Engine/Private/AudioCompressionSettingsUtils.cpp
+++ b/Engine/Source/Runtime/Engine/Private/AudioCompressionSettingsUtils.cpp
@@ -115,7 +115,7 @@ void CacheAudioCookOverrides(FPlatformAudioCookOverrides& OutOverrides, const TC
 	if (OutOverrides.bUseStreamCaching)
 	{
 		// Cache size:
-		const int32 DefaultCacheSize = 32 * 1024;
+		const int32 DefaultCacheSize = 64 * 1024;
 		int32 RetrievedCacheSize = DefaultCacheSize;
 		int32 RetrievedChunkSizeOverride = INDEX_NONE;
 
diff --git a/Engine/Source/Runtime/Engine/Private/AudioStreamingCache.cpp b/Engine/Source/Runtime/Engine/Private/AudioStreamingCache.cpp
--- a/Engine/Source/Runtime/Engine/Private/AudioStreamingCache.cpp
+++ b/Engine/Source/Runtime/Engine/Private/AudioStreamingCache.cpp
@@ -1221,8 +1221,8 @@ FAudioChunkCache::FCacheElement* FAudioChunkCache::FindElementForKey(const FChun
 				return &CachePool[CacheOffset];
 			}
 
-			UE_LOG(LogAudioStreamCaching, Display, TEXT("Cache Offset [%i] currently stores chunk for Soundwave: %s -- (looking for Soundwave %s)")
-				, CacheOffset, *CachePool[CacheOffset].Key.SoundWaveName.ToString(), *InKey.SoundWaveName.ToString());
+			UE_LOG(LogAudioStreamCaching, Verbose, TEXT("Cache Miss for soundwave: %s. (Cache Lookup ID [%i] currently stores chunk for Soundwave: %s"),
+				*InKey.SoundWaveName.ToString(), CacheOffset, *CachePool[CacheOffset].Key.SoundWaveName.ToString());
 		}
 	}
 
