Fix Editor ES31 assert by binding DefaultColorBuffer if bHasColorVertexData is false

diff --git a/Engine/Source/Runtime/Engine/Private/InstancedStaticMesh.cpp b/Engine/Source/Runtime/Engine/Private/InstancedStaticMesh.cpp
--- a/Engine/Source/Runtime/Engine/Private/InstancedStaticMesh.cpp
+++ b/Engine/Source/Runtime/Engine/Private/InstancedStaticMesh.cpp
@@ -905,9 +905,18 @@ void FInstancedStaticMeshRenderData::InitVertexFactories()
 			{
 				RenderData->VertexBuffers.StaticMeshVertexBuffer.BindLightMapVertexBuffer(&VertexFactory, Data, LightMapCoordinateIndex);
 			}
-			RenderData->VertexBuffers.ColorVertexBuffer.BindColorVertexBuffer(&VertexFactory, Data);
 
+			if (RenderData->bHasColorVertexData)
+			{
+				RenderData->VertexBuffers.ColorVertexBuffer.BindColorVertexBuffer(&VertexFactory, Data);
+			}
+			else
+			{
+				FColorVertexBuffer::BindDefaultColorVertexBuffer(&VertexFactory, Data, FColorVertexBuffer::NullBindStride::FColorSizeForComponentOverride);
+			}
+			
 			check(PerInstanceRenderData);
+			
 			PerInstanceRenderData->InstanceBuffer.BindInstanceVertexBuffer(&VertexFactory, Data);
 
 			VertexFactory.SetData(Data);
