Fix crash when toggling Data Layers: don't preload always loaded actors that belongs to a Data Layer.

diff --git a/Engine/Source/Runtime/Engine/Private/Actor.cpp b/Engine/Source/Runtime/Engine/Private/Actor.cpp
--- a/Engine/Source/Runtime/Engine/Private/Actor.cpp
+++ b/Engine/Source/Runtime/Engine/Private/Actor.cpp
@@ -707,6 +707,17 @@ void AActor::RemoveControllingMatineeActor( AMatineeActor& InMatineeActor )
 
 void AActor::BeginDestroy()
 {
+#if WITH_EDITOR && DO_CHECK
+	if (IsPackageExternal())
+	{
+		ForEachObjectWithPackage(GetPackage(), [this](UObject* Object)
+		{
+			ensureMsgf(!Object->HasAnyFlags(RF_Public | RF_Standalone), TEXT("BeginDestroy was called on external actor %s without properly detaching it's package"), *GetPathName());
+			return true;
+		}, false);
+	}
+#endif
+
 	UnregisterAllComponents();
 	ULevel* Level = GetLevel();
 	if (Level && !Level->HasAnyInternalFlags(EInternalObjectFlags::Unreachable))
diff --git a/Engine/Source/Runtime/Engine/Private/WorldPartition/WorldPartition.cpp b/Engine/Source/Runtime/Engine/Private/WorldPartition/WorldPartition.cpp
--- a/Engine/Source/Runtime/Engine/Private/WorldPartition/WorldPartition.cpp
+++ b/Engine/Source/Runtime/Engine/Private/WorldPartition/WorldPartition.cpp
@@ -332,13 +332,16 @@ void UWorldPartition::Initialize(UWorld* InWorld, const FTransform& InTransform)
 
 	if (bEditorOnly)
 	{
-		// Make sure to preload any AInfo actors first (mainly for ShouldActorBeLoaded as it needs to have access to the WorldDataLayers actor)
-		TArray<UWorldPartitionEditorCell::FActorHandle> AlwaysLoadedActors = EditorHash->GetAlwaysLoadedCell()->Actors.Array();
-		AlwaysLoadedActors.Sort([](const UWorldPartitionEditorCell::FActorHandle& A, const UWorldPartitionEditorCell::FActorHandle& B) { return A->GetActorClass()->IsChildOf<AInfo>(); });
-			
-		TArray<FWorldPartitionReference> LoadedActors;
-		LoadedActors.Reserve(AlwaysLoadedActors.Num());
-		Algo::Transform(AlwaysLoadedActors, LoadedActors, [](const UWorldPartitionEditorCell::FActorHandle& ActorHandle) { return ActorHandle.Handle; });
+		// Make sure to preload only AWorldDataLayers actor first (ShouldActorBeLoaded requires it)
+		FWorldPartitionReference WorldDataLayersActor;
+		for (UWorldPartitionEditorCell::FActorHandle& ActorHandle : EditorHash->GetAlwaysLoadedCell()->Actors)
+		{
+			if (ActorHandle->GetActorClass()->IsChildOf<AWorldDataLayers>())
+			{
+				WorldDataLayersActor = ActorHandle.Handle;
+				break;
+			}
+		}
 
 		// Load the always loaded cell, don't call LoadCells to avoid creating a transaction
 		UpdateLoadingEditorCell(EditorHash->GetAlwaysLoadedCell(), true);
