Only add in the destination level package objects found in the source actor package if it's an external actor.

diff --git a/Engine/Source/Runtime/Engine/Private/WorldPartition/WorldPartitionLevelHelper.cpp b/Engine/Source/Runtime/Engine/Private/WorldPartition/WorldPartitionLevelHelper.cpp
--- a/Engine/Source/Runtime/Engine/Private/WorldPartition/WorldPartitionLevelHelper.cpp
+++ b/Engine/Source/Runtime/Engine/Private/WorldPartition/WorldPartitionLevelHelper.cpp
@@ -58,7 +58,7 @@ void FWorldPartitionLevelHelper::MoveExternalActorsToLevel(const TArray<FWorldPa
 		AActor* Actor = FindObject<AActor>(nullptr, *PackageObjectMapping.LoadedPath.ToString());
 		if (ensure(Actor))
 		{
-			UPackage* ActorPackage = Actor->GetPackage();
+			UPackage* ActorExternalPackage = Actor->IsPackageExternal() ? Actor->GetPackage() : nullptr;
 
 			const bool bSameOuter = (InLevel == Actor->GetOuter());
 			Actor->SetPackageExternal(false, false);
@@ -77,13 +77,17 @@ void FWorldPartitionLevelHelper::MoveExternalActorsToLevel(const TArray<FWorldPa
 			}
 
 			// Include objects found in the source actor package in the destination level package
-			TArray<UObject*> Objects;
-			GetObjectsWithOuter(ActorPackage, Objects);
-			for (UObject* Object : Objects)
+			// @todo_ow: Generalize this case and handle properly the case of LevelInstance (LI actor's package is the duplicated LI package, i.e. not an external actor package)
+			if (ActorExternalPackage)
 			{
-				if (Object->GetFName() != NAME_PackageMetaData)
+				TArray<UObject*> Objects;
+				GetObjectsWithOuter(ActorExternalPackage, Objects);
+				for (UObject* Object : Objects)
 				{
-					Object->Rename(nullptr, LevelPackage, REN_ForceNoResetLoaders);
+					if (Object->GetFName() != NAME_PackageMetaData)
+					{
+						Object->Rename(nullptr, LevelPackage, REN_ForceNoResetLoaders);
+					}
 				}
 			}
 		}
