Packaging fails with assertion 'BulkData.cpp [Line: 1099]'

Must not include nested subobjects when moving objects from the actors packages to the (cell) level package

diff --git a/Engine/Source/Runtime/Engine/Private/WorldPartition/WorldPartitionLevelHelper.cpp b/Engine/Source/Runtime/Engine/Private/WorldPartition/WorldPartitionLevelHelper.cpp
--- a/Engine/Source/Runtime/Engine/Private/WorldPartition/WorldPartitionLevelHelper.cpp
+++ b/Engine/Source/Runtime/Engine/Private/WorldPartition/WorldPartitionLevelHelper.cpp
@@ -81,7 +81,8 @@ void FWorldPartitionLevelHelper::MoveExternalActorsToLevel(const TArray<FWorldPa
 			if (ActorExternalPackage)
 			{
 				TArray<UObject*> Objects;
-				GetObjectsWithOuter(ActorExternalPackage, Objects);
+				const bool bIncludeNestedSubobjects = false;
+				GetObjectsWithOuter(ActorExternalPackage, Objects, bIncludeNestedSubobjects);
 				for (UObject* Object : Objects)
 				{
 					if (Object->GetFName() != NAME_PackageMetaData)
