Warn when a material function instance isn't saved with a valid GUID.

diff --git a/Engine/Source/Runtime/Engine/Classes/Materials/MaterialFunctionInstance.h b/Engine/Source/Runtime/Engine/Classes/Materials/MaterialFunctionInstance.h
--- a/Engine/Source/Runtime/Engine/Classes/Materials/MaterialFunctionInstance.h
+++ b/Engine/Source/Runtime/Engine/Classes/Materials/MaterialFunctionInstance.h
@@ -77,6 +77,7 @@ class UMaterialFunctionInstance : public UMaterialFunctionInterface
 #endif
 	virtual bool ValidateFunctionUsage(class FMaterialCompiler* Compiler, const FFunctionExpressionOutput& Output) override;
 
+	virtual void Serialize(FArchive& Ar) override;
 	virtual void PostLoad() override;
 #if WITH_EDITOR
 	virtual int32 Compile(class FMaterialCompiler* Compiler, const struct FFunctionExpressionOutput& Output) override;
diff --git a/Engine/Source/Runtime/Engine/Private/Materials/MaterialExpressions.cpp b/Engine/Source/Runtime/Engine/Private/Materials/MaterialExpressions.cpp
--- a/Engine/Source/Runtime/Engine/Private/Materials/MaterialExpressions.cpp
+++ b/Engine/Source/Runtime/Engine/Private/Materials/MaterialExpressions.cpp
@@ -13232,6 +13232,28 @@ bool UMaterialFunctionInstance::ValidateFunctionUsage(class FMaterialCompiler* C
 	return Parent ? Parent->ValidateFunctionUsage(Compiler, Output) : false;
 }
 
+void UMaterialFunctionInstance::Serialize(FArchive& Ar)
+{
+	FGuid OldStateId;
+	// catch assets saved without proper StateId
+	if (Ar.IsLoading() && Ar.IsPersistent())
+	{
+		OldStateId = StateId;
+		StateId.Invalidate();
+	}
+
+	Super::Serialize(Ar);
+
+	if (Ar.IsLoading() && Ar.IsPersistent())
+	{
+		if (!StateId.IsValid())
+		{
+			UE_LOG(LogMaterial, Warning, TEXT("%s was saved without a valid StateId (old asset?). This will cause all materials using it to recompile their shaders on each load - please resave to fix."), *GetName());
+			StateId = OldStateId;
+		}
+	}
+}
+
 void UMaterialFunctionInstance::PostLoad()
 {
 	Super::PostLoad();
