Avoid overhead from access tracking on wrapped object pointers during high frequency operations related to materials at cook time.
Also go back to using Cast/CastChecked/ExactCast as these have now been made more efficient for use with wrapped object pointers.

diff --git a/Engine/Source/Developer/AssetTools/Private/AssetVtConversion.cpp b/Engine/Source/Developer/AssetTools/Private/AssetVtConversion.cpp
--- a/Engine/Source/Developer/AssetTools/Private/AssetVtConversion.cpp
+++ b/Engine/Source/Developer/AssetTools/Private/AssetVtConversion.cpp
@@ -688,11 +688,11 @@ void FVTConversionWorker::DoConvert()
 
 			bool FuncModified = false;
 			const TArray<TObjectPtr<UMaterialExpression>> *Expressions = Func->GetFunctionExpressions();
-			for (TObjectPtr<UMaterialExpression> Expr : *Expressions)
+			for (const TObjectPtr<UMaterialExpression>& Expr : *Expressions)
 			{
-				if (Expr && Expr.IsA<UMaterialExpressionTextureBase>())
+				UMaterialExpressionTextureBase *TexExpr = Cast<UMaterialExpressionTextureBase>(Expr);
+				if (TexExpr)
 				{
-					UMaterialExpressionTextureBase* TexExpr = (UMaterialExpressionTextureBase*)Expr.Get();
 					if (Textures.Contains(TexExpr->Texture))
 					{
 						UE_LOG(LogVirtualTextureConversion, Display, TEXT("Adjusting sampler %s."), *TexExpr->GetName());
diff --git a/Engine/Source/Editor/VirtualTexturingEditor/Private/RuntimeVirtualTextureAssetTypeActions.cpp b/Engine/Source/Editor/VirtualTexturingEditor/Private/RuntimeVirtualTextureAssetTypeActions.cpp
--- a/Engine/Source/Editor/VirtualTexturingEditor/Private/RuntimeVirtualTextureAssetTypeActions.cpp
+++ b/Engine/Source/Editor/VirtualTexturingEditor/Private/RuntimeVirtualTextureAssetTypeActions.cpp
@@ -154,11 +154,11 @@ namespace
 
 				bool bFunctionModified = false;
 				const TArray<TObjectPtr<UMaterialExpression>> *Expressions = Function->GetFunctionExpressions();
-				for (TObjectPtr<UMaterialExpression> Expression : *Expressions)
+				for (const TObjectPtr<UMaterialExpression>& Expression : *Expressions)
 				{
-					if (Expression && Expression.IsA<UMaterialExpressionRuntimeVirtualTextureSample>())
+					UMaterialExpressionRuntimeVirtualTextureSample* RVTSampleExpression = Cast<UMaterialExpressionRuntimeVirtualTextureSample>(Expression);
+					if (RVTSampleExpression)
 					{
-						UMaterialExpressionRuntimeVirtualTextureSample* RVTSampleExpression = (UMaterialExpressionRuntimeVirtualTextureSample*)Expression.Get();
 						if (RuntimeVirtualTexture == RVTSampleExpression->VirtualTexture)
 						{
 							if (RVTSampleExpression->InitVirtualTextureDependentSettings())
diff --git a/Engine/Source/Runtime/Engine/Classes/Materials/Material.h b/Engine/Source/Runtime/Engine/Classes/Materials/Material.h
--- a/Engine/Source/Runtime/Engine/Classes/Materials/Material.h
+++ b/Engine/Source/Runtime/Engine/Classes/Materials/Material.h
@@ -1261,7 +1261,7 @@ public:
 	template<typename ExpressionType>
 	void GetAllParameterInfo(TArray<FMaterialParameterInfo>& OutParameterInfo, TArray<FGuid>& OutParameterIds, const TArray<FStaticMaterialLayersParameter>* MaterialLayersParameters = nullptr) const
 	{
-		for (class UMaterialExpression* Expression : Expressions)
+		for (const TObjectPtr<UMaterialExpression>& Expression : Expressions)
 		{
 			FMaterialParameterInfo BaseParameterInfo;
 			BaseParameterInfo.Association = EMaterialParameterAssociation::GlobalParameter;
@@ -1924,7 +1924,7 @@ private:
 
 		for (TObjectPtr<UMaterialExpression>& Expression : Expressions)
 		{
-			if (TrySetParameterValue((Expression && Expression.IsA<ParameterType>()) ? (ParameterType*)Expression.Get() : nullptr))
+			if (TrySetParameterValue(Cast<ParameterType>(Expression)))
 			{
 				return true;
 			}
@@ -1941,9 +1941,9 @@ private:
 						const TArray<TObjectPtr<UMaterialExpression>>* ExpressionPtr = Function->GetFunctionExpressions();
 						if (ExpressionPtr)
 						{
-							for (TObjectPtr<UMaterialExpression> FunctionExpression : *ExpressionPtr)
+							for (const TObjectPtr<UMaterialExpression>& FunctionExpression : *ExpressionPtr)
 							{
-								if (TrySetParameterValue((FunctionExpression && FunctionExpression.IsA<ParameterType>()) ? (ParameterType*)FunctionExpression.Get() : nullptr))
+								if (TrySetParameterValue(Cast<ParameterType>(FunctionExpression)))
 								{
 									return true;
 								}
diff --git a/Engine/Source/Runtime/Engine/Classes/Materials/MaterialFunctionInterface.h b/Engine/Source/Runtime/Engine/Classes/Materials/MaterialFunctionInterface.h
--- a/Engine/Source/Runtime/Engine/Classes/Materials/MaterialFunctionInterface.h
+++ b/Engine/Source/Runtime/Engine/Classes/Materials/MaterialFunctionInterface.h
@@ -162,7 +162,7 @@ public:
 		if (const UMaterialFunctionInterface* ParameterFunction = GetBaseFunction())
 		{
 			const UClass* TargetClass = UMaterialExpressionMaterialFunctionCall::StaticClass();
-			for (TObjectPtr<UMaterialExpression> Expression : *ParameterFunction->GetFunctionExpressions())
+			for (const TObjectPtr<UMaterialExpression>& Expression : *ParameterFunction->GetFunctionExpressions())
 			{
 				if (const UMaterialExpressionMaterialFunctionCall* FunctionExpression = (Expression && Expression.IsA(TargetClass)) ? (const UMaterialExpressionMaterialFunctionCall*)Expression.Get() : nullptr)
 				{
@@ -171,9 +171,9 @@ public:
 						FunctionExpression->MaterialFunction->GetAllParameterInfo<const ExpressionType>(OutParameterInfo, OutParameterIds, InBaseParameterInfo);
 					}
 				}
-				else if (Expression && Expression.IsA<const ExpressionType>())
+				else if (const ExpressionType* ParameterExpression = Cast<const ExpressionType>(Expression))
 				{
-					((const ExpressionType*)Expression.Get())->GetAllParameterInfo(OutParameterInfo, OutParameterIds, InBaseParameterInfo);
+					ParameterExpression->GetAllParameterInfo(OutParameterInfo, OutParameterIds, InBaseParameterInfo);
 				}
 			}
 
diff --git a/Engine/Source/Runtime/Engine/Private/Materials/MaterialExpressions.cpp b/Engine/Source/Runtime/Engine/Private/Materials/MaterialExpressions.cpp
--- a/Engine/Source/Runtime/Engine/Private/Materials/MaterialExpressions.cpp
+++ b/Engine/Source/Runtime/Engine/Private/Materials/MaterialExpressions.cpp
@@ -12801,11 +12801,10 @@ bool UMaterialFunction::SetVectorParameterValueEditorOnly(FName ParameterName, F
 					const TArray<TObjectPtr<UMaterialExpression>>* ExpressionPtr = Function->GetFunctionExpressions();
 					if (ExpressionPtr)
 					{
-						for (TObjectPtr<UMaterialExpression> FunctionExpression : *ExpressionPtr)
+						for (const TObjectPtr<UMaterialExpression>&FunctionExpression : *ExpressionPtr)
 						{
-							if (FunctionExpression && FunctionExpression.IsA<UMaterialExpressionVectorParameter>())
+							if (UMaterialExpressionVectorParameter* FunctionExpressionParameter = Cast<UMaterialExpressionVectorParameter>(FunctionExpression))
 							{
-								UMaterialExpressionVectorParameter* FunctionExpressionParameter = (UMaterialExpressionVectorParameter*)FunctionExpression.Get();
 								if (FunctionExpressionParameter->SetParameterValue(ParameterName, InValue))
 								{
 									return true;
@@ -12845,11 +12844,10 @@ bool UMaterialFunction::SetScalarParameterValueEditorOnly(FName ParameterName, f
 					const TArray<TObjectPtr<UMaterialExpression>>* ExpressionPtr = Function->GetFunctionExpressions();
 					if (ExpressionPtr)
 					{
-						for (TObjectPtr<UMaterialExpression> FunctionExpression : *ExpressionPtr)
+						for (const TObjectPtr<UMaterialExpression>&FunctionExpression : *ExpressionPtr)
 						{
-							if (FunctionExpression && FunctionExpression.IsA<UMaterialExpressionScalarParameter>())
+							if (UMaterialExpressionScalarParameter* FunctionExpressionParameter = Cast<UMaterialExpressionScalarParameter>(FunctionExpression))
 							{
-								UMaterialExpressionScalarParameter* FunctionExpressionParameter = (UMaterialExpressionScalarParameter*)FunctionExpression.Get();
 								if (FunctionExpressionParameter->SetParameterValue(ParameterName, InValue))
 								{
 									return true;
@@ -12889,11 +12887,10 @@ bool UMaterialFunction::SetTextureParameterValueEditorOnly(FName ParameterName,
 					const TArray<TObjectPtr<UMaterialExpression>>* ExpressionPtr = Function->GetFunctionExpressions();
 					if (ExpressionPtr)
 					{
-						for (TObjectPtr<UMaterialExpression> FunctionExpression : *ExpressionPtr)
+						for (const TObjectPtr<UMaterialExpression>&FunctionExpression : *ExpressionPtr)
 						{
-							if (FunctionExpression && FunctionExpression.IsA<UMaterialExpressionTextureSampleParameter>())
+							if (UMaterialExpressionTextureSampleParameter* FunctionExpressionParameter = Cast<UMaterialExpressionTextureSampleParameter>(FunctionExpression))
 							{
-								UMaterialExpressionTextureSampleParameter* FunctionExpressionParameter = (UMaterialExpressionTextureSampleParameter*)FunctionExpression.Get();
 								if (FunctionExpressionParameter->SetParameterValue(ParameterName, InValue))
 								{
 									return true;
@@ -12933,11 +12930,10 @@ bool UMaterialFunction::SetRuntimeVirtualTextureParameterValueEditorOnly(FName P
 					const TArray<TObjectPtr<UMaterialExpression>>* ExpressionPtr = Function->GetFunctionExpressions();
 					if (ExpressionPtr)
 					{
-						for (TObjectPtr<UMaterialExpression> FunctionExpression : *ExpressionPtr)
+						for (const TObjectPtr<UMaterialExpression>&FunctionExpression : *ExpressionPtr)
 						{
-							if (FunctionExpression && FunctionExpression.IsA<UMaterialExpressionRuntimeVirtualTextureSampleParameter>())
+							if (UMaterialExpressionRuntimeVirtualTextureSampleParameter* FunctionExpressionParameter = Cast<UMaterialExpressionRuntimeVirtualTextureSampleParameter>(FunctionExpression))
 							{
-								UMaterialExpressionRuntimeVirtualTextureSampleParameter* FunctionExpressionParameter = (UMaterialExpressionRuntimeVirtualTextureSampleParameter*)FunctionExpression.Get();
 								if (FunctionExpressionParameter->SetParameterValue(ParameterName, InValue))
 								{
 									return true;
@@ -12977,7 +12973,7 @@ bool UMaterialFunction::SetFontParameterValueEditorOnly(FName ParameterName, cla
 					const TArray<TObjectPtr<UMaterialExpression>>* ExpressionPtr = Function->GetFunctionExpressions();
 					if (ExpressionPtr)
 					{
-						for (TObjectPtr<UMaterialExpression> FunctionExpression : *ExpressionPtr)
+						for (const TObjectPtr<UMaterialExpression>&FunctionExpression : *ExpressionPtr)
 						{
 							if (FunctionExpression && FunctionExpression.IsA<UMaterialExpressionFontSampleParameter>())
 							{
@@ -13021,7 +13017,7 @@ bool UMaterialFunction::SetStaticSwitchParameterValueEditorOnly(FName ParameterN
 					const TArray<TObjectPtr<UMaterialExpression>>* ExpressionPtr = Function->GetFunctionExpressions();
 					if (ExpressionPtr)
 					{
-						for (TObjectPtr<UMaterialExpression> FunctionExpression : *ExpressionPtr)
+						for (const TObjectPtr<UMaterialExpression>&FunctionExpression : *ExpressionPtr)
 						{
 							if (FunctionExpression && FunctionExpression.IsA<UMaterialExpressionStaticSwitchParameter>())
 							{
@@ -13065,7 +13061,7 @@ bool UMaterialFunction::SetStaticComponentMaskParameterValueEditorOnly(FName Par
 					const TArray<TObjectPtr<UMaterialExpression>>* ExpressionPtr = Function->GetFunctionExpressions();
 					if (ExpressionPtr)
 					{
-						for (TObjectPtr<UMaterialExpression> FunctionExpression : *ExpressionPtr)
+						for (const TObjectPtr<UMaterialExpression>&FunctionExpression : *ExpressionPtr)
 						{
 							if (FunctionExpression && FunctionExpression.IsA<UMaterialExpressionStaticComponentMaskParameter>())
 							{
diff --git a/Engine/Source/Runtime/Engine/Private/Materials/MaterialInstance.cpp b/Engine/Source/Runtime/Engine/Private/Materials/MaterialInstance.cpp
--- a/Engine/Source/Runtime/Engine/Private/Materials/MaterialInstance.cpp
+++ b/Engine/Source/Runtime/Engine/Private/Materials/MaterialInstance.cpp
@@ -5079,31 +5079,23 @@ void FindCollectionExpressionRecursive(TArray<FGuid>& OutGuidList, const TArray<
 {
 	for (int32 ExpressionIndex = 0; ExpressionIndex < InMaterialExpression.Num(); ExpressionIndex++)
 	{
-		if (!InMaterialExpression[ExpressionIndex])
+		if (const UMaterialExpressionCollectionParameter* CollectionPtr = Cast<UMaterialExpressionCollectionParameter>(InMaterialExpression[ExpressionIndex]))
 		{
-			continue;
-		}
-
-		if (InMaterialExpression[ExpressionIndex].IsA<UMaterialExpressionCollectionParameter>())
-		{
-			const UMaterialExpressionCollectionParameter* CollectionPtr = (UMaterialExpressionCollectionParameter*)InMaterialExpression[ExpressionIndex].Get();
 			if (CollectionPtr->Collection)
 			{
 				OutGuidList.Add(CollectionPtr->Collection->StateId);
 			}
 			return;
 		}
-		else if (InMaterialExpression[ExpressionIndex].IsA<UMaterialExpressionMaterialFunctionCall>())
+		else if (const UMaterialExpressionMaterialFunctionCall* MaterialFunctionCall = Cast<UMaterialExpressionMaterialFunctionCall>(InMaterialExpression[ExpressionIndex]))
 		{
-			const UMaterialExpressionMaterialFunctionCall* MaterialFunctionCall = (UMaterialExpressionMaterialFunctionCall*)InMaterialExpression[ExpressionIndex].Get();
-			if (const TArray<TObjectPtr<UMaterialExpression>>* FunctionExpressions = MaterialFunctionCall->MaterialFunction->GetFunctionExpressions())
+			if (const TArray<TObjectPtr<UMaterialExpression>>* FunctionExpressions = MaterialFunctionCall->MaterialFunction ? MaterialFunctionCall->MaterialFunction->GetFunctionExpressions() : nullptr)
 			{
 				FindCollectionExpressionRecursive(OutGuidList, *FunctionExpressions);
 			}
 		}
-		else if (InMaterialExpression[ExpressionIndex].IsA<UMaterialExpressionMaterialAttributeLayers>())
+		else if (const UMaterialExpressionMaterialAttributeLayers* MaterialLayers = Cast<UMaterialExpressionMaterialAttributeLayers>(InMaterialExpression[ExpressionIndex]))
 		{
-			const UMaterialExpressionMaterialAttributeLayers* MaterialLayers = (UMaterialExpressionMaterialAttributeLayers*)InMaterialExpression[ExpressionIndex].Get();
 			const TArray<UMaterialFunctionInterface*>& Layers = MaterialLayers->GetLayers();
 			const TArray<UMaterialFunctionInterface*>& Blends = MaterialLayers->GetBlends();
 
diff --git a/Engine/Source/Runtime/Engine/Private/Materials/MaterialStaticParameterValueResolver.h b/Engine/Source/Runtime/Engine/Private/Materials/MaterialStaticParameterValueResolver.h
--- a/Engine/Source/Runtime/Engine/Private/Materials/MaterialStaticParameterValueResolver.h
+++ b/Engine/Source/Runtime/Engine/Private/Materials/MaterialStaticParameterValueResolver.h
@@ -72,7 +72,7 @@ public:
 				auto GetExpressionParameterByNamePredicate =
 					[this, &PendingParametersForFunction, &ParameterIndicesForFunction, &ParametersAndOwners, TargetClass](UMaterialFunctionInterface* InFunction) -> bool
 				{
-					for (TObjectPtr<UMaterialExpression> FunctionExpression : *InFunction->GetFunctionExpressions())
+					for (const TObjectPtr<UMaterialExpression>& FunctionExpression : *InFunction->GetFunctionExpressions())
 					{
 						if (ParameterType* ExpressionParameter = (FunctionExpression && FunctionExpression.IsA(TargetClass)) ? (ParameterType*)FunctionExpression.Get() : nullptr)
 						{
