Backout 15900716 and don't try to get the file timestamp for in-memory packages in UWorld::GetAssetRegistryTags.

diff --git a/Engine/Source/Runtime/AssetRegistry/Private/AssetRegistry.cpp b/Engine/Source/Runtime/AssetRegistry/Private/AssetRegistry.cpp
--- a/Engine/Source/Runtime/AssetRegistry/Private/AssetRegistry.cpp
+++ b/Engine/Source/Runtime/AssetRegistry/Private/AssetRegistry.cpp
@@ -1045,12 +1045,6 @@ bool UAssetRegistryImpl::EnumerateAssets(const FARCompiledFilter& InFilter, TFun
 				// Package name
 				const FName PackageName = InMemoryPackage->GetFName();
 
-				// Skip assets from memory mounted packages (/Memory/...)
-				if (FPackageName::IsMemoryPackage(PackageName.ToString()))
-				{
-					return;
-				}
-
 				PackagesToSkip.Add(PackageName);
 
 				if (InFilter.PackageNames.Num() && !InFilter.PackageNames.Contains(PackageName))
diff --git a/Engine/Source/Runtime/Engine/Private/World.cpp b/Engine/Source/Runtime/Engine/Private/World.cpp
--- a/Engine/Source/Runtime/Engine/Private/World.cpp
+++ b/Engine/Source/Runtime/Engine/Private/World.cpp
@@ -8069,11 +8069,13 @@ void UWorld::GetAssetRegistryTags(TArray<FAssetRegistryTag>& OutTags) const
 	}
 
 	// Get the full file path with extension
-	const FString FullFilePath = FPackageName::LongPackageNameToFilename(GetOutermost()->GetName(), FPackageName::GetMapPackageExtension());
-
-	// Save/Display the modify date. File size is handled generically for all packages
-	FDateTime AssetDateModified = IFileManager::Get().GetTimeStamp(*FullFilePath);
-	OutTags.Add(FAssetRegistryTag("DateModified", AssetDateModified.ToString(), FAssetRegistryTag::TT_Chronological, FAssetRegistryTag::TD_Date));
+	FString FullFilePath;
+	if (FPackageName::TryConvertLongPackageNameToFilename(GetOutermost()->GetName(), FullFilePath, FPackageName::GetMapPackageExtension()))
+	{
+		// Save/Display the modify date. File size is handled generically for all packages
+		FDateTime AssetDateModified = IFileManager::Get().GetTimeStamp(*FullFilePath);
+		OutTags.Add(FAssetRegistryTag("DateModified", AssetDateModified.ToString(), FAssetRegistryTag::TT_Chronological, FAssetRegistryTag::TD_Date));
+	}
 
 	FWorldDelegates::GetAssetTags.Broadcast(this, OutTags);
 }
