Remove an ensure I added today that is firing in -game, will investigate later.

diff --git a/Engine/Source/Runtime/Engine/Private/Actor.cpp b/Engine/Source/Runtime/Engine/Private/Actor.cpp
--- a/Engine/Source/Runtime/Engine/Private/Actor.cpp
+++ b/Engine/Source/Runtime/Engine/Private/Actor.cpp
@@ -707,17 +707,6 @@ void AActor::RemoveControllingMatineeActor( AMatineeActor& InMatineeActor )
 
 void AActor::BeginDestroy()
 {
-#if WITH_EDITOR && DO_CHECK
-	if (IsPackageExternal())
-	{
-		ForEachObjectWithPackage(GetPackage(), [this](UObject* Object)
-		{
-			ensureMsgf(!Object->HasAnyFlags(RF_Public | RF_Standalone), TEXT("BeginDestroy was called on external actor %s without properly detaching it's package"), *GetPathName());
-			return true;
-		}, false);
-	}
-#endif
-
 	UnregisterAllComponents();
 	ULevel* Level = GetLevel();
 	if (Level && !Level->HasAnyInternalFlags(EInternalObjectFlags::Unreachable))
