Moved GPUSkinCache GPU stat into FGPUSkinCache::DoDispatch

diff --git a/Engine/Source/Runtime/Engine/Private/GPUSkinCache.cpp b/Engine/Source/Runtime/Engine/Private/GPUSkinCache.cpp
--- a/Engine/Source/Runtime/Engine/Private/GPUSkinCache.cpp
+++ b/Engine/Source/Runtime/Engine/Private/GPUSkinCache.cpp
@@ -1279,8 +1279,12 @@ FGPUSkinCache::FRWBuffersAllocation* FGPUSkinCache::TryAllocBuffer(uint32 NumVer
 	return NewAllocation;
 }
 
+DECLARE_GPU_STAT(GPUSkinCache);
+
 void FGPUSkinCache::DoDispatch(FRHICommandListImmediate& RHICmdList)
 {
+	SCOPED_GPU_STAT(RHICmdList, GPUSkinCache);
+
 	int32 BatchCount = BatchDispatches.Num();
 	INC_DWORD_STAT_BY(STAT_GPUSkinCache_TotalNumChunks, BatchCount);
 
@@ -1321,6 +1325,8 @@ void FGPUSkinCache::DoDispatch(FRHICommandListImmediate& RHICmdList)
 
 void FGPUSkinCache::DoDispatch(FRHICommandListImmediate& RHICmdList, FGPUSkinCacheEntry* SkinCacheEntry, int32 Section, int32 RevisionNumber)
 {
+	SCOPED_GPU_STAT(RHICmdList, GPUSkinCache);
+
 	INC_DWORD_STAT(STAT_GPUSkinCache_TotalNumChunks);
 	PrepareUpdateSkinning(SkinCacheEntry, Section, RevisionNumber, nullptr);
 	DispatchUpdateSkinning(RHICmdList, SkinCacheEntry, Section, RevisionNumber);
diff --git a/Engine/Source/Runtime/Engine/Private/LevelTick.cpp b/Engine/Source/Runtime/Engine/Private/LevelTick.cpp
--- a/Engine/Source/Runtime/Engine/Private/LevelTick.cpp
+++ b/Engine/Source/Runtime/Engine/Private/LevelTick.cpp
@@ -996,7 +996,6 @@ FSendAllEndOfFrameUpdates* BeginSendEndOfFrameUpdatesDrawEvent(
 }
 
 DECLARE_GPU_STAT(EndOfFrameUpdates);
-DECLARE_GPU_STAT(GPUSkinCache);
 DECLARE_GPU_STAT(GPUSkinCacheRayTracingGeometry);
 DECLARE_GPU_STAT(ComputeFrameworkExecuteBatches);
 void EndSendEndOfFrameUpdatesDrawEvent(FSendAllEndOfFrameUpdates* SendAllEndOfFrameUpdates)
@@ -1009,8 +1008,6 @@ void EndSendEndOfFrameUpdatesDrawEvent(FSendAllEndOfFrameUpdates* SendAllEndOfFr
 		if (SendAllEndOfFrameUpdates->GPUSkinCache)
 		{
 			{
-				SCOPED_GPU_STAT(RHICmdList, GPUSkinCache);
-
 				// Once all the individual components have received their DoDeferredRenderUpdates_Concurrent()
 				// allow the GPU Skin Cache system to update.
 				SendAllEndOfFrameUpdates->GPUSkinCache->EndBatchDispatch(RHICmdList);
