UEngine needs to call Deinitialize on its EngineSubsystemCollection when destroying.
UEditorEngine needs to call GWorld->CleanupWorld (to Deinitialize the UWorld's SubsystemCollection) earlier than FinishDestroy because it is not guaranteed that its FinishDestroy will be called before GWorld's FinishDestroy.
Revert the hack fix to GersterWaterWaveSubsystem; FGerstnerWaterWaveViewExtension needs to remove itself from the GerstnerWaterWaveSubsystem; if it can no longer find the subsystem because GEngine is null, failing to deregister will possibly cause problems elsewhere. The GEngine=null case should no longer occur due to the other two fixes.

diff --git a/Engine/Plugins/Experimental/Water/Source/Runtime/Private/GerstnerWaterWaveViewExtension.cpp b/Engine/Plugins/Experimental/Water/Source/Runtime/Private/GerstnerWaterWaveViewExtension.cpp
--- a/Engine/Plugins/Experimental/Water/Source/Runtime/Private/GerstnerWaterWaveViewExtension.cpp
+++ b/Engine/Plugins/Experimental/Water/Source/Runtime/Private/GerstnerWaterWaveViewExtension.cpp
@@ -15,7 +15,6 @@ FGerstnerWaterWaveViewExtension::~FGerstnerWaterWaveViewExtension()
 
 void FGerstnerWaterWaveViewExtension::Initialize()
 {
-	check(GEngine);
 	if (UGerstnerWaterWaveSubsystem* GerstnerWaterWaveSubsystem = GEngine->GetEngineSubsystem<UGerstnerWaterWaveSubsystem>())
 	{
 		GerstnerWaterWaveSubsystem->Register(this);
@@ -30,12 +29,9 @@ void FGerstnerWaterWaveViewExtension::Deinitialize()
 		[WaveGPUData=WaveGPUData](FRHICommandListImmediate& RHICmdList){}
 	);
 
-	if (GEngine)
+	if (UGerstnerWaterWaveSubsystem* GerstnerWaterWaveSubsystem = GEngine->GetEngineSubsystem<UGerstnerWaterWaveSubsystem>())
 	{
-		if (UGerstnerWaterWaveSubsystem* GerstnerWaterWaveSubsystem = GEngine->GetEngineSubsystem<UGerstnerWaterWaveSubsystem>())
-		{
-			GerstnerWaterWaveSubsystem->Unregister(this);
-		}
+		GerstnerWaterWaveSubsystem->Unregister(this);
 	}
 }
 
diff --git a/Engine/Source/Editor/UnrealEd/Classes/Editor/EditorEngine.h b/Engine/Source/Editor/UnrealEd/Classes/Editor/EditorEngine.h
--- a/Engine/Source/Editor/UnrealEd/Classes/Editor/EditorEngine.h
+++ b/Engine/Source/Editor/UnrealEd/Classes/Editor/EditorEngine.h
@@ -725,6 +725,7 @@ public:
 	void BroadcastObjectReimported(UObject* InObject);
 
 	//~ Begin UObject Interface.
+	virtual void BeginDestroy() override;
 	virtual void FinishDestroy() override;
 	virtual void PostEditChangeProperty(FPropertyChangedEvent& PropertyChangedEvent) override;
 	static void AddReferencedObjects(UObject* InThis, FReferenceCollector& Collector);
diff --git a/Engine/Source/Editor/UnrealEd/Private/EditorEngine.cpp b/Engine/Source/Editor/UnrealEd/Private/EditorEngine.cpp
--- a/Engine/Source/Editor/UnrealEd/Private/EditorEngine.cpp
+++ b/Engine/Source/Editor/UnrealEd/Private/EditorEngine.cpp
@@ -1201,6 +1201,23 @@ void UEditorEngine::BroadcastObjectReimported(UObject* InObject)
 	GetEditorSubsystem<UImportSubsystem>()->BroadcastAssetReimport(InObject);
 }
 
+void UEditorEngine::BeginDestroy()
+{
+	if (!HasAnyFlags(RF_ClassDefaultObject))
+	{
+		// UWorld::FinishDestroy clears GWorld if this == GWorld, so we need to do any cleanup on GWorld
+		// before any calls to FinishDestroy are made, in case UEditorEngine and UWorld are GC'd in the same
+		// pass, which happens on editor shutdown. Doing it here in BeginDestroy is early enough.
+		UWorld* World = GWorld;
+		if (World != NULL)
+		{
+			World->ClearWorldComponents();
+			World->CleanupWorld();
+		}
+	}
+	Super::BeginDestroy();
+}
+
 void UEditorEngine::FinishDestroy()
 {
 	if ( !HasAnyFlags(RF_ClassDefaultObject) )
@@ -1235,13 +1252,6 @@ void UEditorEngine::FinishDestroy()
 			AssetRegistryModule->Get().OnInMemoryAssetCreated().RemoveAll(this);
 		}
 
-		UWorld* World = GWorld;
-		if( World != NULL )
-		{
-			World->ClearWorldComponents();
-			World->CleanupWorld();
-		}
-	
 		// Shut down transaction tracking system.
 		if( Trans )
 		{
diff --git a/Engine/Source/Runtime/Engine/Classes/Engine/World.h b/Engine/Source/Runtime/Engine/Classes/Engine/World.h
--- a/Engine/Source/Runtime/Engine/Classes/Engine/World.h
+++ b/Engine/Source/Runtime/Engine/Classes/Engine/World.h
@@ -1244,6 +1244,9 @@ private:
 	/** Is there at least one material parameter collection instance waiting for a deferred update?								*/
 	uint8 bMaterialParameterCollectionInstanceNeedsDeferredUpdate : 1;
 
+	/** Whether world object has been initialized via Init and has not yet had CleanupWorld called								*/
+	uint8 bInitializedAndNeedsCleanup : 1;
+
 	/** The world's navigation data manager */
 	UPROPERTY(Transient)
 	TObjectPtr<class UNavigationSystemBase>				NavigationSystem;
diff --git a/Engine/Source/Runtime/Engine/Private/UnrealEngine.cpp b/Engine/Source/Runtime/Engine/Private/UnrealEngine.cpp
--- a/Engine/Source/Runtime/Engine/Private/UnrealEngine.cpp
+++ b/Engine/Source/Runtime/Engine/Private/UnrealEngine.cpp
@@ -1935,6 +1935,8 @@ void UEngine::PreExit()
 	DynamicResolutionState.Reset();
 	NextDynamicResolutionState.Reset();
 #endif
+
+	EngineSubsystemCollection.Deinitialize();
 }
 
 void UEngine::ShutdownHMD()
diff --git a/Engine/Source/Runtime/Engine/Private/World.cpp b/Engine/Source/Runtime/Engine/Private/World.cpp
--- a/Engine/Source/Runtime/Engine/Private/World.cpp
+++ b/Engine/Source/Runtime/Engine/Private/World.cpp
@@ -366,6 +366,7 @@ UWorld::UWorld( const FObjectInitializer& ObjectInitializer )
 , FeatureLevel(GMaxRHIFeatureLevel)
 , bIsBuilt(false)
 , bShouldTick(true)
+, bInitializedAndNeedsCleanup(false)
 , ActiveLevelCollectionIndex(INDEX_NONE)
 , AudioDeviceHandle()
 #if WITH_EDITOR
@@ -992,6 +993,11 @@ void UWorld::ReleasePhysicsScene()
 
 void UWorld::FinishDestroy()
 {
+	if (bInitializedAndNeedsCleanup)
+	{
+		UE_LOG(LogWorld, Warning, TEXT("UWorld::FinishDestroy called after InitWorld without calling CleanupWorld first."));
+	}
+
 	// Avoid cleanup if the world hasn't been initialized, e.g., the default object or a world object that got loaded
 	// due to level streaming.
 	if (bIsWorldInitialized)
@@ -1772,6 +1778,7 @@ void UWorld::InitWorld(const InitializationValues IVS)
 
 	// We're initialized now.
 	bIsWorldInitialized = true;
+	bInitializedAndNeedsCleanup = true;
 
 	//@TODO: Should this happen here, or below here?
 	FWorldDelegates::OnPostWorldInitialization.Broadcast(this, IVS);
@@ -4624,7 +4631,12 @@ bool UWorld::IsNavigationRebuilt() const
 void UWorld::CleanupWorld(bool bSessionEnded, bool bCleanupResources, UWorld* NewWorld)
 {
     CleanupWorldGlobalTag++;
+	if (!bInitializedAndNeedsCleanup)
+	{
+		UE_LOG(LogWorld, Warning, TEXT("UWorld::CleanupWorld called twice or called without InitWorld called first."));
+	}
 	CleanupWorldInternal(bSessionEnded, bCleanupResources, NewWorld);
+	bInitializedAndNeedsCleanup = false;
 }
 
 void UWorld::CleanupWorldInternal(bool bSessionEnded, bool bCleanupResources, UWorld* NewWorld)
