Remove the GUID-checking in material functions from games.

#rnx

diff --git a/Engine/Source/Runtime/Engine/Private/Materials/MaterialExpressions.cpp b/Engine/Source/Runtime/Engine/Private/Materials/MaterialExpressions.cpp
--- a/Engine/Source/Runtime/Engine/Private/Materials/MaterialExpressions.cpp
+++ b/Engine/Source/Runtime/Engine/Private/Materials/MaterialExpressions.cpp
@@ -13228,6 +13228,7 @@ bool UMaterialFunctionInstance::ValidateFunctionUsage(class FMaterialCompiler* C
 	return Parent ? Parent->ValidateFunctionUsage(Compiler, Output) : false;
 }
 
+#if WITH_EDITORONLY_DATA
 void UMaterialFunctionInstance::Serialize(FArchive& Ar)
 {
 	FGuid OldStateId;
@@ -13249,6 +13250,7 @@ void UMaterialFunctionInstance::Serialize(FArchive& Ar)
 		}
 	}
 }
+#endif
 
 void UMaterialFunctionInstance::PostLoad()
 {
