Make DefaultEngine.ini writeable before updating dependent project settings, otherwise Lumen dependencies don't get set

diff --git a/Engine/Source/Runtime/Engine/Private/RendererSettings.cpp b/Engine/Source/Runtime/Engine/Private/RendererSettings.cpp
--- a/Engine/Source/Runtime/Engine/Private/RendererSettings.cpp
+++ b/Engine/Source/Runtime/Engine/Private/RendererSettings.cpp
@@ -101,6 +101,40 @@ void URendererSettings::PreEditChange(FProperty* PropertyAboutToChange)
 	PreEditReflectionCaptureResolution = ReflectionCaptureResolution;
 }
 
+void UpdateDependentPropertyInConfigFile(URendererSettings* RendererSettings, FName PropertyName)
+{
+	// Duplicate SSettingsEditor::NotifyPostChange functionality to make sure we are able to write to the file before we modify the ini
+	// Only duplicating the MakeWriteable behavior, SSettingsEditor::NotifyPostChange will be called after this on the main property and handle source control
+	//@todo - unify with SSettingsEditor::NotifyPostChange 
+
+	check(RendererSettings->GetClass()->HasAnyClassFlags(CLASS_DefaultConfig));
+	
+	FString RelativePath = RendererSettings->GetDefaultConfigFilename();
+	FString FullPath = FPaths::ConvertRelativePathToFull(RelativePath);
+
+	const bool bIsWriteable = !FPlatformFileManager::Get().GetPlatformFile().IsReadOnly(*FullPath);
+
+	if (!bIsWriteable)
+	{
+		FPlatformFileManager::Get().GetPlatformFile().SetReadOnly(*FullPath, false);
+	}
+
+	for (TFieldIterator<FProperty> PropIt(RendererSettings->GetClass()); PropIt; ++PropIt)
+	{
+		FProperty* Property = *PropIt;
+		if (Property->GetFName() == PropertyName)
+		{
+			RendererSettings->UpdateSinglePropertyInConfigFile(Property, RendererSettings->GetDefaultConfigFilename());
+		}
+	}
+
+	// Restore original state for source control
+	if (!bIsWriteable)
+	{
+		FPlatformFileManager::Get().GetPlatformFile().SetReadOnly(*FullPath, true);
+	}
+}
+
 void URendererSettings::PostEditChangeProperty(FPropertyChangedEvent& PropertyChangedEvent)
 {
 	Super::PostEditChangeProperty(PropertyChangedEvent);
@@ -129,28 +163,12 @@ void URendererSettings::PostEditChangeProperty(FPropertyChangedEvent& PropertyCh
 			if (FMessageDialog::Open(EAppMsgType::YesNo, LOCTEXT("Skin Cache Disabled", "Ray Tracing requires enabling skin cache. Do you want to automatically enable skin cache now?")) == EAppReturnType::Yes)
 			{
 				bSupportSkinCacheShaders = 1;
-
-				for (TFieldIterator<FProperty> PropIt(GetClass()); PropIt; ++PropIt)
-				{
-					FProperty* Property = *PropIt;
-					if (Property->GetFName() == GET_MEMBER_NAME_CHECKED(URendererSettings, bSupportSkinCacheShaders))
-					{
-						UpdateSinglePropertyInConfigFile(Property, GetDefaultConfigFilename());
-					}
-				}
+				UpdateDependentPropertyInConfigFile(this, GET_MEMBER_NAME_CHECKED(URendererSettings, bSupportSkinCacheShaders));
 			}
 			else
 			{
 				bEnableRayTracing = 0;
-
-				for (TFieldIterator<FProperty> PropIt(GetClass()); PropIt; ++PropIt)
-				{
-					FProperty* Property = *PropIt;
-					if (Property->GetFName() == GET_MEMBER_NAME_CHECKED(URendererSettings, bEnableRayTracing))
-					{
-						UpdateSinglePropertyInConfigFile(Property, GetDefaultConfigFilename());
-					}
-				}
+				UpdateDependentPropertyInConfigFile(this, GET_MEMBER_NAME_CHECKED(URendererSettings, bEnableRayTracing));
 			}
 		}
 
@@ -166,14 +184,7 @@ void URendererSettings::PostEditChangeProperty(FPropertyChangedEvent& PropertyCh
 				IConsoleVariable* CVar = IConsoleManager::Get().FindConsoleVariable(TEXT("r.ReflectionMethod"));
 				CVar->Set((int32)Reflections, ECVF_SetByProjectSetting);
 
-				for (TFieldIterator<FProperty> PropIt(GetClass()); PropIt; ++PropIt)
-				{
-					FProperty* Property = *PropIt;
-					if (Property->GetFName() == GET_MEMBER_NAME_CHECKED(URendererSettings, Reflections))
-					{
-						UpdateSinglePropertyInConfigFile(Property, GetDefaultConfigFilename());
-					}
-				}
+				UpdateDependentPropertyInConfigFile(this, GET_MEMBER_NAME_CHECKED(URendererSettings, Reflections));
 			}
 
 			if (!bGenerateMeshDistanceFields)
@@ -182,14 +193,7 @@ void URendererSettings::PostEditChangeProperty(FPropertyChangedEvent& PropertyCh
 
 				bGenerateMeshDistanceFields = true;
 
-				for (TFieldIterator<FProperty> PropIt(GetClass()); PropIt; ++PropIt)
-				{
-					FProperty* Property = *PropIt;
-					if (Property->GetFName() == GET_MEMBER_NAME_CHECKED(URendererSettings, bGenerateMeshDistanceFields))
-					{
-						UpdateSinglePropertyInConfigFile(Property, GetDefaultConfigFilename());
-					}
-				}
+				UpdateDependentPropertyInConfigFile(this, GET_MEMBER_NAME_CHECKED(URendererSettings, bGenerateMeshDistanceFields));
 			}
 		}
 
@@ -202,14 +206,7 @@ void URendererSettings::PostEditChangeProperty(FPropertyChangedEvent& PropertyCh
 
 				bEnableNonNaniteVSM = true;
 
-				for (TFieldIterator<FProperty> PropIt(GetClass()); PropIt; ++PropIt)
-				{
-					FProperty* Property = *PropIt;
-					if (Property->GetFName() == GET_MEMBER_NAME_CHECKED(URendererSettings, bEnableNonNaniteVSM))
-					{
-						UpdateSinglePropertyInConfigFile(Property, GetDefaultConfigFilename());
-					}
-				}
+				UpdateDependentPropertyInConfigFile(this, GET_MEMBER_NAME_CHECKED(URendererSettings, bEnableNonNaniteVSM));
 			}
 		}
 
