fixed incorrect usage of IsInAudioThread(), unified launching audio command

diff --git a/Engine/Source/Runtime/AudioMixer/Private/AudioMixerDevice.cpp b/Engine/Source/Runtime/AudioMixer/Private/AudioMixerDevice.cpp
--- a/Engine/Source/Runtime/AudioMixer/Private/AudioMixerDevice.cpp
+++ b/Engine/Source/Runtime/AudioMixer/Private/AudioMixerDevice.cpp
@@ -2320,9 +2320,7 @@ namespace Audio
 			}
 		};
 
-		IsInAudioThread()
-			? RegisterLambda()
-			: AsyncTask(ENamedThreads::AudioThread, MoveTemp(RegisterLambda));
+		FAudioThread::RunCommandOnAudioThread(MoveTemp(RegisterLambda));
 	}
 
 	void FMixerDevice::UnregisterSubmixBufferListener(ISubmixBufferListener* InSubmixBufferListener, USoundSubmix* InSubmix)
@@ -2350,9 +2348,7 @@ namespace Audio
 			}
 		};
 
-		IsInAudioThread()
-			? UnregisterLambda()
-			: AsyncTask(ENamedThreads::AudioThread, MoveTemp(UnregisterLambda));
+		FAudioThread::RunCommandOnAudioThread(MoveTemp(UnregisterLambda));
 	}
 
 	void FMixerDevice::FlushExtended(UWorld* WorldToFlush, bool bClearActivatedReverb)
diff --git a/Engine/Source/Runtime/AudioMixer/Private/AudioMixerSourceManager.cpp b/Engine/Source/Runtime/AudioMixer/Private/AudioMixerSourceManager.cpp
--- a/Engine/Source/Runtime/AudioMixer/Private/AudioMixerSourceManager.cpp
+++ b/Engine/Source/Runtime/AudioMixer/Private/AudioMixerSourceManager.cpp
@@ -634,8 +634,7 @@ namespace Audio
 			FSourceInfo& SourceInfo = SourceInfos[SourceId];
 
 			// Unregister these source effect instances from their owning USoundEffectInstance on the next audio thread tick.
- 			const ENamedThreads::Type UnregistrationThread = IsAudioThreadRunning() ? ENamedThreads::AudioThread: ENamedThreads::GameThread;
-			AsyncTask(UnregistrationThread, [SourceEffects = MoveTemp(SourceInfo.SourceEffects)]() mutable
+			FAudioThread::RunCommandOnAudioThread([SourceEffects = MoveTemp(SourceInfo.SourceEffects)]() mutable
 			{
 				for (int32 i = 0; i < SourceEffects.Num(); ++i)
 				{
@@ -644,11 +643,6 @@ namespace Audio
 			});
 
 			SourceInfo.SourceEffects.Reset();
-
-			for (int32 i = 0; i < SourceInfo.SourceEffectPresets.Num(); ++i)
-			{
-				SourceInfo.SourceEffectPresets[i] = nullptr;
-			}
 			SourceInfo.SourceEffectPresets.Reset();
 		}
 	}
diff --git a/Engine/Source/Runtime/AudioMixer/Private/AudioMixerSubmix.cpp b/Engine/Source/Runtime/AudioMixer/Private/AudioMixerSubmix.cpp
--- a/Engine/Source/Runtime/AudioMixer/Private/AudioMixerSubmix.cpp
+++ b/Engine/Source/Runtime/AudioMixer/Private/AudioMixerSubmix.cpp
@@ -595,9 +595,7 @@ namespace Audio
 		}
 
 		// Unregister these source effect instances from their owning USoundEffectInstance on the next audio thread tick.
-		// If the audio thread isn't currently active (ex. suspended), unregister immediately
-		const ENamedThreads::Type UnregistrationThread = IsAudioThreadRunning() ? ENamedThreads::AudioThread : ENamedThreads::GameThread;
-		AsyncTask(UnregistrationThread, [SubmixEffects = MoveTemp(SubmixEffectsToReset)]() mutable
+		FAudioThread::RunCommandOnAudioThread([SubmixEffects = MoveTemp(SubmixEffectsToReset)]() mutable
 		{
 			for (TSoundEffectSubmixPtr& SubmixPtr : SubmixEffects)
 			{
diff --git a/Engine/Source/Runtime/Engine/Private/AudioThread.cpp b/Engine/Source/Runtime/Engine/Private/AudioThread.cpp
--- a/Engine/Source/Runtime/Engine/Private/AudioThread.cpp
+++ b/Engine/Source/Runtime/Engine/Private/AudioThread.cpp
@@ -275,49 +275,66 @@ struct FAudioAsyncBatcher
 
 static FAudioAsyncBatcher GAudioAsyncBatcher;
 
-void FAudioThread::RunCommandOnAudioThread(TUniqueFunction<void()> InFunction, const TStatId InStatId)
+TUniqueFunction<void()> FAudioThread::GetCommandWrapper(TUniqueFunction<void()> InFunction, const TStatId InStatId)
 {
-	check(FPlatformTLS::GetCurrentThreadId() == GGameThreadId);
-	if (IsAudioThreadRunning())
+	if (GCVarEnableAudioCommandLogging == 1)
 	{
-		if (GCVarEnableAudioCommandLogging == 1)
+		return[Function = MoveTemp(InFunction), InStatId]()
 		{
-			TUniqueFunction<void()> FuncWrapper = [Function = MoveTemp(InFunction), InStatId]()
-			{
-				FTaskTagScope Scope(ETaskTag::EAudioThread);
-				FAudioThread::SetCurrentAudioThreadStatId(InStatId);
+			FTaskTagScope Scope(ETaskTag::EAudioThread);
+			FScopeCycleCounter ScopeCycleCounter(InStatId);
+			FAudioThread::SetCurrentAudioThreadStatId(InStatId);
 
-				// Time the execution of the function
-				const double StartTime = FPlatformTime::Seconds();
+			// Time the execution of the function
+			const double StartTime = FPlatformTime::Seconds();
 
-				// Execute the function
-				Function();
-
-				// Track the longest one
-				const double DeltaTime = (FPlatformTime::Seconds() - StartTime) * 1000.0f;
-				if (DeltaTime > GetCurrentLongestTime())
-				{
-					SetLongestTimeAndId(InStatId, DeltaTime);
-				}
-			};
+			// Execute the function
+			Function();
 
-			FFunctionGraphTask::CreateAndDispatchWhenReady(MoveTemp(FuncWrapper), InStatId, GAudioAsyncBatcher.GetAsyncPrereq(), ENamedThreads::AudioThread);
-		}
-		else
-		{
-			TUniqueFunction<void()> FuncWrapper = [Function = MoveTemp(InFunction), InStatId]()
+			// Track the longest one
+			const double DeltaTime = (FPlatformTime::Seconds() - StartTime) * 1000.0f;
+			if (DeltaTime > GetCurrentLongestTime())
 			{
-				FTaskTagScope Scope(ETaskTag::EAudioThread);
-				// Execute the function
-				Function();
-			};
-			FFunctionGraphTask::CreateAndDispatchWhenReady(MoveTemp(FuncWrapper), InStatId, GAudioAsyncBatcher.GetAsyncPrereq(), ENamedThreads::AudioThread);
-		}
+				SetLongestTimeAndId(InStatId, DeltaTime);
+			}
+		};
 	}
 	else
 	{
+		return[Function = MoveTemp(InFunction), InStatId]()
+		{
+			FTaskTagScope Scope(ETaskTag::EAudioThread);
+			FScopeCycleCounter ScopeCycleCounter(InStatId);
+			Function();
+		};
+	}
+}
+
+void FAudioThread::RunCommandOnAudioThread(TUniqueFunction<void()> InFunction, const TStatId InStatId)
+{
+	if (IsInAudioThread())
+	{
+		// it's audio-thread-safe so execute the command in-place
 		FScopeCycleCounter ScopeCycleCounter(InStatId);
 		InFunction();
+		return;
+	}
+
+	TUniqueFunction<void()> CommandWrapper{ GetCommandWrapper(MoveTemp(InFunction), InStatId) };
+	if (IsInGameThread())
+	{
+		// batch commands to minimise game thread overhead
+		FFunctionGraphTask::CreateAndDispatchWhenReady(MoveTemp(CommandWrapper), TStatId{}, GAudioAsyncBatcher.GetAsyncPrereq(), ENamedThreads::AudioThread);
+	}
+	// we are on an unknown thread
+	else if (IsUsingThreadedAudio())
+	{
+		FFunctionGraphTask::CreateAndDispatchWhenReady(MoveTemp(CommandWrapper), TStatId{}, nullptr, ENamedThreads::AudioThread);
+	}
+	else
+	{
+		// the command must be executed on the game thread
+		FFunctionGraphTask::CreateAndDispatchWhenReady(MoveTemp(CommandWrapper), TStatId{}, nullptr, ENamedThreads::GameThread);
 	}
 }
 
diff --git a/Engine/Source/Runtime/Engine/Private/SoundModulationDestination.cpp b/Engine/Source/Runtime/Engine/Private/SoundModulationDestination.cpp
--- a/Engine/Source/Runtime/Engine/Private/SoundModulationDestination.cpp
+++ b/Engine/Source/Runtime/Engine/Private/SoundModulationDestination.cpp
@@ -259,8 +259,6 @@ namespace Audio
 			}
 		};
 
-		IsInAudioThread()
-			? UpdateHandleLambda()
-			: AsyncTask(ENamedThreads::AudioThread, MoveTemp(UpdateHandleLambda));
+		FAudioThread::RunCommandOnAudioThread(MoveTemp(UpdateHandleLambda));
 	}
 } // namespace Audio
diff --git a/Engine/Source/Runtime/Engine/Public/AudioThread.h b/Engine/Source/Runtime/Engine/Public/AudioThread.h
--- a/Engine/Source/Runtime/Engine/Public/AudioThread.h
+++ b/Engine/Source/Runtime/Engine/Public/AudioThread.h
@@ -54,6 +54,9 @@ private:
 	static void SetLongestTimeAndId(TStatId NewLongestId, double LongestTimeMsec);
 	static double GetCurrentLongestTime() { return LongestAudioThreadTimeMsec; }
 
+	// a helper to apply stats and optional additional logging
+	static TUniqueFunction<void()> GetCommandWrapper(TUniqueFunction<void()> InFunction, const TStatId InStatId);
+
 public:
 
 	FAudioThread();
