WorldPartition: Remove warning until we have a better way of finding invalid bound actors.

diff --git a/Engine/Source/Runtime/Engine/Private/WorldPartition/WorldPartitionRuntimeSpatialHash.cpp b/Engine/Source/Runtime/Engine/Private/WorldPartition/WorldPartitionRuntimeSpatialHash.cpp
--- a/Engine/Source/Runtime/Engine/Private/WorldPartition/WorldPartitionRuntimeSpatialHash.cpp
+++ b/Engine/Source/Runtime/Engine/Private/WorldPartition/WorldPartitionRuntimeSpatialHash.cpp
@@ -490,51 +490,6 @@ FName UWorldPartitionRuntimeSpatialHash::GetActorRuntimeGrid(const AActor* Actor
 	return Super::GetActorRuntimeGrid(Actor);
 }
 
-void UWorldPartitionRuntimeSpatialHash::CheckForErrorsInternal(const TMap<FGuid, FWorldPartitionActorViewProxy>& ActorDescList) const
-{
-	Super::CheckForErrorsInternal(ActorDescList);
-
-	auto GetActorLabel = [](const FWorldPartitionActorDescView& ActorDescView) -> FString
-	{
-		const FName ActorLabel = ActorDescView.GetActorLabel();
-		if (!ActorLabel.IsNone())
-		{
-			return ActorLabel.ToString();
-		}
-
-		const FString ActorPath = ActorDescView.GetActorPath().ToString();
-
-		FString SubObjectName;
-		FString SubObjectContext;
-		if (FString(ActorPath).Split(TEXT("."), &SubObjectContext, &SubObjectName))
-		{
-			return SubObjectName;
-		}
-
-		return ActorPath;
-	};
-
-	for (auto& ActorDescListPair: ActorDescList)
-	{
-		const FWorldPartitionActorViewProxy& ActorDescView = ActorDescListPair.Value;
-
-		if (!ActorDescView.GetActorIsEditorOnly())
-		{
-			const FBox2D ActorBounds2D = FBox2D(FVector2D(ActorDescView.GetBounds().Min), FVector2D(ActorDescView.GetBounds().Max));
-			const float ActorBoundsArea2D = ActorBounds2D.GetArea();
-
-			if (ActorBoundsArea2D < KINDA_SMALL_NUMBER)
-			{
-				TSharedRef<FTokenizedMessage> Error = FMessageLog("MapCheck").Warning()
-					->AddToken(FTextToken::Create(LOCTEXT("MapCheck_WorldPartition_Actor", "Actor")))
-					->AddToken(FAssetNameToken::Create(GetActorLabel(ActorDescView)))
-					->AddToken(FTextToken::Create(LOCTEXT("MapCheck_WorldPartition_HasEmptyBounds", "has empty bounds")))
-					->AddToken(FMapErrorToken::Create(FName(TEXT("WorldPartition_ActorHasEmptyBounds_CheckForErrors"))));
-			}
-		}
-	}
-}
-
 void UWorldPartitionRuntimeSpatialHash::SetDefaultValues()
 {
 	FSpatialHashRuntimeGrid& MainGrid = Grids.AddDefaulted_GetRef();
diff --git a/Engine/Source/Runtime/Engine/Public/WorldPartition/WorldPartitionRuntimeSpatialHash.h b/Engine/Source/Runtime/Engine/Public/WorldPartition/WorldPartitionRuntimeSpatialHash.h
--- a/Engine/Source/Runtime/Engine/Public/WorldPartition/WorldPartitionRuntimeSpatialHash.h
+++ b/Engine/Source/Runtime/Engine/Public/WorldPartition/WorldPartitionRuntimeSpatialHash.h
@@ -207,7 +207,6 @@ public:
 	virtual bool GenerateHLOD(ISourceControlHelper* SourceControlHelper, bool bCreateActorsOnly) override;
 	virtual bool GenerateNavigationData() override;
 	virtual FName GetActorRuntimeGrid(const AActor* Actor) const override;
-	virtual void CheckForErrorsInternal(const TMap<FGuid, FWorldPartitionActorViewProxy>& ActorDescList) const override;
 	virtual void DrawPreview() const override;
 
 	FName GetCellName(FName InGridName, const FIntVector& InCellGlobalCoord, const FDataLayersID& InDataLayerID) const;
