LevelInstance: override PostEditUndo with annotation (because the UObject implementation explicitely calls the UObject parameterless version)

diff --git a/Engine/Source/Runtime/Engine/Private/LevelInstance/LevelInstanceActor.cpp b/Engine/Source/Runtime/Engine/Private/LevelInstance/LevelInstanceActor.cpp
--- a/Engine/Source/Runtime/Engine/Private/LevelInstance/LevelInstanceActor.cpp
+++ b/Engine/Source/Runtime/Engine/Private/LevelInstance/LevelInstanceActor.cpp
@@ -255,6 +255,16 @@ void ALevelInstance::PreEditUndo()
 	}
 }
 
+void ALevelInstance::PostEditUndo(TSharedPtr<ITransactionObjectAnnotation> TransactionAnnotation)
+{
+	{
+		TGuardValue<bool> LoadUnloadGuard(bGuardLoadUnload, true);
+		Super::PostEditUndo(TransactionAnnotation);
+	}
+
+	PostEditUndoInternal();
+}
+
 void ALevelInstance::PostEditUndo()
 {
 	{
@@ -262,6 +272,11 @@ void ALevelInstance::PostEditUndo()
 		Super::PostEditUndo();
 	}
 
+	PostEditUndoInternal();
+}
+
+void ALevelInstance::PostEditUndoInternal()
+{
 	if (CachedWorldAsset != WorldAsset)
 	{
 		OnWorldAssetChanged();
diff --git a/Engine/Source/Runtime/Engine/Public/LevelInstance/LevelInstanceActor.h b/Engine/Source/Runtime/Engine/Public/LevelInstance/LevelInstanceActor.h
--- a/Engine/Source/Runtime/Engine/Public/LevelInstance/LevelInstanceActor.h
+++ b/Engine/Source/Runtime/Engine/Public/LevelInstance/LevelInstanceActor.h
@@ -55,6 +55,7 @@ public:
 	// UObject overrides
 	virtual void PreEditUndo() override;
 	virtual void PostEditUndo() override;
+	virtual void PostEditUndo(TSharedPtr<ITransactionObjectAnnotation> TransactionAnnotation) override;
 	virtual void PostLoad() override;
 	virtual bool CanEditChange(const FProperty* InProperty) const override;
 	virtual void PreEditChange(FProperty* PropertyThatWillChange) override;
@@ -111,6 +112,7 @@ public:
 private:
 
 #if WITH_EDITOR
+	void PostEditUndoInternal();
 	bool CanSetValue(TSoftObjectPtr<UWorld> LevelInstance, FString* Reason = nullptr) const;
 	
 	TSoftObjectPtr<UWorld> CachedWorldAsset;
