Add missing PostRegisterAllComponent call for symmetry

diff --git a/Engine/Source/Runtime/Engine/Private/ActorEditor.cpp b/Engine/Source/Runtime/Engine/Private/ActorEditor.cpp
--- a/Engine/Source/Runtime/Engine/Private/ActorEditor.cpp
+++ b/Engine/Source/Runtime/Engine/Private/ActorEditor.cpp
@@ -166,6 +166,10 @@ void AActor::PostEditChangeProperty(FPropertyChangedEvent& PropertyChangedEvent)
 			{
 				ReregisterAllComponents();
 			}
+			else
+			{
+				PostRegisterAllComponents();
+			}
 		}
 		else
 		{
