The Chaos physics mesh cooker will now change the winding order of triangles according to the bFlipNormals flag (instead of assuming flipped normals)

diff --git a/Engine/Source/Runtime/Engine/Private/PhysicsEngine/Experimental/ChaosCooking.cpp b/Engine/Source/Runtime/Engine/Private/PhysicsEngine/Experimental/ChaosCooking.cpp
--- a/Engine/Source/Runtime/Engine/Private/PhysicsEngine/Experimental/ChaosCooking.cpp
+++ b/Engine/Source/Runtime/Engine/Private/PhysicsEngine/Experimental/ChaosCooking.cpp
@@ -30,9 +30,11 @@ namespace Chaos
 			FinalIndices.Reserve(Desc.Indices.Num() * 3);
 			for(const FTriIndices& Tri : Desc.Indices)
 			{
-				//question: It seems like unreal triangles are CW, but couldn't find confirmation for this
-				FinalIndices.Add(Tri.v1);
-				FinalIndices.Add(Tri.v0);
+				// NOTE: This is where the Winding order of the triangles are changed to be consistent throughout the rest of the physics engine
+				// After this point we should have clockwise (CW) winding in left handed (LH) coordinates (or equivalently CCW in RH)
+				// This is the opposite convention followed in most of the unreal engine
+				FinalIndices.Add(Desc.bFlipNormals ? Tri.v1 : Tri.v0);
+				FinalIndices.Add(Desc.bFlipNormals ? Tri.v0 : Tri.v1);
 				FinalIndices.Add(Tri.v2);
 			}
 
@@ -200,9 +202,11 @@ namespace Chaos
 			FinalIndices.Reserve(InParams.TriangleMeshDesc.Indices.Num() * 3);
 			for(const FTriIndices& Tri : InParams.TriangleMeshDesc.Indices)
 			{
-				//question: It seems like unreal triangles are CW, but couldn't find confirmation for this
-				FinalIndices.Add(Tri.v1);
-				FinalIndices.Add(Tri.v0);
+				// NOTE: This is where the Winding order of the triangles are changed to be consistent throughout the rest of the physics engine
+				// After this point we should have clockwise (CW) winding in left handed (LH) coordinates (or equivalently CCW in RH)
+				// This is the opposite convention followed in most of the unreal engine
+				FinalIndices.Add(InParams.TriangleMeshDesc.bFlipNormals ? Tri.v1 : Tri.v0);
+				FinalIndices.Add(InParams.TriangleMeshDesc.bFlipNormals ? Tri.v0 : Tri.v1);
 				FinalIndices.Add(Tri.v2);
 			}
 
diff --git a/Engine/Source/Runtime/Experimental/ChaosCore/Public/Chaos/Core.h b/Engine/Source/Runtime/Experimental/ChaosCore/Public/Chaos/Core.h
--- a/Engine/Source/Runtime/Experimental/ChaosCore/Public/Chaos/Core.h
+++ b/Engine/Source/Runtime/Experimental/ChaosCore/Public/Chaos/Core.h
@@ -35,5 +35,5 @@ namespace Chaos
 	// NOTE: if you get a merge conflict on the GUID, you must replace it with a new GUID - do not accept the source or target
 	// or you will likely get DDC version conflicts resulting in crashes during load.
 	// Core version string for Chaos data. Any DDC builder dependent on Chaos for serialization should depend on this version
-	static const FString ChaosVersionString = TEXT("7E91D439-4ACF-4A10-91A8-69D17564B7AF");
+	static const FString ChaosVersionString = TEXT("A493942D-9C01-4A47-99AB-A69B84457659");
 }
