Optimized world partition runtime hash 2D debug drawing

diff --git a/Engine/Source/Runtime/Engine/Private/WorldPartition/WorldPartitionRuntimeSpatialHash.cpp b/Engine/Source/Runtime/Engine/Private/WorldPartition/WorldPartitionRuntimeSpatialHash.cpp
--- a/Engine/Source/Runtime/Engine/Private/WorldPartition/WorldPartitionRuntimeSpatialHash.cpp
+++ b/Engine/Source/Runtime/Engine/Private/WorldPartition/WorldPartitionRuntimeSpatialHash.cpp
@@ -294,6 +294,12 @@ void FSpatialHashStreamingGrid::Draw2D(UCanvas* Canvas, const TArray<FWorldParti
 	int32 MinGridLevel = FMath::Clamp<int32>(GShowRuntimeSpatialHashGridLevel, 0, GridLevels.Num() - 1);
 	int32 MaxGridLevel = FMath::Clamp<int32>(MinGridLevel + GShowRuntimeSpatialHashGridLevelCount - 1, 0, GridLevels.Num() - 1);
 
+	// Precompute a cell coordinate text with/height using a generic coordinate
+	// This will be used later to filter out drawing of cell coordinates (avoid doing expesive calls to Canvas->StrLen)
+	const FString CellCoordString = UWorldPartitionRuntimeSpatialHash::GetCellCoordString(FIntVector(88,88,88));
+	float MaxCellCoordTextWidth, MaxCellCoordTextHeight;
+	Canvas->StrLen(GEngine->GetTinyFont(), CellCoordString, MaxCellCoordTextWidth, MaxCellCoordTextHeight);
+
 	for (int32 GridLevel = MinGridLevel; GridLevel <= MaxGridLevel; ++GridLevel)
 	{
 		// Draw X/Y Axis
@@ -341,15 +347,15 @@ void FSpatialHashStreamingGrid::Draw2D(UCanvas* Canvas, const TArray<FWorldParti
 			{
 				FIntVector CellGlobalCoords;
 				verify(Helper.GetCellGlobalCoords(FIntVector(Coords.X, Coords.Y, GridLevel), CellGlobalCoords));
-				const FString GridInfoText = UWorldPartitionRuntimeSpatialHash::GetCellCoordString(CellGlobalCoords);
-				float TextWidth, TextHeight;
-				Canvas->SetDrawColor(255, 255, 0);
-				Canvas->StrLen(GEngine->GetTinyFont(), GridInfoText, TextWidth, TextHeight);
-				FVector2D CellBoundsSize = CellScreenBounds.GetSize();
-				if (TextWidth < CellBoundsSize.X && TextHeight < CellBoundsSize.Y)
+				const FString CellCoordString = UWorldPartitionRuntimeSpatialHash::GetCellCoordString(CellGlobalCoords);
+				const FVector2D CellBoundsSize = CellScreenBounds.GetSize();
+				if (MaxCellCoordTextWidth < CellBoundsSize.X && MaxCellCoordTextHeight < CellBoundsSize.Y)
 				{
-					FVector2D GridInfoPos = CellScreenBounds.GetCenter() - FVector2D(TextWidth / 2, TextHeight / 2);
-					Canvas->DrawText(GEngine->GetTinyFont(), GridInfoText, GridInfoPos.X, GridInfoPos.Y);
+					float CellCoordTextWidth, CellCoordTextHeight;
+					Canvas->StrLen(GEngine->GetTinyFont(), CellCoordString, CellCoordTextWidth, CellCoordTextHeight);
+					FVector2D GridInfoPos = CellScreenBounds.GetCenter() - FVector2D(CellCoordTextWidth / 2, CellCoordTextHeight / 2);
+					Canvas->SetDrawColor(255, 255, 0);
+					Canvas->DrawText(GEngine->GetTinyFont(), CellCoordString, GridInfoPos.X, GridInfoPos.Y);
 				}
 			}
 
@@ -1117,7 +1123,7 @@ void UWorldPartitionRuntimeSpatialHash::Draw2D(UCanvas* Canvas, const TArray<FWo
 		FSphere AverageSphere(ForceInit);
 		for (const FWorldPartitionStreamingSource& Source : Sources)
 		{
-			AverageSphere += FSphere(Source.Location, 0.5f * GridSideDistance);
+			AverageSphere += FSphere(FVector(Source.Location.X, Source.Location.Y, 0.f), 0.5f * GridSideDistance);
 		}
 		const FVector2D GridReferenceWorldPos = FVector2D(AverageSphere.Center);
 		const FBox Region(AverageSphere.Center - AverageSphere.W, AverageSphere.Center + AverageSphere.W);
