Skin cache recompute tangent - fix broken logic of merging duplicated vertices

There are double errors:
1) GPUSkinCacheNeedsDuplicatedVertices has the toggle wrong way round, duplicated vertex buffer is needed when GAllowDupedVertsForRecomputeTangents is false, hence merging vertices is performed in recompute tangents
2) Which triangle pass shader to bind with GAllowDupedVertsForRecomputeTangents is wrong way round

diff --git a/Engine/Source/Runtime/Engine/Private/GPUSkinCache.cpp b/Engine/Source/Runtime/Engine/Private/GPUSkinCache.cpp
--- a/Engine/Source/Runtime/Engine/Private/GPUSkinCache.cpp
+++ b/Engine/Source/Runtime/Engine/Private/GPUSkinCache.cpp
@@ -107,7 +107,6 @@ static TAutoConsoleVariable<float> CVarGPUSkinCacheSceneMemoryLimitInMB(
 	ECVF_RenderThreadSafe
 );
 
-////temporary disable until resource lifetimes are safe for all cases
 static int32 GAllowDupedVertsForRecomputeTangents = 0;
 FAutoConsoleVariableRef CVarGPUSkinCacheAllowDupedVertesForRecomputeTangents(
 	TEXT("r.SkinCache.AllowDupedVertsForRecomputeTangents"),
@@ -164,7 +163,7 @@ ENGINE_API bool GPUSkinCacheNeedsDuplicatedVertices()
 #if WITH_EDITOR // Duplicated vertices are used in the editor when merging meshes
 	return true;
 #else
-	return (bool)GAllowDupedVertsForRecomputeTangents;
+	return GAllowDupedVertsForRecomputeTangents == 0;
 #endif
 }
 
@@ -1174,13 +1173,13 @@ void FGPUSkinCache::DispatchUpdateSkinTangents(FRHICommandListImmediate& RHICmdL
 
 			if (bFullPrecisionUV)
 			{
-				if (GAllowDupedVertsForRecomputeTangents) Shader = ComputeShader11;
-				else Shader = ComputeShader01;
+				if (GAllowDupedVertsForRecomputeTangents) Shader = ComputeShader01;
+				else Shader = ComputeShader11;
 			}
 			else
 			{
-				if (GAllowDupedVertsForRecomputeTangents) Shader = ComputeShader10;
-				else Shader = ComputeShader00;
+				if (GAllowDupedVertsForRecomputeTangents) Shader = ComputeShader00;
+				else Shader = ComputeShader10;
 			}
 
 			check(Shader.IsValid());
