Modified wp.Runtime.HLOD console command to also turn on/off rendering of always loaded HLODs.

diff --git a/Engine/Source/Runtime/Engine/Private/WorldPartition/HLOD/HLODSubsystem.cpp b/Engine/Source/Runtime/Engine/Private/WorldPartition/HLOD/HLODSubsystem.cpp
--- a/Engine/Source/Runtime/Engine/Private/WorldPartition/HLOD/HLODSubsystem.cpp
+++ b/Engine/Source/Runtime/Engine/Private/WorldPartition/HLOD/HLODSubsystem.cpp
@@ -27,6 +27,39 @@ UHLODSubsystem::~UHLODSubsystem()
 {
 }
 
+bool UHLODSubsystem::WorldPartitionHLODEnabled = true;
+
+FAutoConsoleCommand UHLODSubsystem::EnableHLODCommand(
+	TEXT("wp.Runtime.HLOD"),
+	TEXT("Turn on/off loading & rendering of world partition HLODs."),
+	FConsoleCommandWithArgsDelegate::CreateLambda([](const TArray<FString>& Args)
+	{
+		UHLODSubsystem::WorldPartitionHLODEnabled = (Args.Num() != 1) || (Args[0] != TEXT("0"));
+		for (const FWorldContext& Context : GEngine->GetWorldContexts())
+		{
+			UWorld* World = Context.World();
+			if (World && World->IsGameWorld())
+			{
+				UHLODSubsystem* HLODSubSystem = World->GetSubsystem<UHLODSubsystem>();
+				for (const auto& CellHLODMapping : HLODSubSystem->CellsHLODMapping)
+				{
+					const FCellHLODMapping& CellHLODs = CellHLODMapping.Value;
+					bool bIsHLODVisible = UHLODSubsystem::WorldPartitionHLODEnabled && !CellHLODs.bIsCellVisible;
+					for (AWorldPartitionHLOD* HLODActor : CellHLODs.LoadedHLODs)
+					{
+						HLODActor->SetVisibility(bIsHLODVisible);
+					}
+				}
+			}
+		}
+	})
+);
+
+bool UHLODSubsystem::IsHLODEnabled()
+{
+	return UHLODSubsystem::WorldPartitionHLODEnabled;
+}
+
 bool UHLODSubsystem::ShouldCreateSubsystem(UObject* Outer) const
 {
 	if (!Super::ShouldCreateSubsystem(Outer))
@@ -90,7 +123,7 @@ void UHLODSubsystem::RegisterHLODActor(AWorldPartitionHLOD* InWorldPartitionHLOD
 	if (CellHLODs)
 	{
 		CellHLODs->LoadedHLODs.Add(InWorldPartitionHLOD);
-		InWorldPartitionHLOD->SetVisibility(!CellHLODs->bIsCellVisible);
+		InWorldPartitionHLOD->SetVisibility(UHLODSubsystem::WorldPartitionHLODEnabled && !CellHLODs->bIsCellVisible);
 	}
 	else
 	{
@@ -148,7 +181,7 @@ void UHLODSubsystem::OnCellHidden(const UWorldPartitionRuntimeCell* InCell)
 #if WITH_EDITOR
 		UE_LOG(LogHLODSubsystem, Verbose, TEXT("\t\t%s - %s"), *HLODActor->GetActorLabel(), *HLODActor->GetActorGuid().ToString());
 #endif
-		HLODActor->SetVisibility(true);
+		HLODActor->SetVisibility(UHLODSubsystem::WorldPartitionHLODEnabled);
 	}
 }
 
diff --git a/Engine/Source/Runtime/Engine/Private/WorldPartition/WorldPartitionStreamingPolicy.cpp b/Engine/Source/Runtime/Engine/Private/WorldPartition/WorldPartitionStreamingPolicy.cpp
--- a/Engine/Source/Runtime/Engine/Private/WorldPartition/WorldPartitionStreamingPolicy.cpp
+++ b/Engine/Source/Runtime/Engine/Private/WorldPartition/WorldPartitionStreamingPolicy.cpp
@@ -11,6 +11,7 @@
 #include "WorldPartition/WorldPartition.h"
 #include "WorldPartition/DataLayer/DataLayerSubsystem.h"
 #include "WorldPartition/HLOD/HLODActor.h"
+#include "WorldPartition/HLOD/HLODSubsystem.h"
 #include "GameFramework/PlayerController.h"
 #include "Engine/LocalPlayer.h"
 #include "Engine/World.h"
@@ -26,14 +27,6 @@
 #include "LevelEditorViewport.h"
 #endif
 
-int32 GWorldPartitionShowHLODs = 1;
-static FAutoConsoleVariableRef CVarWorldPartitionShowHLODs(
-	TEXT("wp.Runtime.HLOD"),
-	GWorldPartitionShowHLODs,
-	TEXT("Turn on/off loading & rendering of world partition HLODs."),
-	ECVF_Scalability
-);
-
 UWorldPartitionStreamingPolicy::UWorldPartitionStreamingPolicy(const FObjectInitializer& ObjectInitializer)
 : Super(ObjectInitializer) 
 , bIsServerLoadingDone(false)
@@ -157,7 +150,7 @@ void UWorldPartitionStreamingPolicy::UpdateStreamingState()
 			LoadStreamingCells = LoadStreamingCells.Difference(ActivateStreamingCells);
 		}
 
-		if (!GWorldPartitionShowHLODs)
+		if (!UHLODSubsystem::IsHLODEnabled())
 		{
 			// Remove all HLOD cells from the Activate & Load cells
 			auto RemoveHLODCells = [](TSet<const UWorldPartitionRuntimeCell*>& Cells)
@@ -217,6 +210,7 @@ void UWorldPartitionStreamingPolicy::UpdateStreamingState()
 bool UWorldPartitionStreamingPolicy::IsStreamingCompleted(EWorldPartitionRuntimeCellState QueryState, const TArray<FWorldPartitionStreamingQuerySource>& QuerySources, bool bExactState) const
 {
 	const UDataLayerSubsystem* DataLayerSubsystem = WorldPartition->GetWorld()->GetSubsystem<UDataLayerSubsystem>();
+	const bool bIsHLODEnabled = UHLODSubsystem::IsHLODEnabled();
 
 	for (const FWorldPartitionStreamingQuerySource& QuerySource : QuerySources)
 	{
@@ -231,7 +225,7 @@ bool UWorldPartitionStreamingPolicy::IsStreamingCompleted(EWorldPartitionRuntime
 				bool bSkipCell = false;
 
 				// Don't consider HLOD cells if HLODs are disabled.
-				if (!GWorldPartitionShowHLODs)
+				if (!bIsHLODEnabled)
 				{
 					bSkipCell = Cell->HasCellData<UWorldPartitionRuntimeHLODCellData>();
 				}
diff --git a/Engine/Source/Runtime/Engine/Public/WorldPartition/HLOD/HLODSubsystem.h b/Engine/Source/Runtime/Engine/Public/WorldPartition/HLOD/HLODSubsystem.h
--- a/Engine/Source/Runtime/Engine/Public/WorldPartition/HLOD/HLODSubsystem.h
+++ b/Engine/Source/Runtime/Engine/Public/WorldPartition/HLOD/HLODSubsystem.h
@@ -34,12 +34,20 @@ public:
 
 	void OnCellShown(const UWorldPartitionRuntimeCell* InCell);
 	void OnCellHidden(const UWorldPartitionRuntimeCell* InCell);
+
+	static bool IsHLODEnabled();
 	
 private:
 	void OnWorldPartitionRegistered(UWorldPartition* InWorldPartition);
 	void OnWorldPartitionUnregistered(UWorldPartition* InWorldPartition);
 
+	/** Console command used to turn on/off loading & rendering of world partition HLODs */
+	static class FAutoConsoleCommand EnableHLODCommand;
+
 private:
+
+	static bool WorldPartitionHLODEnabled;
+
 	struct FCellHLODMapping
 	{
 		bool						bIsCellVisible;
