WorldPartition - Dynamic data layers will now be generating HLODs by default

diff --git a/Engine/Source/Runtime/Engine/Private/WorldPartition/DataLayer/DataLayer.cpp b/Engine/Source/Runtime/Engine/Private/WorldPartition/DataLayer/DataLayer.cpp
--- a/Engine/Source/Runtime/Engine/Private/WorldPartition/DataLayer/DataLayer.cpp
+++ b/Engine/Source/Runtime/Engine/Private/WorldPartition/DataLayer/DataLayer.cpp
@@ -16,7 +16,7 @@ UDataLayer::UDataLayer(const FObjectInitializer& ObjectInitializer)
 , bIsInitiallyVisible(true)
 , bIsDynamicallyLoadedInEditor(true)
 , bIsLocked(false)
-, bGeneratesHLODs(false)
+, bGeneratesHLODs(true)
 , DefaultHLODLayer()
 #endif
 , DataLayerLabel(GetFName())
