Add ability to dump material DDC keys (UE-43174).

- r.Material.DumpDDCKeys 1 (togglable at runtime) will cause DDC keys to be dumped into Saved/MaterialDDCKeys

diff --git a/Engine/Source/Runtime/Engine/Private/Materials/MaterialShader.cpp b/Engine/Source/Runtime/Engine/Private/Materials/MaterialShader.cpp
--- a/Engine/Source/Runtime/Engine/Private/Materials/MaterialShader.cpp
+++ b/Engine/Source/Runtime/Engine/Private/Materials/MaterialShader.cpp
@@ -27,6 +27,14 @@ static FAutoConsoleVariableRef CVarMaterialExcludeNonPipelinedShaders(
 	ECVF_ReadOnly
 );
 
+int32 GMaterialDumpDDCKeys = 0;
+static FAutoConsoleVariableRef CVarMaterialDumpDDCKeys(
+	TEXT("r.Material.DumpDDCKeys"),
+	GMaterialDumpDDCKeys,
+	TEXT("if != 0, DDC keys for each material shadermap will be dumped into project's Saved directory (MaterialDDCKeys subdirectory)"),
+	ECVF_Default
+);
+
 #if ENABLE_COOK_STATS
 namespace MaterialShaderCookStats
 {
@@ -1181,6 +1189,19 @@ void FMaterialShaderMap::LoadFromDerivedDataCache(const FMaterial* Material, con
 			const FString DataKey = GetMaterialShaderMapKeyString(ShaderMapId, InPlatform);
 			OutDDCKeyDesc = FSHA1_HashString(DataKey);
 
+			if (UNLIKELY(GMaterialDumpDDCKeys))
+			{
+				FString TempPath = FPaths::ProjectSavedDir() / TEXT("MaterialDDCKeys") / LexToString(InPlatform);
+				IFileManager::Get().MakeDirectory(*TempPath, true);
+
+				FString FileName = FString::Printf(TEXT("%s-%s-%s.txt"), *Material->GetAssetName(), *LexToString(ShaderMapId.FeatureLevel), *LexToString(ShaderMapId.QualityLevel));
+				FString TempFile = TempPath / FileName;
+
+				TUniquePtr<FArchive> DumpAr(IFileManager::Get().CreateFileWriter(*TempFile));
+				// serializing the string via << produces a non-textual file because it saves string's length, too
+				DumpAr->Serialize(const_cast<TCHAR*>(*DataKey), DataKey.Len() * sizeof(TCHAR));
+			}
+
 			bool CheckCache = true;
 
 			// If NoShaderDDC then don't check for a material the first time we encounter it to simulate
