Add perf markers for skinning & recompute tangent batching

diff --git a/Engine/Source/Runtime/Engine/Private/GPUSkinCache.cpp b/Engine/Source/Runtime/Engine/Private/GPUSkinCache.cpp
--- a/Engine/Source/Runtime/Engine/Private/GPUSkinCache.cpp
+++ b/Engine/Source/Runtime/Engine/Private/GPUSkinCache.cpp
@@ -1301,24 +1301,30 @@ void FGPUSkinCache::DoDispatch(FRHICommandListImmediate& RHICmdList)
 	}
 
 	RHICmdList.BeginUAVOverlap(OverlappedUAVBuffers);
-	for (int32 i = 0; i < BatchCount; ++i)
 	{
-		FDispatchEntry& DispatchItem = BatchDispatches[i];
-		DispatchUpdateSkinning(RHICmdList, DispatchItem.SkinCacheEntry, DispatchItem.Section, DispatchItem.RevisionNumber);
+		SCOPED_DRAW_EVENT(RHICmdList, GPUSkinCache_UpdateSkinningBatches);
+		for (int32 i = 0; i < BatchCount; ++i)
+		{
+			FDispatchEntry& DispatchItem = BatchDispatches[i];
+			DispatchUpdateSkinning(RHICmdList, DispatchItem.SkinCacheEntry, DispatchItem.Section, DispatchItem.RevisionNumber);
+		}
 	}
 	RHICmdList.EndUAVOverlap(OverlappedUAVBuffers);
 
-	for (int32 i = 0; i < BatchCount; ++i)
 	{
-		FDispatchEntry& DispatchItem = BatchDispatches[i];
-		DispatchItem.SkinCacheEntry->UpdateVertexFactoryDeclaration(DispatchItem.Section);
-
-		if (DispatchItem.SkinCacheEntry->DispatchData[DispatchItem.Section].IndexBuffer)
+		SCOPED_DRAW_EVENT(RHICmdList, GPUSkinCache_RecomputeTangentsBatches);
+		for (int32 i = 0; i < BatchCount; ++i)
 		{
-			DispatchUpdateSkinTangents(RHICmdList, DispatchItem.SkinCacheEntry, DispatchItem.Section);
+			FDispatchEntry& DispatchItem = BatchDispatches[i];
+			DispatchItem.SkinCacheEntry->UpdateVertexFactoryDeclaration(DispatchItem.Section);
+	
+			if (DispatchItem.SkinCacheEntry->DispatchData[DispatchItem.Section].IndexBuffer)
+			{
+				DispatchUpdateSkinTangents(RHICmdList, DispatchItem.SkinCacheEntry, DispatchItem.Section);
+			}
+	
+			DispatchItem.SkinCacheEntry->UpdateVertexFactoryDeclaration(DispatchItem.Section);
 		}
-
-		DispatchItem.SkinCacheEntry->UpdateVertexFactoryDeclaration(DispatchItem.Section);
 	}
 }
 
