Implemented the full fix for UE-106993 and removed workaround put in place for 4.26 hotfix. BeginReleaseResource() can now be used with FRayTracingGeometry, as it now implements FRenderResource::ReleaseResource() interface.

diff --git a/Engine/Source/Runtime/Engine/Private/SkeletalRenderCPUSkin.cpp b/Engine/Source/Runtime/Engine/Private/SkeletalRenderCPUSkin.cpp
--- a/Engine/Source/Runtime/Engine/Private/SkeletalRenderCPUSkin.cpp
+++ b/Engine/Source/Runtime/Engine/Private/SkeletalRenderCPUSkin.cpp
@@ -473,16 +473,7 @@ void FSkeletalMeshObjectCPUSkin::FSkeletalMeshObjectLOD::ReleaseResources()
 	BeginReleaseResource(&StaticMeshVertexBuffer);
 
 #if RHI_RAYTRACING
-	// BeginReleaseResource(&RayTracingGeometry);
-	// Workaround for UE-106993:
-	// Destroy ray tracing geometry on the render thread, as it may hold references to render resources.
-	// These references should be cleared in FRayTracingGeometry::ReleaseResource(), however FRayTracingGeometry does
-	// not implement this method and it can't be added due to 4.26 hotfix rules.
-	ENQUEUE_RENDER_COMMAND(ReleaseRayTracingGeometry)([Ptr = &RayTracingGeometry](FRHICommandListImmediate&)
-	{
-		Ptr->ReleaseResource();
-		*Ptr = FRayTracingGeometry(); // Explicitly reset all contents, including any resource references.
-	});
+	BeginReleaseResource(&RayTracingGeometry);
 #endif // RHI_RAYTRACING
 
 	bResourcesInitialized = false;
diff --git a/Engine/Source/Runtime/Engine/Private/SkeletalRenderGPUSkin.cpp b/Engine/Source/Runtime/Engine/Private/SkeletalRenderGPUSkin.cpp
--- a/Engine/Source/Runtime/Engine/Private/SkeletalRenderGPUSkin.cpp
+++ b/Engine/Source/Runtime/Engine/Private/SkeletalRenderGPUSkin.cpp
@@ -226,23 +226,14 @@ void FSkeletalMeshObjectGPUSkin::ReleaseResources()
 	);
 
 #if RHI_RAYTRACING
-	// BeginReleaseResource(&RayTracingGeometry);
-	// Workaround for UE-106993:
-	// Destroy ray tracing geometry on the render thread, as it may hold references to render resources.
-	// These references should be cleared in FRayTracingGeometry::ReleaseResource(), however FRayTracingGeometry does
-	// not implement this method and it can't be added due to 4.26 hotfix rules.
-	ENQUEUE_RENDER_COMMAND(ReleaseRayTracingGeometry)([Ptr = &RayTracingGeometry](FRHICommandListImmediate&)
-	{
-		Ptr->ReleaseResource();
-		*Ptr = FRayTracingGeometry(); // Explicitly reset all contents, including any resource references.
-	});
+	BeginReleaseResource(&RayTracingGeometry);
 
 	// Only enqueue when intialized
 	if (RayTracingDynamicVertexBuffer.NumBytes > 0)
 	{
 		ENQUEUE_RENDER_COMMAND(ReleaseRayTracingDynamicVertexBuffer)([&RayTracingDynamicVertexBuffer = RayTracingDynamicVertexBuffer](FRHICommandListImmediate& RHICmdList) mutable { RayTracingDynamicVertexBuffer.Release(); });
 	}
-#endif
+#endif // RHI_RAYTRACING
 }
 
 void FSkeletalMeshObjectGPUSkin::InitMorphResources(bool bInUsePerBoneMotionBlur, const TArray<float>& MorphTargetWeights)
diff --git a/Engine/Source/Runtime/Engine/Private/SkeletalRenderStatic.cpp b/Engine/Source/Runtime/Engine/Private/SkeletalRenderStatic.cpp
--- a/Engine/Source/Runtime/Engine/Private/SkeletalRenderStatic.cpp
+++ b/Engine/Source/Runtime/Engine/Private/SkeletalRenderStatic.cpp
@@ -223,16 +223,7 @@ void FSkeletalMeshObjectStatic::FSkeletalMeshObjectLOD::ReleaseResources()
 	BeginReleaseResource(&VertexFactory);
 
 #if RHI_RAYTRACING
-	// BeginReleaseResource(&RayTracingGeometry);
-	// Workaround for UE-106993:
-	// Destroy ray tracing geometry on the render thread, as it may hold references to render resources.
-	// These references should be cleared in FRayTracingGeometry::ReleaseResource(), however FRayTracingGeometry does
-	// not implement this method and it can't be added due to 4.26 hotfix rules.
-	ENQUEUE_RENDER_COMMAND(ReleaseRayTracingGeometry)([Ptr = &RayTracingGeometry](FRHICommandListImmediate&)
-	{
-		Ptr->ReleaseResource();
-		*Ptr = FRayTracingGeometry(); // Explicitly reset all contents, including any resource references.
-	});
+	BeginReleaseResource(&RayTracingGeometry);
 #endif // RHI_RAYTRACING
 
 	bResourcesInitialized = false;
diff --git a/Engine/Source/Runtime/Engine/Private/StaticMesh.cpp b/Engine/Source/Runtime/Engine/Private/StaticMesh.cpp
--- a/Engine/Source/Runtime/Engine/Private/StaticMesh.cpp
+++ b/Engine/Source/Runtime/Engine/Private/StaticMesh.cpp
@@ -1356,17 +1356,9 @@ void FStaticMeshLODResources::ReleaseResources()
 		BeginReleaseResource(&AdditionalIndexBuffers->WireframeIndexBuffer);
 		BeginReleaseResource(&AdditionalIndexBuffers->ReversedDepthOnlyIndexBuffer);
 	}
+
 #if RHI_RAYTRACING
-	// BeginReleaseResource(&RayTracingGeometry);
-	// Workaround for UE-106993:
-	// Destroy ray tracing geometry on the render thread, as it may hold references to render resources.
-	// These references should be cleared in FRayTracingGeometry::ReleaseResource(), however FRayTracingGeometry does
-	// not implement this method and it can't be added due to 4.26 hotfix rules.
-	ENQUEUE_RENDER_COMMAND(ReleaseRayTracingGeometry)([Ptr = &RayTracingGeometry](FRHICommandListImmediate&)
-	{
-		Ptr->ReleaseResource();
-		*Ptr = FRayTracingGeometry(); // Explicitly reset all contents, including any resource references.
-	});
+	BeginReleaseResource(&RayTracingGeometry);
 #endif // RHI_RAYTRACING
 }
 
diff --git a/Engine/Source/Runtime/RenderCore/Private/RenderResource.cpp b/Engine/Source/Runtime/RenderCore/Private/RenderResource.cpp
--- a/Engine/Source/Runtime/RenderCore/Private/RenderResource.cpp
+++ b/Engine/Source/Runtime/RenderCore/Private/RenderResource.cpp
@@ -416,6 +416,14 @@ void FRayTracingGeometry::CreateRayTracingGeometry(ERTAccelerationStructureBuild
 	}
 }
 
+void FRayTracingGeometry::InitRHI()
+{
+	if (!IsRayTracingEnabled())
+		return;
+
+	CreateRayTracingGeometry(ERTAccelerationStructureBuildPriority::Normal);
+}
+
 void FRayTracingGeometry::ReleaseRHI()
 {
 	if (HasPendingBuildRequest())
@@ -426,6 +434,15 @@ void FRayTracingGeometry::ReleaseRHI()
 	RayTracingGeometryRHI.SafeRelease();
 }
 
+void FRayTracingGeometry::ReleaseResource()
+{
+	// Release any resource references held by the initializer.
+	// This includes index and vertex buffers used for building the BLAS.
+	Initializer = FRayTracingGeometryInitializer {};
+
+	FRenderResource::ReleaseResource();
+}
+
 void FRayTracingGeometry::BoostBuildPriority(float InBoostValue) const
 {
 	check(HasPendingBuildRequest());
diff --git a/Engine/Source/Runtime/RenderCore/Public/RenderResource.h b/Engine/Source/Runtime/RenderCore/Public/RenderResource.h
--- a/Engine/Source/Runtime/RenderCore/Public/RenderResource.h
+++ b/Engine/Source/Runtime/RenderCore/Public/RenderResource.h
@@ -670,24 +670,12 @@ public:
 	FRayTracingGeometryInitializer Initializer;
 	FRayTracingGeometryRHIRef RayTracingGeometryRHI;
 
-	// FRenderResource interface.
-	virtual FString GetFriendlyName() const override { return TEXT("FRayTracingGeometry"); }
-
 	void SetInitializer(const FRayTracingGeometryInitializer& InInitializer)
 	{
 		Initializer = InInitializer;
 	}
 
-	virtual void InitRHI() override
-	{
-		if (!IsRayTracingEnabled())
-			return;
-
-		CreateRayTracingGeometry(ERTAccelerationStructureBuildPriority::Normal);
-	}
-
 	void CreateRayTracingGeometry(ERTAccelerationStructureBuildPriority InBuildPriority);
-	virtual void ReleaseRHI() override;
 
 	bool HasPendingBuildRequest() const
 	{
@@ -695,6 +683,15 @@ public:
 	}
 	void BoostBuildPriority(float InBoostValue = 0.01f) const;
 
+	// FRenderResource interface
+
+	virtual FString GetFriendlyName() const override { return TEXT("FRayTracingGeometry"); }
+
+	virtual void InitRHI() override;
+	virtual void ReleaseRHI() override;
+
+	virtual void ReleaseResource() override;
+
 protected:
 
 	friend class FRayTracingGeometryManager;
