Fixed propagation of bSupportRayTracing flag from skeletal mesh into skin cache.

diff --git a/Engine/Source/Runtime/Engine/Private/GPUSkinCache.cpp b/Engine/Source/Runtime/Engine/Private/GPUSkinCache.cpp
--- a/Engine/Source/Runtime/Engine/Private/GPUSkinCache.cpp
+++ b/Engine/Source/Runtime/Engine/Private/GPUSkinCache.cpp
@@ -1523,7 +1523,7 @@ void FGPUSkinCache::ProcessRayTracingGeometryToUpdate(
 	bool bAnySegmentUsesWorldPositionOffset
 	)
 {
-	if (IsRayTracingEnabled() && GEnableGPUSkinCache && SkinCacheEntry)
+	if (IsRayTracingEnabled() && GEnableGPUSkinCache && SkinCacheEntry && SkinCacheEntry->GPUSkin->bSupportRayTracing)
 	{
 		FRayTracingGeometry& RayTracingGeometry = SkinCacheEntry->GPUSkin->RayTracingGeometry;
 
diff --git a/Engine/Source/Runtime/Engine/Private/SkeletalRender.cpp b/Engine/Source/Runtime/Engine/Private/SkeletalRender.cpp
--- a/Engine/Source/Runtime/Engine/Private/SkeletalRender.cpp
+++ b/Engine/Source/Runtime/Engine/Private/SkeletalRender.cpp
@@ -31,6 +31,7 @@ FSkeletalMeshObject::FSkeletalMeshObject(USkinnedMeshComponent* InMeshComponent,
 ,   bHasBeenUpdatedAtLeastOnce(false)
 #if RHI_RAYTRACING
 , bRequireRecreatingRayTracingGeometry(false)
+, bSupportRayTracing(InMeshComponent->SkeletalMesh->bSupportRayTracing)
 #endif
 #if !UE_BUILD_SHIPPING
 , DebugName(InMeshComponent->SkeletalMesh->GetFName())
diff --git a/Engine/Source/Runtime/Engine/Private/SkeletalRenderGPUSkin.cpp b/Engine/Source/Runtime/Engine/Private/SkeletalRenderGPUSkin.cpp
--- a/Engine/Source/Runtime/Engine/Private/SkeletalRenderGPUSkin.cpp
+++ b/Engine/Source/Runtime/Engine/Private/SkeletalRenderGPUSkin.cpp
@@ -194,7 +194,10 @@ void FSkeletalMeshObjectGPUSkin::InitResources(USkinnedMeshComponent* InMeshComp
 	}
 
 #if RHI_RAYTRACING
-	BeginInitResource(&RayTracingGeometry);
+	if (bSupportRayTracing)
+	{
+		BeginInitResource(&RayTracingGeometry);
+	}
 #endif
 }
 void FSkeletalMeshObjectGPUSkin::ReleaseResources()
@@ -226,7 +229,10 @@ void FSkeletalMeshObjectGPUSkin::ReleaseResources()
 	);
 
 #if RHI_RAYTRACING
-	BeginReleaseResource(&RayTracingGeometry);
+	if (bSupportRayTracing)
+	{
+		BeginReleaseResource(&RayTracingGeometry);
+	}
 
 	// Only enqueue when intialized
 	if (RayTracingDynamicVertexBuffer.NumBytes > 0)
