Fix ULevelStreamingDynamic::LoadLevelInstance, which recently changed to use AssetRegistry::GetFirstPackageByName, to pass the PackageName in instead of the ObjectPath.
#rnx

diff --git a/Engine/Source/Runtime/Engine/Private/LevelStreaming.cpp b/Engine/Source/Runtime/Engine/Private/LevelStreaming.cpp
--- a/Engine/Source/Runtime/Engine/Private/LevelStreaming.cpp
+++ b/Engine/Source/Runtime/Engine/Private/LevelStreaming.cpp
@@ -1920,7 +1920,7 @@ void ULevelStreamingDynamic::SetShouldBeLoaded(const bool bInShouldBeLoaded)
 	}
 }
 
-ULevelStreamingDynamic* ULevelStreamingDynamic::LoadLevelInstance(UObject* WorldContextObject, const FString LevelName, const FVector Location, const FRotator Rotation, bool& bOutSuccess, const FString& OptionalLevelNameOverride, TSubclassOf<ULevelStreamingDynamic> OptionalLevelStreamingClass)
+ULevelStreamingDynamic* ULevelStreamingDynamic::LoadLevelInstance(UObject* WorldContextObject, const FString LevelObjectPath, const FVector Location, const FRotator Rotation, bool& bOutSuccess, const FString& OptionalLevelNameOverride, TSubclassOf<ULevelStreamingDynamic> OptionalLevelStreamingClass)
 {
 	bOutSuccess = false;
 	UWorld* const World = GEngine->GetWorldFromContextObject(WorldContextObject, EGetWorldErrorMode::LogAndReturnNull);
@@ -1930,16 +1930,41 @@ ULevelStreamingDynamic* ULevelStreamingDynamic::LoadLevelInstance(UObject* World
 	}
 
 	// Check whether requested map exists
+	// LevelObjectPath may be either an objectpath, a package name, or a file path; and packagename/object path may be a LongPackageName or ShortPackageName
+	// convert that flexible input to ObjectPath with a LongPackageName
+	FString PackageName;
+	FString ObjectRelativePath;
+	if (FPackageName::IsShortPackageName(LevelObjectPath) || FPackageName::IsValidObjectPath(LevelObjectPath))
+	{
+		FString UnusedClassName;
+		FString ObjectName;
+		FString SubObjectName;
+		FPackageName::SplitFullObjectPath(LevelObjectPath, UnusedClassName, PackageName, ObjectName, SubObjectName);
+		if (!ObjectName.IsEmpty())
+		{
+			ObjectRelativePath = FString(TEXT(".")) + ObjectName;
+			if (!SubObjectName.IsEmpty())
+			{
+				ObjectRelativePath += FString(SUBOBJECT_DELIMITER) + SubObjectName;
+			}
+		}
+	}
+	else if (!FPackageName::TryConvertFilenameToLongPackageName(LevelObjectPath, PackageName))
+	{
+		// An unrecognized path or format
+		return nullptr;
+	}
+
 	IAssetRegistry& AssetRegistry = FModuleManager::LoadModuleChecked<FAssetRegistryModule>(TEXT("AssetRegistry")).Get();
-	FName ExistingPackageName = AssetRegistry.GetFirstPackageByName(LevelName);
+	FName ExistingPackageName = AssetRegistry.GetFirstPackageByName(PackageName);
 	bOutSuccess = !ExistingPackageName.IsNone();
 	if (!bOutSuccess)
 	{
 		return nullptr;
 	}
-	FString LongPackageName = ExistingPackageName.ToString();
+	FString LongPackageNameObjectPath = ExistingPackageName.ToString() + ObjectRelativePath;
 
-	return LoadLevelInstance_Internal(World, LongPackageName, FTransform(Rotation, Location), bOutSuccess, OptionalLevelNameOverride, OptionalLevelStreamingClass);
+	return LoadLevelInstance_Internal(World, LongPackageNameObjectPath, FTransform(Rotation, Location), bOutSuccess, OptionalLevelNameOverride, OptionalLevelStreamingClass);
 }
 
 ULevelStreamingDynamic* ULevelStreamingDynamic::LoadLevelInstanceBySoftObjectPtr(UObject* WorldContextObject, const TSoftObjectPtr<UWorld> Level, const FVector Location, const FRotator Rotation, bool& bOutSuccess, const FString& OptionalLevelNameOverride, TSubclassOf<ULevelStreamingDynamic> OptionalLevelStreamingClass)
