Fix a crash related to adding virtual bones in combination with skeleton remapping.

#rnx

diff --git a/Engine/Source/Runtime/Engine/Private/Animation/Skeleton.cpp b/Engine/Source/Runtime/Engine/Private/Animation/Skeleton.cpp
--- a/Engine/Source/Runtime/Engine/Private/Animation/Skeleton.cpp
+++ b/Engine/Source/Runtime/Engine/Private/Animation/Skeleton.cpp
@@ -2257,6 +2257,8 @@ void USkeleton::HandleVirtualBoneChanges()
 	const bool bRebuildNameMap = false;
 	ReferenceSkeleton.RebuildRefSkeleton(this, bRebuildNameMap);
 
+	BuildSkeletonRemappings(true);
+
 	for (TObjectIterator<USkeletalMesh> ItMesh; ItMesh; ++ItMesh)
 	{
 		USkeletalMesh* SkelMesh = *ItMesh;
