Fix an ensure failure caused by generating VT platform data with virtual texture disabled.

diff --git a/Engine/Source/Runtime/Engine/Private/TextureDerivedData.cpp b/Engine/Source/Runtime/Engine/Private/TextureDerivedData.cpp
--- a/Engine/Source/Runtime/Engine/Private/TextureDerivedData.cpp
+++ b/Engine/Source/Runtime/Engine/Private/TextureDerivedData.cpp
@@ -37,6 +37,7 @@
 #include "Interfaces/ITextureFormat.h"
 #include "ProfilingDebugging/CookStats.h"
 #include "VT/VirtualTextureDataBuilder.h"
+#include "VT/LightmapVirtualTexture.h"
 #include "TextureCompiler.h"
 
 /*------------------------------------------------------------------------------
@@ -429,7 +430,9 @@ static void GetTextureBuildSettings(
 		);
 
 	static const auto CVarVirtualTexturesEnabled = IConsoleManager::Get().FindTConsoleVariableDataInt(TEXT("r.VirtualTextures")); check(CVarVirtualTexturesEnabled);
-	const bool bVirtualTextureStreaming = CVarVirtualTexturesEnabled->GetValueOnAnyThread() && bPlatformSupportsVirtualTextureStreaming && Texture.VirtualTextureStreaming;
+	// A ULightMapVirtualTexture2D with multiple layers saved in MapBuildData could be loaded with the r.VirtualTexture disabled, it will generate DDC before we decide to invalidate the light map data, to skip the ensure failure let it generate VT DDC anyway.
+	const bool bUseVirtualTexture = CVarVirtualTexturesEnabled->GetValueOnAnyThread() || ULightMapVirtualTexture2D::StaticClass() == Texture.GetClass();
+	const bool bVirtualTextureStreaming = bUseVirtualTexture && bPlatformSupportsVirtualTextureStreaming && Texture.VirtualTextureStreaming;
 	const FIntPoint SourceSize = Texture.Source.GetLogicalSize();
 
 	OutBuildSettings.MipGenSettings = MipGenSettings;
