Ensure that streaming levels that have their bRequireFullVisibilityToRender member set are completely hidden as soon as they reach the MakingInvisible state
* Without this change, depending on the amount of components to process, it was possible to see them dissapear slowly when the level was in this state
* Fixed bRequireFullVisibilityToRender not taken into account in SIE
* Fixed FPrimitiveSceneProxy::IsShadowCast() not taking into account the bRequireFullVisibilityToRender flag, resulting in shadows popping up before the actual meshes are displayed
* Note that there's still an issue related to Nanite meshes showing up even when this flag is set on the level. To be investigated.

diff --git a/Engine/Source/Runtime/Engine/Private/LevelStreaming.cpp b/Engine/Source/Runtime/Engine/Private/LevelStreaming.cpp
--- a/Engine/Source/Runtime/Engine/Private/LevelStreaming.cpp
+++ b/Engine/Source/Runtime/Engine/Private/LevelStreaming.cpp
@@ -615,15 +615,34 @@ void ULevelStreaming::UpdateStreamingState(bool& bOutUpdateAgain, bool& bOutRede
 	case ECurrentState::MakingInvisible:
 		if (ensure(LoadedLevel))
 		{
+			auto RemoveLevelFromScene = [World, this]()
+			{
+				if (World->Scene)
+				{
+					World->Scene->OnLevelRemovedFromWorld(LoadedLevel->GetOutermost()->GetFName(), World, LoadedLevel->bIsLightingScenario);
+				}
+			};
+
+			const bool bWasVisible = LoadedLevel->bIsVisible;
+
 			// Hide loaded level, incrementally if necessary
 			World->RemoveFromWorld(LoadedLevel, !bShouldBlockOnUnload && World->IsGameWorld());
 
-			// Inform the scene once we have finished making the level invisible
-			if (!LoadedLevel->bIsVisible)
+			// Hide loaded level immediately if bRequireFullVisibilityToRender is set
+			const bool LevelBecameInvisible = bWasVisible && !LoadedLevel->bIsVisible;
+			if (LoadedLevel->bRequireFullVisibilityToRender && LevelBecameInvisible)
 			{
-				if (World->Scene)
+				RemoveLevelFromScene();
+			}
+
+			// If the level is now hidden & all components have been removed from the world
+			const bool LevelWasRemovedFromWorld = !LoadedLevel->bIsVisible && !LoadedLevel->bIsBeingRemoved;
+			if (LevelWasRemovedFromWorld)
+			{
+				// Remove level from scene if we haven't done it already
+				if (!LoadedLevel->bRequireFullVisibilityToRender)
 				{
-					World->Scene->OnLevelRemovedFromWorld(World, LoadedLevel->bIsLightingScenario);
+					RemoveLevelFromScene();
 				}
 
 				CurrentState = ECurrentState::LoadedNotVisible;
diff --git a/Engine/Source/Runtime/Engine/Private/PrimitiveSceneProxy.cpp b/Engine/Source/Runtime/Engine/Private/PrimitiveSceneProxy.cpp
--- a/Engine/Source/Runtime/Engine/Private/PrimitiveSceneProxy.cpp
+++ b/Engine/Source/Runtime/Engine/Private/PrimitiveSceneProxy.cpp
@@ -793,6 +793,12 @@ bool FPrimitiveSceneProxy::IsShown(const FSceneView* View) const
 		return false;
 	}
 
+	// if primitive requires component level to be visible
+	if (bRequiresVisibleLevelToRender && !bIsComponentLevelVisible)
+	{
+		return false;
+	}
+
 	// After checking for VR/Desktop Edit mode specific actors, check for Editor vs. Game
 	if(View->Family->EngineShowFlags.Editor)
 	{
@@ -832,12 +838,6 @@ bool FPrimitiveSceneProxy::IsShown(const FSceneView* View) const
 			return false;
 		}
 
-		// if primitive requires component level to be visible
-		if (bRequiresVisibleLevelToRender && !bIsComponentLevelVisible)
-		{
-			return false;
-		}
-
 		if (DrawInVirtualTextureOnly(false) && !View->bIsVirtualTexture)
 		{
 			return false;
@@ -914,6 +914,12 @@ bool FPrimitiveSceneProxy::IsShadowCast(const FSceneView* View) const
 		}
 #endif	//#if WITH_EDITOR
 
+		// if primitive requires component level to be visible
+		if (bRequiresVisibleLevelToRender && !bIsComponentLevelVisible)
+		{
+			return false;
+		}
+
 		if (DrawInVirtualTextureOnly(View->Family->EngineShowFlags.Editor) && !View->bIsVirtualTexture)
 		{
 			return false;
diff --git a/Engine/Source/Runtime/Engine/Private/World.cpp b/Engine/Source/Runtime/Engine/Private/World.cpp
--- a/Engine/Source/Runtime/Engine/Private/World.cpp
+++ b/Engine/Source/Runtime/Engine/Private/World.cpp
@@ -2871,9 +2871,20 @@ void UWorld::RemoveFromWorld( ULevel* Level, bool bAllowIncrementalRemoval )
 
 	Level->bIsDisassociatingLevel = true;
 
+	auto BeginRemoval = [Level, this]()
+	{
+		Level->bIsBeingRemoved = true;
+
+		if (Level->bRequireFullVisibilityToRender)
+		{
+			Level->bIsVisible = false;
+			ULevelStreaming::BroadcastLevelVisibleStatus(this, Level->GetOutermost()->GetFName(), false);
+		}
+	};
+
 	// To be removed from the world a world must be visible and not pending being made visible (this may be redundent, but for safety)
 	// If the level may be removed incrementally then there must also be no level pending visibility
-	if ( ((CurrentLevelPendingVisibility == nullptr) || (!bAllowIncrementalRemoval && (CurrentLevelPendingVisibility != Level))) && Level->bIsVisible )
+	if ( ((CurrentLevelPendingVisibility == nullptr) || (!bAllowIncrementalRemoval && (CurrentLevelPendingVisibility != Level))) && (Level->bIsVisible || Level->bIsBeingRemoved) )
 	{
 		// Keep track of timing.
 		double StartTime = FPlatformTime::Seconds();	
@@ -2887,7 +2898,8 @@ void UWorld::RemoveFromWorld( ULevel* Level, bool bAllowIncrementalRemoval )
 				// Mark level as being the one in process of being made invisible. 
 				// This will prevent this level from being unloaded or made visible in the meantime
 				CurrentLevelPendingInvisibility = Level;
-				Level->bIsBeingRemoved = true;
+
+				BeginRemoval();
 			}
 
 			if ( CurrentLevelPendingInvisibility == Level )
@@ -2910,7 +2922,7 @@ void UWorld::RemoveFromWorld( ULevel* Level, bool bAllowIncrementalRemoval )
 		}
 		else
 		{
-			Level->bIsBeingRemoved = true;
+			BeginRemoval();
 		}
 
 		if ( bFinishRemovingLevel )
@@ -2966,6 +2978,8 @@ void UWorld::RemoveFromWorld( ULevel* Level, bool bAllowIncrementalRemoval )
 				}
 			}
 
+			// We expect level that use the bRequireFullVisibilityToRender flag to already be !visible at this point
+			check(!Level->bRequireFullVisibilityToRender || !Level->bIsVisible);
 			Level->bIsVisible = false;
 
 			// Notify world composition: will place a level at original position
@@ -2985,7 +2999,11 @@ void UWorld::RemoveFromWorld( ULevel* Level, bool bAllowIncrementalRemoval )
 			FWorldDelegates::LevelRemovedFromWorld.Broadcast(Level, this);
 			BroadcastLevelsChanged();
 
-			ULevelStreaming::BroadcastLevelVisibleStatus(this, Level->GetOutermost()->GetFName(), false);
+			// If the level requires full visibility to be rendered, we already made it non visible in BeginRemoval()
+			if (!Level->bRequireFullVisibilityToRender)
+			{
+				ULevelStreaming::BroadcastLevelVisibleStatus(this, Level->GetOutermost()->GetFName(), false);
+			}
 
 			Level->bIsBeingRemoved = false;
 		} // if ( bFinishRemovingLevel )
diff --git a/Engine/Source/Runtime/Engine/Public/SceneInterface.h b/Engine/Source/Runtime/Engine/Public/SceneInterface.h
--- a/Engine/Source/Runtime/Engine/Public/SceneInterface.h
+++ b/Engine/Source/Runtime/Engine/Public/SceneInterface.h
@@ -489,7 +489,7 @@ public:
 	 * @param	InLevelName		Level name
 	 */
 	virtual void OnLevelAddedToWorld(FName InLevelName, UWorld* InWorld, bool bIsLightingScenario) {}
-	virtual void OnLevelRemovedFromWorld(UWorld* InWorld, bool bIsLightingScenario) {}
+	virtual void OnLevelRemovedFromWorld(FName InLevelName, UWorld* InWorld, bool bIsLightingScenario) {}
 
 	/**
 	 * @return True if there are any lights in the scene
diff --git a/Engine/Source/Runtime/Renderer/Private/RendererScene.cpp b/Engine/Source/Runtime/Renderer/Private/RendererScene.cpp
--- a/Engine/Source/Runtime/Renderer/Private/RendererScene.cpp
+++ b/Engine/Source/Runtime/Renderer/Private/RendererScene.cpp
@@ -3662,6 +3662,8 @@ void FScene::OnLevelAddedToWorld(FName LevelAddedName, UWorld* InWorld, bool bIs
 
 void FScene::OnLevelAddedToWorld_RenderThread(FName InLevelName)
 {
+	TRACE_CPUPROFILER_EVENT_SCOPE(FScene::OnLevelAddedToWorld_RenderThread);
+
 	// Mark level primitives
 	TArray<FPrimitiveSceneInfo*> PrimitivesToAdd;
 
@@ -3684,6 +3686,38 @@ void FScene::OnLevelAddedToWorld_RenderThread(FName InLevelName)
 	FPrimitiveSceneInfo::AddStaticMeshes(FRHICommandListExecutor::GetImmediateCommandList(), this, PrimitivesToAdd);
 }
 
+void FScene::OnLevelRemovedFromWorld(FName LevelRemovedName, UWorld* InWorld, bool bIsLightingScenario)
+{
+	if (bIsLightingScenario)
+	{
+		InWorld->PropagateLightingScenarioChange();
+	}
+
+	FScene* Scene = this;
+	ENQUEUE_RENDER_COMMAND(FLevelRemovedFromWorld)(
+		[Scene, LevelRemovedName](FRHICommandListImmediate& RHICmdList)
+		{
+			Scene->UpdateAllPrimitiveSceneInfos(RHICmdList);
+			Scene->OnLevelRemovedFromWorld_RenderThread(LevelRemovedName);
+		});
+}
+
+
+void FScene::OnLevelRemovedFromWorld_RenderThread(FName InLevelName)
+{
+	TRACE_CPUPROFILER_EVENT_SCOPE(FScene::OnLevelRemovedFromWorld_RenderThread);
+
+	for (auto It = Primitives.CreateIterator(); It; ++It)
+	{
+		FPrimitiveSceneProxy* Proxy = (*It)->Proxy;
+		if (Proxy->LevelName == InLevelName)
+		{
+			Proxy->bIsComponentLevelVisible = false;
+		}
+	}
+}
+
+
 void FScene::ProcessAtmosphereLightAddition_RenderThread(FLightSceneInfo* LightSceneInfo)
 {
 	if (LightSceneInfo->Proxy->IsUsedAsAtmosphereSunLight())
@@ -3722,14 +3756,6 @@ void FScene::ProcessAtmosphereLightRemoval_RenderThread(FLightSceneInfo* LightSc
 	}
 }
 
-void FScene::OnLevelRemovedFromWorld(UWorld* InWorld, bool bIsLightingScenario)
-{
-	if (bIsLightingScenario)
-	{
-		InWorld->PropagateLightingScenarioChange();
-	}
-}
-
 #if WITH_EDITOR
 bool FScene::InitializePixelInspector(FRenderTarget* BufferFinalColor, FRenderTarget* BufferSceneColor, FRenderTarget* BufferDepth, FRenderTarget* BufferHDR, FRenderTarget* BufferA, FRenderTarget* BufferBCDEF, int32 BufferIndex)
 {
diff --git a/Engine/Source/Runtime/Renderer/Private/ScenePrivate.h b/Engine/Source/Runtime/Renderer/Private/ScenePrivate.h
--- a/Engine/Source/Runtime/Renderer/Private/ScenePrivate.h
+++ b/Engine/Source/Runtime/Renderer/Private/ScenePrivate.h
@@ -3103,7 +3103,7 @@ public:
 	virtual void ApplyWorldOffset(FVector InOffset) override;
 
 	virtual void OnLevelAddedToWorld(FName InLevelName, UWorld* InWorld, bool bIsLightingScenario) override;
-	virtual void OnLevelRemovedFromWorld(UWorld* InWorld, bool bIsLightingScenario) override;
+	virtual void OnLevelRemovedFromWorld(FName LevelRemovedName, UWorld* InWorld, bool bIsLightingScenario) override;
 
 	virtual bool HasAnyLights() const override 
 	{ 
@@ -3282,6 +3282,13 @@ private:
 	 */
 	void OnLevelAddedToWorld_RenderThread(FName InLevelName);
 
+	/**
+	 * Notification from game thread that level was removed from a world
+	 *
+	 * @param	InLevelName		Level name
+	 */
+	void OnLevelRemovedFromWorld_RenderThread(FName InLevelName);
+
 	void ProcessAtmosphereLightRemoval_RenderThread(FLightSceneInfo* LightSceneInfo);
 	void ProcessAtmosphereLightAddition_RenderThread(FLightSceneInfo* LightSceneInfo);
 
