LevelInstance: Fix Embedded mode

diff --git a/Engine/Source/Runtime/Engine/Private/WorldPartition/WorldPartitionActorCluster.cpp b/Engine/Source/Runtime/Engine/Private/WorldPartition/WorldPartitionActorCluster.cpp
--- a/Engine/Source/Runtime/Engine/Private/WorldPartition/WorldPartitionActorCluster.cpp
+++ b/Engine/Source/Runtime/Engine/Private/WorldPartition/WorldPartitionActorCluster.cpp
@@ -278,7 +278,7 @@ void CreateActorCluster(const FWorldPartitionActorDescView& ActorDescView, TMap<
 	}
 }
 
-void FActorClusterContext::CreateContainerInstanceRecursive(uint32 ID, const FTransform& Transform, EContainerClusterMode ClusterMode, const UActorDescContainer* Container, const TSet<FName>& DataLayers, FBox* ParentBounds)
+void FActorClusterContext::CreateContainerInstanceRecursive(uint32 ID, const FTransform& Transform, EContainerClusterMode ClusterMode, const UActorDescContainer* Container, const TSet<FName>& DataLayers, FBox& ParentBounds)
 {
 	InstanceCountHint += Container->GetActorDescCount();
 		
@@ -311,7 +311,7 @@ void FActorClusterContext::CreateContainerInstanceRecursive(uint32 ID, const FTr
 			}
 			// Always inherite parent container DataLayers
 			ChildDataLayers.Append(DataLayers);
-			CreateContainerInstanceRecursive(ChildContainerHashBuilder.GetHash(), OutTransform * Transform, OutClusterMode, OutContainer, ChildDataLayers, &Bounds);
+			CreateContainerInstanceRecursive(ChildContainerHashBuilder.GetHash(), OutTransform * Transform, OutClusterMode, OutContainer, ChildDataLayers, Bounds);
 		}
 		else
 		{
@@ -332,11 +332,8 @@ void FActorClusterContext::CreateContainerInstanceRecursive(uint32 ID, const FTr
 		}
 	}
 	
-	if (ParentBounds)
-	{
-		*ParentBounds += Bounds;
-	}
-
+	ParentBounds += Bounds;
+	
 	UE_LOG(LogWorldPartition, Verbose, TEXT("ContainerInstance (%08x) Bounds (%s) Package (%s)"), ID, *Bounds.TransformBy(Transform).ToString(), *Container->GetContainerPackage().ToString());
 	ContainerInstances.Add(FActorContainerInstance(ID, Transform, Bounds, DataLayers, ClusterMode, Container, MoveTemp(ChildContainers), MoveTemp(ActorDescViewMap)));
 }
@@ -385,11 +382,14 @@ const TArray<FActorCluster>& FActorClusterContext::CreateActorClustersImpl(const
 void FActorClusterContext::CreateActorClusters()
 {
 	// First Instance is the main WorldPartition
-	CreateContainerInstanceRecursive(0, FTransform::Identity, EContainerClusterMode::Partitioned, WorldPartition, TSet<FName>());
-			
+	FBox WorldBounds(ForceInit);
+	CreateContainerInstanceRecursive(0, FTransform::Identity, EContainerClusterMode::Partitioned, WorldPartition, TSet<FName>(), WorldBounds);
+		
 	ClusterInstances.Reserve(InstanceCountHint);
-	for (const FActorContainerInstance& ContainerInstance : ContainerInstances)
+	for (FActorContainerInstance& ContainerInstance : ContainerInstances)
 	{
+		RuntimeHash->UpdateActorDescViewMap(WorldBounds, ContainerInstance.ActorDescViewMap);
+
 		const TArray<FActorCluster>& NewClusters = CreateActorClustersImpl(ContainerInstance);
 		for (const FActorCluster& Cluster : NewClusters)
 		{
diff --git a/Engine/Source/Runtime/Engine/Private/WorldPartition/WorldPartitionLevelStreamingDynamic.cpp b/Engine/Source/Runtime/Engine/Private/WorldPartition/WorldPartitionLevelStreamingDynamic.cpp
--- a/Engine/Source/Runtime/Engine/Private/WorldPartition/WorldPartitionLevelStreamingDynamic.cpp
+++ b/Engine/Source/Runtime/Engine/Private/WorldPartition/WorldPartitionLevelStreamingDynamic.cpp
@@ -225,15 +225,14 @@ bool UWorldPartitionLevelStreamingDynamic::IssueLoadRequests()
 	UPackage* RuntimePackage = RuntimeLevel->GetPackage();
 	InstancingContext.AddMapping(OriginalLevelPackageName, RuntimePackage->GetFName());
 
-	TArray<FWorldPartitionRuntimeCellObjectMapping> ChildPackagesToLoad;
-	ChildPackagesToLoad.Reserve(ChildPackages.Num());
+	ChildPackagesToLoad.Reset(ChildPackages.Num());
 
 	UWorld* World = GetWorld();
 	for (FWorldPartitionRuntimeCellObjectMapping& ChildPackage : ChildPackages)
 	{
+		bool bNeedDup = false;
 		if (ChildPackage.ContainerID == 0)
 		{
-			bool bNeedDup = false;
 			if (ActorContainer)
 			{
 				FString SubObjectName;
@@ -246,11 +245,11 @@ bool UWorldPartitionLevelStreamingDynamic::IssueLoadRequests()
 					}
 				}
 			}
-			
-			if(!bNeedDup)
-			{
-				ChildPackagesToLoad.Add(ChildPackage);
-			}
+		}
+
+		if (!bNeedDup)
+		{
+			ChildPackagesToLoad.Add(ChildPackage);
 		}
 	}
 
diff --git a/Engine/Source/Runtime/Engine/Private/WorldPartition/WorldPartitionRuntimeSpatialHash.cpp b/Engine/Source/Runtime/Engine/Private/WorldPartition/WorldPartitionRuntimeSpatialHash.cpp
--- a/Engine/Source/Runtime/Engine/Private/WorldPartition/WorldPartitionRuntimeSpatialHash.cpp
+++ b/Engine/Source/Runtime/Engine/Private/WorldPartition/WorldPartitionRuntimeSpatialHash.cpp
@@ -632,37 +632,9 @@ FName UWorldPartitionRuntimeSpatialHash::GetCellName(FName InGridName, const FIn
 	return UWorldPartitionRuntimeSpatialHash::GetCellName(WorldPartition, InGridName, InCellGlobalCoord, InDataLayerID);
 }
 
-void UWorldPartitionRuntimeSpatialHash::CreateActorDescViewMap(const UActorDescContainer* Container, TMap<FGuid, FWorldPartitionActorDescView>& OutActorDescViewMap) const
+void UWorldPartitionRuntimeSpatialHash::UpdateActorDescViewMap(const FBox& WorldBounds, TMap<FGuid, FWorldPartitionActorDescView>& ActorDescViewMap) const
 {
-	Super::CreateActorDescViewMap(Container, OutActorDescViewMap);
-
-	// Compute the initial world bounds.
-	// This bound will change after we update actor descriptor views grid placement.
-	auto GetWorldBounds = [](const TMap<FGuid, FWorldPartitionActorDescView> ActorDescViewMap)
-	{
-		FBox WorldBounds(ForceInit);
-		for (auto& ActorDescViewPair : ActorDescViewMap)
-		{
-			const FWorldPartitionActorDescView& ActorDescView = ActorDescViewPair.Value;
-			switch (ActorDescView.GetGridPlacement())
-			{
-				case EActorGridPlacement::Location:
-				{
-					FVector Location = ActorDescView.GetOrigin();
-					WorldBounds += FBox(Location, Location);
-				}
-				break;
-				case EActorGridPlacement::Bounds:
-				{
-					WorldBounds += ActorDescView.GetBounds();
-				}
-				break;
-			}
-		}
-		return WorldBounds;
-	};
-
- 	const FBox WorldBounds = GetWorldBounds(OutActorDescViewMap);
+	Super::UpdateActorDescViewMap(WorldBounds, ActorDescViewMap);
 
 	TMap<FName, int32> GridsMapping;
 	GridsMapping.Add(NAME_None, 0);
@@ -676,7 +648,7 @@ void UWorldPartitionRuntimeSpatialHash::CreateActorDescViewMap(const UActorDescC
 	const FBox2D WorldBounds2D = FBox2D(FVector2D(WorldBounds.Min), FVector2D(WorldBounds.Max));
 	const float WorldBoundsArea = WorldBounds2D.GetArea();
 
-	for (auto& ActorDescViewPair : OutActorDescViewMap)
+	for (auto& ActorDescViewPair : ActorDescViewMap)
 	{
 		FWorldPartitionActorDescView& ActorDescView = ActorDescViewPair.Value;
 
diff --git a/Engine/Source/Runtime/Engine/Public/WorldPartition/WorldPartitionActorCluster.h b/Engine/Source/Runtime/Engine/Public/WorldPartition/WorldPartitionActorCluster.h
--- a/Engine/Source/Runtime/Engine/Public/WorldPartition/WorldPartitionActorCluster.h
+++ b/Engine/Source/Runtime/Engine/Public/WorldPartition/WorldPartitionActorCluster.h
@@ -118,7 +118,7 @@ public:
 
 private:
 	void CreateActorClusters();
-	void CreateContainerInstanceRecursive(uint32 ID, const FTransform& Transform, EContainerClusterMode ClusterMode, const UActorDescContainer* ActorDescContainer, const TSet<FName>& DataLayers, FBox* ParentBounds = nullptr);
+	void CreateContainerInstanceRecursive(uint32 ID, const FTransform& Transform, EContainerClusterMode ClusterMode, const UActorDescContainer* ActorDescContainer, const TSet<FName>& DataLayers, FBox& ParentBounds);
 	const TArray<FActorCluster>& CreateActorClustersImpl(const FActorContainerInstance& ContainerInstance);
 	
 	// Init data
diff --git a/Engine/Source/Runtime/Engine/Public/WorldPartition/WorldPartitionLevelStreamingDynamic.h b/Engine/Source/Runtime/Engine/Public/WorldPartition/WorldPartitionLevelStreamingDynamic.h
--- a/Engine/Source/Runtime/Engine/Public/WorldPartition/WorldPartitionLevelStreamingDynamic.h
+++ b/Engine/Source/Runtime/Engine/Public/WorldPartition/WorldPartitionLevelStreamingDynamic.h
@@ -49,6 +49,7 @@ private:
 
 	FName OriginalLevelPackageName;
 	TArray<FWorldPartitionRuntimeCellObjectMapping> ChildPackages;
+	TArray<FWorldPartitionRuntimeCellObjectMapping> ChildPackagesToLoad;
 
 #if WITH_EDITORONLY_DATA
 	UPROPERTY()
diff --git a/Engine/Source/Runtime/Engine/Public/WorldPartition/WorldPartitionRuntimeHash.h b/Engine/Source/Runtime/Engine/Public/WorldPartition/WorldPartitionRuntimeHash.h
--- a/Engine/Source/Runtime/Engine/Public/WorldPartition/WorldPartitionRuntimeHash.h
+++ b/Engine/Source/Runtime/Engine/Public/WorldPartition/WorldPartitionRuntimeHash.h
@@ -63,6 +63,7 @@ protected:
 
 	virtual bool GenerateStreaming(EWorldPartitionStreamingMode Mode, class UWorldPartitionStreamingPolicy* Policy, TArray<FString>* OutPackagesToGenerate) { return false; }
 	virtual void CreateActorDescViewMap(const UActorDescContainer* Container, TMap<FGuid, FWorldPartitionActorDescView>& OutActorDescViewMap) const;
+	virtual void UpdateActorDescViewMap(const FBox& WorldBounds, TMap<FGuid, FWorldPartitionActorDescView>& ActorDescViewMap) const {}
 	void ChangeActorDescViewGridPlacement(FWorldPartitionActorDescView& ActorDescView, EActorGridPlacement GridPlacement) const;	
 #endif
 
diff --git a/Engine/Source/Runtime/Engine/Public/WorldPartition/WorldPartitionRuntimeSpatialHash.h b/Engine/Source/Runtime/Engine/Public/WorldPartition/WorldPartitionRuntimeSpatialHash.h
--- a/Engine/Source/Runtime/Engine/Public/WorldPartition/WorldPartitionRuntimeSpatialHash.h
+++ b/Engine/Source/Runtime/Engine/Public/WorldPartition/WorldPartitionRuntimeSpatialHash.h
@@ -212,7 +212,7 @@ public:
 	FName GetCellName(FName InGridName, const FIntVector& InCellGlobalCoord, const FDataLayersID& InDataLayerID) const;
 	static FName GetCellName(UWorldPartition* WorldPartition, FName InGridName, const FIntVector& InCellGlobalCoord, const FDataLayersID& InDataLayerID);
 
-	virtual void CreateActorDescViewMap(const UActorDescContainer* Container, TMap<FGuid, FWorldPartitionActorDescView>& OutActorDescViewMap) const override;
+	virtual void UpdateActorDescViewMap(const FBox& WorldBounds, TMap<FGuid, FWorldPartitionActorDescView>& ActorDescViewMap) const override;
 #endif
 
 	// streaming interface
