Relax the restriction that getting an array element in script can only work if the field classes of the destination and the array's inner property are identical.  Allow interrelated classes of the same size to successfully produce output when getting elements out of an array.
Same change made to Map Find operations from script.

diff --git a/Engine/Source/Runtime/Engine/Classes/Kismet/BlueprintMapLibrary.h b/Engine/Source/Runtime/Engine/Classes/Kismet/BlueprintMapLibrary.h
--- a/Engine/Source/Runtime/Engine/Classes/Kismet/BlueprintMapLibrary.h
+++ b/Engine/Source/Runtime/Engine/Classes/Kismet/BlueprintMapLibrary.h
@@ -214,7 +214,19 @@ public:
 		
 		Stack.MostRecentPropertyAddress = NULL;
 		Stack.StepCompiledIn<FProperty>(ValueStorageSpace);
-		void* ItemPtr = (Stack.MostRecentPropertyAddress != NULL && Stack.MostRecentProperty->GetClass() == CurrValueProp->GetClass()) ? Stack.MostRecentPropertyAddress : ValueStorageSpace;
+		const FFieldClass* CurrValuePropClass = CurrValueProp->GetClass();
+		const FFieldClass* MostRecentPropClass = Stack.MostRecentProperty->GetClass();
+		void* ItemPtr;
+		// If the destination and the inner type are identical in size and their field classes derive from one another, then permit the writing out of the array element to the destination memory
+		if (Stack.MostRecentPropertyAddress != NULL && (ValuePropertySize == Stack.MostRecentProperty->ElementSize*Stack.MostRecentProperty->ArrayDim) &&
+			(MostRecentPropClass->IsChildOf(CurrValuePropClass) || CurrValuePropClass->IsChildOf(MostRecentPropClass)))
+		{
+			ItemPtr = Stack.MostRecentPropertyAddress;
+		}
+		else
+		{
+			ItemPtr = ValueStorageSpace;
+		}
 		
 		P_FINISH;
 		P_NATIVE_BEGIN;
diff --git a/Engine/Source/Runtime/Engine/Classes/Kismet/KismetArrayLibrary.h b/Engine/Source/Runtime/Engine/Classes/Kismet/KismetArrayLibrary.h
--- a/Engine/Source/Runtime/Engine/Classes/Kismet/KismetArrayLibrary.h
+++ b/Engine/Source/Runtime/Engine/Classes/Kismet/KismetArrayLibrary.h
@@ -655,7 +655,19 @@ public:
 
 		Stack.MostRecentPropertyAddress = NULL;
 		Stack.StepCompiledIn<FProperty>(StorageSpace);
-		void* ItemPtr = (Stack.MostRecentPropertyAddress != NULL && Stack.MostRecentProperty->GetClass() == InnerProp->GetClass()) ? Stack.MostRecentPropertyAddress : StorageSpace;
+		const FFieldClass* InnerPropClass = InnerProp->GetClass();
+		const FFieldClass* MostRecentPropClass = Stack.MostRecentProperty->GetClass();
+		void* ItemPtr;
+		// If the destination and the inner type are identical in size and their field classes derive from one another, then permit the writing out of the array element to the destination memory
+		if (Stack.MostRecentPropertyAddress != NULL && (PropertySize == Stack.MostRecentProperty->ElementSize*Stack.MostRecentProperty->ArrayDim) &&
+			(MostRecentPropClass->IsChildOf(InnerPropClass) || InnerPropClass->IsChildOf(MostRecentPropClass)))
+		{
+			ItemPtr = Stack.MostRecentPropertyAddress;
+		}
+		else
+		{
+			ItemPtr = StorageSpace;
+		}
 
 		P_FINISH;
 		P_NATIVE_BEGIN;
