Don't set IsSaving in FArchiveGatherExternalActorRefs to avoid having custom serialization fixup code running which could mark the map package dirty+ don't recurse into assets.

diff --git a/Engine/Source/Runtime/Engine/Private/ActorReferencesUtils.cpp b/Engine/Source/Runtime/Engine/Private/ActorReferencesUtils.cpp
--- a/Engine/Source/Runtime/Engine/Private/ActorReferencesUtils.cpp
+++ b/Engine/Source/Runtime/Engine/Private/ActorReferencesUtils.cpp
@@ -12,16 +12,16 @@ public:
 		, ActorReferences(InActorReferences)
 		, RequiredFlags(InRequiredFlags)
 	{
-		SetIsSaving(true);
+		// Don't gather transient actor references
 		SetIsPersistent(true);
+
+		// Don't trigger serialization of compilable assets
+		SetShouldSkipCompilingAssets(true);
+
 		ArIgnoreOuterRef = true;
 		ArIsObjectReferenceCollector = true;
 		ArShouldSkipBulkData = true;
 
-		// Avoid having to wait until compilable asset are built just to scan for actors. 
-		// Assets like Textures/Shaders/StaticMesh should not refer to actors anyway.
-		SetShouldSkipCompilingAssets(true);
-
 		SubObjects.Add(Root);
 
 		Root->Serialize(*this);
