On world's BeginPlay, partitioned world will block till level streaming is completed.

diff --git a/Engine/Source/Runtime/Engine/Private/WorldPartition/ActorDescContainer.cpp b/Engine/Source/Runtime/Engine/Private/WorldPartition/ActorDescContainer.cpp
--- a/Engine/Source/Runtime/Engine/Private/WorldPartition/ActorDescContainer.cpp
+++ b/Engine/Source/Runtime/Engine/Private/WorldPartition/ActorDescContainer.cpp
@@ -17,7 +17,7 @@ UActorDescContainer::UActorDescContainer(const FObjectInitializer& ObjectInitial
 #endif
 {}
 
-void UActorDescContainer::Initialize(UWorld* InWorld, FName InPackageName, bool bRegisterDelegates)
+void UActorDescContainer::Initialize(UWorld* InWorld, FName InPackageName)
 {
 	check(!World || World == InWorld);
 	World = InWorld;
@@ -96,10 +96,7 @@ void UActorDescContainer::Initialize(UWorld* InWorld, FName InPackageName, bool
 		}
 	}
 
-	if (bRegisterDelegates)
-	{
-		RegisterDelegates();
-	}
+	RegisterEditorDelegates();
 
 	bContainerInitialized = true;
 #endif
@@ -107,11 +104,10 @@ void UActorDescContainer::Initialize(UWorld* InWorld, FName InPackageName, bool
 
 void UActorDescContainer::Uninitialize()
 {
-	World = nullptr;
 #if WITH_EDITOR
 	if (bContainerInitialized)
 	{
-		UnregisterDelegates();
+		UnregisterEditorDelegates();
 		bContainerInitialized = false;
 	}
 
@@ -120,6 +116,7 @@ void UActorDescContainer::Uninitialize()
 		ActorDescPtr.Reset();
 	}
 #endif
+	World = nullptr;
 }
 
 UWorld* UActorDescContainer::GetWorld() const
@@ -219,18 +216,18 @@ void UActorDescContainer::RemoveActor(const FGuid& ActorGuid)
 	}
 }
 
-void UActorDescContainer::RegisterDelegates()
+void UActorDescContainer::RegisterEditorDelegates()
 {
-	if (GEditor && !IsTemplate())
+	if (GEditor && !IsTemplate() && !World->IsGameWorld())
 	{
 		FCoreUObjectDelegates::OnObjectSaved.AddUObject(this, &UActorDescContainer::OnObjectPreSave);
 		FEditorDelegates::OnPackageDeleted.AddUObject(this, &UActorDescContainer::OnPackageDeleted);
 	}
 }
 
-void UActorDescContainer::UnregisterDelegates()
+void UActorDescContainer::UnregisterEditorDelegates()
 {
-	if (GEditor && !IsTemplate())
+	if (GEditor && !IsTemplate() && !World->IsGameWorld())
 	{
 		FCoreUObjectDelegates::OnObjectSaved.RemoveAll(this);
 		FEditorDelegates::OnPackageDeleted.RemoveAll(this);
diff --git a/Engine/Source/Runtime/Engine/Private/WorldPartition/WorldPartition.cpp b/Engine/Source/Runtime/Engine/Private/WorldPartition/WorldPartition.cpp
--- a/Engine/Source/Runtime/Engine/Private/WorldPartition/WorldPartition.cpp
+++ b/Engine/Source/Runtime/Engine/Private/WorldPartition/WorldPartition.cpp
@@ -259,6 +259,8 @@ void UWorldPartition::Initialize(UWorld* InWorld, const FTransform& InTransform)
 
 	check(IsMainWorldPartition());
 
+	RegisterDelegates();
+
 #if WITH_EDITOR
 	bool bEditorOnly = !World->IsGameWorld();
 	if (bEditorOnly)
@@ -284,9 +286,10 @@ void UWorldPartition::Initialize(UWorld* InWorld, const FTransform& InTransform)
 
 	if (bEditorOnly || IsRunningGame())
 	{
-		TArray<FAssetData> Assets;
 		UPackage* LevelPackage = OuterWorld->PersistentLevel->GetOutermost();
-		FName PackageName = LevelPackage->GetLoadedPath().GetPackageFName();
+		const FName PackageName = LevelPackage->GetLoadedPath().GetPackageFName();
+		UActorDescContainer::Initialize(World, PackageName);
+		check(bContainerInitialized);
 
 		bool bIsInstanced = (bEditorOnly && !IsRunningCommandlet()) ? OuterWorld->PersistentLevel->IsInstancedLevel() : false;
 
@@ -295,7 +298,7 @@ void UWorldPartition::Initialize(UWorld* InWorld, const FTransform& InTransform)
 
 		if (bIsInstanced)
 		{
-			InstancingContext.AddMapping(LevelPackage->GetLoadedPath().GetPackageFName(), LevelPackage->GetFName());
+			InstancingContext.AddMapping(PackageName, LevelPackage->GetFName());
 
 			const FString SourceWorldName = FPaths::GetBaseFilename(LevelPackage->GetLoadedPath().GetPackageName());
 			const FString DestWorldName = FPaths::GetBaseFilename(LevelPackage->GetFName().ToString());
@@ -304,10 +307,6 @@ void UWorldPartition::Initialize(UWorld* InWorld, const FTransform& InTransform)
 			ReplaceTo = DestWorldName + TEXT(".") + DestWorldName;
 		}
 
-		const bool bRegisterDelegates = bEditorOnly;
-		UActorDescContainer::Initialize(World, PackageName, bRegisterDelegates);
-		check(bContainerInitialized);
-
 		for (UActorDescContainer::TIterator<> ActorDescIterator(this); ActorDescIterator; ++ActorDescIterator)
 		{
 			if (bIsInstanced)
@@ -435,6 +434,8 @@ void UWorldPartition::Uninitialize()
 		check(World);
 
 		InitState = EWorldPartitionInitState::Uninitializing;
+
+		UnregisterDelegates();
 		
 		// Unload all loaded cells
 		if (World->IsGameWorld())
@@ -485,8 +486,6 @@ void UWorldPartition::Uninitialize()
 			WorldPartitionSubsystem->UnregisterWorldPartition(this);
 		}
 
-		World = nullptr;
-
 		InitState = EWorldPartitionInitState::Uninitialized;
 	}
 
@@ -509,6 +508,53 @@ bool UWorldPartition::IsMainWorldPartition() const
 	return World == GetTypedOuter<UWorld>();
 }
 
+void UWorldPartition::RegisterDelegates()
+{
+	check(World); 
+
+#if WITH_EDITOR
+	if (GEditor && !IsTemplate() && !World->IsGameWorld())
+	{
+		FEditorDelegates::PreBeginPIE.AddUObject(this, &UWorldPartition::OnPreBeginPIE);
+		FEditorDelegates::PrePIEEnded.AddUObject(this, &UWorldPartition::OnPrePIEEnded);
+		FEditorDelegates::CancelPIE.AddUObject(this, &UWorldPartition::OnCancelPIE);
+		FGameDelegates::Get().GetEndPlayMapDelegate().AddUObject(this, &UWorldPartition::OnEndPlay);
+	}
+#endif
+
+	if (World->IsGameWorld())
+	{
+		World->OnWorldBeginPlay.AddUObject(this, &UWorldPartition::OnWorldBeginPlay);
+	}
+}
+
+void UWorldPartition::UnregisterDelegates()
+{
+	check(World);
+
+#if WITH_EDITOR
+	if (GEditor && !IsTemplate() && !World->IsGameWorld())
+	{
+		FEditorDelegates::PreBeginPIE.RemoveAll(this);
+		FEditorDelegates::PrePIEEnded.RemoveAll(this);
+		FEditorDelegates::CancelPIE.RemoveAll(this);
+		FGameDelegates::Get().GetEndPlayMapDelegate().RemoveAll(this);
+	}
+#endif
+
+	if (World->IsGameWorld())
+	{
+		World->OnWorldBeginPlay.RemoveAll(this);
+	}
+}
+
+void UWorldPartition::OnWorldBeginPlay()
+{
+	check(GetWorld()->IsGameWorld());
+	// Wait for any level streaming to complete
+	GetWorld()->BlockTillLevelStreamingCompleted();
+}
+
 #if WITH_EDITOR
 UWorldPartition* UWorldPartition::CreateWorldPartition(AWorldSettings* WorldSettings, TSubclassOf<UWorldPartitionEditorHash> EditorHashClass, TSubclassOf<UWorldPartitionRuntimeHash> RuntimeHashClass)
 {
@@ -538,30 +584,6 @@ UWorldPartition* UWorldPartition::CreateWorldPartition(AWorldSettings* WorldSett
 	return WorldPartition;
 }
 
-void UWorldPartition::RegisterDelegates()
-{
-	Super::RegisterDelegates();
-	if (GEditor && !IsTemplate())
-	{
-		FEditorDelegates::PreBeginPIE.AddUObject(this, &UWorldPartition::OnPreBeginPIE);
-		FEditorDelegates::PrePIEEnded.AddUObject(this, &UWorldPartition::OnPrePIEEnded);
-		FEditorDelegates::CancelPIE.AddUObject(this, &UWorldPartition::OnCancelPIE);
-		FGameDelegates::Get().GetEndPlayMapDelegate().AddUObject(this, &UWorldPartition::OnEndPlay);
-	}
-}
-
-void UWorldPartition::UnregisterDelegates()
-{
-	Super::UnregisterDelegates();
-	if (GEditor && !IsTemplate())
-	{
-		FEditorDelegates::PreBeginPIE.RemoveAll(this);
-		FEditorDelegates::PrePIEEnded.RemoveAll(this);
-		FEditorDelegates::CancelPIE.RemoveAll(this);
-		FGameDelegates::Get().GetEndPlayMapDelegate().RemoveAll(this);
-	}
-}
-
 void UWorldPartition::ForEachIntersectingActorDesc(const FBox& Box, TSubclassOf<AActor> ActorClass, TFunctionRef<bool(const FWorldPartitionActorDesc*)> Predicate) const
 {
 	EditorHash->ForEachIntersectingActor(Box, [&ActorClass, Predicate](const FWorldPartitionActorDesc* ActorDesc)
@@ -1149,7 +1171,7 @@ const UActorDescContainer* UWorldPartition::RegisterActorDescContainer(FName Pac
 	}
 		
 	UActorDescContainer* NewContainer = NewObject<UActorDescContainer>(GetTransientPackage());
-	NewContainer->Initialize(GetWorld(), PackageName, true);
+	NewContainer->Initialize(GetWorld(), PackageName);
 	ActorDescContainers.Add(PackageName, TWeakObjectPtr<UActorDescContainer>(NewContainer));
 
 	return NewContainer;
diff --git a/Engine/Source/Runtime/Engine/Public/WorldPartition/ActorDescContainer.h b/Engine/Source/Runtime/Engine/Public/WorldPartition/ActorDescContainer.h
--- a/Engine/Source/Runtime/Engine/Public/WorldPartition/ActorDescContainer.h
+++ b/Engine/Source/Runtime/Engine/Public/WorldPartition/ActorDescContainer.h
@@ -17,7 +17,7 @@ class ENGINE_API UActorDescContainer : public UObject, public FActorDescList
 	friend class FWorldPartitionActorDesc;
 
 public:
-	void Initialize(UWorld* World, FName InPackageName, bool bRegisterDelegates);
+	void Initialize(UWorld* World, FName InPackageName);
 	virtual void Uninitialize();
 	
 	virtual UWorld* GetWorld() const override;
@@ -49,9 +49,6 @@ protected:
 	//~ End UObject Interface
 
 #if WITH_EDITOR
-	virtual void RegisterDelegates();
-	virtual void UnregisterDelegates();
-
 	virtual void OnActorDescAdded(FWorldPartitionActorDesc* NewActorDesc);
 	virtual void OnActorDescRemoved(FWorldPartitionActorDesc* ActorDesc);
 	virtual void OnActorDescUpdating(FWorldPartitionActorDesc* ActorDesc) {}
@@ -66,4 +63,10 @@ protected:
 
 	FName ContainerPackageName;
 #endif
+
+private:
+#if WITH_EDITOR
+	void RegisterEditorDelegates();
+	void UnregisterEditorDelegates();
+#endif
 };
\ No newline at end of file
diff --git a/Engine/Source/Runtime/Engine/Public/WorldPartition/WorldPartition.h b/Engine/Source/Runtime/Engine/Public/WorldPartition/WorldPartition.h
--- a/Engine/Source/Runtime/Engine/Public/WorldPartition/WorldPartition.h
+++ b/Engine/Source/Runtime/Engine/Public/WorldPartition/WorldPartition.h
@@ -225,12 +225,13 @@ private:
 #endif
 
 	bool IsMainWorldPartition() const;
-		
-#if WITH_EDITOR
+	void OnWorldBeginPlay();
+
 	// Delegates registration
-	virtual void RegisterDelegates() override;
-	virtual void UnregisterDelegates() override;
+	void RegisterDelegates();
+	void UnregisterDelegates();	
 
+#if WITH_EDITOR
 	void UpdateLoadingEditorCell(UWorldPartitionEditorCell* Cell, bool bShouldBeLoaded);
 	void HashActorDesc(FWorldPartitionActorDesc* ActorDesc);
 	void UnhashActorDesc(FWorldPartitionActorDesc* ActorDesc);
