Fix PhysicsField Rhi validation

diff --git a/Engine/Source/Runtime/Engine/Private/PhysicsField/PhysicsFieldComponent.cpp b/Engine/Source/Runtime/Engine/Private/PhysicsField/PhysicsFieldComponent.cpp
--- a/Engine/Source/Runtime/Engine/Private/PhysicsField/PhysicsFieldComponent.cpp
+++ b/Engine/Source/Runtime/Engine/Private/PhysicsField/PhysicsFieldComponent.cpp
@@ -113,7 +113,7 @@ void InitInternalBuffer(const uint32 ElementCount, FRWBuffer& OutputBuffer)
 
 		if (OutputBuffer.UAV)
 		{
-			FRHICommandListExecutor::GetImmediateCommandList().Transition(FRHITransitionInfo(OutputBuffer.UAV, ERHIAccess::SRVCompute, ERHIAccess::UAVCompute));
+			FRHICommandListExecutor::GetImmediateCommandList().Transition(FRHITransitionInfo(OutputBuffer.UAV, ERHIAccess::Unknown, ERHIAccess::UAVCompute));
 			if (PixelFormat == EPixelFormat::PF_R32_FLOAT)
 			{
 				FRHICommandListExecutor::GetImmediateCommandList().ClearUAVFloat(OutputBuffer.UAV, FVector4(ForceInitToZero));
@@ -136,7 +136,7 @@ void UpdateInternalBuffer(const uint32 ElementCount, const BufferType* InputData
 		const uint32 BufferCount = ElementCount * ElementSize;
 		const uint32 BufferBytes = sizeof(BufferType) * BufferCount;
 
-		FRHICommandListExecutor::GetImmediateCommandList().Transition(FRHITransitionInfo(OutputBuffer.UAV, ERHIAccess::SRVCompute, ERHIAccess::UAVCompute));
+		FRHICommandListExecutor::GetImmediateCommandList().Transition(FRHITransitionInfo(OutputBuffer.UAV, ERHIAccess::Unknown, ERHIAccess::UAVCompute));
 
 		if(bInitField)
 		{
