Add option to 'RecompileShaders material <name>' command to specify platform, to test compiling shaders for specific platforms

diff --git a/Engine/Source/Runtime/Engine/Private/ShaderCompiler/ShaderCompiler.cpp b/Engine/Source/Runtime/Engine/Private/ShaderCompiler/ShaderCompiler.cpp
--- a/Engine/Source/Runtime/Engine/Private/ShaderCompiler/ShaderCompiler.cpp
+++ b/Engine/Source/Runtime/Engine/Private/ShaderCompiler/ShaderCompiler.cpp
@@ -5319,18 +5319,40 @@ bool RecompileShaders(const TCHAR* Cmd, FOutputDevice& Ar)
 		{
 			FString RequestedMaterialName(FParse::Token(Cmd, 0));
 			FRecompileShadersTimer TestTimer(FString::Printf(TEXT("Recompile Material %s"), *RequestedMaterialName));
+
+			ITargetPlatformManagerModule& TPM = GetTargetPlatformManagerRef();
+			FString TargetPlatformName(FParse::Token(Cmd, 0));
+			const ITargetPlatform* TargetPlatform = nullptr;
+			if (TargetPlatformName.Len() > 0)
+			{
+				TargetPlatform = TPM.FindTargetPlatform(TargetPlatformName);
+			}
+
 			bool bMaterialFound = false;
-			for( TObjectIterator<UMaterial> It; It; ++It )
+			for( TObjectIterator<UMaterialInterface> It; It; ++It )
 			{
-				UMaterial* Material = *It;
+				UMaterialInterface* Material = *It;
 				if( Material && Material->GetName() == RequestedMaterialName)
 				{
 					bMaterialFound = true;
 #if WITH_EDITOR
 					// <Pre/Post>EditChange will force a re-creation of the resource,
 					// in turn recompiling the shader.
-					Material->PreEditChange(NULL);
-					Material->PostEditChange();
+					if (TargetPlatform)
+					{
+						Material->BeginCacheForCookedPlatformData(TargetPlatform);
+						while (!Material->IsCachedCookedPlatformDataLoaded(TargetPlatform))
+						{
+							FPlatformProcess::Sleep(0.1f);
+							GShaderCompilingManager->ProcessAsyncResults(false, false);
+						}
+						Material->ClearCachedCookedPlatformData(TargetPlatform);
+					}
+					else
+					{
+						Material->PreEditChange(NULL);
+						Material->PostEditChange();
+					}
 #endif // WITH_EDITOR
 					break;
 				}
