AssetManager's BulkDataScanning: keep the AssetRegistry's temporary caching optimization going when we detect that it is invalidated; use the AssetRegistry's internal invalidation functionality rather than disabling the caching.
#rnx

diff --git a/Engine/Source/Runtime/AssetRegistry/Private/AssetRegistry.cpp b/Engine/Source/Runtime/AssetRegistry/Private/AssetRegistry.cpp
--- a/Engine/Source/Runtime/AssetRegistry/Private/AssetRegistry.cpp
+++ b/Engine/Source/Runtime/AssetRegistry/Private/AssetRegistry.cpp
@@ -3419,6 +3419,11 @@ void UAssetRegistryImpl::SetTemporaryCachingMode(bool bEnable)
 	}
 }
 
+void UAssetRegistryImpl::SetTemporaryCachingModeInvalidated()
+{
+	bIsTempCachingUpToDate = false;
+}
+
 bool UAssetRegistryImpl::GetTemporaryCachingMode() const
 {
 	return bIsTempCachingEnabled;
diff --git a/Engine/Source/Runtime/AssetRegistry/Private/AssetRegistry.h b/Engine/Source/Runtime/AssetRegistry/Private/AssetRegistry.h
--- a/Engine/Source/Runtime/AssetRegistry/Private/AssetRegistry.h
+++ b/Engine/Source/Runtime/AssetRegistry/Private/AssetRegistry.h
@@ -84,6 +84,7 @@ public:
 	virtual void ExpandRecursiveFilter(const FARFilter& InFilter, FARFilter& ExpandedFilter) const override;
 	virtual void CompileFilter(const FARFilter& InFilter, FARCompiledFilter& OutCompiledFilter) const override;
 	virtual void SetTemporaryCachingMode(bool bEnable) override;
+	virtual void SetTemporaryCachingModeInvalidated() override;
 	virtual bool GetTemporaryCachingMode() const override;
 	virtual EAssetAvailability::Type GetAssetAvailability(const FAssetData& AssetData) const override;	
 	virtual float GetAssetAvailabilityProgress(const FAssetData& AssetData, EAssetAvailabilityProgressReportingType::Type ReportType) const override;
diff --git a/Engine/Source/Runtime/AssetRegistry/Public/AssetRegistry/IAssetRegistry.h b/Engine/Source/Runtime/AssetRegistry/Public/AssetRegistry/IAssetRegistry.h
--- a/Engine/Source/Runtime/AssetRegistry/Public/AssetRegistry/IAssetRegistry.h
+++ b/Engine/Source/Runtime/AssetRegistry/Public/AssetRegistry/IAssetRegistry.h
@@ -357,6 +357,11 @@ public:
 
 	/** Enables or disable temporary search caching, when this is enabled scanning/searching is faster because we assume no objects are loaded between scans. Disabling frees any caches created */
 	virtual void SetTemporaryCachingMode(bool bEnable) = 0;
+	/**
+	 * Mark that the temporary cached needs to be updated before being used again, because e.g. a new class was loaded.
+	 * Does nothing if TemporaryCachingMode is not enabled
+	 */
+	virtual void SetTemporaryCachingModeInvalidated() = 0;
 
 	/** Returns true if temporary caching mode enabled */
 	virtual bool GetTemporaryCachingMode() const = 0;
diff --git a/Engine/Source/Runtime/Engine/Private/AssetManager.cpp b/Engine/Source/Runtime/Engine/Private/AssetManager.cpp
--- a/Engine/Source/Runtime/Engine/Private/AssetManager.cpp
+++ b/Engine/Source/Runtime/Engine/Private/AssetManager.cpp
@@ -2975,8 +2975,8 @@ bool UAssetManager::ShouldScanPrimaryAssetType(FPrimaryAssetTypeInfo& TypeInfo)
 
 	if (bBaseClassWasLoaded)
 	{
-		// Had to load a class, leave temporary caching mode for future scans
-		GetAssetRegistry().SetTemporaryCachingMode(false);
+		// Had to load a class, mark that the temporary cache needs to be updated
+		GetAssetRegistry().SetTemporaryCachingModeInvalidated();
 	}
 
 	return bIsValid;
@@ -3987,8 +3987,8 @@ void UAssetManager::RefreshPrimaryAssetDirectory(bool bForceRefresh)
 
 			if (bBaseClassWasLoaded)
 			{
-				// Had to load a class, leave temporary caching mode for future scans
-				GetAssetRegistry().SetTemporaryCachingMode(false);
+				// Had to load a class, mark that the temporary cache needs to be updated
+				GetAssetRegistry().SetTemporaryCachingModeInvalidated();
 			}
 
 			if (!bIsValid)
