Make sure any async static mesh task are cancelled or completed before their UPackage::BeginDestroy gets called
Remove strong ref on UPackage that is not needed anymore
Remove strong ref in ICompilable that is not needed anymore

diff --git a/Engine/Source/Runtime/Engine/Classes/Engine/StaticMesh.h b/Engine/Source/Runtime/Engine/Classes/Engine/StaticMesh.h
--- a/Engine/Source/Runtime/Engine/Classes/Engine/StaticMesh.h
+++ b/Engine/Source/Runtime/Engine/Classes/Engine/StaticMesh.h
@@ -1845,13 +1845,6 @@ public:
 	FStaticMeshCompilationContext& operator=(FStaticMeshCompilationContext&&) = default;
 
 	bool bShouldComputeExtendedBounds = false;
-
-	// Prevent the garbage collection of the static mesh's package while the async compilation is in progress.
-	// This may happen because we allow garbage collection of UStaticMesh to be initiated even in the presence of
-	// an async operation in order to cancel it if possible. In that case, we need to wait until the async operation
-	// is fully cancelled before trying to destroy the UPackage otherwise the bulk data of the static mesh could
-	// become inaccessible while the async operation is still using it.
-	TStrongObjectPtr<UPackage> StaticMeshOwner;
 };
 
 class FStaticMeshPostLoadContext : public FStaticMeshCompilationContext
diff --git a/Engine/Source/Runtime/Engine/Private/AsyncCompilationHelpers.cpp b/Engine/Source/Runtime/Engine/Private/AsyncCompilationHelpers.cpp
--- a/Engine/Source/Runtime/Engine/Private/AsyncCompilationHelpers.cpp
+++ b/Engine/Source/Runtime/Engine/Private/AsyncCompilationHelpers.cpp
@@ -96,8 +96,9 @@ namespace AsyncCompilationHelpers
 		uint64 StartTime = FPlatformTime::Cycles64();
 		TOptional<FScopedSlowTask> SlowTask;
 
-		// Do not create a progress during PostLoad as TickSlate could endup calling too many things
-		if (!FUObjectThreadContext::Get().IsRoutingPostLoad)
+		// Do not create a progress during PostLoad or garbage collect as ticking slate might
+		// introduce too many unwanted side effects
+		if (!FUObjectThreadContext::Get().IsRoutingPostLoad && !GIsGarbageCollecting)
 		{
 			FFormatNamedArguments Args;
 			Args.Add(TEXT("AssetType"), AssetType.ToLower());
diff --git a/Engine/Source/Runtime/Engine/Private/StaticMesh.cpp b/Engine/Source/Runtime/Engine/Private/StaticMesh.cpp
--- a/Engine/Source/Runtime/Engine/Private/StaticMesh.cpp
+++ b/Engine/Source/Runtime/Engine/Private/StaticMesh.cpp
@@ -5101,7 +5101,6 @@ void UStaticMesh::PostLoad()
 			FStaticMeshCompilingManager::Get().FinishCompilation({ ComplexCollisionMesh });
 		}
 
-		Context.StaticMeshOwner.Reset(GetPackage());
 		AsyncTask = MakeUnique<FStaticMeshAsyncBuildTask>(this, MakeUnique<FStaticMeshPostLoadContext>(MoveTemp(Context)));
 		AsyncTask->StartBackgroundTask(StaticMeshThreadPool, BasePriority, EQueuedWorkFlags::None, GetBuildRequiredMemory());
 		FStaticMeshCompilingManager::Get().AddStaticMeshes({this});
diff --git a/Engine/Source/Runtime/Engine/Private/StaticMeshBuild.cpp b/Engine/Source/Runtime/Engine/Private/StaticMeshBuild.cpp
--- a/Engine/Source/Runtime/Engine/Private/StaticMeshBuild.cpp
+++ b/Engine/Source/Runtime/Engine/Private/StaticMeshBuild.cpp
@@ -247,7 +247,6 @@ void UStaticMesh::BatchBuild(const TArray<UStaticMesh*>& InStaticMeshes, bool bI
 				if (OutErrors == nullptr && FStaticMeshCompilingManager::Get().IsAsyncCompilationAllowed(StaticMesh))
 				{
 					TUniquePtr<FStaticMeshBuildContext> Context = MakeUnique<FStaticMeshBuildContext>();
-					Context->StaticMeshOwner.Reset(StaticMesh->GetPackage());
 					StaticMesh->BeginBuildInternal(Context.Get());
 
 					FQueuedThreadPool* StaticMeshThreadPool = FStaticMeshCompilingManager::Get().GetThreadPool();
diff --git a/Engine/Source/Runtime/Engine/Private/StaticMeshCompiler.cpp b/Engine/Source/Runtime/Engine/Private/StaticMeshCompiler.cpp
--- a/Engine/Source/Runtime/Engine/Private/StaticMeshCompiler.cpp
+++ b/Engine/Source/Runtime/Engine/Private/StaticMeshCompiler.cpp
@@ -87,6 +87,38 @@ namespace StaticMeshCompilingManagerImpl
 FStaticMeshCompilingManager::FStaticMeshCompilingManager()
 	: Notification(LOCTEXT("StaticMeshes", "Static Meshes"))
 {
+	PostReachabilityAnalysisHandle = FCoreUObjectDelegates::PostReachabilityAnalysis.AddRaw(this, &FStaticMeshCompilingManager::OnPostReachabilityAnalysis);
+}
+
+void FStaticMeshCompilingManager::OnPostReachabilityAnalysis()
+{
+	if (GetNumRemainingMeshes())
+	{
+		TRACE_CPUPROFILER_EVENT_SCOPE(FStaticMeshCompilingManager::CancelUnreachableMeshes);
+
+		TArray<UStaticMesh*> PendingStaticMeshes;
+		PendingStaticMeshes.Reserve(GetNumRemainingMeshes());
+
+		for (auto Iterator = RegisteredStaticMesh.CreateIterator(); Iterator; ++Iterator)
+		{
+			UStaticMesh* StaticMesh = Iterator->GetEvenIfUnreachable();
+			if (StaticMesh && StaticMesh->IsUnreachable())
+			{
+				UE_LOG(LogStaticMesh, Verbose, TEXT("Cancelling static mesh %s async compilation because it's being garbage collected"), *StaticMesh->GetName());
+
+				if (StaticMesh->TryCancelAsyncTasks())
+				{
+					Iterator.RemoveCurrent();
+				}
+				else
+				{
+					PendingStaticMeshes.Add(StaticMesh);
+				}
+			}
+		}
+
+		FinishCompilation(PendingStaticMeshes);
+	}
 }
 
 EQueuedWorkPriority FStaticMeshCompilingManager::GetBasePriority(UStaticMesh* InStaticMesh) const
@@ -131,15 +163,7 @@ void FStaticMeshCompilingManager::Shutdown()
 			if (WeakStaticMesh.IsValid())
 			{
 				UStaticMesh* StaticMesh = WeakStaticMesh.Get();
-				if (!StaticMesh->IsAsyncTaskComplete())
-				{
-					if (StaticMesh->AsyncTask->Cancel())
-					{
-						StaticMesh->AsyncTask.Reset();
-					}
-				}
-
-				if (StaticMesh->AsyncTask)
+				if (!StaticMesh->TryCancelAsyncTasks())
 				{
 					PendingStaticMeshes.Add(StaticMesh);
 				}
@@ -148,6 +172,8 @@ void FStaticMeshCompilingManager::Shutdown()
 
 		FinishCompilation(PendingStaticMeshes);
 	}
+
+	FCoreUObjectDelegates::PostReachabilityAnalysis.Remove(PostReachabilityAnalysisHandle);
 }
 
 bool FStaticMeshCompilingManager::IsAsyncStaticMeshCompilationEnabled() const
@@ -178,10 +204,17 @@ void FStaticMeshCompilingManager::PostCompilation(TArrayView<UStaticMesh* const>
 
 		for (UStaticMesh* StaticMesh : InStaticMeshes)
 		{
-			AssetsData.Emplace(StaticMesh);
+			// Do not broadcast an event for unreachable objects
+			if (!StaticMesh->IsUnreachable())
+			{
+				AssetsData.Emplace(StaticMesh);
+			}
 		}
 
-		FAssetCompilingManager::Get().OnAssetPostCompileEvent().Broadcast(AssetsData);
+		if (AssetsData.Num())
+		{
+			FAssetCompilingManager::Get().OnAssetPostCompileEvent().Broadcast(AssetsData);
+		}
 	}
 }
 
@@ -195,8 +228,6 @@ void FStaticMeshCompilingManager::PostCompilation(UStaticMesh* StaticMesh)
 		check(IsInGameThread());
 		TRACE_CPUPROFILER_EVENT_SCOPE(PostCompilation);
 
-		UE_LOG(LogStaticMesh, Verbose, TEXT("Refreshing static mesh %s because it is ready"), *StaticMesh->GetName());
-
 		FObjectCacheContextScope ObjectCacheScope;
 
 		// The scope is important here to destroy the FStaticMeshAsyncBuildScope before broadcasting events
@@ -205,6 +236,14 @@ void FStaticMeshCompilingManager::PostCompilation(UStaticMesh* StaticMesh)
 			TUniquePtr<FStaticMeshAsyncBuildTask> LocalAsyncTask = MoveTemp(StaticMesh->AsyncTask);
 			LocalAsyncTask->EnsureCompletion();
 
+			// Do not do anything else if the staticmesh is being garbage collected
+			if (StaticMesh->IsUnreachable())
+			{
+				return;
+			}
+
+			UE_LOG(LogStaticMesh, Verbose, TEXT("Refreshing static mesh %s because it is ready"), *StaticMesh->GetName());
+
 			FStaticMeshAsyncBuildScope AsyncBuildScope(StaticMesh);
 
 			if (LocalAsyncTask->GetTask().PostLoadContext.IsValid())
@@ -315,7 +354,7 @@ void FStaticMeshCompilingManager::FinishCompilation(TArrayView<UStaticMesh* cons
 				return StaticMesh->AsyncTask.Get();
 			}
 
-			TStrongObjectPtr<UStaticMesh> StaticMesh;
+			UStaticMesh* StaticMesh;
 			FName GetName() override { return StaticMesh->GetOutermost()->GetFName(); }
 		};
 
@@ -328,7 +367,7 @@ void FStaticMeshCompilingManager::FinishCompilation(TArrayView<UStaticMesh* cons
 			LogStaticMesh,
 			[this](AsyncCompilationHelpers::ICompilable* Object)
 			{
-				UStaticMesh* StaticMesh = static_cast<FCompilableStaticMesh*>(Object)->StaticMesh.Get();
+				UStaticMesh* StaticMesh = static_cast<FCompilableStaticMesh*>(Object)->StaticMesh;
 				PostCompilation(StaticMesh);
 				RegisteredStaticMesh.Remove(StaticMesh);
 			}
diff --git a/Engine/Source/Runtime/Engine/Public/StaticMeshCompiler.h b/Engine/Source/Runtime/Engine/Public/StaticMeshCompiler.h
--- a/Engine/Source/Runtime/Engine/Public/StaticMeshCompiler.h
+++ b/Engine/Source/Runtime/Engine/Public/StaticMeshCompiler.h
@@ -84,6 +84,9 @@ private:
 
 	void PostCompilation(TArrayView<UStaticMesh* const> InStaticMeshes);
 	void PostCompilation(UStaticMesh* StaticMesh);
+
+	void OnPostReachabilityAnalysis();
+	FDelegateHandle PostReachabilityAnalysisHandle;
 };
 
 #endif // #if WITH_EDITOR
\ No newline at end of file
