LevelInstance: Change FLevelInstanceID to struct so it can handle hash collision and use 64bit hash to greatly reduce chances of collision.
WorldPartition Container ID also converted to 64bit hash

diff --git a/Engine/Source/Runtime/Engine/Private/LevelInstance/LevelInstanceActor.cpp b/Engine/Source/Runtime/Engine/Private/LevelInstance/LevelInstanceActor.cpp
--- a/Engine/Source/Runtime/Engine/Private/LevelInstance/LevelInstanceActor.cpp
+++ b/Engine/Source/Runtime/Engine/Private/LevelInstance/LevelInstanceActor.cpp
@@ -22,11 +22,9 @@
 ALevelInstance::ALevelInstance(const FObjectInitializer& ObjectInitializer)
 	: Super(ObjectInitializer)
 #if WITH_EDITOR
-	, CachedLevelInstanceID(InvalidLevelInstanceID)
 	, bGuardLoadUnload(false)
 	, bEditLockLocation(false)
 #endif
-	, LevelInstanceID(InvalidLevelInstanceID)
 {
 	RootComponent = CreateDefaultSubobject<ULevelInstanceComponent>(TEXT("Root"));
 	RootComponent->Mobility = EComponentMobility::Static;
@@ -115,7 +113,7 @@ void ALevelInstance::PostUnregisterAllComponents()
 		}
 
 		// To avoid processing PostUnregisterAllComponents multiple times (BP Recompile is one use case)
-		LevelInstanceID = InvalidLevelInstanceID;
+		LevelInstanceID = FLevelInstanceID();
 	}
 
 }
@@ -166,7 +164,7 @@ bool ALevelInstance::IsLevelInstancePathValid() const
 
 bool ALevelInstance::HasValidLevelInstanceID() const
 {
-	return LevelInstanceID != InvalidLevelInstanceID;
+	return LevelInstanceID.IsValid();
 }
 
 const FLevelInstanceID& ALevelInstance::GetLevelInstanceID() const
@@ -305,7 +303,7 @@ void ALevelInstance::PostEditUndoInternal()
 		}
 	}
 
-	CachedLevelInstanceID = InvalidLevelInstanceID;
+	CachedLevelInstanceID = FLevelInstanceID();
 	CachedWorldAsset.Reset();
 
 	if (ULevelInstanceComponent* LevelInstanceComponent = Cast<ULevelInstanceComponent>(GetRootComponent()))
diff --git a/Engine/Source/Runtime/Engine/Private/LevelInstance/LevelInstanceEditorInstanceActor.cpp b/Engine/Source/Runtime/Engine/Private/LevelInstance/LevelInstanceEditorInstanceActor.cpp
--- a/Engine/Source/Runtime/Engine/Private/LevelInstance/LevelInstanceEditorInstanceActor.cpp
+++ b/Engine/Source/Runtime/Engine/Private/LevelInstance/LevelInstanceEditorInstanceActor.cpp
@@ -7,9 +7,6 @@
 
 ALevelInstanceEditorInstanceActor::ALevelInstanceEditorInstanceActor(const FObjectInitializer& ObjectInitializer)
 	: Super(ObjectInitializer)
-#if WITH_EDITOR
-	, LevelInstanceID(InvalidLevelInstanceID)
-#endif
 {
 	RootComponent = CreateDefaultSubobject<USceneComponent>(TEXT("RootComponent"));
 	RootComponent->Mobility = EComponentMobility::Static;
diff --git a/Engine/Source/Runtime/Engine/Private/LevelInstance/LevelInstanceEditorLevelStreaming.cpp b/Engine/Source/Runtime/Engine/Private/LevelInstance/LevelInstanceEditorLevelStreaming.cpp
--- a/Engine/Source/Runtime/Engine/Private/LevelInstance/LevelInstanceEditorLevelStreaming.cpp
+++ b/Engine/Source/Runtime/Engine/Private/LevelInstance/LevelInstanceEditorLevelStreaming.cpp
@@ -13,7 +13,7 @@
 #endif
 
 #if WITH_EDITOR
-static FLevelInstanceID EditLevelInstanceID = InvalidLevelInstanceID;
+static FLevelInstanceID EditLevelInstanceID;
 #endif
 
 ULevelStreamingLevelInstanceEditor::ULevelStreamingLevelInstanceEditor(const FObjectInitializer& ObjectInitializer)
diff --git a/Engine/Source/Runtime/Engine/Private/LevelInstance/LevelInstanceLevelStreaming.cpp b/Engine/Source/Runtime/Engine/Private/LevelInstance/LevelInstanceLevelStreaming.cpp
--- a/Engine/Source/Runtime/Engine/Private/LevelInstance/LevelInstanceLevelStreaming.cpp
+++ b/Engine/Source/Runtime/Engine/Private/LevelInstance/LevelInstanceLevelStreaming.cpp
@@ -18,7 +18,6 @@
 
 ULevelStreamingLevelInstance::ULevelStreamingLevelInstance(const FObjectInitializer& ObjectInitializer)
 	: Super(ObjectInitializer)
-	, LevelInstanceID(InvalidLevelInstanceID)
 {
 #if WITH_EDITOR
 	SetShouldBeVisibleInEditor(true);
@@ -65,7 +64,7 @@ ULevelStreamingLevelInstance* ULevelStreamingLevelInstance::LoadInstance(ALevelI
 	FString ShortPackageName = FPackageName::GetShortName(LevelInstanceActor->GetWorldAsset().GetLongPackageName());
 	// Build a unique and deterministic LevelInstance level instance name by using LevelInstanceID. 
 	// Distinguish game from editor since we don't want to duplicate for PIE already loaded editor instances (not yet supported).
-	FString Suffix = FString::Printf(TEXT("%s_LevelInstance_%08X_%d"), *ShortPackageName, LevelInstanceActor->GetLevelInstanceID(), LevelInstanceActor->GetWorld()->IsGameWorld() ? 1 : 0);
+	FString Suffix = FString::Printf(TEXT("%s_LevelInstance_%016llx_%d"), *ShortPackageName, LevelInstanceActor->GetLevelInstanceID().GetHash(), LevelInstanceActor->GetWorld()->IsGameWorld() ? 1 : 0);
 	ULevelStreamingLevelInstance* LevelStreaming = Cast<ULevelStreamingLevelInstance>(ULevelStreamingDynamic::LoadLevelInstanceBySoftObjectPtr(LevelInstanceActor->GetWorld(), LevelInstanceActor->GetWorldAsset(), LevelInstanceActor->GetActorTransform(), bOutSuccess, Suffix, ULevelStreamingLevelInstance::StaticClass()));
 	if (bOutSuccess)
 	{
diff --git a/Engine/Source/Runtime/Engine/Private/LevelInstance/LevelInstanceSubsystem.cpp b/Engine/Source/Runtime/Engine/Private/LevelInstance/LevelInstanceSubsystem.cpp
--- a/Engine/Source/Runtime/Engine/Private/LevelInstance/LevelInstanceSubsystem.cpp
+++ b/Engine/Source/Runtime/Engine/Private/LevelInstance/LevelInstanceSubsystem.cpp
@@ -12,7 +12,7 @@
 #include "EngineUtils.h"
 #include "LevelInstancePrivate.h"
 #include "LevelUtils.h"
-#include "Misc/HashBuilder.h"
+#include "Hash/CityHash.h"
 
 #if WITH_EDITOR
 #include "LevelInstance/LevelInstanceEditorLevelStreaming.h"
@@ -89,23 +89,21 @@ ALevelInstance* ULevelInstanceSubsystem::GetLevelInstance(FLevelInstanceID Level
 	return nullptr;
 }
 
-FLevelInstanceID ULevelInstanceSubsystem::ComputeLevelInstanceID(ALevelInstance* LevelInstanceActor) const
+FLevelInstanceID::FLevelInstanceID(ULevelInstanceSubsystem* LevelInstanceSubsystem, ALevelInstance* LevelInstanceActor)
 {
-	FLevelInstanceID LevelInstanceID = InvalidLevelInstanceID;
-	FHashBuilder HashBuilder;
-	ForEachLevelInstanceAncestorsAndSelf(LevelInstanceActor, [&HashBuilder](const ALevelInstance* AncestorOrSelf)
-	{
-		HashBuilder << AncestorOrSelf->GetLevelInstanceActorGuid();
-		return true;
-	});
-	LevelInstanceID = HashBuilder.GetHash();
-	return LevelInstanceID;
+	LevelInstanceSubsystem->ForEachLevelInstanceAncestorsAndSelf(LevelInstanceActor, [this](const ALevelInstance* AncestorOrSelf)
+		{
+			Guids.Add(AncestorOrSelf->GetLevelInstanceActorGuid());
+			return true;
+		});
+	check(!Guids.IsEmpty());
+	Hash = CityHash64((const char*)Guids.GetData(), Guids.Num() * sizeof(FGuid));
 }
 
 FLevelInstanceID ULevelInstanceSubsystem::RegisterLevelInstance(ALevelInstance* LevelInstanceActor)
 {
-	FLevelInstanceID LevelInstanceID = ComputeLevelInstanceID(LevelInstanceActor);
-	check(LevelInstanceID != InvalidLevelInstanceID);
+	FLevelInstanceID LevelInstanceID(this, LevelInstanceActor);
+	check(LevelInstanceID.IsValid());
 	ALevelInstance*& Value = RegisteredLevelInstances.FindOrAdd(LevelInstanceID);
 	check(GIsReinstancing || Value == nullptr || Value == LevelInstanceActor);
 	Value = LevelInstanceActor;
@@ -369,7 +367,7 @@ void ULevelInstanceSubsystem::Tick()
 		UpdateStreamingState();
 
 		// Begin editing the pending LevelInstance when loads complete
-		if (PendingLevelInstanceToEdit != InvalidLevelInstanceID && !LevelInstancesToLoadOrUpdate.Num())
+		if (PendingLevelInstanceToEdit.IsValid() && !LevelInstancesToLoadOrUpdate.Num())
 		{
 			if (ALevelInstance** LevelInstanceActor = RegisteredLevelInstances.Find(PendingLevelInstanceToEdit))
 			{
@@ -1361,8 +1359,8 @@ bool ULevelInstanceSubsystem::CanCommitLevelInstance(const ALevelInstance* Level
 void ULevelInstanceSubsystem::EditLevelInstance(ALevelInstance* LevelInstanceActor, TWeakObjectPtr<AActor> ContextActorPtr)
 {
 	check(CanEditLevelInstance(LevelInstanceActor));
-	const bool bWasPendingEdit = PendingLevelInstanceToEdit != InvalidLevelInstanceID;
-	PendingLevelInstanceToEdit = InvalidLevelInstanceID;
+	const bool bWasPendingEdit = PendingLevelInstanceToEdit.IsValid();
+	PendingLevelInstanceToEdit = FLevelInstanceID();
 		
 	FScopedSlowTask SlowTask(0, LOCTEXT("BeginEditLevelInstance", "Loading Level Instance for edit..."), !GetWorld()->IsGameWorld());
 	SlowTask.MakeDialog();
diff --git a/Engine/Source/Runtime/Engine/Private/WorldPartition/WorldPartitionActorCluster.cpp b/Engine/Source/Runtime/Engine/Private/WorldPartition/WorldPartitionActorCluster.cpp
--- a/Engine/Source/Runtime/Engine/Private/WorldPartition/WorldPartitionActorCluster.cpp
+++ b/Engine/Source/Runtime/Engine/Private/WorldPartition/WorldPartitionActorCluster.cpp
@@ -10,7 +10,7 @@
 #include "WorldPartition/ActorDescContainer.h"
 #include "WorldPartition/WorldPartitionActorDescView.h"
 #include "WorldPartition/WorldPartitionRuntimeHash.h"
-#include "Misc/HashBuilder.h"
+#include "Hash/CityHash.h"
 
 template<class LayerNameContainer>
 TSet<const UDataLayer*> GetDataLayers(UWorld* InWorld, const LayerNameContainer& DataLayerNames)
@@ -173,7 +173,7 @@ FActorContainerInstance::FActorContainerInstance(const UActorDescContainer* InCo
 	: FActorContainerInstance(0, FTransform::Identity, FBox(ForceInit), TSet<FName>(), EContainerClusterMode::Partitioned, InContainer, TSet<FGuid>(), InActorDescViewMap)
 {}
 
-FActorContainerInstance::FActorContainerInstance(uint32 InID, const FTransform& InTransform, const FBox& InBounds, const TSet<FName>& InDataLayers, EContainerClusterMode InClusterMode, const UActorDescContainer* InContainer, TSet<FGuid> InChildContainers, TMap<FGuid, FWorldPartitionActorDescView> InActorDescViewMap)
+FActorContainerInstance::FActorContainerInstance(uint64 InID, const FTransform& InTransform, const FBox& InBounds, const TSet<FName>& InDataLayers, EContainerClusterMode InClusterMode, const UActorDescContainer* InContainer, TSet<FGuid> InChildContainers, TMap<FGuid, FWorldPartitionActorDescView> InActorDescViewMap)
 	: ID(InID)
 	, Transform(InTransform)
 	, Bounds(InBounds)
@@ -278,7 +278,7 @@ void CreateActorCluster(const FWorldPartitionActorDescView& ActorDescView, TMap<
 	}
 }
 
-void FActorClusterContext::CreateContainerInstanceRecursive(uint32 ID, const FTransform& Transform, EContainerClusterMode ClusterMode, const UActorDescContainer* Container, const TSet<FName>& DataLayers, FBox& ParentBounds)
+void FActorClusterContext::CreateContainerInstanceRecursive(uint64 ID, const FTransform& Transform, EContainerClusterMode ClusterMode, const UActorDescContainer* Container, const TSet<FName>& DataLayers, FBox& ParentBounds)
 {
 	InstanceCountHint += Container->GetActorDescCount();
 		
@@ -299,8 +299,8 @@ void FActorClusterContext::CreateContainerInstanceRecursive(uint32 ID, const FTr
 			// Add Child Container Guid so we can discard the actor later
 			ChildContainers.Add(ActorDescView.GetGuid());
 
-			FHashBuilder ChildContainerHashBuilder(ID);
-			ChildContainerHashBuilder << ActorDescView.GetGuid();
+			FGuid ActorGuid = ActorDescView.GetGuid();
+			uint64 Hash = CityHash64WithSeed((const char*)&ActorGuid, sizeof(ActorGuid), ID);
 			
 			TSet<FName> ChildDataLayers;
 			ChildDataLayers.Reserve(DataLayers.Num() + ActorDescView.GetDataLayers().Num());
@@ -311,7 +311,7 @@ void FActorClusterContext::CreateContainerInstanceRecursive(uint32 ID, const FTr
 			}
 			// Always inherite parent container DataLayers
 			ChildDataLayers.Append(DataLayers);
-			CreateContainerInstanceRecursive(ChildContainerHashBuilder.GetHash(), OutTransform * Transform, OutClusterMode, OutContainer, ChildDataLayers, Bounds);
+			CreateContainerInstanceRecursive(Hash, OutTransform * Transform, OutClusterMode, OutContainer, ChildDataLayers, Bounds);
 		}
 		else
 		{
diff --git a/Engine/Source/Runtime/Engine/Private/WorldPartition/WorldPartitionLevelHelper.cpp b/Engine/Source/Runtime/Engine/Private/WorldPartition/WorldPartitionLevelHelper.cpp
--- a/Engine/Source/Runtime/Engine/Private/WorldPartition/WorldPartitionLevelHelper.cpp
+++ b/Engine/Source/Runtime/Engine/Private/WorldPartition/WorldPartitionLevelHelper.cpp
@@ -202,7 +202,7 @@ bool FWorldPartitionLevelHelper::LoadActors(ULevel* InDestLevel, TArrayView<FWor
 	ActorPackages.Reserve(InActorPackages.Num());
 
 	// Levels to Load with actors to duplicate
-	TMap<uint32, TArray<FWorldPartitionRuntimeCellObjectMapping*>> PackagesToDuplicate;
+	TMap<uint64, TArray<FWorldPartitionRuntimeCellObjectMapping*>> PackagesToDuplicate;
 	check(!bLoadForPlay || InOutInstancingContext);
 
 	for (FWorldPartitionRuntimeCellObjectMapping& PackageObjectMapping : InActorPackages)
@@ -290,7 +290,7 @@ bool FWorldPartitionLevelHelper::LoadActors(ULevel* InDestLevel, TArrayView<FWor
 
 			if (LoadedPackage)
 			{
-				FName DuplicatePackageName(*FString::Printf(TEXT("%s_%08x"), *LoadedPackage->GetName(), LevelInstanceID));
+				FName DuplicatePackageName(*FString::Printf(TEXT("%s_%016llx"), *LoadedPackage->GetName(), LevelInstanceID));
 				UPackage* DuplicatedPackage = InPackageCache.FindPackage(DuplicatePackageName);
 				UWorld* DuplicatedWorld = nullptr;
 
@@ -325,7 +325,7 @@ bool FWorldPartitionLevelHelper::LoadActors(ULevel* InDestLevel, TArrayView<FWor
 					AActor* DuplicatedActor = FindObject<AActor>(DuplicatedWorld->PersistentLevel, *ActorName);
 					check(DuplicatedActor);
 
-					DuplicatedActor->Rename(*FString::Printf(TEXT("%s_%08x"), *DuplicatedActor->GetName(), Mapping->ContainerID), InDestLevel, REN_NonTransactional | REN_ForceNoResetLoaders | REN_DoNotDirty | REN_DontCreateRedirectors);
+					DuplicatedActor->Rename(*FString::Printf(TEXT("%s_%016llx"), *DuplicatedActor->GetName(), Mapping->ContainerID), InDestLevel, REN_NonTransactional | REN_ForceNoResetLoaders | REN_DoNotDirty | REN_DontCreateRedirectors);
 					USceneComponent* RootComponent = DuplicatedActor->GetRootComponent();
 					
 					FLevelUtils::FApplyLevelTransformParams TransformParams(nullptr, Mapping->ContainerTransform);
diff --git a/Engine/Source/Runtime/Engine/Private/WorldPartition/WorldPartitionRuntimeLevelStreamingCell.cpp b/Engine/Source/Runtime/Engine/Private/WorldPartition/WorldPartitionRuntimeLevelStreamingCell.cpp
--- a/Engine/Source/Runtime/Engine/Private/WorldPartition/WorldPartitionRuntimeLevelStreamingCell.cpp
+++ b/Engine/Source/Runtime/Engine/Private/WorldPartition/WorldPartitionRuntimeLevelStreamingCell.cpp
@@ -69,7 +69,7 @@ void UWorldPartitionRuntimeLevelStreamingCell::SetIsAlwaysLoaded(bool bInIsAlway
 
 #if WITH_EDITOR
 
-void UWorldPartitionRuntimeLevelStreamingCell::AddActorToCell(const FWorldPartitionActorDescView& ActorDescView, uint32 InContainerID, const FTransform& InContainerTransform, const UActorDescContainer* InContainer)
+void UWorldPartitionRuntimeLevelStreamingCell::AddActorToCell(const FWorldPartitionActorDescView& ActorDescView, uint64 InContainerID, const FTransform& InContainerTransform, const UActorDescContainer* InContainer)
 {
 	check(!ActorDescView.GetActorIsEditorOnly());
 	Packages.Emplace(ActorDescView.GetActorPackage(), ActorDescView.GetActorPath(), InContainerID, InContainerTransform, InContainer->GetContainerPackage());
diff --git a/Engine/Source/Runtime/Engine/Public/LevelInstance/LevelInstanceSubsystem.h b/Engine/Source/Runtime/Engine/Public/LevelInstance/LevelInstanceSubsystem.h
--- a/Engine/Source/Runtime/Engine/Public/LevelInstance/LevelInstanceSubsystem.h
+++ b/Engine/Source/Runtime/Engine/Public/LevelInstance/LevelInstanceSubsystem.h
@@ -99,7 +99,6 @@ private:
 	void ForEachActorInLevel(ULevel* Level, TFunctionRef<bool(AActor * LevelActor)> Operation) const;
 	void ForEachLevelInstanceAncestors(AActor* Actor, TFunctionRef<bool(ALevelInstance*)> Operation) const;
 	ALevelInstance* GetOwningLevelInstance(const ULevel* Level) const;
-	FLevelInstanceID ComputeLevelInstanceID(ALevelInstance* LevelInstanceActor) const;
 #if WITH_EDITOR
 	bool ForEachLevelInstanceChildImpl(const ALevelInstance* LevelInstanceActor, bool bRecursive, TFunctionRef<bool(const ALevelInstance*)> Operation) const;
 	bool ForEachLevelInstanceChildImpl(ALevelInstance* LevelInstanceActor, bool bRecursive, TFunctionRef<bool(ALevelInstance*)> Operation) const;
@@ -150,5 +149,5 @@ private:
 	TSet<FLevelInstanceID> LevelInstancesToUnload;
 	TMap<FLevelInstanceID, FLevelInstance> LevelInstances;
 	TMap<FLevelInstanceID, ALevelInstance*> RegisteredLevelInstances;
-	FLevelInstanceID PendingLevelInstanceToEdit = InvalidLevelInstanceID;
+	FLevelInstanceID PendingLevelInstanceToEdit;
 };
\ No newline at end of file
diff --git a/Engine/Source/Runtime/Engine/Public/LevelInstance/LevelInstanceTypes.h b/Engine/Source/Runtime/Engine/Public/LevelInstance/LevelInstanceTypes.h
--- a/Engine/Source/Runtime/Engine/Public/LevelInstance/LevelInstanceTypes.h
+++ b/Engine/Source/Runtime/Engine/Public/LevelInstance/LevelInstanceTypes.h
@@ -3,13 +3,39 @@
 #pragma once
 
 #include "CoreMinimal.h"
+#include "Misc/Guid.h"
 #include "LevelInstanceTypes.generated.h"
 
 // FLevelInstanceID is a runtime unique id that is computed from the Hash of LevelInstance Actor Guid and all its ancestor LevelInstance Actor Guids.
 // Resulting in a different ID for all instances whether they load the same level or not.
-using FLevelInstanceID = uint32;
+struct FLevelInstanceID
+{
+	FLevelInstanceID() {}
+	FLevelInstanceID(class ULevelInstanceSubsystem* LevelInstanceSubsystem, class ALevelInstance* Actor);
+
+	inline friend uint32 GetTypeHash(const FLevelInstanceID& Key)
+	{
+		return ::GetTypeHash(Key.GetHash());
+	}
+
+	inline bool operator!=(const FLevelInstanceID& Other) const
+	{
+		return !(*this == Other);
+	}
+
+	inline bool operator==(const FLevelInstanceID& Other) const
+	{
+		return Hash == Other.Hash && Guids == Other.Guids;
+	}
 
-static constexpr FLevelInstanceID InvalidLevelInstanceID = 0;
+	inline bool IsValid() const { return !Guids.IsEmpty(); }
+
+	inline uint64 GetHash() const { return Hash; }
+
+private:
+	uint64 Hash = 0;
+	TArray<FGuid> Guids;
+};
 
 UENUM()
 enum class ELevelInstanceCreationType : uint8
diff --git a/Engine/Source/Runtime/Engine/Public/WorldPartition/WorldPartitionActorCluster.h b/Engine/Source/Runtime/Engine/Public/WorldPartition/WorldPartitionActorCluster.h
--- a/Engine/Source/Runtime/Engine/Public/WorldPartition/WorldPartitionActorCluster.h
+++ b/Engine/Source/Runtime/Engine/Public/WorldPartition/WorldPartitionActorCluster.h
@@ -38,9 +38,9 @@ struct FActorClusterInstance;
 struct FActorContainerInstance
 {
 	FActorContainerInstance(const UActorDescContainer* InContainer, TMap<FGuid, FWorldPartitionActorDescView> InActorDescViewMap);
-	FActorContainerInstance(uint32 InID, const FTransform& InTransform, const FBox& InBounds, const TSet<FName>& InDataLayers, EContainerClusterMode InClusterMode, const UActorDescContainer* InContainer, TSet<FGuid> InChildContainers, TMap<FGuid, FWorldPartitionActorDescView> InActorDescViewMap);
+	FActorContainerInstance(uint64 InID, const FTransform& InTransform, const FBox& InBounds, const TSet<FName>& InDataLayers, EContainerClusterMode InClusterMode, const UActorDescContainer* InContainer, TSet<FGuid> InChildContainers, TMap<FGuid, FWorldPartitionActorDescView> InActorDescViewMap);
 	
-	uint32						ID;
+	uint64						ID;
 	FTransform					Transform;
 	FBox						Bounds;
 	EContainerClusterMode		ClusterMode;
@@ -118,7 +118,7 @@ public:
 
 private:
 	void CreateActorClusters();
-	void CreateContainerInstanceRecursive(uint32 ID, const FTransform& Transform, EContainerClusterMode ClusterMode, const UActorDescContainer* ActorDescContainer, const TSet<FName>& DataLayers, FBox& ParentBounds);
+	void CreateContainerInstanceRecursive(uint64 ID, const FTransform& Transform, EContainerClusterMode ClusterMode, const UActorDescContainer* ActorDescContainer, const TSet<FName>& DataLayers, FBox& ParentBounds);
 	const TArray<FActorCluster>& CreateActorClustersImpl(const FActorContainerInstance& ContainerInstance);
 	
 	// Init data
diff --git a/Engine/Source/Runtime/Engine/Public/WorldPartition/WorldPartitionRuntimeCell.h b/Engine/Source/Runtime/Engine/Public/WorldPartition/WorldPartitionRuntimeCell.h
--- a/Engine/Source/Runtime/Engine/Public/WorldPartition/WorldPartitionRuntimeCell.h
+++ b/Engine/Source/Runtime/Engine/Public/WorldPartition/WorldPartitionRuntimeCell.h
@@ -23,7 +23,7 @@ struct FWorldPartitionRuntimeCellObjectMapping
 #endif
 	{}
 
-	FWorldPartitionRuntimeCellObjectMapping(FName InPackage, FName InPath, uint32 InContainerID, const FTransform& InContainerTransform, FName InContainerPackage)
+	FWorldPartitionRuntimeCellObjectMapping(FName InPackage, FName InPath, uint64 InContainerID, const FTransform& InContainerTransform, FName InContainerPackage)
 #if WITH_EDITORONLY_DATA
 		: Package(InPackage)
 		, Path(InPath)
@@ -51,7 +51,7 @@ struct FWorldPartitionRuntimeCellObjectMapping
 	 * ID of the owning container instance
 	 */
 	UPROPERTY()
-	uint32 ContainerID;
+	uint64 ContainerID;
 
 	/** 
 	 * Transform of the owning container instance
@@ -119,7 +119,7 @@ class UWorldPartitionRuntimeCell : public UObject
 #if WITH_EDITOR
 	void SetDataLayers(const TArray<const UDataLayer*> InDataLayers);
 	void AddCellData(const UWorldPartitionRuntimeCellData* InCellData);
-	virtual void AddActorToCell(const FWorldPartitionActorDescView& ActorDescView, uint32 InContainerID, const FTransform& InContainerTransform, const UActorDescContainer* InContainer) PURE_VIRTUAL(UWorldPartitionRuntimeCell::AddActorToCell,);
+	virtual void AddActorToCell(const FWorldPartitionActorDescView& ActorDescView, uint64 InContainerID, const FTransform& InContainerTransform, const UActorDescContainer* InContainer) PURE_VIRTUAL(UWorldPartitionRuntimeCell::AddActorToCell,);
 	virtual int32 GetActorCount() const PURE_VIRTUAL(UWorldPartitionRuntimeCell::GetActorCount, return 0;);
 
 	// Cook methods
diff --git a/Engine/Source/Runtime/Engine/Public/WorldPartition/WorldPartitionRuntimeLevelStreamingCell.h b/Engine/Source/Runtime/Engine/Public/WorldPartition/WorldPartitionRuntimeLevelStreamingCell.h
--- a/Engine/Source/Runtime/Engine/Public/WorldPartition/WorldPartitionRuntimeLevelStreamingCell.h
+++ b/Engine/Source/Runtime/Engine/Public/WorldPartition/WorldPartitionRuntimeLevelStreamingCell.h
@@ -20,7 +20,7 @@ class UWorldPartitionRuntimeLevelStreamingCell : public UWorldPartitionRuntimeSp
 	EWorldPartitionRuntimeCellState GetCurrentState() const;
 
 #if WITH_EDITOR
-	virtual void AddActorToCell(const FWorldPartitionActorDescView& ActorDescView, uint32 InContainerID, const FTransform& InContainerTransform, const UActorDescContainer* InContainer) override;
+	virtual void AddActorToCell(const FWorldPartitionActorDescView& ActorDescView, uint64 InContainerID, const FTransform& InContainerTransform, const UActorDescContainer* InContainer) override;
 	virtual int32 GetActorCount() const override;
 	const TArray<FWorldPartitionRuntimeCellObjectMapping>& GetPackages() const { return Packages; }
 
