Properly set whether or not we need to compute extended bounds during mesh build

diff --git a/Engine/Source/Runtime/Engine/Private/StaticMeshBuild.cpp b/Engine/Source/Runtime/Engine/Private/StaticMeshBuild.cpp
--- a/Engine/Source/Runtime/Engine/Private/StaticMeshBuild.cpp
+++ b/Engine/Source/Runtime/Engine/Private/StaticMeshBuild.cpp
@@ -415,6 +415,10 @@ void UStaticMesh::BeginBuildInternal(FStaticMeshBuildContext* Context)
 			// Release the property now that it contains valid values to avoid game-thread stalls while we do the rest async
 			ReleaseAsyncProperty(EStaticMeshAsyncProperties::ExtendedBounds);
 		}
+		else
+		{
+			Context->bShouldComputeExtendedBounds = true;
+		}
 	}
 }
 
