Quantize values in FormatStatValueFloat(). Fixes wrong rounding in on-screen stats. Preserves the original formatting behavior of the integer part.

diff --git a/Engine/Source/Runtime/Engine/Private/StatsRender2.cpp b/Engine/Source/Runtime/Engine/Private/StatsRender2.cpp
--- a/Engine/Source/Runtime/Engine/Private/StatsRender2.cpp
+++ b/Engine/Source/Runtime/Engine/Private/StatsRender2.cpp
@@ -426,9 +426,10 @@ static int32 RenderCycle( const FComplexStatMessage& Item, class FCanvas* Canvas
 
 static FString FormatStatValueFloat(const float Value)
 {
-	const float Frac = FMath::Frac(Value);
+	const float QuantizedValue = FMath::RoundToFloat(Value * 100.0f) / 100.0f;
+	const float Frac = FMath::Frac(QuantizedValue);
 	// #TODO: Move to stats thread, add support for int64 type, int32 may not be sufficient all the time.
-	const int32 Integer = FMath::FloorToInt(Value);
+	const int32 Integer = FMath::FloorToInt(QuantizedValue);
 	const FString IntString = FString::FormatAsNumber(Integer);
 	const FString FracString = FString::Printf(TEXT("%0.2f"), Frac);
 	const FString Result = FString::Printf(TEXT("%s.%s"), *IntString, *FracString.Mid(2));
