Skin cache - add mesh name and LOD index info to gpu draw events

diff --git a/Engine/Source/Runtime/Engine/Private/GPUSkinCache.cpp b/Engine/Source/Runtime/Engine/Private/GPUSkinCache.cpp
--- a/Engine/Source/Runtime/Engine/Private/GPUSkinCache.cpp
+++ b/Engine/Source/Runtime/Engine/Private/GPUSkinCache.cpp
@@ -1120,6 +1120,11 @@ void FGPUSkinCache::DispatchUpdateSkinTangents(FRHICommandListImmediate& RHICmdL
 	{
 		// no need to clear the intermediate buffer because we create it cleared and clear it after each usage in the per vertex pass
 
+		FString MeshName = TEXT("None");
+#if !UE_BUILD_SHIPPING
+		MeshName = Entry->GPUSkin->DebugName.ToString();
+#endif // !UE_BUILD_SHIPPING
+
 		FSkeletalMeshRenderData& SkelMeshRenderData = Entry->GPUSkin->GetSkeletalMeshRenderData();
 		int32 LODIndex = Entry->LOD;
 		FSkeletalMeshLODRenderData& LodData = SkelMeshRenderData.LODRenderData[Entry->LOD];
@@ -1187,8 +1192,8 @@ void FGPUSkinCache::DispatchUpdateSkinTangents(FRHICommandListImmediate& RHICmdL
 			uint32 NumTriangles = DispatchData.NumTriangles;
 			uint32 ThreadGroupCountValue = FMath::DivideAndRoundUp(NumTriangles, FBaseRecomputeTangentsPerTriangleShader::ThreadGroupSizeX);
 
-			SCOPED_DRAW_EVENTF(RHICmdList, SkinTangents_PerTrianglePass, TEXT("TangentsTri IndexStart=%d Tri=%d BoneInfluenceType=%d UVPrecision=%d"),
-				DispatchData.IndexBufferOffsetValue, DispatchData.NumTriangles, Entry->BoneInfluenceType, bFullPrecisionUV);
+			SCOPED_DRAW_EVENTF(RHICmdList, SkinTangents_PerTrianglePass, TEXT("TangentsTri Mesh=%s, LOD=%d, Chunk=%d, IndexStart=%d Tri=%d BoneInfluenceType=%d UVPrecision=%d"),
+				*MeshName, LODIndex, SectionIndex, DispatchData.IndexBufferOffsetValue, DispatchData.NumTriangles, Entry->BoneInfluenceType, bFullPrecisionUV);
 
 			FRHIComputeShader* ShaderRHI = Shader.GetComputeShader();
 			RHICmdList.SetComputeShader(ShaderRHI);
@@ -1220,8 +1225,8 @@ void FGPUSkinCache::DispatchUpdateSkinTangents(FRHICommandListImmediate& RHICmdL
 		}
 
 		{
-			SCOPED_DRAW_EVENTF(RHICmdList, SkinTangents_PerVertexPass, TEXT("TangentsVertex InputStreamStart=%d, OutputStreamStart=%d, Vert=%d"),
-				DispatchData.InputStreamStart, DispatchData.OutputStreamStart, DispatchData.NumVertices);
+			SCOPED_DRAW_EVENTF(RHICmdList, SkinTangents_PerVertexPass, TEXT("TangentsVertex Mesh=%s, LOD=%d, Chunk=%d, InputStreamStart=%d, OutputStreamStart=%d, Vert=%d"),
+				*MeshName, LODIndex, SectionIndex, DispatchData.InputStreamStart, DispatchData.OutputStreamStart, DispatchData.NumVertices);
 			//#todo-gpuskin Feature level?
 			auto* GlobalShaderMap = GetGlobalShaderMap(GetFeatureLevel());
 			TShaderMapRef<FRecomputeTangentsPerVertexPassCS<0>> ComputeShader0(GlobalShaderMap);
@@ -1806,9 +1811,14 @@ void FGPUSkinCache::DispatchUpdateSkinning(FRHICommandListImmediate& RHICmdList,
 	FGPUSkinCacheEntry::FSectionDispatchData& DispatchData = Entry->DispatchData[Section];
 	FGPUBaseSkinVertexFactory::FShaderDataType& ShaderData = DispatchData.SourceVertexFactory->GetShaderData();
 
+	FString MeshName = TEXT("None");
+#if !UE_BUILD_SHIPPING
+	MeshName = Entry->GPUSkin->DebugName.ToString();
+#endif // !UE_BUILD_SHIPPING
+
 	SCOPED_DRAW_EVENTF(RHICmdList, SkinCacheDispatch,
-		TEXT("Skinning%d%d%d Chunk=%d InStreamStart=%d OutStart=%d Vert=%d Morph=%d/%d"),
-		(int32)Entry->bUse16BitBoneIndex, (int32)Entry->BoneInfluenceType, DispatchData.SkinType,
+		TEXT("Skinning%d%d%d Mesh=%s LOD=%d Chunk=%d InStreamStart=%d OutStart=%d Vert=%d Morph=%d/%d"),
+		(int32)Entry->bUse16BitBoneIndex, (int32)Entry->BoneInfluenceType, DispatchData.SkinType, *MeshName, Entry->GPUSkin->GetLOD(),
 		DispatchData.SectionIndex, DispatchData.InputStreamStart, DispatchData.OutputStreamStart, DispatchData.NumVertices, Entry->MorphBuffer != 0, DispatchData.MorphBufferOffset);
 	auto* GlobalShaderMap = GetGlobalShaderMap(GetFeatureLevel());
 	TShaderMapRef<TGPUSkinCacheCS<0>> SkinCacheCS000(GlobalShaderMap);		// 16bit_0, BoneInfluenceType_0, SkinType_0
