Downgrade the "failed to put and then read back" error to a display, since the error fires in DDCCommandlet due to use of non-readable backend PakFileDerivedDataBackend, and the failure to read does not cause a problem in the DDCCommandlet.
#rnx

diff --git a/Engine/Source/Runtime/Engine/Private/TextureDerivedDataTask.cpp b/Engine/Source/Runtime/Engine/Private/TextureDerivedDataTask.cpp
--- a/Engine/Source/Runtime/Engine/Private/TextureDerivedDataTask.cpp
+++ b/Engine/Source/Runtime/Engine/Private/TextureDerivedDataTask.cpp
@@ -306,7 +306,7 @@ void FTextureCacheDerivedDataWorker::BuildTexture(bool bReplaceExistingDDC)
 				bSucceeded = !bInlineMips || DerivedData->TryInlineMipData(BuildSettingsPerLayer[0].LODBiasWithCinematicMips, &Texture);
 				if (!bSucceeded)
 				{
-					UE_LOG(LogTexture, Error, TEXT("Failed to put and then read back mipmap data from DDC for %s"), *Texture.GetPathName());
+					UE_LOG(LogTexture, Display, TEXT("Failed to put and then read back mipmap data from DDC for %s"), *Texture.GetPathName());
 				}
 			}
 			else
@@ -396,7 +396,7 @@ void FTextureCacheDerivedDataWorker::BuildTexture(bool bReplaceExistingDDC)
 			bSucceeded = !bInlineMips || DerivedData->TryInlineMipData(BuildSettingsPerLayer[0].LODBiasWithCinematicMips, &Texture);
 			if (!bSucceeded)
 			{
-				UE_LOG(LogTexture, Error, TEXT("Failed to put and then read back mipmap data from DDC for %s"), *Texture.GetPathName());
+				UE_LOG(LogTexture, Display, TEXT("Failed to put and then read back mipmap data from DDC for %s"), *Texture.GetPathName());
 			}
 		}
 		else
