World Partition: always detach standalone objects in actor packages after the GC reachability pass to ensure we don't have lingering actor packages, resulting in invalid actors loading afterward.
This is a symptom of not being able to properly gather all external actor references in the actor descriptor, which will be addressed later.

diff --git a/Engine/Source/Runtime/Engine/Private/WorldPartition/WorldPartition.cpp b/Engine/Source/Runtime/Engine/Private/WorldPartition/WorldPartition.cpp
--- a/Engine/Source/Runtime/Engine/Private/WorldPartition/WorldPartition.cpp
+++ b/Engine/Source/Runtime/Engine/Private/WorldPartition/WorldPartition.cpp
@@ -185,6 +185,33 @@ UWorldPartition::UWorldPartition(const FObjectInitializer& ObjectInitializer)
 }
 
 #if WITH_EDITOR
+void UWorldPartition::OnGCPostReachabilityAnalysis()
+{
+	for (FRawObjectIterator It; It; ++It)
+	{
+		if (AActor* Actor = Cast<AActor>(static_cast<UObject*>(It->Object)))
+		{
+			if (Actor->IsUnreachable() && Actor->IsPackageExternal())
+			{
+				ForEachObjectWithPackage(Actor->GetPackage(), [this, Actor](UObject* Object)
+				{
+					if (Object->HasAnyFlags(RF_Standalone))
+					{
+						UE_LOG(LogWorldPartition, Warning, TEXT("Actor %s is unreachable without properly detaching object %s in its package"), *Actor->GetPathName(), *Object->GetPathName());
+
+						Object->ClearFlags(RF_Standalone);
+
+						// Make sure we trigger a second GC at the next tick to properly destroy packages there were fixed in this pass
+						bForceGarbageCollection = true;
+						bForceGarbageCollectionPurge = true;
+					}
+					return true;
+				}, false);
+			}
+		}
+	}
+}
+
 void UWorldPartition::OnPreBeginPIE(bool bStartSimulate)
 {
 	check(!bIsPIE);
@@ -519,6 +546,11 @@ void UWorldPartition::RegisterDelegates()
 		FEditorDelegates::PrePIEEnded.AddUObject(this, &UWorldPartition::OnPrePIEEnded);
 		FEditorDelegates::CancelPIE.AddUObject(this, &UWorldPartition::OnCancelPIE);
 		FGameDelegates::Get().GetEndPlayMapDelegate().AddUObject(this, &UWorldPartition::OnEndPlay);
+
+		if (IsMainWorldPartition())
+		{
+			FCoreUObjectDelegates::PostReachabilityAnalysis.AddUObject(this, &UWorldPartition::OnGCPostReachabilityAnalysis);
+		}
 	}
 #endif
 
@@ -539,6 +571,11 @@ void UWorldPartition::UnregisterDelegates()
 		FEditorDelegates::PrePIEEnded.RemoveAll(this);
 		FEditorDelegates::CancelPIE.RemoveAll(this);
 		FGameDelegates::Get().GetEndPlayMapDelegate().RemoveAll(this);
+
+		if (!IsEngineExitRequested() && IsMainWorldPartition())
+		{
+			FCoreUObjectDelegates::PostReachabilityAnalysis.RemoveAll(this);
+		}
 	}
 #endif
 
diff --git a/Engine/Source/Runtime/Engine/Public/WorldPartition/WorldPartition.h b/Engine/Source/Runtime/Engine/Public/WorldPartition/WorldPartition.h
--- a/Engine/Source/Runtime/Engine/Public/WorldPartition/WorldPartition.h
+++ b/Engine/Source/Runtime/Engine/Public/WorldPartition/WorldPartition.h
@@ -93,6 +93,8 @@ private:
 		
 	void FlushStreaming();
 
+	void OnGCPostReachabilityAnalysis();
+
 	// PIE/Game Methods
 	void OnPreBeginPIE(bool bStartSimulate);
 	void OnPrePIEEnded(bool bWasSimulatingInEditor);
