Wraps BlendSpace input properly

diff --git a/Engine/Source/Runtime/Engine/Private/Animation/AnimInterpFilter.cpp b/Engine/Source/Runtime/Engine/Private/Animation/AnimInterpFilter.cpp
--- a/Engine/Source/Runtime/Engine/Private/Animation/AnimInterpFilter.cpp
+++ b/Engine/Source/Runtime/Engine/Private/Animation/AnimInterpFilter.cpp
@@ -142,6 +142,37 @@ void FFIRFilterTimeBased::RefreshValidFilters()
 	}
 }
 
+void FFIRFilterTimeBased::WrapToValue(float Input, float Range)
+{
+	if (Range <= 0.0f)
+	{
+		return;
+	}
+	float HalfRange = Range / 2.0f;
+
+	switch (InterpolationType)
+	{
+	case BSIT_ExponentialDecay:
+	case BSIT_SpringDamper:
+	{
+		if (FilterWindow.Num() != 0)
+		{
+			FilterWindow[0].Input = FMath::Wrap(FilterWindow[0].Input, Input - HalfRange, Input + HalfRange);
+		}
+	}
+	break;
+	default:
+	{
+		for (int32 Index = 0; Index != FilterWindow.Num() ; ++Index)
+		{
+			FilterWindow[Index].Input = FMath::Wrap(FilterWindow[Index].Input, Input - HalfRange, Input + HalfRange);
+		}
+	}
+	break;
+	}
+}
+
+
 float FFIRFilterTimeBased::UpdateAndGetFilteredData(float Input, float DeltaTime)
 {
 	if (DeltaTime > KINDA_SMALL_NUMBER)
diff --git a/Engine/Source/Runtime/Engine/Private/Animation/BlendSpace.cpp b/Engine/Source/Runtime/Engine/Private/Animation/BlendSpace.cpp
--- a/Engine/Source/Runtime/Engine/Private/Animation/BlendSpace.cpp
+++ b/Engine/Source/Runtime/Engine/Private/Animation/BlendSpace.cpp
@@ -1676,9 +1676,16 @@ FVector UBlendSpace::FilterInput(FBlendFilter* Filter, const FVector& BlendInput
 	UpdateFilterParams(Filter);
 #endif
 	FVector FilteredBlendInput;
-	FilteredBlendInput.X = Filter->FilterPerAxis[0].UpdateAndGetFilteredData(BlendInput.X, DeltaTime);
-	FilteredBlendInput.Y = Filter->FilterPerAxis[1].UpdateAndGetFilteredData(BlendInput.Y, DeltaTime);
-	FilteredBlendInput.Z = Filter->FilterPerAxis[2].UpdateAndGetFilteredData(BlendInput.Z, DeltaTime);
+	for (int32 AxisIndex = 0; AxisIndex < 3; ++AxisIndex)
+	{
+		if (BlendParameters[AxisIndex].bWrapInput)
+		{
+			Filter->FilterPerAxis[AxisIndex].WrapToValue(
+				BlendInput[AxisIndex], BlendParameters[AxisIndex].Max - BlendParameters[AxisIndex].Min);
+		}
+		FilteredBlendInput[AxisIndex] = Filter->FilterPerAxis[AxisIndex].UpdateAndGetFilteredData(
+			BlendInput[AxisIndex], DeltaTime);
+	}
 	return FilteredBlendInput;
 }
 
diff --git a/Engine/Source/Runtime/Engine/Public/AnimInterpFilter.h b/Engine/Source/Runtime/Engine/Public/AnimInterpFilter.h
--- a/Engine/Source/Runtime/Engine/Public/AnimInterpFilter.h
+++ b/Engine/Source/Runtime/Engine/Public/AnimInterpFilter.h
@@ -146,6 +146,10 @@ public:
 
 	// This adds Input to the stack representing an update of DeltaTime, and returns the new filtered value. 
 	float UpdateAndGetFilteredData(float Input, float DeltaTime);
+
+	// Wraps the internal state by steps of Range so that it is as close as possible to Input
+	void WrapToValue(float Input, float Range);
+
 	bool IsValid() const { return WindowDuration > 0.f; }
 	float LastOutput;
 
