Small fix to the previous fix for BlendSpace wrapping, when FIR filters are being used.

diff --git a/Engine/Source/Runtime/Engine/Private/Animation/AnimInterpFilter.cpp b/Engine/Source/Runtime/Engine/Private/Animation/AnimInterpFilter.cpp
--- a/Engine/Source/Runtime/Engine/Private/Animation/AnimInterpFilter.cpp
+++ b/Engine/Source/Runtime/Engine/Private/Animation/AnimInterpFilter.cpp
@@ -163,9 +163,18 @@ void FFIRFilterTimeBased::WrapToValue(float Input, float Range)
 	break;
 	default:
 	{
-		for (int32 Index = 0; Index != FilterWindow.Num() ; ++Index)
+		if (IsValid())
 		{
-			FilterWindow[Index].Input = FMath::Wrap(FilterWindow[Index].Input, Input - HalfRange, Input + HalfRange);
+			float NewLastOutput = FMath::Wrap(LastOutput, Input - HalfRange, Input + HalfRange);
+			float Delta = NewLastOutput - LastOutput;
+			if (Delta)
+			{
+				LastOutput = NewLastOutput;
+				for (int32 Index = 0; Index != FilterWindow.Num() ; ++Index)
+				{
+					FilterWindow[Index].Input += Delta;
+				}
+			}
 		}
 	}
 	break;
