Character Movement component now has a larger tolerance for physics solver velocity errors when pushing objects.
Solving Issue: Unpushable cube can be pushed on TM-CharacterMovement map of QAGame with Chaos enabled

diff --git a/Engine/Source/Runtime/Engine/Private/Components/CharacterMovementComponent.cpp b/Engine/Source/Runtime/Engine/Private/Components/CharacterMovementComponent.cpp
--- a/Engine/Source/Runtime/Engine/Private/Components/CharacterMovementComponent.cpp
+++ b/Engine/Source/Runtime/Engine/Private/Components/CharacterMovementComponent.cpp
@@ -6851,8 +6851,8 @@ void UCharacterMovementComponent::ApplyImpactPhysicsForces(const FHitResult& Imp
 				}
 
 				Force *= PushForceModificator;
-
-				if (ComponentVelocity.IsNearlyZero())
+				const float ZeroVelocityTolerance = 1.0f;
+				if (ComponentVelocity.IsNearlyZero(ZeroVelocityTolerance))
 				{
 					Force *= InitialPushForceFactor;
 					ImpactComponent->AddImpulseAtLocation(Force, ForcePoint, Impact.BoneName);
@@ -11943,4 +11943,4 @@ void UCharacterMovementComponent::FlushServerMoves()
 ETeleportType UCharacterMovementComponent::GetTeleportType() const
 { 
 	return bJustTeleported || bNetworkLargeClientCorrection ? ETeleportType::TeleportPhysics : ETeleportType::None;
-}
\ No newline at end of file
+}
