Fix crash when creating a new level.

diff --git a/Engine/Source/Runtime/Engine/Private/WorldPartition/WorldPartition.cpp b/Engine/Source/Runtime/Engine/Private/WorldPartition/WorldPartition.cpp
--- a/Engine/Source/Runtime/Engine/Private/WorldPartition/WorldPartition.cpp
+++ b/Engine/Source/Runtime/Engine/Private/WorldPartition/WorldPartition.cpp
@@ -547,10 +547,7 @@ void UWorldPartition::RegisterDelegates()
 		FEditorDelegates::CancelPIE.AddUObject(this, &UWorldPartition::OnCancelPIE);
 		FGameDelegates::Get().GetEndPlayMapDelegate().AddUObject(this, &UWorldPartition::OnEndPlay);
 
-		if (IsMainWorldPartition())
-		{
-			FCoreUObjectDelegates::PostReachabilityAnalysis.AddUObject(this, &UWorldPartition::OnGCPostReachabilityAnalysis);
-		}
+		FCoreUObjectDelegates::PostReachabilityAnalysis.AddUObject(this, &UWorldPartition::OnGCPostReachabilityAnalysis);
 	}
 #endif
 
@@ -572,7 +569,7 @@ void UWorldPartition::UnregisterDelegates()
 		FEditorDelegates::CancelPIE.RemoveAll(this);
 		FGameDelegates::Get().GetEndPlayMapDelegate().RemoveAll(this);
 
-		if (!IsEngineExitRequested() && IsMainWorldPartition())
+		if (!IsEngineExitRequested())
 		{
 			FCoreUObjectDelegates::PostReachabilityAnalysis.RemoveAll(this);
 		}
