Ensure the content of World Partition runtime cells is shown/hidden in a single frame
Required to ensure seamless transition to HLOD cells.

diff --git a/Engine/Source/Runtime/Engine/Private/Components/SceneComponent.cpp b/Engine/Source/Runtime/Engine/Private/Components/SceneComponent.cpp
--- a/Engine/Source/Runtime/Engine/Private/Components/SceneComponent.cpp
+++ b/Engine/Source/Runtime/Engine/Private/Components/SceneComponent.cpp
@@ -3004,6 +3004,13 @@ bool USceneComponent::IsVisible() const
 	{
 		return false;
 	}
+
+	// If the level has the bRequireFullVisibilityToRender flag and it ain't visible yet, this component should not be visible.
+	const ULevel* ComponentLevel = GetComponentLevel();
+	if (ComponentLevel && ComponentLevel->bRequireFullVisibilityToRender && !ComponentLevel->bIsVisible)
+	{
+		return false;
+	}
 	
 	return (GetVisibleFlag() && (!CachedLevelCollection || CachedLevelCollection->IsVisible())); 
 }
diff --git a/Engine/Source/Runtime/Engine/Private/LevelStreaming.cpp b/Engine/Source/Runtime/Engine/Private/LevelStreaming.cpp
--- a/Engine/Source/Runtime/Engine/Private/LevelStreaming.cpp
+++ b/Engine/Source/Runtime/Engine/Private/LevelStreaming.cpp
@@ -594,6 +594,12 @@ void ULevelStreaming::UpdateStreamingState(bool& bOutUpdateAgain, bool& bOutRede
 
 				if (LoadedLevel->bIsVisible)
 				{
+					if (LoadedLevel->bRequireFullVisibilityToRender)
+					{
+						// Make sure components from this level are made visible
+						LoadedLevel->MarkLevelComponentsRenderStateDirty();
+					}
+
 					// immediately discard previous level
 					DiscardPendingUnloadLevel(World);
 
@@ -632,6 +638,8 @@ void ULevelStreaming::UpdateStreamingState(bool& bOutUpdateAgain, bool& bOutRede
 			const bool LevelBecameInvisible = bWasVisible && !LoadedLevel->bIsVisible;
 			if (LoadedLevel->bRequireFullVisibilityToRender && LevelBecameInvisible)
 			{
+				// Make sure components from this level are hidden
+				LoadedLevel->MarkLevelComponentsRenderStateDirty();
 				RemoveLevelFromScene();
 			}
 
