Fix Delete Asset crash: replace check by if because of CDOs can be at the root of a level package

diff --git a/Engine/Source/Runtime/Engine/Private/Actor.cpp b/Engine/Source/Runtime/Engine/Private/Actor.cpp
--- a/Engine/Source/Runtime/Engine/Private/Actor.cpp
+++ b/Engine/Source/Runtime/Engine/Private/Actor.cpp
@@ -1122,8 +1122,7 @@ float AActor::GetActorTickInterval() const
 
 bool AActor::IsMainPackageActor() const
 {
-	check(IsPackageExternal());
-	return GetExternalPackage()->GetFName().ToString() == ULevel::GetActorPackageName(GetLevel()->GetPackage(), GetPathName());
+	return IsPackageExternal() && (GetExternalPackage()->GetFName().ToString() == ULevel::GetActorPackageName(GetLevel()->GetPackage(), GetPathName()));
 }
 
 AActor* AActor::FindActorInPackage(UPackage* InPackage)
