Fix for call to FStreamedAudioChunk::GetCopy() when running without cooked data causing nullptr deref.

diff --git a/Engine/Source/Runtime/Engine/Classes/Sound/SoundWave.h b/Engine/Source/Runtime/Engine/Classes/Sound/SoundWave.h
--- a/Engine/Source/Runtime/Engine/Classes/Sound/SoundWave.h
+++ b/Engine/Source/Runtime/Engine/Classes/Sound/SoundWave.h
@@ -61,7 +61,8 @@ struct FStreamedAudioChunk
 	/** Serialization. */
 	void Serialize(FArchive& Ar, UObject* Owner, int32 ChunkIndex);
 
-	void GetCopy(void** OutChunkData);
+	/**  returns false if data retrieval failed */
+	bool GetCopy(void** OutChunkData);
 
 	/** Size of the chunk of data in bytes including zero padding */
 	int32 DataSize = 0;
diff --git a/Engine/Source/Runtime/Engine/Private/AudioStreamingCache.cpp b/Engine/Source/Runtime/Engine/Private/AudioStreamingCache.cpp
--- a/Engine/Source/Runtime/Engine/Private/AudioStreamingCache.cpp
+++ b/Engine/Source/Runtime/Engine/Private/AudioStreamingCache.cpp
@@ -870,6 +870,14 @@ TArrayView<uint8> FAudioChunkCache::GetChunk(const FChunkKey& InKey, bool bBlock
 
 		OutCacheOffset = FoundElement->CacheLookupID;
 
+// In cooked / packaged builds we need to retrieve the data from the pak file
+// if we are running with WITH_EDITORONLY_DATA, then this data has already been 
+// accessed via the DDC, and the underlying FByteBulk we access below has been 
+// cleared out
+#if WITH_EDITORONLY_DATA
+		bBlockForLoadCompletion = false;
+#endif // #if !WITH_EDITORONLY_DATA
+
 		if (bBlockForLoadCompletion)
 		{
 			FStreamedAudioChunk& Chunk = InKey.GetChunk(InKey.ChunkIndex);
@@ -882,9 +890,14 @@ TArrayView<uint8> FAudioChunkCache::GetChunk(const FChunkKey& InKey, bool bBlock
 
 			// Reallocate our chunk data This allows us to shrink if possible.
 			FoundElement->ChunkData = (uint8*)FMemory::Realloc(FoundElement->ChunkData, ChunkAudioDataSize);
-
 			void* DataDestPtr = FoundElement->ChunkData;
-			Chunk.GetCopy(&DataDestPtr);
+			const bool Result = Chunk.GetCopy(&DataDestPtr);
+
+			if (!Result)
+			{
+				UE_LOG(LogAudioStreamCaching, Warning, TEXT("Failed to retrieve chunk data from Bulk Data for soundwave: %s"), *InKey.SoundWaveName.ToString());
+				return TArrayView<uint8>();
+			}
 
 			MemoryCounterBytes += ChunkAudioDataSize;
 
diff --git a/Engine/Source/Runtime/Engine/Private/SoundWave.cpp b/Engine/Source/Runtime/Engine/Private/SoundWave.cpp
--- a/Engine/Source/Runtime/Engine/Private/SoundWave.cpp
+++ b/Engine/Source/Runtime/Engine/Private/SoundWave.cpp
@@ -186,7 +186,7 @@ void FStreamedAudioChunk::Serialize(FArchive& Ar, UObject* Owner, int32 ChunkInd
 #endif // #if WITH_EDITORONLY_DATA
 }
 
-void FStreamedAudioChunk::GetCopy(void** OutChunkData)
+bool FStreamedAudioChunk::GetCopy(void** OutChunkData)
 {
 	if (!CachedDataPtr)
 	{
@@ -205,8 +205,13 @@ void FStreamedAudioChunk::GetCopy(void** OutChunkData)
 		}
 	}
 
-	// todo: expose as shared ptr(?)
-	FMemory::Memcpy(*OutChunkData, CachedDataPtr, DataSize);
+	if (CachedDataPtr)
+	{
+		FMemory::Memcpy(*OutChunkData, CachedDataPtr, DataSize);
+		return true;
+	}
+
+	return false;
 }
 
 #if WITH_EDITORONLY_DATA
@@ -3245,6 +3250,7 @@ void FSoundWaveProxy::EnsureZerothChunkIsLoaded()
 		return;
 	}
 
+// #if WITH_EDITOR
 #ifdef WITH_EDITOR
 	// If we're running the editor, we'll need to retrieve the chunked audio from the DDC:
 	uint8* TempChunkBuffer = nullptr;
@@ -3260,7 +3266,7 @@ void FSoundWaveProxy::EnsureZerothChunkIsLoaded()
 #else // WITH_EDITOR
 	// Otherwise, the zeroth chunk is cooked out to SharedRuntimeData::RunningPlatformData, and we just need to retrieve it.
 	check(GetNumChunks() > 0);
-	FStreamedAudioChunk& ZerothChunk = GetChunk(0);
+	const FStreamedAudioChunk& ZerothChunk = GetChunk(0);
 	// Some sanity checks to ensure that the bulk size set up
 	UE_CLOG(ZerothChunk.BulkData.GetBulkDataSize() != ZerothChunk.DataSize, LogAudio, Warning, TEXT("Bulk data serialized out had a mismatched size with the DataSize field. Soundwave: %s Bulk Data Reported Size: %d Bulk Data Actual Size: %ld"), *GetFullName(), ZerothChunk.DataSize, ZerothChunk.BulkData.GetBulkDataSize());
 
@@ -3300,7 +3306,7 @@ const FStreamedAudioChunk& FSoundWaveProxy::GetChunk(uint32 ChunkIndex) const
 
 	check((ChunkIndex < (uint32)GetNumChunks()));
 	check(RunningPlatformData.IsValid());
-	return RunningPlatformData->Chunks[ChunkIndex];
+ 	return RunningPlatformData->Chunks[ChunkIndex];
 }
 
 int32 FSoundWaveProxy::GetChunkFromDDC(int32 ChunkIndex, uint8** OutChunkData, bool bMakeSureChunkIsLoaded)
