Remove the BoxTransitionDistance from the calculation of the BoxReflectionComponent's InfluenceBoundingRadius.

diff --git a/Engine/Source/Runtime/Engine/Private/Components/ReflectionCaptureComponent.cpp b/Engine/Source/Runtime/Engine/Private/Components/ReflectionCaptureComponent.cpp
--- a/Engine/Source/Runtime/Engine/Private/Components/ReflectionCaptureComponent.cpp
+++ b/Engine/Source/Runtime/Engine/Private/Components/ReflectionCaptureComponent.cpp
@@ -1169,7 +1169,7 @@ void UBoxReflectionCaptureComponent::UpdatePreviewShape()
 
 float UBoxReflectionCaptureComponent::GetInfluenceBoundingRadius() const
 {
-	return (GetComponentTransform().GetScale3D() + FVector(BoxTransitionDistance)).Size();
+	return GetComponentTransform().GetScale3D().Size();
 }
 
 #if WITH_EDITOR
