Passing "-csv" to the memreport command will now add the "-csv" argument to all subcommands. Added number of mips to ListTextures command output.

diff --git a/Engine/Source/Runtime/Engine/Private/UnrealEngine.cpp b/Engine/Source/Runtime/Engine/Private/UnrealEngine.cpp
--- a/Engine/Source/Runtime/Engine/Private/UnrealEngine.cpp
+++ b/Engine/Source/Runtime/Engine/Private/UnrealEngine.cpp
@@ -3585,9 +3585,10 @@ struct FSortedTexture
 	bool		bIsStreaming;
 	bool		bIsVirtual;
 	int32		UsageCount;
+	int32		NumMips;
 
 	/** Constructor, initializing every member variable with passed in values. */
-	FSortedTexture(	int32 InMaxAllowedSizeX, int32 InMaxAllowedSizeY, EPixelFormat InFormat, int32 InCurSizeX, int32 InCurSizeY, int32 InLODBias, int32 InMaxAllowedSize, int32 InCurrentSize, const FString& InName, int32 InLODGroup, bool bInIsStreaming, bool bInIsVirtual, int32 InUsageCount )
+	FSortedTexture(	int32 InMaxAllowedSizeX, int32 InMaxAllowedSizeY, EPixelFormat InFormat, int32 InCurSizeX, int32 InCurSizeY, int32 InLODBias, int32 InMaxAllowedSize, int32 InCurrentSize, const FString& InName, int32 InLODGroup, bool bInIsStreaming, bool bInIsVirtual, int32 InUsageCount, int32 InNumMips )
 		:	MaxAllowedSizeX( InMaxAllowedSizeX )
 		,	MaxAllowedSizeY( InMaxAllowedSizeY )
 		,	Format( InFormat )
@@ -3601,6 +3602,7 @@ struct FSortedTexture
 		,	bIsStreaming( bInIsStreaming )
 		,	bIsVirtual( bInIsVirtual )
 		,	UsageCount( InUsageCount )
+		,	NumMips( InNumMips )
 	{}
 };
 struct FCompareFSortedTexture
@@ -5389,7 +5391,8 @@ bool UEngine::HandleListTexturesCommand( const TCHAR* Cmd, FOutputDevice& Ar )
 				LODGroup, 
 				bIsStreamingTexture,
 				bIsVirtual,
-				UsageCount);
+				UsageCount,
+				NumMips);
 		}
 	}
 
@@ -5417,15 +5420,15 @@ bool UEngine::HandleListTexturesCommand( const TCHAR* Cmd, FOutputDevice& Ar )
 
 	if (bCSV)
 	{
-		Ar.Logf(TEXT(",Max Width,Max Height,Max Size (KB),Bias Authored,Current Width,Current Height,Current Size (KB),Format,LODGroup,Name,Streaming,VT,Usage Count"));
+		Ar.Logf(TEXT(",Max Width,Max Height,Max Size (KB),Bias Authored,Current Width,Current Height,Current Size (KB),Format,LODGroup,Name,Streaming,VT,Usage Count,NumMips"));
 	}
 	else if (!FPlatformProperties::RequiresCookedData())
 	{
-		Ar.Logf(TEXT("MaxAllowedSize: Width x Height (Size in KB, Authored Bias), Current/InMem: Width x Height (Size in KB), Format, LODGroup, Name, Streaming, VT, Usage Count"));
+		Ar.Logf(TEXT("MaxAllowedSize: Width x Height (Size in KB, Authored Bias), Current/InMem: Width x Height (Size in KB), Format, LODGroup, Name, Streaming, VT, Usage Count, NumMips"));
 	}
 	else
 	{
-		Ar.Logf(TEXT("Cooked/OnDisk: Width x Height (Size in KB, Authored Bias), Current/InMem: Width x Height (Size in KB), Format, LODGroup, Name, Streaming, VT, Usage Count"));
+		Ar.Logf(TEXT("Cooked/OnDisk: Width x Height (Size in KB, Authored Bias), Current/InMem: Width x Height (Size in KB), Format, LODGroup, Name, Streaming, VT, Usage Count, NumMips"));
 	}
 
 	for( int32 TextureIndex=0; TextureIndex<SortedTextures.Num(); TextureIndex++ )
@@ -5442,7 +5445,7 @@ bool UEngine::HandleListTexturesCommand( const TCHAR* Cmd, FOutputDevice& Ar )
 
 		if (bCSV)
 		{
-			Ar.Logf(TEXT(",%i, %i, %i, %s, %i, %i, %i, %s, %s, %s, %s, %s, %i"),
+			Ar.Logf(TEXT(",%i, %i, %i, %s, %i, %i, %i, %s, %s, %s, %s, %s, %i, %i"),
 				SortedTexture.MaxAllowedSizeX, SortedTexture.MaxAllowedSizeY, (SortedTexture.MaxAllowedSize + 512) / 1024,
 				*AuthoredBiasString,
 				SortedTexture.CurSizeX, SortedTexture.CurSizeY, (SortedTexture.CurrentSize + 512) / 1024,
@@ -5451,11 +5454,12 @@ bool UEngine::HandleListTexturesCommand( const TCHAR* Cmd, FOutputDevice& Ar )
 				*SortedTexture.Name,
 				SortedTexture.bIsStreaming ? TEXT("YES") : TEXT("NO"),
 				SortedTexture.bIsVirtual ? TEXT("YES") : TEXT("NO"),
-				SortedTexture.UsageCount);
+				SortedTexture.UsageCount,
+				SortedTexture.NumMips);
 		}
 		else
 		{
-			Ar.Logf(TEXT("%ix%i (%i KB, %s), %ix%i (%i KB), %s, %s, %s, %s, %s, %i"),
+			Ar.Logf(TEXT("%ix%i (%i KB, %s), %ix%i (%i KB), %s, %s, %s, %s, %s, %i, %i"),
 				SortedTexture.MaxAllowedSizeX, SortedTexture.MaxAllowedSizeY, (SortedTexture.MaxAllowedSize + 512) / 1024,
 				*AuthoredBiasString,
 				SortedTexture.CurSizeX, SortedTexture.CurSizeY, (SortedTexture.CurrentSize + 512) / 1024,
@@ -5464,7 +5468,8 @@ bool UEngine::HandleListTexturesCommand( const TCHAR* Cmd, FOutputDevice& Ar )
 				*SortedTexture.Name,
 				SortedTexture.bIsStreaming ? TEXT("YES") : TEXT("NO"),
 				SortedTexture.bIsVirtual ? TEXT("YES") : TEXT("NO"),
-				SortedTexture.UsageCount);
+				SortedTexture.UsageCount,
+				SortedTexture.NumMips);
 		}
 
 		if (bValidTextureGroup)
@@ -6424,6 +6429,7 @@ bool UEngine::HandleMemReportDeferredCommand( const TCHAR* Cmd, FOutputDevice& A
 
 	const bool bPerformSlowCommands = FParse::Param( Cmd, TEXT("FULL") );
 	const bool bLogOutputToFile = !FParse::Param( Cmd, TEXT("LOG") );
+	const bool bCSV = FParse::Param(Cmd, TEXT("CSV"));
 	FString InFileName;
 	FParse::Value(Cmd, TEXT("NAME="), InFileName);
 
@@ -6465,7 +6471,13 @@ bool UEngine::HandleMemReportDeferredCommand( const TCHAR* Cmd, FOutputDevice& A
 	{
 		for (FConfigSectionMap::TIterator It(*CommandsToRun); It; ++It)
 		{
-			Exec( InWorld, *It.Value().GetValue(), *ReportAr );
+			FString Command = It.Value().GetValue();
+			if (bCSV)
+			{
+				Command += " -csv";
+			}
+
+			Exec( InWorld, *Command, *ReportAr );
 			ReportAr->Logf( LINE_TERMINATOR );
 		}
 	}
@@ -6478,7 +6490,13 @@ bool UEngine::HandleMemReportDeferredCommand( const TCHAR* Cmd, FOutputDevice& A
 		{
 			for (FConfigSectionMap::TIterator It(*CommandsToRun); It; ++It)
 			{
-				Exec( InWorld, *It.Value().GetValue(), *ReportAr );
+				FString Command = It.Value().GetValue();
+				if (bCSV)
+				{
+					Command += " -csv";
+				}
+
+				Exec( InWorld, *Command, *ReportAr );
 				ReportAr->Logf( LINE_TERMINATOR );
 			}
 		}
