deprecate old Color::Quantize

diff --git a/Engine/Source/Runtime/Core/Public/Math/Color.h b/Engine/Source/Runtime/Core/Public/Math/Color.h
--- a/Engine/Source/Runtime/Core/Public/Math/Color.h
+++ b/Engine/Source/Runtime/Core/Public/Math/Color.h
@@ -322,6 +322,7 @@ struct FLinearColor
 	/** backwards compatible Quantize function name, does QuantizeFloor.
 	* @todo deprecate me
 	*/
+	UE_DEPRECATED(5.0,"Most callers of Quantize should have been calling QuantizeRound; to match old behavior use QuantizeFloor")
 	CORE_API FColor Quantize() const;
 
 	/** Quantizes the linear color and returns the result as a FColor with optional sRGB conversion. 
diff --git a/Engine/Source/Runtime/Engine/Classes/Kismet/KismetMathLibrary.inl b/Engine/Source/Runtime/Engine/Classes/Kismet/KismetMathLibrary.inl
--- a/Engine/Source/Runtime/Engine/Classes/Kismet/KismetMathLibrary.inl
+++ b/Engine/Source/Runtime/Engine/Classes/Kismet/KismetMathLibrary.inl
@@ -3562,7 +3562,10 @@ FLinearColor UKismetMathLibrary::RGBLinearToHSV(FLinearColor RGB)
 KISMET_MATH_FORCEINLINE
 FColor UKismetMathLibrary::LinearColor_Quantize(FLinearColor InColor)
 {
-	return InColor.Quantize();
+	// @todo remove LinearColor_Quantize from Kismet?
+	//	QuantizeRound should almost always be used instead
+	// call QuantizeFloor to match old behavior :
+	return InColor.QuantizeFloor();
 }
 
 KISMET_MATH_FORCEINLINE
