put DeprecatedFunction meta on Kismet LinearColor_Quantize

diff --git a/Engine/Source/Runtime/Engine/Classes/Kismet/KismetMathLibrary.h b/Engine/Source/Runtime/Engine/Classes/Kismet/KismetMathLibrary.h
--- a/Engine/Source/Runtime/Engine/Classes/Kismet/KismetMathLibrary.h
+++ b/Engine/Source/Runtime/Engine/Classes/Kismet/KismetMathLibrary.h
@@ -3010,7 +3010,7 @@ class ENGINE_API UKismetMathLibrary : public UBlueprintFunctionLibrary
 	static FColor Conv_LinearColorToColor(FLinearColor InLinearColor, bool InUseSRGB = true);
 
 	/** Quantizes the linear color and returns the result as an 8-bit color.  This bypasses the SRGB conversion. */
-	UFUNCTION(BlueprintPure, meta = (DisplayName = "Quantize to 8-bit (LinearColor)", ScriptMethod = "Quantize", Keywords = "cast convert"), Category = "Math|Color")
+	UFUNCTION(BlueprintPure, meta = (DeprecatedFunction, DeprecationMessage = "Use LinearColor_QuantizeRound instead for correct color conversion.", DisplayName = "Quantize to 8-bit (LinearColor)", ScriptMethod = "Quantize", Keywords = "cast convert"), Category = "Math|Color")
 	static FColor LinearColor_Quantize(FLinearColor InColor);
 
 	/** Quantizes the linear color with rounding and returns the result as an 8-bit color.  This bypasses the SRGB conversion. */
