Remove ~0.9M of ~4M allocations during startup of an internal demo project

Allocations came from ULevelStreamingDynamic::PostLoad()  -> IsValidStreamingLevel() -> FPackageName::DoesPackageExist() -> TryConvertToMountedPathComponents()

diff --git a/Engine/Source/Runtime/CoreUObject/Private/Misc/PackageName.cpp b/Engine/Source/Runtime/CoreUObject/Private/Misc/PackageName.cpp
--- a/Engine/Source/Runtime/CoreUObject/Private/Misc/PackageName.cpp
+++ b/Engine/Source/Runtime/CoreUObject/Private/Misc/PackageName.cpp
@@ -1115,12 +1115,15 @@ bool FPackageName::TryGetMountPointForPath(FStringView InFilePathOrPackageName,
 		return false;
 	}
 
-	FString PossibleAbsFilePath(FPaths::ConvertRelativePathToFull(FString(InFilePathOrPackageName)));
+	TStringBuilder<256> PossibleAbsFilePath;
+	FPathViews::ToAbsolutePath(InFilePathOrPackageName, PossibleAbsFilePath);
 	const FLongPackagePathsSingleton& Paths = FLongPackagePathsSingleton::Get();
 	FReadScopeLock ScopeLock(ContentMountPointCriticalSection);
 	for (const auto& Pair : Paths.ContentRootToPath)
 	{
-		FString RootFileAbsPath = FPaths::ConvertRelativePathToFull(Pair.ContentPath);
+		TStringBuilder<256> RootFileAbsPath;
+		FPathViews::ToAbsolutePath(Pair.ContentPath, RootFileAbsPath);
+
 		FStringView RelPath;
 		if (FPathViews::TryMakeChildPathRelativeTo(InFilePathOrPackageName, Pair.RootPath, RelPath))
 		{
diff --git a/Engine/Source/Runtime/CoreUObject/Private/Misc/PackagePath.cpp b/Engine/Source/Runtime/CoreUObject/Private/Misc/PackagePath.cpp
--- a/Engine/Source/Runtime/CoreUObject/Private/Misc/PackagePath.cpp
+++ b/Engine/Source/Runtime/CoreUObject/Private/Misc/PackagePath.cpp
@@ -246,9 +246,9 @@ bool FPackagePath::TryFromPackageName(FStringView InPackageName, FPackagePath& O
 
 bool FPackagePath::TryFromPackageName(FName InPackageName, FPackagePath& OutPackagePath)
 {
-	TStringBuilder<256> PackageNameString;
-	InPackageName.AppendString(PackageNameString);
-	return TryFromPackageName(PackageNameString, OutPackagePath);
+	TCHAR Buffer[FName::StringBufferSize];
+	FStringView PackageName(Buffer, InPackageName.ToString(Buffer));
+	return TryFromPackageName(PackageName, OutPackagePath);
 }
 
 bool FPackagePath::TryFromPackageName(const TCHAR* InPackageName, FPackagePath& OutPackagePath)
diff --git a/Engine/Source/Runtime/CoreUObject/Private/UObject/SoftObjectPath.cpp b/Engine/Source/Runtime/CoreUObject/Private/UObject/SoftObjectPath.cpp
--- a/Engine/Source/Runtime/CoreUObject/Private/UObject/SoftObjectPath.cpp
+++ b/Engine/Source/Runtime/CoreUObject/Private/UObject/SoftObjectPath.cpp
@@ -10,6 +10,7 @@
 #include "UObject/UObjectThreadContext.h"
 #include "UObject/CoreRedirects.h"
 #include "Misc/RedirectCollector.h"
+#include "String/Find.h"
 
 FSoftObjectPath::FSoftObjectPath(const UObject* InObject)
 {
@@ -57,6 +58,14 @@ void FSoftObjectPath::ToString(FStringBuilderBase& Builder) const
 	}
 }
 
+FName FSoftObjectPath::GetLongPackageFName() const
+{
+	TCHAR Buffer[NAME_SIZE];
+	FStringView PlainAssetPath(Buffer, /* len */ AssetPathName.GetPlainNameString(Buffer));
+	int32 DotPos = UE::String::FindFirstChar(PlainAssetPath, '.');
+	return DotPos == INDEX_NONE ? AssetPathName : FName(PlainAssetPath.Left(DotPos));
+}
+
 /** Helper function that adds info about the object currently being serialized when triggering an ensure about invalid soft object path */
 static FString GetObjectBeingSerializedForSoftObjectPath()
 {
diff --git a/Engine/Source/Runtime/CoreUObject/Public/UObject/SoftObjectPath.h b/Engine/Source/Runtime/CoreUObject/Public/UObject/SoftObjectPath.h
--- a/Engine/Source/Runtime/CoreUObject/Public/UObject/SoftObjectPath.h
+++ b/Engine/Source/Runtime/CoreUObject/Public/UObject/SoftObjectPath.h
@@ -111,6 +111,9 @@ struct COREUOBJECT_API FSoftObjectPath
 		return PackageName;
 	}
 
+	/** Returns /package/path, leaving off the asset name and sub object */
+	FName GetLongPackageFName() const;
+
 	/** Returns assetname string, leaving off the /package/path part and sub object */
 	FString GetAssetName() const
 	{
diff --git a/Engine/Source/Runtime/Engine/Private/LevelStreaming.cpp b/Engine/Source/Runtime/Engine/Private/LevelStreaming.cpp
--- a/Engine/Source/Runtime/Engine/Private/LevelStreaming.cpp
+++ b/Engine/Source/Runtime/Engine/Private/LevelStreaming.cpp
@@ -1553,7 +1553,7 @@ FName ULevelStreaming::GetWorldAssetPackageFName() const
 {
 	if (!bHasCachedWorldAssetPackageFName)
 	{
-		CachedWorldAssetPackageFName = FName(*FPackageName::ObjectPathToPackageName(WorldAsset.ToString()));
+		CachedWorldAssetPackageFName = WorldAsset.ToSoftObjectPath().GetLongPackageFName();
 		bHasCachedWorldAssetPackageFName = true;
 	}
 	return CachedWorldAssetPackageFName;
@@ -1800,8 +1800,9 @@ bool ULevelStreaming::IsValidStreamingLevel() const
 	const bool PIESession = GetWorld()->WorldType == EWorldType::PIE || GetOutermost()->HasAnyPackageFlags(PKG_PlayInEditor);
 	if (!PIESession && !WorldAsset.IsNull())
 	{
-		const FString WorldPackageName = GetWorldAssetPackageName();
-		return FPackageName::DoesPackageExist(WorldPackageName);
+		FName WorldPackageName = GetWorldAssetPackageFName();
+		FPackagePath WorldPackagePath;
+		return FPackagePath::TryFromPackageName(WorldPackageName, /* Out*/ WorldPackagePath) && FPackageName::DoesPackageExist(WorldPackagePath);
 	}
 	return true;
 }
