WorldPartition: Fix new unsaved map broken PIE

diff --git a/Engine/Source/Runtime/Engine/Private/WorldPartition/WorldPartitionRuntimeHash.cpp b/Engine/Source/Runtime/Engine/Private/WorldPartition/WorldPartitionRuntimeHash.cpp
--- a/Engine/Source/Runtime/Engine/Private/WorldPartition/WorldPartitionRuntimeHash.cpp
+++ b/Engine/Source/Runtime/Engine/Private/WorldPartition/WorldPartitionRuntimeHash.cpp
@@ -118,8 +118,7 @@ void UWorldPartitionRuntimeHash::CreateActorDescViewMap(const UActorDescContaine
 				if (Actor && 
 					Actor->IsPackageExternal() && 
 					!Actor->IsPendingKill() && 
-					Actor->GetPackage()->IsDirty() &&
-					Actor->GetPackage()->HasAnyPackageFlags(PKG_NewlyCreated))
+					!Container->GetActorDesc(Actor->GetActorGuid())) // Actor not on disk yet so not found in container
 				{
 					FWorldPartitionActorDesc* ActorDesc = ModifiedActorDescListForPIE.AddActor(Actor);
 					OutActorDescViewMap.Emplace(ActorDesc->GetGuid(), ActorDesc);
