New FVector4(ENoInit) constructor.
Various shader parameter constructors updated to use the new FVector4 constructor to avoid init when their state is about to be overwritten anyway.

diff --git a/Engine/Source/Runtime/Core/Public/Math/Vector4.h b/Engine/Source/Runtime/Core/Public/Math/Vector4.h
--- a/Engine/Source/Runtime/Core/Public/Math/Vector4.h
+++ b/Engine/Source/Runtime/Core/Public/Math/Vector4.h
@@ -80,6 +80,13 @@ public:
 	 */
 	explicit FVector4(EForceInit);
 
+	/**
+	 * Creates an uninitialized new vector.
+	 *
+	 * @param ENoInit Force uninitialized enum.
+	 */
+	FVector4(ENoInit);
+
 public:
 
 	/**
@@ -494,6 +501,11 @@ FORCEINLINE FVector4::FVector4(EForceInit)
 }
 
 
+FORCEINLINE FVector4::FVector4(ENoInit)
+{
+}
+
+
 FORCEINLINE FVector4::FVector4(FVector2D InXY, FVector2D InZW)
 	: X(InXY.X)
 	, Y(InXY.Y)
diff --git a/Engine/Source/Runtime/Engine/Private/InstanceUniformShaderParameters.cpp b/Engine/Source/Runtime/Engine/Private/InstanceUniformShaderParameters.cpp
--- a/Engine/Source/Runtime/Engine/Private/InstanceUniformShaderParameters.cpp
+++ b/Engine/Source/Runtime/Engine/Private/InstanceUniformShaderParameters.cpp
@@ -5,6 +5,7 @@
 #include "PrimitiveSceneInfo.h"
 
 FInstanceSceneShaderData::FInstanceSceneShaderData(const FPrimitiveInstance& Instance)
+	: Data(InPlace, NoInit)
 {
 	Setup(GetInstanceUniformShaderParameters(
 		Instance.LocalToWorld,
diff --git a/Engine/Source/Runtime/Engine/Private/LightmapUniformShaderParameters.cpp b/Engine/Source/Runtime/Engine/Private/LightmapUniformShaderParameters.cpp
--- a/Engine/Source/Runtime/Engine/Private/LightmapUniformShaderParameters.cpp
+++ b/Engine/Source/Runtime/Engine/Private/LightmapUniformShaderParameters.cpp
@@ -9,6 +9,7 @@
 IMPLEMENT_GLOBAL_SHADER_PARAMETER_STRUCT(FPrecomputedLightingUniformParameters, "PrecomputedLightingBuffer");
 
 FLightmapSceneShaderData::FLightmapSceneShaderData(const class FLightCacheInterface* LCI, ERHIFeatureLevel::Type FeatureLevel)
+	: Data(InPlace, NoInit)
 {
 	FPrecomputedLightingUniformParameters Parameters;
 	GetPrecomputedLightingParameters(FeatureLevel, Parameters, LCI);
diff --git a/Engine/Source/Runtime/Engine/Private/PrimitiveUniformShaderParameters.cpp b/Engine/Source/Runtime/Engine/Private/PrimitiveUniformShaderParameters.cpp
--- a/Engine/Source/Runtime/Engine/Private/PrimitiveUniformShaderParameters.cpp
+++ b/Engine/Source/Runtime/Engine/Private/PrimitiveUniformShaderParameters.cpp
@@ -48,15 +48,15 @@ void FSinglePrimitiveStructured::UploadToGPU()
 		void* LockedData = nullptr;
 
 		LockedData = RHILockBuffer(PrimitiveSceneDataBufferRHI, 0, FPrimitiveSceneShaderData::PrimitiveDataStrideInFloat4s * sizeof(FVector4), RLM_WriteOnly);
-		FPlatformMemory::Memcpy(LockedData, PrimitiveSceneData.Data, FPrimitiveSceneShaderData::PrimitiveDataStrideInFloat4s * sizeof(FVector4));
+		FPlatformMemory::Memcpy(LockedData, PrimitiveSceneData.Data.GetData(), FPrimitiveSceneShaderData::PrimitiveDataStrideInFloat4s * sizeof(FVector4));
 		RHIUnlockBuffer(PrimitiveSceneDataBufferRHI);
 
 		LockedData = RHILockBuffer(LightmapSceneDataBufferRHI, 0, FLightmapSceneShaderData::LightmapDataStrideInFloat4s * sizeof(FVector4), RLM_WriteOnly);
-		FPlatformMemory::Memcpy(LockedData, LightmapSceneData.Data, FLightmapSceneShaderData::LightmapDataStrideInFloat4s * sizeof(FVector4));
+		FPlatformMemory::Memcpy(LockedData, LightmapSceneData.Data.GetData(), FLightmapSceneShaderData::LightmapDataStrideInFloat4s * sizeof(FVector4));
 		RHIUnlockBuffer(LightmapSceneDataBufferRHI);
 
 		LockedData = RHILockBuffer(InstanceSceneDataBufferRHI, 0, FInstanceSceneShaderData::InstanceDataStrideInFloat4s * sizeof(FVector4), RLM_WriteOnly);
-		FPlatformMemory::Memcpy(LockedData, InstanceSceneData.Data, FInstanceSceneShaderData::InstanceDataStrideInFloat4s * sizeof(FVector4));
+		FPlatformMemory::Memcpy(LockedData, InstanceSceneData.Data.GetData(), FInstanceSceneShaderData::InstanceDataStrideInFloat4s * sizeof(FVector4));
 		RHIUnlockBuffer(InstanceSceneDataBufferRHI);
 	}
 
@@ -94,6 +94,7 @@ TGlobalResource<FSinglePrimitiveStructured> GIdentityPrimitiveBuffer;
 TGlobalResource<FSinglePrimitiveStructured> GTilePrimitiveBuffer;
 
 FPrimitiveSceneShaderData::FPrimitiveSceneShaderData(const FPrimitiveSceneProxy* RESTRICT Proxy)
+	: Data(InPlace, NoInit)
 {
 	bool bHasPrecomputedVolumetricLightmap;
 	FMatrix PreviousLocalToWorld;
diff --git a/Engine/Source/Runtime/Engine/Public/InstanceUniformShaderParameters.h b/Engine/Source/Runtime/Engine/Public/InstanceUniformShaderParameters.h
--- a/Engine/Source/Runtime/Engine/Public/InstanceUniformShaderParameters.h
+++ b/Engine/Source/Runtime/Engine/Public/InstanceUniformShaderParameters.h
@@ -8,6 +8,7 @@
 #include "RenderResource.h"
 #include "ShaderParameters.h"
 #include "UniformBuffer.h"
+#include "Containers/StaticArray.h"
 
 #define INSTANCE_SCENE_DATA_FLAG_CAST_SHADOWS			0x1
 #define INSTANCE_SCENE_DATA_FLAG_DETERMINANT_SIGN		0x2
@@ -165,9 +166,10 @@ struct FInstanceSceneShaderData
 	// Must match GetInstanceData() in SceneData.ush
 	enum { InstanceDataStrideInFloat4s = 10 };
 
-	FVector4 Data[InstanceDataStrideInFloat4s];
+	TStaticArray<FVector4, InstanceDataStrideInFloat4s> Data;
 
 	FInstanceSceneShaderData()
+		: Data(InPlace, NoInit)
 	{
 		Setup(GetInstanceUniformShaderParameters(
 			FMatrix::Identity,
@@ -186,6 +188,7 @@ struct FInstanceSceneShaderData
 	}
 
 	explicit FInstanceSceneShaderData(const FInstanceUniformShaderParameters& InstanceUniformShaderParameters)
+		: Data(InPlace, NoInit)
 	{
 		Setup(InstanceUniformShaderParameters);
 	}
diff --git a/Engine/Source/Runtime/Engine/Public/LightmapUniformShaderParameters.h b/Engine/Source/Runtime/Engine/Public/LightmapUniformShaderParameters.h
--- a/Engine/Source/Runtime/Engine/Public/LightmapUniformShaderParameters.h
+++ b/Engine/Source/Runtime/Engine/Public/LightmapUniformShaderParameters.h
@@ -7,6 +7,7 @@
 #include "RenderResource.h"
 #include "ShaderParameters.h"
 #include "UniformBuffer.h"
+#include "Containers/StaticArray.h"
 
 class FLightCacheInterface;
 
@@ -36,9 +37,10 @@ struct FLightmapSceneShaderData
 	// Must match usf
 	enum { LightmapDataStrideInFloat4s = 15 };
 
-	FVector4 Data[LightmapDataStrideInFloat4s];
+	TStaticArray<FVector4, LightmapDataStrideInFloat4s> Data;
 
 	FLightmapSceneShaderData()
+		: Data(InPlace, NoInit)
 	{
 		FPrecomputedLightingUniformParameters ShaderParameters;
 		GetDefaultPrecomputedLightingParameters(ShaderParameters);
@@ -46,6 +48,7 @@ struct FLightmapSceneShaderData
 	}
 
 	explicit FLightmapSceneShaderData(const FPrecomputedLightingUniformParameters& ShaderParameters)
+		: Data(InPlace, NoInit)
 	{
 		Setup(ShaderParameters);
 	}
diff --git a/Engine/Source/Runtime/Engine/Public/PrimitiveUniformShaderParameters.h b/Engine/Source/Runtime/Engine/Public/PrimitiveUniformShaderParameters.h
--- a/Engine/Source/Runtime/Engine/Public/PrimitiveUniformShaderParameters.h
+++ b/Engine/Source/Runtime/Engine/Public/PrimitiveUniformShaderParameters.h
@@ -11,6 +11,7 @@
 #include "InstanceUniformShaderParameters.h"
 #include "LightmapUniformShaderParameters.h"
 #include "UnifiedBuffer.h"
+#include "Containers/StaticArray.h"
 
 /** 
  * The uniform shader parameters associated with a primitive. 
@@ -282,15 +283,17 @@ struct FPrimitiveSceneShaderData
 	// Must match usf
 	enum { PrimitiveDataStrideInFloat4s = 37 };
 
-	FVector4 Data[PrimitiveDataStrideInFloat4s];
+	TStaticArray<FVector4, PrimitiveDataStrideInFloat4s> Data;
 
 	FPrimitiveSceneShaderData()
+		: Data(InPlace, NoInit)
 	{
 		static_assert(FPrimitiveSceneShaderData::PrimitiveDataStrideInFloat4s == FScatterUploadBuffer::PrimitiveDataStrideInFloat4s,"");
 		Setup(GetIdentityPrimitiveParameters());
 	}
 
 	explicit FPrimitiveSceneShaderData(const FPrimitiveUniformShaderParameters& PrimitiveUniformShaderParameters)
+		: Data(InPlace, NoInit)
 	{
 		Setup(PrimitiveUniformShaderParameters);
 	}
