Add backwards compatibility support for replicated TObjectPtr properties in replays.

diff --git a/Engine/Source/Runtime/Engine/Private/RepLayout.cpp b/Engine/Source/Runtime/Engine/Private/RepLayout.cpp
--- a/Engine/Source/Runtime/Engine/Private/RepLayout.cpp
+++ b/Engine/Source/Runtime/Engine/Private/RepLayout.cpp
@@ -887,8 +887,21 @@ static uint32 GetRepLayoutCmdCompatibleChecksum(
 	// Evolve checksum on name
 	uint32 CompatibleChecksum = FCrc::StrCrc32(*Property->GetName().ToLower(), InChecksum);	
 	
-	// Evolve by property type			
-	CompatibleChecksum = FCrc::StrCrc32(*Property->GetCPPType(nullptr, 0).ToLower(), CompatibleChecksum);
+	// Evolve by property type
+	const FObjectPtrProperty* const ObjectPtrProperty = CastField<const FObjectPtrProperty>(Property);
+
+	FString CPPType;
+	if (ObjectPtrProperty)
+	{
+		// To remain compatible with TObjectPtr, use the underlying pointer type in the checksum since the net-serialized data is compatible.
+		CPPType = ObjectPtrProperty->FObjectProperty::GetCPPType(nullptr, 0).ToLower();
+	}
+	else
+	{
+		CPPType = Property->GetCPPType(nullptr, 0).ToLower();
+	}
+
+	CompatibleChecksum = FCrc::StrCrc32(*CPPType, CompatibleChecksum);
 	
 	// Evolve by StaticArrayIndex (to make all unrolled static array elements unique)
 	if ((ServerConnection == nullptr) ||
