Angular motor setup scaling issue on velocity drive

diff --git a/Engine/Source/Runtime/Engine/Private/PhysicsEngine/Experimental/PhysInterface_Chaos.cpp b/Engine/Source/Runtime/Engine/Private/PhysicsEngine/Experimental/PhysInterface_Chaos.cpp
--- a/Engine/Source/Runtime/Engine/Private/PhysicsEngine/Experimental/PhysInterface_Chaos.cpp
+++ b/Engine/Source/Runtime/Engine/Private/PhysicsEngine/Experimental/PhysInterface_Chaos.cpp
@@ -470,7 +470,7 @@ void FPhysInterface_Chaos::UpdateAngularDrive_AssumesLocked(const FPhysicsConstr
 				}
 				else
 				{
-					Constraint->SetAngularDriveVelocityTarget(InDriveParams.AngularVelocityTarget);
+					Constraint->SetAngularDriveVelocityTarget(InDriveParams.AngularVelocityTarget * 2.0f * PI); // Rev/s to Rad/s
 				}
 			}
 
