DotProduct preshader will properly splat float1 values

diff --git a/Engine/Source/Runtime/Engine/Private/Shader/ShaderValue.cpp b/Engine/Source/Runtime/Engine/Private/Shader/ShaderValue.cpp
--- a/Engine/Source/Runtime/Engine/Private/Shader/ShaderValue.cpp
+++ b/Engine/Source/Runtime/Engine/Private/Shader/ShaderValue.cpp
@@ -337,21 +337,18 @@ inline FValue UnaryOp(const Operation& Op, const FValue& Value)
 	return Result;
 }
 
+inline int8 GetNumComponentsResult(int8 Lhs, int8 Rhs)
+{
+	// operations between scalar and non-scalar will splat the scalar value
+	// otherwise, operations should only be between types with same number of components
+	return (Lhs == 1 || Rhs == 1) ? FMath::Max(Lhs, Rhs) : FMath::Min(Lhs, Rhs);
+}
+
 template<typename Operation>
 inline FValue BinaryOp(const Operation& Op, const FValue& Lhs, const FValue& Rhs)
 {
-	int8 NumComponents;
-	if (Lhs.NumComponents == 1 || Rhs.NumComponents == 1)
-	{
-		// operations between scalar and non-scalar will splat the scalar value
-		NumComponents = FMath::Max(Lhs.NumComponents, Rhs.NumComponents);
-	}
-	else
-	{
-		// otherwise, operations should only be between types with same number of components
-		NumComponents = FMath::Min(Lhs.NumComponents, Rhs.NumComponents);
-	}
-
+	const int8 NumComponents = GetNumComponentsResult(Lhs.NumComponents, Rhs.NumComponents);
+	
 	FValue Result;
 	Result.NumComponents = NumComponents;
 
@@ -538,7 +535,7 @@ UE::Shader::FValue UE::Shader::Clamp(const FValue& Value, const FValue& Low, con
 
 UE::Shader::FValue UE::Shader::Dot(const FValue& Lhs, const FValue& Rhs)
 {
-	const int32 NumComponents = FMath::Min(Lhs.NumComponents, Rhs.NumComponents);
+	const int8 NumComponents = Private::GetNumComponentsResult(Lhs.NumComponents, Rhs.NumComponents);
 
 	FValue Result;
 	Result.NumComponents = 1;
