WorldPartition: wp.Editor.DumpActorDesc when called with no params will output information on selected actors

diff --git a/Engine/Source/Runtime/Engine/Private/WorldPartition/WorldPartition.cpp b/Engine/Source/Runtime/Engine/Private/WorldPartition/WorldPartition.cpp
--- a/Engine/Source/Runtime/Engine/Private/WorldPartition/WorldPartition.cpp
+++ b/Engine/Source/Runtime/Engine/Private/WorldPartition/WorldPartition.cpp
@@ -21,6 +21,8 @@
 #if WITH_EDITOR
 #include "LevelUtils.h"
 #include "Editor.h"
+#include "Editor/EditorEngine.h"
+#include "Selection.h"
 #include "FileHelpers.h"
 #include "HAL/FileManager.h"
 #include "LevelEditorViewport.h"
@@ -48,7 +50,23 @@ static FAutoConsoleCommand DumpActorDesc(
 	TEXT("Dump a specific actor descriptor on the console."),
 	FConsoleCommandWithArgsDelegate::CreateLambda([](const TArray<FString>& Args)
 	{
+		TArray<FString> ActorPaths;
 		if (Args.Num() > 0)
+		{
+			ActorPaths.Add(Args[0]);
+		}
+		else
+		{
+			for (FSelectionIterator SelectionIt(*GEditor->GetSelectedActors()); SelectionIt; ++SelectionIt)
+			{
+				if (const AActor* Actor = CastChecked<AActor>(*SelectionIt))
+				{
+					ActorPaths.Add(Actor->GetPathName());
+				}
+			}
+		}
+
+		if (!ActorPaths.IsEmpty())
 		{
 			if (UWorld* World = GEditor->GetEditorWorldContext().World())
 			{
@@ -56,22 +74,26 @@ static FAutoConsoleCommand DumpActorDesc(
 				{
 					if (UWorldPartition* WorldPartition = World->GetWorldPartition())
 					{
-						if (const FWorldPartitionActorDesc* ActorDesc = WorldPartition->GetActorDesc(Args[0]))
+						UE_LOG(LogWorldPartition, Log, TEXT("Guid, Class, Name, Package, BVCenterX, BVCenterY, BVCenterZ, BVExtentX, BVExtentY, BVExtentZ"));
+						for (const FString& ActorPath : ActorPaths)
 						{
-							UE_LOG(LogWorldPartition, Log, TEXT("Guid, Class, Name, Package, BVCenterX, BVCenterY, BVCenterZ, BVExtentX, BVExtentY, BVExtentZ"));
-							UE_LOG(LogWorldPartition, Log, TEXT("%s, %s, %s, %s, %s, %.2f, %.2f, %.2f, %.2f, %.2f, %.2f"),
-								*ActorDesc->GetGuid().ToString(), 
-								*ActorDesc->GetClass().ToString(), 
-								*ActorDesc->GetActorName().ToString(), 
-								*ActorDesc->GetActorPackage().ToString(), 
-								*ActorDesc->GetActorLabel().ToString(),
-								ActorDesc->GetBounds().GetCenter().X,
-								ActorDesc->GetBounds().GetCenter().Y,
-								ActorDesc->GetBounds().GetCenter().Z,
-								ActorDesc->GetBounds().GetExtent().X,
-								ActorDesc->GetBounds().GetExtent().Y,
-								ActorDesc->GetBounds().GetExtent().Z
-							);
+							if (const FWorldPartitionActorDesc* ActorDesc = WorldPartition->GetActorDesc(ActorPath))
+							{
+
+								UE_LOG(LogWorldPartition, Log, TEXT("%s, %s, %s, %s, %s, %.2f, %.2f, %.2f, %.2f, %.2f, %.2f"),
+									*ActorDesc->GetGuid().ToString(),
+									*ActorDesc->GetClass().ToString(),
+									*ActorDesc->GetActorName().ToString(),
+									*ActorDesc->GetActorPackage().ToString(),
+									*ActorDesc->GetActorLabel().ToString(),
+									ActorDesc->GetBounds().GetCenter().X,
+									ActorDesc->GetBounds().GetCenter().Y,
+									ActorDesc->GetBounds().GetCenter().Z,
+									ActorDesc->GetBounds().GetExtent().X,
+									ActorDesc->GetBounds().GetExtent().Y,
+									ActorDesc->GetBounds().GetExtent().Z
+								);
+							}
 						}
 					}
 				}
