Exposed to Blueprint methods to get UDataLayer object from Name, Label or FActorDataLayer.

diff --git a/Engine/Source/Runtime/Engine/Private/WorldPartition/DataLayer/DataLayerSubsystem.cpp b/Engine/Source/Runtime/Engine/Private/WorldPartition/DataLayer/DataLayerSubsystem.cpp
--- a/Engine/Source/Runtime/Engine/Private/WorldPartition/DataLayer/DataLayerSubsystem.cpp
+++ b/Engine/Source/Runtime/Engine/Private/WorldPartition/DataLayer/DataLayerSubsystem.cpp
@@ -86,13 +86,18 @@ void UDataLayerSubsystem::Deinitialize()
 	Super::Deinitialize();
 }
 
-UDataLayer* UDataLayerSubsystem::GetDataLayerFromLabel(const FName& InDataLayerLabel) const
+UDataLayer* UDataLayerSubsystem::GetDataLayer(const FActorDataLayer& InDataLayer) const
+{
+	return GetDataLayerFromName(InDataLayer.Name);
+}
+
+UDataLayer* UDataLayerSubsystem::GetDataLayerFromLabel(FName InDataLayerLabel) const
 {
 	const AWorldDataLayers* WorldDataLayers = AWorldDataLayers::Get(GetWorld());
 	return WorldDataLayers ? const_cast<UDataLayer*>(WorldDataLayers->GetDataLayerFromLabel(InDataLayerLabel)) : nullptr;
 }
 
-UDataLayer* UDataLayerSubsystem::GetDataLayerFromName(const FName& InDataLayerName) const
+UDataLayer* UDataLayerSubsystem::GetDataLayerFromName(FName InDataLayerName) const
 {
 	const AWorldDataLayers* WorldDataLayers = AWorldDataLayers::Get(GetWorld());
 	return WorldDataLayers ? const_cast<UDataLayer*>(WorldDataLayers->GetDataLayerFromName(InDataLayerName)) : nullptr;
diff --git a/Engine/Source/Runtime/Engine/Public/WorldPartition/DataLayer/DataLayerSubsystem.h b/Engine/Source/Runtime/Engine/Public/WorldPartition/DataLayer/DataLayerSubsystem.h
--- a/Engine/Source/Runtime/Engine/Public/WorldPartition/DataLayer/DataLayerSubsystem.h
+++ b/Engine/Source/Runtime/Engine/Public/WorldPartition/DataLayer/DataLayerSubsystem.h
@@ -51,6 +51,9 @@ public:
 	UFUNCTION(BlueprintCallable, Category = DataLayers)
 	void SetDataLayerState(const FActorDataLayer& InDataLayer, EDataLayerState InState);
 
+	UFUNCTION(BlueprintCallable, Category = DataLayers)
+	UDataLayer* GetDataLayer(const FActorDataLayer& InDataLayer) const;
+
 	UFUNCTION(BlueprintCallable, Category = DataLayers)
 	void SetDataLayerStateByLabel(const FName& InDataLayerLabel, EDataLayerState InState);
 		
@@ -65,10 +68,14 @@ public:
 
 	UPROPERTY(BlueprintAssignable)
 	FOnDataLayerStateChanged OnDataLayerStateChanged;
+
+	UFUNCTION(BlueprintCallable, Category = DataLayers)
+	UDataLayer* GetDataLayerFromLabel(FName InDataLayerLabel) const;
+
+	UFUNCTION(BlueprintCallable, Category = DataLayers)
+	UDataLayer* GetDataLayerFromName(FName InDataLayerName) const;
 	//~ End Blueprint callable functions
 
-	UDataLayer* GetDataLayerFromLabel(const FName& InDataLayerLabel) const;
-	UDataLayer* GetDataLayerFromName(const FName& InDataLayerName) const;
 	void SetDataLayerState(const UDataLayer* InDataLayer, EDataLayerState InState);
 	void SetDataLayerStateByName(const FName& InDataLayerName, EDataLayerState InState);
 	EDataLayerState GetDataLayerState(const UDataLayer* InDataLayer) const;
