Allow imported jpegs to continue to be stored as jpeg data in .uassets until the data is actually edited and needs to be re-compressed. This can reduce megascan assets by up to 40%. The feature is disabled by default and can be enabled for a project by setting the editor config value [TextureImporter]RetainJpegFormat=True.

* Texture Work
- When importing a jpeg we will now check the config file to see if we should try and keep the data in compressed JPEG format and decompress it as it is requested rather than returning the data in raw uncompressed format.
- If the data is kept in jpeg compressed format we can initialize the texture via the new method InitWithCompressedSourceData instead. (note that we can probably merge ::InitWithCompressedSourceData and ::Init but that introduces further risk, this way we know that ONLY textures that were originally jpeg can be going through ::InitWithCompressedSourceData for now)
- We now record the compression format of the bulkdata in FTextureSource via 'CompressionFormat' an enum rather than a single bool 'bPNGCompressed' which only wqorked for png format.
 - We should deprecate bPNGCompressed at some point and loading existing textures off disk will not have 'bPNGCompressed' and 'CompressionFormat' in sync, fixing this is out of scope for EA work.
- Note that when PNG compression is used we can decompress the data, use that, then recompress it without loss in quality, the same is not to be said for JPEG so if the data is modified by a call to ::LockMip/::Unlock mip the data will end up falling back to being stored as a compressed PNG or raw data. If this occurs we will log a warning.
- To reduce the number of places where this happens a new ::LockMipReadOnly method has been added, locking the texture source as read only will not replace the stored data when unlocked (as read only means no changes)
- Attempting to lock a texture for a different mode than it is currently locked for is considered a programming error and will raise an assert.
- Updated the 'SourceCompression' asset registry tag. This can be seen when hovering the mouse over the asset in the content browser and will show which compression is being used on the source data, png, jpeg or none.

* NormalMapIdentification/PaperAtlasTextureHelpers
- Switched to use LockMipReadOnly as we never modify the locked mip data.
- Every other remaining place in code that I could find that used LockMip() is actually calling that to edit the data.

* JPEG Decompression
- Worth noting that in early access the jpeg compression is slower than decompressing the same data if it was stored in compressed png format.
- On main the jpeg decompressor has been upgraded and the data ends up being smaller AND faster to decode, so wins all around.
- For early access we are willing to take the increase decoding cost to reduce the risk of the submit.

#[fyi] Jonathan.Bard

diff --git a/Engine/Config/BaseEditor.ini b/Engine/Config/BaseEditor.ini
--- a/Engine/Config/BaseEditor.ini
+++ b/Engine/Config/BaseEditor.ini
@@ -291,6 +291,7 @@ Class=UnrealEd.ActorFactory
 
 [TextureImporter]
 AllowNonPowerOfTwoTextures=True
+RetainJpegFormat=False
 
 [/Script/Cascade.CascadeConfiguration]
 +ModuleMenu_ModuleRejections=ParticleModuleTypeDataBase
diff --git a/Engine/Plugins/2D/Paper2D/Source/Paper2DEditor/Private/Atlasing/PaperAtlasTextureHelpers.cpp b/Engine/Plugins/2D/Paper2D/Source/Paper2DEditor/Private/Atlasing/PaperAtlasTextureHelpers.cpp
--- a/Engine/Plugins/2D/Paper2D/Source/Paper2DEditor/Private/Atlasing/PaperAtlasTextureHelpers.cpp
+++ b/Engine/Plugins/2D/Paper2D/Source/Paper2DEditor/Private/Atlasing/PaperAtlasTextureHelpers.cpp
@@ -33,7 +33,7 @@ bool FPaperAtlasTextureHelpers::ReadSpriteTexture(UTexture* SourceTexture, const
 	if (SourceData.GetFormat() == TSF_BGRA8)
 	{
 		uint32 BytesPerPixel = SourceData.GetBytesPerPixel();
-		uint8* OffsetSource = SourceData.LockMip(0) + (SourceXY.X + SourceXY.Y * SourceData.GetSizeX()) * BytesPerPixel;
+		const uint8* OffsetSource = SourceData.LockMipReadOnly(0) + (SourceXY.X + SourceXY.Y * SourceData.GetSizeX()) * BytesPerPixel;
 		uint8* OffsetDest = TargetBuffer.GetData();
 		CopyTextureData(OffsetSource, OffsetDest, SourceSize.X, SourceSize.Y, BytesPerPixel, SourceData.GetSizeX() * BytesPerPixel, SourceSize.X * BytesPerPixel);
 		SourceData.UnlockMip(0);
diff --git a/Engine/Source/Editor/UnrealEd/Classes/Factories/TextureFactory.h b/Engine/Source/Editor/UnrealEd/Classes/Factories/TextureFactory.h
--- a/Engine/Source/Editor/UnrealEd/Classes/Factories/TextureFactory.h
+++ b/Engine/Source/Editor/UnrealEd/Classes/Factories/TextureFactory.h
@@ -19,6 +19,8 @@ struct FImportImage
 	int32 SizeX = 0;
 	int32 SizeY = 0;
 	bool SRGB = true;
+	/** Which compression format (if any) that is applied to RawData */
+	ETextureSourceCompressionFormat RawDataCompressionFormat = TSCF_None;
 
 	void Init2DWithParams(int32 InSizeX, int32 InSizeY, ETextureSourceFormat InFormat, bool InSRGB);
 	void Init2DWithOneMip(int32 InSizeX, int32 InSizeY, ETextureSourceFormat InFormat, const void* InData = nullptr);
diff --git a/Engine/Source/Editor/UnrealEd/Private/Factories/EditorFactories.cpp b/Engine/Source/Editor/UnrealEd/Private/Factories/EditorFactories.cpp
--- a/Engine/Source/Editor/UnrealEd/Private/Factories/EditorFactories.cpp
+++ b/Engine/Source/Editor/UnrealEd/Private/Factories/EditorFactories.cpp
@@ -3109,7 +3109,7 @@ bool UTextureFactory::ImportImage(const uint8* Buffer, uint32 Length, FFeedbackC
 		{
 			if (BitDepth <= 8)
 			{
-				TextureFormat = TSF_G8;
+				TextureFormat =  TSF_G8;
 				Format = ERGBFormat::Gray;
 				BitDepth = 8;
 			}
@@ -3136,11 +3136,24 @@ bool UTextureFactory::ImportImage(const uint8* Buffer, uint32 Length, FFeedbackC
 			TextureFormat,
 			BitDepth < 16
 		);
-		
-		if (!JpegImageWrapper->GetRaw(Format, BitDepth, OutImage.RawData))
+
+		// For now this option is opt in via the config files once there is no technical risk this will
+		// become the default path.
+		bool bRetainJpegFormat = false;
+		GConfig->GetBool(TEXT("TextureImporter"), TEXT("RetainJpegFormat"), bRetainJpegFormat, GEditorIni);
+
+		if (bRetainJpegFormat)
 		{
-			Warn->Logf(ELogVerbosity::Error, TEXT("Failed to decode JPEG."));
-			return false;
+			OutImage.RawData.Append(Buffer, Length);
+			OutImage.RawDataCompressionFormat = ETextureSourceCompressionFormat::TSCF_JPEG;
+		}
+		else
+		{
+			if (!JpegImageWrapper->GetRaw(Format, BitDepth, OutImage.RawData))
+			{
+				Warn->Logf(ELogVerbosity::Error, TEXT("Failed to decode JPEG."));
+				return false;
+			}
 		}
 
 		return true;
@@ -3645,14 +3658,29 @@ UTexture* UTextureFactory::ImportTexture(UClass* Class, UObject* InParent, FName
 		UTexture2D* Texture = CreateTexture2D(InParent, Name, Flags);
 		if (Texture)
 		{
-			Texture->Source.Init(
-				Image.SizeX,
-				Image.SizeY,
-				/*NumSlices=*/ 1,
-				Image.NumMips,
-				Image.Format,
-				Image.RawData.GetData()
-			);
+			if (Image.RawDataCompressionFormat == ETextureSourceCompressionFormat::TSCF_None)
+			{
+				Texture->Source.Init(
+					Image.SizeX,
+					Image.SizeY,
+					/*NumSlices=*/ 1,
+					Image.NumMips,
+					Image.Format,
+					Image.RawData.GetData()
+				);
+			}
+			else
+			{
+				Texture->Source.InitWithCompressedSourceData(
+					Image.SizeX,
+					Image.SizeY,
+					Image.NumMips,
+					Image.Format,
+					Image.RawData,
+					Image.RawDataCompressionFormat
+				);
+			}
+
 			Texture->CompressionSettings = Image.CompressionSettings;
 
 			if (ColorSpaceMode == ETextureSourceColorSpace::Auto)
diff --git a/Engine/Source/Editor/UnrealEd/Private/NormalMapIdentification.cpp b/Engine/Source/Editor/UnrealEd/Private/NormalMapIdentification.cpp
--- a/Engine/Source/Editor/UnrealEd/Private/NormalMapIdentification.cpp
+++ b/Engine/Source/Editor/UnrealEd/Private/NormalMapIdentification.cpp
@@ -77,13 +77,13 @@ public:
 		SourceTexture = Texture;
 		TextureSizeX = Texture->Source.GetSizeX();
 		TextureSizeY = Texture->Source.GetSizeY();
-		SourceTextureData = Texture->Source.LockMip(0);
+		SourceTextureData = Texture->Source.LockMipReadOnly(0);
 	}
 
 	UTexture* SourceTexture;
 	int32 TextureSizeX;
 	int32 TextureSizeY;
-	uint8* SourceTextureData;
+	const uint8* SourceTextureData;
 };
 
 template<int RIdx, int GIdx, int BIdx, int AIdx> class SampleNormalMapPixel8 : public NormalMapSamplerBase
@@ -95,7 +95,7 @@ public:
 	FLinearColor DoSampleColor( int32 X, int32 Y )
 	{
 		FLinearColor Result;
-		uint8* PixelToSample = SourceTextureData + ((Y * TextureSizeX + X) * 4);
+		const uint8* PixelToSample = SourceTextureData + ((Y * TextureSizeX + X) * 4);
 
 		const float OneOver255 = 1.0f / 255.0f;
 
@@ -124,7 +124,7 @@ public:
 	FLinearColor DoSampleColor( int32 X, int32 Y )
 	{
 		FLinearColor Result;
-		uint8* PixelToSample = SourceTextureData + ((Y * TextureSizeX + X) * 8);
+		const uint8* PixelToSample = SourceTextureData + ((Y * TextureSizeX + X) * 8);
 
 		const float OneOver65535 = 1.0f / 65535.0f;
 
@@ -150,7 +150,7 @@ public:
 	FLinearColor DoSampleColor( int32 X, int32 Y )
 	{
 		FLinearColor Result;
-		uint8* PixelToSample = SourceTextureData + ((Y * TextureSizeX + X) * 8);
+		const uint8* PixelToSample = SourceTextureData + ((Y * TextureSizeX + X) * 8);
 
 		// this assume the normal map to be in linear (not the case if Photoshop converts a 8bit normalmap to float and saves it as 16bit dds)
 		Result.R = (float)((FFloat16*)PixelToSample)[0];
diff --git a/Engine/Source/Runtime/Engine/Classes/Engine/Texture.h b/Engine/Source/Runtime/Engine/Classes/Engine/Texture.h
--- a/Engine/Source/Runtime/Engine/Classes/Engine/Texture.h
+++ b/Engine/Source/Runtime/Engine/Classes/Engine/Texture.h
@@ -131,6 +131,16 @@ enum ETextureSourceFormat
 	TSF_MAX
 };
 
+UENUM()
+enum ETextureSourceCompressionFormat
+{
+	TSCF_None	UMETA(DisplayName = "ZLib"),
+	TSCF_PNG	UMETA(DisplayName = "PNG"),
+	TSCF_JPEG	UMETA(DisplayName = "JPEG"),
+
+	TSCF_MAX
+};
+
 UENUM()
 enum ETextureCompressionQuality
 {
@@ -250,12 +260,33 @@ struct FTextureSource
 		ETextureSourceFormat NewFormat
 		);
 
+	/**
+	 * Initialize the source data with the given size, number of mips, and format.
+	 * @param NewSizeX - Width of the texture source data.
+	 * @param NewSizeY - Height of the texture source data.
+	 * @param NewNumMips - The number of mips in the texture source data.
+	 * @param NewFormat - The format in which source data is stored.
+	 * @param NewSourceData -The new source data.
+	 * @param NewSourceFormat -The compression format of the new source data.
+	 */
+	ENGINE_API void InitWithCompressedSourceData(
+		int32 NewSizeX,
+		int32 NewSizeY,
+		int32 NewNumMips,
+		ETextureSourceFormat NewFormat,
+		const TArrayView64<uint8> NewSourceData,
+		ETextureSourceCompressionFormat NewSourceFormat
+	);
+
 	/** PNG Compresses the source art if possible or tells the bulk data to zlib compress when it saves out to disk. */
 	ENGINE_API void Compress();
 
 	/** Force the GUID to change even if mip data has not been modified. */
 	void ForceGenerateGuid();
 
+	/** Lock a mip for reading. */
+	ENGINE_API const uint8* LockMipReadOnly(int32 BlockIndex, int32 LayerIndex, int32 MipIndex);
+
 	/** Lock a mip for editing. */
 	ENGINE_API uint8* LockMip(int32 BlockIndex, int32 LayerIndex, int32 MipIndex);
 
@@ -289,6 +320,9 @@ struct FTextureSource
 	/** Returns the unique ID string for this source art. */
 	FString GetIdString() const;
 
+	/** Returns the compression format of the source data in string format for use with the UI. */
+	FString GetSourceCompressionAsString() const;
+
 	/** Trivial accessors. These will only give values for Block0 so may not be correct for UDIM/multi-block textures, use GetBlock() for this case. */
 	FORCEINLINE FGuid GetId() const { return Id; }
 	FORCEINLINE int32 GetSizeX() const { return SizeX; }
@@ -314,6 +348,9 @@ struct FTextureSource
 	}
 
 	FORCEINLINE int64 CalcMipSize(int32 MipIndex) const { return CalcMipSize(0, 0, MipIndex); }
+	/** Lock a mip for reading. */
+	FORCEINLINE const uint8* LockMipReadOnly(int32 MipIndex) { return LockMipReadOnly(0, 0, MipIndex); }
+	/** Lock a mip for editing. */
 	FORCEINLINE uint8* LockMip(int32 MipIndex) { return LockMip(0, 0, MipIndex); }
 	FORCEINLINE void UnlockMip(int32 MipIndex) { UnlockMip(0, 0, MipIndex); }
 
@@ -337,8 +374,22 @@ private:
 	uint8* LockedMipData;
 	/** Number of mips that are locked. */
 	uint32 NumLockedMips;
+
+	enum class ELockState : uint8
+	{
+		None,
+		ReadOnly,
+		ReadWrite
+	};
+	/** The state of any lock being held on the mip data */
+	ELockState LockState;
+
 #if WITH_EDITOR
 
+	// Internal implementation for locking the mip data, called by LockMipReadOnly or LockMip */
+	uint8* LockMipInternal(int32 BlockIndex, int32 LayerIndex, int32 MipIndex, ELockState RequestedLockState);
+	/** Attempt to decompress the source data from Jpeg format. All failures will be logged and result in the method returning false */
+	bool TryDecompressJpegData(IImageWrapperModule* ImageWrapperModule, const void* RawSourceData, int32 MipIndex, TArray64<uint8>& OutRawData);
 	/** Return true if the source art is not png compressed but could be. */
 	bool CanPNGCompress() const;
 	/** Removes source data. */
@@ -396,6 +447,10 @@ private:
 	UPROPERTY(VisibleAnywhere, Category=TextureSource)
 	bool bPNGCompressed;
 
+	/** Compression format that source data is stored as. */
+	UPROPERTY(VisibleAnywhere, Category = TextureSource)
+	TEnumAsByte<enum ETextureSourceCompressionFormat> CompressionFormat;
+
 	/** Uses hash instead of guid to identify content to improve DDC cache hit. */
 	UPROPERTY(VisibleAnywhere, Category=TextureSource)
 	bool bGuidIsHash;
diff --git a/Engine/Source/Runtime/Engine/Private/Texture.cpp b/Engine/Source/Runtime/Engine/Private/Texture.cpp
--- a/Engine/Source/Runtime/Engine/Private/Texture.cpp
+++ b/Engine/Source/Runtime/Engine/Private/Texture.cpp
@@ -622,7 +622,7 @@ void UTexture::GetAssetRegistryTags(TArray<FAssetRegistryTag>& OutTags) const
 		OutTags.Add( FAssetRegistryTag(SourceFileTagName(), AssetImportData->GetSourceData().ToJson(), FAssetRegistryTag::TT_Hidden) );
 	}
 
-	OutTags.Add(FAssetRegistryTag("SourceCompression", Source.IsPNGCompressed() ? TEXT("PNG") : LexToString(NAME_Zlib), FAssetRegistryTag::TT_Alphabetical));
+	OutTags.Add(FAssetRegistryTag("SourceCompression", Source.GetSourceCompressionAsString(), FAssetRegistryTag::TT_Alphabetical));
 	
 	Super::GetAssetRegistryTags(OutTags);
 }
@@ -944,6 +944,7 @@ FStreamableRenderResourceState UTexture::GetResourcePostInitState(FTexturePlatfo
 FTextureSource::FTextureSource()
 	: LockedMipData(NULL)
 	, NumLockedMips(0u)
+	, LockState(ELockState::None)
 #if WITH_EDITOR
 	, bHasHadBulkDataCleared(false)
 #endif
@@ -956,6 +957,7 @@ FTextureSource::FTextureSource()
 	, NumMips(0)
 	, NumLayers(1) // Default to 1 so old data has the correct value
 	, bPNGCompressed(false)
+	, CompressionFormat(TSCF_None)
 	, bGuidIsHash(false)
 	, Format(TSF_Invalid)
 #endif // WITH_EDITORONLY_DATA
@@ -1029,6 +1031,7 @@ void FTextureSource::InitBlocked(const ETextureSourceFormat* InLayerFormats,
 		TotalBytes += CalcBlockSize(i);
 	}
 
+	checkf(LockState == ELockState::None, TEXT("InitBlocked shouldn't be called in-between LockMip/UnlockMip"));
 	BulkData.Lock(LOCK_READ_WRITE);
 	uint8* DestData = (uint8*)BulkData.Realloc(TotalBytes);
 	if (InDataPerBlock)
@@ -1074,6 +1077,7 @@ void FTextureSource::InitLayered(
 		TotalBytes += CalcLayerSize(0, i);
 	}
 
+	checkf(LockState == ELockState::None, TEXT("InitLayered shouldn't be called in-between LockMip/UnlockMip"));
 	BulkData.Lock(LOCK_READ_WRITE);
 	uint8* DestData = (uint8*)BulkData.Realloc(TotalBytes);
 	if (NewData)
@@ -1125,8 +1129,44 @@ void FTextureSource::InitCubeWithMipChain(
 	Init(NewSizeX, NewSizeY, 6, NewMipCount, NewFormat);
 }
 
+void FTextureSource::InitWithCompressedSourceData(
+	int32 NewSizeX,
+	int32 NewSizeY,
+	int32 NewNumMips,
+	ETextureSourceFormat NewFormat,
+	const TArrayView64<uint8> NewData,
+	ETextureSourceCompressionFormat NewSourceFormat
+)
+{
+	RemoveSourceData();
+
+	SizeX = NewSizeX;
+	SizeY = NewSizeY;
+
+	NumLayers = 1;
+	NumSlices = 1;
+	NumMips = NewNumMips;
+
+	Format = NewFormat;
+	LayerFormat.SetNum(1, true);
+	LayerFormat[0] = NewFormat;
+	
+	CompressionFormat = NewSourceFormat;
+
+	checkf(LockState == ELockState::None, TEXT("InitWithCompressedSourceData shouldn't be called in-between LockMip/UnlockMip"));
+	BulkData.Lock(LOCK_READ_WRITE);
+	uint8* DestData = (uint8*)BulkData.Realloc(NewData.Num());
+	if (NewData.GetData() != nullptr)
+	{
+		FMemory::Memcpy(DestData, NewData.GetData(), NewData.Num());
+	}
+	BulkData.Unlock();
+}
+
 void FTextureSource::Compress()
 {
+	checkf(LockState == ELockState::None, TEXT("Compress shouldn't be called in-between LockMip/UnlockMip"));
+
 #if WITH_EDITOR
 	FWriteScopeLock BulkDataExclusiveScope(BulkDataLock.Get());
 #endif
@@ -1147,24 +1187,48 @@ void FTextureSource::Compress()
 				FMemory::Memcpy(BulkDataPtr, CompressedData.GetData(), CompressedData.Num());
 				BulkData.Unlock();
 				bPNGCompressed = true;
+				CompressionFormat = TSCF_PNG;
 			}
 		}
 	}
 	
-	// If the data is already compressed as a png then we don't need to compress it
-	// again when serializing to disk.
-	BulkData.StoreCompressedOnDisk(bPNGCompressed ? NAME_None : NAME_Zlib);
+	// Disable the internal bulkdata compression if the source data is already compressed
+	BulkData.StoreCompressedOnDisk(CompressionFormat == TSCF_None ? NAME_Zlib : NAME_None);
+}
+
+const uint8* FTextureSource::LockMipReadOnly(int32 BlockIndex, int32 LayerIndex, int32 MipIndex)
+{
+	return LockMipInternal(BlockIndex, LayerIndex, MipIndex, ELockState::ReadOnly);
 }
 
 uint8* FTextureSource::LockMip(int32 BlockIndex, int32 LayerIndex, int32 MipIndex)
 {
+	return LockMipInternal(BlockIndex, LayerIndex, MipIndex, ELockState::ReadWrite);
+}
+
+uint8* FTextureSource::LockMipInternal(int32 BlockIndex, int32 LayerIndex, int32 MipIndex, ELockState RequestedLockState )
+{
+	checkf(RequestedLockState != ELockState::None, TEXT("Cannot call FTextureSource::LockMipInternal with a RequestedLockState of type ELockState::None"));
+
 	uint8* MipData = NULL;
 	if (BlockIndex < GetNumBlocks() && LayerIndex < NumLayers && MipIndex < NumMips)
 	{
 		if (LockedMipData == NULL)
 		{
 			LockedMipData = (uint8*)BulkData.Lock(LOCK_READ_WRITE);
-			if (bPNGCompressed)
+
+			if (CompressionFormat == TSCF_JPEG)
+			{
+				TArray64<uint8> RawData;
+				// Note that the method itself will give warnings on all failure conditions so we 
+				// don't need to handle that here
+				if (TryDecompressJpegData(nullptr, LockedMipData, MipIndex, RawData))
+				{
+					LockedMipData = (uint8*)FMemory::Malloc(RawData.Num());
+					FMemory::Memcpy(LockedMipData, RawData.GetData(), RawData.Num());
+				}
+			}	
+			else if (bPNGCompressed)
 			{
 				bool bCanPngCompressFormat = (Format == TSF_G8 || Format == TSF_G16 || Format == TSF_RGBA8 || Format == TSF_BGRA8 || Format == TSF_RGBA16);
 				check(Blocks.Num() == 0 && NumLayers == 1 && NumSlices == 1 && bCanPngCompressFormat);
@@ -1200,7 +1264,17 @@ uint8* FTextureSource::LockMip(int32 BlockIndex, int32 LayerIndex, int32 MipInde
 		}
 
 		MipData = LockedMipData + CalcMipOffset(BlockIndex, LayerIndex, MipIndex);
-		++NumLockedMips;
+
+		if (NumLockedMips == 0)
+		{
+			LockState = RequestedLockState;
+		}
+		else
+		{
+			checkf(LockState == RequestedLockState, TEXT("Cannot change the lock type until UnlockMip is called"));
+		}
+
+		++NumLockedMips;		
 	}
 	return MipData;
 }
@@ -1212,28 +1286,53 @@ void FTextureSource::UnlockMip(int32 BlockIndex, int32 LayerIndex, int32 MipInde
 	check(MipIndex < MAX_TEXTURE_MIP_COUNT);
 
 	check(NumLockedMips > 0u);
+	check(LockState != ELockState::None);
+
 	--NumLockedMips;
 	if (NumLockedMips == 0u)
 	{
-		if (bPNGCompressed)
+		// If the source data is stored in a compressed format then we returned a copy when the mip was
+		// locked. We now need to deal with that copy.
+		if (CompressionFormat == TSCF_JPEG || bPNGCompressed)
 		{
-			check(BlockIndex == 0);
-			check(LayerIndex == 0);
-			check(MipIndex == 0);
-			int32 MipSize = CalcMipSize(0, 0, 0);
-			uint8* UncompressedData = (uint8*)BulkData.Realloc(MipSize);
-			FMemory::Memcpy(UncompressedData, LockedMipData, MipSize);
-			FMemory::Free(LockedMipData);
-			bPNGCompressed = false;
+			// If we had a write lock then we need to update the bulkdata with the copy as we have
+			// to assume that it contains changes. This means we will no longer be storing it 
+			// in a compressed format.
+			// If the lock was read only then we can just delete the locked mip data as we know
+			// that there were no changes to the data.
+			if (LockState == ELockState::ReadWrite)
+			{
+				UE_CLOG(CompressionFormat == TSCF_JPEG, LogTexture, Warning, TEXT("Call to FTextureSource::UnlockMip will cause texture source to lose it's jpeg storage format"));
+
+				check(BlockIndex == 0);
+				check(LayerIndex == 0);
+				check(MipIndex == 0);
+				int32 MipSize = CalcMipSize(0, 0, 0);
+				uint8* UncompressedData = (uint8*)BulkData.Realloc(MipSize);
+				FMemory::Memcpy(UncompressedData, LockedMipData, MipSize);
+				FMemory::Free(LockedMipData);
+				bPNGCompressed = false;
+				CompressionFormat = TSCF_None;
+			}
+			else
+			{
+				FMemory::Free(LockedMipData);
+			}
 		}
+
+		LockState = ELockState::None;
 		LockedMipData = NULL;
 		BulkData.Unlock();
+
+		// TODO: Only call this if the LockState was ELockState::ReadWrite
 		ForceGenerateGuid();
 	}
 }
 
 bool FTextureSource::GetMipData(TArray64<uint8>& OutMipData, int32 BlockIndex, int32 LayerIndex, int32 MipIndex, IImageWrapperModule* ImageWrapperModule)
 {
+	checkf(LockState == ELockState::None, TEXT("GetMipData shouldn't be called in-between LockMip/UnlockMip"));
+
 	bool bSuccess = false;
 	if (BlockIndex < GetNumBlocks() && LayerIndex < NumLayers && MipIndex < NumMips && BulkData.GetBulkDataSize() > 0)
 	{
@@ -1241,7 +1340,11 @@ bool FTextureSource::GetMipData(TArray64<uint8>& OutMipData, int32 BlockIndex, i
 		FWriteScopeLock BulkDataExclusiveScope(BulkDataLock.Get());
 #endif
 		void* RawSourceData = BulkData.Lock(LOCK_READ_ONLY);
-		if (bPNGCompressed)
+		if (CompressionFormat == TSCF_JPEG)
+		{
+			bSuccess = TryDecompressJpegData(ImageWrapperModule, RawSourceData, MipIndex, OutMipData);
+		}
+		else if (bPNGCompressed)
 		{
 			bool bCanPngCompressFormat = (Format == TSF_G8 || Format == TSF_G16 || Format == TSF_RGBA8 || Format == TSF_BGRA8 || Format == TSF_RGBA16);
 			if (MipIndex == 0 && NumLayers == 1 && NumSlices == 1 && Blocks.Num() == 0 && bCanPngCompressFormat)
@@ -1304,6 +1407,7 @@ bool FTextureSource::GetMipData(TArray64<uint8>& OutMipData, int32 BlockIndex, i
 		}
 		BulkData.Unlock();
 	}
+
 	return bSuccess;
 }
 
@@ -1399,6 +1503,56 @@ FString FTextureSource::GetIdString() const
 	return GuidString;
 }
 
+FString FTextureSource::GetSourceCompressionAsString() const
+{
+	// Until we deprecate bPNGCompressed it might not be 100% in sync with CompressionFormat
+	// so if it is set we should use that rather than the enum.
+	if (bPNGCompressed)
+	{
+		return StaticEnum<ETextureSourceCompressionFormat>()->GetDisplayNameTextByValue(ETextureSourceCompressionFormat::TSCF_PNG).ToString();
+	}
+
+	return StaticEnum<ETextureSourceCompressionFormat>()->GetDisplayNameTextByValue(CompressionFormat).ToString();
+}
+
+bool FTextureSource::TryDecompressJpegData(IImageWrapperModule* ImageWrapperModule, const void* RawSourceData, int32 MipIndex, TArray64<uint8>& OutRawData)
+{
+	if (MipIndex == 0 && NumLayers == 1 && NumSlices == 1 && Blocks.Num() == 0)
+	{
+		if (!ImageWrapperModule)
+		{
+			ImageWrapperModule = &FModuleManager::LoadModuleChecked<IImageWrapperModule>(FName("ImageWrapper"));
+		}
+
+		TSharedPtr<IImageWrapper> JpegImageWrapper = ImageWrapperModule->CreateImageWrapper(EImageFormat::JPEG);
+		if (JpegImageWrapper.IsValid() && JpegImageWrapper->SetCompressed(RawSourceData, BulkData.GetBulkDataSize()))
+		{
+			// The two formats we support for JPEG imports, see UTextureFactory::ImportImage
+			const ERGBFormat JpegFormat = Format == TSF_G8 ? ERGBFormat::Gray : ERGBFormat::BGRA;
+			if (JpegImageWrapper->GetRaw(JpegFormat, 8, OutRawData))
+			{
+				return true;
+			}
+			else
+			{
+				UE_LOG(LogTexture, Warning, TEXT("JPEG decompression of source art failed to return uncompressed data"));
+				return false;
+			}
+		}
+		else
+		{
+			UE_LOG(LogTexture, Warning, TEXT("JPEG decompression of source art failed initialization"));
+			return false;
+		}
+	}
+	else
+	{
+		UE_LOG(LogTexture, Warning, TEXT("JPEG compressed source art is in an invalid format MipIndex:(%d) NumLayers:(%d) NumSlices:(%d) NumBlocks:(%d)"),
+			MipIndex, NumLayers, NumSlices, Blocks.Num());
+		return false;
+	}
+}
+
 bool FTextureSource::CanPNGCompress() const
 {
 	bool bCanPngCompressFormat = (Format == TSF_G8 || Format == TSF_G16 || Format == TSF_RGBA8 || Format == TSF_BGRA8 || Format == TSF_RGBA16);
@@ -1411,7 +1565,8 @@ bool FTextureSource::CanPNGCompress() const
 		SizeX > 4 &&
 		SizeY > 4 &&
 		BulkData.GetBulkDataSize() > 0 &&
-		bCanPngCompressFormat)
+		bCanPngCompressFormat &&
+		CompressionFormat == TSCF_None)
 	{
 		return true;
 	}
@@ -1447,6 +1602,7 @@ void FTextureSource::RemoveSourceData()
 	bPNGCompressed = false;
 	LockedMipData = NULL;
 	NumLockedMips = 0u;
+	LockState = ELockState::None;
 	if (BulkData.IsLocked())
 	{
 		BulkData.Unlock();
