Chaos Physics - make sure rigid bodies are woken up when applying a physics profile

diff --git a/Engine/Source/Runtime/Engine/Private/PhysicsEngine/PhysicalAnimationComponent.cpp b/Engine/Source/Runtime/Engine/Private/PhysicsEngine/PhysicalAnimationComponent.cpp
--- a/Engine/Source/Runtime/Engine/Private/PhysicsEngine/PhysicalAnimationComponent.cpp
+++ b/Engine/Source/Runtime/Engine/Private/PhysicsEngine/PhysicalAnimationComponent.cpp
@@ -6,12 +6,12 @@
 #include "PhysicsEngine/PhysicsAsset.h"
 #include "PhysXPublic.h"
 #include "Physics/PhysicsInterfaceCore.h"
+#include "Chaos/ChaosConstraintSettings.h"
 #include "Chaos/ParticleHandle.h"
 #include "Chaos/Sphere.h"
 #include "ChaosCheck.h"
 #include "PhysicsProxy/SingleParticlePhysicsProxy.h"
 
-
 const FConstraintProfileProperties UPhysicalAnimationComponent::PhysicalAnimationProfile = []()
 {
 	//Setup the default constraint profile for all joints created by physical animation system
@@ -345,21 +345,46 @@ void UPhysicalAnimationComponent::TickComponent(float DeltaTime, enum ELevelTick
 	UpdateTargetActors(ETeleportType::None);
 }
 
-void SetMotorStrength(FConstraintInstance& ConstraintInstance, const FPhysicalAnimationData& PhysAnimData, float StrengthMultiplyer)
+void SetMotorStrength(FConstraintInstance& ConstraintInstance, const FPhysicalAnimationData& PhysAnimData, float StrengthMultiplier)
 {
-	ConstraintInstance.SetAngularDriveParams(PhysAnimData.OrientationStrength * StrengthMultiplyer, PhysAnimData.AngularVelocityStrength * StrengthMultiplyer, PhysAnimData.MaxAngularForce * StrengthMultiplyer);
+	float PositionStrengthMultiplier = StrengthMultiplier;
+	float VelocityStrengthMultiplier = StrengthMultiplier;
+	float OrientationStrengthMultiplier = StrengthMultiplier;
+	float AngularVelocityStrengthMultiplier = StrengthMultiplier;
+#if WITH_CHAOS
+	// Chaos has it's own global adjustments
+	PositionStrengthMultiplier *= Chaos::ConstraintSettings::LinearDriveStiffnessScale();
+	VelocityStrengthMultiplier *= Chaos::ConstraintSettings::LinearDriveDampingScale();
+	OrientationStrengthMultiplier *= Chaos::ConstraintSettings::AngularDriveStiffnessScale();
+	AngularVelocityStrengthMultiplier *= Chaos::ConstraintSettings::AngularDriveDampingScale();
+#endif
+
+	ConstraintInstance.SetAngularDriveParams(
+		PhysAnimData.OrientationStrength * OrientationStrengthMultiplier, 
+		PhysAnimData.AngularVelocityStrength * AngularVelocityStrengthMultiplier,
+		PhysAnimData.MaxAngularForce * StrengthMultiplier
+	);
+
 	if (PhysAnimData.bIsLocalSimulation)	//linear only works for world space simulation
 	{
 		ConstraintInstance.SetLinearDriveParams(0.f, 0.f, 0.f);
 	}
 	else
 	{
-		ConstraintInstance.SetLinearDriveParams(PhysAnimData.PositionStrength * StrengthMultiplyer, PhysAnimData.VelocityStrength * StrengthMultiplyer, PhysAnimData.MaxLinearForce * StrengthMultiplyer);
+		ConstraintInstance.SetLinearDriveParams(
+			PhysAnimData.PositionStrength * PositionStrengthMultiplier, 
+			PhysAnimData.VelocityStrength * VelocityStrengthMultiplier, 
+			PhysAnimData.MaxLinearForce * StrengthMultiplier
+		);
 	}
 }
 
 void UPhysicalAnimationComponent::UpdatePhysicsEngine()
 {
+	if (SkeletalMeshComponent)
+	{
+		SkeletalMeshComponent->WakeAllRigidBodies();
+	}
 	bPhysicsEngineNeedsUpdating = true;	//must defer until tick so that animation can finish
 }
 
