Add a couple of meaningful Insight traces

diff --git a/Engine/Source/Runtime/Engine/Private/ContentStreaming.cpp b/Engine/Source/Runtime/Engine/Private/ContentStreaming.cpp
--- a/Engine/Source/Runtime/Engine/Private/ContentStreaming.cpp
+++ b/Engine/Source/Runtime/Engine/Private/ContentStreaming.cpp
@@ -846,6 +846,7 @@ void FStreamingManagerCollection::Tick( float DeltaTime, bool bProcessEverything
 
 void FStreamingManagerCollection::UpdateResourceStreaming( float DeltaTime, bool bProcessEverything/*=false*/ )
 {
+	TRACE_CPUPROFILER_EVENT_SCOPE(FStreamingManagerCollection::UpdateResourceStreaming);
 	SetupViewInfos( DeltaTime );
 
 	// only allow this if its not disabled
diff --git a/Engine/Source/Runtime/Engine/Private/StaticMesh.cpp b/Engine/Source/Runtime/Engine/Private/StaticMesh.cpp
--- a/Engine/Source/Runtime/Engine/Private/StaticMesh.cpp
+++ b/Engine/Source/Runtime/Engine/Private/StaticMesh.cpp
@@ -5776,6 +5776,8 @@ void UStaticMesh::BuildFromStaticMeshDescriptions(const TArray<UStaticMeshDescri
 
 bool UStaticMesh::BuildFromMeshDescriptions(const TArray<const FMeshDescription*>& MeshDescriptions, const FBuildMeshDescriptionsParams& Params)
 {
+	TRACE_CPUPROFILER_EVENT_SCOPE(UStaticMesh::BuildFromMeshDescriptions);
+
 	// Set up
 	SetIsBuiltAtRuntime(true);
 	NeverStream = true;
diff --git a/Engine/Source/Runtime/Engine/Private/Texture.cpp b/Engine/Source/Runtime/Engine/Private/Texture.cpp
--- a/Engine/Source/Runtime/Engine/Private/Texture.cpp
+++ b/Engine/Source/Runtime/Engine/Private/Texture.cpp
@@ -348,6 +348,8 @@ void UTexture::PostEditChangeProperty(FPropertyChangedEvent& PropertyChangedEven
 	}
 	else if (!GDisableAutomaticTextureMaterialUpdateDependencies)
 	{
+		TRACE_CPUPROFILER_EVENT_SCOPE(UpdateDependentMaterials);
+
 		// Update any material that uses this texture and must force a recompile of cache resource
 		TArray<UMaterial*> MaterialsToUpdate;
 		TSet<UMaterial*> BaseMaterialsThatUseThisTexture;
@@ -401,6 +403,8 @@ void UTexture::PostEditChangeProperty(FPropertyChangedEvent& PropertyChangedEven
 #if WITH_EDITORONLY_DATA
 	// any texture that is referencing this texture as AssociatedNormalMap needs to be informed
 	{
+		TRACE_CPUPROFILER_EVENT_SCOPE(UpdateDependentTextures);
+
 		TArray<UTexture*> TexturesThatUseThisTexture;
 
 		for (TObjectIterator<UTexture> It; It; ++It)
diff --git a/Engine/Source/Runtime/GameplayTags/Private/GameplayTagsManager.cpp b/Engine/Source/Runtime/GameplayTags/Private/GameplayTagsManager.cpp
--- a/Engine/Source/Runtime/GameplayTags/Private/GameplayTagsManager.cpp
+++ b/Engine/Source/Runtime/GameplayTags/Private/GameplayTagsManager.cpp
@@ -1960,6 +1960,8 @@ void UGameplayTagsManager::DoneAddingNativeTags()
 	// is initialized (DoneAddingNativeTags is bound to PostEngineInit to cover anything that's skipped).
 	if (GEngine && !bDoneAddingNativeTags)
 	{
+		TRACE_CPUPROFILER_EVENT_SCOPE(UGameplayTagsManager::DoneAddingNativeTags);
+
 		UE_CLOG(GAMEPLAYTAGS_VERBOSE, LogGameplayTags, Display, TEXT("UGameplayTagsManager::DoneAddingNativeTags. DelegateIsBound: %d"), (int32)OnLastChanceToAddNativeTags().IsBound());
 		OnLastChanceToAddNativeTags().Broadcast();
 		bDoneAddingNativeTags = true;
diff --git a/Engine/Source/Runtime/MeshDescription/Private/MeshDescription.cpp b/Engine/Source/Runtime/MeshDescription/Private/MeshDescription.cpp
--- a/Engine/Source/Runtime/MeshDescription/Private/MeshDescription.cpp
+++ b/Engine/Source/Runtime/MeshDescription/Private/MeshDescription.cpp
@@ -1970,6 +1970,8 @@ void FMeshDescriptionBulkData::Serialize( FArchive& Ar, UObject* Owner )
 		return;
 	}
 
+	TRACE_CPUPROFILER_EVENT_SCOPE(FMeshDescriptionBulkData::Serialize);
+
 	if ( Ar.IsTransacting() )
 	{
 		// If transacting, keep these members alive the other side of an undo, otherwise their values will get lost
@@ -2101,6 +2103,8 @@ void FMeshDescriptionBulkData::SaveMeshDescription( FMeshDescription& MeshDescri
 
 void FMeshDescriptionBulkData::LoadMeshDescription( FMeshDescription& MeshDescription )
 {
+	TRACE_CPUPROFILER_EVENT_SCOPE(FMeshDescriptionBulkData::LoadMeshDescription);
+
 	MeshDescription.Empty();
 
 	if (BulkData.GetBulkDataSize() > 0)
