Write shader compilation stats to the log periodically in the editor.

- Keep the log information simple.  Just total compiled shaders.
- Added a PrintStats function to the shader compiler manager which will print both shader job cache stats and shader compilation stats.
- Always log at the end of editor startup.

diff --git a/Engine/Source/Editor/UnrealEd/Private/UnrealEdMisc.cpp b/Engine/Source/Editor/UnrealEd/Private/UnrealEdMisc.cpp
--- a/Engine/Source/Editor/UnrealEd/Private/UnrealEdMisc.cpp
+++ b/Engine/Source/Editor/UnrealEd/Private/UnrealEdMisc.cpp
@@ -536,6 +536,8 @@ void FUnrealEdMisc::OnInit()
 	UE_LOG(LogUnrealEdMisc, Log, TEXT("Total Editor Startup Time, took %.3f"), TotalEditorStartupTime);
 
 	FStudioAnalytics::FireEvent_Loading(TEXT("TotalEditorStartup"), TotalEditorStartupTime);
+
+	GShaderCompilingManager->PrintStats(true);
 }
 
 FString FUnrealEdMisc::FindMapFileFromPartialName(const FString& PartialMapName)
diff --git a/Engine/Source/Runtime/Engine/Private/ShaderCompiler/ShaderCompiler.cpp b/Engine/Source/Runtime/Engine/Private/ShaderCompiler/ShaderCompiler.cpp
--- a/Engine/Source/Runtime/Engine/Private/ShaderCompiler/ShaderCompiler.cpp
+++ b/Engine/Source/Runtime/Engine/Private/ShaderCompiler/ShaderCompiler.cpp
@@ -521,7 +521,7 @@ void FShaderCompileJobCollection::SubmitJobs(const TArray<FShaderCommonCompileJo
 
 void FShaderCompileJobCollection::HandleLogJobsCacheStats()
 {
-	LogCachingStats(true);
+	GShaderCompilingManager->PrintStats(true);
 }
 
 void FShaderCompileJobCollection::ProcessFinishedJob(FShaderCommonCompileJob* FinishedJob, bool bWasCached)
@@ -597,21 +597,10 @@ void FShaderCompileJobCollection::AddToCacheAndProcessPending(FShaderCommonCompi
 	JobsInFlight.Remove(InputHash);
 }
 
-void FShaderCompileJobCollection::LogCachingStats(bool bForceLogIgnoringTimeInverval)
+void FShaderCompileJobCollection::LogCachingStats()
 {
-	static double LastTimeStatsPrinted = FPlatformTime::Seconds();
-	// do not print if
-	//  - job cache is disabled
-	//  - not enough time passed since the previous time and we're not forced to print
-	if (!ShaderCompiler::IsJobCacheEnabled() || 
-		(!bForceLogIgnoringTimeInverval && GShaderCompilerCacheStatsPrintoutInterval > 0 && FPlatformTime::Seconds() - LastTimeStatsPrinted < GShaderCompilerCacheStatsPrintoutInterval))
-	{
-		return;
-	}
-
 	FWriteScopeLock Locker(Lock);	// write lock because logging actually changes the cache state (in a minor way - updating the memory used - but still).
 	CompletedJobsCache.LogStats();
-	LastTimeStatsPrinted = FPlatformTime::Seconds();
 }
 
 int32 FShaderCompileJobCollection::GetNumPendingJobs() const
@@ -2621,6 +2610,33 @@ void FShaderCompilerStats::WriteStats(FOutputDevice* Ar)
 	}
 #endif // ALLOW_DEBUG_FILES
 }
+
+void FShaderCompilerStats::WriteStatSummary()
+{
+	uint32 TotalCompiled = 0;
+
+	{
+		FScopeLock Lock(&CompileStatsLock);
+		const TSparseArray<ShaderCompilerStats>& PlatformStats = GetShaderCompilerStats();
+		for (int32 Platform = 0; Platform < PlatformStats.GetMaxIndex(); ++Platform)
+		{
+			if (PlatformStats.IsValidIndex(Platform))
+			{
+				const ShaderCompilerStats& Stats = PlatformStats[Platform];
+				for (const auto& Pair : Stats)
+				{
+					const FShaderCompilerStats::FShaderStats& SingleStats = Pair.Value;
+
+					TotalCompiled += SingleStats.Compiled;
+				}
+			}
+		}
+	}
+
+	UE_LOG(LogShaderCompilers, Display, TEXT("=== Shader Compilation stats ==="));
+	UE_LOG(LogShaderCompilers, Display, TEXT("Shaders Compiled: %u"), TotalCompiled);
+}
+
 void FShaderCompilerStats::RegisterCookedShaders(uint32 NumCooked, float CompileTime, EShaderPlatform Platform, const FString MaterialPath, FString PermutationString)
 {
 	FScopeLock Lock(&CompileStatsLock);
@@ -3773,7 +3789,7 @@ void FShaderCompilingManager::ProcessCompiledShaderMaps(
 		}
 	}
 
-	AllJobs.LogCachingStats();
+	GShaderCompilingManager->PrintStats();
 #endif // WITH_EDITOR
 }
 
@@ -3842,8 +3858,7 @@ void FShaderCompilingManager::PropagateMaterialChangesToPrimitives(const TMap<TR
  */
 void FShaderCompilingManager::Shutdown()
 {
-	// print the statistics on shutdown
-	AllJobs.LogCachingStats(true);
+	PrintStats();
 
 	for (const auto& Thread : Threads)
 	{
@@ -3852,6 +3867,30 @@ void FShaderCompilingManager::Shutdown()
 	}
 }
 
+void FShaderCompilingManager::PrintStats(bool bForceLogIgnoringTimeInverval)
+{
+	static double LastTimeStatsPrinted = FPlatformTime::Seconds();
+	// do not print if
+	//  - job cache is disabled
+	//  - not enough time passed since the previous time and we're not forced to print
+	if ((!bForceLogIgnoringTimeInverval && GShaderCompilerCacheStatsPrintoutInterval > 0 && FPlatformTime::Seconds() - LastTimeStatsPrinted < GShaderCompilerCacheStatsPrintoutInterval))
+	{
+		return;
+	}
+
+	UE_LOG(LogShaderCompilers, Display, TEXT("================================================"));
+
+	if (ShaderCompiler::IsJobCacheEnabled())
+	{
+		AllJobs.LogCachingStats();
+	}
+
+	GShaderCompilerStats->WriteStatSummary();
+
+	UE_LOG(LogShaderCompilers, Display, TEXT("================================================"));
+
+	LastTimeStatsPrinted = FPlatformTime::Seconds();
+}
 
 bool FShaderCompilingManager::HandlePotentialRetryOnError(TMap<int32, FShaderMapFinalizeResults>& CompletedShaderMaps)
 {
@@ -5153,7 +5192,7 @@ namespace
 		else
 		{
 			// tell other side all the materials to load, by pathname
-			for( TObjectIterator<UMaterialInterface> It; It; ++It )
+			for (TObjectIterator<UMaterialInterface> It; It; ++It)
 			{
 				OutMaterialsToLoad.Add(It->GetPathName());
 			}
@@ -6327,20 +6366,20 @@ void RecompileShadersForRemote(
 					FMaterialShaderMap::SaveForRemoteRecompile(Ar, CompiledShaderMaps);
 				}
 
-				// save it out so the client can get it (and it's up to date next time)
-				FString GlobalShaderFilename = SaveGlobalShaderFile(ShaderPlatform, OutputDirectory, TargetPlatform);
+					// save it out so the client can get it (and it's up to date next time)
+					FString GlobalShaderFilename = SaveGlobalShaderFile(ShaderPlatform, OutputDirectory, TargetPlatform);
 
-				// add this to the list of files to tell the other end about
-				if (ModifiedFiles)
-				{
-					// need to put it in non-sandbox terms
-					FString SandboxPath(GlobalShaderFilename);
-					check(SandboxPath.StartsWith(OutputDirectory));
-					SandboxPath.ReplaceInline(*OutputDirectory, TEXT("../../../"));
-					FPaths::NormalizeFilename(SandboxPath);
-					ModifiedFiles->Add(SandboxPath);
+					// add this to the list of files to tell the other end about
+					if (ModifiedFiles)
+					{
+						// need to put it in non-sandbox terms
+						FString SandboxPath(GlobalShaderFilename);
+						check(SandboxPath.StartsWith(OutputDirectory));
+						SandboxPath.ReplaceInline(*OutputDirectory, TEXT("../../../"));
+						FPaths::NormalizeFilename(SandboxPath);
+						ModifiedFiles->Add(SandboxPath);
+					}
 				}
-			}
 		}
 	}
 
@@ -6806,7 +6845,6 @@ void FShaderJobCache::LogStats()
 	{
 		UE_LOG(LogShaderCompilers, Display, TEXT("RAM used: %.2f MB (%.2f GB), no memory limit set"), MemUsedMB, MemUsedGB);
 	}
-	UE_LOG(LogShaderCompilers, Display, TEXT("================================================"));
 }
 
 #undef LOCTEXT_NAMESPACE
diff --git a/Engine/Source/Runtime/Engine/Public/ShaderCompiler.h b/Engine/Source/Runtime/Engine/Public/ShaderCompiler.h
--- a/Engine/Source/Runtime/Engine/Public/ShaderCompiler.h
+++ b/Engine/Source/Runtime/Engine/Public/ShaderCompiler.h
@@ -395,9 +395,8 @@ public:
 
 	/** Log caching statistics.
 	 *
-	 * @param bForceLogIgnoringTimeInverval - this function is called often, so not every invocation normally will actually log the stats. This parameter being true bypasses this pacing.
 	 */
-	void LogCachingStats(bool bForceLogIgnoringTimeInverval = false);
+	void LogCachingStats();
 
 	inline int32 GetNumPendingJobs(EShaderCompileJobPriority InPriority) const
 	{
@@ -909,6 +908,7 @@ public:
 	ENGINE_API void RegisterCompiledShaders(uint32 NumPermutations, EShaderPlatform Platform, const FString MaterialPath, FString PermutationString = FString(""));
 	ENGINE_API const TSparseArray<ShaderCompilerStats>& GetShaderCompilerStats() { return CompileStats; }
 	ENGINE_API void WriteStats(class FOutputDevice* Ar = nullptr);
+	ENGINE_API void WriteStatSummary();
 
 private:
 	FCriticalSection CompileStatsLock;
@@ -1178,6 +1178,11 @@ public:
 	 */
 	ENGINE_API void Shutdown();
 
+	/**
+	 * Prints stats related to shader compilation to the log.
+	 * @param bForceLogIgnoringTimeInverval - this function is called often, so not every invocation normally will actually log the stats. This parameter being true bypasses this pacing.
+	 */
+	ENGINE_API void PrintStats(bool bForceLogIgnoringTimeInverval = false);
 
 	/** 
 	 * Processes completed asynchronous shader maps, and assigns them to relevant materials.
