Skin cache parallel dispatch optimizations

*) batches share common vertex/tangent buffers because sk mesh is dispatched in sections, find all unique resources that require transitioning and transition them before skinning, so skinning dispatches can be processed in parallel
*) add an option to enable recompute tangents to do parallel dispatches, when it is turned on: resource transitions -> all triangle pass dispatches -> resource transitions -> all vertex pass dispatches, it requires an accumulated tangent buffer per sk mesh so uses more memory.
*) existing recompute tangents logic is kept and is the default, for projects that cannot afford to use more memory in tradeoff of faster performance

diff --git a/Engine/Shaders/Private/RecomputeTangentsCommon.ush b/Engine/Shaders/Private/RecomputeTangentsCommon.ush
--- a/Engine/Shaders/Private/RecomputeTangentsCommon.ush
+++ b/Engine/Shaders/Private/RecomputeTangentsCommon.ush
@@ -55,3 +55,4 @@ uint SkinCacheStart;
 // intermediate tangent accumulation buffer, int se we can use atomics and to get deterministic results no matter what order
 // uses as output on MainCS(), used and input and output (to clear) for ResolveCS()
 RWBuffer<int> IntermediateAccumBufferUAV;
+uint IntermediateAccumBufferOffset;
\ No newline at end of file
diff --git a/Engine/Shaders/Private/RecomputeTangentsPerTrianglePass.usf b/Engine/Shaders/Private/RecomputeTangentsPerTrianglePass.usf
--- a/Engine/Shaders/Private/RecomputeTangentsPerTrianglePass.usf
+++ b/Engine/Shaders/Private/RecomputeTangentsPerTrianglePass.usf
@@ -171,7 +171,7 @@ void MainCS(
 				// accumulate the contributions of each triangle on the triangle vertices
 				
 				// no start offset as we reuse the same buffer always from 0 on
-				uint Index = VertexIndex * INTERMEDIATE_ACCUM_BUFFER_NUM_INTS;
+				uint Index = (IntermediateAccumBufferOffset + VertexIndex) * INTERMEDIATE_ACCUM_BUFFER_NUM_INTS;
 
 				// TangentZ
 				InterlockedAdd(IntermediateAccumBufferUAV[Index + 0], IntTangentZ.x);
@@ -190,7 +190,7 @@ void MainCS(
                 for(uint DupIndexOffset = 0; DupIndexOffset < DupVertexIndicesLength; ++DupIndexOffset)
                 {
                     uint DupVertexIndex = DuplicatedIndices[DupIndexStart + DupIndexOffset];
-                    uint DupIndex = DupVertexIndex * INTERMEDIATE_ACCUM_BUFFER_NUM_INTS;
+                    uint DupIndex = (IntermediateAccumBufferOffset + DupVertexIndex) * INTERMEDIATE_ACCUM_BUFFER_NUM_INTS;
                     // TangentZ
                     InterlockedAdd(IntermediateAccumBufferUAV[DupIndex + 0], IntTangentZ.x);
                     InterlockedAdd(IntermediateAccumBufferUAV[DupIndex + 1], IntTangentZ.y);
diff --git a/Engine/Shaders/Private/RecomputeTangentsPerVertexPass.usf b/Engine/Shaders/Private/RecomputeTangentsPerVertexPass.usf
--- a/Engine/Shaders/Private/RecomputeTangentsPerVertexPass.usf
+++ b/Engine/Shaders/Private/RecomputeTangentsPerVertexPass.usf
@@ -134,7 +134,7 @@ void MainCS(
 		// the sign defines the handyness of the tangent matrix
   
         // no start offset as we reuse the same buffer always from 0 on
-        uint Index = DispatchThreadId * INTERMEDIATE_ACCUM_BUFFER_NUM_INTS;
+        uint Index = (IntermediateAccumBufferOffset + DispatchThreadId) * INTERMEDIATE_ACCUM_BUFFER_NUM_INTS;
 		
         // we don't have to scale down as we normalize anyway
         TangentZ = normalize(float3(IntermediateAccumBufferUAV[Index + 0], IntermediateAccumBufferUAV[Index + 1], IntermediateAccumBufferUAV[Index + 2]));
diff --git a/Engine/Source/Runtime/Engine/Private/GPUSkinCache.cpp b/Engine/Source/Runtime/Engine/Private/GPUSkinCache.cpp
--- a/Engine/Source/Runtime/Engine/Private/GPUSkinCache.cpp
+++ b/Engine/Source/Runtime/Engine/Private/GPUSkinCache.cpp
@@ -128,6 +128,16 @@ FAutoConsoleVariableRef CVarGPUSkinCacheBlendUsingVertexColorForRecomputeTangent
 	ECVF_RenderThreadSafe
 );
 
+int32 GRecomputeTangentsParallelDispatch = 0;
+FAutoConsoleVariableRef CVarRecomputeTangentsParallelDispatch(
+	TEXT("r.SkinCache.RecomputeTangentsParallelDispatch"),
+	GRecomputeTangentsParallelDispatch,
+	TEXT("This option enables parallel dispatches for recompute tangents.\n")
+	TEXT(" 0: off (default), triangle pass is interleaved with vertex pass, requires resource barriers in between. \n")
+	TEXT(" 1: on, batch triangle passes together, resource barrier, followed by vertex passes together, cost more memory. \n"),
+	ECVF_RenderThreadSafe
+);
+
 static int32 GGPUSkinCacheFlushCounter = 0;
 
 #if RHI_RAYTRACING
@@ -277,6 +287,7 @@ public:
 
 		FRWBuffer* TangentBuffer = nullptr;
 		FRWBuffer* IntermediateTangentBuffer = nullptr;
+		FRWBuffer* IntermediateAccumulatedTangentBuffer = nullptr;
 		FRWBuffer* PositionBuffer = nullptr;
 		FRWBuffer* PreviousPositionBuffer = nullptr;
 
@@ -325,6 +336,12 @@ public:
 			}
 		}
 
+		inline FRWBuffer* GetIntermediateAccumulatedTangentBuffer()
+		{
+			check(IntermediateAccumulatedTangentBuffer);
+			return IntermediateAccumulatedTangentBuffer;
+		}
+
 		void UpdateVertexFactoryDeclaration()
 		{
 			TargetVertexFactory->UpdateVertexDeclaration(SourceVertexFactory, GetPositionRWBuffer(), GetPreSkinPositionSRV(), GetTangentRWBuffer());
@@ -761,28 +778,20 @@ void FGPUSkinCache::Cleanup()
 	ensure(Allocations.Num() == 0);
 }
 
-void FGPUSkinCache::AddBufferToTransition(FRHIUnorderedAccessView* InUAV)
-{
-	// add UAV to set to remove duplicated entries but could still be different UAVs on the same resource
-	// then this code will need better filtering because mutliple transitions on the same resource
-	// is not allowed
-	BuffersToTransition.Add(InUAV);
-}
-
 void FGPUSkinCache::TransitionAllToReadable(FRHICommandList& RHICmdList)
 {
-	if (BuffersToTransition.Num() > 0)
+	if (BuffersToTransitionToRead.Num() > 0)
 	{
 		FMemMark Mark(FMemStack::Get());
 		TArray<FRHITransitionInfo, SceneRenderingAllocator> UAVs;
-		UAVs.Reserve(BuffersToTransition.Num());
-		for (TSet<FRHIUnorderedAccessView*>::TConstIterator SetIt(BuffersToTransition); SetIt; ++SetIt)
+		UAVs.Reserve(BuffersToTransitionToRead.Num());
+		for (TSet<FRHIUnorderedAccessView*>::TConstIterator SetIt(BuffersToTransitionToRead); SetIt; ++SetIt)
 		{
 			UAVs.Add(FRHITransitionInfo(*SetIt, ERHIAccess::Unknown, ERHIAccess::VertexOrIndexBuffer | ERHIAccess::SRVMask));
 		}
 		RHICmdList.Transition(UAVs);
 
-		BuffersToTransition.Empty(BuffersToTransition.Num());
+		BuffersToTransitionToRead.Empty(BuffersToTransitionToRead.Num());
 	}
 }
 
@@ -908,6 +917,7 @@ public:
 		: FGlobalShader(Initializer)
 	{
 		IntermediateAccumBufferUAV.Bind(Initializer.ParameterMap, TEXT("IntermediateAccumBufferUAV"));
+		IntermediateAccumBufferOffset.Bind(Initializer.ParameterMap, TEXT("IntermediateAccumBufferOffset"));
 		NumTriangles.Bind(Initializer.ParameterMap, TEXT("NumTriangles"));
 		GPUPositionCacheBuffer.Bind(Initializer.ParameterMap, TEXT("GPUPositionCacheBuffer"));
 		GPUTangentCacheBuffer.Bind(Initializer.ParameterMap, TEXT("GPUTangentCacheBuffer"));
@@ -950,6 +960,7 @@ public:
 
 		// UAV
 		SetUAVParameter(RHICmdList, ShaderRHI, IntermediateAccumBufferUAV, StagingBuffer.UAV);
+		SetShaderValue(RHICmdList, ShaderRHI, IntermediateAccumBufferOffset, GRecomputeTangentsParallelDispatch * DispatchData.IndexBufferOffsetValue);
 
         if (!GAllowDupedVertsForRecomputeTangents)
         {
@@ -966,6 +977,7 @@ public:
 	}
 
 	LAYOUT_FIELD(FShaderResourceParameter, IntermediateAccumBufferUAV);
+	LAYOUT_FIELD(FShaderParameter, IntermediateAccumBufferOffset);
 	LAYOUT_FIELD(FShaderParameter, NumTriangles);
 	LAYOUT_FIELD(FShaderResourceParameter, GPUPositionCacheBuffer);
 	LAYOUT_FIELD(FShaderResourceParameter, GPUTangentCacheBuffer);
@@ -1028,6 +1040,7 @@ public:
 	static const uint32 ThreadGroupSizeX = 64;
 
 	LAYOUT_FIELD(FShaderResourceParameter, IntermediateAccumBufferUAV);
+	LAYOUT_FIELD(FShaderParameter, IntermediateAccumBufferOffset);
 	LAYOUT_FIELD(FShaderResourceParameter, TangentBufferUAV);
 	LAYOUT_FIELD(FShaderResourceParameter, TangentInputBuffer);
 	LAYOUT_FIELD(FShaderResourceParameter, ColorInputBuffer);
@@ -1042,6 +1055,7 @@ public:
 		: FGlobalShader(Initializer)
 	{
 		IntermediateAccumBufferUAV.Bind(Initializer.ParameterMap, TEXT("IntermediateAccumBufferUAV"));
+		IntermediateAccumBufferOffset.Bind(Initializer.ParameterMap, TEXT("IntermediateAccumBufferOffset"));
 		TangentBufferUAV.Bind(Initializer.ParameterMap, TEXT("TangentBufferUAV"));
 		TangentInputBuffer.Bind(Initializer.ParameterMap, TEXT("TangentInputBuffer"));
 		ColorInputBuffer.Bind(Initializer.ParameterMap, TEXT("ColorInputBuffer"));
@@ -1066,6 +1080,7 @@ public:
 
 		// UAVs
 		SetUAVParameter(RHICmdList, ShaderRHI, IntermediateAccumBufferUAV, StagingBuffer.UAV);
+		SetShaderValue(RHICmdList, ShaderRHI, IntermediateAccumBufferOffset, GRecomputeTangentsParallelDispatch * DispatchData.IndexBufferOffsetValue);
 		SetUAVParameter(RHICmdList, ShaderRHI, TangentBufferUAV, DispatchData.GetTangentRWBuffer()->UAV);
 
 		SetSRVParameter(RHICmdList, ShaderRHI, TangentInputBuffer, DispatchData.IntermediateTangentBuffer ? DispatchData.IntermediateTangentBuffer->SRV : nullptr);
@@ -1113,58 +1128,51 @@ IMPLEMENT_SHADER_TYPE(template<>, FRecomputeTangentsPerVertexPassCS<2>, TEXT("/E
 IMPLEMENT_SHADER_TYPE(template<>, FRecomputeTangentsPerVertexPassCS<3>, TEXT("/Engine/Private/RecomputeTangentsPerVertexPass.usf"), TEXT("MainCS"), SF_Compute);
 IMPLEMENT_SHADER_TYPE(template<>, FRecomputeTangentsPerVertexPassCS<4>, TEXT("/Engine/Private/RecomputeTangentsPerVertexPass.usf"), TEXT("MainCS"), SF_Compute);
 
-void FGPUSkinCache::DispatchUpdateSkinTangents(FRHICommandListImmediate& RHICmdList, FGPUSkinCacheEntry* Entry, int32 SectionIndex)
+void FGPUSkinCache::DispatchUpdateSkinTangents(FRHICommandListImmediate& RHICmdList, FGPUSkinCacheEntry* Entry, int32 SectionIndex, FRWBuffer*& StagingBuffer, bool bTrianglePass)
 {
-	FGPUSkinCacheEntry::FSectionDispatchData& DispatchData = Entry->DispatchData[SectionIndex];
-
-	{
-		// no need to clear the intermediate buffer because we create it cleared and clear it after each usage in the per vertex pass
-
-		FString MeshName = TEXT("None");
+	FString MeshName = TEXT("None");
 #if !UE_BUILD_SHIPPING
-		MeshName = Entry->GPUSkin->DebugName.ToString();
+	MeshName = Entry->GPUSkin->DebugName.ToString();
 #endif // !UE_BUILD_SHIPPING
 
-		FSkeletalMeshRenderData& SkelMeshRenderData = Entry->GPUSkin->GetSkeletalMeshRenderData();
-		int32 LODIndex = Entry->LOD;
-		FSkeletalMeshLODRenderData& LodData = SkelMeshRenderData.LODRenderData[Entry->LOD];
-
-		//SetRenderTarget(RHICmdList, FTextureRHIRef(), FTextureRHIRef());
-
-		FRawStaticIndexBuffer16or32Interface* IndexBuffer = LodData.MultiSizeIndexContainer.GetIndexBuffer();
-		FRHIBuffer* IndexBufferRHI = IndexBuffer->IndexBufferRHI;
-
-		const uint32 RequiredVertexCount = LodData.GetNumVertices();
+	FGPUSkinCacheEntry::FSectionDispatchData& DispatchData = Entry->DispatchData[SectionIndex];
 
-		uint32 MaxVertexCount = RequiredVertexCount;
+	FSkeletalMeshRenderData& SkelMeshRenderData = Entry->GPUSkin->GetSkeletalMeshRenderData();
+	int32 LODIndex = Entry->LOD;
+	FSkeletalMeshLODRenderData& LodData = SkelMeshRenderData.LODRenderData[Entry->LOD];
 
-		if (StagingBuffers.Num() != GNumTangentIntermediateBuffers)
-		{
-			// Release extra buffers if shrinking
-			for (int32 Index = GNumTangentIntermediateBuffers; Index < StagingBuffers.Num(); ++Index)
+	if (bTrianglePass)
+	{
+		if (!GRecomputeTangentsParallelDispatch)
+		{	
+			if (StagingBuffers.Num() != GNumTangentIntermediateBuffers)
 			{
-				StagingBuffers[Index].Release();
+				// Release extra buffers if shrinking
+				for (int32 Index = GNumTangentIntermediateBuffers; Index < StagingBuffers.Num(); ++Index)
+				{
+					StagingBuffers[Index].Release();
+				}
+				StagingBuffers.SetNum(GNumTangentIntermediateBuffers, false);
 			}
-			StagingBuffers.SetNum(GNumTangentIntermediateBuffers, false);
-		}
 
-		uint32 NumIntsPerBuffer = DispatchData.NumTriangles * 3 * FGPUSkinCache::IntermediateAccumBufferNumInts;
-		CurrentStagingBufferIndex = (CurrentStagingBufferIndex + 1) % StagingBuffers.Num();
-		FRWBuffer& StagingBuffer = StagingBuffers[CurrentStagingBufferIndex];
-		if (StagingBuffer.NumBytes < NumIntsPerBuffer * sizeof(uint32))
-		{
-			StagingBuffer.Release();
-			StagingBuffer.Initialize(TEXT("SkinTangentIntermediate"), sizeof(int32), NumIntsPerBuffer, PF_R32_SINT, BUF_UnorderedAccess);
-			RHICmdList.BindDebugLabelName(StagingBuffer.UAV, TEXT("SkinTangentIntermediate"));
+			// no need to clear the staging buffer because we create it cleared and clear it after each usage in the per vertex pass
+			uint32 NumIntsPerBuffer = DispatchData.NumTriangles * 3 * FGPUSkinCache::IntermediateAccumBufferNumInts;
+			CurrentStagingBufferIndex = (CurrentStagingBufferIndex + 1) % StagingBuffers.Num();
+			StagingBuffer = &StagingBuffers[CurrentStagingBufferIndex];
+			if (StagingBuffer->NumBytes < NumIntsPerBuffer * sizeof(uint32))
+			{
+				StagingBuffer->Release();
+				StagingBuffer->Initialize(TEXT("SkinTangentIntermediate"), sizeof(int32), NumIntsPerBuffer, PF_R32_SINT, BUF_UnorderedAccess);
+				RHICmdList.BindDebugLabelName(StagingBuffer->UAV, TEXT("SkinTangentIntermediate"));
 
-			const uint32 MemSize = NumIntsPerBuffer * sizeof(uint32);
-			SET_MEMORY_STAT(STAT_GPUSkinCache_TangentsIntermediateMemUsed, MemSize);
+				const uint32 MemSize = NumIntsPerBuffer * sizeof(uint32);
+				SET_MEMORY_STAT(STAT_GPUSkinCache_TangentsIntermediateMemUsed, MemSize);
 
-			// The UAV must be zero-filled. We leave it zeroed after each round (see RecomputeTangentsPerVertexPass.usf), so this is only needed on when the buffer is first created.
-			RHICmdList.ClearUAVUint(StagingBuffer.UAV, FUintVector4(0, 0, 0, 0));
+				// The UAV must be zero-filled. We leave it zeroed after each round (see RecomputeTangentsPerVertexPass.usf), so this is only needed on when the buffer is first created.
+				RHICmdList.ClearUAVUint(StagingBuffer->UAV, FUintVector4(0, 0, 0, 0));
+			}
 		}
 
-		// This code can be optimized by batched up and doing it with less Dispatch calls (costs more memory)
 		{
 			auto* GlobalShaderMap = GetGlobalShaderMap(GetFeatureLevel());
 			TShaderMapRef<FRecomputeTangentsPerTrianglePassCS<0>> ComputeShader00(GlobalShaderMap);
@@ -1192,7 +1200,7 @@ void FGPUSkinCache::DispatchUpdateSkinTangents(FRHICommandListImmediate& RHICmdL
 			uint32 NumTriangles = DispatchData.NumTriangles;
 			uint32 ThreadGroupCountValue = FMath::DivideAndRoundUp(NumTriangles, FBaseRecomputeTangentsPerTriangleShader::ThreadGroupSizeX);
 
-			SCOPED_DRAW_EVENTF(RHICmdList, SkinTangents_PerTrianglePass, TEXT("TangentsTri Mesh=%s, LOD=%d, Chunk=%d, IndexStart=%d Tri=%d BoneInfluenceType=%d UVPrecision=%d"),
+			SCOPED_DRAW_EVENTF(RHICmdList, SkinTangents_PerTrianglePass, TEXT("TangentsTri  Mesh=%s, LOD=%d, Chunk=%d, IndexStart=%d Tri=%d BoneInfluenceType=%d UVPrecision=%d"),
 				*MeshName, LODIndex, SectionIndex, DispatchData.IndexBufferOffsetValue, DispatchData.NumTriangles, Entry->BoneInfluenceType, bFullPrecisionUV);
 
 			FRHIComputeShader* ShaderRHI = Shader.GetComputeShader();
@@ -1205,67 +1213,64 @@ void FGPUSkinCache::DispatchUpdateSkinTangents(FRHICommandListImmediate& RHICmdL
 				DispatchData.DuplicatedIndicesIndices = LodData.RenderSections[SectionIndex].DuplicatedVerticesBuffer.LengthAndIndexDuplicatedVerticesIndexBuffer.VertexBufferSRV;
 			}
 
-			TArray<FRHITransitionInfo, TInlineAllocator<4>> Transitions;
-			Transitions.Add(FRHITransitionInfo(StagingBuffer.UAV, ERHIAccess::Unknown, ERHIAccess::UAVCompute));
-			Transitions.Add(FRHITransitionInfo(DispatchData.GetPositionRWBuffer()->UAV.GetReference(), ERHIAccess::Unknown, ERHIAccess::SRVCompute));
-			if (!GBlendUsingVertexColorForRecomputeTangents)
+			if (!GRecomputeTangentsParallelDispatch)
 			{
-				Transitions.Add(FRHITransitionInfo(DispatchData.GetTangentRWBuffer()->UAV.GetReference(), ERHIAccess::Unknown, ERHIAccess::SRVCompute));
+				// When triangle & vertex passes are interleaved, resource transition is needed in between.
+				RHICmdList.Transition(FRHITransitionInfo(StagingBuffer->UAV, ERHIAccess::Unknown, ERHIAccess::UAVCompute));
 			}
-			else if (DispatchData.IntermediateTangentBuffer)
-			{
-				Transitions.Add(FRHITransitionInfo(DispatchData.IntermediateTangentBuffer->UAV.GetReference(), ERHIAccess::Unknown, ERHIAccess::SRVCompute));
-			}
-			RHICmdList.Transition(Transitions);
 
 			INC_DWORD_STAT_BY(STAT_GPUSkinCache_NumTrianglesForRecomputeTangents, NumTriangles);
-			Shader->SetParameters(RHICmdList, Entry, DispatchData, StagingBuffer);
+			Shader->SetParameters(RHICmdList, Entry, DispatchData, GRecomputeTangentsParallelDispatch ? *DispatchData.GetIntermediateAccumulatedTangentBuffer() : *StagingBuffer);
 			DispatchComputeShader(RHICmdList, Shader.GetShader(), ThreadGroupCountValue, 1, 1);
 			Shader->UnsetParameters(RHICmdList);
 		}
+	}
+	else
+	{
+		SCOPED_DRAW_EVENTF(RHICmdList, SkinTangents_PerVertexPass, TEXT("TangentsVertex Mesh=%s, LOD=%d, Chunk=%d, InputStreamStart=%d, OutputStreamStart=%d, Vert=%d"),
+			*MeshName, LODIndex, SectionIndex, DispatchData.InputStreamStart, DispatchData.OutputStreamStart, DispatchData.NumVertices);
+		//#todo-gpuskin Feature level?
+		auto* GlobalShaderMap = GetGlobalShaderMap(GetFeatureLevel());
+		TShaderMapRef<FRecomputeTangentsPerVertexPassCS<0>> ComputeShader0(GlobalShaderMap);
+		TShaderMapRef<FRecomputeTangentsPerVertexPassCS<1>> ComputeShader1(GlobalShaderMap);
+		TShaderMapRef<FRecomputeTangentsPerVertexPassCS<2>> ComputeShader2(GlobalShaderMap);
+		TShaderMapRef<FRecomputeTangentsPerVertexPassCS<3>> ComputeShader3(GlobalShaderMap);
+		TShaderMapRef<FRecomputeTangentsPerVertexPassCS<4>> ComputeShader4(GlobalShaderMap);
+		TShaderRef<FBaseRecomputeTangentsPerVertexShader> ComputeShader;
+		if (GBlendUsingVertexColorForRecomputeTangents == 1)
+			ComputeShader = ComputeShader1;
+		else if (GBlendUsingVertexColorForRecomputeTangents == 2)
+			ComputeShader = ComputeShader2;
+		else if (GBlendUsingVertexColorForRecomputeTangents == 3)
+			ComputeShader = ComputeShader3;
+		else if (GBlendUsingVertexColorForRecomputeTangents == 4)
+			ComputeShader = ComputeShader4;
+		else
+			ComputeShader = ComputeShader0;
+		RHICmdList.SetComputeShader(ComputeShader.GetComputeShader());
 
-		{
-			SCOPED_DRAW_EVENTF(RHICmdList, SkinTangents_PerVertexPass, TEXT("TangentsVertex Mesh=%s, LOD=%d, Chunk=%d, InputStreamStart=%d, OutputStreamStart=%d, Vert=%d"),
-				*MeshName, LODIndex, SectionIndex, DispatchData.InputStreamStart, DispatchData.OutputStreamStart, DispatchData.NumVertices);
-			//#todo-gpuskin Feature level?
-			auto* GlobalShaderMap = GetGlobalShaderMap(GetFeatureLevel());
-			TShaderMapRef<FRecomputeTangentsPerVertexPassCS<0>> ComputeShader0(GlobalShaderMap);
-			TShaderMapRef<FRecomputeTangentsPerVertexPassCS<1>> ComputeShader1(GlobalShaderMap);
-			TShaderMapRef<FRecomputeTangentsPerVertexPassCS<2>> ComputeShader2(GlobalShaderMap);
-			TShaderMapRef<FRecomputeTangentsPerVertexPassCS<3>> ComputeShader3(GlobalShaderMap);
-			TShaderMapRef<FRecomputeTangentsPerVertexPassCS<4>> ComputeShader4(GlobalShaderMap);
-			TShaderRef<FBaseRecomputeTangentsPerVertexShader> ComputeShader;
-			if (GBlendUsingVertexColorForRecomputeTangents == 1)
-				ComputeShader = ComputeShader1;
-			else if (GBlendUsingVertexColorForRecomputeTangents == 2)
-				ComputeShader = ComputeShader2;
-			else if (GBlendUsingVertexColorForRecomputeTangents == 3)
-				ComputeShader = ComputeShader3;
-			else if (GBlendUsingVertexColorForRecomputeTangents == 4)
-				ComputeShader = ComputeShader4;
-			else
-				ComputeShader = ComputeShader0;
-			RHICmdList.SetComputeShader(ComputeShader.GetComputeShader());
-
-			uint32 VertexCount = DispatchData.NumVertices;
-			uint32 ThreadGroupCountValue = FMath::DivideAndRoundUp(VertexCount, ComputeShader->ThreadGroupSizeX);
+		uint32 VertexCount = DispatchData.NumVertices;
+		uint32 ThreadGroupCountValue = FMath::DivideAndRoundUp(VertexCount, ComputeShader->ThreadGroupSizeX);
 
+		if (!GRecomputeTangentsParallelDispatch)
+		{
+			// When triangle & vertex passes are interleaved, resource transition is needed in between.
 			RHICmdList.Transition({
 				FRHITransitionInfo(DispatchData.GetTangentRWBuffer()->UAV.GetReference(), GBlendUsingVertexColorForRecomputeTangents ? ERHIAccess::Unknown : ERHIAccess::SRVCompute, ERHIAccess::UAVCompute),
-				FRHITransitionInfo(StagingBuffer.UAV, ERHIAccess::UAVCompute, ERHIAccess::UAVCompute)
-			});
-
-			ComputeShader->SetParameters(RHICmdList, Entry, DispatchData, StagingBuffer);
-			DispatchComputeShader(RHICmdList, ComputeShader.GetShader(), ThreadGroupCountValue, 1, 1);
-			ComputeShader->UnsetParameters(RHICmdList);
+				FRHITransitionInfo(StagingBuffer->UAV, ERHIAccess::UAVCompute, ERHIAccess::UAVCompute)
+				});
 		}
+
+		ComputeShader->SetParameters(RHICmdList, Entry, DispatchData, GRecomputeTangentsParallelDispatch ? *DispatchData.GetIntermediateAccumulatedTangentBuffer() : *StagingBuffer);
+		DispatchComputeShader(RHICmdList, ComputeShader.GetShader(), ThreadGroupCountValue, 1, 1);
+		ComputeShader->UnsetParameters(RHICmdList);
 	}
 }
 
-FGPUSkinCache::FRWBuffersAllocation* FGPUSkinCache::TryAllocBuffer(uint32 NumVertices, bool WithTangnents)
+FGPUSkinCache::FRWBuffersAllocation* FGPUSkinCache::TryAllocBuffer(uint32 NumVertices, bool WithTangnents, uint32 NumTriangles, FRHICommandListImmediate& RHICmdList)
 {
 	uint64 MaxSizeInBytes = (uint64)(GSkinCacheSceneMemoryLimitInMB * 1024.0f * 1024.0f);
-	uint64 RequiredMemInBytes = FRWBuffersAllocation::CalculateRequiredMemory(NumVertices, WithTangnents);
+	uint64 RequiredMemInBytes = FRWBuffersAllocation::CalculateRequiredMemory(NumVertices, WithTangnents, NumTriangles);
 	if (bRequiresMemoryLimit && UsedMemoryInBytes + RequiredMemInBytes >= MaxSizeInBytes)
 	{
 		ExtraRequiredMemory += RequiredMemInBytes;
@@ -1274,7 +1279,7 @@ FGPUSkinCache::FRWBuffersAllocation* FGPUSkinCache::TryAllocBuffer(uint32 NumVer
 		return nullptr;
 	}
 
-	FRWBuffersAllocation* NewAllocation = new FRWBuffersAllocation(NumVertices, WithTangnents);
+	FRWBuffersAllocation* NewAllocation = new FRWBuffersAllocation(NumVertices, WithTangnents, NumTriangles, RHICmdList);
 	Allocations.Add(NewAllocation);
 
 	UsedMemoryInBytes += RequiredMemInBytes;
@@ -1285,6 +1290,21 @@ FGPUSkinCache::FRWBuffersAllocation* FGPUSkinCache::TryAllocBuffer(uint32 NumVer
 
 DECLARE_GPU_STAT(GPUSkinCache);
 
+void FGPUSkinCache::MakeBufferTransitions(FRHICommandListImmediate& RHICmdList, TArray<FRHIUnorderedAccessView*>& Buffers, ERHIAccess FromState, ERHIAccess ToState)
+{
+	if (Buffers.Num() > 0)
+	{
+		FMemMark Mark(FMemStack::Get());
+		TArray<FRHITransitionInfo, SceneRenderingAllocator> UAVs;
+		UAVs.Reserve(Buffers.Num());
+		for (FRHIUnorderedAccessView* Buffer : Buffers)
+		{
+			UAVs.Add(FRHITransitionInfo(Buffer, FromState, ToState));
+		}
+		RHICmdList.Transition(MakeArrayView(UAVs.GetData(), UAVs.Num()));
+	}
+}
+
 void FGPUSkinCache::DoDispatch(FRHICommandListImmediate& RHICmdList)
 {
 	SCOPED_GPU_STAT(RHICmdList, GPUSkinCache);
@@ -1292,20 +1312,22 @@ void FGPUSkinCache::DoDispatch(FRHICommandListImmediate& RHICmdList)
 	int32 BatchCount = BatchDispatches.Num();
 	INC_DWORD_STAT_BY(STAT_GPUSkinCache_TotalNumChunks, BatchCount);
 
-	TArray<FRHIUnorderedAccessView*> OverlappedUAVBuffers;
-	OverlappedUAVBuffers.Reserve(BatchCount * 2);
+	TArray<FRHIUnorderedAccessView*> BuffersToTransitionForSkinning;
+	BuffersToTransitionForSkinning.Reserve(BatchCount * 2);
 	{
 		for (int32 i = 0; i < BatchCount; ++i)
 		{
 			FDispatchEntry& DispatchItem = BatchDispatches[i];
-			PrepareUpdateSkinning(DispatchItem.SkinCacheEntry, DispatchItem.Section, DispatchItem.RevisionNumber, &OverlappedUAVBuffers);
+			PrepareUpdateSkinning(DispatchItem.SkinCacheEntry, DispatchItem.Section, DispatchItem.RevisionNumber, &BuffersToTransitionForSkinning);
 		}
 
-		Algo::Sort(OverlappedUAVBuffers, [](const FRHIUnorderedAccessView* A, const FRHIUnorderedAccessView* B){return A < B;});
-		OverlappedUAVBuffers.SetNum(Algo::Unique(OverlappedUAVBuffers));
+		Algo::Sort(BuffersToTransitionForSkinning, [](const FRHIUnorderedAccessView* A, const FRHIUnorderedAccessView* B){return A < B;});
+		BuffersToTransitionForSkinning.SetNum(Algo::Unique(BuffersToTransitionForSkinning));
+		MakeBufferTransitions(RHICmdList, BuffersToTransitionForSkinning, ERHIAccess::Unknown, ERHIAccess::UAVCompute);
+		BuffersToTransitionToRead.Append(BuffersToTransitionForSkinning);
 	}
 
-	RHICmdList.BeginUAVOverlap(OverlappedUAVBuffers);
+	RHICmdList.BeginUAVOverlap(BuffersToTransitionForSkinning);
 	{
 		SCOPED_DRAW_EVENT(RHICmdList, GPUSkinCache_UpdateSkinningBatches);
 		for (int32 i = 0; i < BatchCount; ++i)
@@ -1314,23 +1336,77 @@ void FGPUSkinCache::DoDispatch(FRHICommandListImmediate& RHICmdList)
 			DispatchUpdateSkinning(RHICmdList, DispatchItem.SkinCacheEntry, DispatchItem.Section, DispatchItem.RevisionNumber);
 		}
 	}
-	RHICmdList.EndUAVOverlap(OverlappedUAVBuffers);
+	RHICmdList.EndUAVOverlap(BuffersToTransitionForSkinning);
+
+	// Do necessary buffer transitions before recomputing tangents
+	TArray<FRHIUnorderedAccessView*> BuffersToSRVForRecomputeTangents;
+	TArray<FRHIUnorderedAccessView*> IntermediateAccumulatedTangentBuffers;
+	for (int32 i = 0; i < BatchCount; ++i)
+	{
+		FDispatchEntry& DispatchItem = BatchDispatches[i];
+		FGPUSkinCacheEntry::FSectionDispatchData& DispatchData = DispatchItem.SkinCacheEntry->DispatchData[DispatchItem.Section];
+		if (DispatchData.IndexBuffer)
+		{
+			BuffersToSRVForRecomputeTangents.AddUnique(DispatchData.GetPositionRWBuffer()->UAV);
+			BuffersToSRVForRecomputeTangents.AddUnique(DispatchData.GetActiveTangentRWBuffer()->UAV);
+			if (GRecomputeTangentsParallelDispatch)
+			{
+				IntermediateAccumulatedTangentBuffers.AddUnique(DispatchData.GetIntermediateAccumulatedTangentBuffer()->UAV);
+			}
+		}	
+	}
+	MakeBufferTransitions(RHICmdList, BuffersToSRVForRecomputeTangents, ERHIAccess::Unknown, ERHIAccess::SRVCompute);
+	MakeBufferTransitions(RHICmdList, IntermediateAccumulatedTangentBuffers, ERHIAccess::Unknown, ERHIAccess::UAVCompute);
 
+	RHICmdList.BeginUAVOverlap(IntermediateAccumulatedTangentBuffers);
 	{
 		SCOPED_DRAW_EVENT(RHICmdList, GPUSkinCache_RecomputeTangentsBatches);
+		FRWBuffer* StagingBuffer = nullptr;
 		for (int32 i = 0; i < BatchCount; ++i)
 		{
 			FDispatchEntry& DispatchItem = BatchDispatches[i];
-			DispatchItem.SkinCacheEntry->UpdateVertexFactoryDeclaration(DispatchItem.Section);
-	
 			if (DispatchItem.SkinCacheEntry->DispatchData[DispatchItem.Section].IndexBuffer)
 			{
-				DispatchUpdateSkinTangents(RHICmdList, DispatchItem.SkinCacheEntry, DispatchItem.Section);
+				DispatchUpdateSkinTangents(RHICmdList, DispatchItem.SkinCacheEntry, DispatchItem.Section, StagingBuffer, true);
+				if (!GRecomputeTangentsParallelDispatch)
+				{
+					// When parallel dispatching is off, triangle pass and vertex pass are dispatched interleaved.
+					DispatchUpdateSkinTangents(RHICmdList, DispatchItem.SkinCacheEntry, DispatchItem.Section, StagingBuffer, false);
+				}
+			}
+		}
+		if (GRecomputeTangentsParallelDispatch)
+		{
+			// Do necessary buffer transitions before vertex pass dispatches
+			TArray<FRHIUnorderedAccessView*> TangentBuffers;
+			for (int32 i = 0; i < BatchCount; ++i)
+			{
+				FDispatchEntry& DispatchItem = BatchDispatches[i];
+				FGPUSkinCacheEntry::FSectionDispatchData& DispatchData = DispatchItem.SkinCacheEntry->DispatchData[DispatchItem.Section];
+				TangentBuffers.AddUnique(DispatchData.GetTangentRWBuffer()->UAV);
 			}
-	
-			DispatchItem.SkinCacheEntry->UpdateVertexFactoryDeclaration(DispatchItem.Section);
+			MakeBufferTransitions(RHICmdList, TangentBuffers, GBlendUsingVertexColorForRecomputeTangents ? ERHIAccess::Unknown : ERHIAccess::SRVCompute, ERHIAccess::UAVCompute);
+			MakeBufferTransitions(RHICmdList, IntermediateAccumulatedTangentBuffers, ERHIAccess::UAVCompute, ERHIAccess::UAVCompute);
+		
+			RHICmdList.BeginUAVOverlap(TangentBuffers);
+			for (int32 i = 0; i < BatchCount; ++i)
+			{
+				FDispatchEntry& DispatchItem = BatchDispatches[i];
+				if (DispatchItem.SkinCacheEntry->DispatchData[DispatchItem.Section].IndexBuffer)
+				{
+					DispatchUpdateSkinTangents(RHICmdList, DispatchItem.SkinCacheEntry, DispatchItem.Section, StagingBuffer, false);
+				}
+			}
+			RHICmdList.EndUAVOverlap(TangentBuffers);
 		}
 	}
+	RHICmdList.EndUAVOverlap(IntermediateAccumulatedTangentBuffers);
+
+	for (int32 i = 0; i < BatchCount; ++i)
+	{
+		FDispatchEntry& DispatchItem = BatchDispatches[i];
+		DispatchItem.SkinCacheEntry->UpdateVertexFactoryDeclaration(DispatchItem.Section);
+	}
 }
 
 void FGPUSkinCache::DoDispatch(FRHICommandListImmediate& RHICmdList, FGPUSkinCacheEntry* SkinCacheEntry, int32 Section, int32 RevisionNumber)
@@ -1338,14 +1414,41 @@ void FGPUSkinCache::DoDispatch(FRHICommandListImmediate& RHICmdList, FGPUSkinCac
 	SCOPED_GPU_STAT(RHICmdList, GPUSkinCache);
 
 	INC_DWORD_STAT(STAT_GPUSkinCache_TotalNumChunks);
-	PrepareUpdateSkinning(SkinCacheEntry, Section, RevisionNumber, nullptr);
+
+	TArray<FRHIUnorderedAccessView*> BuffersToTransitionForSkinning;
+	PrepareUpdateSkinning(SkinCacheEntry, Section, RevisionNumber, &BuffersToTransitionForSkinning);
+	MakeBufferTransitions(RHICmdList, BuffersToTransitionForSkinning, ERHIAccess::Unknown, ERHIAccess::UAVCompute);
+	BuffersToTransitionToRead.Append(BuffersToTransitionForSkinning);
+
+	RHICmdList.BeginUAVOverlap(BuffersToTransitionForSkinning);
 	DispatchUpdateSkinning(RHICmdList, SkinCacheEntry, Section, RevisionNumber);
-	SkinCacheEntry->UpdateVertexFactoryDeclaration(Section);
+	RHICmdList.EndUAVOverlap(BuffersToTransitionForSkinning);
 
-	if (SkinCacheEntry->DispatchData[Section].IndexBuffer)
+	FGPUSkinCacheEntry::FSectionDispatchData& DispatchData = SkinCacheEntry->DispatchData[Section];
+	if (DispatchData.IndexBuffer)
 	{
-		DispatchUpdateSkinTangents(RHICmdList, SkinCacheEntry, Section);
+		RHICmdList.Transition({
+			FRHITransitionInfo(DispatchData.GetPositionRWBuffer()->UAV.GetReference(), ERHIAccess::Unknown, ERHIAccess::SRVCompute),
+			FRHITransitionInfo(DispatchData.GetActiveTangentRWBuffer()->UAV.GetReference(), ERHIAccess::Unknown, ERHIAccess::SRVCompute)
+		});
+		if (GRecomputeTangentsParallelDispatch)
+		{
+			RHICmdList.Transition(FRHITransitionInfo(DispatchData.GetIntermediateAccumulatedTangentBuffer()->UAV.GetReference(), ERHIAccess::Unknown, ERHIAccess::UAVCompute));
+		}
+
+		FRWBuffer* StagingBuffer = nullptr;
+		DispatchUpdateSkinTangents(RHICmdList, SkinCacheEntry, Section, StagingBuffer, true);
+		if (GRecomputeTangentsParallelDispatch)
+		{
+			RHICmdList.Transition({
+				FRHITransitionInfo(DispatchData.GetTangentRWBuffer()->UAV.GetReference(), GBlendUsingVertexColorForRecomputeTangents ? ERHIAccess::Unknown : ERHIAccess::SRVCompute, ERHIAccess::UAVCompute),
+				FRHITransitionInfo(DispatchData.GetIntermediateAccumulatedTangentBuffer()->UAV.GetReference(), ERHIAccess::UAVCompute, ERHIAccess::UAVCompute)
+			});
+		}
+		DispatchUpdateSkinTangents(RHICmdList, SkinCacheEntry, Section, StagingBuffer, false);
 	}
+
+	SkinCacheEntry->UpdateVertexFactoryDeclaration(Section);
 }
 
 void FGPUSkinCache::ProcessEntry(
@@ -1405,7 +1508,9 @@ void FGPUSkinCache::ProcessEntry(
 	{
 		bool WithTangents = RecomputeTangentsMode > 0;
 		int32 TotalNumVertices = VertexFactory->GetNumVertices();
-		FRWBuffersAllocation* NewPositionAllocation = TryAllocBuffer(TotalNumVertices, WithTangents);
+
+		uint32 TotalNumTriangles = GRecomputeTangentsParallelDispatch ? LodData.GetTotalFaces() : 0;
+		FRWBuffersAllocation* NewPositionAllocation = TryAllocBuffer(TotalNumVertices, WithTangents, TotalNumTriangles, RHICmdList);
 		if (!NewPositionAllocation)
 		{
 			// Couldn't fit; caller will notify OOM
@@ -1797,6 +1902,7 @@ void FGPUSkinCache::PrepareUpdateSkinning(FGPUSkinCacheEntry* Entry, int32 Secti
 
 	DispatchData.TangentBuffer = DispatchData.PositionTracker.GetTangentBuffer();
 	DispatchData.IntermediateTangentBuffer = DispatchData.PositionTracker.GetIntermediateTangentBuffer();
+	DispatchData.IntermediateAccumulatedTangentBuffer = DispatchData.PositionTracker.GetIntermediateAccumulatedTangentBuffer();
 
 	if (OverlappedUAVs && DispatchData.DispatchFlags != 0 && DispatchData.GetActiveTangentRWBuffer())
 	{
@@ -1941,15 +2047,6 @@ void FGPUSkinCache::DispatchUpdateSkinning(FRHICommandListImmediate& RHICmdList,
 			DispatchData.GetActiveTangentRWBuffer() ? DispatchData.GetActiveTangentRWBuffer()->UAV : nullptr
 			);
 
-		RHICmdList.Transition(FRHITransitionInfo(DispatchData.GetPreviousPositionRWBuffer()->UAV.GetReference(), ERHIAccess::Unknown, ERHIAccess::UAVCompute));
-		AddBufferToTransition(DispatchData.GetPreviousPositionRWBuffer()->UAV);
-
-		if (DispatchData.GetActiveTangentRWBuffer())
-		{
-			RHICmdList.Transition(FRHITransitionInfo(DispatchData.GetActiveTangentRWBuffer()->UAV.GetReference(), ERHIAccess::Unknown, ERHIAccess::UAVCompute));
-			AddBufferToTransition(DispatchData.GetActiveTangentRWBuffer()->UAV);
-		}
-
 		uint32 VertexCountAlign64 = FMath::DivideAndRoundUp(DispatchData.NumVertices, (uint32)64);
 		INC_DWORD_STAT_BY(STAT_GPUSkinCache_TotalNumVertices, VertexCountAlign64 * 64);
 		RHICmdList.DispatchComputeShader(VertexCountAlign64, 1, 1);
@@ -1970,15 +2067,6 @@ void FGPUSkinCache::DispatchUpdateSkinning(FRHICommandListImmediate& RHICmdList,
 			DispatchData.GetActiveTangentRWBuffer() ? DispatchData.GetActiveTangentRWBuffer()->UAV : nullptr
 			);
 
-		RHICmdList.Transition(FRHITransitionInfo(DispatchData.GetPositionRWBuffer()->UAV.GetReference(), ERHIAccess::Unknown, ERHIAccess::UAVCompute));
-		AddBufferToTransition(DispatchData.GetPositionRWBuffer()->UAV);
-
-		if (DispatchData.GetActiveTangentRWBuffer())
-		{
-			RHICmdList.Transition(FRHITransitionInfo(DispatchData.GetActiveTangentRWBuffer()->UAV.GetReference(), ERHIAccess::Unknown, ERHIAccess::UAVCompute));
-			AddBufferToTransition(DispatchData.GetActiveTangentRWBuffer()->UAV);
-		}
-
 		uint32 VertexCountAlign64 = FMath::DivideAndRoundUp(DispatchData.NumVertices, (uint32)64);
 		INC_DWORD_STAT_BY(STAT_GPUSkinCache_TotalNumVertices, VertexCountAlign64 * 64);
 		RHICmdList.DispatchComputeShader(VertexCountAlign64, 1, 1);
@@ -2028,7 +2116,7 @@ void FGPUSkinCache::ReleaseSkinCacheEntry(FGPUSkinCacheEntry* SkinCacheEntry)
 		DEC_MEMORY_STAT_BY(STAT_GPUSkinCache_TotalMemUsed, RequiredMemInBytes);
 
 		SkinCache->Allocations.Remove(PositionAllocation);
-		PositionAllocation->RemoveAllFromTransitionArray(SkinCache->BuffersToTransition);
+		PositionAllocation->RemoveAllFromTransitionArray(SkinCache->BuffersToTransitionToRead);
 
 		delete PositionAllocation;
 
diff --git a/Engine/Source/Runtime/Engine/Public/GPUSkinCache.h b/Engine/Source/Runtime/Engine/Public/GPUSkinCache.h
--- a/Engine/Source/Runtime/Engine/Public/GPUSkinCache.h
+++ b/Engine/Source/Runtime/Engine/Public/GPUSkinCache.h
@@ -221,8 +221,8 @@ public:
 	{
 		friend struct FRWBufferTracker;
 
-		FRWBuffersAllocation(uint32 InNumVertices, bool InWithTangents)
-			: NumVertices(InNumVertices), WithTangents(InWithTangents)
+		FRWBuffersAllocation(uint32 InNumVertices, bool InWithTangents, uint32 InNumTriangles, FRHICommandListImmediate& RHICmdList)
+			: NumVertices(InNumVertices), WithTangents(InWithTangents), NumTriangles(InNumTriangles)
 		{
 			for (int32 Index = 0; Index < NUM_BUFFERS; ++Index)
 			{
@@ -236,6 +236,12 @@ public:
 					IntermediateTangents.Initialize(TEXT("SkinCacheIntermediateTangents"), 8, NumVertices * 2, PF_R16G16B16A16_SNORM, BUF_Static);
 				}
 			}
+			if (NumTriangles > 0)
+			{
+				IntermediateAccumulatedTangents.Initialize(TEXT("SkinCacheIntermediateAccumulatedTangents"), sizeof(int32), NumTriangles * 3 * FGPUSkinCache::IntermediateAccumBufferNumInts, PF_R32_SINT, BUF_UnorderedAccess);
+				// The UAV must be zero-filled. We leave it zeroed after each round (see RecomputeTangentsPerVertexPass.usf), so this is only needed on when the buffer is first created.
+				RHICmdList.ClearUAVUint(IntermediateAccumulatedTangents.UAV, FUintVector4(0, 0, 0, 0));
+			}
 		}
 
 		~FRWBuffersAllocation()
@@ -249,23 +255,28 @@ public:
 				Tangents.Release();
 				IntermediateTangents.Release();
 			}
+			if (NumTriangles > 0)
+			{
+				IntermediateAccumulatedTangents.Release();
+			}
 		}
 
-		static uint64 CalculateRequiredMemory(uint32 NumVertices, bool WithTangents)
+		static uint64 CalculateRequiredMemory(uint32 InNumVertices, bool InWithTangents, uint32 InNumTriangles)
 		{
-			uint64 PositionBufferSize = 4 * 3 * NumVertices * NUM_BUFFERS;
-			uint64 TangentBufferSize = WithTangents ? 2 * 4 * NumVertices : 0;
+			uint64 PositionBufferSize = 4 * 3 * InNumVertices * NUM_BUFFERS;
+			uint64 TangentBufferSize = InWithTangents ? 2 * 4 * InNumVertices : 0;
 			uint64 IntermediateTangentBufferSize = 0;
 			if (FGPUSkinCache::UseIntermediateTangents())
 			{
-				IntermediateTangentBufferSize = WithTangents ? 2 * 4 * NumVertices : 0;
+				IntermediateTangentBufferSize = InWithTangents ? 2 * 4 * InNumVertices : 0;
 			}
-			return TangentBufferSize + IntermediateTangentBufferSize + PositionBufferSize;
+			uint64 AccumulatedTangentBufferSize = InNumTriangles * 3 * FGPUSkinCache::IntermediateAccumBufferNumInts * sizeof(int32);
+			return TangentBufferSize + IntermediateTangentBufferSize + PositionBufferSize + AccumulatedTangentBufferSize;
 		}
 
 		uint64 GetNumBytes() const
 		{
-			return CalculateRequiredMemory(NumVertices, WithTangents);
+			return CalculateRequiredMemory(NumVertices, WithTangents, NumTriangles);
 		}
 
 		FRWBuffer* GetTangentBuffer()
@@ -278,6 +289,11 @@ public:
 			return WithTangents ? &IntermediateTangents : nullptr;
 		}
 
+		FRWBuffer* GetIntermediateAccumulatedTangentBuffer()
+		{
+			return NumTriangles > 0 ? &IntermediateAccumulatedTangents : nullptr;
+		}
+
 		void RemoveAllFromTransitionArray(TSet<FRHIUnorderedAccessView*>& BuffersToTransition);
 
 	private:
@@ -286,8 +302,11 @@ public:
 
 		FRWBuffer Tangents;
 		FRWBuffer IntermediateTangents;
+		FRWBuffer IntermediateAccumulatedTangents;	// Intermediate buffer used to accumulate results of triangle pass to be passed onto vertex pass
+
 		const uint32 NumVertices;
 		const bool WithTangents;
+		const uint32 NumTriangles;
 	};
 
 	struct FRWBufferTracker
@@ -329,12 +348,17 @@ public:
 
 		FRWBuffer* GetTangentBuffer()
 		{
-			return Allocation->GetTangentBuffer();
+			return Allocation ? Allocation->GetTangentBuffer() : nullptr;
 		}
 
 		FRWBuffer* GetIntermediateTangentBuffer()
 		{
-			return Allocation->GetIntermediateTangentBuffer();
+			return Allocation ? Allocation->GetIntermediateTangentBuffer() : nullptr;
+		}
+
+		FRWBuffer* GetIntermediateAccumulatedTangentBuffer()
+		{
+			return Allocation ? Allocation->GetIntermediateAccumulatedTangentBuffer() : nullptr;
 		}
 
 		void Advance(const FVertexBufferAndSRV& BoneBuffer1, uint32 Revision1, const FVertexBufferAndSRV& BoneBuffer2, uint32 Revision2)
@@ -396,10 +420,9 @@ public:
 	inline ERHIFeatureLevel::Type GetFeatureLevel() const { return FeatureLevel; }
 
 protected:
+	void MakeBufferTransitions(FRHICommandListImmediate& RHICmdList, TArray<FRHIUnorderedAccessView*>& Buffers, ERHIAccess FromState, ERHIAccess ToState);
 
-	void AddBufferToTransition(FRHIUnorderedAccessView* InUAV);
-
-	TSet<FRHIUnorderedAccessView*> BuffersToTransition;
+	TSet<FRHIUnorderedAccessView*> BuffersToTransitionToRead;
 #if RHI_RAYTRACING
 	TMap<FRayTracingGeometry*, EAccelerationStructureBuildMode> RayTracingGeometriesToUpdate;
 	uint64 RayTracingGeometryMemoryPendingRelease = 0;
@@ -409,10 +432,10 @@ protected:
 	TArray<FGPUSkinCacheEntry*> Entries;
 	TArray<FDispatchEntry> BatchDispatches;
 
-	FRWBuffersAllocation* TryAllocBuffer(uint32 NumVertices, bool WithTangnents);
+	FRWBuffersAllocation* TryAllocBuffer(uint32 NumVertices, bool WithTangnents, uint32 NumTriangles, FRHICommandListImmediate& RHICmdList);
 	void DoDispatch(FRHICommandListImmediate& RHICmdList);
 	void DoDispatch(FRHICommandListImmediate& RHICmdList, FGPUSkinCacheEntry* SkinCacheEntry, int32 Section, int32 RevisionNumber);
-	void DispatchUpdateSkinTangents(FRHICommandListImmediate& RHICmdList, FGPUSkinCacheEntry* Entry, int32 SectionIndex);
+	void DispatchUpdateSkinTangents(FRHICommandListImmediate& RHICmdList, FGPUSkinCacheEntry* Entry, int32 SectionIndex, FRWBuffer*& StagingBuffer, bool bTrianglePass);
 
 	void PrepareUpdateSkinning(
 		FGPUSkinCacheEntry* Entry, 
