Don't crash when trying to access editor-only cached material data that may not exist
This can happen when cooked data is loaded in editor

diff --git a/Engine/Source/Runtime/Engine/Private/Materials/Material.cpp b/Engine/Source/Runtime/Engine/Private/Materials/Material.cpp
--- a/Engine/Source/Runtime/Engine/Private/Materials/Material.cpp
+++ b/Engine/Source/Runtime/Engine/Private/Materials/Material.cpp
@@ -2267,23 +2267,7 @@ bool UMaterial::IsScalarParameterUsedAsAtlasPosition(const FHashedMaterialParame
 {
 	if (CachedExpressionData)
 	{
-		const int32 Index = CachedExpressionData->Parameters.FindParameterIndex(EMaterialParameterType::Scalar, ParameterInfo);
-		if (Index != INDEX_NONE)
-		{
-			UCurveLinearColor* Curve = CachedExpressionData->Parameters.ScalarCurveValues[Index];
-			UCurveLinearColorAtlas* Atlas = CachedExpressionData->Parameters.ScalarCurveAtlasValues[Index];
-			if (Curve && Atlas)
-			{
-				OutCurve = TSoftObjectPtr<UCurveLinearColor>(FSoftObjectPath(Curve->GetPathName()));
-				OutAtlas = TSoftObjectPtr<UCurveLinearColorAtlas>(FSoftObjectPath(Atlas->GetPathName()));
-				OutValue = true;
-			}
-			else
-			{
-				OutValue = false;
-			}
-			return true;
-		}
+		return CachedExpressionData->Parameters.IsScalarParameterUsedAsAtlasPosition(ParameterInfo, OutValue, OutCurve, OutAtlas);
 	}
 
 	return false;
@@ -2293,13 +2277,7 @@ bool UMaterial::GetScalarParameterSliderMinMax(const FHashedMaterialParameterInf
 {
 	if (CachedExpressionData)
 	{
-		const int32 Index = CachedExpressionData->Parameters.FindParameterIndex(EMaterialParameterType::Scalar, ParameterInfo);
-		if (Index != INDEX_NONE)
-		{
-			OutSliderMin = CachedExpressionData->Parameters.ScalarMinMaxValues[Index].X;
-			OutSliderMax = CachedExpressionData->Parameters.ScalarMinMaxValues[Index].Y;
-			return true;
-		}
+		return CachedExpressionData->Parameters.GetScalarParameterSliderMinMax(ParameterInfo, OutSliderMin, OutSliderMax);
 	}
 
 	return false;
@@ -2411,12 +2389,7 @@ bool UMaterial::IsVectorParameterUsedAsChannelMask(const FHashedMaterialParamete
 {
 	if (CachedExpressionData)
 	{
-		const int32 Index = CachedExpressionData->Parameters.FindParameterIndex(EMaterialParameterType::Vector, ParameterInfo);
-		if (Index != INDEX_NONE)
-		{
-			OutValue = CachedExpressionData->Parameters.VectorUsedAsChannelMaskValues[Index];
-			return true;
-		}
+		return CachedExpressionData->Parameters.IsVectorParameterUsedAsChannelMask(ParameterInfo, OutValue);
 	}
 	return false;
 }
@@ -2425,12 +2398,7 @@ bool UMaterial::GetVectorParameterChannelNames(const FHashedMaterialParameterInf
 {
 	if (CachedExpressionData)
 	{
-		const int32 Index = CachedExpressionData->Parameters.FindParameterIndex(EMaterialParameterType::Vector, ParameterInfo);
-		if (Index != INDEX_NONE)
-		{
-			OutValue = CachedExpressionData->Parameters.VectorChannelNameValues[Index];
-			return true;
-		}
+		return CachedExpressionData->Parameters.GetVectorParameterChannelNames(ParameterInfo, OutValue);
 	}
 	return false;
 }
@@ -2661,12 +2629,7 @@ bool UMaterial::GetTextureParameterChannelNames(const FHashedMaterialParameterIn
 {
 	if (CachedExpressionData)
 	{
-		const int32 Index = CachedExpressionData->Parameters.FindParameterIndex(EMaterialParameterType::Texture, ParameterInfo);
-		if (Index != INDEX_NONE)
-		{
-			OutValue = CachedExpressionData->Parameters.TextureChannelNameValues[Index];
-			return true;
-		}
+		return CachedExpressionData->Parameters.GetTextureParameterChannelNames(ParameterInfo, OutValue);
 	}
 	return false;
 }
diff --git a/Engine/Source/Runtime/Engine/Private/Materials/MaterialCachedData.cpp b/Engine/Source/Runtime/Engine/Private/Materials/MaterialCachedData.cpp
--- a/Engine/Source/Runtime/Engine/Private/Materials/MaterialCachedData.cpp
+++ b/Engine/Source/Runtime/Engine/Private/Materials/MaterialCachedData.cpp
@@ -757,6 +757,98 @@ void FMaterialCachedParameters::GetAllGlobalParameterInfoOfType(EMaterialParamet
 	}
 }
 
+bool FMaterialCachedParameters::GetScalarParameterSliderMinMax(const FMemoryImageMaterialParameterInfo& ParameterInfo, float& OutSliderMin, float& OutSliderMax) const
+{
+#if WITH_EDITORONLY_DATA
+	if (ScalarMinMaxValues.Num() != 0)
+	{
+		const int32 ParameterIndex = FindParameterIndex(EMaterialParameterType::Scalar, ParameterInfo);
+		if (ParameterIndex != INDEX_NONE)
+		{
+			OutSliderMin = ScalarMinMaxValues[ParameterIndex].X;
+			OutSliderMax = ScalarMinMaxValues[ParameterIndex].Y;
+			return true;
+		}
+	}
+#endif // WITH_EDITORONLY_DATA
+	return false;
+}
+
+bool FMaterialCachedParameters::IsScalarParameterUsedAsAtlasPosition(const FMemoryImageMaterialParameterInfo& ParameterInfo, bool& OutValue, TSoftObjectPtr<UCurveLinearColor>& OutCurve, TSoftObjectPtr<UCurveLinearColorAtlas>& OutAtlas) const
+{
+#if WITH_EDITORONLY_DATA
+	if (ScalarCurveValues.Num() != 0)
+	{
+		const int32 ParameterIndex = FindParameterIndex(EMaterialParameterType::Scalar, ParameterInfo);
+		if (ParameterIndex != INDEX_NONE)
+		{
+			UCurveLinearColor* Curve = ScalarCurveValues[ParameterIndex];
+			UCurveLinearColorAtlas* Atlas = ScalarCurveAtlasValues[ParameterIndex];
+			if (Curve && Atlas)
+			{
+				OutCurve = Curve;
+				OutAtlas = Atlas;
+				OutValue = true;
+			}
+			else
+			{
+				OutValue = false;
+			}
+			return true;
+		}
+	}
+#endif // WITH_EDITORONLY_DATA
+	return false;
+}
+
+bool FMaterialCachedParameters::IsVectorParameterUsedAsChannelMask(const FMemoryImageMaterialParameterInfo& ParameterInfo, bool& OutValue) const
+{
+#if WITH_EDITORONLY_DATA
+	if (VectorUsedAsChannelMaskValues.Num() != 0)
+	{
+		const int32 ParameterIndex = FindParameterIndex(EMaterialParameterType::Vector, ParameterInfo);
+		if (ParameterIndex != INDEX_NONE)
+		{
+			OutValue = VectorUsedAsChannelMaskValues[ParameterIndex];
+			return true;
+		}
+	}
+#endif // WITH_EDITORONLY_DATA
+	return false;
+}
+
+bool FMaterialCachedParameters::GetVectorParameterChannelNames(const FMemoryImageMaterialParameterInfo& ParameterInfo, FParameterChannelNames& OutValue) const
+{
+#if WITH_EDITORONLY_DATA
+	if (VectorChannelNameValues.Num() != 0)
+	{
+		const int32 ParameterIndex = FindParameterIndex(EMaterialParameterType::Vector, ParameterInfo);
+		if (ParameterIndex != INDEX_NONE)
+		{
+			OutValue = VectorChannelNameValues[ParameterIndex];
+			return true;
+		}
+	}
+#endif // WITH_EDITORONLY_DATA
+	return false;
+}
+
+bool FMaterialCachedParameters::GetTextureParameterChannelNames(const FMemoryImageMaterialParameterInfo& ParameterInfo, FParameterChannelNames& OutValue) const
+{
+#if WITH_EDITORONLY_DATA
+	if (TextureChannelNameValues.Num() != 0)
+	{
+		const int32 ParameterIndex = FindParameterIndex(EMaterialParameterType::Texture, ParameterInfo);
+		if (ParameterIndex != INDEX_NONE)
+		{
+			OutValue = TextureChannelNameValues[ParameterIndex];
+			return true;
+		}
+	}
+#endif // WITH_EDITORONLY_DATA
+	return false;
+}
+
 void FMaterialInstanceCachedData::Initialize(FMaterialCachedExpressionData&& InCachedExpressionData)
 {
 	Parameters = MoveTemp(InCachedExpressionData.Parameters);
diff --git a/Engine/Source/Runtime/Engine/Private/Materials/MaterialInstance.cpp b/Engine/Source/Runtime/Engine/Private/Materials/MaterialInstance.cpp
--- a/Engine/Source/Runtime/Engine/Private/Materials/MaterialInstance.cpp
+++ b/Engine/Source/Runtime/Engine/Private/Materials/MaterialInstance.cpp
@@ -775,11 +775,8 @@ bool UMaterialInstance::GetScalarParameterSliderMinMax(const FHashedMaterialPara
 
 	if (ParameterInfo.Association != EMaterialParameterAssociation::GlobalParameter && CachedData)
 	{
-		const int32 ParameterIndex = CachedData->Parameters.FindParameterIndex(EMaterialParameterType::Scalar, ParameterInfo);
-		if (ParameterIndex != INDEX_NONE)
+		if (CachedData->Parameters.GetScalarParameterSliderMinMax(ParameterInfo, OutSliderMin, OutSliderMax))
 		{
-			OutSliderMin = CachedData->Parameters.ScalarMinMaxValues[ParameterIndex].X;
-			OutSliderMax = CachedData->Parameters.ScalarMinMaxValues[ParameterIndex].Y;
 			return true;
 		}
 	}
@@ -858,21 +855,8 @@ bool UMaterialInstance::IsScalarParameterUsedAsAtlasPosition(const FHashedMateri
 	// Instance-included default
 	if (ParameterInfo.Association != EMaterialParameterAssociation::GlobalParameter && CachedData)
 	{
-		const int32 ParameterIndex = CachedData->Parameters.FindParameterIndex(EMaterialParameterType::Scalar, ParameterInfo);
-		if (ParameterIndex != INDEX_NONE)
+		if (CachedData->Parameters.IsScalarParameterUsedAsAtlasPosition(ParameterInfo, OutValue, OutCurve, OutAtlas))
 		{
-			UCurveLinearColor* Curve = CachedData->Parameters.ScalarCurveValues[ParameterIndex];
-			UCurveLinearColorAtlas* Atlas = CachedData->Parameters.ScalarCurveAtlasValues[ParameterIndex];
-			if (Curve && Atlas)
-			{
-				OutCurve = TSoftObjectPtr<UCurveLinearColor>(FSoftObjectPath(Curve->GetPathName()));
-				OutAtlas = TSoftObjectPtr<UCurveLinearColorAtlas>(FSoftObjectPath(Atlas->GetPathName()));
-				OutValue = true;
-			}
-			else
-			{
-				OutValue = false;
-			}
 			return true;
 		}
 	}
@@ -940,10 +924,8 @@ bool UMaterialInstance::IsVectorParameterUsedAsChannelMask(const FHashedMaterial
 	// Instance-included default
 	if (ParameterInfo.Association != EMaterialParameterAssociation::GlobalParameter && CachedData)
 	{
-		const int32 ParameterIndex = CachedData->Parameters.FindParameterIndex(EMaterialParameterType::Vector, ParameterInfo);
-		if (ParameterIndex != INDEX_NONE)
+		if (CachedData->Parameters.IsVectorParameterUsedAsChannelMask(ParameterInfo, OutValue))
 		{
-			OutValue = CachedData->Parameters.VectorUsedAsChannelMaskValues[ParameterIndex];
 			return true;
 		}
 	}
@@ -970,10 +952,8 @@ bool  UMaterialInstance::GetVectorParameterChannelNames(const FHashedMaterialPar
 	// Instance-included default
 	if (ParameterInfo.Association != EMaterialParameterAssociation::GlobalParameter && CachedData)
 	{
-		const int32 ParameterIndex = CachedData->Parameters.FindParameterIndex(EMaterialParameterType::Vector, ParameterInfo);
-		if (ParameterIndex != INDEX_NONE)
+		if (CachedData->Parameters.GetVectorParameterChannelNames(ParameterInfo, OutValue))
 		{
-			OutValue = CachedData->Parameters.VectorChannelNameValues[ParameterIndex];
 			return true;
 		}
 	}
@@ -1085,10 +1065,8 @@ bool  UMaterialInstance::GetTextureParameterChannelNames(const FHashedMaterialPa
 	// Instance-included default
 	if (ParameterInfo.Association != EMaterialParameterAssociation::GlobalParameter && CachedData)
 	{
-		const int32 ParameterIndex = CachedData->Parameters.FindParameterIndex(EMaterialParameterType::Texture, ParameterInfo);
-		if (ParameterIndex != INDEX_NONE)
+		if (CachedData->Parameters.GetTextureParameterChannelNames(ParameterInfo, OutValue))
 		{
-			OutValue = CachedData->Parameters.TextureChannelNameValues[ParameterIndex];
 			return true;
 		}
 	}
diff --git a/Engine/Source/Runtime/Engine/Public/MaterialCachedData.h b/Engine/Source/Runtime/Engine/Public/MaterialCachedData.h
--- a/Engine/Source/Runtime/Engine/Public/MaterialCachedData.h
+++ b/Engine/Source/Runtime/Engine/Public/MaterialCachedData.h
@@ -160,6 +160,12 @@ struct FMaterialCachedParameters
 
 	void AddReferencedObjects(FReferenceCollector& Collector);
 
+	bool GetScalarParameterSliderMinMax(const FMemoryImageMaterialParameterInfo& ParameterInfo, float& OutSliderMin, float& OutSliderMax) const;
+	bool IsScalarParameterUsedAsAtlasPosition(const FMemoryImageMaterialParameterInfo& ParameterInfo, bool& OutValue, TSoftObjectPtr<UCurveLinearColor>& OutCurve, TSoftObjectPtr<UCurveLinearColorAtlas>& OutAtlas) const;
+	bool IsVectorParameterUsedAsChannelMask(const FMemoryImageMaterialParameterInfo& ParameterInfo, bool& OutValue) const;
+	bool GetVectorParameterChannelNames(const FMemoryImageMaterialParameterInfo& ParameterInfo, FParameterChannelNames& OutValue) const;
+	bool GetTextureParameterChannelNames(const FMemoryImageMaterialParameterInfo& ParameterInfo, FParameterChannelNames& OutValue) const;
+
 	UPROPERTY()
 	FMaterialCachedParameterEntry RuntimeEntries[NumMaterialRuntimeParameterTypes];
 
