Fix Packed Level Instance bounds (same fix as early access but with versioning)

diff --git a/Engine/Source/Runtime/Core/Public/UObject/UE5MainStreamObjectVersion.h b/Engine/Source/Runtime/Core/Public/UObject/UE5MainStreamObjectVersion.h
--- a/Engine/Source/Runtime/Core/Public/UObject/UE5MainStreamObjectVersion.h
+++ b/Engine/Source/Runtime/Core/Public/UObject/UE5MainStreamObjectVersion.h
@@ -112,6 +112,9 @@ struct CORE_API FUE5MainStreamObjectVersion
 		// Sparse class data serializes its associated structure to allow for BP types to be used
 		SparseClassDataStructSerialization,
 
+		// PackedLevelInstance bounds fix
+		PackedLevelInstanceBoundsFix,
+
         // -----<new versions can be added above this line>-------------------------------------------------
 		VersionPlusOne,
 		LatestVersion = VersionPlusOne - 1
diff --git a/Engine/Source/Runtime/Core/Public/UObject/UE5ReleaseStreamObjectVersion.h b/Engine/Source/Runtime/Core/Public/UObject/UE5ReleaseStreamObjectVersion.h
--- a/Engine/Source/Runtime/Core/Public/UObject/UE5ReleaseStreamObjectVersion.h
+++ b/Engine/Source/Runtime/Core/Public/UObject/UE5ReleaseStreamObjectVersion.h
@@ -39,6 +39,9 @@ struct CORE_API FUE5ReleaseStreamObjectVersion
 		// PackedLevelInstance version
 		PackedLevelInstanceVersion,
 
+		// PackedLevelInstance bounds fix
+		PackedLevelInstanceBoundsFix,
+
 		// -----<new versions can be added above this line>-------------------------------------------------
 		VersionPlusOne,
 		LatestVersion = VersionPlusOne - 1
diff --git a/Engine/Source/Runtime/Engine/Private/LevelInstance/LevelInstanceActor.cpp b/Engine/Source/Runtime/Engine/Private/LevelInstance/LevelInstanceActor.cpp
--- a/Engine/Source/Runtime/Engine/Private/LevelInstance/LevelInstanceActor.cpp
+++ b/Engine/Source/Runtime/Engine/Private/LevelInstance/LevelInstanceActor.cpp
@@ -650,12 +650,17 @@ bool ALevelInstance::IsLockLocation() const
 FBox ALevelInstance::GetComponentsBoundingBox(bool bNonColliding, bool bIncludeFromChildActors) const
 {
 	FBox Box = Super::GetComponentsBoundingBox(bNonColliding, bIncludeFromChildActors);
-	if (ULevelInstanceSubsystem* LevelInstanceSubsystem = GetLevelInstanceSubsystem())
+	
+	// Add Level Bounds
+	if (SupportsLoading())
 	{
-		FBox LevelInstanceBounds;
-		if(LevelInstanceSubsystem->GetLevelInstanceBounds(this, LevelInstanceBounds))
+		if (ULevelInstanceSubsystem* LevelInstanceSubsystem = GetLevelInstanceSubsystem())
 		{
-			Box += LevelInstanceBounds;
+			FBox LevelInstanceBounds;
+			if (LevelInstanceSubsystem->GetLevelInstanceBounds(this, LevelInstanceBounds))
+			{
+				Box += LevelInstanceBounds;
+			}
 		}
 	}
 
diff --git a/Engine/Source/Runtime/Engine/Private/LevelInstance/LevelInstanceSubsystem.cpp b/Engine/Source/Runtime/Engine/Private/LevelInstance/LevelInstanceSubsystem.cpp
--- a/Engine/Source/Runtime/Engine/Private/LevelInstance/LevelInstanceSubsystem.cpp
+++ b/Engine/Source/Runtime/Engine/Private/LevelInstance/LevelInstanceSubsystem.cpp
@@ -482,10 +482,27 @@ bool ULevelInstanceSubsystem::GetLevelInstanceBounds(const ALevelInstance* Level
 	}
 	else if(LevelInstanceActor->IsLevelInstancePathValid())
 	{
-		if (ULevel::GetLevelBoundsFromPackage(FName(*LevelInstanceActor->GetWorldAssetPackage()), OutBounds))
-		{
-			return true;
-		}
+		return GetLevelInstanceBoundsFromPackage(LevelInstanceActor->GetActorTransform(), FName(*LevelInstanceActor->GetWorldAssetPackage()), OutBounds);
+	}
+
+	return false;
+}
+
+bool ULevelInstanceSubsystem::GetLevelInstanceBoundsFromPackage(const FTransform& InstanceTransform, FName LevelPackage, FBox& OutBounds)
+{
+	FBox LevelBounds;
+	if (ULevel::GetLevelBoundsFromPackage(LevelPackage, LevelBounds))
+	{
+		FVector LevelBoundsLocation;
+		FVector BoundsLocation;
+		FVector BoundsExtent;
+		LevelBounds.GetCenterAndExtents(BoundsLocation, BoundsExtent);
+
+		//@todo_ow: This will result in a new BoundsExtent that is larger than it should. To fix this, we would need the Object Oriented BoundingBox of the actor (the BV of the actor without rotation)
+		const FVector BoundsMin = BoundsLocation - BoundsExtent;
+		const FVector BoundsMax = BoundsLocation + BoundsExtent;
+		OutBounds = FBox(BoundsMin, BoundsMax).TransformBy(InstanceTransform);
+		return true;
 	}
 
 	return false;
diff --git a/Engine/Source/Runtime/Engine/Private/WorldPartition/LevelInstance/LevelInstanceActorDesc.cpp b/Engine/Source/Runtime/Engine/Private/WorldPartition/LevelInstance/LevelInstanceActorDesc.cpp
--- a/Engine/Source/Runtime/Engine/Private/WorldPartition/LevelInstance/LevelInstanceActorDesc.cpp
+++ b/Engine/Source/Runtime/Engine/Private/WorldPartition/LevelInstance/LevelInstanceActorDesc.cpp
@@ -7,8 +7,10 @@
 #include "Engine/World.h"
 #include "Engine/Level.h"
 #include "LevelInstance/LevelInstanceActor.h"
+#include "LevelInstance/LevelInstanceSubsystem.h"
 #include "WorldPartition/WorldPartition.h"
 #include "UObject/UE5ReleaseStreamObjectVersion.h"
+#include "UObject/UE5MainStreamObjectVersion.h"
 
 FLevelInstanceActorDesc::FLevelInstanceActorDesc()
 	: DesiredRuntimeBehavior(ELevelInstanceRuntimeBehavior::Embedded)
@@ -69,6 +71,7 @@ void FLevelInstanceActorDesc::Serialize(FArchive& Ar)
 {
 	FWorldPartitionActorDesc::Serialize(Ar);
 	Ar.UsingCustomVersion(FUE5ReleaseStreamObjectVersion::GUID);
+	Ar.UsingCustomVersion(FUE5MainStreamObjectVersion::GUID);
 
 	Ar << LevelPackage << LevelInstanceTransform;
 
@@ -79,19 +82,15 @@ void FLevelInstanceActorDesc::Serialize(FArchive& Ar)
 
 	if (Ar.IsLoading())
 	{
-		if (!LevelPackage.IsNone())
+		const bool bFixupOldVersion = (Ar.CustomVer(FUE5MainStreamObjectVersion::GUID) < FUE5MainStreamObjectVersion::PackedLevelInstanceBoundsFix) && 
+									  (Ar.CustomVer(FUE5ReleaseStreamObjectVersion::GUID) < FUE5ReleaseStreamObjectVersion::PackedLevelInstanceBoundsFix);
+
+		if (!LevelPackage.IsNone() && (GetActorClass()->GetDefaultObject<ALevelInstance>()->SupportsLoading() || bFixupOldVersion))
 		{
-			FBox LevelBounds;
-			if (ULevel::GetLevelBoundsFromPackage(LevelPackage, LevelBounds))
+			FBox OutBounds;
+			if (ULevelInstanceSubsystem::GetLevelInstanceBoundsFromPackage(LevelInstanceTransform, LevelPackage, OutBounds))
 			{
-				FVector LevelBoundsLocation;
-				LevelBounds.GetCenterAndExtents(BoundsLocation, BoundsExtent);
-
-				//@todo_ow: This will result in a new BoundsExtent that is larger than it should. To fix this, we would need the Object Oriented BoundingBox of the actor (the BV of the actor without rotation)
-				const FVector BoundsMin = BoundsLocation - BoundsExtent;
-				const FVector BoundsMax = BoundsLocation + BoundsExtent;
-				const FBox NewBounds = FBox(BoundsMin, BoundsMax).TransformBy(LevelInstanceTransform);
-				NewBounds.GetCenterAndExtents(BoundsLocation, BoundsExtent);
+				OutBounds.GetCenterAndExtents(BoundsLocation, BoundsExtent);
 			}
 		}
 	}
diff --git a/Engine/Source/Runtime/Engine/Public/LevelInstance/LevelInstanceSubsystem.h b/Engine/Source/Runtime/Engine/Public/LevelInstance/LevelInstanceSubsystem.h
--- a/Engine/Source/Runtime/Engine/Public/LevelInstance/LevelInstanceSubsystem.h
+++ b/Engine/Source/Runtime/Engine/Public/LevelInstance/LevelInstanceSubsystem.h
@@ -60,6 +60,7 @@ public:
 	bool IsEditingLevelInstance(const ALevelInstance* LevelInstanceActor) const { return GetLevelInstanceEdit(LevelInstanceActor) != nullptr; }
 	
 	bool GetLevelInstanceBounds(const ALevelInstance* LevelInstanceActor, FBox& OutBounds) const;
+	static bool GetLevelInstanceBoundsFromPackage(const FTransform& InstanceTransform, FName LevelPackage, FBox& OutBounds);
 	
 	void ForEachActorInLevelInstance(const ALevelInstance* LevelInstanceActor, TFunctionRef<bool(AActor * LevelActor)> Operation) const;
 	void ForEachLevelInstanceAncestorsAndSelf(const AActor* Actor, TFunctionRef<bool(const ALevelInstance*)> Operation) const;
