Restoring removed material parameters and hiding them since there's no good way to automatically remove their connections from nodes that reference them via their GUIDs.

diff --git a/Engine/Source/Runtime/Engine/Private/Materials/Material.cpp b/Engine/Source/Runtime/Engine/Private/Materials/Material.cpp
--- a/Engine/Source/Runtime/Engine/Private/Materials/Material.cpp
+++ b/Engine/Source/Runtime/Engine/Private/Materials/Material.cpp
@@ -3371,7 +3371,7 @@ void UMaterial::Serialize(FArchive& Ar)
 	}
 #endif // #if WITH_EDITOR
 
-	static_assert(MP_MAX == 31, "New material properties must have DoMaterialAttributeReorder called on them to ensure that any future reordering of property pins is correctly applied.");
+	static_assert(MP_MAX == 33, "New material properties must have DoMaterialAttributeReorder called on them to ensure that any future reordering of property pins is correctly applied.");
 
 	if (Ar.UEVer() < VER_UE4_MATERIAL_MASKED_BLENDMODE_TIDY)
 	{
diff --git a/Engine/Source/Runtime/Engine/Private/Materials/MaterialExpressions.cpp b/Engine/Source/Runtime/Engine/Private/Materials/MaterialExpressions.cpp
--- a/Engine/Source/Runtime/Engine/Private/Materials/MaterialExpressions.cpp
+++ b/Engine/Source/Runtime/Engine/Private/Materials/MaterialExpressions.cpp
@@ -5728,7 +5728,7 @@ int32 UMaterialExpressionMakeMaterialAttributes::Compile(class FMaterialCompiler
 	int32 Ret = INDEX_NONE;
 	UMaterialExpression* Expression = nullptr;
 
- 	static_assert(MP_MAX == 31, 
+ 	static_assert(MP_MAX == 33, 
 		"New material properties should be added to the end of the inputs for this expression. \
 		The order of properties here should match the material results pins, the make material attriubtes node inputs and the mapping of IO indices to properties in GetMaterialPropertyFromInputOutputIndex().\
 		Insertions into the middle of the properties or a change in the order of properties will also require that existing data is fixed up in DoMaterialAttributeReorder().\
@@ -5835,7 +5835,7 @@ UMaterialExpressionBreakMaterialAttributes::UMaterialExpressionBreakMaterialAttr
 
 	MenuCategories.Add(ConstructorStatics.NAME_MaterialAttributes);
 	
- 	static_assert(MP_MAX == 31, 
+ 	static_assert(MP_MAX == 33, 
 		"New material properties should be added to the end of the outputs for this expression. \
 		The order of properties here should match the material results pins, the make material attriubtes node inputs and the mapping of IO indices to properties in GetMaterialPropertyFromInputOutputIndex().\
 		Insertions into the middle of the properties or a change in the order of properties will also require that existing data is fixed up in DoMaterialAttriubtesReorder().\
diff --git a/Engine/Source/Runtime/Engine/Private/Materials/MaterialShared.cpp b/Engine/Source/Runtime/Engine/Private/Materials/MaterialShared.cpp
--- a/Engine/Source/Runtime/Engine/Private/Materials/MaterialShared.cpp
+++ b/Engine/Source/Runtime/Engine/Private/Materials/MaterialShared.cpp
@@ -4338,6 +4338,10 @@ void FMaterialAttributeDefinitionMap::InitializeAttributeMap()
 	// Debug attributes
 	Add(FGuid(0x5BF6BA94, 0xA3264629, 0xA253A05B, 0x0EABBB86), TEXT("Missing"), MP_MAX, MCT_Float, FVector4(0,0,0,0), SF_Pixel, INDEX_NONE, bHideAttribute);
 
+	// Removed attributes
+	Add(FGuid(0x2091ECA2, 0xB59248EE, 0x8E2CD578, 0xD371926D), TEXT("WorldDisplacement"), MP_WorldDisplacement_DEPRECATED, MCT_Float3, FVector4(0, 0, 0, 0), SF_Vertex, INDEX_NONE, bHideAttribute);
+	Add(FGuid(0xA0119D44, 0xC456450D, 0x9C39C933, 0x1F72D8D1), TEXT("TessellationMultiplier"), MP_TessellationMultiplier_DEPRECATED, MCT_Float, FVector4(1, 0, 0, 0), SF_Vertex, INDEX_NONE, bHideAttribute);
+
 	// UMaterialExpression custom outputs
 	AddCustomAttribute(FGuid(0xfbd7b46e, 0xb1234824, 0xbde76b23, 0x609f984c), "BentNormal", "GetBentNormal", MCT_Float3, FVector4(0, 0, 1, 0));
 	AddCustomAttribute(FGuid(0xAA3D5C04, 0x16294716, 0xBBDEC869, 0x6A27DD72), "ClearCoatBottomNormal", "ClearCoatBottomNormal", MCT_Float3, FVector4(0, 0, 1, 0));
@@ -4425,6 +4429,10 @@ FText FMaterialAttributeDefinitionMap::GetAttributeOverrideForMaterial(const FGu
 		return LOCTEXT("Tangent", "Tangent");
 	case MP_WorldPositionOffset:
 		return Material->IsUIMaterial() ? LOCTEXT("ScreenPosition", "Screen Position") : LOCTEXT("WorldPositionOffset", "World Position Offset");
+	case MP_WorldDisplacement_DEPRECATED:
+		return LOCTEXT("WorldDisplacement", "World Displacement");
+	case MP_TessellationMultiplier_DEPRECATED:
+		return LOCTEXT("TessellationMultiplier", "Tessellation Multiplier");
 	case MP_SubsurfaceColor:
 		if (Material->MaterialDomain == MD_Volume)
 		{
diff --git a/Engine/Source/Runtime/Engine/Public/SceneTypes.h b/Engine/Source/Runtime/Engine/Public/SceneTypes.h
--- a/Engine/Source/Runtime/Engine/Public/SceneTypes.h
+++ b/Engine/Source/Runtime/Engine/Public/SceneTypes.h
@@ -164,6 +164,8 @@ enum EMaterialProperty
 	MP_Normal UMETA(DisplayName = "Normal"),
 	MP_Tangent UMETA(DisplayName = "Tangent"),
 	MP_WorldPositionOffset UMETA(Hidden),
+	MP_WorldDisplacement_DEPRECATED UMETA(Hidden),
+	MP_TessellationMultiplier_DEPRECATED UMETA(Hidden),
 	MP_SubsurfaceColor UMETA(DisplayName = "Subsurface"),
 	MP_CustomData0 UMETA(Hidden),
 	MP_CustomData1 UMETA(Hidden),
