Fixed PIE creating WorldPartitionActorDesc for child actors which could lead to unexpected behaviors when building runtime grid. One example was actor clustering clearing DataLayers for clusters with child actors.

diff --git a/Engine/Source/Runtime/Engine/Private/WorldPartition/WorldPartitionRuntimeHash.cpp b/Engine/Source/Runtime/Engine/Private/WorldPartition/WorldPartitionRuntimeHash.cpp
--- a/Engine/Source/Runtime/Engine/Private/WorldPartition/WorldPartitionRuntimeHash.cpp
+++ b/Engine/Source/Runtime/Engine/Private/WorldPartition/WorldPartitionRuntimeHash.cpp
@@ -118,7 +118,8 @@ void UWorldPartitionRuntimeHash::CreateActorDescViewMap(const UActorDescContaine
 				if (Actor && 
 					Actor->IsPackageExternal() && 
 					!Actor->IsPendingKill() && 
-					!Container->GetActorDesc(Actor->GetActorGuid())) // Actor not on disk yet so not found in container
+					!Container->GetActorDesc(Actor->GetActorGuid()) && // Actor not on disk yet so not found in container
+					Actor->IsMainPackageActor())
 				{
 					FWorldPartitionActorDesc* ActorDesc = ModifiedActorDescListForPIE.AddActor(Actor);
 					OutActorDescViewMap.Emplace(ActorDesc->GetGuid(), ActorDesc);
