Add missing bOverride variables to copy constructor.  Fixes random values creaping in on M1 macs under Rosetta.

diff --git a/Engine/Source/Runtime/Engine/Private/Scene.cpp b/Engine/Source/Runtime/Engine/Private/Scene.cpp
--- a/Engine/Source/Runtime/Engine/Private/Scene.cpp
+++ b/Engine/Source/Runtime/Engine/Private/Scene.cpp
@@ -738,6 +738,8 @@ FPostProcessSettings::FPostProcessSettings(const FPostProcessSettings& Settings)
 	, bOverride_MotionBlurMax(Settings.bOverride_MotionBlurMax)
 	, bOverride_MotionBlurPerObjectSize(Settings.bOverride_MotionBlurPerObjectSize)
 	, bOverride_ScreenPercentage(Settings.bOverride_ScreenPercentage)
+	, bOverride_ReflectionMethod(Settings.bOverride_ReflectionMethod)
+	, bOverride_LumenReflectionQuality(Settings.bOverride_LumenReflectionQuality)
 	, bOverride_ScreenSpaceReflectionIntensity(Settings.bOverride_ScreenSpaceReflectionIntensity)
 	, bOverride_ScreenSpaceReflectionQuality(Settings.bOverride_ScreenSpaceReflectionQuality)
 	, bOverride_ScreenSpaceReflectionMaxRoughness(Settings.bOverride_ScreenSpaceReflectionMaxRoughness)
