Wait until all static meshes are built during PIE to limit backward compatibility problems (i.e. blueprint/scripts)

diff --git a/Engine/Source/Runtime/Engine/Private/StaticMeshCompiler.cpp b/Engine/Source/Runtime/Engine/Private/StaticMeshCompiler.cpp
--- a/Engine/Source/Runtime/Engine/Private/StaticMeshCompiler.cpp
+++ b/Engine/Source/Runtime/Engine/Private/StaticMeshCompiler.cpp
@@ -42,7 +42,7 @@ static AsyncCompilationHelpers::FAsyncCompilationStandardCVars CVarAsyncStaticMe
 
 static TAutoConsoleVariable<int32> CVarAsyncStaticMeshPlayInEditorMode(
 	TEXT("Editor.AsyncStaticMeshPlayInEditorMode"),
-	2,
+	0,
 	TEXT("0 - Wait until all static meshes are built before entering PIE. (Slowest but causes no visual or behavior artifacts.) \n")
 	TEXT("1 - Wait until all static meshes affecting navigation and physics are built before entering PIE. (Some visuals might be missing during compilation.)\n")
 	TEXT("2 - Wait only on static meshes affecting navigation and physics when they are close to the player. (Fastest while still preventing falling through the floor and going through objects.)\n"),
diff --git a/Engine/Source/Runtime/Engine/Private/UnrealEngine.cpp b/Engine/Source/Runtime/Engine/Private/UnrealEngine.cpp
--- a/Engine/Source/Runtime/Engine/Private/UnrealEngine.cpp
+++ b/Engine/Source/Runtime/Engine/Private/UnrealEngine.cpp
@@ -155,6 +155,7 @@ UnrealEngine.cpp: Implements the UEngine class and helpers.
 #include "LandscapeSubsystem.h"
 #include "TextureCompiler.h"
 #include "StaticMeshCompiler.h"
+#include "AssetCompilingManager.h"
 #endif
 // @todo this is here only due to circular dependency to AIModule. To be removed
 
@@ -13345,6 +13346,11 @@ bool UEngine::LoadMap( FWorldContext& WorldContext, FURL URL, class UPendingNetG
 		}
 	}
 
+#if WITH_EDITOR
+	// Gives a chance to any assets being used for PIE/game to complete
+	FAssetCompilingManager::Get().ProcessAsyncTasks();
+#endif
+
 	// Note that AI system will be created only if ai-system-creation conditions are met
 	WorldContext.World()->CreateAISystem();
 
diff --git a/Engine/Source/Runtime/Engine/Private/World.cpp b/Engine/Source/Runtime/Engine/Private/World.cpp
--- a/Engine/Source/Runtime/Engine/Private/World.cpp
+++ b/Engine/Source/Runtime/Engine/Private/World.cpp
@@ -2637,8 +2637,8 @@ void UWorld::AddToWorld( ULevel* Level, const FTransform& LevelTransform, bool b
 		bExecuteNextStep = (!bConsiderTimeLimit || !IsTimeLimitExceeded( TEXT("shifting actors"), StartTime, Level, TimeLimit ));
 	}
 
-	// Wait on any async DDC handles
 #if WITH_EDITOR
+	// Wait on any async DDC handles
 	if (bExecuteNextStep && AsyncPreRegisterDDCRequests.Num())
 	{
 		if (!bConsiderTimeLimit)
@@ -2667,6 +2667,9 @@ void UWorld::AddToWorld( ULevel* Level, const FTransform& LevelTransform, bool b
 			}
 		}
 	}
+
+	// Gives a chance to any assets being used for PIE/game to complete
+	FAssetCompilingManager::Get().ProcessAsyncTasks();
 #endif
 
 	// Updates the level components (Actor components and UModelComponents).
