Virtual Texture - Downgraded errors fetching from DDC to display & logs
Failures to read from the DDC shouldn't be reported as errors

diff --git a/Engine/Source/Runtime/Engine/Private/TextureDerivedDataTask.cpp b/Engine/Source/Runtime/Engine/Private/TextureDerivedDataTask.cpp
--- a/Engine/Source/Runtime/Engine/Private/TextureDerivedDataTask.cpp
+++ b/Engine/Source/Runtime/Engine/Private/TextureDerivedDataTask.cpp
@@ -563,7 +563,7 @@ void FTextureCacheDerivedDataWorker::DoWork()
 			bSucceeded = DerivedData->VTData->ValidateData(Texture.GetPathName(), false);
 			if (!bSucceeded)
 			{
-				UE_LOG(LogTexture, Error, TEXT("Texture %s has corrupt Virtual Texture compression. The texture will be rebuilt."), *Texture.GetFullName());
+				UE_LOG(LogTexture, Display, TEXT("Texture %s has invalid cached VT data. The texture will be rebuilt."), *Texture.GetFullName());
 				bInvalidVirtualTextureCompression = true;
 			}
 		}
diff --git a/Engine/Source/Runtime/Engine/Private/VT/VirtualTextureBuiltData.cpp b/Engine/Source/Runtime/Engine/Private/VT/VirtualTextureBuiltData.cpp
--- a/Engine/Source/Runtime/Engine/Private/VT/VirtualTextureBuiltData.cpp
+++ b/Engine/Source/Runtime/Engine/Private/VT/VirtualTextureBuiltData.cpp
@@ -227,7 +227,7 @@ bool FVirtualTextureBuiltData::ValidateData(FStringView const& InDDCDebugContext
 			const bool bDDCResult = GetDerivedDataCacheRef().GetSynchronous(*Chunk.DerivedDataKey, ChunkDataDDC, InDDCDebugContext);
 			if (!bDDCResult)
 			{
-				UE_LOG(LogTexture, Error, TEXT("Virtual Texture %s failed to retrieve DDC data (%s) for chunk %d"), *TextureName, *Chunk.DerivedDataKey, ChunkIndex);
+				UE_LOG(LogTexture, Log, TEXT("Virtual Texture %s failed to retrieve DDC data (%s) for chunk %d"), *TextureName, *Chunk.DerivedDataKey, ChunkIndex);
 				bResult = false;
 				break;
 			}
diff --git a/Engine/Source/Runtime/Engine/Private/VT/VirtualTextureChunkDDCCache.cpp b/Engine/Source/Runtime/Engine/Private/VT/VirtualTextureChunkDDCCache.cpp
--- a/Engine/Source/Runtime/Engine/Private/VT/VirtualTextureChunkDDCCache.cpp
+++ b/Engine/Source/Runtime/Engine/Private/VT/VirtualTextureChunkDDCCache.cpp
@@ -109,7 +109,7 @@ public:
 		const bool DDCResult = DDC.GetSynchronous(*Chunk->DerivedDataKey, Results, FinalFilename);
 		if (DDCResult == false)
 		{
-			UE_LOG(LogVTDiskCache, Error, TEXT("Failed to fetch data from DDC (key: %s)"), *Chunk->DerivedDataKey);
+			UE_LOG(LogVTDiskCache, Log, TEXT("Failed to fetch data from DDC (key: %s)"), *Chunk->DerivedDataKey);
 			return;
 		}
 
