Fix the generating VT issue with virtual texture disabled.

diff --git a/Engine/Source/Runtime/Engine/Private/TextureDerivedData.cpp b/Engine/Source/Runtime/Engine/Private/TextureDerivedData.cpp
--- a/Engine/Source/Runtime/Engine/Private/TextureDerivedData.cpp
+++ b/Engine/Source/Runtime/Engine/Private/TextureDerivedData.cpp
@@ -431,8 +431,8 @@ static void GetTextureBuildSettings(
 
 	static const auto CVarVirtualTexturesEnabled = IConsoleManager::Get().FindTConsoleVariableDataInt(TEXT("r.VirtualTextures")); check(CVarVirtualTexturesEnabled);
 	// A ULightMapVirtualTexture2D with multiple layers saved in MapBuildData could be loaded with the r.VirtualTexture disabled, it will generate DDC before we decide to invalidate the light map data, to skip the ensure failure let it generate VT DDC anyway.
-	const bool bUseVirtualTexture = CVarVirtualTexturesEnabled->GetValueOnAnyThread() || ULightMapVirtualTexture2D::StaticClass() == Texture.GetClass();
-	const bool bVirtualTextureStreaming = bUseVirtualTexture && bPlatformSupportsVirtualTextureStreaming && Texture.VirtualTextureStreaming;
+	const bool bUseVirtualTexture = ULightMapVirtualTexture2D::StaticClass() == Texture.GetClass();
+	const bool bVirtualTextureStreaming = bUseVirtualTexture || (CVarVirtualTexturesEnabled->GetValueOnAnyThread() && bPlatformSupportsVirtualTextureStreaming && Texture.VirtualTextureStreaming);
 	const FIntPoint SourceSize = Texture.Source.GetLogicalSize();
 
 	OutBuildSettings.MipGenSettings = MipGenSettings;
