Implements r.Test.OverrideTimeMaterialExpressions

diff --git a/Engine/Source/Runtime/Engine/Private/SceneView.cpp b/Engine/Source/Runtime/Engine/Private/SceneView.cpp
--- a/Engine/Source/Runtime/Engine/Private/SceneView.cpp
+++ b/Engine/Source/Runtime/Engine/Private/SceneView.cpp
@@ -328,6 +328,15 @@ static TAutoConsoleVariable<int32> CVarEnableTemporalUpsample(
 	TEXT(" 1: TemporalAA performs spatial and temporal upscale as screen percentage method (default)."),
 	ECVF_Default);
 
+#if !(UE_BUILD_SHIPPING || UE_BUILD_TEST)
+
+static TAutoConsoleVariable<float> CVarOverrideTimeMaterialExpressions(
+	TEXT("r.Test.OverrideTimeMaterialExpressions"), -1.0f,
+	TEXT("Value to freeze time material expressions with."),
+	ECVF_RenderThreadSafe);
+
+#endif
+
 // Conversion factor used when "r.DefaultFeature.AutoExposure.ExtendDefaultLuminanceRange" 
 FORCEINLINE float LuminanceToEV100(float Luminance)
 {
@@ -2393,8 +2402,6 @@ void FSceneView::SetupCommonViewUniformBufferParameters(
 	ViewUniformShaderParameters.SpecularOverrideParameter = SpecularOverrideParameter;
 	ViewUniformShaderParameters.NormalOverrideParameter = NormalOverrideParameter;
 	ViewUniformShaderParameters.RoughnessOverrideParameter = LocalRoughnessOverrideParameter;
-	ViewUniformShaderParameters.PrevFrameGameTime = Family->CurrentWorldTime - Family->DeltaWorldTime;
-	ViewUniformShaderParameters.PrevFrameRealTime = Family->CurrentRealTime  - Family->DeltaWorldTime;
 	ViewUniformShaderParameters.WorldCameraMovementSinceLastFrame = InViewMatrices.GetViewOrigin() - InPrevViewMatrices.GetViewOrigin();
 	ViewUniformShaderParameters.CullingSign = bReverseCulling ? -1.0f : 1.0f;
 	ViewUniformShaderParameters.NearPlane = InViewMatrices.ComputeNearPlane();
@@ -2445,9 +2452,25 @@ void FSceneView::SetupCommonViewUniformBufferParameters(
 	ViewUniformShaderParameters.UnlitViewmodeMask = !Family->EngineShowFlags.Lighting || Family->EngineShowFlags.PathTracing ? 1 : 0;
 	ViewUniformShaderParameters.OutOfBoundsMask = Family->EngineShowFlags.VisualizeOutOfBoundsPixels ? 1 : 0;
 
-	ViewUniformShaderParameters.GameTime = Family->CurrentWorldTime;
-	ViewUniformShaderParameters.RealTime = Family->CurrentRealTime;
-	ViewUniformShaderParameters.DeltaTime = Family->DeltaWorldTime;
+#if !(UE_BUILD_SHIPPING || UE_BUILD_TEST)
+	float OverrideTimeMaterialExpression = CVarOverrideTimeMaterialExpressions.GetValueOnRenderThread();
+	if (OverrideTimeMaterialExpression >= 0.0f)
+	{
+		ViewUniformShaderParameters.PrevFrameGameTime = OverrideTimeMaterialExpression;
+		ViewUniformShaderParameters.PrevFrameRealTime = OverrideTimeMaterialExpression;
+		ViewUniformShaderParameters.GameTime = OverrideTimeMaterialExpression;
+		ViewUniformShaderParameters.RealTime = OverrideTimeMaterialExpression;
+		ViewUniformShaderParameters.DeltaTime = 0.0f;
+	}
+	else
+#endif
+	{
+		ViewUniformShaderParameters.PrevFrameGameTime = Family->CurrentWorldTime - Family->DeltaWorldTime;
+		ViewUniformShaderParameters.PrevFrameRealTime = Family->CurrentRealTime - Family->DeltaWorldTime;
+		ViewUniformShaderParameters.GameTime = Family->CurrentWorldTime;
+		ViewUniformShaderParameters.RealTime = Family->CurrentRealTime;
+		ViewUniformShaderParameters.DeltaTime = Family->DeltaWorldTime;
+	}
 
 #if !(UE_BUILD_SHIPPING || UE_BUILD_TEST)
 	static FIntPoint LockedCursorPos = CursorPos;
