Anim node data/compiler refactor

diff --git a/Engine/Source/Runtime/Engine/Classes/Animation/AnimBlueprintGeneratedClass.h b/Engine/Source/Runtime/Engine/Classes/Animation/AnimBlueprintGeneratedClass.h
--- a/Engine/Source/Runtime/Engine/Classes/Animation/AnimBlueprintGeneratedClass.h
+++ b/Engine/Source/Runtime/Engine/Classes/Animation/AnimBlueprintGeneratedClass.h
@@ -12,7 +12,7 @@
 #include "Animation/AnimClassInterface.h"
 #include "Animation/AnimNodeBase.h"
 #include "Animation/BlendSpace.h"
-#include "PropertyAccess.h"
+#include "Animation/ExposedValueHandler.h"
 
 #include "AnimBlueprintGeneratedClass.generated.h"
 
@@ -25,6 +25,9 @@ class UAnimStateAliasNode;
 class UAnimStateTransitionNode;
 class UEdGraph;
 class USkeleton;
+struct FAnimSubsystem;
+struct FAnimSubsystemInstance;
+struct FPropertyAccessLibrary;
 
 // Represents the debugging information for a single state within a state machine
 USTRUCT()
@@ -308,6 +311,22 @@ public:
 #endif
 };
 
+// 'Marker' structure for mutable data. This is used as a base struct for mutable data to be inserted into by the anim
+// BP compiler.
+USTRUCT()
+struct ENGINE_API FAnimBlueprintMutableData
+{
+	GENERATED_BODY()
+};
+
+// 'Marker' structure for constant data. This is used as a base struct for constant data to be inserted into by the anim
+// BP compiler if there is no existing archetype sparse class data.
+USTRUCT()
+struct ENGINE_API FAnimBlueprintConstantData
+{
+	GENERATED_BODY()
+};
+
 #if WITH_EDITORONLY_DATA
 namespace EPropertySearchMode
 {
@@ -359,9 +378,11 @@ class ENGINE_API UAnimBlueprintGeneratedClass : public UBlueprintGeneratedClass,
 	UPROPERTY()
 	TArray<FName> SyncGroupNames;
 
-	// The default handler for graph-exposed inputs
+#if WITH_EDITORONLY_DATA
+	// Deprecated - moved to FAnimSubsystem_Base
 	UPROPERTY()
-	TArray<FExposedValueHandler> EvaluateGraphExposedInputs;
+	TArray<FExposedValueHandler> EvaluateGraphExposedInputs_DEPRECATED;
+#endif
 
 	// Indices for any Asset Player found within a specific (named) Anim Layer Graph, or implemented Anim Interface Graph
 	UPROPERTY()
@@ -372,14 +393,28 @@ class ENGINE_API UAnimBlueprintGeneratedClass : public UBlueprintGeneratedClass,
 	TMap<FName, FAnimGraphBlendOptions> GraphBlendOptions;
 
 private:
-	// The library holding the property access data
+	// Constant/folded anim node data
 	UPROPERTY()
-	FPropertyAccessLibrary PropertyAccessLibrary;
+	TArray<FAnimNodeData> AnimNodeData;
 
-	// Any internal blendspaces we host
+	// Map from anim node struct to info about that struct (used to accelerate property name lookups)
 	UPROPERTY()
-	TArray<TObjectPtr<UBlendSpace>> BlendSpaces;
+	TMap<const UScriptStruct*, FAnimNodeStructData> NodeTypeMap;
 
+	// Cached properties used to access 'folded' anim node properties
+	TArray<FProperty*> MutableProperties;
+	TArray<FProperty*> ConstantProperties;
+
+	// Cached properties used to access subsystem properties
+	TArray<FStructProperty*> ConstantSubsystemProperties;
+	TArray<FStructProperty*> MutableSubsystemProperties;
+
+	// Property for the object's mutable data area
+	FStructProperty* MutableNodeDataProperty = nullptr;
+
+	// Pointers to each subsystem, for easier debugging
+	TArray<const FAnimSubsystem*> Subsystems;
+	
 public:
 	// IAnimClassInterface interface
 	virtual const TArray<FBakedAnimationStateMachine>& GetBakedStateMachines() const override { return GetRootClass()->GetBakedStateMachines_Direct(); }
@@ -395,11 +430,10 @@ public:
 	virtual const TArray<FName>& GetSyncGroupNames() const override { return GetRootClass()->GetSyncGroupNames_Direct(); }
 	virtual const TMap<FName, FCachedPoseIndices>& GetOrderedSavedPoseNodeIndicesMap() const override { return GetRootClass()->GetOrderedSavedPoseNodeIndicesMap_Direct(); }
 	virtual int32 GetSyncGroupIndex(FName SyncGroupName) const override { return GetSyncGroupNames().IndexOfByKey(SyncGroupName); }
-	virtual const TArray<FExposedValueHandler>& GetExposedValueHandlers() const { return EvaluateGraphExposedInputs; }
+
 	virtual const TArray<FAnimBlueprintFunction>& GetAnimBlueprintFunctions() const override { return AnimBlueprintFunctions; }
 	virtual const TMap<FName, FGraphAssetPlayerInformation>& GetGraphAssetPlayerInformation() const override { return GetRootClass()->GetGraphAssetPlayerInformation_Direct(); }
 	virtual const TMap<FName, FAnimGraphBlendOptions>& GetGraphBlendOptions() const override { return GetRootClass()->GetGraphBlendOptions_Direct(); }
-	virtual const FPropertyAccessLibrary& GetPropertyAccessLibrary() const override { return GetRootClass()->GetPropertyAccessLibrary_Direct(); }
 
 private:
 	virtual const TArray<FBakedAnimationStateMachine>& GetBakedStateMachines_Direct() const override { return BakedStateMachines; }
@@ -408,8 +442,29 @@ private:
 	virtual const TMap<FName, FCachedPoseIndices>& GetOrderedSavedPoseNodeIndicesMap_Direct() const override { return OrderedSavedPoseIndicesMap; }
 	virtual const TMap<FName, FGraphAssetPlayerInformation>& GetGraphAssetPlayerInformation_Direct() const override { return GraphAssetPlayerInformation; }
 	virtual const TMap<FName, FAnimGraphBlendOptions>& GetGraphBlendOptions_Direct() const override { return GraphBlendOptions; }
-	virtual const FPropertyAccessLibrary& GetPropertyAccessLibrary_Direct() const override { return PropertyAccessLibrary; }
+	
+private:
+	virtual const void* GetConstantNodeValueRaw(int32 InIndex) const override;
+	virtual const void* GetMutableNodeValueRaw(int32 InIndex, const UObject* InObject) const override;
+	virtual const FAnimBlueprintMutableData* GetMutableNodeData(const UObject* InObject) const override;
+	virtual FAnimBlueprintMutableData* GetMutableNodeData(UObject* InObject) const override;
+	virtual const void* GetConstantNodeData() const override;
+	virtual TArrayView<const FAnimNodeData> GetNodeData() const override { return AnimNodeData; }
+
+	virtual int32 GetAnimNodePropertyIndex(const UScriptStruct* InNodeType, FName InPropertyName) const override;
+	virtual int32 GetAnimNodePropertyCount(const UScriptStruct* InNodeType) const override;
+	
+	virtual void ForEachSubsystem(TFunctionRef<EAnimSubsystemEnumeration(const FAnimSubsystemContext&)> InFunction) const override;
+	virtual void ForEachSubsystem(UObject* InObject, TFunctionRef<EAnimSubsystemEnumeration(const FAnimSubsystemInstanceContext&)> InFunction) const override;
+	virtual const FAnimSubsystem* FindSubsystem(UScriptStruct* InSubsystemType) const override;
+
+	// Called internally post-load defaults and by the compiler after compilation is completed 
+	void OnPostLoadDefaults(UObject* Object);
 
+	// Called by the compiler to make sure that data tables are initialized. This is needed to patch the sparse class
+	// data for child anim BP overrides 
+	void InitializeAnimNodeData(UObject* DefaultObject);
+	
 public:
 #if WITH_EDITORONLY_DATA
 	FAnimBlueprintDebugData AnimBlueprintDebugData;
@@ -521,7 +576,7 @@ public:
 	// Gets the remapped property index from the original UAnimGraphNode's GUID. Can be used to index the AnimNodeProperties array.
 	int32 GetNodeIndexFromGuid(FGuid Guid, EPropertySearchMode::Type SearchMode = EPropertySearchMode::OnlyThis);
 
-	const UEdGraphNode* GetVisualNodeFromNodePropertyIndex(int32 PropertyIndex) const;
+	const UEdGraphNode* GetVisualNodeFromNodePropertyIndex(int32 PropertyIndex, EPropertySearchMode::Type SearchMode = EPropertySearchMode::OnlyThis) const;
 #endif
 
 	// Called after Link to patch up references to the nodes in the CDO
@@ -530,6 +585,9 @@ public:
 	// Populates AnimBlueprintFunctions according to the UFunction(s) on this class
 	void GenerateAnimationBlueprintFunctions();
 
+	// Build the properties that we cache for our constant data
+	void BuildConstantProperties();
+
 	// UObject interface
 	virtual void Serialize(FArchive& Ar) override;
 	// End of UObject interface
@@ -557,4 +615,4 @@ NodeType* GetNodeFromPropertyIndex(UObject* AnimInstanceObject, const IAnimClass
 	}
 
 	return NULL;
-}
\ No newline at end of file
+}
diff --git a/Engine/Source/Runtime/Engine/Classes/Animation/AnimClassData.h b/Engine/Source/Runtime/Engine/Classes/Animation/AnimClassData.h
--- a/Engine/Source/Runtime/Engine/Classes/Animation/AnimClassData.h
+++ b/Engine/Source/Runtime/Engine/Classes/Animation/AnimClassData.h
@@ -107,18 +107,10 @@ public:
 	UPROPERTY()
 	TArray<FName> SyncGroupNames;
 
-	// The default handler for graph-exposed inputs
-	UPROPERTY()
-	TArray<FExposedValueHandler> EvaluateGraphExposedInputs;
-
 	// Per layer graph blending options
 	UPROPERTY()
 	TMap<FName, FAnimGraphBlendOptions> GraphBlendOptions;
 
-	// Property access library
-	UPROPERTY()
-	FPropertyAccessLibrary PropertyAccessLibrary;
-
 public:
 	// IAnimClassInterface interface
 	virtual const TArray<FBakedAnimationStateMachine>& GetBakedStateMachines() const override { return BakedStateMachines; }
@@ -135,10 +127,8 @@ public:
 	virtual const TArray<FStructProperty*>& GetInitializationNodeProperties() const override { return ResolvedInitializationNodeProperties; }
 	virtual const TArray<FName>& GetSyncGroupNames() const override { return SyncGroupNames; }
 	virtual int32 GetSyncGroupIndex(FName SyncGroupName) const override { return SyncGroupNames.IndexOfByKey(SyncGroupName); }
-	virtual const TArray<FExposedValueHandler>& GetExposedValueHandlers() const override { return EvaluateGraphExposedInputs; }
 	virtual const TMap<FName, FGraphAssetPlayerInformation>& GetGraphAssetPlayerInformation() const override { return GraphNameAssetPlayers; }
 	virtual const TMap<FName, FAnimGraphBlendOptions>& GetGraphBlendOptions() const override { return GraphBlendOptions; }
-	virtual const FPropertyAccessLibrary& GetPropertyAccessLibrary() const override { return PropertyAccessLibrary; }
 
 private:
 	virtual const TArray<FBakedAnimationStateMachine>& GetBakedStateMachines_Direct() const override { return BakedStateMachines; }
@@ -147,8 +137,21 @@ private:
 	virtual const TMap<FName, FCachedPoseIndices>& GetOrderedSavedPoseNodeIndicesMap_Direct() const override { return OrderedSavedPoseIndicesMap; }
 	virtual const TMap<FName, FGraphAssetPlayerInformation>& GetGraphAssetPlayerInformation_Direct() const override { return GraphNameAssetPlayers; }
 	virtual const TMap<FName, FAnimGraphBlendOptions>& GetGraphBlendOptions_Direct() const override { return GraphBlendOptions; }
-	virtual const FPropertyAccessLibrary& GetPropertyAccessLibrary_Direct() const override { return PropertyAccessLibrary; }
-
+	
+private:
+	virtual const void* GetConstantNodeValueRaw(int32 InIndex) const override { return nullptr; }
+	virtual const void* GetMutableNodeValueRaw(int32 InIndex, const UObject* InObject) const override { return nullptr; }
+	virtual const FAnimBlueprintMutableData* GetMutableNodeData(const UObject* InObject) const override { return nullptr; }
+	virtual FAnimBlueprintMutableData* GetMutableNodeData(UObject* InObject) const override { return nullptr; }
+	virtual const void* GetConstantNodeData() const override { return nullptr; }
+	virtual TArrayView<const FAnimNodeData> GetNodeData() const override;
+	virtual int32 GetAnimNodePropertyIndex(const UScriptStruct* InNodeType, FName InPropertyName) const override { return INDEX_NONE; }
+	virtual int32 GetAnimNodePropertyCount(const UScriptStruct* InNodeType) const override { return 0; }
+	
+	virtual void ForEachSubsystem(TFunctionRef<EAnimSubsystemEnumeration(const FAnimSubsystemContext&)> InFunction) const override {}
+	virtual void ForEachSubsystem(UObject* InObject, TFunctionRef<EAnimSubsystemEnumeration(const FAnimSubsystemInstanceContext&)> InFunction) const override {}
+	virtual const FAnimSubsystem* FindSubsystem(UScriptStruct* InStruct) const override { return nullptr; }
+	
 public:
 	// Resolve TFieldPaths to FStructPropertys, init value handlers
 	void DynamicClassInitialization(UDynamicClass* InDynamicClass);
diff --git a/Engine/Source/Runtime/Engine/Classes/Animation/AnimClassInterface.h b/Engine/Source/Runtime/Engine/Classes/Animation/AnimClassInterface.h
--- a/Engine/Source/Runtime/Engine/Classes/Animation/AnimClassInterface.h
+++ b/Engine/Source/Runtime/Engine/Classes/Animation/AnimClassInterface.h
@@ -5,11 +5,10 @@
 #include "CoreMinimal.h"
 #include "UObject/ObjectMacros.h"
 #include "UObject/Object.h"
-#include "UObject/Class.h"
-#include "Templates/Casts.h"
 #include "UObject/Interface.h"
 #include "Animation/AnimTypes.h"
 #include "Animation/AnimStateMachineTypes.h"
+#include "Animation/AnimSubsystem.h"
 #include "UObject/FieldPath.h"
 
 #include "AnimClassInterface.generated.h"
@@ -17,6 +16,13 @@
 class USkeleton;
 struct FExposedValueHandler;
 struct FPropertyAccessLibrary;
+struct FAnimNodeConstantData;
+class IAnimClassInterface;
+struct FAnimBlueprintConstantData;
+struct FAnimBlueprintMutableData;
+struct FAnimNode_Base;
+struct FAnimNodeData;
+namespace UE { namespace Anim { struct FNodeDataId; } }
 
 /** Describes the input and output of an anim blueprint 'function' */
 USTRUCT()
@@ -130,6 +136,12 @@ struct FAnimGraphBlendOptions
 	{}
 };
 
+// How to proceed when enumerating subsystems
+enum class EAnimSubsystemEnumeration
+{
+	Stop,
+	Continue
+};
 
 UINTERFACE()
 class ENGINE_API UAnimClassInterface : public UInterface
@@ -156,7 +168,8 @@ public:
 	virtual const TArray<FStructProperty*>& GetDynamicResetNodeProperties() const = 0;
 	virtual const TArray<FStructProperty*>& GetStateMachineNodeProperties() const = 0;
 	virtual const TArray<FStructProperty*>& GetInitializationNodeProperties() const = 0;
-	virtual const TArray<FExposedValueHandler>& GetExposedValueHandlers() const = 0;
+	UE_DEPRECATED(5.0, "Please use GetSubsystem<FAnimSubsystem_Base>")
+	virtual const TArray<FExposedValueHandler>& GetExposedValueHandlers() const;
 	virtual const TArray<FName>& GetSyncGroupNames() const = 0;
 	virtual const TMap<FName, FCachedPoseIndices>& GetOrderedSavedPoseNodeIndicesMap() const = 0;
 	virtual const TArray<FAnimBlueprintFunction>& GetAnimBlueprintFunctions() const = 0;
@@ -164,11 +177,35 @@ public:
 	virtual const TMap<FName, FAnimGraphBlendOptions>& GetGraphBlendOptions() const = 0;
 	virtual USkeleton* GetTargetSkeleton() const = 0;
 	virtual int32 GetSyncGroupIndex(FName SyncGroupName) const = 0;
-	virtual const FPropertyAccessLibrary& GetPropertyAccessLibrary() const = 0;
+	UE_DEPRECATED(5.0, "Please use GetSubsystem<FAnimSubsystem_PropertyAccess>")
+	virtual const FPropertyAccessLibrary& GetPropertyAccessLibrary() const;
+
+	// Iterate over each subsystem for this class, supplying both the constant (FAnimSubsystem) and mutable (FAnimSubsystemInstance) data
+	virtual void ForEachSubsystem(TFunctionRef<EAnimSubsystemEnumeration(const FAnimSubsystemContext&)> InFunction) const  = 0;
+	virtual void ForEachSubsystem(UObject* InObject, TFunctionRef<EAnimSubsystemEnumeration(const FAnimSubsystemInstanceContext&)> InFunction) const = 0;
+
+	// Find a subsystem's class-resident data. If no subsystem of the type exists this will return nullptr.
+	// @param	InSubsystemType	The subsystem's type
+	virtual const FAnimSubsystem* FindSubsystem(UScriptStruct* InSubsystemType) const = 0;
+
+	// Get a subsystem's class-resident data. If no subsystem of the type exists this will return nullptr.
+	template<typename SubsystemType>
+	const SubsystemType* FindSubsystem() const
+	{
+		const FAnimSubsystem* Subsystem = FindSubsystem(SubsystemType::StaticStruct());
+		return static_cast<const SubsystemType*>(Subsystem);
+	}
+	
+	// Get a subsystem's class-resident data. If no subsystem of the type exists this will assert.
+	template<typename SubsystemType>
+	const SubsystemType& GetSubsystem() const
+	{
+		const FAnimSubsystem* Subsystem = FindSubsystem(SubsystemType::StaticStruct());
+		check(Subsystem);
+		return static_cast<const SubsystemType&>(*Subsystem);
+	}
 	
 protected:
-	friend class UAnimBlueprintGeneratedClass;
-
 	// These direct accessors are here to allow internal access that doesnt redirect to the root class
 	virtual const TArray<FBakedAnimationStateMachine>& GetBakedStateMachines_Direct() const = 0;
 	virtual const TArray<FAnimNotifyEvent>& GetAnimNotifies_Direct() const = 0;
@@ -176,101 +213,52 @@ protected:
 	virtual const TMap<FName, FCachedPoseIndices>& GetOrderedSavedPoseNodeIndicesMap_Direct() const = 0;
 	virtual const TMap<FName, FGraphAssetPlayerInformation>& GetGraphAssetPlayerInformation_Direct() const = 0;
 	virtual const TMap<FName, FAnimGraphBlendOptions>& GetGraphBlendOptions_Direct() const = 0;
-	virtual const FPropertyAccessLibrary& GetPropertyAccessLibrary_Direct() const = 0;
+	UE_DEPRECATED(5.0, "Please use GetSubsystem<FAnimSubsystem_PropertyAccess>")
+	virtual const FPropertyAccessLibrary& GetPropertyAccessLibrary_Direct() const;
+
+protected:
+	friend class UAnimBlueprintGeneratedClass;
+	friend struct FAnimNodeData;
+	friend struct UE::Anim::FNodeDataId;
+
+	// Access the various constant and mutable values
+	virtual const void* GetConstantNodeValueRaw(int32 InIndex) const = 0;
+	virtual const void* GetMutableNodeValueRaw(int32 InIndex, const UObject* InObject) const = 0;
+
+	// Get the struct that holds the mutable data
+	// @param	InObject	The anim instance object that holds the mutable data
+	virtual const FAnimBlueprintMutableData* GetMutableNodeData(const UObject* InObject) const = 0;
+	virtual FAnimBlueprintMutableData* GetMutableNodeData(UObject* InObject) const = 0;
+
+	// Get the struct that holds the constant data
+	virtual const void* GetConstantNodeData() const = 0;
 
+	// Get the anim node data used for each node's constant/folded data
+	virtual TArrayView<const FAnimNodeData> GetNodeData() const = 0;
+
+	// Get the (editor-only data) index of the property
+	virtual int32 GetAnimNodePropertyIndex(const UScriptStruct* InNodeType, FName InPropertyName) const = 0;
+
+	// Get the number of properties (including editor only properties) that the anim node type has
+	virtual int32 GetAnimNodePropertyCount(const UScriptStruct* InNodeType) const = 0;
+	
 public:
 
 	// Get the root anim class interface (i.e. if this is a derived class).
 	// Some properties that are derived from the compiled anim graph are routed to the 'Root' class
 	// as child classes don't get fully compiled. Instead they just override various asset players leaving the
 	// full compilation up to the base class. 
-	const IAnimClassInterface* GetRootClass() const
-	{
-		auto GetSuperClassInterface = [](const IAnimClassInterface* InClass) -> const IAnimClassInterface*
-		{
-			if(const UClass* ActualClass = GetActualAnimClass(InClass))
-			{
-				return GetFromClass(ActualClass->GetSuperClass());
-			}
-
-			return nullptr;
-		};
-
-		const IAnimClassInterface* RootClass = this;
-		while(const IAnimClassInterface* NextClass = GetSuperClassInterface(RootClass))
-		{
-			RootClass = NextClass;
-		}
-
-		return RootClass;
-	}
+	const IAnimClassInterface* GetRootClass() const;
 
-	static IAnimClassInterface* GetFromClass(UClass* InClass)
-	{
-		if (auto AnimClassInterface = Cast<IAnimClassInterface>(InClass))
-		{
-			return AnimClassInterface;
-		}
-		if (auto DynamicClass = Cast<UDynamicClass>(InClass))
-		{
-			DynamicClass->GetDefaultObject(true);
-			return CastChecked<IAnimClassInterface>(DynamicClass->AnimClassImplementation, ECastCheckedType::NullAllowed);
-		}
-		return nullptr;
-	}
+	static IAnimClassInterface* GetFromClass(UClass* InClass);
 
-	static const IAnimClassInterface* GetFromClass(const UClass* InClass)
-	{
-		if (auto AnimClassInterface = Cast<const IAnimClassInterface>(InClass))
-		{
-			return AnimClassInterface;
-		}
-		if (auto DynamicClass = Cast<const UDynamicClass>(InClass))
-		{
-			DynamicClass->GetDefaultObject(true);
-			return CastChecked<const IAnimClassInterface>(DynamicClass->AnimClassImplementation, ECastCheckedType::NullAllowed);
-		}
-		return nullptr;
-	}
+	static const IAnimClassInterface* GetFromClass(const UClass* InClass);
 
-	static UClass* GetActualAnimClass(IAnimClassInterface* AnimClassInterface)
-	{
-		if (UClass* ActualAnimClass = Cast<UClass>(AnimClassInterface))
-		{
-			return ActualAnimClass;
-		}
-		if (UObject* AsObject = Cast<UObject>(AnimClassInterface))
-		{
-			return Cast<UClass>(AsObject->GetOuter());
-		}
-		return nullptr;
-	}
+	static UClass* GetActualAnimClass(IAnimClassInterface* AnimClassInterface);
 
-	static const UClass* GetActualAnimClass(const IAnimClassInterface* AnimClassInterface)
-	{
-		if (const UClass* ActualAnimClass = Cast<const UClass>(AnimClassInterface))
-		{
-			return ActualAnimClass;
-		}
-		if (const UObject* AsObject = Cast<const UObject>(AnimClassInterface))
-		{
-			return Cast<const UClass>(AsObject->GetOuter());
-		}
-		return nullptr;
-	}
+	static const UClass* GetActualAnimClass(const IAnimClassInterface* AnimClassInterface);
 
-	static const FAnimBlueprintFunction* FindAnimBlueprintFunction(IAnimClassInterface* AnimClassInterface, const FName& InFunctionName)
-	{
-		for(const FAnimBlueprintFunction& Function : AnimClassInterface->GetAnimBlueprintFunctions())
-		{
-			if(Function.Name == InFunctionName)
-			{
-				return &Function;
-			}
-		}
-
-		return nullptr;
-	}
+	static const FAnimBlueprintFunction* FindAnimBlueprintFunction(IAnimClassInterface* AnimClassInterface, const FName& InFunctionName);
 
 	/**
 	 * Check if a function is an anim function on this class
@@ -278,21 +266,11 @@ public:
 	 * @param	InFunction				The function to check
 	 * @return true if the supplied function is an anim function on the specified class
 	 */
-	static bool IsAnimBlueprintFunction(IAnimClassInterface* InAnimClassInterface, const UFunction* InFunction)
-	{
-		if(InFunction->GetOuterUClass() == GetActualAnimClass(InAnimClassInterface))
-		{
-			for(const FAnimBlueprintFunction& Function : InAnimClassInterface->GetAnimBlueprintFunctions())
-			{
-				if(Function.Name == InFunction->GetFName())
-				{
-					return true;
-				}
-			}
-		}
-		return false;
-	}
+	static bool IsAnimBlueprintFunction(IAnimClassInterface* InAnimClassInterface, const UFunction* InFunction);
 
+	// Get the object ptr given an anim node
+	static const UObject* GetObjectPtrFromAnimNode(IAnimClassInterface* InAnimClassInterface, const FAnimNode_Base* InNode);
+	
 	UE_DEPRECATED(4.23, "Please use GetAnimBlueprintFunctions()")
 	virtual int32 GetRootAnimNodeIndex() const { return INDEX_NONE; }
 
diff --git a/Engine/Source/Runtime/Engine/Classes/Animation/AnimInstance.h b/Engine/Source/Runtime/Engine/Classes/Animation/AnimInstance.h
--- a/Engine/Source/Runtime/Engine/Classes/Animation/AnimInstance.h
+++ b/Engine/Source/Runtime/Engine/Classes/Animation/AnimInstance.h
@@ -1015,7 +1015,7 @@ public:
 	float GetRelevantAnimTimeFraction(int32 MachineIndex, int32 StateIndex);
 
 	/** Gets the runtime instance of the specified state machine by Name */
-	FAnimNode_StateMachine* GetStateMachineInstanceFromName(FName MachineName);
+	const FAnimNode_StateMachine* GetStateMachineInstanceFromName(FName MachineName) const;
 
 	/** Get the machine description for the specified instance. Does not rely on PRIVATE_MachineDescription being initialized */
 	const FBakedAnimationStateMachine* GetMachineDescription(IAnimClassInterface* AnimBlueprintClass, FAnimNode_StateMachine* MachineInstance);
@@ -1026,25 +1026,28 @@ public:
 	int32 GetSyncGroupIndexFromName(FName SyncGroupName) const;
 
 	/** Gets the index of the state machine matching MachineName */
-	int32 GetStateMachineIndex(FName MachineName);
+	int32 GetStateMachineIndex(FName MachineName) const;
 
 	/** Gets the runtime instance of the specified state machine */
-	FAnimNode_StateMachine* GetStateMachineInstance(int32 MachineIndex);
+	const FAnimNode_StateMachine* GetStateMachineInstance(int32 MachineIndex) const;
 
 	/** 
 	 * Get the index of the specified instance asset player. Useful to pass to GetInstanceAssetPlayerLength (etc.).
 	 * Passing NAME_None to InstanceName will return the first (assumed only) player instance index found.
 	 */
-	int32 GetInstanceAssetPlayerIndex(FName MachineName, FName StateName, FName InstanceName = NAME_None);
+	int32 GetInstanceAssetPlayerIndex(FName MachineName, FName StateName, FName InstanceName = NAME_None) const;
 
 	/** Returns all Animation Nodes of FAnimNode_AssetPlayerBase class within the specified (named) Animation Graph */
-	TArray<FAnimNode_AssetPlayerBase*> GetInstanceAssetPlayers(const FName& GraphName);
+	TArray<const FAnimNode_AssetPlayerBase*> GetInstanceAssetPlayers(const FName& GraphName) const;
+
+	/** Returns all Animation Nodes of FAnimNode_AssetPlayerBase class within the specified (named) Animation Graph */
+	TArray<FAnimNode_AssetPlayerBase*> GetMutableInstanceAssetPlayers(const FName& GraphName);
 
 	/** Gets the runtime instance desc of the state machine specified by name */
-	const FBakedAnimationStateMachine* GetStateMachineInstanceDesc(FName MachineName);
+	const FBakedAnimationStateMachine* GetStateMachineInstanceDesc(FName MachineName) const;
 
 	/** Gets the most relevant asset player in a specified state */
-	FAnimNode_AssetPlayerBase* GetRelevantAssetPlayerFromState(int32 MachineIndex, int32 StateIndex);
+	const FAnimNode_AssetPlayerBase* GetRelevantAssetPlayerFromState(int32 MachineIndex, int32 StateIndex) const;
 
 	//////////////////////////////////////////////////////////////////////////
 
diff --git a/Engine/Source/Runtime/Engine/Classes/Animation/AnimNodeBase.h b/Engine/Source/Runtime/Engine/Classes/Animation/AnimNodeBase.h
--- a/Engine/Source/Runtime/Engine/Classes/Animation/AnimNodeBase.h
+++ b/Engine/Source/Runtime/Engine/Classes/Animation/AnimNodeBase.h
@@ -10,13 +10,13 @@
 #include "Animation/AnimTypes.h"
 #include "Animation/AnimCurveTypes.h"
 #include "BonePose.h"
-#include "Logging/TokenizedMessage.h"
 #include "Stats/StatsHierarchical.h"
 #include "Animation/AnimTrace.h"
 #include "Animation/AnimationPoseData.h"
-#include "UObject/FieldPath.h"
 #include "Animation/AttributesRuntime.h"
 #include "Animation/AnimNodeMessages.h"
+#include "Animation/AnimNodeData.h"
+#include "Animation/ExposedValueHandler.h"
 
 // WARNING: This should always be the last include in any file that needs it (except .generated.h)
 #include "UObject/UndefineUPropertyMacros.h"
@@ -33,6 +33,7 @@ struct FAnimInstanceProxy;
 struct FAnimNode_Base;
 class UProperty;
 struct FPropertyAccessLibrary;
+struct FAnimNodeConstantData;
 
 /**
  * DEPRECATED - This system is now supplanted by UE::Anim::FMessageStack
@@ -173,6 +174,9 @@ public:
 	// Note: This can return NULL, so check the result.
 	ENGINE_API IAnimClassInterface* GetAnimClass() const;
 
+	// Get the anim instance associated with the current proxy
+	ENGINE_API UObject* GetAnimInstanceObject() const;
+
 #if WITH_EDITORONLY_DATA
 	// Get the AnimBlueprint associated with this context, if there is one.
 	// Note: This can return NULL, so check the result.
@@ -657,6 +661,11 @@ struct ENGINE_API FPoseLinkBase
 {
 	GENERATED_USTRUCT_BODY()
 
+protected:
+	/** The non serialized node pointer. */
+	FAnimNode_Base* LinkedNode;
+
+public:
 	/** Serialized link ID, used to build the non-serialized pointer map. */
 	UPROPERTY()
 	int32 LinkID;
@@ -675,36 +684,40 @@ struct ENGINE_API FPoseLinkBase
 #endif
 
 protected:
+#if DO_CHECK
 	/** Flag to prevent reentry when dealing with circular trees. */
 	bool bProcessed;
-
-	/** The non serialized node pointer. */
-	struct FAnimNode_Base* LinkedNode;
+#endif
 
 public:
 	FPoseLinkBase()
-		: LinkID(INDEX_NONE)
+		: LinkedNode(nullptr)
+		, LinkID(INDEX_NONE)
 #if WITH_EDITORONLY_DATA
 		, SourceLinkID(INDEX_NONE)
 #endif
+#if DO_CHECK
 		, bProcessed(false)
-		, LinkedNode(NULL)
+#endif
 	{
 	}
 
 	// Interface
 
 	void Initialize(const FAnimationInitializeContext& Context);
-	void CacheBones(const FAnimationCacheBonesContext& Context) ;
+	void CacheBones(const FAnimationCacheBonesContext& Context);
 	void Update(const FAnimationUpdateContext& Context);
 	void GatherDebugData(FNodeDebugData& DebugData);
 
 	/** Try to re-establish the linked node pointer. */
 	void AttemptRelink(const FAnimationBaseContext& Context);
+
 	/** This only used by custom handlers, and it is advanced feature. */
-	void SetLinkNode(struct FAnimNode_Base* NewLinkNode);
+	void SetLinkNode(FAnimNode_Base* NewLinkNode);
+
 	/** This only used when dynamic linking other graphs to this one. */
 	void SetDynamicLinkNode(struct FPoseLinkBase* InPoseLink);
+
 	/** This only used by custom handlers, and it is advanced feature. */
 	FAnimNode_Base* GetLinkNode();
 };
@@ -739,88 +752,6 @@ public:
 	void EvaluateComponentSpace(FComponentSpacePoseContext& Output);
 };
 
-UENUM()
-enum class EPostCopyOperation : uint8
-{
-	None,
-
-	LogicalNegateBool,
-};
-
-USTRUCT()
-struct FExposedValueCopyRecord
-{
-	GENERATED_BODY()
-
-	FExposedValueCopyRecord() = default;
-
-	FExposedValueCopyRecord(int32 InCopyIndex, EPostCopyOperation InPostCopyOperation)
-		: CopyIndex(InCopyIndex)
-		, PostCopyOperation(InPostCopyOperation)
-	{
-	}
-
-	UPROPERTY()
-	int32 CopyIndex = INDEX_NONE;
-
-	UPROPERTY()
-	EPostCopyOperation PostCopyOperation = EPostCopyOperation::None;
-};
-
-// An exposed value updater
-USTRUCT()
-struct ENGINE_API FExposedValueHandler
-{
-	GENERATED_USTRUCT_BODY()
-
-	FExposedValueHandler()
-		: BoundFunction(NAME_None)
-		, Function(nullptr)
-		, ValueHandlerNodeProperty(nullptr)
-		, PropertyAccessLibrary(nullptr)
-		, bInitialized(false)
-	{
-	}
-
-	// The function to call to update associated properties (can be NULL)
-	UPROPERTY()
-	FName BoundFunction;
-
-	// Direct data access to property in anim instance
-	UPROPERTY()
-	TArray<FExposedValueCopyRecord> CopyRecords;
-
-	// function pointer if BoundFunction != NAME_None
-	UPROPERTY()
-	TObjectPtr<UFunction> Function;
-
-	// Node property that this value handler is associated with, when the node
-	// is instantiated from this property the node's ExposedValueHandler will 
-	// point back to this FExposedValueHandler:
-	UPROPERTY()
-	TFieldPath<FStructProperty> ValueHandlerNodeProperty;
-
-	// Cached property access library ptr
-	const FPropertyAccessLibrary* PropertyAccessLibrary;
-
-	// Prevent multiple initialization
-	bool bInitialized;
-
-	// Helper function to bind an array of handlers.
-	// This is called for nativized builds to initialize against a dynamic class
-	static void DynamicClassInitialization(TArray<FExposedValueHandler>& Handlers, UDynamicClass* InDynamicClass);
-
-	// Helper function to bind an array of handlers.
-	// This is called for non-nativized builds to initialize against a UAnimBlueprintGeneratedClass
-	static void ClassInitialization(TArray<FExposedValueHandler>& Handlers, UObject* ClassDefaultObject);
-
-	// Bind copy records and cache UFunction if necessary
-	void Initialize(UClass* InClass, const FPropertyAccessLibrary& InPropertyAccessLibrary);
-
-	// Execute the function and copy records
-	void Execute(const FAnimationBaseContext& Context) const;
-};
-
 /**
  * This is the base of all runtime animation nodes
  *
@@ -831,7 +762,7 @@ struct ENGINE_API FExposedValueHandler
 USTRUCT()
 struct ENGINE_API FAnimNode_Base
 {
-	GENERATED_USTRUCT_BODY()
+	GENERATED_BODY()
 
 	/** 
 	 * Called when the node first runs. If the node is inside a state machine or cached pose branch then this can be called multiple times. 
@@ -872,12 +803,6 @@ struct ENGINE_API FAnimNode_Base
 	 * @param	Output		Output structure to write pose or curve data to. Also provides access to relevant data as a context.
 	 */	
 	virtual void EvaluateComponentSpace_AnyThread(FComponentSpacePoseContext& Output);
-	/** 
-	 * If a derived anim node should respond to asset overrides, OverrideAsset should be defined to handle changing the asset 
-	 * This is called during anim blueprint compilation to handle child anim blueprints.
-	 * @param	NewAsset	The new asset that is being set
-	 */
-	virtual void OverrideAsset(class UAnimationAsset* NewAsset) {}
 
 	/**
 	 * Called to gather on-screen debug data. 
@@ -923,7 +848,7 @@ struct ENGINE_API FAnimNode_Base
 
 	/** 
 	 * For nodes that need some kind of initialization that is not dependent on node relevancy 
-	 * (i.e. it is insufficent or inefficent to use Initialize_AnyThread), return true here.
+	 * (i.e. it is insufficient or inefficient to use Initialize_AnyThread), return true here.
 	 * Note that this is called at load on the UAnimInstance CDO to avoid needing to call this at runtime.
 	 */
 	virtual bool NeedsOnInitializeAnimInstance() const { return false; }
@@ -947,19 +872,64 @@ struct ENGINE_API FAnimNode_Base
 	virtual bool WantsSkippedUpdates() const { return false; }
 	UE_DEPRECATED(5.0, "Please use IGraphMessage instead")
 	virtual void OnUpdatesSkipped(TArrayView<const FAnimationUpdateContext*> SkippedUpdateContexts) {}
-
+	UE_DEPRECATED(5.0, "Please use the OverrideAssets API on UAnimGraphNode_Base to opt-in to child anim BP override functionality, or per-node specific asset override calls.")
+	virtual void OverrideAsset(class UAnimationAsset* NewAsset) {}
+	
 	// The default handler for graph-exposed inputs:
-	const FExposedValueHandler& GetEvaluateGraphExposedInputs();
+	const FExposedValueHandler& GetEvaluateGraphExposedInputs() const;
 
 	// Initialization function for the default handler for graph-exposed inputs, used only by instancing code:
-	void SetExposedValueHandler(const FExposedValueHandler* Handler) 
-	{ 
-		ExposedValueHandler = Handler; 
+	UE_DEPRECATED(5.0, "Exposed value handlers are now accessed via FAnimNodeConstantData")
+	void SetExposedValueHandler(const FExposedValueHandler* Handler) { }
+
+	// Get this node's index. The node index provides a unique key into its location within the class data
+	int32 GetNodeIndex() const
+	{
+		check(NodeData);
+		return NodeData->GetNodeIndex();
 	}
 
+protected:
+	// Get anim node constant/folded data of the specified type given the identifier. Do not use directly - use GET_ANIM_NODE_DATA
+	template<typename DataType>
+	const DataType& GetData(UE::Anim::FNodeDataId InId, const UObject* InObject = nullptr) const
+	{
+#if WITH_EDITORONLY_DATA
+		if(NodeData)
+		{
+			return *static_cast<const DataType*>(NodeData->GetData(InId, this, InObject));
+		}
+		else
+		{
+			return *InId.GetProperty()->ContainerPtrToValuePtr<const DataType>(this);
+		}
+#else
+		check(NodeData);
+		return *static_cast<const DataType*>(NodeData->GetData(InId, this, InObject));
+#endif
+	}
+
+#if WITH_EDITORONLY_DATA
+	// Get anim node constant/folded data of the specified type given the identifier. Do not use directly - use GET_MUTABLE_ANIM_NODE_DATA
+	template<typename DataType>
+	DataType& GetMutableData(UE::Anim::FNodeDataId InId, UObject* InObject = nullptr)
+	{
+		if(NodeData)
+		{
+			return *static_cast<DataType*>(NodeData->GetMutableData(InId, this, InObject));
+		}
+		else
+		{
+			return *InId.GetProperty()->ContainerPtrToValuePtr<DataType>(this);
+		}
+	}
+#endif
+
 protected:
 	/** return true if enabled, otherwise, return false. This is utility function that can be used per node level */
 	bool IsLODEnabled(FAnimInstanceProxy* AnimInstanceProxy);
+
+	/** Get the LOD threshold at which this node is enabled. Node is enabled if the current LOD >= threshold. */
 	virtual int32 GetLODThreshold() const { return INDEX_NONE; }
 
 	/** Deprecated function */
@@ -971,10 +941,60 @@ protected:
 
 	friend struct FAnimInstanceProxy;
 
-private:		
-	// Reference to the exposed value handler used by this node. Allocated on the class, rather than per instance:
-	const FExposedValueHandler* ExposedValueHandler = nullptr;
+private:
+	friend class IAnimClassInterface;
+	friend class UAnimBlueprintGeneratedClass;
+	friend struct UE::Anim::FNodeDataId;
+
+	// Set the cached ptr to the constant/folded data for this node
+	void SetNodeData(const FAnimNodeData& InNodeData) { NodeData = &InNodeData; }
+
+	// Reference to the constant/folded data for this node
+	const FAnimNodeData* NodeData = nullptr;
 };
 
+#if WITH_EDITORONLY_DATA
+#define ANIM_NODE_VERIFY_MEMBER(Identifier) (this->Identifier);
+#else
+#define ANIM_NODE_VERIFY_MEMBER(Identifier)
+#endif
+
+// Get some (potentially folded) anim node data. Only usable from within an anim node.
+// This caches the node data ID in static contained in a local lambda for improved performance
+#define GET_ANIM_NODE_DATA(Type, Identifier) \
+	( \
+		GetData<Type>( \
+			[this]() -> UE::Anim::FNodeDataId \
+			{ \
+				ANIM_NODE_VERIFY_MEMBER(Identifier); \
+				static UE::Anim::FNodeDataId CachedId_##Identifier; \
+				if(!CachedId_##Identifier.IsValid()) \
+				{ \
+					static const FName NAME_##Identifier(#Identifier); \
+					CachedId_##Identifier = UE::Anim::FNodeDataId(NAME_##Identifier, this, StaticStruct()); \
+				} \
+				return CachedId_##Identifier; \
+			}() \
+		) \
+	)
+
+#if WITH_EDITORONLY_DATA
+#define GET_MUTABLE_ANIM_NODE_DATA(Type, Identifier) \
+	( \
+		GetMutableData<Type>( \
+			[this]() -> UE::Anim::FNodeDataId \
+			{ \
+				ANIM_NODE_VERIFY_MEMBER(Identifier); \
+				static UE::Anim::FNodeDataId CachedId_##Identifier; \
+				if(!CachedId_##Identifier.IsValid()) \
+				{ \
+					static const FName NAME_##Identifier(#Identifier); \
+					CachedId_##Identifier = UE::Anim::FNodeDataId(NAME_##Identifier, this, StaticStruct()); \
+				} \
+				return CachedId_##Identifier; \
+			}() \
+		) \
+	)
+#endif
 
 #include "UObject/DefineUPropertyMacros.h"
\ No newline at end of file
diff --git a/Engine/Source/Runtime/Engine/Classes/Animation/AnimNode_AssetPlayerBase.h b/Engine/Source/Runtime/Engine/Classes/Animation/AnimNode_AssetPlayerBase.h
--- a/Engine/Source/Runtime/Engine/Classes/Animation/AnimNode_AssetPlayerBase.h
+++ b/Engine/Source/Runtime/Engine/Classes/Animation/AnimNode_AssetPlayerBase.h
@@ -14,7 +14,9 @@ struct ENGINE_API FAnimNode_AssetPlayerBase : public FAnimNode_Base
 {
 	GENERATED_BODY();
 
-	FAnimNode_AssetPlayerBase();
+	friend class UAnimGraphNode_AssetPlayerBase;
+
+	FAnimNode_AssetPlayerBase() = default;
 
 	/** Get the last encountered blend weight for this node */
 	virtual float GetCachedBlendWeight() const;
@@ -29,7 +31,7 @@ struct ENGINE_API FAnimNode_AssetPlayerBase : public FAnimNode_Base
 	virtual void SetAccumulatedTime(const float& NewTime);
 
 	/** Get the animation asset associated with the node, derived classes should implement this */
-	virtual UAnimationAsset* GetAnimAsset();
+	virtual UAnimationAsset* GetAnimAsset() const;
 
 	/** Initialize function for setup purposes */
 	virtual void Initialize_AnyThread(const FAnimationInitializeContext& Context) override;
@@ -43,53 +45,80 @@ struct ENGINE_API FAnimNode_AssetPlayerBase : public FAnimNode_Base
 	/** Update method for the asset player, to be implemented by derived classes */
 	virtual void UpdateAssetPlayer(const FAnimationUpdateContext& Context) {};
 
-	// The group name (NAME_None if it is not part of any group)
-	UPROPERTY()
-	FName GroupName;
+	// Create a tick record for this node
+	void CreateTickRecordForNode(const FAnimationUpdateContext& Context, UAnimSequenceBase* Sequence, bool bLooping, float PlayRate);
+
+	// Functions to report data to getters, this is required for all asset players (but can't be pure abstract because of struct instantiation generated code).
+	virtual float GetCurrentAssetLength() const { return 0.0f; }
+	virtual float GetCurrentAssetTime() const { return 0.0f; }
+	virtual float GetCurrentAssetTimePlayRateAdjusted() const { return GetCurrentAssetTime(); }
+
+	// Get the sync group name we are using
+	FName GetGroupName() const;
+
+	// Get the sync group role we are using
+	EAnimGroupRole::Type GetGroupRole() const;
+
+	// Get the sync group method we are using
+	EAnimSyncMethod GetGroupMethod() const;
+
+	// Check whether this node should be ignored when testing for relevancy in state machines
+	bool GetIgnoreForRelevancyTest() const;
+
+#if WITH_EDITORONLY_DATA
+	// Set the sync group name we are using
+	void SetGroupName(FName InGroupName) { GroupName = InGroupName; }
+
+	// Set the sync group role we are using
+	void GetGroupRole(EAnimGroupRole::Type InRole) { GroupRole = InRole; }
+
+	// Set the sync group method we are using
+	void SetGroupMethod(EAnimSyncMethod InMethod) { Method = InMethod; }
+
+	// Set whether this node should be ignored when testing for relevancy in state machines
+	void SetIgnoreForRelevancyTest(bool bInIgnoreForRelevancyTest) { bIgnoreForRelevancyTest = bInIgnoreForRelevancyTest; }
+#endif
+
+private:
 
 #if WITH_EDITORONLY_DATA
+	// The group name (NAME_None if it is not part of any group)
+	UPROPERTY(EditAnywhere, Category=Sync, meta=(FoldProperty))
+	FName GroupName = NAME_None;
+
 	UPROPERTY()
-	int32 GroupIndex_DEPRECATED;
+	int32 GroupIndex_DEPRECATED = INDEX_NONE;
 
 	UPROPERTY()
-	EAnimSyncGroupScope GroupScope_DEPRECATED;
-#endif
+	EAnimSyncGroupScope GroupScope_DEPRECATED = EAnimSyncGroupScope::Local;
 
 	// The role this player can assume within the group (ignored if GroupIndex is INDEX_NONE)
-	UPROPERTY()
-	TEnumAsByte<EAnimGroupRole::Type> GroupRole;
+	UPROPERTY(EditAnywhere, Category=Sync, meta=(FoldProperty))
+	TEnumAsByte<EAnimGroupRole::Type> GroupRole = EAnimGroupRole::CanBeLeader;
 
 	// How synchronization is determined
-	UPROPERTY()
+	UPROPERTY(EditAnywhere, Category=Sync, meta=(FoldProperty))
 	EAnimSyncMethod Method = EAnimSyncMethod::DoNotSync;
 
 	/** If true, "Relevant anim" nodes that look for the highest weighted animation in a state will ignore
 	 *  this node
 	 */
-	UPROPERTY(EditAnywhere, BlueprintReadWrite, Category=Relevancy, meta=(PinHiddenByDefault))
-	bool bIgnoreForRelevancyTest;
-
-	// Create a tick record for this node
-	void CreateTickRecordForNode(const FAnimationUpdateContext& Context, UAnimSequenceBase* Sequence, bool bLooping, float PlayRate);
-
-	// Functions to report data to getters, this is required for all asset players (but can't be pure abstract because of struct instantiation generated code).
-	virtual float GetCurrentAssetLength() { return 0.0f; }
-	virtual float GetCurrentAssetTime() { return 0.0f; }
-	virtual float GetCurrentAssetTimePlayRateAdjusted() { return GetCurrentAssetTime(); }
+	UPROPERTY(EditAnywhere, Category=Relevancy, meta=(FoldProperty, PinHiddenByDefault))
+	bool bIgnoreForRelevancyTest = false;
+#endif
 
 protected:
-
-	/** Track whether we have been full weight previously. Reset when we reach 0 weight*/
-	bool bHasBeenFullWeight;
+	/** Store data about current marker position when using marker based syncing*/
+	FMarkerTickRecord MarkerTickRecord;
 
 	/** Last encountered blendweight for this node */
 	UPROPERTY(BlueprintReadWrite, Transient, Category=DoNotEdit)
-	float BlendWeight;
+	float BlendWeight = 0.0f;
 
 	/** Accumulated time used to reference the asset in this node */
 	UPROPERTY(BlueprintReadWrite, Transient, Category=DoNotEdit)
-	float InternalTimeAccumulator;
+	float InternalTimeAccumulator = 0.0f;
 
-	/** Store data about current marker position when using marker based syncing*/
-	FMarkerTickRecord MarkerTickRecord;
+	/** Track whether we have been full weight previously. Reset when we reach 0 weight*/
+	bool bHasBeenFullWeight = false;
 };
diff --git a/Engine/Source/Runtime/Engine/Classes/Animation/AnimNode_Root.h b/Engine/Source/Runtime/Engine/Classes/Animation/AnimNode_Root.h
--- a/Engine/Source/Runtime/Engine/Classes/Animation/AnimNode_Root.h
+++ b/Engine/Source/Runtime/Engine/Classes/Animation/AnimNode_Root.h
@@ -16,13 +16,16 @@ struct ENGINE_API FAnimNode_Root : public FAnimNode_Base
 	UPROPERTY(EditAnywhere, BlueprintReadWrite, Category=Links)
 	FPoseLink Result;
 
+#if WITH_EDITORONLY_DATA
+protected:
 	/** The name of this root node, used to identify the output of this graph. Filled in by the compiler, propagated from the parent graph. */
-	UPROPERTY()
+	UPROPERTY(meta=(FoldProperty))
 	FName Name;
 
 	/** The group of this root node, used to group this output with others when used in a layer. */
-	UPROPERTY()
+	UPROPERTY(meta=(FoldProperty))
 	FName Group;
+#endif
 
 public:	
 	FAnimNode_Root();
@@ -34,4 +37,18 @@ public:
 	virtual void Evaluate_AnyThread(FPoseContext& Output) override;
 	virtual void GatherDebugData(FNodeDebugData& DebugData) override;
 	// End of FAnimNode_Base interface
+
+#if WITH_EDITORONLY_DATA
+	/** Set the name of this root node, used to identify the output of this graph */
+	void SetName(FName InName) { Name = InName; }
+
+	/** Set the group of this root node, used to group this output with others when used in a layer. */
+	void SetGroup(FName InGroup) { Group = InGroup; }
+#endif
+
+	/** Get the name of this root node, used to identify the output of this graph */
+	FName GetName() const;
+
+	/** Get the group of this root node, used to group this output with others when used in a layer. */
+	FName GetGroup() const;
 };
diff --git a/Engine/Source/Runtime/Engine/Classes/Animation/AnimNode_SequencePlayer.h b/Engine/Source/Runtime/Engine/Classes/Animation/AnimNode_SequencePlayer.h
--- a/Engine/Source/Runtime/Engine/Classes/Animation/AnimNode_SequencePlayer.h
+++ b/Engine/Source/Runtime/Engine/Classes/Animation/AnimNode_SequencePlayer.h
@@ -10,70 +10,172 @@
 #include "AnimNode_SequencePlayer.generated.h"
 
 
-// Sequence player node
+// Sequence player node. Not instantiated directly, use FAnimNode_SequencePlayer or FAnimNode_SequencePlayer_Standalone
 USTRUCT(BlueprintInternalUseOnly)
-struct ENGINE_API FAnimNode_SequencePlayer : public FAnimNode_AssetPlayerBase
+struct ENGINE_API FAnimNode_SequencePlayerBase : public FAnimNode_AssetPlayerBase
 {
-	GENERATED_USTRUCT_BODY()
+	GENERATED_BODY()
+
+protected:
+	// Corresponding state for PlayRateScaleBiasClampConstants
+	UPROPERTY()
+	FInputScaleBiasClampState PlayRateScaleBiasClampState;
+
 public:
+	// FAnimNode_AssetPlayerBase interface
+	virtual float GetCurrentAssetTime() const override;
+	virtual float GetCurrentAssetTimePlayRateAdjusted() const override;
+	virtual float GetCurrentAssetLength() const override;
+	virtual UAnimationAsset* GetAnimAsset() const override { return GetSequence(); }
+	// End of FAnimNode_AssetPlayerBase interface
+
+	// FAnimNode_Base interface
+	virtual void Initialize_AnyThread(const FAnimationInitializeContext& Context) override;
+	virtual void CacheBones_AnyThread(const FAnimationCacheBonesContext& Context) override;
+	virtual void UpdateAssetPlayer(const FAnimationUpdateContext& Context) override;
+	virtual void Evaluate_AnyThread(FPoseContext& Output) override;
+	virtual void GatherDebugData(FNodeDebugData& DebugData) override;
+	// End of FAnimNode_Base interface
+
+	float GetTimeFromEnd(float CurrentNodeTime) const;
+	float GetEffectiveStartPosition(const FAnimationBaseContext& Context) const;
+
 	// The animation sequence asset to play
-	UPROPERTY(EditAnywhere, BlueprintReadWrite, Category = Settings, meta = (PinHiddenByDefault, DisallowedClasses="AnimMontage"))
-	TObjectPtr<UAnimSequenceBase> Sequence;
+	virtual UAnimSequenceBase* GetSequence() const { return nullptr; }
+
+protected:
+	// Set the animation sequence asset to play
+	virtual void SetSequence(UAnimSequenceBase* InSequence) PURE_VIRTUAL(SetSequence, )
 
+	// Set the animation to continue looping when it reaches the end
+	virtual void SetLoopAnimation(bool bInLoopAnimation) PURE_VIRTUAL(SetLoopAnimation, )
+	
 	// The Basis in which the PlayRate is expressed in. This is used to rescale PlayRate inputs.
 	// For example a Basis of 100 means that the PlayRate input will be divided by 100.
-	UPROPERTY(EditAnywhere, BlueprintReadWrite, Category = Settings, meta = (PinHiddenByDefault))
-	float PlayRateBasis;
+	virtual float GetPlayRateBasis() const PURE_VIRTUAL(GetPlayRateBasis, return 1.0f;)
 
 	// The play rate multiplier. Can be negative, which will cause the animation to play in reverse.
-	UPROPERTY(EditAnywhere, BlueprintReadWrite, Category = Settings, meta = (PinHiddenByDefault))
-	float PlayRate;
+	virtual float GetPlayRate() const PURE_VIRTUAL(GetPlayRate, return 1.0f;)
 	
 	// Additional scaling, offsetting and clamping of PlayRate input.
 	// Performed after PlayRateBasis.
-	UPROPERTY(EditAnywhere, BlueprintReadWrite, Category = Settings)
-	FInputScaleBiasClamp PlayRateScaleBiasClamp;
+	virtual const FInputScaleBiasClampConstants& GetPlayRateScaleBiasClampConstants() const PURE_VIRTUAL(GetPlayRateScaleBiasClampConstants, static FInputScaleBiasClampConstants Dummy; return Dummy; )
 
 	// The start up position, it only applies when reinitialized
 	// if you loop, it will still start from 0.f after finishing the round
-	UPROPERTY(EditAnywhere, BlueprintReadWrite, Category = Settings, meta = (PinHiddenByDefault))
-	float StartPosition;
+	virtual float GetStartPosition() const PURE_VIRTUAL(GetStartPosition, return 0.0f;)
 
 	// Should the animation continue looping when it reaches the end?
-	UPROPERTY(EditAnywhere, BlueprintReadWrite, Category = Settings, meta = (PinHiddenByDefault))
-	bool bLoopAnimation;
+	virtual bool GetLoopAnimation() const PURE_VIRTUAL(GetLoopAnimation, return true;)
 
 	// Use pose matching to choose the start position. Requires experimental PoseSearch plugin.
-	UPROPERTY(EditAnywhere, BlueprintReadWrite, Category = PoseMatching, meta = (PinHiddenByDefault))
-	bool bStartFromMatchingPose;
+	virtual bool GetStartFromMatchingPose() const PURE_VIRTUAL(GetStartFromMatchingPose, return false;)
+};
+
+// Sequence player node that can be used with constant folding
+USTRUCT(BlueprintInternalUseOnly)
+struct ENGINE_API FAnimNode_SequencePlayer : public FAnimNode_SequencePlayerBase
+{
+	GENERATED_BODY()
+
+protected:
+	friend class UAnimGraphNode_SequencePlayer;
+
+#if WITH_EDITORONLY_DATA
+	// The animation sequence asset to play
+	UPROPERTY(EditAnywhere, Category = Settings, meta = (PinHiddenByDefault, DisallowedClasses="AnimMontage", FoldProperty))
+	TObjectPtr<UAnimSequenceBase> Sequence = nullptr;
+
+	// The Basis in which the PlayRate is expressed in. This is used to rescale PlayRate inputs.
+	// For example a Basis of 100 means that the PlayRate input will be divided by 100.
+	UPROPERTY(EditAnywhere, Category = Settings, meta = (PinHiddenByDefault, FoldProperty))
+	float PlayRateBasis = 1.0f;
+
+	// The play rate multiplier. Can be negative, which will cause the animation to play in reverse.
+	UPROPERTY(EditAnywhere, Category = Settings, meta = (PinHiddenByDefault, FoldProperty))
+	float PlayRate = 1.0f;
+	
+	// Additional scaling, offsetting and clamping of PlayRate input.
+	// Performed after PlayRateBasis.
+	UPROPERTY(EditAnywhere, Category = Settings, meta = (DisplayName="PlayRateScaleBiasClamp", FoldProperty))
+	FInputScaleBiasClampConstants PlayRateScaleBiasClampConstants;
+
+	UPROPERTY()
+	FInputScaleBiasClamp PlayRateScaleBiasClamp_DEPRECATED;
+
+	// The start up position, it only applies when reinitialized
+	// if you loop, it will still start from 0.f after finishing the round
+	UPROPERTY(EditAnywhere, Category = Settings, meta = (PinHiddenByDefault, FoldProperty))
+	float StartPosition = 0.0f;
+
+	// Should the animation continue looping when it reaches the end?
+	UPROPERTY(EditAnywhere, Category = Settings, meta = (PinHiddenByDefault, FoldProperty))
+	bool bLoopAnimation = true;
+
+	// Use pose matching to choose the start position. Requires experimental PoseSearch plugin.
+	UPROPERTY(EditAnywhere, Category = PoseMatching, meta = (PinHiddenByDefault, FoldProperty))
+	bool bStartFromMatchingPose = false;
+#endif
 
 public:
-	FAnimNode_SequencePlayer()
-		: Sequence(nullptr)
-		, PlayRateBasis(1.0f)
-		, PlayRate(1.0f)
-		, StartPosition(0.f)
-		, bLoopAnimation(true)
-		, bStartFromMatchingPose(false)
-	{
-	}
+	// FAnimNode_SequencePlayerBase interface
+	virtual void SetSequence(UAnimSequenceBase* InSequence) override;
+	virtual UAnimSequenceBase* GetSequence() const override;
+	virtual float GetPlayRateBasis() const override;
+	virtual float GetPlayRate() const override;
+	virtual const FInputScaleBiasClampConstants& GetPlayRateScaleBiasClampConstants() const override;
+	virtual float GetStartPosition() const override;
+	virtual bool GetLoopAnimation() const override;
+	virtual bool GetStartFromMatchingPose() const override;
+};
 
-	// FAnimNode_AssetPlayerBase interface
-	virtual float GetCurrentAssetTime() override;
-	virtual float GetCurrentAssetTimePlayRateAdjusted() override;
-	virtual float GetCurrentAssetLength() override;
-	virtual UAnimationAsset* GetAnimAsset() override { return Sequence; }
-	// End of FAnimNode_AssetPlayerBase interface
+// Sequence player node that can be used standalone (without constant folding)
+USTRUCT(BlueprintInternalUseOnly)
+struct ENGINE_API FAnimNode_SequencePlayer_Standalone : public FAnimNode_SequencePlayerBase
+{	
+	GENERATED_BODY()
 
-	// FAnimNode_Base interface
-	virtual void Initialize_AnyThread(const FAnimationInitializeContext& Context) override;
-	virtual void CacheBones_AnyThread(const FAnimationCacheBonesContext& Context) override;
-	virtual void UpdateAssetPlayer(const FAnimationUpdateContext& Context) override;
-	virtual void Evaluate_AnyThread(FPoseContext& Output) override;
-	virtual void OverrideAsset(UAnimationAsset* NewAsset) override;
-	virtual void GatherDebugData(FNodeDebugData& DebugData) override;
-	// End of FAnimNode_Base interface
+protected:
+	// The animation sequence asset to play
+	UPROPERTY(EditAnywhere, Category = Settings, meta = (PinHiddenByDefault, DisallowedClasses="AnimMontage"))
+	TObjectPtr<UAnimSequenceBase> Sequence = nullptr;
 
-	float GetTimeFromEnd(float CurrentNodeTime);
-	float GetEffectiveStartPosition(const FAnimationBaseContext& Context) const;
+	// The Basis in which the PlayRate is expressed in. This is used to rescale PlayRate inputs.
+	// For example a Basis of 100 means that the PlayRate input will be divided by 100.
+	UPROPERTY(EditAnywhere, Category = Settings, meta = (PinHiddenByDefault))
+	float PlayRateBasis = 1.0f;
+
+	// The play rate multiplier. Can be negative, which will cause the animation to play in reverse.
+	UPROPERTY(EditAnywhere, Category = Settings, meta = (PinHiddenByDefault))
+	float PlayRate = 1.0f;
+	
+	// Additional scaling, offsetting and clamping of PlayRate input.
+	// Performed after PlayRateBasis.
+	UPROPERTY(EditAnywhere, Category = Settings, meta = (DisplayName="PlayRateScaleBiasClamp"))
+	FInputScaleBiasClampConstants PlayRateScaleBiasClampConstants;
+
+	// The start up position, it only applies when reinitialized
+	// if you loop, it will still start from 0.f after finishing the round
+	UPROPERTY(EditAnywhere, Category = Settings, meta = (PinHiddenByDefault))
+	float StartPosition = 0.0f;
+
+	// Should the animation continue looping when it reaches the end?
+	UPROPERTY(EditAnywhere, Category = Settings, meta = (PinHiddenByDefault))
+	bool bLoopAnimation = true;
+
+	// Use pose matching to choose the start position. Requires experimental PoseSearch plugin.
+	UPROPERTY(EditAnywhere, Category = PoseMatching, meta = (PinHiddenByDefault))
+	bool bStartFromMatchingPose = false;
+
+public:
+	// FAnimNode_SequencePlayerBase interface
+	virtual void SetSequence(UAnimSequenceBase* InSequence) override;
+	virtual void SetLoopAnimation(bool bInLoopAnimation) override;
+	virtual UAnimSequenceBase* GetSequence() const override;
+	virtual float GetPlayRateBasis() const override;
+	virtual float GetPlayRate() const override;
+	virtual const FInputScaleBiasClampConstants& GetPlayRateScaleBiasClampConstants() const override;
+	virtual float GetStartPosition() const override;
+	virtual bool GetLoopAnimation() const override;
+	virtual bool GetStartFromMatchingPose() const override;
 };
diff --git a/Engine/Source/Runtime/Engine/Classes/Animation/AnimNode_StateMachine.h b/Engine/Source/Runtime/Engine/Classes/Animation/AnimNode_StateMachine.h
--- a/Engine/Source/Runtime/Engine/Classes/Animation/AnimNode_StateMachine.h
+++ b/Engine/Source/Runtime/Engine/Classes/Animation/AnimNode_StateMachine.h
@@ -262,7 +262,7 @@ protected:
 	void EvaluateTransitionStandardBlendInternal(FPoseContext& Output, FAnimationActiveTransitionEntry& Transition, const FPoseContext& PreviousStateResult, const FPoseContext& NextStateResult);
 	void EvaluateTransitionCustomBlend(FPoseContext& Output, FAnimationActiveTransitionEntry& Transition, bool bIntermediatePoseIsValid);
 
-	FAnimNode_AssetPlayerBase* GetRelevantAssetPlayerFromState(const FAnimationUpdateContext& Context, const FBakedAnimationState& StateInfo);
+	const FAnimNode_AssetPlayerBase* GetRelevantAssetPlayerFromState(const FAnimationUpdateContext& Context, const FBakedAnimationState& StateInfo) const;
 
 	void LogInertializationRequestError(const FAnimationUpdateContext& Context, int32 PreviousState, int32 NextState);
 
diff --git a/Engine/Source/Runtime/Engine/Classes/Animation/InputScaleBias.h b/Engine/Source/Runtime/Engine/Classes/Animation/InputScaleBias.h
--- a/Engine/Source/Runtime/Engine/Classes/Animation/InputScaleBias.h
+++ b/Engine/Source/Runtime/Engine/Classes/Animation/InputScaleBias.h
@@ -130,7 +130,99 @@ public:
 	// Apply scale, bias, and clamp to value
 	float ApplyTo(float Value, float InDeltaTime) const;
 	void Reinitialize() { bInitialized = false; }
+
+#if WITH_EDITOR
 	FText GetFriendlyName(FText InFriendlyName) const;
+#endif
+};
+
+USTRUCT(BlueprintType)
+struct ENGINE_API FInputScaleBiasClampConstants
+{
+	GENERATED_BODY()
+
+	UPROPERTY(EditAnywhere, BlueprintReadWrite, Category = Settings)
+	bool bMapRange;
+
+	UPROPERTY(EditAnywhere, BlueprintReadWrite, Category = Settings)
+	bool bClampResult;
+
+	UPROPERTY(EditAnywhere, BlueprintReadWrite, Category = Settings)
+	bool bInterpResult;
+
+	UPROPERTY(EditAnywhere, BlueprintReadWrite, Category = Settings, meta = (EditCondition = "bMapRange"))
+	FInputRange InRange;
+
+	UPROPERTY(EditAnywhere, BlueprintReadWrite, Category = Settings, meta = (EditCondition = "bMapRange"))
+	FInputRange OutRange;
+
+	UPROPERTY(EditAnywhere, BlueprintReadWrite, Category = Settings)
+	float Scale;
+
+	UPROPERTY(EditAnywhere, BlueprintReadWrite, Category = Settings)
+	float Bias;
+
+	UPROPERTY(EditAnywhere, BlueprintReadWrite, Category = Settings, meta = (EditCondition = "bClampResult"))
+	float ClampMin;
+
+	UPROPERTY(EditAnywhere, BlueprintReadWrite, Category = Settings, meta = (EditCondition = "bClampResult"))
+	float ClampMax;
+
+	UPROPERTY(EditAnywhere, BlueprintReadWrite, Category = Settings, meta = (EditCondition = "bInterpResult"))
+	float InterpSpeedIncreasing;
+
+	UPROPERTY(EditAnywhere, BlueprintReadWrite, Category = Settings, meta = (EditCondition = "bInterpResult"))
+	float InterpSpeedDecreasing;
+
+public:
+	FInputScaleBiasClampConstants()
+		: bMapRange(false)
+		, bClampResult(false)
+		, bInterpResult(false)
+		, Scale(1.0f)
+		, Bias(0.0f)
+		, ClampMin(0.f)
+		, ClampMax(1.f)
+		, InterpSpeedIncreasing(10.f)
+		, InterpSpeedDecreasing(10.f)
+	{
+	}
+
+#if WITH_EDITOR
+	// Get a friendly name to display on a pin
+	FText GetFriendlyName(FText InFriendlyName) const;
+
+	// Copy parameters from the legacy combined constants/state structure
+	void CopyFromLegacy(const FInputScaleBiasClamp& InLegacy);
+#endif
+};
+
+// Mutable state struct to be used with FInputScaleBiasClampConstants
+USTRUCT(BlueprintType)
+struct ENGINE_API FInputScaleBiasClampState
+{
+	GENERATED_BODY()
+
+	// The interpolated result
+	float InterpolatedResult;
+
+	// Whether this state is initialized
+	bool bInitialized;
+
+public:
+	FInputScaleBiasClampState()
+		: InterpolatedResult(0.f)
+		, bInitialized(false)
+	{
+	}
+
+	// Apply scale, bias, and clamp to value
+	float ApplyTo(const FInputScaleBiasClampConstants& InConstants, float Value, float InDeltaTime);
+
+	// Apply but dont modify InterpolatedResult
+	float ApplyTo(const FInputScaleBiasClampConstants& InConstants, float Value) const;
+
+	void Reinitialize() { bInitialized = false; }
 };
 
 // AnimNodes using an Alpha can choose how it is driven.
diff --git a/Engine/Source/Runtime/Engine/Private/AnimBlueprint.cpp b/Engine/Source/Runtime/Engine/Private/AnimBlueprint.cpp
--- a/Engine/Source/Runtime/Engine/Private/AnimBlueprint.cpp
+++ b/Engine/Source/Runtime/Engine/Private/AnimBlueprint.cpp
@@ -18,7 +18,10 @@ UAnimBlueprint::UAnimBlueprint(const FObjectInitializer& ObjectInitializer)
 	: Super(ObjectInitializer)
 {
 	bUseMultiThreadedAnimationUpdate = true;
-
+#if WITH_EDITORONLY_DATA
+	bRefreshExtensions = true;
+#endif
+	
 #if WITH_EDITOR
 	if(!HasAnyFlags(RF_ClassDefaultObject))
 	{
@@ -84,8 +87,6 @@ int32 UAnimBlueprint::FindOrAddGroup(FName GroupName)
 	}
 }
 
-
-/** Returns the most base anim blueprint for a given blueprint (if it is inherited from another anim blueprint, returning null if only native / non-anim BP classes are it's parent) */
 UAnimBlueprint* UAnimBlueprint::FindRootAnimBlueprint(const UAnimBlueprint* DerivedBlueprint)
 {
 	UAnimBlueprint* ParentBP = nullptr;
@@ -102,6 +103,19 @@ UAnimBlueprint* UAnimBlueprint::FindRootAnimBlueprint(const UAnimBlueprint* Deri
 	return ParentBP;
 }
 
+UAnimBlueprint* UAnimBlueprint::GetParentAnimBlueprint(const UAnimBlueprint* DerivedBlueprint)
+{
+	UAnimBlueprint* ParentBP = nullptr;
+	UClass* ParentClass = DerivedBlueprint->ParentClass;
+
+	if (UAnimBlueprint* TestBP = Cast<UAnimBlueprint>(ParentClass->ClassGeneratedBy))
+	{
+		ParentBP = TestBP;
+	}
+
+	return ParentBP;
+}
+
 FAnimParentNodeAssetOverride* UAnimBlueprint::GetAssetOverrideForNode(FGuid NodeGuid, bool bIgnoreSelf) const
 {
 	TArray<UBlueprint*> Hierarchy;
diff --git a/Engine/Source/Runtime/Engine/Private/Animation/AnimBlueprintGeneratedClass.cpp b/Engine/Source/Runtime/Engine/Private/Animation/AnimBlueprintGeneratedClass.cpp
--- a/Engine/Source/Runtime/Engine/Private/Animation/AnimBlueprintGeneratedClass.cpp
+++ b/Engine/Source/Runtime/Engine/Private/Animation/AnimBlueprintGeneratedClass.cpp
@@ -8,7 +8,6 @@
 #include "Serialization/ObjectWriter.h"
 #include "Serialization/ObjectReader.h"
 #include "ObjectEditorUtils.h"
-#include "BonePose.h"
 #include "Animation/AnimNodeBase.h"
 #include "Animation/AnimInstance.h"
 #include "UObject/AnimObjectVersion.h"
@@ -17,10 +16,10 @@
 #include "Animation/AnimNode_Root.h"
 #include "Animation/AnimNode_LinkedInputPose.h"
 #include "Animation/AnimNode_LinkedAnimLayer.h"
-#include "Animation/AnimNode_AssetPlayerBase.h"
 #include "Animation/AnimNode_StateMachine.h"
+#include "Animation/AnimSubsystem.h"
 #include "EdGraph/EdGraphNode.h"
-#include "Algo/Reverse.h"
+#include "Animation/AnimSubsystemInstance.h"
 
 /////////////////////////////////////////////////////
 // FStateMachineDebugData
@@ -330,13 +329,115 @@ void FAnimationFrameSnapshot::CopyToInstance(UAnimInstance* Instance)
 #endif
 
 /////////////////////////////////////////////////////
-// UAnimBlueprintGeneratedClas
+// UAnimBlueprintGeneratedClass
 
 UAnimBlueprintGeneratedClass::UAnimBlueprintGeneratedClass(const FObjectInitializer& ObjectInitializer)
 	: Super(ObjectInitializer)
 {
 }
 
+const void* UAnimBlueprintGeneratedClass::GetConstantNodeValueRaw(int32 InIndex) const
+{
+	const FProperty* Property = ConstantProperties[InIndex];
+	check(Property->GetOwner<UStruct>() && GetSparseClassDataStruct()->IsChildOf(Property->GetOwner<UStruct>()));
+	const void* ConstantData = GetConstantNodeData();
+	check(ConstantData);
+	return Property->ContainerPtrToValuePtr<void>(static_cast<const void*>(ConstantData));
+}
+
+const void* UAnimBlueprintGeneratedClass::GetMutableNodeValueRaw(int32 InIndex, const UObject* InObject) const
+{
+	const FProperty* Property = MutableProperties[InIndex];
+	check(Property->GetOwner<UStruct>() && Property->GetOwner<UStruct>()->IsChildOf(FAnimBlueprintMutableData::StaticStruct()));
+	const FAnimBlueprintMutableData* MutableData = GetMutableNodeData(InObject);
+	check(MutableData);	
+	return Property->ContainerPtrToValuePtr<void>(static_cast<const void*>(MutableData));
+}
+
+FAnimBlueprintMutableData* UAnimBlueprintGeneratedClass::GetMutableNodeData(UObject* InObject) const
+{
+	if(MutableNodeDataProperty != nullptr)
+	{
+		check(MutableNodeDataProperty->Struct->IsChildOf(FAnimBlueprintMutableData::StaticStruct()));
+		check(MutableNodeDataProperty->GetOwner<UClass>() && InObject->GetClass()->IsChildOf(MutableNodeDataProperty->GetOwner<UClass>()));
+		return MutableNodeDataProperty->ContainerPtrToValuePtr<FAnimBlueprintMutableData>(InObject);
+	}
+
+	return nullptr;
+}
+
+const FAnimBlueprintMutableData* UAnimBlueprintGeneratedClass::GetMutableNodeData(const UObject* InObject) const
+{
+	if(MutableNodeDataProperty != nullptr)
+	{
+		check(MutableNodeDataProperty->Struct->IsChildOf(FAnimBlueprintMutableData::StaticStruct()));
+		check(MutableNodeDataProperty->GetOwner<UClass>() && InObject->GetClass()->IsChildOf(MutableNodeDataProperty->GetOwner<UClass>()));
+		return MutableNodeDataProperty->ContainerPtrToValuePtr<FAnimBlueprintMutableData>(InObject);
+	}
+
+	return nullptr;
+}
+
+const void* UAnimBlueprintGeneratedClass::GetConstantNodeData() const 
+{
+	return SparseClassData;
+}
+
+void UAnimBlueprintGeneratedClass::ForEachSubsystem(TFunctionRef<EAnimSubsystemEnumeration(const FAnimSubsystemContext&)> InFunction) const
+{
+	for(int32 SubsystemIndex = 0; SubsystemIndex < ConstantSubsystemProperties.Num(); ++SubsystemIndex)
+	{
+		FStructProperty* ConstantSubsystemProperty = ConstantSubsystemProperties[SubsystemIndex];
+		check(ConstantSubsystemProperty->GetOwner<UStruct>() && GetSparseClassDataStruct()->IsChildOf(ConstantSubsystemProperty->GetOwner<UStruct>()));
+		FAnimSubsystem& Subsystem = *ConstantSubsystemProperty->ContainerPtrToValuePtr<FAnimSubsystem>(const_cast<void*>(GetConstantNodeData()));
+
+		if(InFunction(FAnimSubsystemContext(Subsystem, ConstantSubsystemProperty->Struct)) == EAnimSubsystemEnumeration::Stop)
+		{
+			break;
+		}
+	}
+}
+
+void UAnimBlueprintGeneratedClass::ForEachSubsystem(UObject* InObject, TFunctionRef<EAnimSubsystemEnumeration(const FAnimSubsystemInstanceContext&)> InFunction) const
+{
+	check(ConstantSubsystemProperties.Num() == MutableSubsystemProperties.Num());
+	
+	for(int32 SubsystemIndex = 0; SubsystemIndex < ConstantSubsystemProperties.Num(); ++SubsystemIndex)
+	{
+		FStructProperty* ConstantSubsystemProperty = ConstantSubsystemProperties[SubsystemIndex];
+		FStructProperty* MutableSubsystemProperty = MutableSubsystemProperties[SubsystemIndex];
+		
+		check(ConstantSubsystemProperty->GetOwner<UStruct>() && GetSparseClassDataStruct()->IsChildOf(ConstantSubsystemProperty->GetOwner<UStruct>()));
+		const FAnimSubsystem& Subsystem = *ConstantSubsystemProperty->ContainerPtrToValuePtr<FAnimSubsystem>(GetConstantNodeData());
+
+		check(MutableSubsystemProperty->GetOwner<UClass>() && InObject->GetClass()->IsChildOf(MutableSubsystemProperty->GetOwner<UClass>()));
+		FAnimSubsystemInstance& SubsystemInstance = *MutableSubsystemProperty->ContainerPtrToValuePtr<FAnimSubsystemInstance>(InObject);
+		
+		if(InFunction(FAnimSubsystemInstanceContext(Subsystem, ConstantSubsystemProperty->Struct, SubsystemInstance, MutableSubsystemProperty->Struct)) == EAnimSubsystemEnumeration::Stop)
+		{
+			break;
+		}
+	}
+}
+
+const FAnimSubsystem* UAnimBlueprintGeneratedClass::FindSubsystem(UScriptStruct* InSubsystemType) const
+{
+	const FAnimSubsystem* FoundSubsystem = nullptr;
+	
+	ForEachSubsystem([&FoundSubsystem, InSubsystemType](const FAnimSubsystemContext& InContext)
+	{
+		if(InContext.SubsystemStruct == InSubsystemType)
+		{
+			FoundSubsystem = &InContext.Subsystem;
+			return EAnimSubsystemEnumeration::Stop;
+		}
+
+		return EAnimSubsystemEnumeration::Continue;
+	});
+
+	return FoundSubsystem;
+}
+
 void UAnimBlueprintGeneratedClass::Serialize(FArchive& Ar)
 {
 	Ar.UsingCustomVersion(FAnimObjectVersion::GUID);
@@ -357,15 +458,10 @@ void UAnimBlueprintGeneratedClass::Link(FArchive& Ar, bool bRelinkExistingProper
 	DynamicResetNodeProperties.Empty();
 	StateMachineNodeProperties.Empty();
 	InitializationNodeProperties.Empty();
-
-#if WITH_EDITORONLY_DATA
-	for (FExposedValueHandler& Handler : EvaluateGraphExposedInputs)
-	{
-		// handle potential renames of the class package
-		Handler.ValueHandlerNodeProperty.ResolveWithRenamedStructPackage(this);
-		check(Handler.ValueHandlerNodeProperty.Get() || Handler.ValueHandlerNodeProperty.IsPathToFieldEmpty());
-	}
-#endif // WITH_EDITORONLY_DATA
+	MutableProperties.Empty();
+	ConstantProperties.Empty();
+	MutableSubsystemProperties.Empty();
+	ConstantSubsystemProperties.Empty();
 
 #if WITH_EDITOR
 	// This relies on the entire class being fully loaded, this is not the case with EDL async-loading, in which case the functions are generated in PostLoad
@@ -393,6 +489,19 @@ void UAnimBlueprintGeneratedClass::Link(FArchive& Ar, bool bRelinkExistingProper
 				}
 				AnimNodeProperties.Add(StructProp);
 			}
+			else if(StructProp->Struct->IsChildOf(FAnimBlueprintMutableData::StaticStruct()))
+			{
+				MutableNodeDataProperty = StructProp;
+
+				for(TFieldIterator<FProperty> MutableIt(StructProp->Struct); MutableIt; ++MutableIt)
+				{
+					MutableProperties.Add(*MutableIt);
+				}
+			}
+			else if(StructProp->Struct->IsChildOf(FAnimSubsystemInstance::StaticStruct()))
+			{
+				MutableSubsystemProperties.Add(StructProp);
+			}
 		}
 	}
 }
@@ -425,28 +534,46 @@ void UAnimBlueprintGeneratedClass::PostLoadDefaultObject(UObject* Object)
 {
 	Super::PostLoadDefaultObject(Object);
 
+	OnPostLoadDefaults(Object);
+}
+
+void UAnimBlueprintGeneratedClass::OnPostLoadDefaults(UObject* Object)
+{
 	// Make sure that 'functions' are generated before we use them in LinkFunctionsToDefaultObjectNodes()
 	if (AnimBlueprintFunctions.IsEmpty()) 
 	{ 
 		GenerateAnimationBlueprintFunctions();
 	}
 
-	UAnimBlueprintGeneratedClass* Iter = this;
-	while(Iter)
-	{
-		FExposedValueHandler::ClassInitialization(Iter->EvaluateGraphExposedInputs, Object);
-		Iter = Cast<UAnimBlueprintGeneratedClass>(Iter->GetSuperClass());
-	}
+	BuildConstantProperties();
 
 	LinkFunctionsToDefaultObjectNodes(Object);
+
+	Subsystems.Empty();
+	
+	ForEachSubsystem(Object, [this, Object](const FAnimSubsystemInstanceContext& InContext)
+	{
+		Subsystems.Add(&InContext.Subsystem);
+
+		FAnimSubsystemPostLoadDefaultsContext Context(InContext, Object);
+		const_cast<FAnimSubsystem&>(InContext.Subsystem).OnPostLoadDefaults(Context);
+		return EAnimSubsystemEnumeration::Continue;
+	});
 }
 
 void UAnimBlueprintGeneratedClass::PostLoad()
 {
 	Super::PostLoad();
 
-	// Post-load property access library
-	PropertyAccess::PostLoadLibrary(PropertyAccessLibrary);
+	BuildConstantProperties();
+
+	// Call postload on the class subsystems
+	ForEachSubsystem([this](const FAnimSubsystemContext& InContext)
+	{
+		FAnimSubsystemPostLoadContext Context(InContext, *this);
+		const_cast<FAnimSubsystem&>(InContext.Subsystem).OnPostLoad(Context);
+		return EAnimSubsystemEnumeration::Continue;
+	});
 }
 
 void UAnimBlueprintGeneratedClass::GenerateAnimationBlueprintFunctions()
@@ -535,6 +662,21 @@ void UAnimBlueprintGeneratedClass::GenerateAnimationBlueprintFunctions()
 	}
 }
 
+void UAnimBlueprintGeneratedClass::InitializeAnimNodeData(UObject* DefaultObject)
+{
+	// Link functions to their nodes
+	for(int32 AnimNodeIndex = 0; AnimNodeIndex < AnimNodeProperties.Num(); ++AnimNodeIndex)
+	{
+		FStructProperty* StructProperty = AnimNodeProperties[AnimNodeIndex];
+
+		if(StructProperty->Struct->IsChildOf(FAnimNode_Base::StaticStruct()))
+		{
+			FAnimNode_Base* Node = StructProperty->ContainerPtrToValuePtr<FAnimNode_Base>(DefaultObject);
+			Node->SetNodeData(GetNodeData()[AnimNodeIndex]);
+		}
+	}
+}
+
 void UAnimBlueprintGeneratedClass::LinkFunctionsToDefaultObjectNodes(UObject* DefaultObject)
 {
 	PreUpdateNodeProperties.Empty();
@@ -545,12 +687,34 @@ void UAnimBlueprintGeneratedClass::LinkFunctionsToDefaultObjectNodes(UObject* De
 	for(int32 AnimNodeIndex = 0; AnimNodeIndex < AnimNodeProperties.Num(); ++AnimNodeIndex)
 	{
 		FStructProperty* StructProperty = AnimNodeProperties[AnimNodeIndex];
+
+		if(StructProperty->Struct->IsChildOf(FAnimNode_Base::StaticStruct()))
+		{
+			FAnimNode_Base* Node = StructProperty->ContainerPtrToValuePtr<FAnimNode_Base>(DefaultObject);
+			Node->SetNodeData(GetNodeData()[AnimNodeIndex]);
+
+			if(Node->NeedsDynamicReset())
+			{
+				DynamicResetNodeProperties.Add(StructProperty);
+			}
+
+			if(Node->HasPreUpdate())
+			{
+				PreUpdateNodeProperties.Add(StructProperty);
+			}
+
+			if(Node->NeedsOnInitializeAnimInstance())
+			{
+				InitializationNodeProperties.Add(StructProperty);
+			}
+		}
+		
 		if (StructProperty->Struct->IsChildOf(FAnimNode_Root::StaticStruct()))
 		{
 			FAnimNode_Root* RootNode = StructProperty->ContainerPtrToValuePtr<FAnimNode_Root>(DefaultObject);
-			if(FAnimBlueprintFunction* FoundFunction = AnimBlueprintFunctions.FindByPredicate([RootNode](const FAnimBlueprintFunction& InFunction){ return InFunction.Name == RootNode->Name; }))
+			if(FAnimBlueprintFunction* FoundFunction = AnimBlueprintFunctions.FindByPredicate([RootNode](const FAnimBlueprintFunction& InFunction){ return InFunction.Name == RootNode->GetName(); }))
 			{
-				FoundFunction->Group = RootNode->Group;
+				FoundFunction->Group = RootNode->GetGroup();
 				FoundFunction->OutputPoseNodeIndex = AnimNodeIndex;
 				FoundFunction->OutputPoseNodeProperty = StructProperty;
 			}
@@ -570,24 +734,6 @@ void UAnimBlueprintGeneratedClass::LinkFunctionsToDefaultObjectNodes(UObject* De
 				}
 			}
 		}
-		else if(StructProperty->Struct->IsChildOf(FAnimNode_Base::StaticStruct()))
-		{
-			FAnimNode_Base* Node = StructProperty->ContainerPtrToValuePtr<FAnimNode_Base>(DefaultObject);
-			if(Node->NeedsDynamicReset())
-			{
-				DynamicResetNodeProperties.Add(StructProperty);
-			}
-
-			if(Node->HasPreUpdate())
-			{
-				PreUpdateNodeProperties.Add(StructProperty);
-			}
-
-			if(Node->NeedsOnInitializeAnimInstance())
-			{
-				InitializationNodeProperties.Add(StructProperty);
-			}
-		}
 	}
 
 	// Determine whether functions are 'implemented'
@@ -650,7 +796,6 @@ const int32* UAnimBlueprintGeneratedClass::GetNodePropertyIndexFromGuid(FGuid Gu
 					return NodeIndex;
 				}
 			}
-
 		}
 	}
 
@@ -667,11 +812,81 @@ int32 UAnimBlueprintGeneratedClass::GetNodeIndexFromGuid(FGuid Guid, EPropertySe
 	return INDEX_NONE;
 }
 
-const UEdGraphNode* UAnimBlueprintGeneratedClass::GetVisualNodeFromNodePropertyIndex(int32 PropertyIndex) const
+const UEdGraphNode* UAnimBlueprintGeneratedClass::GetVisualNodeFromNodePropertyIndex(int32 PropertyIndex, EPropertySearchMode::Type SearchMode) const
 {
 	const int32 ReversedPropertyIndex = AnimNodeProperties.Num() - PropertyIndex - 1;
-	const TWeakObjectPtr<const UEdGraphNode>* Node = AnimBlueprintDebugData.NodePropertyIndexToNodeMap.Find(ReversedPropertyIndex);
-	return Node ? Node->Get() : nullptr;
+	if (SearchMode == EPropertySearchMode::OnlyThis)
+	{
+		const TWeakObjectPtr<const UEdGraphNode>* Node = AnimBlueprintDebugData.NodePropertyIndexToNodeMap.Find(ReversedPropertyIndex);
+		return Node ? Node->Get() : nullptr;
+	}
+	else
+	{
+		TArray<const UBlueprintGeneratedClass*> BlueprintHierarchy;
+		GetGeneratedClassesHierarchy(this, BlueprintHierarchy);
+
+		for (const UBlueprintGeneratedClass* Blueprint : BlueprintHierarchy)
+		{
+			if (const UAnimBlueprintGeneratedClass* AnimBlueprintClass = Cast<UAnimBlueprintGeneratedClass>(Blueprint))
+			{
+				const TWeakObjectPtr<const UEdGraphNode>* Node = AnimBlueprintClass->AnimBlueprintDebugData.NodePropertyIndexToNodeMap.Find(ReversedPropertyIndex);
+				if (Node)
+				{
+					return Node->Get();
+				}
+			}
+		}
+	}
+
+	return nullptr;
 }
 
 #endif // WITH_EDITORONLY_DATA
+
+void UAnimBlueprintGeneratedClass::BuildConstantProperties()
+{
+	ConstantProperties.Empty();
+	ConstantSubsystemProperties.Empty();
+
+	if(UScriptStruct* ConstantsStruct = GetSparseClassDataStruct())
+	{
+		// Make sure we create the data area up-front here
+		GetOrCreateSparseClassData();
+
+		// We must init sparse class data here once the class has been serialized
+		for (TFieldIterator<FProperty> It(ConstantsStruct); It; ++It)
+		{
+			if (FStructProperty* StructProp = CastField<FStructProperty>(*It))
+			{
+				if(StructProp->Struct->IsChildOf(FAnimSubsystem::StaticStruct()))
+				{
+					ConstantSubsystemProperties.Add(StructProp);
+				}
+			}
+			
+			ConstantProperties.Add(*It);
+		}
+	}
+}
+
+int32 UAnimBlueprintGeneratedClass::GetAnimNodePropertyIndex(const UScriptStruct* InNodeType, FName InPropertyName) const
+{
+	const FAnimNodeStructData* ScriptStructDataPtr = NodeTypeMap.Find(InNodeType);
+	if(ScriptStructDataPtr != nullptr)
+	{
+		return ScriptStructDataPtr->GetPropertyIndex(InPropertyName);
+	}
+
+	return INDEX_NONE;
+}
+
+int32 UAnimBlueprintGeneratedClass::GetAnimNodePropertyCount(const UScriptStruct* InNodeType) const
+{
+	const FAnimNodeStructData* ScriptStructDataPtr = NodeTypeMap.Find(InNodeType);
+	if(ScriptStructDataPtr != nullptr)
+	{
+		return ScriptStructDataPtr->GetNumProperties();
+	}
+	
+	return 0;
+}
\ No newline at end of file
diff --git a/Engine/Source/Runtime/Engine/Private/Animation/AnimClassData.cpp b/Engine/Source/Runtime/Engine/Private/Animation/AnimClassData.cpp
--- a/Engine/Source/Runtime/Engine/Private/Animation/AnimClassData.cpp
+++ b/Engine/Source/Runtime/Engine/Private/Animation/AnimClassData.cpp
@@ -2,7 +2,7 @@
 
 #include "Animation/AnimClassData.h"
 #include "Animation/AnimNode_Root.h"
-#include "Animation/AnimNode_LinkedInputPose.h"
+#include "PropertyAccess.h"
 
 void UAnimClassData::DynamicClassInitialization(UDynamicClass* InDynamicClass)
 {
@@ -28,12 +28,6 @@ void UAnimClassData::DynamicClassInitialization(UDynamicClass* InDynamicClass)
 		Algo::Transform(AnimBlueprintFunctionData[FunctionIndex].InputProperties, AnimBlueprintFunctions[FunctionIndex].InputProperties, [](const TFieldPath<FProperty>& InPropertyPath){ return InPropertyPath.Get(); });
 		Algo::Transform(AnimBlueprintFunctionData[FunctionIndex].InputPoseNodeProperties, AnimBlueprintFunctions[FunctionIndex].InputPoseNodeProperties, ResolveTransform);
 	}
-
-	// Init property access library
-	PropertyAccess::PostLoadLibrary(PropertyAccessLibrary);
-
-	// Init exposed value handlers
-	FExposedValueHandler::DynamicClassInitialization(EvaluateGraphExposedInputs, InDynamicClass);
 }
 
 #if WITH_EDITOR
@@ -77,9 +71,13 @@ void UAnimClassData::CopyFrom(UAnimBlueprintGeneratedClass* AnimClass)
 	ResolvedInitializationNodeProperties = AnimClass->GetInitializationNodeProperties();
 
 	SyncGroupNames = AnimClass->GetSyncGroupNames();
-	EvaluateGraphExposedInputs = AnimClass->GetExposedValueHandlers();
 	GraphNameAssetPlayers = AnimClass->GetGraphAssetPlayerInformation();
 	GraphBlendOptions = AnimClass->GetGraphBlendOptions();
-	PropertyAccessLibrary = AnimClass->GetPropertyAccessLibrary();
 }
-#endif // WITH_EDITOR
\ No newline at end of file
+#endif // WITH_EDITOR
+
+TArrayView<const FAnimNodeData> UAnimClassData::GetNodeData() const
+{
+	TArray<FAnimNodeData> Dummy;
+	return Dummy;
+}
\ No newline at end of file
diff --git a/Engine/Source/Runtime/Engine/Private/Animation/AnimInstance.cpp b/Engine/Source/Runtime/Engine/Private/Animation/AnimInstance.cpp
--- a/Engine/Source/Runtime/Engine/Private/Animation/AnimInstance.cpp
+++ b/Engine/Source/Runtime/Engine/Private/Animation/AnimInstance.cpp
@@ -30,6 +30,7 @@
 #include "Animation/AnimNode_LinkedInputPose.h"
 #include "Animation/AnimNode_LinkedAnimLayer.h"
 #include "UObject/UE5MainStreamObjectVersion.h"
+#include "Animation/AnimSubsystem.h"
 
 /** Anim stats */
 
@@ -502,15 +503,17 @@ PRAGMA_ENABLE_DEPRECATION_WARNINGS
 		CSV_SCOPED_TIMING_STAT(Animation, BlueprintUpdate);
 		BlueprintUpdateAnimation(DeltaSeconds);
 	}
-
-	// Perform property access copies after the event graph work - this allows properties that are transformed and
-	// accessed via BP functions to be correctly transformed.
+	
 	if (IAnimClassInterface* AnimBlueprintClass = IAnimClassInterface::GetFromClass(GetClass()))
 	{
-		// Process internal batched property copies
-		PropertyAccess::ProcessCopies(this, AnimBlueprintClass->GetPropertyAccessLibrary(), EPropertyAccessCopyBatch::ExternalBatched);
+		AnimBlueprintClass->ForEachSubsystem(this, [this, DeltaSeconds](const FAnimSubsystemInstanceContext& InContext)
+		{
+			FAnimSubsystemUpdateContext Context(InContext, this, DeltaSeconds);
+			InContext.Subsystem.OnUpdate(Context);
+			return EAnimSubsystemEnumeration::Continue;
+		});
 	}
-	
+
 	// Determine whether or not the animation should be immediately updated according to current state
 	const bool bWantsImmediateUpdate = bNeedsValidRootMotion || NeedsImmediateUpdate(DeltaSeconds);
 
@@ -3381,22 +3384,22 @@ float UAnimInstance::GetRelevantAnimTimeFraction(int32 MachineIndex, int32 State
 	return GetProxyOnAnyThread<FAnimInstanceProxy>().GetRelevantAnimTimeFraction(MachineIndex, StateIndex);
 }
 
-FAnimNode_StateMachine* UAnimInstance::GetStateMachineInstance(int32 MachineIndex)
+const FAnimNode_StateMachine* UAnimInstance::GetStateMachineInstance(int32 MachineIndex) const
 {
 	return GetProxyOnAnyThread<FAnimInstanceProxy>().GetStateMachineInstance(MachineIndex);
 }
 
-FAnimNode_StateMachine* UAnimInstance::GetStateMachineInstanceFromName(FName MachineName)
+const FAnimNode_StateMachine* UAnimInstance::GetStateMachineInstanceFromName(FName MachineName) const
 {
 	return GetProxyOnAnyThread<FAnimInstanceProxy>().GetStateMachineInstanceFromName(MachineName);
 }
 
-const FBakedAnimationStateMachine* UAnimInstance::GetStateMachineInstanceDesc(FName MachineName)
+const FBakedAnimationStateMachine* UAnimInstance::GetStateMachineInstanceDesc(FName MachineName) const
 {
 	return GetProxyOnAnyThread<FAnimInstanceProxy>().GetStateMachineInstanceDesc(MachineName);
 }
 
-int32 UAnimInstance::GetStateMachineIndex(FName MachineName)
+int32 UAnimInstance::GetStateMachineIndex(FName MachineName) const
 {
 	return GetProxyOnAnyThread<FAnimInstanceProxy>().GetStateMachineIndex(MachineName);
 }
@@ -3411,17 +3414,22 @@ const FBakedAnimationStateMachine* UAnimInstance::GetMachineDescription(IAnimCla
 	return GetProxyOnGameThread<FAnimInstanceProxy>().GetMachineDescription(AnimBlueprintClass, MachineInstance);
 }
 
-int32 UAnimInstance::GetInstanceAssetPlayerIndex(FName MachineName, FName StateName, FName AssetName)
+int32 UAnimInstance::GetInstanceAssetPlayerIndex(FName MachineName, FName StateName, FName AssetName) const
 {
 	return GetProxyOnGameThread<FAnimInstanceProxy>().GetInstanceAssetPlayerIndex(MachineName, StateName, AssetName);
 }
 
-TArray<FAnimNode_AssetPlayerBase*> UAnimInstance::GetInstanceAssetPlayers(const FName& GraphName)
+TArray<const FAnimNode_AssetPlayerBase*> UAnimInstance::GetInstanceAssetPlayers(const FName& GraphName) const
 {
 	return GetProxyOnAnyThread<FAnimInstanceProxy>().GetInstanceAssetPlayers(GraphName);
 }
 
-FAnimNode_AssetPlayerBase* UAnimInstance::GetRelevantAssetPlayerFromState(int32 MachineIndex, int32 StateIndex)
+TArray<FAnimNode_AssetPlayerBase*> UAnimInstance::GetMutableInstanceAssetPlayers(const FName& GraphName)
+{
+	return GetProxyOnAnyThread<FAnimInstanceProxy>().GetMutableInstanceAssetPlayers(GraphName);
+}
+
+const FAnimNode_AssetPlayerBase* UAnimInstance::GetRelevantAssetPlayerFromState(int32 MachineIndex, int32 StateIndex) const
 {
 	return GetProxyOnGameThread<FAnimInstanceProxy>().GetRelevantAssetPlayerFromState(MachineIndex, StateIndex);
 }
