Anim node data/compiler refactor

diff --git a/Engine/Source/Runtime/Engine/Private/Animation/AnimInstanceProxy.cpp b/Engine/Source/Runtime/Engine/Private/Animation/AnimInstanceProxy.cpp
--- a/Engine/Source/Runtime/Engine/Private/Animation/AnimInstanceProxy.cpp
+++ b/Engine/Source/Runtime/Engine/Private/Animation/AnimInstanceProxy.cpp
@@ -59,10 +59,13 @@ void FAnimInstanceProxy::UpdateAnimationNode_WithRoot(const FAnimationUpdateCont
 		InRootNode->Update_AnyThread(InContext);
 
 		// We've updated the graph, now update the fractured saved pose sections
-		TArray<FAnimNode_SaveCachedPose*>& SavedPoseQueue = SavedPoseQueueMap.FindChecked(InLayerName);
-		for(FAnimNode_SaveCachedPose* PoseNode : SavedPoseQueue)
+		// Note SavedPoseQueueMap can be empty if we have no cached poses in use
+		if(TArray<FAnimNode_SaveCachedPose*>* SavedPoseQueue = SavedPoseQueueMap.Find(InLayerName))
 		{
-			PoseNode->PostGraphUpdate();
+			for(FAnimNode_SaveCachedPose* PoseNode : *SavedPoseQueue)
+			{
+				PoseNode->PostGraphUpdate();
+			}
 		}
 	}
 }
@@ -890,9 +893,9 @@ float FAnimInstanceProxy::CalcSlotMontageLocalWeight(const FName& SlotNodeName)
 	return out_SlotNodeLocalWeight;
 }
 
-FAnimNode_Base* FAnimInstanceProxy::GetCheckedNodeFromIndexUntyped(int32 NodeIdx, UScriptStruct* RequiredStructType)
+const FAnimNode_Base* FAnimInstanceProxy::GetCheckedNodeFromIndexUntyped(int32 NodeIdx, UScriptStruct* RequiredStructType) const
 {
-	FAnimNode_Base* NodePtr = nullptr;
+	const FAnimNode_Base* NodePtr = nullptr;
 	if (AnimClassInterface)
 	{
 		const TArray<FStructProperty*>& AnimNodeProperties = AnimClassInterface->GetAnimNodeProperties();
@@ -922,9 +925,14 @@ FAnimNode_Base* FAnimInstanceProxy::GetCheckedNodeFromIndexUntyped(int32 NodeIdx
 	return NodePtr;
 }
 
-FAnimNode_Base* FAnimInstanceProxy::GetNodeFromIndexUntyped(int32 NodeIdx, UScriptStruct* RequiredStructType)
+FAnimNode_Base* FAnimInstanceProxy::GetCheckedMutableNodeFromIndexUntyped(int32 NodeIdx, UScriptStruct* RequiredStructType)
+{
+	return const_cast<FAnimNode_Base*>(GetCheckedNodeFromIndexUntyped(NodeIdx, RequiredStructType));
+}
+
+const FAnimNode_Base* FAnimInstanceProxy::GetNodeFromIndexUntyped(int32 NodeIdx, UScriptStruct* RequiredStructType) const
 {
-	FAnimNode_Base* NodePtr = nullptr;
+	const FAnimNode_Base* NodePtr = nullptr;
 	if (AnimClassInterface)
 	{
 		const TArray<FStructProperty*>& AnimNodeProperties = AnimClassInterface->GetAnimNodeProperties();
@@ -944,6 +952,11 @@ FAnimNode_Base* FAnimInstanceProxy::GetNodeFromIndexUntyped(int32 NodeIdx, UScri
 	return NodePtr;
 }
 
+FAnimNode_Base* FAnimInstanceProxy::GetMutableNodeFromIndexUntyped(int32 NodeIdx, UScriptStruct* RequiredStructType)
+{
+	return const_cast<FAnimNode_Base*>(GetNodeFromIndexUntyped(NodeIdx, RequiredStructType));
+}
+
 void FAnimInstanceProxy::RecalcRequiredBones(USkeletalMeshComponent* Component, UObject* Asset)
 {
 	DECLARE_SCOPE_HIERARCHICAL_COUNTER_FUNC()
@@ -1006,10 +1019,14 @@ void FAnimInstanceProxy::UpdateAnimation()
 
 	FMemMark Mark(FMemStack::Get());
 
-	// Process internal batched property copies
 	if(AnimClassInterface)
 	{
-		PropertyAccess::ProcessCopies(GetAnimInstanceObject(), AnimClassInterface->GetPropertyAccessLibrary(), EPropertyAccessCopyBatch::InternalBatched);
+		AnimClassInterface->ForEachSubsystem(GetAnimInstanceObject(), [this](const FAnimSubsystemInstanceContext& InContext)
+		{
+		    FAnimSubsystemParallelUpdateContext Context(InContext, *this, CurrentDeltaSeconds);
+		    InContext.Subsystem.OnParallelUpdate(Context);
+		    return EAnimSubsystemEnumeration::Continue;
+		});
 	}
 
 #if WITH_EDITORONLY_DATA
@@ -1024,7 +1041,7 @@ void FAnimInstanceProxy::UpdateAnimation()
 
 	if(AnimClassInterface && AnimClassInterface->GetAnimBlueprintFunctions().Num() > 0)
 	{
-		Context = Context.WithNodeId(AnimClassInterface->GetAnimBlueprintFunctions()[0].OutputPoseNodeIndex);
+		Context.SetNodeId(AnimClassInterface->GetAnimBlueprintFunctions()[0].OutputPoseNodeIndex);
 	}
 
 	UpdateAnimation_WithRoot(Context, RootNode, NAME_AnimGraph);
@@ -1876,10 +1893,12 @@ void FAnimInstanceProxy::GatherDebugData_WithRoot(FNodeDebugData& DebugData, FAn
 	}
 
 	// Gather debug data for Cached Poses.
-	TArray<FAnimNode_SaveCachedPose*>& SavedPoseQueue = SavedPoseQueueMap.FindChecked(InLayerName);
-	for (FAnimNode_SaveCachedPose* PoseNode : SavedPoseQueue)
+	if(TArray<FAnimNode_SaveCachedPose*>* SavedPoseQueue = SavedPoseQueueMap.Find(InLayerName))
 	{
-		PoseNode->GatherDebugData(DebugData);
+		for (FAnimNode_SaveCachedPose* PoseNode : *SavedPoseQueue)
+		{
+			PoseNode->GatherDebugData(DebugData);
+		}
 	}
 }
 
@@ -2013,9 +2032,9 @@ void FAnimInstanceProxy::AnimDrawDebugPoint(const FVector& Loc, float Size, cons
 }
 #endif // ENABLE_ANIM_DRAW_DEBUG
 
-float FAnimInstanceProxy::GetInstanceAssetPlayerLength(int32 AssetPlayerIndex)
+float FAnimInstanceProxy::GetInstanceAssetPlayerLength(int32 AssetPlayerIndex) const
 {
-	if(FAnimNode_AssetPlayerBase* PlayerNode = GetNodeFromIndex<FAnimNode_AssetPlayerBase>(AssetPlayerIndex))
+	if(const FAnimNode_AssetPlayerBase* PlayerNode = GetNodeFromIndex<FAnimNode_AssetPlayerBase>(AssetPlayerIndex))
 	{
 		return PlayerNode->GetCurrentAssetLength();
 	}
@@ -2023,9 +2042,9 @@ float FAnimInstanceProxy::GetInstanceAssetPlayerLength(int32 AssetPlayerIndex)
 	return 0.0f;
 }
 
-float FAnimInstanceProxy::GetInstanceAssetPlayerTime(int32 AssetPlayerIndex)
+float FAnimInstanceProxy::GetInstanceAssetPlayerTime(int32 AssetPlayerIndex) const
 {
-	if(FAnimNode_AssetPlayerBase* PlayerNode = GetNodeFromIndex<FAnimNode_AssetPlayerBase>(AssetPlayerIndex))
+	if(const FAnimNode_AssetPlayerBase* PlayerNode = GetNodeFromIndex<FAnimNode_AssetPlayerBase>(AssetPlayerIndex))
 	{
 		return PlayerNode->GetCurrentAssetTimePlayRateAdjusted();
 	}
@@ -2033,9 +2052,9 @@ float FAnimInstanceProxy::GetInstanceAssetPlayerTime(int32 AssetPlayerIndex)
 	return 0.0f;
 }
 
-float FAnimInstanceProxy::GetInstanceAssetPlayerTimeFraction(int32 AssetPlayerIndex)
+float FAnimInstanceProxy::GetInstanceAssetPlayerTimeFraction(int32 AssetPlayerIndex) const
 {
-	if(FAnimNode_AssetPlayerBase* PlayerNode = GetNodeFromIndex<FAnimNode_AssetPlayerBase>(AssetPlayerIndex))
+	if(const FAnimNode_AssetPlayerBase* PlayerNode = GetNodeFromIndex<FAnimNode_AssetPlayerBase>(AssetPlayerIndex))
 	{
 		float Length = PlayerNode->GetCurrentAssetLength();
 
@@ -2048,9 +2067,9 @@ float FAnimInstanceProxy::GetInstanceAssetPlayerTimeFraction(int32 AssetPlayerIn
 	return 0.0f;
 }
 
-float FAnimInstanceProxy::GetInstanceAssetPlayerTimeFromEndFraction(int32 AssetPlayerIndex)
+float FAnimInstanceProxy::GetInstanceAssetPlayerTimeFromEndFraction(int32 AssetPlayerIndex) const
 {
-	if(FAnimNode_AssetPlayerBase* PlayerNode = GetNodeFromIndex<FAnimNode_AssetPlayerBase>(AssetPlayerIndex))
+	if(const FAnimNode_AssetPlayerBase* PlayerNode = GetNodeFromIndex<FAnimNode_AssetPlayerBase>(AssetPlayerIndex))
 	{
 		float Length = PlayerNode->GetCurrentAssetLength();
 
@@ -2063,9 +2082,9 @@ float FAnimInstanceProxy::GetInstanceAssetPlayerTimeFromEndFraction(int32 AssetP
 	return 1.0f;
 }
 
-float FAnimInstanceProxy::GetInstanceAssetPlayerTimeFromEnd(int32 AssetPlayerIndex)
+float FAnimInstanceProxy::GetInstanceAssetPlayerTimeFromEnd(int32 AssetPlayerIndex) const
 {
-	if(FAnimNode_AssetPlayerBase* PlayerNode = GetNodeFromIndex<FAnimNode_AssetPlayerBase>(AssetPlayerIndex))
+	if(const FAnimNode_AssetPlayerBase* PlayerNode = GetNodeFromIndex<FAnimNode_AssetPlayerBase>(AssetPlayerIndex))
 	{
 		return PlayerNode->GetCurrentAssetLength() - PlayerNode->GetCurrentAssetTimePlayRateAdjusted();
 	}
@@ -2073,9 +2092,9 @@ float FAnimInstanceProxy::GetInstanceAssetPlayerTimeFromEnd(int32 AssetPlayerInd
 	return MAX_flt;
 }
 
-float FAnimInstanceProxy::GetInstanceMachineWeight(int32 MachineIndex)
+float FAnimInstanceProxy::GetInstanceMachineWeight(int32 MachineIndex) const
 {
-	if (FAnimNode_StateMachine* MachineInstance = GetStateMachineInstance(MachineIndex))
+	if (const FAnimNode_StateMachine* MachineInstance = GetStateMachineInstance(MachineIndex))
 	{
 		return GetRecordedMachineWeight(MachineInstance->StateMachineIndexInClass);
 	}
@@ -2083,9 +2102,9 @@ float FAnimInstanceProxy::GetInstanceMachineWeight(int32 MachineIndex)
 	return 0.0f;
 }
 
-float FAnimInstanceProxy::GetInstanceStateWeight(int32 MachineIndex, int32 StateIndex)
+float FAnimInstanceProxy::GetInstanceStateWeight(int32 MachineIndex, int32 StateIndex) const
 {
-	if(FAnimNode_StateMachine* MachineInstance = GetStateMachineInstance(MachineIndex))
+	if(const FAnimNode_StateMachine* MachineInstance = GetStateMachineInstance(MachineIndex))
 	{
 		return GetRecordedStateWeight(MachineInstance->StateMachineIndexInClass, StateIndex);
 	}
@@ -2093,9 +2112,9 @@ float FAnimInstanceProxy::GetInstanceStateWeight(int32 MachineIndex, int32 State
 	return 0.0f;
 }
 
-float FAnimInstanceProxy::GetInstanceCurrentStateElapsedTime(int32 MachineIndex)
+float FAnimInstanceProxy::GetInstanceCurrentStateElapsedTime(int32 MachineIndex) const
 {
-	if(FAnimNode_StateMachine* MachineInstance = GetStateMachineInstance(MachineIndex))
+	if(const FAnimNode_StateMachine* MachineInstance = GetStateMachineInstance(MachineIndex))
 	{
 		return MachineInstance->GetCurrentStateElapsedTime();
 	}
@@ -2103,9 +2122,9 @@ float FAnimInstanceProxy::GetInstanceCurrentStateElapsedTime(int32 MachineIndex)
 	return 0.0f;
 }
 
-float FAnimInstanceProxy::GetInstanceTransitionCrossfadeDuration(int32 MachineIndex, int32 TransitionIndex)
+float FAnimInstanceProxy::GetInstanceTransitionCrossfadeDuration(int32 MachineIndex, int32 TransitionIndex) const
 {
-	if(FAnimNode_StateMachine* MachineInstance = GetStateMachineInstance(MachineIndex))
+	if(const FAnimNode_StateMachine* MachineInstance = GetStateMachineInstance(MachineIndex))
 	{
 		if(MachineInstance->IsValidTransitionIndex(TransitionIndex))
 		{
@@ -2116,14 +2135,14 @@ float FAnimInstanceProxy::GetInstanceTransitionCrossfadeDuration(int32 MachineIn
 	return 0.0f;
 }
 
-float FAnimInstanceProxy::GetInstanceTransitionTimeElapsed(int32 MachineIndex, int32 TransitionIndex)
+float FAnimInstanceProxy::GetInstanceTransitionTimeElapsed(int32 MachineIndex, int32 TransitionIndex) const
 {
 	// Just an alias for readability in the anim graph
-	if(FAnimNode_StateMachine* MachineInstance = GetStateMachineInstance(MachineIndex))
+	if(const FAnimNode_StateMachine* MachineInstance = GetStateMachineInstance(MachineIndex))
 	{
 		if(MachineInstance->IsValidTransitionIndex(TransitionIndex))
 		{
-			for(FAnimationActiveTransitionEntry& ActiveTransition : MachineInstance->ActiveTransitionArray)
+			for(const FAnimationActiveTransitionEntry& ActiveTransition : MachineInstance->ActiveTransitionArray)
 			{
 				if(ActiveTransition.SourceTransitionIndices.Contains(TransitionIndex))
 				{
@@ -2136,13 +2155,13 @@ float FAnimInstanceProxy::GetInstanceTransitionTimeElapsed(int32 MachineIndex, i
 	return 0.0f;
 }
 
-float FAnimInstanceProxy::GetInstanceTransitionTimeElapsedFraction(int32 MachineIndex, int32 TransitionIndex)
+float FAnimInstanceProxy::GetInstanceTransitionTimeElapsedFraction(int32 MachineIndex, int32 TransitionIndex) const
 {
-	if(FAnimNode_StateMachine* MachineInstance = GetStateMachineInstance(MachineIndex))
+	if(const FAnimNode_StateMachine* MachineInstance = GetStateMachineInstance(MachineIndex))
 	{
 		if(MachineInstance->IsValidTransitionIndex(TransitionIndex))
 		{
-			for(FAnimationActiveTransitionEntry& ActiveTransition : MachineInstance->ActiveTransitionArray)
+			for(const FAnimationActiveTransitionEntry& ActiveTransition : MachineInstance->ActiveTransitionArray)
 			{
 				if(ActiveTransition.SourceTransitionIndices.Contains(TransitionIndex))
 				{
@@ -2155,9 +2174,9 @@ float FAnimInstanceProxy::GetInstanceTransitionTimeElapsedFraction(int32 Machine
 	return 0.0f;
 }
 
-float FAnimInstanceProxy::GetRelevantAnimTimeRemaining(int32 MachineIndex, int32 StateIndex)
+float FAnimInstanceProxy::GetRelevantAnimTimeRemaining(int32 MachineIndex, int32 StateIndex) const
 {
-	if(FAnimNode_AssetPlayerBase* AssetPlayer = GetRelevantAssetPlayerFromState(MachineIndex, StateIndex))
+	if(const FAnimNode_AssetPlayerBase* AssetPlayer = GetRelevantAssetPlayerFromState(MachineIndex, StateIndex))
 	{
 		if(AssetPlayer->GetAnimAsset())
 		{
@@ -2168,9 +2187,9 @@ float FAnimInstanceProxy::GetRelevantAnimTimeRemaining(int32 MachineIndex, int32
 	return MAX_flt;
 }
 
-float FAnimInstanceProxy::GetRelevantAnimTimeRemainingFraction(int32 MachineIndex, int32 StateIndex)
+float FAnimInstanceProxy::GetRelevantAnimTimeRemainingFraction(int32 MachineIndex, int32 StateIndex) const
 {
-	if(FAnimNode_AssetPlayerBase* AssetPlayer = GetRelevantAssetPlayerFromState(MachineIndex, StateIndex))
+	if(const FAnimNode_AssetPlayerBase* AssetPlayer = GetRelevantAssetPlayerFromState(MachineIndex, StateIndex))
 	{
 		if(AssetPlayer->GetAnimAsset())
 		{
@@ -2185,9 +2204,9 @@ float FAnimInstanceProxy::GetRelevantAnimTimeRemainingFraction(int32 MachineInde
 	return 1.0f;
 }
 
-float FAnimInstanceProxy::GetRelevantAnimLength(int32 MachineIndex, int32 StateIndex)
+float FAnimInstanceProxy::GetRelevantAnimLength(int32 MachineIndex, int32 StateIndex) const
 {
-	if(FAnimNode_AssetPlayerBase* AssetPlayer = GetRelevantAssetPlayerFromState(MachineIndex, StateIndex))
+	if(const FAnimNode_AssetPlayerBase* AssetPlayer = GetRelevantAssetPlayerFromState(MachineIndex, StateIndex))
 	{
 		if(AssetPlayer->GetAnimAsset())
 		{
@@ -2198,9 +2217,9 @@ float FAnimInstanceProxy::GetRelevantAnimLength(int32 MachineIndex, int32 StateI
 	return 0.0f;
 }
 
-float FAnimInstanceProxy::GetRelevantAnimTime(int32 MachineIndex, int32 StateIndex)
+float FAnimInstanceProxy::GetRelevantAnimTime(int32 MachineIndex, int32 StateIndex) const
 {
-	if(FAnimNode_AssetPlayerBase* AssetPlayer = GetRelevantAssetPlayerFromState(MachineIndex, StateIndex))
+	if(const FAnimNode_AssetPlayerBase* AssetPlayer = GetRelevantAssetPlayerFromState(MachineIndex, StateIndex))
 	{
 		return AssetPlayer->GetCurrentAssetTimePlayRateAdjusted();
 	}
@@ -2208,9 +2227,9 @@ float FAnimInstanceProxy::GetRelevantAnimTime(int32 MachineIndex, int32 StateInd
 	return 0.0f;
 }
 
-float FAnimInstanceProxy::GetRelevantAnimTimeFraction(int32 MachineIndex, int32 StateIndex)
+float FAnimInstanceProxy::GetRelevantAnimTimeFraction(int32 MachineIndex, int32 StateIndex) const
 {
-	if(FAnimNode_AssetPlayerBase* AssetPlayer = GetRelevantAssetPlayerFromState(MachineIndex, StateIndex))
+	if(const FAnimNode_AssetPlayerBase* AssetPlayer = GetRelevantAssetPlayerFromState(MachineIndex, StateIndex))
 	{
 		float Length = AssetPlayer->GetCurrentAssetLength();
 		if(Length > 0.0f)
@@ -2234,18 +2253,18 @@ const TArray<FMontageEvaluationState>& FAnimInstanceProxy::GetMontageEvaluationD
 	return *MainMontageEvaluationData;
 }
 
-FAnimNode_AssetPlayerBase* FAnimInstanceProxy::GetRelevantAssetPlayerFromState(int32 MachineIndex, int32 StateIndex)
+const FAnimNode_AssetPlayerBase* FAnimInstanceProxy::GetRelevantAssetPlayerFromState(int32 MachineIndex, int32 StateIndex) const
 {
-	FAnimNode_AssetPlayerBase* ResultPlayer = nullptr;
-	if(FAnimNode_StateMachine* MachineInstance = GetStateMachineInstance(MachineIndex))
+	const FAnimNode_AssetPlayerBase* ResultPlayer = nullptr;
+	if(const FAnimNode_StateMachine* MachineInstance = GetStateMachineInstance(MachineIndex))
 	{
 		float MaxWeight = 0.0f;
 		const FBakedAnimationState& State = MachineInstance->GetStateInfo(StateIndex);
 		for(const int32& PlayerIdx : State.PlayerNodeIndices)
 		{
-			if(FAnimNode_AssetPlayerBase* Player = GetNodeFromIndex<FAnimNode_AssetPlayerBase>(PlayerIdx))
+			if(const FAnimNode_AssetPlayerBase* Player = GetNodeFromIndex<FAnimNode_AssetPlayerBase>(PlayerIdx))
 			{
-				if(!Player->bIgnoreForRelevancyTest && Player->GetCachedBlendWeight() > MaxWeight)
+				if(!Player->GetIgnoreForRelevancyTest() && Player->GetCachedBlendWeight() > MaxWeight)
 				{
 					MaxWeight = Player->GetCachedBlendWeight();
 					ResultPlayer = Player;
@@ -2256,7 +2275,7 @@ FAnimNode_AssetPlayerBase* FAnimInstanceProxy::GetRelevantAssetPlayerFromState(i
 	return ResultPlayer;
 }
 
-FAnimNode_StateMachine* FAnimInstanceProxy::GetStateMachineInstance(int32 MachineIndex)
+const FAnimNode_StateMachine* FAnimInstanceProxy::GetStateMachineInstance(int32 MachineIndex) const
 {
 	if (AnimClassInterface)
 	{
@@ -2438,13 +2457,13 @@ void FAnimInstanceProxy::BindNativeDelegates()
 	}
 }
 
-const FBakedAnimationStateMachine* FAnimInstanceProxy::GetMachineDescription(IAnimClassInterface* AnimBlueprintClass, FAnimNode_StateMachine* MachineInstance)
+const FBakedAnimationStateMachine* FAnimInstanceProxy::GetMachineDescription(IAnimClassInterface* AnimBlueprintClass, const FAnimNode_StateMachine* MachineInstance)
 {
 	const TArray<FBakedAnimationStateMachine>& BakedStateMachines = AnimBlueprintClass->GetBakedStateMachines();
 	return BakedStateMachines.IsValidIndex(MachineInstance->StateMachineIndexInClass) ? &(BakedStateMachines[MachineInstance->StateMachineIndexInClass]) : nullptr;
 }
 
-FAnimNode_StateMachine* FAnimInstanceProxy::GetStateMachineInstanceFromName(FName MachineName)
+const FAnimNode_StateMachine* FAnimInstanceProxy::GetStateMachineInstanceFromName(FName MachineName) const
 {
 	if (AnimClassInterface)
 	{
@@ -2454,7 +2473,7 @@ FAnimNode_StateMachine* FAnimInstanceProxy::GetStateMachineInstanceFromName(FNam
 			FStructProperty* Property = AnimNodeProperties[AnimNodeProperties.Num() - 1 - MachineIndex];
 			if (Property && Property->Struct->IsChildOf(FAnimNode_StateMachine::StaticStruct()))
 			{
-				FAnimNode_StateMachine* StateMachine = Property->ContainerPtrToValuePtr<FAnimNode_StateMachine>(AnimInstanceObject);
+				const FAnimNode_StateMachine* StateMachine = Property->ContainerPtrToValuePtr<FAnimNode_StateMachine>(AnimInstanceObject);
 				if (StateMachine)
 				{
 					if (const FBakedAnimationStateMachine* MachineDescription = GetMachineDescription(AnimClassInterface, StateMachine))
@@ -2472,7 +2491,7 @@ FAnimNode_StateMachine* FAnimInstanceProxy::GetStateMachineInstanceFromName(FNam
 	return nullptr;
 }
 
-const FBakedAnimationStateMachine* FAnimInstanceProxy::GetStateMachineInstanceDesc(FName MachineName)
+const FBakedAnimationStateMachine* FAnimInstanceProxy::GetStateMachineInstanceDesc(FName MachineName) const
 {
 	if (AnimClassInterface)
 	{
@@ -2500,7 +2519,7 @@ const FBakedAnimationStateMachine* FAnimInstanceProxy::GetStateMachineInstanceDe
 	return nullptr;
 }
 
-int32 FAnimInstanceProxy::GetStateMachineIndex(FName MachineName)
+int32 FAnimInstanceProxy::GetStateMachineIndex(FName MachineName) const
 {
 	if (AnimClassInterface)
 	{
@@ -2510,7 +2529,7 @@ int32 FAnimInstanceProxy::GetStateMachineIndex(FName MachineName)
 			FStructProperty* Property = AnimNodeProperties[AnimNodeProperties.Num() - 1 - MachineIndex];
 			if(Property && Property->Struct->IsChildOf(FAnimNode_StateMachine::StaticStruct()))
 			{
-				FAnimNode_StateMachine* StateMachine = Property->ContainerPtrToValuePtr<FAnimNode_StateMachine>(AnimInstanceObject);
+				const FAnimNode_StateMachine* StateMachine = Property->ContainerPtrToValuePtr<FAnimNode_StateMachine>(AnimInstanceObject);
 				if(StateMachine)
 				{
 					if (const FBakedAnimationStateMachine* MachineDescription = GetMachineDescription(AnimClassInterface, StateMachine))
@@ -2528,7 +2547,7 @@ int32 FAnimInstanceProxy::GetStateMachineIndex(FName MachineName)
 	return INDEX_NONE;
 }
 
-void FAnimInstanceProxy::GetStateMachineIndexAndDescription(FName InMachineName, int32& OutMachineIndex, const FBakedAnimationStateMachine** OutMachineDescription)
+void FAnimInstanceProxy::GetStateMachineIndexAndDescription(FName InMachineName, int32& OutMachineIndex, const FBakedAnimationStateMachine** OutMachineDescription) const
 {
 	if (AnimClassInterface)
 	{
@@ -2538,7 +2557,7 @@ void FAnimInstanceProxy::GetStateMachineIndexAndDescription(FName InMachineName,
 			FStructProperty* Property = AnimNodeProperties[AnimNodeProperties.Num() - 1 - MachineIndex];
 			if (Property && Property->Struct->IsChildOf(FAnimNode_StateMachine::StaticStruct()))
 			{
-				FAnimNode_StateMachine* StateMachine = Property->ContainerPtrToValuePtr<FAnimNode_StateMachine>(AnimInstanceObject);
+				const FAnimNode_StateMachine* StateMachine = Property->ContainerPtrToValuePtr<FAnimNode_StateMachine>(AnimInstanceObject);
 				if (StateMachine)
 				{
 					if (const FBakedAnimationStateMachine* MachineDescription = GetMachineDescription(AnimClassInterface, StateMachine))
@@ -2565,7 +2584,7 @@ void FAnimInstanceProxy::GetStateMachineIndexAndDescription(FName InMachineName,
 	}
 }
 
-int32 FAnimInstanceProxy::GetInstanceAssetPlayerIndex(FName MachineName, FName StateName, FName AssetName)
+int32 FAnimInstanceProxy::GetInstanceAssetPlayerIndex(FName MachineName, FName StateName, FName AssetName) const
 {
 	if (AnimClassInterface)
 	{
@@ -2583,7 +2602,7 @@ int32 FAnimInstanceProxy::GetInstanceAssetPlayerIndex(FName MachineName, FName S
 						FStructProperty* AssetPlayerProperty = AnimNodeProperties[AnimNodeProperties.Num() - 1 - State.PlayerNodeIndices[PlayerIndex]];
 						if(AssetPlayerProperty && AssetPlayerProperty->Struct->IsChildOf(FAnimNode_AssetPlayerBase::StaticStruct()))
 						{
-							FAnimNode_AssetPlayerBase* AssetPlayer = AssetPlayerProperty->ContainerPtrToValuePtr<FAnimNode_AssetPlayerBase>(AnimInstanceObject);
+							const FAnimNode_AssetPlayerBase* AssetPlayer = AssetPlayerProperty->ContainerPtrToValuePtr<FAnimNode_AssetPlayerBase>(AnimInstanceObject);
 							if(AssetPlayer)
 							{
 								if(AssetName == NAME_None || AssetPlayer->GetAnimAsset()->GetFName() == AssetName)
@@ -2682,7 +2701,30 @@ void FAnimInstanceProxy::TraceMontageEvaluationData(const FAnimationUpdateContex
 }
 #endif
 
-TArray<FAnimNode_AssetPlayerBase*> FAnimInstanceProxy::GetInstanceAssetPlayers(const FName& GraphName)
+TArray<const FAnimNode_AssetPlayerBase*> FAnimInstanceProxy::GetInstanceAssetPlayers(const FName& GraphName) const
+{
+	TArray<const FAnimNode_AssetPlayerBase*> Nodes;
+
+	// Retrieve all asset player nodes from the (named) Animation Layer Graph
+	if (AnimClassInterface)
+	{
+		const TMap<FName, FGraphAssetPlayerInformation>& GrapInformationMap = AnimClassInterface->GetGraphAssetPlayerInformation();
+		if (const FGraphAssetPlayerInformation* Information = GrapInformationMap.Find(GraphName))
+		{
+			for (const int32& NodeIndex : Information->PlayerNodeIndices)
+			{
+				if (const FAnimNode_AssetPlayerBase* Node = GetNodeFromIndex<FAnimNode_AssetPlayerBase>(NodeIndex))
+				{
+					Nodes.Add(Node);
+				}
+			}
+		}
+	}
+
+	return Nodes;
+}
+
+TArray<FAnimNode_AssetPlayerBase*> FAnimInstanceProxy::GetMutableInstanceAssetPlayers(const FName& GraphName)
 {
 	TArray<FAnimNode_AssetPlayerBase*> Nodes;
 
@@ -2694,7 +2736,7 @@ TArray<FAnimNode_AssetPlayerBase*> FAnimInstanceProxy::GetInstanceAssetPlayers(c
 		{
 			for (const int32& NodeIndex : Information->PlayerNodeIndices)
 			{
-				if (FAnimNode_AssetPlayerBase* Node = GetNodeFromIndex<FAnimNode_AssetPlayerBase>(NodeIndex))
+				if (FAnimNode_AssetPlayerBase* Node = GetMutableNodeFromIndex<FAnimNode_AssetPlayerBase>(NodeIndex))
 				{
 					Nodes.Add(Node);
 				}
diff --git a/Engine/Source/Runtime/Engine/Private/Animation/AnimNodeBase.cpp b/Engine/Source/Runtime/Engine/Private/Animation/AnimNodeBase.cpp
--- a/Engine/Source/Runtime/Engine/Private/Animation/AnimNodeBase.cpp
+++ b/Engine/Source/Runtime/Engine/Private/Animation/AnimNodeBase.cpp
@@ -5,8 +5,8 @@
 #include "Animation/AnimInstanceProxy.h"
 #include "Animation/AnimTrace.h"
 #include "UObject/CoreObjectVersion.h"
-#include "PropertyAccess.h"
 #include "Animation/AnimAttributes.h"
+#include "Animation/AnimSubsystem_Base.h"
 
 /////////////////////////////////////////////////////
 // FAnimationBaseContext
@@ -40,6 +40,11 @@ IAnimClassInterface* FAnimationBaseContext::GetAnimClass() const
 	return AnimInstanceProxy ? AnimInstanceProxy->GetAnimClassInterface() : nullptr;
 }
 
+UObject* FAnimationBaseContext::GetAnimInstanceObject() const 
+{
+	return AnimInstanceProxy ? AnimInstanceProxy->GetAnimInstanceObject() : nullptr;
+}
+
 #if WITH_EDITORONLY_DATA
 UAnimBlueprint* FAnimationBaseContext::GetAnimBlueprint() const
 {
@@ -151,7 +156,7 @@ void FAnimNode_Base::ResetDynamics(ETeleportType InTeleportType)
 void FPoseLinkBase::AttemptRelink(const FAnimationBaseContext& Context)
 {
 	// Do the linkage
-	if ((LinkedNode == NULL) && (LinkID != INDEX_NONE))
+	if ((LinkedNode == nullptr) && (LinkID != INDEX_NONE))
 	{
 		IAnimClassInterface* AnimBlueprintClass = Context.GetAnimClass();
 		check(AnimBlueprintClass);
@@ -167,27 +172,27 @@ void FPoseLinkBase::AttemptRelink(const FAnimationBaseContext& Context)
 	}
 }
 
-void FPoseLinkBase::Initialize(const FAnimationInitializeContext& Context)
+void FPoseLinkBase::Initialize(const FAnimationInitializeContext& InContext)
 {
 #if DO_CHECK
 	checkf(!bProcessed, TEXT("Initialize already in progress, circular link for AnimInstance [%s] Blueprint [%s]"), \
-		*Context.AnimInstanceProxy->GetAnimInstanceName(), *GetFullNameSafe(IAnimClassInterface::GetActualAnimClass(Context.AnimInstanceProxy->GetAnimClassInterface())));
+		*InContext.AnimInstanceProxy->GetAnimInstanceName(), *GetFullNameSafe(IAnimClassInterface::GetActualAnimClass(InContext.AnimInstanceProxy->GetAnimClassInterface())));
 	TGuardValue<bool> CircularGuard(bProcessed, true);
 #endif
 
-	AttemptRelink(Context);
+	AttemptRelink(InContext);
 
 #if ENABLE_ANIMGRAPH_TRAVERSAL_DEBUG
-	InitializationCounter.SynchronizeWith(Context.AnimInstanceProxy->GetInitializationCounter());
+	InitializationCounter.SynchronizeWith(InContext.AnimInstanceProxy->GetInitializationCounter());
 
 	// Initialization will require update to be called before an evaluate.
 	UpdateCounter.Reset();
 #endif
 
 	// Do standard initialization
-	if (LinkedNode != NULL)
+	if (LinkedNode != nullptr)
 	{
-		FAnimationInitializeContext LinkContext(Context);
+		FAnimationInitializeContext LinkContext(InContext);
 		LinkContext.SetNodeId(LinkID);
 		TRACE_SCOPED_ANIM_NODE(LinkContext);
 		LinkedNode->Initialize_AnyThread(LinkContext);
@@ -225,40 +230,41 @@ FAnimNode_Base* FPoseLinkBase::GetLinkNode()
 	return LinkedNode;
 }
 
-const FExposedValueHandler& FAnimNode_Base::GetEvaluateGraphExposedInputs()
+const FExposedValueHandler& FAnimNode_Base::GetEvaluateGraphExposedInputs() const
 {
-	// Inverting control (entering via the immutable data rather than the mutable data) would allow
-	// us to remove this static local. Would also allow us to remove the vtable from FAnimNode_Base.
-	static const FExposedValueHandler Default;
-	if(ExposedValueHandler)
+	if(NodeData)
 	{
-		return *ExposedValueHandler;
+		const int32 NodeIndex = NodeData->GetNodeIndex();
+		return NodeData->GetAnimClassInterface().GetSubsystem<FAnimSubsystem_Base>().GetExposedValueHandlers()[NodeIndex];
 	}
 	else
 	{
+		// Inverting control (entering via the immutable data rather than the mutable data) would allow
+		// us to remove this static local. Would also allow us to remove the vtable from FAnimNode_Base.
+		static const FExposedValueHandler Default;	
 		return Default;
 	}
 }
 
-void FPoseLinkBase::CacheBones(const FAnimationCacheBonesContext& Context) 
+void FPoseLinkBase::CacheBones(const FAnimationCacheBonesContext& InContext) 
 {
 #if DO_CHECK
 	checkf( !bProcessed, TEXT( "CacheBones already in progress, circular link for AnimInstance [%s] Blueprint [%s]" ), \
-		*Context.AnimInstanceProxy->GetAnimInstanceName(), *GetFullNameSafe(IAnimClassInterface::GetActualAnimClass(Context.AnimInstanceProxy->GetAnimClassInterface())));
+		*InContext.AnimInstanceProxy->GetAnimInstanceName(), *GetFullNameSafe(IAnimClassInterface::GetActualAnimClass(InContext.AnimInstanceProxy->GetAnimClassInterface())));
 	TGuardValue<bool> CircularGuard(bProcessed, true);
 #endif
 
 #if ENABLE_ANIMGRAPH_TRAVERSAL_DEBUG
-	CachedBonesCounter.SynchronizeWith(Context.AnimInstanceProxy->GetCachedBonesCounter());
+	CachedBonesCounter.SynchronizeWith(InContext.AnimInstanceProxy->GetCachedBonesCounter());
 #endif
 
-	if (LinkedNode != NULL)
+	if (LinkedNode != nullptr)
 	{
-		LinkedNode->CacheBones_AnyThread(Context);
+		LinkedNode->CacheBones_AnyThread(InContext);
 	}
 }
 
-void FPoseLinkBase::Update(const FAnimationUpdateContext& Context)
+void FPoseLinkBase::Update(const FAnimationUpdateContext& InContext)
 {
 #if ENABLE_VERBOSE_ANIM_PERF_TRACKING
 	QUICK_SCOPE_CYCLE_COUNTER(STAT_FPoseLinkBase_Update);
@@ -266,48 +272,48 @@ void FPoseLinkBase::Update(const FAnimationUpdateContext& Context)
 
 #if DO_CHECK
 	checkf( !bProcessed, TEXT( "Update already in progress, circular link for AnimInstance [%s] Blueprint [%s]" ), \
-		*Context.AnimInstanceProxy->GetAnimInstanceName(), *GetFullNameSafe(IAnimClassInterface::GetActualAnimClass(Context.AnimInstanceProxy->GetAnimClassInterface())));
+		*InContext.AnimInstanceProxy->GetAnimInstanceName(), *GetFullNameSafe(IAnimClassInterface::GetActualAnimClass(InContext.AnimInstanceProxy->GetAnimClassInterface())));
 	TGuardValue<bool> CircularGuard(bProcessed, true);
 #endif
 
 #if WITH_EDITOR
 	if (GIsEditor)
 	{
-		if (LinkedNode == NULL)
+		if (LinkedNode == nullptr)
 		{
 			//@TODO: Should only do this when playing back
-			AttemptRelink(Context);
+			AttemptRelink(InContext);
 		}
 
 		// Record the node line activation
-		if (LinkedNode != NULL)
+		if (LinkedNode != nullptr)
 		{
-			if (Context.AnimInstanceProxy->IsBeingDebugged())
+			if (InContext.AnimInstanceProxy->IsBeingDebugged())
 			{
-				Context.AnimInstanceProxy->RecordNodeVisit(LinkID, SourceLinkID, Context.GetFinalBlendWeight());
+				InContext.AnimInstanceProxy->RecordNodeVisit(LinkID, SourceLinkID, InContext.GetFinalBlendWeight());
 			}
 		}
 	}
 #endif
 
 #if ENABLE_ANIMGRAPH_TRAVERSAL_DEBUG
-	checkf(InitializationCounter.IsSynchronized_Counter(Context.AnimInstanceProxy->GetInitializationCounter()), TEXT("Calling Update without initialization!"));
-	UpdateCounter.SynchronizeWith(Context.AnimInstanceProxy->GetUpdateCounter());
+	checkf(InitializationCounter.IsSynchronized_Counter(InContext.AnimInstanceProxy->GetInitializationCounter()), TEXT("Calling Update without initialization!"));
+	UpdateCounter.SynchronizeWith(InContext.AnimInstanceProxy->GetUpdateCounter());
 #endif
 
-	if (LinkedNode != NULL)
+	if (LinkedNode != nullptr)
 	{
-		FAnimationUpdateContext LinkContext(Context.WithNodeId(LinkID));
+		FAnimationUpdateContext LinkContext(InContext.WithNodeId(LinkID));
 		TRACE_SCOPED_ANIM_NODE(LinkContext);
 		LinkedNode->Update_AnyThread(LinkContext);
 	}
 }
 
-void FPoseLinkBase::GatherDebugData(FNodeDebugData& DebugData)
+void FPoseLinkBase::GatherDebugData(FNodeDebugData& InDebugData)
 {
-	if(LinkedNode != NULL)
+	if(LinkedNode != nullptr)
 	{
-		LinkedNode->GatherDebugData(DebugData);
+		LinkedNode->GatherDebugData(InDebugData);
 	}
 }
 
@@ -323,7 +329,7 @@ void FPoseLink::Evaluate(FPoseContext& Output)
 #endif
 
 #if WITH_EDITOR
-	if ((LinkedNode == NULL) && GIsEditor)
+	if ((LinkedNode == nullptr) && GIsEditor)
 	{
 		//@TODO: Should only do this when playing back
 		AttemptRelink(Output);
@@ -337,7 +343,7 @@ void FPoseLink::Evaluate(FPoseContext& Output)
 	EvaluationCounter.SynchronizeWith(Output.AnimInstanceProxy->GetEvaluationCounter());
 #endif
 
-	if (LinkedNode != NULL)
+	if (LinkedNode != nullptr)
 	{
 #if ENABLE_ANIMNODE_POSE_DEBUG
 		CurrentPose.ResetToAdditiveIdentity();
@@ -442,7 +448,7 @@ void FComponentSpacePoseLink::EvaluateComponentSpace(FComponentSpacePoseContext&
 	EvaluationCounter.SynchronizeWith(Output.AnimInstanceProxy->GetEvaluationCounter());
 #endif
 
-	if (LinkedNode != NULL)
+	if (LinkedNode != nullptr)
 	{
 		int32 SourceID = Output.GetCurrentNodeId();
 
@@ -577,88 +583,3 @@ void FNodeDebugData::GetFlattenedDebugData(TArray<FFlattenedDebugData>& Flattene
 		}
 	}
 }
-
-void FExposedValueHandler::DynamicClassInitialization(TArray<FExposedValueHandler>& InHandlers, UDynamicClass* InDynamicClass)
-{
-	const FPropertyAccessLibrary& PropertyAccessLibrary = IAnimClassInterface::GetFromClass(InDynamicClass)->GetPropertyAccessLibrary();
-
-	for(FExposedValueHandler& Handler : InHandlers)
-	{
-		Handler.Initialize(InDynamicClass, PropertyAccessLibrary);
-	}
-}
-
-void FExposedValueHandler::ClassInitialization(TArray<FExposedValueHandler>& InHandlers, UObject* InClassDefaultObject)
-{
-	UClass* Class = InClassDefaultObject->GetClass();
-	const FPropertyAccessLibrary& PropertyAccessLibrary = IAnimClassInterface::GetFromClass(Class)->GetPropertyAccessLibrary();
-
-	for(FExposedValueHandler& Handler : InHandlers)
-	{
-		FAnimNode_Base* AnimNode = Handler.ValueHandlerNodeProperty->ContainerPtrToValuePtr<FAnimNode_Base>(InClassDefaultObject);
-		check(AnimNode);
-		AnimNode->SetExposedValueHandler(&Handler);
-		Handler.Initialize(Class, PropertyAccessLibrary);
-	}
-}
-
-void FExposedValueHandler::Initialize(UClass* InClass, const FPropertyAccessLibrary& InPropertyAccessLibrary)
-{
-	// bInitialized may no longer be necessary, but leaving alone for now:
-	if (bInitialized)
-	{
-		return;
-	}
-
-	if (BoundFunction != NAME_None)
-	{
-		// This cached function is NULL when the CDO is initially serialized, or (in editor) when the class has been
-		// recompiled and any instances have been re-instanced. When new instances are spawned, this function is
-		// duplicated (it is a FProperty) onto those instances so we dont pay the cost of the FindFunction() call
-#if !WITH_EDITOR
-		if (Function == nullptr)
-#endif
-		{
-			// we cant call FindFunction on anything but the game thread as it accesses a shared map in the object's class
-			check(IsInGameThread());
-			Function = InClass->FindFunctionByName(BoundFunction);
-			check(Function);
-		}
-	}
-	else
-	{
-		Function = NULL;
-	}
-
-	// Cache property access library
-	PropertyAccessLibrary = &InPropertyAccessLibrary;
-
-	bInitialized = true;
-}
-
-void FExposedValueHandler::Execute(const FAnimationBaseContext& Context) const
-{
-	if (Function != nullptr)
-	{
-		Context.AnimInstanceProxy->GetAnimInstanceObject()->ProcessEvent(Function, NULL);
-	}
-
-	if(CopyRecords.Num() > 0)
-	{
-		if(PropertyAccessLibrary != nullptr)
-		{
-			UObject* AnimInstanceObject = Context.AnimInstanceProxy->GetAnimInstanceObject();
-			for(const FExposedValueCopyRecord& CopyRecord : CopyRecords)
-			{
-				PropertyAccess::ProcessCopy(AnimInstanceObject, *PropertyAccessLibrary, EPropertyAccessCopyBatch::InternalUnbatched, CopyRecord.CopyIndex, [&CopyRecord](const FProperty* InProperty, void* InAddress)
-				{
-					if(CopyRecord.PostCopyOperation == EPostCopyOperation::LogicalNegateBool)
-					{
-						bool bValue = static_cast<const FBoolProperty*>(InProperty)->GetPropertyValue(InAddress);
-						static_cast<const FBoolProperty*>(InProperty)->SetPropertyValue(InAddress, !bValue);
-					}
-				});
-			}
-		}
-	}
-}
diff --git a/Engine/Source/Runtime/Engine/Private/Animation/AnimNode_AssetPlayerBase.cpp b/Engine/Source/Runtime/Engine/Private/Animation/AnimNode_AssetPlayerBase.cpp
--- a/Engine/Source/Runtime/Engine/Private/Animation/AnimNode_AssetPlayerBase.cpp
+++ b/Engine/Source/Runtime/Engine/Private/Animation/AnimNode_AssetPlayerBase.cpp
@@ -4,21 +4,6 @@
 #include "Animation/AnimInstanceProxy.h"
 #include "Animation/AnimSyncScope.h"
 
-FAnimNode_AssetPlayerBase::FAnimNode_AssetPlayerBase()
-	: GroupName(NAME_None)
-#if WITH_EDITORONLY_DATA
-	, GroupIndex_DEPRECATED(INDEX_NONE)
-	, GroupScope_DEPRECATED(EAnimSyncGroupScope::Local)
-#endif
-	, GroupRole(EAnimGroupRole::CanBeLeader)
-	, bIgnoreForRelevancyTest(false)
-	, bHasBeenFullWeight(false)
-	, BlendWeight(0.0f)
-	, InternalTimeAccumulator(0.0f)
-{
-
-}
-
 void FAnimNode_AssetPlayerBase::Initialize_AnyThread(const FAnimationInitializeContext& Context)
 {
 	FAnimNode_Base::Initialize_AnyThread(Context);
@@ -43,14 +28,17 @@ void FAnimNode_AssetPlayerBase::CreateTickRecordForNode(const FAnimationUpdateCo
 
 	UE::Anim::FAnimSyncGroupScope& SyncScope = Context.GetMessageChecked<UE::Anim::FAnimSyncGroupScope>();
 
-	FName GroupNameToUse = ((GroupRole < EAnimGroupRole::TransitionLeader) || bHasBeenFullWeight) ? GroupName : NAME_None;
-	EAnimSyncMethod MethodToUse = Method;
+	const EAnimGroupRole::Type SyncGroupRole = GetGroupRole();
+	const FName SyncGroupName = GetGroupName();
+
+	FName GroupNameToUse = ((SyncGroupRole < EAnimGroupRole::TransitionLeader) || bHasBeenFullWeight) ? SyncGroupName : NAME_None;
+	EAnimSyncMethod MethodToUse = GetGroupMethod();
 	if(GroupNameToUse == NAME_None && MethodToUse == EAnimSyncMethod::SyncGroup)
 	{
 		MethodToUse = EAnimSyncMethod::DoNotSync;
 	}
 
-	UE::Anim::FAnimSyncParams SyncParams(GroupNameToUse, GroupRole, MethodToUse);
+	UE::Anim::FAnimSyncParams SyncParams(GroupNameToUse, SyncGroupRole, MethodToUse);
 	FAnimTickRecord TickRecord(Sequence, bLooping, PlayRate, FinalBlendWeight, /*inout*/ InternalTimeAccumulator, MarkerTickRecord);
 	TickRecord.RootMotionWeightModifier = Context.GetRootMotionWeightModifier();
 
@@ -74,7 +62,7 @@ void FAnimNode_AssetPlayerBase::SetAccumulatedTime(const float& NewTime)
 	InternalTimeAccumulator = NewTime;
 }
 
-UAnimationAsset* FAnimNode_AssetPlayerBase::GetAnimAsset()
+UAnimationAsset* FAnimNode_AssetPlayerBase::GetAnimAsset() const
 {
 	return nullptr;
 }
@@ -84,3 +72,22 @@ void FAnimNode_AssetPlayerBase::ClearCachedBlendWeight()
 	BlendWeight = 0.0f;
 }
 
+FName FAnimNode_AssetPlayerBase::GetGroupName() const
+{
+	return GET_ANIM_NODE_DATA(FName, GroupName);
+}
+
+EAnimGroupRole::Type FAnimNode_AssetPlayerBase::GetGroupRole() const
+{
+	return GET_ANIM_NODE_DATA(EAnimGroupRole::Type, GroupRole);
+}
+
+EAnimSyncMethod FAnimNode_AssetPlayerBase::GetGroupMethod() const
+{
+	return GET_ANIM_NODE_DATA(EAnimSyncMethod, Method);
+}
+
+bool FAnimNode_AssetPlayerBase::GetIgnoreForRelevancyTest() const
+{
+	return GET_ANIM_NODE_DATA(bool, bIgnoreForRelevancyTest);
+}
\ No newline at end of file
diff --git a/Engine/Source/Runtime/Engine/Private/Animation/AnimNode_LinkedAnimGraph.cpp b/Engine/Source/Runtime/Engine/Private/Animation/AnimNode_LinkedAnimGraph.cpp
--- a/Engine/Source/Runtime/Engine/Private/Animation/AnimNode_LinkedAnimGraph.cpp
+++ b/Engine/Source/Runtime/Engine/Private/Animation/AnimNode_LinkedAnimGraph.cpp
@@ -80,7 +80,7 @@ void FAnimNode_LinkedAnimGraph::CacheBonesSubGraph_AnyThread(const FAnimationCac
 		// Note not calling Proxy.CacheBones_WithRoot here as it is guarded by
 		// bBoneCachesInvalidated, which is handled at a higher level
 		FAnimationCacheBonesContext LinkedContext(&Proxy);
-		LinkedContext = LinkedContext.WithNodeId(CachedLinkedNodeIndex);
+		LinkedContext.SetNodeId(CachedLinkedNodeIndex);
 		LinkedRoot->CacheBones_AnyThread(LinkedContext);
 	}
 }
@@ -112,7 +112,7 @@ void FAnimNode_LinkedAnimGraph::Update_AnyThread(const FAnimationUpdateContext&
 		// in USkeletalMeshComponent::TickAnimation. It used to be the case that we could do non-parallel work in 
 		// USkeletalMeshComponent::TickAnimation, which would mean we would have to skip doing that work here.
 		FAnimationUpdateContext NewContext = InContext.WithOtherProxy(&Proxy);
-		NewContext = NewContext.WithNodeId(CachedLinkedNodeIndex);
+		NewContext.SetNodeId(CachedLinkedNodeIndex);
 		Proxy.UpdateAnimation_WithRoot(NewContext, LinkedRoot, GetDynamicLinkFunctionName());
 	}
 	else if(InputPoses.Num() > 0)
diff --git a/Engine/Source/Runtime/Engine/Private/Animation/AnimNode_LinkedInputPose.cpp b/Engine/Source/Runtime/Engine/Private/Animation/AnimNode_LinkedInputPose.cpp
--- a/Engine/Source/Runtime/Engine/Private/Animation/AnimNode_LinkedInputPose.cpp
+++ b/Engine/Source/Runtime/Engine/Private/Animation/AnimNode_LinkedInputPose.cpp
@@ -29,7 +29,7 @@ void FAnimNode_LinkedInputPose::Update_AnyThread(const FAnimationUpdateContext&
 	if(InputProxy)
 	{
 		FAnimationUpdateContext InputContext = Context.WithOtherProxy(InputProxy);
-		InputContext = InputContext.WithNodeId(OuterGraphNodeIndex);
+		InputContext.SetNodeId(OuterGraphNodeIndex);
 		InputPose.Update(InputContext);
 	}
 }
diff --git a/Engine/Source/Runtime/Engine/Private/Animation/AnimNode_Root.cpp b/Engine/Source/Runtime/Engine/Private/Animation/AnimNode_Root.cpp
--- a/Engine/Source/Runtime/Engine/Private/Animation/AnimNode_Root.cpp
+++ b/Engine/Source/Runtime/Engine/Private/Animation/AnimNode_Root.cpp
@@ -11,8 +11,6 @@ FAnimNode_Root::FAnimNode_Root()
 
 void FAnimNode_Root::Initialize_AnyThread(const FAnimationInitializeContext& Context)
 {
-	FAnimNode_Base::Initialize_AnyThread(Context);
-
 	Result.Initialize(Context);
 }
 
@@ -23,9 +21,8 @@ void FAnimNode_Root::CacheBones_AnyThread(const FAnimationCacheBonesContext& Con
 
 void FAnimNode_Root::Update_AnyThread(const FAnimationUpdateContext& Context)
 {
-	TRACE_ANIM_NODE_VALUE(Context, TEXT("Name"), Name);
+	TRACE_ANIM_NODE_VALUE(Context, TEXT("Name"), GetName());
 
-	GetEvaluateGraphExposedInputs().Execute(Context);
 	Result.Update(Context);
 }
 
@@ -40,3 +37,13 @@ void FAnimNode_Root::GatherDebugData(FNodeDebugData& DebugData)
 	DebugData.AddDebugItem(DebugLine);
 	Result.GatherDebugData(DebugData);
 }
+
+FName FAnimNode_Root::GetName() const
+{
+	return GET_ANIM_NODE_DATA(FName, Name);
+}
+
+FName FAnimNode_Root::GetGroup() const
+{
+	return GET_ANIM_NODE_DATA(FName, Group);
+}
\ No newline at end of file
diff --git a/Engine/Source/Runtime/Engine/Private/Animation/AnimNode_SequencePlayer.cpp b/Engine/Source/Runtime/Engine/Private/Animation/AnimNode_SequencePlayer.cpp
--- a/Engine/Source/Runtime/Engine/Private/Animation/AnimNode_SequencePlayer.cpp
+++ b/Engine/Source/Runtime/Engine/Private/Animation/AnimNode_SequencePlayer.cpp
@@ -12,25 +12,30 @@
 /////////////////////////////////////////////////////
 // FAnimSequencePlayerNode
 
-float FAnimNode_SequencePlayer::GetCurrentAssetTime()
+float FAnimNode_SequencePlayerBase::GetCurrentAssetTime() const
 {
 	return InternalTimeAccumulator;
 }
 
-float FAnimNode_SequencePlayer::GetCurrentAssetTimePlayRateAdjusted()
+float FAnimNode_SequencePlayerBase::GetCurrentAssetTimePlayRateAdjusted() const
 {
-	const float SequencePlayRate = (Sequence ? Sequence->RateScale : 1.f);
-	const float AdjustedPlayRate = PlayRateScaleBiasClamp.ApplyTo(FMath::IsNearlyZero(PlayRateBasis) ? 0.f : (PlayRate / PlayRateBasis), 0.f);
+	UAnimSequenceBase* CurrentSequence = GetSequence();
+	const float SequencePlayRate = (CurrentSequence ? CurrentSequence->RateScale : 1.f);
+	const float CurrentPlayRate = GetPlayRate();
+	const float CurrentPlayRateBasis = GetPlayRateBasis();
+
+	const float AdjustedPlayRate = PlayRateScaleBiasClampState.ApplyTo(GetPlayRateScaleBiasClampConstants(), FMath::IsNearlyZero(CurrentPlayRateBasis) ? 0.f : (CurrentPlayRate / CurrentPlayRateBasis));
 	const float EffectivePlayrate = SequencePlayRate * AdjustedPlayRate;
 	return (EffectivePlayrate < 0.0f) ? GetCurrentAssetLength() - InternalTimeAccumulator : InternalTimeAccumulator;
 }
 
-float FAnimNode_SequencePlayer::GetCurrentAssetLength()
+float FAnimNode_SequencePlayerBase::GetCurrentAssetLength() const
 {
-	return Sequence ? Sequence->GetPlayLength() : 0.0f;
+	UAnimSequenceBase* CurrentSequence = GetSequence();
+	return CurrentSequence ? CurrentSequence->GetPlayLength() : 0.0f;
 }
 
-void FAnimNode_SequencePlayer::Initialize_AnyThread(const FAnimationInitializeContext& Context)
+void FAnimNode_SequencePlayerBase::Initialize_AnyThread(const FAnimationInitializeContext& Context)
 {
 	DECLARE_SCOPE_HIERARCHICAL_COUNTER_ANIMNODE(Initialize_AnyThread);
 
@@ -38,81 +43,89 @@ void FAnimNode_SequencePlayer::Initialize_AnyThread(const FAnimationInitializeCo
 
 	GetEvaluateGraphExposedInputs().Execute(Context);
 
-	if (Sequence && !ensureMsgf(!Sequence->IsA<UAnimMontage>(), TEXT("Sequence players do not support anim montages.")))
+	UAnimSequenceBase* CurrentSequence = GetSequence();
+	if (CurrentSequence && !ensureMsgf(!CurrentSequence->IsA<UAnimMontage>(), TEXT("Sequence players do not support anim montages.")))
 	{
-		Sequence = nullptr;
+		CurrentSequence = nullptr;
 	}
 
-	InternalTimeAccumulator = StartPosition;
-	PlayRateScaleBiasClamp.Reinitialize();
+	InternalTimeAccumulator = GetStartPosition();
+	PlayRateScaleBiasClampState.Reinitialize();
 
-	if (Sequence != nullptr)
+	if (CurrentSequence != nullptr)
 	{
 		const float EffectiveStartPosition = GetEffectiveStartPosition(Context);
+		const float CurrentPlayRate = GetPlayRate();
+		const float CurrentPlayRateBasis = GetPlayRateBasis();
 
-		InternalTimeAccumulator = FMath::Clamp(EffectiveStartPosition, 0.f, Sequence->GetPlayLength());
-		const float AdjustedPlayRate = PlayRateScaleBiasClamp.ApplyTo(FMath::IsNearlyZero(PlayRateBasis) ? 0.f : (PlayRate / PlayRateBasis), 0.f);
-		const float EffectivePlayrate = Sequence->RateScale * AdjustedPlayRate;
+		InternalTimeAccumulator = FMath::Clamp(EffectiveStartPosition, 0.f, CurrentSequence->GetPlayLength());
+		const float AdjustedPlayRate = PlayRateScaleBiasClampState.ApplyTo(GetPlayRateScaleBiasClampConstants(), FMath::IsNearlyZero(CurrentPlayRateBasis) ? 0.f : (CurrentPlayRate / CurrentPlayRateBasis), 0.f);
+		const float EffectivePlayrate = CurrentSequence->RateScale * AdjustedPlayRate;
 		if ((EffectiveStartPosition == 0.f) && (EffectivePlayrate < 0.f))
 		{
-			InternalTimeAccumulator = Sequence->GetPlayLength();
+			InternalTimeAccumulator = CurrentSequence->GetPlayLength();
 		}
 	}
 }
 
-void FAnimNode_SequencePlayer::CacheBones_AnyThread(const FAnimationCacheBonesContext& Context)
+void FAnimNode_SequencePlayerBase::CacheBones_AnyThread(const FAnimationCacheBonesContext& Context)
 {
 	DECLARE_SCOPE_HIERARCHICAL_COUNTER_ANIMNODE(CacheBones_AnyThread);
 }
 
-void FAnimNode_SequencePlayer::UpdateAssetPlayer(const FAnimationUpdateContext& Context)
+void FAnimNode_SequencePlayerBase::UpdateAssetPlayer(const FAnimationUpdateContext& Context)
 {
 	DECLARE_SCOPE_HIERARCHICAL_COUNTER_ANIMNODE(UpdateAssetPlayer);
 
 	GetEvaluateGraphExposedInputs().Execute(Context);
 
-	if (Sequence && !ensureMsgf(!Sequence->IsA<UAnimMontage>(), TEXT("Sequence players do not support anim montages.")))
+	UAnimSequenceBase* CurrentSequence = GetSequence();
+	if (CurrentSequence && !ensureMsgf(!CurrentSequence->IsA<UAnimMontage>(), TEXT("Sequence players do not support anim montages.")))
 	{
-		Sequence = nullptr;
+		CurrentSequence = nullptr;
 	}
 
-	if ((Sequence != nullptr) && Context.AnimInstanceProxy->IsSkeletonCompatible(Sequence->GetSkeleton()))
+	if ((CurrentSequence != nullptr) && (Context.AnimInstanceProxy->IsSkeletonCompatible(CurrentSequence->GetSkeleton())))
 	{
-		InternalTimeAccumulator = FMath::Clamp(InternalTimeAccumulator, 0.f, Sequence->GetPlayLength());
-		const float AdjustedPlayRate = PlayRateScaleBiasClamp.ApplyTo(FMath::IsNearlyZero(PlayRateBasis) ? 0.f : (PlayRate / PlayRateBasis), Context.GetDeltaTime());
-		CreateTickRecordForNode(Context, Sequence, bLoopAnimation, AdjustedPlayRate);
+		const float CurrentPlayRate = GetPlayRate();
+		const float CurrentPlayRateBasis = GetPlayRateBasis();
+
+		InternalTimeAccumulator = FMath::Clamp(InternalTimeAccumulator, 0.f, CurrentSequence->GetPlayLength());
+		const float AdjustedPlayRate = PlayRateScaleBiasClampState.ApplyTo(GetPlayRateScaleBiasClampConstants(), FMath::IsNearlyZero(CurrentPlayRateBasis) ? 0.f : (CurrentPlayRate / CurrentPlayRateBasis), Context.GetDeltaTime());
+		CreateTickRecordForNode(Context, CurrentSequence, GetLoopAnimation(), AdjustedPlayRate);
 	}
 
 #if WITH_EDITORONLY_DATA
 	if (FAnimBlueprintDebugData* DebugData = Context.AnimInstanceProxy->GetAnimBlueprintDebugData())
 	{
-		DebugData->RecordSequencePlayer(Context.GetCurrentNodeId(), GetAccumulatedTime(), Sequence != nullptr ? Sequence->GetPlayLength() : 0.0f, Sequence != nullptr ? Sequence->GetNumberOfSampledKeys() : 0);
+		DebugData->RecordSequencePlayer(Context.GetCurrentNodeId(), GetAccumulatedTime(), CurrentSequence != nullptr ? CurrentSequence->GetPlayLength() : 0.0f, CurrentSequence != nullptr ? CurrentSequence->GetNumberOfSampledKeys() : 0);
 	}
 #endif
 
 	TRACE_ANIM_SEQUENCE_PLAYER(Context, *this);
-	TRACE_ANIM_NODE_VALUE(Context, TEXT("Name"), Sequence != nullptr ? Sequence->GetFName() : NAME_None);
-	TRACE_ANIM_NODE_VALUE(Context, TEXT("Sequence"), Sequence);
+	TRACE_ANIM_NODE_VALUE(Context, TEXT("Name"), CurrentSequence != nullptr ? CurrentSequence->GetFName() : NAME_None);
+	TRACE_ANIM_NODE_VALUE(Context, TEXT("Sequence"), CurrentSequence);
 	TRACE_ANIM_NODE_VALUE(Context, TEXT("Playback Time"), InternalTimeAccumulator);
 }
 
-void FAnimNode_SequencePlayer::Evaluate_AnyThread(FPoseContext& Output)
+void FAnimNode_SequencePlayerBase::Evaluate_AnyThread(FPoseContext& Output)
 {
 	DECLARE_SCOPE_HIERARCHICAL_COUNTER_ANIMNODE(Evaluate_AnyThread);
 
-	if ((Sequence != nullptr) && (Output.AnimInstanceProxy->IsSkeletonCompatible(Sequence->GetSkeleton())))
+	UAnimSequenceBase* CurrentSequence = GetSequence();
+	if ((CurrentSequence != nullptr) && (Output.AnimInstanceProxy->IsSkeletonCompatible(CurrentSequence->GetSkeleton())))
 	{
 		const bool bExpectedAdditive = Output.ExpectsAdditivePose();
-		const bool bIsAdditive = Sequence->IsValidAdditive();
+		const bool bIsAdditive = CurrentSequence->IsValidAdditive();
 
 		if (bExpectedAdditive && !bIsAdditive)
 		{
-			FText Message = FText::Format(LOCTEXT("AdditiveMismatchWarning", "Trying to play a non-additive animation '{0}' into a pose that is expected to be additive in anim instance '{1}'"), FText::FromString(Sequence->GetName()), FText::FromString(Output.AnimInstanceProxy->GetAnimInstanceName()));
+			FText Message = FText::Format(LOCTEXT("AdditiveMismatchWarning", "Trying to play a non-additive animation '{0}' into a pose that is expected to be additive in anim instance '{1}'"), FText::FromString(CurrentSequence->GetName()), FText::FromString(Output.AnimInstanceProxy->GetAnimInstanceName()));
 			Output.LogMessage(EMessageSeverity::Warning, Message);
 		}
 
 		FAnimationPoseData AnimationPoseData(Output);
-		Sequence->GetAnimationPose(AnimationPoseData, FAnimExtractContext(InternalTimeAccumulator, Output.AnimInstanceProxy->ShouldExtractRootMotion()));
+		CurrentSequence->GetAnimationPose(AnimationPoseData, FAnimExtractContext(InternalTimeAccumulator, Output.AnimInstanceProxy->ShouldExtractRootMotion()));
 	}
 	else
 	{
@@ -120,36 +133,31 @@ void FAnimNode_SequencePlayer::Evaluate_AnyThread(FPoseContext& Output)
 	}
 }
 
-void FAnimNode_SequencePlayer::OverrideAsset(UAnimationAsset* NewAsset)
-{
-	if (UAnimSequenceBase* AnimSequence = Cast<UAnimSequenceBase>(NewAsset))
-	{
-		Sequence = AnimSequence;
-	}
-}
-
-void FAnimNode_SequencePlayer::GatherDebugData(FNodeDebugData& DebugData)
+void FAnimNode_SequencePlayerBase::GatherDebugData(FNodeDebugData& DebugData)
 {
 	FString DebugLine = DebugData.GetNodeName(this);
 
-	DebugLine += FString::Printf(TEXT("('%s' Play Time: %.3f)"), Sequence ? *Sequence->GetName() : TEXT("NULL"), InternalTimeAccumulator);
+	UAnimSequenceBase* CurrentSequence = GetSequence();
+	DebugLine += FString::Printf(TEXT("('%s' Play Time: %.3f)"), CurrentSequence ? *CurrentSequence->GetName() : TEXT("NULL"), InternalTimeAccumulator);
 	DebugData.AddDebugItem(DebugLine, true);
 }
 
-float FAnimNode_SequencePlayer::GetTimeFromEnd(float CurrentNodeTime)
+float FAnimNode_SequencePlayerBase::GetTimeFromEnd(float CurrentNodeTime) const
 {
-	return Sequence->GetPlayLength() - CurrentNodeTime;
+	UAnimSequenceBase* CurrentSequence = GetSequence();
+	return CurrentSequence ? CurrentSequence->GetPlayLength() - CurrentNodeTime : 0.0f;
 }
 
-float FAnimNode_SequencePlayer::GetEffectiveStartPosition(const FAnimationBaseContext& Context) const
+float FAnimNode_SequencePlayerBase::GetEffectiveStartPosition(const FAnimationBaseContext& Context) const
 {
 	// Override the start position if pose matching is enabled
-	if (Sequence != nullptr && bStartFromMatchingPose)
+	UAnimSequenceBase* CurrentSequence = GetSequence();
+	if (CurrentSequence != nullptr && GetStartFromMatchingPose())
 	{
 		UE::Anim::IPoseSearchProvider* PoseSearchProvider = UE::Anim::IPoseSearchProvider::Get();
 		if (PoseSearchProvider)
 		{
-			UE::Anim::IPoseSearchProvider::FSearchResult Result = PoseSearchProvider->Search(Context, Sequence);
+			UE::Anim::IPoseSearchProvider::FSearchResult Result = PoseSearchProvider->Search(Context, CurrentSequence);
 			if (Result.PoseIdx >= 0)
 			{
 				return Result.TimeOffsetSeconds;
@@ -157,7 +165,95 @@ float FAnimNode_SequencePlayer::GetEffectiveStartPosition(const FAnimationBaseCo
 		}
 	}
 
+	return GetStartPosition();
+}
+
+void FAnimNode_SequencePlayer::SetSequence(UAnimSequenceBase* InSequence)
+{
+#if WITH_EDITORONLY_DATA	
+	TObjectPtr<UAnimSequenceBase>& SequenceToSet = GET_MUTABLE_ANIM_NODE_DATA(TObjectPtr<UAnimSequenceBase>, Sequence);
+	SequenceToSet = InSequence;
+#endif
+}
+
+UAnimSequenceBase* FAnimNode_SequencePlayer::GetSequence() const
+{
+	return GET_ANIM_NODE_DATA(UAnimSequenceBase*, Sequence);
+}
+
+float FAnimNode_SequencePlayer::GetPlayRateBasis() const
+{
+	return GET_ANIM_NODE_DATA(float, PlayRateBasis);
+}
+
+float FAnimNode_SequencePlayer::GetPlayRate() const
+{
+	return GET_ANIM_NODE_DATA(float, PlayRate);
+}
+
+const FInputScaleBiasClampConstants& FAnimNode_SequencePlayer::GetPlayRateScaleBiasClampConstants() const
+{
+	return GET_ANIM_NODE_DATA(FInputScaleBiasClampConstants, PlayRateScaleBiasClampConstants);
+}
+
+float FAnimNode_SequencePlayer::GetStartPosition() const
+{
+	return GET_ANIM_NODE_DATA(float, StartPosition);
+}
+
+bool FAnimNode_SequencePlayer::GetLoopAnimation() const
+{
+	return GET_ANIM_NODE_DATA(bool, bLoopAnimation);
+}
+
+bool FAnimNode_SequencePlayer::GetStartFromMatchingPose() const
+{
+	return GET_ANIM_NODE_DATA(bool, bStartFromMatchingPose);
+}
+
+void FAnimNode_SequencePlayer_Standalone::SetSequence(UAnimSequenceBase* InSequence)
+{
+	Sequence = InSequence;
+}
+
+void FAnimNode_SequencePlayer_Standalone::SetLoopAnimation(bool bInLoopAnimation)
+{
+	bLoopAnimation = bInLoopAnimation;
+}
+
+UAnimSequenceBase* FAnimNode_SequencePlayer_Standalone::GetSequence() const
+{
+	return Sequence;
+}
+
+float FAnimNode_SequencePlayer_Standalone::GetPlayRateBasis() const
+{
+	return PlayRateBasis;
+}
+
+float FAnimNode_SequencePlayer_Standalone::GetPlayRate() const
+{
+	return PlayRate;
+}
+
+const FInputScaleBiasClampConstants& FAnimNode_SequencePlayer_Standalone::GetPlayRateScaleBiasClampConstants() const
+{
+	return PlayRateScaleBiasClampConstants;
+}
+
+float FAnimNode_SequencePlayer_Standalone::GetStartPosition() const
+{
 	return StartPosition;
 }
 
-#undef LOCTEXT_NAMESPACE
\ No newline at end of file
+bool FAnimNode_SequencePlayer_Standalone::GetLoopAnimation() const
+{
+	return bLoopAnimation;
+}
+
+bool FAnimNode_SequencePlayer_Standalone::GetStartFromMatchingPose() const
+{
+	return bStartFromMatchingPose;
+}
+
+#undef LOCTEXT_NAMESPACE
diff --git a/Engine/Source/Runtime/Engine/Private/Animation/AnimNode_StateMachine.cpp b/Engine/Source/Runtime/Engine/Private/Animation/AnimNode_StateMachine.cpp
--- a/Engine/Source/Runtime/Engine/Private/Animation/AnimNode_StateMachine.cpp
+++ b/Engine/Source/Runtime/Engine/Private/Animation/AnimNode_StateMachine.cpp
@@ -561,14 +561,14 @@ void FAnimNode_StateMachine::Update_AnyThread(const FAnimationUpdateContext& Con
 #endif
 }
 
-FAnimNode_AssetPlayerBase* FAnimNode_StateMachine::GetRelevantAssetPlayerFromState(const FAnimationUpdateContext& Context, const FBakedAnimationState& StateInfo)
+const FAnimNode_AssetPlayerBase* FAnimNode_StateMachine::GetRelevantAssetPlayerFromState(const FAnimationUpdateContext& Context, const FBakedAnimationState& StateInfo) const
 {
-	FAnimNode_AssetPlayerBase* ResultPlayer = nullptr;
+	const FAnimNode_AssetPlayerBase* ResultPlayer = nullptr;
 	float MaxWeight = 0.0f;
 
-	auto EvaluatePlayerWeight = [&MaxWeight, &ResultPlayer](FAnimNode_AssetPlayerBase* Player)
+	auto EvaluatePlayerWeight = [&MaxWeight, &ResultPlayer, &Context](const FAnimNode_AssetPlayerBase* Player)
 	{
-		if (!Player->bIgnoreForRelevancyTest && (Player->GetCachedBlendWeight() > MaxWeight))
+		if (!Player->GetIgnoreForRelevancyTest() && Player->GetCachedBlendWeight() > MaxWeight)
 		{
 			MaxWeight = Player->GetCachedBlendWeight();
 			ResultPlayer = Player;
@@ -577,7 +577,7 @@ FAnimNode_AssetPlayerBase* FAnimNode_StateMachine::GetRelevantAssetPlayerFromSta
 
 	for (const int32& PlayerIdx : StateInfo.PlayerNodeIndices)
 	{
-		if (FAnimNode_AssetPlayerBase* Player = Context.AnimInstanceProxy->GetNodeFromIndex<FAnimNode_AssetPlayerBase>(PlayerIdx))
+		if (const FAnimNode_AssetPlayerBase* Player = Context.AnimInstanceProxy->GetNodeFromIndex<FAnimNode_AssetPlayerBase>(PlayerIdx))
 		{
 			EvaluatePlayerWeight(Player);
 		}
@@ -587,14 +587,14 @@ FAnimNode_AssetPlayerBase* FAnimNode_StateMachine::GetRelevantAssetPlayerFromSta
 	for (const int32& LayerIdx : StateInfo.LayerNodeIndices)
 	{
 		// Try and retrieve the actual node object
-		if (FAnimNode_LinkedAnimLayer* Layer = Context.AnimInstanceProxy->GetNodeFromIndex<FAnimNode_LinkedAnimLayer>(LayerIdx))
+		if (const FAnimNode_LinkedAnimLayer* Layer = Context.AnimInstanceProxy->GetNodeFromIndex<FAnimNode_LinkedAnimLayer>(LayerIdx))
 		{
 			// Retrieve the AnimInstance running for this layer
-			if (UAnimInstance* CurrentTarget = Layer->GetTargetInstance<UAnimInstance>())
+			if (const UAnimInstance* CurrentTarget = Layer->GetTargetInstance<UAnimInstance>())
 			{
 				// Retrieve all asset player nodes from the corresponding Anim blueprint class and apply same logic to find highest weighted asset player 
-				TArray<FAnimNode_AssetPlayerBase*> PlayerNodesInLayer = CurrentTarget->GetInstanceAssetPlayers(Layer->Layer);
-				for (FAnimNode_AssetPlayerBase* Player : PlayerNodesInLayer)
+				TArray<const FAnimNode_AssetPlayerBase*> PlayerNodesInLayer = CurrentTarget->GetInstanceAssetPlayers(Layer->Layer);
+				for (const FAnimNode_AssetPlayerBase* Player : PlayerNodesInLayer)
 				{
 					EvaluatePlayerWeight(Player);
 				}
@@ -672,9 +672,9 @@ bool FAnimNode_StateMachine::FindValidTransition(const FAnimationUpdateContext&
 		else if (TransitionRule.bAutomaticRemainingTimeRule)
 		{
 			bool bCanEnterTransition = false;
-			if (FAnimNode_AssetPlayerBase* RelevantPlayer = GetRelevantAssetPlayerFromState(Context, StateInfo))
+			if (const FAnimNode_AssetPlayerBase* RelevantPlayer = GetRelevantAssetPlayerFromState(Context, StateInfo))
 			{
-				if (UAnimationAsset* AnimAsset = RelevantPlayer->GetAnimAsset())
+				if (const UAnimationAsset* AnimAsset = RelevantPlayer->GetAnimAsset())
 				{
 					const float AnimTimeRemaining = AnimAsset->GetPlayLength() - RelevantPlayer->GetAccumulatedTime();
 					const FAnimationTransitionBetweenStates& TransitionInfo = GetTransitionInfo(TransitionRule.TransitionIndex);
@@ -996,7 +996,7 @@ void FAnimNode_StateMachine::SetState(const FAnimationBaseContext& Context, int3
 		// This stops any zero length blends holding on to old weights
 		for(int32 PlayerIndex : GetStateInfo(CurrentState).PlayerNodeIndices)
 		{
-			if(FAnimNode_AssetPlayerBase* Player = Context.AnimInstanceProxy->GetNodeFromIndex<FAnimNode_AssetPlayerBase>(PlayerIndex))
+			if(FAnimNode_AssetPlayerBase* Player = Context.AnimInstanceProxy->GetMutableNodeFromIndex<FAnimNode_AssetPlayerBase>(PlayerIndex))
 			{
 				Player->ClearCachedBlendWeight();
 			}
@@ -1006,13 +1006,13 @@ void FAnimNode_StateMachine::SetState(const FAnimationBaseContext& Context, int3
 		for (const int32& LayerIdx : GetStateInfo(CurrentState).LayerNodeIndices)
 		{
 			// Try and retrieve the actual node object
-			if (FAnimNode_LinkedAnimLayer* Layer = Context.AnimInstanceProxy->GetNodeFromIndex<FAnimNode_LinkedAnimLayer>(LayerIdx))
+			if (FAnimNode_LinkedAnimLayer* Layer = Context.AnimInstanceProxy->GetMutableNodeFromIndex<FAnimNode_LinkedAnimLayer>(LayerIdx))
 			{
 				// Retrieve the AnimInstance running for this layer
 				if (UAnimInstance* CurrentTarget = Layer->GetTargetInstance<UAnimInstance>())
 				{
 					// Retrieve all asset player nodes from the corresponding Anim blueprint class and clear their cached blend weight
-					TArray<FAnimNode_AssetPlayerBase*> PlayerNodesInLayer = CurrentTarget->GetInstanceAssetPlayers(Layer->Layer);
+					TArray<FAnimNode_AssetPlayerBase*> PlayerNodesInLayer = CurrentTarget->GetMutableInstanceAssetPlayers(Layer->Layer);
 					for (FAnimNode_AssetPlayerBase* Player : PlayerNodesInLayer)
 					{
 						Player->ClearCachedBlendWeight();
diff --git a/Engine/Source/Runtime/Engine/Private/Animation/AnimSubsystem_Base.cpp b/Engine/Source/Runtime/Engine/Private/Animation/AnimSubsystem_Base.cpp
new file mode 100644
--- /dev/null
+++ b/Engine/Source/Runtime/Engine/Private/Animation/AnimSubsystem_Base.cpp
@@ -0,0 +1,8 @@
+// Copyright Epic Games, Inc. All Rights Reserved.
+
+#include "Animation/AnimSubsystem_Base.h"
+
+void FAnimSubsystem_Base::OnPostLoadDefaults(FAnimSubsystemPostLoadDefaultsContext& InContext)
+{
+	FExposedValueHandler::ClassInitialization(ExposedValueHandlers, InContext.DefaultAnimInstance->GetClass());
+}
diff --git a/Engine/Source/Runtime/Engine/Private/Animation/AnimSubsystem_PropertyAccess.cpp b/Engine/Source/Runtime/Engine/Private/Animation/AnimSubsystem_PropertyAccess.cpp
new file mode 100644
--- /dev/null
+++ b/Engine/Source/Runtime/Engine/Private/Animation/AnimSubsystem_PropertyAccess.cpp
@@ -0,0 +1,23 @@
+// Copyright Epic Games, Inc. All Rights Reserved.
+
+#include "Animation/AnimSubsystem_PropertyAccess.h"
+#include "Animation/AnimInstance.h"
+#include "Animation/AnimInstanceProxy.h"
+
+void FAnimSubsystem_PropertyAccess::OnUpdate(FAnimSubsystemUpdateContext& InContext) const
+{
+	// Process internal batched property copies
+	PropertyAccess::ProcessCopies(InContext.AnimInstance, Library, EPropertyAccessCopyBatch::ExternalBatched);
+}
+
+void FAnimSubsystem_PropertyAccess::OnParallelUpdate(FAnimSubsystemParallelUpdateContext& InContext) const
+{
+	// Process internal batched property copies
+	PropertyAccess::ProcessCopies(InContext.Proxy.GetAnimInstanceObject(), Library, EPropertyAccessCopyBatch::InternalBatched);
+}
+
+void FAnimSubsystem_PropertyAccess::OnPostLoad(FAnimSubsystemPostLoadContext& InContext)
+{
+	// Patch the library on load to fixup property offsets
+	PropertyAccess::PostLoadLibrary(Library);
+}
\ No newline at end of file
diff --git a/Engine/Source/Runtime/Engine/Private/Animation/AnimTrace.cpp b/Engine/Source/Runtime/Engine/Private/Animation/AnimTrace.cpp
--- a/Engine/Source/Runtime/Engine/Private/Animation/AnimTrace.cpp
+++ b/Engine/Source/Runtime/Engine/Private/Animation/AnimTrace.cpp
@@ -1004,7 +1004,7 @@ void FAnimTrace::OutputAnimNodeValue(const FAnimationBaseContext& InContext, con
 		<< AnimNodeValueClass.Attachment(InKey, KeyLength * sizeof(TCHAR));
 }
 
-void FAnimTrace::OutputAnimSequencePlayer(const FAnimationBaseContext& InContext, const FAnimNode_SequencePlayer& InNode)
+void FAnimTrace::OutputAnimSequencePlayer(const FAnimationBaseContext& InContext, const FAnimNode_SequencePlayerBase& InNode)
 {
 	bool bChannelEnabled = UE_TRACE_CHANNELEXPR_IS_ENABLED(AnimationChannel);
 	if (!bChannelEnabled)
@@ -1022,13 +1022,15 @@ void FAnimTrace::OutputAnimSequencePlayer(const FAnimationBaseContext& InContext
 	UObject* AnimInstance = InContext.AnimInstanceProxy->GetAnimInstanceObject();
 	TRACE_OBJECT(AnimInstance);
 
+	UAnimSequenceBase* Sequence = InNode.GetSequence();
+
 	UE_TRACE_LOG(Animation, AnimSequencePlayer, AnimationChannel)
 		<< AnimSequencePlayer.Cycle(FPlatformTime::Cycles64())
 		<< AnimSequencePlayer.AnimInstanceId(FObjectTrace::GetObjectId(AnimInstance))
 		<< AnimSequencePlayer.NodeId(InContext.GetCurrentNodeId())
 		<< AnimSequencePlayer.Position(InNode.GetAccumulatedTime())
-		<< AnimSequencePlayer.Length(InNode.Sequence ? InNode.Sequence->GetPlayLength() : 0.0f)
-		<< AnimSequencePlayer.FrameCounter(InNode.Sequence ? InNode.Sequence->GetNumberOfSampledKeys() : 0);
+		<< AnimSequencePlayer.Length(Sequence ? Sequence->GetPlayLength() : 0.0f)
+		<< AnimSequencePlayer.FrameCounter(Sequence ? Sequence->GetNumberOfSampledKeys() : 0);
 }
 
 void FAnimTrace::OutputStateMachineState(const FAnimationBaseContext& InContext, int32 InStateMachineIndex, int32 InStateIndex, float InStateWeight, float InElapsedTime)
diff --git a/Engine/Source/Runtime/Engine/Private/Animation/ExposedValueHandler.cpp b/Engine/Source/Runtime/Engine/Private/Animation/ExposedValueHandler.cpp
new file mode 100644
--- /dev/null
+++ b/Engine/Source/Runtime/Engine/Private/Animation/ExposedValueHandler.cpp
@@ -0,0 +1,74 @@
+// Copyright Epic Games, Inc. All Rights Reserved.
+
+#include "Animation/ExposedValueHandler.h"
+#include "UObject/Class.h"
+#include "Animation/AnimSubsystem_PropertyAccess.h"
+#include "Animation/AnimNodeBase.h"
+#include "Animation/AnimClassInterface.h"
+#include "Animation/AnimInstanceProxy.h"
+
+void FExposedValueHandler::ClassInitialization(TArray<FExposedValueHandler>& InHandlers, UClass* InClass)
+{
+	if(const FAnimSubsystem_PropertyAccess* Subsystem = IAnimClassInterface::GetFromClass(InClass)->FindSubsystem<FAnimSubsystem_PropertyAccess>())
+	{
+		const FPropertyAccessLibrary& PropertyAccessLibrary = Subsystem->GetLibrary();
+
+		for(FExposedValueHandler& Handler : InHandlers)
+		{
+			Handler.Initialize(InClass, PropertyAccessLibrary);
+		}
+	}
+}
+
+void FExposedValueHandler::Initialize(UClass* InClass, const FPropertyAccessLibrary& InPropertyAccessLibrary)
+{
+	if (BoundFunction != NAME_None)
+	{
+		// This cached function is nullptr when the CDO is initially serialized, or (in editor) when the class has been
+		// recompiled and any instances have been re-instanced. When new instances are spawned, this function is
+		// duplicated (it is a FProperty) onto those instances so we dont pay the cost of the FindFunction() call
+#if !WITH_EDITOR
+		if (Function == nullptr)
+#endif
+		{
+			// we cant call FindFunction on anything but the game thread as it accesses a shared map in the object's class
+			check(IsInGameThread());
+			Function = InClass->FindFunctionByName(BoundFunction);
+			check(Function);
+		}
+	}
+	else
+	{
+		Function = nullptr;
+	}
+
+	// Cache property access library
+	PropertyAccessLibrary = &InPropertyAccessLibrary;
+}
+
+void FExposedValueHandler::Execute(const FAnimationBaseContext& Context) const
+{
+	if (Function != nullptr)
+	{
+		Context.AnimInstanceProxy->GetAnimInstanceObject()->ProcessEvent(Function, nullptr);
+	}
+
+	if(CopyRecords.Num() > 0)
+	{
+		if(PropertyAccessLibrary != nullptr)
+		{
+			UObject* AnimInstanceObject = Context.AnimInstanceProxy->GetAnimInstanceObject();
+			for(const FExposedValueCopyRecord& CopyRecord : CopyRecords)
+			{
+				PropertyAccess::ProcessCopy(AnimInstanceObject, *PropertyAccessLibrary, EPropertyAccessCopyBatch::InternalUnbatched, CopyRecord.CopyIndex, [&CopyRecord](const FProperty* InProperty, void* InAddress)
+				{
+					if(CopyRecord.PostCopyOperation == EPostCopyOperation::LogicalNegateBool)
+					{
+						bool bValue = static_cast<const FBoolProperty*>(InProperty)->GetPropertyValue(InAddress);
+						static_cast<const FBoolProperty*>(InProperty)->SetPropertyValue(InAddress, !bValue);
+					}
+				});
+			}
+		}
+	}
+}
diff --git a/Engine/Source/Runtime/Engine/Private/Animation/InputScaleBias.cpp b/Engine/Source/Runtime/Engine/Private/Animation/InputScaleBias.cpp
--- a/Engine/Source/Runtime/Engine/Private/Animation/InputScaleBias.cpp
+++ b/Engine/Source/Runtime/Engine/Private/Animation/InputScaleBias.cpp
@@ -86,6 +86,7 @@ float FInputScaleBiasClamp::ApplyTo(float Value, float InDeltaTime) const
 	return Result;
 }
 
+#if WITH_EDITOR
 FText FInputScaleBiasClamp::GetFriendlyName(FText InFriendlyName) const
 {
 	FText OutFriendlyName = InFriendlyName;
@@ -157,7 +158,146 @@ FText FInputScaleBiasClamp::GetFriendlyName(FText InFriendlyName) const
 
 	return OutFriendlyName;
 }
+#endif
 
+float FInputScaleBiasClampState::ApplyTo(const FInputScaleBiasClampConstants& InConstants, float InValue, float InDeltaTime)
+{
+	float Result = InValue;
+
+	if (InConstants.bMapRange)
+	{
+		Result = FMath::GetMappedRangeValueUnclamped(InConstants.InRange.ToVector2D(), InConstants.OutRange.ToVector2D(), Result);
+	}
+
+	Result = Result * InConstants.Scale + InConstants.Bias;
+
+	if (InConstants.bClampResult)
+	{
+		Result = FMath::Clamp<float>(Result, InConstants.ClampMin, InConstants.ClampMax);
+	}
+
+	if (InConstants.bInterpResult)
+	{
+		if (bInitialized)
+		{
+			const float InterpSpeed = (Result >= InterpolatedResult) ? InConstants.InterpSpeedIncreasing : InConstants.InterpSpeedDecreasing;
+			Result = FMath::FInterpTo(InterpolatedResult, Result, InDeltaTime, InterpSpeed);
+		}
+
+		InterpolatedResult = Result;
+	}
+
+	bInitialized = true;
+	return Result;
+}
+
+float FInputScaleBiasClampState::ApplyTo(const FInputScaleBiasClampConstants& InConstants, float InValue) const
+{
+	float Result = InValue;
+
+	if (InConstants.bMapRange)
+	{
+		Result = FMath::GetMappedRangeValueUnclamped(InConstants.InRange.ToVector2D(), InConstants.OutRange.ToVector2D(), Result);
+	}
+
+	Result = Result * InConstants.Scale + InConstants.Bias;
+
+	if (InConstants.bClampResult)
+	{
+		Result = FMath::Clamp<float>(Result, InConstants.ClampMin, InConstants.ClampMax);
+	}
+
+	return Result;
+}
+
+#if WITH_EDITOR
+FText FInputScaleBiasClampConstants::GetFriendlyName(FText InFriendlyName) const
+{
+	FText OutFriendlyName = InFriendlyName;
+
+	// MapRange
+	if (bMapRange)
+	{
+		FFormatNamedArguments Args;
+		Args.Add(TEXT("PinFriendlyName"), OutFriendlyName);
+		Args.Add(TEXT("InRangeMin"), InRange.Min);
+		Args.Add(TEXT("InRangeMax"), InRange.Max);
+		Args.Add(TEXT("OutRangeMin"), OutRange.Min);
+		Args.Add(TEXT("OutRangeMax"), OutRange.Max);
+		OutFriendlyName = FText::Format(LOCTEXT("FInputScaleBias_MapRange", "MapRange({PinFriendlyName}, In({InRangeMin}:{InRangeMax}), Out({OutRangeMin}:{OutRangeMax}))"), Args);
+	}
+
+	if (Scale != 1.f)
+	{
+		FFormatNamedArguments Args;
+		Args.Add(TEXT("PinFriendlyName"), OutFriendlyName);
+		Args.Add(TEXT("Scale"), FText::AsNumber(Scale));
+
+		if (Scale == -1.f)
+		{
+			OutFriendlyName = FText::Format(LOCTEXT("FInputScaleBias_Scale", "- {PinFriendlyName}"), Args);
+		}
+		else 
+		{
+			OutFriendlyName = FText::Format(LOCTEXT("FInputScaleBias_ScaleMul", "{Scale} * {PinFriendlyName}"), Args);
+		}
+	}
+
+	if (Bias != 0.f)
+	{
+		FFormatNamedArguments Args;
+		Args.Add(TEXT("PinFriendlyName"), OutFriendlyName);
+		Args.Add(TEXT("Bias"), FText::AsNumber(Bias));
+
+		// '-' Sign already included in Scale above.
+		if (Scale < 0.f)
+		{
+			OutFriendlyName = FText::Format(LOCTEXT("FInputScaleBias_Bias", "{Bias} {PinFriendlyName}"), Args);
+		}
+		else
+		{
+			OutFriendlyName = FText::Format(LOCTEXT("FInputScaleBias_BiasPlus", "{Bias} + {PinFriendlyName}"), Args);
+		}
+	}
+
+	// Clamp
+	if (bClampResult)
+	{
+		FFormatNamedArguments Args;
+		Args.Add(TEXT("PinFriendlyName"), OutFriendlyName);
+		Args.Add(TEXT("ClampMin"), ClampMin);
+		Args.Add(TEXT("ClampMax"), ClampMax);
+		OutFriendlyName = FText::Format(LOCTEXT("FInputScaleBias_Clamp", "Clamp({PinFriendlyName}, {ClampMin}, {ClampMax})"), Args);
+	}
+
+	// Interp
+	if (bInterpResult)
+	{
+		FFormatNamedArguments Args;
+		Args.Add(TEXT("PinFriendlyName"), OutFriendlyName);
+		Args.Add(TEXT("InterpSpeedIncreasing"), InterpSpeedIncreasing);
+		Args.Add(TEXT("InterpSpeedDecreasing"), InterpSpeedDecreasing);
+		OutFriendlyName = FText::Format(LOCTEXT("FInputScaleBias_Interp", "FInterp({PinFriendlyName}, ({InterpSpeedIncreasing}:{InterpSpeedDecreasing}))"), Args);
+	}
+
+	return OutFriendlyName;
+}
+
+void FInputScaleBiasClampConstants::CopyFromLegacy(const FInputScaleBiasClamp& InLegacy)
+{
+	bMapRange = InLegacy.bMapRange;
+	bClampResult = InLegacy.bClampResult;
+	bInterpResult = InLegacy.bInterpResult;
+	InRange = InLegacy.InRange;
+	OutRange = InLegacy.OutRange;
+	Scale = InLegacy.Scale;
+	Bias = InLegacy.Bias;
+	ClampMin = InLegacy.ClampMin;
+	ClampMax = InLegacy.ClampMax;
+	InterpSpeedIncreasing = InLegacy.InterpSpeedIncreasing;
+	InterpSpeedDecreasing = InLegacy.InterpSpeedDecreasing;
+}
+#endif
 
 /////////////////////////////////////////////////////
 // FInputAlphaBool
diff --git a/Engine/Source/Runtime/Engine/Public/Animation/AnimInstanceProxy.h b/Engine/Source/Runtime/Engine/Public/Animation/AnimInstanceProxy.h
--- a/Engine/Source/Runtime/Engine/Public/Animation/AnimInstanceProxy.h
+++ b/Engine/Source/Runtime/Engine/Public/Animation/AnimInstanceProxy.h
@@ -314,21 +314,41 @@ public:
 	const UObject* GetAnimInstanceObject() const { return AnimInstanceObject; }
 
 	/** Gets an unchecked (can return nullptr) node given an index into the node property array */
-	FAnimNode_Base* GetNodeFromIndexUntyped(int32 NodeIdx, UScriptStruct* RequiredStructType);
+	FAnimNode_Base* GetMutableNodeFromIndexUntyped(int32 NodeIdx, UScriptStruct* RequiredStructType);
+
+	/** Gets an unchecked (can return nullptr) node given an index into the node property array */
+	const FAnimNode_Base* GetNodeFromIndexUntyped(int32 NodeIdx, UScriptStruct* RequiredStructType) const;
 
 	/** Gets a checked node given an index into the node property array */
-	FAnimNode_Base* GetCheckedNodeFromIndexUntyped(int32 NodeIdx, UScriptStruct* RequiredStructType);
+	FAnimNode_Base* GetCheckedMutableNodeFromIndexUntyped(int32 NodeIdx, UScriptStruct* RequiredStructType);
+
+	/** Gets a checked node given an index into the node property array */
+	const FAnimNode_Base* GetCheckedNodeFromIndexUntyped(int32 NodeIdx, UScriptStruct* RequiredStructType) const;
 
 	/** Gets a checked node given an index into the node property array */
 	template<class NodeType>
-	NodeType* GetCheckedNodeFromIndex(int32 NodeIdx)
+	const NodeType* GetCheckedNodeFromIndex(int32 NodeIdx) const
 	{
 		return (NodeType*)GetCheckedNodeFromIndexUntyped(NodeIdx, NodeType::StaticStruct());
 	}
 
+	/** Gets a checked node given an index into the node property array */
+	template<class NodeType>
+	NodeType* GetCheckedMutableNodeFromIndex(int32 NodeIdx) const
+	{
+		return (NodeType*)GetCheckedMutableNodeFromIndexUntyped(NodeIdx, NodeType::StaticStruct());
+	}
+
 	/** Gets an unchecked (can return nullptr) node given an index into the node property array */
 	template<class NodeType>
-	NodeType* GetNodeFromIndex(int32 NodeIdx)
+	const NodeType* GetNodeFromIndex(int32 NodeIdx) const
+	{
+		return (NodeType*)GetNodeFromIndexUntyped(NodeIdx, NodeType::StaticStruct());
+	}
+
+	/** Gets an unchecked (can return nullptr) node given an index into the node property array */
+	template<class NodeType>
+	NodeType* GetMutableNodeFromIndex(int32 NodeIdx) const
 	{
 		return (NodeType*)GetNodeFromIndexUntyped(NodeIdx, NodeType::StaticStruct());
 	}
@@ -497,16 +517,16 @@ public:
 	}
 
 	/** Gets the runtime instance of the specified state machine by Name */
-	FAnimNode_StateMachine* GetStateMachineInstanceFromName(FName MachineName);
+	const FAnimNode_StateMachine* GetStateMachineInstanceFromName(FName MachineName) const;
 
 	/** Get the machine description for the specified instance. Does not rely on PRIVATE_MachineDescription being initialized */
-	static const FBakedAnimationStateMachine* GetMachineDescription(IAnimClassInterface* AnimBlueprintClass, FAnimNode_StateMachine* MachineInstance);
+	static const FBakedAnimationStateMachine* GetMachineDescription(IAnimClassInterface* AnimBlueprintClass, const FAnimNode_StateMachine* MachineInstance);
 
 	/** 
 	 * Get the index of the specified instance asset player. Useful to pass to GetInstanceAssetPlayerLength (etc.).
 	 * Passing NAME_None to InstanceName will return the first (assumed only) player instance index found.
 	 */
-	int32 GetInstanceAssetPlayerIndex(FName MachineName, FName StateName, FName InstanceName = NAME_None);
+	int32 GetInstanceAssetPlayerIndex(FName MachineName, FName StateName, FName InstanceName = NAME_None) const;
 
 	float GetRecordedMachineWeight(const int32 InMachineClassIndex) const;
 	void RecordMachineWeight(const int32 InMachineClassIndex, const float InMachineWeight);
@@ -519,7 +539,10 @@ public:
 	void ResetDynamics(ETeleportType InTeleportType);
 
 	/** Returns all Animation Nodes of FAnimNode_AssetPlayerBase class within the specified (named) Animation Graph */
-	TArray<FAnimNode_AssetPlayerBase*> GetInstanceAssetPlayers(const FName& GraphName);
+	TArray<const FAnimNode_AssetPlayerBase*> GetInstanceAssetPlayers(const FName& GraphName) const;
+
+	/** Returns all Animation Nodes of FAnimNode_AssetPlayerBase class within the specified (named) Animation Graph */
+	TArray<FAnimNode_AssetPlayerBase*> GetMutableInstanceAssetPlayers(const FName& GraphName);
 
 	/** Returns true if SyncGroupName is valid (exists, and if is based on markers, has valid markers) */
 	bool IsSyncGroupValid(FName InSyncGroupName) const;
@@ -775,10 +798,10 @@ protected:
 
 	/** Check whether we have active morph target curves */
 	/** Gets the most relevant asset player in a specified state */
-	FAnimNode_AssetPlayerBase* GetRelevantAssetPlayerFromState(int32 MachineIndex, int32 StateIndex);
+	const FAnimNode_AssetPlayerBase* GetRelevantAssetPlayerFromState(int32 MachineIndex, int32 StateIndex) const;
 
 	/** Gets the runtime instance of the specified state machine */
-	FAnimNode_StateMachine* GetStateMachineInstance(int32 MachineIndex);
+	const FAnimNode_StateMachine* GetStateMachineInstance(int32 MachineIndex) const;
 
 	/** Gets an unchecked (can return nullptr) node given a property of the anim instance */
 	template<class NodeType>
@@ -788,52 +811,52 @@ protected:
 	}
 
 	/** Gets the length in seconds of the asset referenced in an asset player node */
-	float GetInstanceAssetPlayerLength(int32 AssetPlayerIndex);
+	float GetInstanceAssetPlayerLength(int32 AssetPlayerIndex) const;
 
 	/** Get the current accumulated time in seconds for an asset player node */
-	float GetInstanceAssetPlayerTime(int32 AssetPlayerIndex);
+	float GetInstanceAssetPlayerTime(int32 AssetPlayerIndex) const;
 
 	/** Get the current accumulated time as a fraction for an asset player node */
-	float GetInstanceAssetPlayerTimeFraction(int32 AssetPlayerIndex);
+	float GetInstanceAssetPlayerTimeFraction(int32 AssetPlayerIndex) const;
 
 	/** Get the time in seconds from the end of an animation in an asset player node */
-	float GetInstanceAssetPlayerTimeFromEnd(int32 AssetPlayerIndex);
+	float GetInstanceAssetPlayerTimeFromEnd(int32 AssetPlayerIndex) const;
 
 	/** Get the time as a fraction of the asset length of an animation in an asset player node */
-	float GetInstanceAssetPlayerTimeFromEndFraction(int32 AssetPlayerIndex);
+	float GetInstanceAssetPlayerTimeFromEndFraction(int32 AssetPlayerIndex) const;
 
 	/** Get the blend weight of a specified state */
-	float GetInstanceMachineWeight(int32 MachineIndex);
+	float GetInstanceMachineWeight(int32 MachineIndex) const;
 
 	/** Get the blend weight of a specified state */
-	float GetInstanceStateWeight(int32 MachineIndex, int32 StateIndex);
+	float GetInstanceStateWeight(int32 MachineIndex, int32 StateIndex) const;
 
 	/** Get the current elapsed time of a state within the specified state machine */
-	float GetInstanceCurrentStateElapsedTime(int32 MachineIndex);
+	float GetInstanceCurrentStateElapsedTime(int32 MachineIndex) const;
 
 	/** Get the crossfade duration of a specified transition */
-	float GetInstanceTransitionCrossfadeDuration(int32 MachineIndex, int32 TransitionIndex);
+	float GetInstanceTransitionCrossfadeDuration(int32 MachineIndex, int32 TransitionIndex) const;
 
 	/** Get the elapsed time in seconds of a specified transition */
-	float GetInstanceTransitionTimeElapsed(int32 MachineIndex, int32 TransitionIndex);
+	float GetInstanceTransitionTimeElapsed(int32 MachineIndex, int32 TransitionIndex) const;
 
 	/** Get the elapsed time as a fraction of the crossfade duration of a specified transition */
-	float GetInstanceTransitionTimeElapsedFraction(int32 MachineIndex, int32 TransitionIndex);
+	float GetInstanceTransitionTimeElapsedFraction(int32 MachineIndex, int32 TransitionIndex) const;
 
 	/** Get the time remaining in seconds for the most relevant animation in the source state */
-	float GetRelevantAnimTimeRemaining(int32 MachineIndex, int32 StateIndex);
+	float GetRelevantAnimTimeRemaining(int32 MachineIndex, int32 StateIndex) const;
 
 	/** Get the time remaining as a fraction of the duration for the most relevant animation in the source state */
-	float GetRelevantAnimTimeRemainingFraction(int32 MachineIndex, int32 StateIndex);
+	float GetRelevantAnimTimeRemainingFraction(int32 MachineIndex, int32 StateIndex) const;
 
 	/** Get the length in seconds of the most relevant animation in the source state */
-	float GetRelevantAnimLength(int32 MachineIndex, int32 StateIndex);
+	float GetRelevantAnimLength(int32 MachineIndex, int32 StateIndex) const;
 
 	/** Get the current accumulated time in seconds for the most relevant animation in the source state */
-	float GetRelevantAnimTime(int32 MachineIndex, int32 StateIndex);
+	float GetRelevantAnimTime(int32 MachineIndex, int32 StateIndex) const;
 
 	/** Get the current accumulated time as a fraction of the length of the most relevant animation in the source state */
-	float GetRelevantAnimTimeFraction(int32 MachineIndex, int32 StateIndex);
+	float GetRelevantAnimTimeFraction(int32 MachineIndex, int32 StateIndex) const;
 
 	// Sets up a native transition delegate between states with PrevStateName and NextStateName, in the state machine with name MachineName.
 	// Note that a transition already has to exist for this to succeed
@@ -858,12 +881,12 @@ protected:
 	void BindNativeDelegates();
 
 	/** Gets the runtime instance desc of the state machine specified by name */
-	const FBakedAnimationStateMachine* GetStateMachineInstanceDesc(FName MachineName);
+	const FBakedAnimationStateMachine* GetStateMachineInstanceDesc(FName MachineName) const;
 
 	/** Gets the index of the state machine matching MachineName */
-	int32 GetStateMachineIndex(FName MachineName);
+	int32 GetStateMachineIndex(FName MachineName) const;
 
-	void GetStateMachineIndexAndDescription(FName InMachineName, int32& OutMachineIndex, const FBakedAnimationStateMachine** OutMachineDescription);
+	void GetStateMachineIndexAndDescription(FName InMachineName, int32& OutMachineIndex, const FBakedAnimationStateMachine** OutMachineDescription) const;
 
 	/** Initialize the root node - split into a separate function for backwards compatibility (initialization order) reasons */
 	void InitializeRootNode(bool bInDeferRootNodeInitialization = false);
diff --git a/Engine/Source/Runtime/Engine/Public/Animation/AnimSubsystem_Base.h b/Engine/Source/Runtime/Engine/Public/Animation/AnimSubsystem_Base.h
new file mode 100644
--- /dev/null
+++ b/Engine/Source/Runtime/Engine/Public/Animation/AnimSubsystem_Base.h
@@ -0,0 +1,26 @@
+// Copyright Epic Games, Inc. All Rights Reserved.
+
+#pragma once
+
+#include "AnimSubsystem.h"
+#include "Animation/ExposedValueHandler.h"
+#include "AnimSubsystem_Base.generated.h"
+
+/** Provides common built-in anim node functionality as a subsystem */
+USTRUCT()
+struct ENGINE_API FAnimSubsystem_Base : public FAnimSubsystem
+{
+	GENERATED_BODY()
+
+	friend class UAnimBlueprintExtension_Base;
+
+	// FAnimSubsystem interface
+	virtual void OnPostLoadDefaults(FAnimSubsystemPostLoadDefaultsContext& InContext) override;
+
+	// Get the exposed value handlers held on this subsystem
+	const TArray<FExposedValueHandler>& GetExposedValueHandlers() const { return ExposedValueHandlers; }
+
+private:
+	UPROPERTY()
+	TArray<FExposedValueHandler> ExposedValueHandlers;
+};
\ No newline at end of file
diff --git a/Engine/Source/Runtime/Engine/Public/Animation/AnimSubsystem_PropertyAccess.h b/Engine/Source/Runtime/Engine/Public/Animation/AnimSubsystem_PropertyAccess.h
new file mode 100644
--- /dev/null
+++ b/Engine/Source/Runtime/Engine/Public/Animation/AnimSubsystem_PropertyAccess.h
@@ -0,0 +1,28 @@
+// Copyright Epic Games, Inc. All Rights Reserved.
+
+#pragma once
+
+#include "CoreMinimal.h"
+#include "Animation/AnimSubsystem.h"
+#include "PropertyAccess.h"
+#include "AnimSubsystem_PropertyAccess.generated.h"
+
+USTRUCT()
+struct ENGINE_API FAnimSubsystem_PropertyAccess : public FAnimSubsystem
+{
+	GENERATED_BODY()
+
+	friend class UAnimBlueprintExtension_PropertyAccess;
+
+	/** FAnimSubsystem interface */
+	virtual void OnUpdate(FAnimSubsystemUpdateContext& InContext) const override;
+	virtual void OnParallelUpdate(FAnimSubsystemParallelUpdateContext& InContext) const override;
+	virtual void OnPostLoad(FAnimSubsystemPostLoadContext& InContext) override;
+
+	/** Access the property access library */
+	const FPropertyAccessLibrary& GetLibrary() const { return Library; }
+
+private:
+	UPROPERTY()
+	FPropertyAccessLibrary Library;
+};
\ No newline at end of file
diff --git a/Engine/Source/Runtime/Engine/Public/Animation/AnimTrace.h b/Engine/Source/Runtime/Engine/Public/Animation/AnimTrace.h
--- a/Engine/Source/Runtime/Engine/Public/Animation/AnimTrace.h
+++ b/Engine/Source/Runtime/Engine/Public/Animation/AnimTrace.h
@@ -31,7 +31,7 @@ struct FComponentSpacePoseContext;
 class FName;
 struct FVector;
 struct FRotator;
-struct FAnimNode_SequencePlayer;
+struct FAnimNode_SequencePlayerBase;
 struct FAnimNotifyEvent;
 struct FPassedMarker;
 struct FAnimSyncMarker;
@@ -166,7 +166,7 @@ struct FAnimTrace
 	ENGINE_API static void OutputAnimNodeValue(const FAnimationBaseContext& InContext, const TCHAR* InKey, const UObject* InValue);
 
 	/** Helper function to output debug info for sequence player nodes */
-	ENGINE_API static void OutputAnimSequencePlayer(const FAnimationBaseContext& InContext, const FAnimNode_SequencePlayer& InNode);
+	ENGINE_API static void OutputAnimSequencePlayer(const FAnimationBaseContext& InContext, const FAnimNode_SequencePlayerBase& InNode);
 
 	/** 
 	 * Helper function to output a name to the trace stream, referenced by ID. 
diff --git a/Engine/Source/Runtime/Engine/Public/Animation/ExposedValueHandler.h b/Engine/Source/Runtime/Engine/Public/Animation/ExposedValueHandler.h
new file mode 100644
--- /dev/null
+++ b/Engine/Source/Runtime/Engine/Public/Animation/ExposedValueHandler.h
@@ -0,0 +1,79 @@
+// Copyright Epic Games, Inc. All Rights Reserved.
+
+#pragma once
+
+#include "CoreMinimal.h"
+#include "UObject/UnrealType.h"
+
+#include "ExposedValueHandler.generated.h"
+
+struct FPropertyAccessLibrary;
+class UDynamicClass;
+class UClass;
+struct FAnimationBaseContext;
+
+UENUM()
+enum class EPostCopyOperation : uint8
+{
+	None,
+
+	LogicalNegateBool,
+};
+
+USTRUCT()
+struct FExposedValueCopyRecord
+{
+	GENERATED_BODY()
+
+	FExposedValueCopyRecord() = default;
+
+	FExposedValueCopyRecord(int32 InCopyIndex, EPostCopyOperation InPostCopyOperation)
+		: CopyIndex(InCopyIndex)
+		, PostCopyOperation(InPostCopyOperation)
+	{
+	}
+
+	UPROPERTY()
+	int32 CopyIndex = INDEX_NONE;
+
+	UPROPERTY()
+	EPostCopyOperation PostCopyOperation = EPostCopyOperation::None;
+};
+
+// An exposed value updater
+USTRUCT()
+struct ENGINE_API FExposedValueHandler
+{
+	GENERATED_USTRUCT_BODY()
+
+	FExposedValueHandler()
+		: Function(nullptr)
+		, PropertyAccessLibrary(nullptr)
+		, BoundFunction(NAME_None)
+	{
+	}
+
+	// Direct data access to property in anim instance
+	UPROPERTY()
+	TArray<FExposedValueCopyRecord> CopyRecords;
+
+	// function pointer if BoundFunction != NAME_None
+	UPROPERTY()
+	TObjectPtr<UFunction> Function;
+
+	// Cached property access library ptr
+	const FPropertyAccessLibrary* PropertyAccessLibrary;
+
+	// The function to call to update associated properties (can be NAME_None)
+	UPROPERTY()
+	FName BoundFunction;
+
+	// Helper function to bind an array of handlers.
+	static void ClassInitialization(TArray<FExposedValueHandler>& Handlers, UClass* InClass);
+
+	// Bind copy records and cache UFunction if necessary
+	void Initialize(UClass* InClass, const FPropertyAccessLibrary& InPropertyAccessLibrary);
+
+	// Execute the function and copy records
+	void Execute(const FAnimationBaseContext& Context) const;
+};
\ No newline at end of file
