Skin cache - add max dispatches per command list option to help with potential dx12 TDRs

diff --git a/Engine/Source/Runtime/Engine/Private/GPUSkinCache.cpp b/Engine/Source/Runtime/Engine/Private/GPUSkinCache.cpp
--- a/Engine/Source/Runtime/Engine/Private/GPUSkinCache.cpp
+++ b/Engine/Source/Runtime/Engine/Private/GPUSkinCache.cpp
@@ -138,6 +138,14 @@ FAutoConsoleVariableRef CVarRecomputeTangentsParallelDispatch(
 	ECVF_RenderThreadSafe
 );
 
+static int32 GSkinCacheMaxDispatchesPerCmdList = 0;
+FAutoConsoleVariableRef CVarGPUSkinCacheMaxDispatchesPerCmdList(
+	TEXT("r.SkinCache.MaxDispatchesPerCmdList"),
+	GSkinCacheMaxDispatchesPerCmdList,
+	TEXT("Maximum number of compute shader dispatches which are batched together into a single command list to fix potential TDRs."),
+	ECVF_RenderThreadSafe
+);
+
 static int32 GGPUSkinCacheFlushCounter = 0;
 
 #if RHI_RAYTRACING
@@ -1223,6 +1231,7 @@ void FGPUSkinCache::DispatchUpdateSkinTangents(FRHICommandListImmediate& RHICmdL
 			Shader->SetParameters(RHICmdList, Entry, DispatchData, GRecomputeTangentsParallelDispatch ? *DispatchData.GetIntermediateAccumulatedTangentBuffer() : *StagingBuffer);
 			DispatchComputeShader(RHICmdList, Shader.GetShader(), ThreadGroupCountValue, 1, 1);
 			Shader->UnsetParameters(RHICmdList);
+			IncrementDispatchCounter(RHICmdList);
 		}
 	}
 	else
@@ -1264,6 +1273,7 @@ void FGPUSkinCache::DispatchUpdateSkinTangents(FRHICommandListImmediate& RHICmdL
 		ComputeShader->SetParameters(RHICmdList, Entry, DispatchData, GRecomputeTangentsParallelDispatch ? *DispatchData.GetIntermediateAccumulatedTangentBuffer() : *StagingBuffer);
 		DispatchComputeShader(RHICmdList, ComputeShader.GetShader(), ThreadGroupCountValue, 1, 1);
 		ComputeShader->UnsetParameters(RHICmdList);
+		IncrementDispatchCounter(RHICmdList);
 	}
 }
 
@@ -1747,6 +1757,7 @@ void FGPUSkinCache::BeginBatchDispatch(FRHICommandListImmediate& RHICmdList)
 {
 	check(BatchDispatches.Num() == 0);
 	bShouldBatchDispatches = true;
+	DispatchCounter = 0;
 }
 
 void FGPUSkinCache::EndBatchDispatch(FRHICommandListImmediate& RHICmdList)
@@ -2051,7 +2062,7 @@ void FGPUSkinCache::DispatchUpdateSkinning(FRHICommandListImmediate& RHICmdList,
 		INC_DWORD_STAT_BY(STAT_GPUSkinCache_TotalNumVertices, VertexCountAlign64 * 64);
 		RHICmdList.DispatchComputeShader(VertexCountAlign64, 1, 1);
 		Shader->UnsetParameters(RHICmdList);
-
+		IncrementDispatchCounter(RHICmdList);
 	}
 
 	if ((DispatchData.DispatchFlags & (uint32)EGPUSkinCacheDispatchFlags::DispatchPosition) != 0)
@@ -2071,6 +2082,7 @@ void FGPUSkinCache::DispatchUpdateSkinning(FRHICommandListImmediate& RHICmdList,
 		INC_DWORD_STAT_BY(STAT_GPUSkinCache_TotalNumVertices, VertexCountAlign64 * 64);
 		RHICmdList.DispatchComputeShader(VertexCountAlign64, 1, 1);
 		Shader->UnsetParameters(RHICmdList);
+		IncrementDispatchCounter(RHICmdList);
 	}
 
 	check(DispatchData.PreviousPositionBuffer != DispatchData.PositionBuffer);
@@ -2244,3 +2256,14 @@ void FGPUSkinCache::CVarSinkFunction()
 }
 
 FAutoConsoleVariableSink FGPUSkinCache::CVarSink(FConsoleCommandDelegate::CreateStatic(&CVarSinkFunction));
+
+void FGPUSkinCache::IncrementDispatchCounter(FRHICommandListImmediate& RHICmdList)
+{
+	DispatchCounter++;
+	if (GSkinCacheMaxDispatchesPerCmdList > 0 && DispatchCounter >= GSkinCacheMaxDispatchesPerCmdList)
+	{
+		//UE_LOG(LogSkinCache, Log, TEXT("SubmitCommandsHint issued after %d dispatches"), DispatchCounter);
+		RHICmdList.SubmitCommandsHint();
+		DispatchCounter = 0;
+	}
+}
\ No newline at end of file
diff --git a/Engine/Source/Runtime/Engine/Public/GPUSkinCache.h b/Engine/Source/Runtime/Engine/Public/GPUSkinCache.h
--- a/Engine/Source/Runtime/Engine/Public/GPUSkinCache.h
+++ b/Engine/Source/Runtime/Engine/Public/GPUSkinCache.h
@@ -469,6 +469,9 @@ protected:
 
 	static void CVarSinkFunction();
 	static FAutoConsoleVariableSink CVarSink;
+
+	void IncrementDispatchCounter(FRHICommandListImmediate& RHICmdList);
+	int32 DispatchCounter = 0;
 };
 
 DECLARE_STATS_GROUP(TEXT("GPU Skin Cache"), STATGROUP_GPUSkinCache, STATCAT_Advanced);
