Fix GetLoadedExternalActorPackages() possibly returning packages of another Level (ex: /Game/PackageA returning external actors of /Game/PackageAB)

diff --git a/Engine/Source/Runtime/Engine/Private/Level.cpp b/Engine/Source/Runtime/Engine/Private/Level.cpp
--- a/Engine/Source/Runtime/Engine/Private/Level.cpp
+++ b/Engine/Source/Runtime/Engine/Private/Level.cpp
@@ -2503,6 +2503,12 @@ TArray<UPackage*> ULevel::GetLoadedExternalActorPackages() const
 	FString ExternalActorsPath = ULevel::GetExternalActorsPath(World->GetPackage(), World->OriginalWorldName == NAME_None ? World->GetName() : World->OriginalWorldName.ToString());
 	if (!ExternalActorsPath.IsEmpty())
 	{
+		// Make sure we filter out similar external folders
+		if (!ExternalActorsPath.EndsWith("/"))
+		{
+			ExternalActorsPath += "/";
+		}
+
 		for (TObjectIterator<UPackage> It; It; ++It)
 		{
 			if (!ActorPackages.Contains(*It) && It->GetLoadedPath().GetPackageName().Contains(ExternalActorsPath))
