remove a lot of PLATFORM_XBOXONE from Engine code.

diff --git a/Engine/Plugins/Compression/OodleNetwork/Source/Private/OodleHandlerComponent.cpp b/Engine/Plugins/Compression/OodleNetwork/Source/Private/OodleHandlerComponent.cpp
--- a/Engine/Plugins/Compression/OodleNetwork/Source/Private/OodleHandlerComponent.cpp
+++ b/Engine/Plugins/Compression/OodleNetwork/Source/Private/OodleHandlerComponent.cpp
@@ -29,6 +29,9 @@ DEFINE_LOG_CATEGORY(OodleHandlerComponentLog);
 #define OODLE_HANDLER_VERBOSE_LOG	1
 #endif
 
+#ifndef OODLE_USE_FALLBACK_DICTIONARY
+	#define OODLE_USE_FALLBACK_DICTIONARY !(PLATFORM_PS4 || PLATFORM_SWITCH)
+#endif
 
 // @todo #JohnB: You're not taking into account, the overhead of sending 'DecompressedLength', in the stats
 
@@ -524,7 +527,7 @@ void OodleHandlerComponent::InitializeDictionaries()
 
 	bInitializedDictionaries = true;
 
-#if (!UE_BUILD_SHIPPING || OODLE_DEV_SHIPPING) && !(PLATFORM_PS4 || PLATFORM_XBOXONE || PLATFORM_SWITCH)
+#if (!UE_BUILD_SHIPPING || OODLE_DEV_SHIPPING) && OODLE_USE_FALLBACK_DICTIONARY
 	if (bUseDictionaryIfPresent)
 	{
 		bGotDictionaryPath = FindFallbackDictionaries(ServerDictionaryPath, ClientDictionaryPath);
diff --git a/Engine/Plugins/Experimental/CommonUI/Source/CommonInput/CommonInput.Build.cs b/Engine/Plugins/Experimental/CommonUI/Source/CommonInput/CommonInput.Build.cs
--- a/Engine/Plugins/Experimental/CommonUI/Source/CommonInput/CommonInput.Build.cs
+++ b/Engine/Plugins/Experimental/CommonUI/Source/CommonInput/CommonInput.Build.cs
@@ -47,6 +47,9 @@ public class CommonInput : ModuleRules
 		{
 			PublicDefinitions.Add("UE_COMMONINPUT_PLATFORM_TYPE = " + CommonUIPlatform);
 		}
+
+		PrivateDefinitions.Add("UE_COMMONINPUT_PLATFORM_USE_GAMEPAD_IF_MOUSE_REMOVED=" + (PlatformUsesGamepadIfMouseRemoved ? "1" : "0") );
+		PrivateDefinitions.Add("UE_COMMONINPUT_PLATFORM_KBM_REQUIRES_ATTACHED_MOUSE=" + (PlatformKBMRequiresAttachedMouse ? "1" : "0") );
 	}
 
 	static public string ToCommonUIPlatform(UnrealTargetPlatform TargetPlatform)
@@ -74,4 +77,8 @@ public class CommonInput : ModuleRules
 
 		return string.Empty;
 	}
+
+	protected virtual bool PlatformUsesGamepadIfMouseRemoved { get { return false; } }
+	protected virtual bool PlatformKBMRequiresAttachedMouse { get { return false; } }
+
 }
diff --git a/Engine/Plugins/Experimental/CommonUI/Source/CommonInput/Private/CommonInputSubsystem.cpp b/Engine/Plugins/Experimental/CommonUI/Source/CommonInput/Private/CommonInputSubsystem.cpp
--- a/Engine/Plugins/Experimental/CommonUI/Source/CommonInput/Private/CommonInputSubsystem.cpp
+++ b/Engine/Plugins/Experimental/CommonUI/Source/CommonInput/Private/CommonInputSubsystem.cpp
@@ -42,7 +42,7 @@ public:
 	
 	virtual void Tick(const float DeltaTime, FSlateApplication& SlateApp, TSharedRef<ICursor> Cursor) override
 	{
-#if PLATFORM_XBOXONE
+#if UE_COMMONINPUT_PLATFORM_USE_GAMEPAD_IF_MOUSE_REMOVED
 		//fall back to gamepad controls if the mouse is removed... seems to stop slate getting confused and preventing gamepad input
 		if( InputSubsystem.GetCurrentInputType() == ECommonInputType::MouseAndKeyboard && !FSlateApplication::Get().IsMouseAttached() )
 		{
@@ -682,7 +682,7 @@ bool UCommonInputSubsystem::PlatformSupportsInputType(ECommonInputType InInputTy
 	{
 	case ECommonInputType::MouseAndKeyboard:
 	{
-#if PLATFORM_XBOXONE || PLATFORM_SWITCH
+#if UE_COMMONINPUT_PLATFORM_KBM_REQUIRES_ATTACHED_MOUSE || PLATFORM_SWITCH
 		bPlatformSupportsInput &= FSlateApplication::Get().IsMouseAttached();
 #endif
 	}
diff --git a/Engine/Plugins/Media/MfMedia/Source/MfMedia/Private/Mf/MfMediaUtils.cpp b/Engine/Plugins/Media/MfMedia/Source/MfMedia/Private/Mf/MfMediaUtils.cpp
--- a/Engine/Plugins/Media/MfMedia/Source/MfMedia/Private/Mf/MfMediaUtils.cpp
+++ b/Engine/Plugins/Media/MfMedia/Source/MfMedia/Private/Mf/MfMediaUtils.cpp
@@ -20,15 +20,15 @@
 #endif
 
 #ifndef MFMEDIA_PLATFORM_SUPPORTS_ONLY_AAC
-	#define MFMEDIA_PLATFORM_SUPPORTS_ONLY_AAC PLATFORM_XBOXONE
+	#define MFMEDIA_PLATFORM_SUPPORTS_ONLY_AAC 0
 #endif
 
 #ifndef MFMEDIA_PLATFORM_SUPPORTS_ONLY_H264
-	#define MFMEDIA_PLATFORM_SUPPORTS_ONLY_H264 PLATFORM_XBOXONE
+	#define MFMEDIA_PLATFORM_SUPPORTS_ONLY_H264 0
 #endif
 
 #ifndef  MFMEDIA_PLATFORM_SUPPORTS_ONLY_NV12
-	#define MFMEDIA_PLATFORM_SUPPORTS_ONLY_NV12 PLATFORM_XBOXONE
+	#define MFMEDIA_PLATFORM_SUPPORTS_ONLY_NV12 0
 #endif
 
 
diff --git a/Engine/Plugins/Media/MfMedia/Source/MfMedia/Private/MfMediaPrivate.h b/Engine/Plugins/Media/MfMedia/Source/MfMedia/Private/MfMediaPrivate.h
--- a/Engine/Plugins/Media/MfMedia/Source/MfMedia/Private/MfMediaPrivate.h
+++ b/Engine/Plugins/Media/MfMedia/Source/MfMedia/Private/MfMediaPrivate.h
@@ -5,7 +5,7 @@
 #include "Logging/LogMacros.h"
 
 #if !defined(MFMEDIA_SUPPORTED_PLATFORM)
-	#define MFMEDIA_SUPPORTED_PLATFORM (PLATFORM_XBOXONE || PLATFORM_HOLOLENS || (PLATFORM_WINDOWS && !UE_SERVER))
+	#define MFMEDIA_SUPPORTED_PLATFORM (PLATFORM_HOLOLENS || (PLATFORM_WINDOWS && !UE_SERVER))
 	#define MFMEDIA_NEED_PLATFORM_PRIVATE 0
 #else
 	#define MFMEDIA_NEED_PLATFORM_PRIVATE 1
diff --git a/Engine/Plugins/Runtime/ChunkDownloader/Source/Private/ChunkDownloader.cpp b/Engine/Plugins/Runtime/ChunkDownloader/Source/Private/ChunkDownloader.cpp
--- a/Engine/Plugins/Runtime/ChunkDownloader/Source/Private/ChunkDownloader.cpp
+++ b/Engine/Plugins/Runtime/ChunkDownloader/Source/Private/ChunkDownloader.cpp
@@ -1179,12 +1179,7 @@ void FChunkDownloader::TryDownloadBuildManifest(int TryNumber)
 	check(BuildBaseUrls.Num() > 0);
 
 	// download the manifest from CDN, then load it
-#if PLATFORM_PS4 || PLATFORM_SWITCH || PLATFORM_XBOXONE
-	// TODO: Remove when the other platforms' build manifests exist in the CDN
-	FString ManifestFileName = FString::Printf(TEXT("BuildManifest-Windows.txt"));
-#else
 	FString ManifestFileName = FString::Printf(TEXT("BuildManifest-%s.txt"), *PlatformName);
-#endif
 	FString Url = BuildBaseUrls[TryNumber % BuildBaseUrls.Num()] / ManifestFileName;
 	UE_LOG(LogChunkDownloader, Log, TEXT("Downloading build manifest (attempt #%d) from %s"), TryNumber+1, *Url);
 
diff --git a/Engine/Source/Runtime/Engine/Private/OnlineReplStructs.cpp b/Engine/Source/Runtime/Engine/Private/OnlineReplStructs.cpp
--- a/Engine/Source/Runtime/Engine/Private/OnlineReplStructs.cpp
+++ b/Engine/Source/Runtime/Engine/Private/OnlineReplStructs.cpp
@@ -775,14 +775,12 @@ void TestUniqueIdRepl(UWorld* InWorld)
 	}
 
 	bool bPlatformSerializationSuccess = true;
-#if PLATFORM_XBOXONE || PLATFORM_PS4
-	if (bSetupSuccess)
+	FString NativePlatformService;
+	if (bSetupSuccess && 
+		GConfig->GetString(TEXT("OnlineSubsystem"), TEXT("NativePlatformService"), NativePlatformService, GEngineIni) && 
+		!NativePlatformService.IsEmpty())
 	{
-#if PLATFORM_XBOXONE
-		FUniqueNetIdPtr PlatformUserId = UOnlineEngineInterface::Get()->GetUniquePlayerId(InWorld, 0, FName(TEXT("LIVE")));
-#elif PLATFORM_PS4
-		FUniqueNetIdPtr PlatformUserId = UOnlineEngineInterface::Get()->GetUniquePlayerId(InWorld, 0, FName(TEXT("PS4")));
-#endif
+		FUniqueNetIdPtr PlatformUserId = UOnlineEngineInterface::Get()->GetUniquePlayerId(InWorld, 0, FName(*NativePlatformService));
 
 		FUniqueNetIdRepl ValidPlatformIdIn(PlatformUserId);
 		if (!ValidPlatformIdIn.IsValid() || PlatformUserId != ValidPlatformIdIn.GetUniqueNetId() || *PlatformUserId != *ValidPlatformIdIn)
@@ -833,7 +831,6 @@ void TestUniqueIdRepl(UWorld* InWorld)
 			CHECK_REPL_EQUALITY(ValidPlatformIdIn, ValidPlatformIdOut, bPlatformSerializationSuccess);
 		}
 	}
-#endif
 
 	bool bJSONSerializationSuccess = true;
 	if (bSetupSuccess)
diff --git a/Engine/Source/Runtime/PreLoadScreen/Private/PreLoadScreenManager.cpp b/Engine/Source/Runtime/PreLoadScreen/Private/PreLoadScreenManager.cpp
--- a/Engine/Source/Runtime/PreLoadScreen/Private/PreLoadScreenManager.cpp
+++ b/Engine/Source/Runtime/PreLoadScreen/Private/PreLoadScreenManager.cpp
@@ -466,10 +466,6 @@ void FPreLoadScreenManager::PlatformSpecificGameLogicFrameTick()
 #if PLATFORM_IOS
 	IOS_PlatformSpecificGameLogicFrameTick();
 #endif //PLATFORM_IOS
-
-#if PLATFORM_XBOXONE
-	XboxOne_PlatformSpecificGameLogicFrameTick();
-#endif // PLATFORM_XBOXONE
 }
 
 void FPreLoadScreenManager::EnableRendering(bool bEnabled)
@@ -724,11 +720,3 @@ void FPreLoadScreenManager::IOS_PlatformSpecificGameLogicFrameTick()
 	[FIOSAsyncTask ProcessAsyncTasks];
 }
 #endif //PLATFORM_IOS
-
-#if PLATFORM_XBOXONE
-void FPreLoadScreenManager::XboxOne_PlatformSpecificGameLogicFrameTick()
-{
-	// Xbox doesn't run engine init on the main thread, so we need to flush logging periodically
-	GLog->FlushThreadedLogs();
-}
-#endif //PLATFORM_XBOXONE
diff --git a/Engine/Source/Runtime/PreLoadScreen/Public/PreLoadScreenManager.h b/Engine/Source/Runtime/PreLoadScreen/Public/PreLoadScreenManager.h
--- a/Engine/Source/Runtime/PreLoadScreen/Public/PreLoadScreenManager.h
+++ b/Engine/Source/Runtime/PreLoadScreen/Public/PreLoadScreenManager.h
@@ -155,8 +155,4 @@ private:
 #if PLATFORM_IOS
 	void IOS_PlatformSpecificGameLogicFrameTick();
 #endif //PLATFORM_IOS
-
-#if PLATFORM_XBOXONE
-	void XboxOne_PlatformSpecificGameLogicFrameTick();
-#endif //PLATFORM_XBOXONE
 };
