Added a concenience function to NavMovementComponent to fetch UpdatedComponents world transform

diff --git a/Engine/Source/Runtime/Engine/Classes/GameFramework/NavMovementComponent.h b/Engine/Source/Runtime/Engine/Classes/GameFramework/NavMovementComponent.h
--- a/Engine/Source/Runtime/Engine/Classes/GameFramework/NavMovementComponent.h
+++ b/Engine/Source/Runtime/Engine/Classes/GameFramework/NavMovementComponent.h
@@ -87,6 +87,8 @@ public:
 	virtual FBasedPosition GetActorFeetLocationBased() const;
 	/** Returns navigation location of controlled actor */
 	FORCEINLINE FVector GetActorNavLocation() const { INavAgentInterface* MyOwner = Cast<INavAgentInterface>(GetOwner()); return MyOwner ? MyOwner->GetNavAgentLocation() : FNavigationSystem::InvalidLocation; }
+	/** Returns the full world-coordinates transform of the associated scene component (the UpdatedComponent) */
+	FORCEINLINE FTransform GetActorTransform() const { return UpdatedComponent ? UpdatedComponent->GetComponentTransform() : FTransform(); }
 
 	/** path following: request new velocity */
 	virtual void RequestDirectMove(const FVector& MoveVelocity, bool bForceMaxSpeed);
