Instrumenting and adding counts to measure shader compilation.

- Adding material name metadata to CacheResourceShadersForRendering() so you can see in Unreal Insights what material is being compiled.
- Add scope markers for the following functions to track unaccounted for time and to see where we are blocking:
    - SaveToDerivedDataCache
    - BlockOnShaderMapCompletion
    - BlockOnAllShaderMapCompletion
    - ProcessCompiledShaderMaps
- Added a counter for the following items so you can see them graphed over time in Unreal Insights:
    - On demand shader requests: How many shaders has rendering requested? aka. What are we actually using.
    - Num total shaders compiled
    - JobCacheSearchAttempts
    - JobCacheHits

diff --git a/Engine/Source/Runtime/Engine/Private/Materials/Material.cpp b/Engine/Source/Runtime/Engine/Private/Materials/Material.cpp
--- a/Engine/Source/Runtime/Engine/Private/Materials/Material.cpp
+++ b/Engine/Source/Runtime/Engine/Private/Materials/Material.cpp
@@ -3058,9 +3058,16 @@ void UMaterial::UpdateMaterialShaderCacheAndTextureReferences()
 
 #endif //WITH_EDITOR
 
+UE_TRACE_EVENT_BEGIN(Cpu, CacheResourceShadersForRendering, NoSync)
+	UE_TRACE_EVENT_FIELD(UE::Trace::WideString, MaterialName)
+UE_TRACE_EVENT_END()
+
 void UMaterial::CacheResourceShadersForRendering(bool bRegenerateId, EMaterialShaderPrecompileMode PrecompileMode)
 {
 	TRACE_CPUPROFILER_EVENT_SCOPE(UMaterial::CacheResourceShadersForRendering);
+	UE_TRACE_LOG_SCOPED_T(Cpu, CacheResourceShadersForRendering, CpuChannel)
+		<< CacheResourceShadersForRendering.MaterialName(*GetFullName());
+
 #if WITH_EDITOR
 	// Always rebuild the shading model field on recompile
 	RebuildShadingModelField();
diff --git a/Engine/Source/Runtime/Engine/Private/Materials/MaterialShader.cpp b/Engine/Source/Runtime/Engine/Private/Materials/MaterialShader.cpp
--- a/Engine/Source/Runtime/Engine/Private/Materials/MaterialShader.cpp
+++ b/Engine/Source/Runtime/Engine/Private/Materials/MaterialShader.cpp
@@ -1171,6 +1171,8 @@ void FMaterialShaderMap::GetAllOutdatedTypes(TArray<const FShaderType*>& Outdate
 
 void FMaterialShaderMap::LoadFromDerivedDataCache(const FMaterial* Material, const FMaterialShaderMapId& ShaderMapId, EShaderPlatform InPlatform, const ITargetPlatform* TargetPlatform, TRefCountPtr<FMaterialShaderMap>& InOutShaderMap, FString& OutDDCKeyDesc)
 {
+	TRACE_CPUPROFILER_EVENT_SCOPE(FMaterialShaderMap::LoadFromDerivedDataCache);
+
 	if (InOutShaderMap != NULL)
 	{
 		check(InOutShaderMap->GetShaderPlatform() == InPlatform);
@@ -1249,6 +1251,7 @@ void FMaterialShaderMap::LoadFromDerivedDataCache(const FMaterial* Material, con
 
 void FMaterialShaderMap::SaveToDerivedDataCache()
 {
+	TRACE_CPUPROFILER_EVENT_SCOPE(FMaterialShaderMap::SaveToDerivedDataCache);
 	COOK_STAT(auto Timer = MaterialShaderCookStats::UsageStats.TimeSyncWork());
 	TArray<uint8> SaveData;
 	FMemoryWriter Ar(SaveData, true);
diff --git a/Engine/Source/Runtime/Engine/Private/Materials/MaterialShared.cpp b/Engine/Source/Runtime/Engine/Private/Materials/MaterialShared.cpp
--- a/Engine/Source/Runtime/Engine/Private/Materials/MaterialShared.cpp
+++ b/Engine/Source/Runtime/Engine/Private/Materials/MaterialShared.cpp
@@ -56,6 +56,7 @@
 #if WITH_ODSC
 #include "ODSC/ODSCManager.h"
 #endif
+#include "ProfilingDebugging/CountersTrace.h"
 
 #define LOCTEXT_NAMESPACE "MaterialShared"
 
@@ -2553,6 +2554,7 @@ TShaderRef<FShader> FMaterial::GetShader(FMeshMaterialShaderType* ShaderType, FV
 	return TShaderRef<FShader>(Shader, *RenderingThreadShaderMap);
 }
 
+TRACE_DECLARE_INT_COUNTER(Shaders_OnDemandShaderRequests, TEXT("Shaders/OnDemandShaderRequests"));
 bool FMaterial::TryGetShaders(const FMaterialShaderTypes& InTypes, const FVertexFactoryType* InVertexFactoryType, FMaterialShaders& OutShaders) const
 {
 	// TRACE_CPUPROFILER_EVENT_SCOPE(FMaterial::TryGetShaders); // <-- disabled by default due to verbosity (hundreds of calls per frame)
@@ -2693,6 +2695,7 @@ bool FMaterial::TryGetShaders(const FMaterialShaderTypes& InTypes, const FVertex
 
 	if (CompileJobs.Num() > 0)
 	{
+		TRACE_COUNTER_ADD(Shaders_OnDemandShaderRequests, CompileJobs.Num());
 		GShaderCompilingManager->SubmitJobs(CompileJobs, GetBaseMaterialPathName(), ShaderMap->GetDebugDescription());
 	}
 
diff --git a/Engine/Source/Runtime/Engine/Private/ShaderCompiler/ShaderCompiler.cpp b/Engine/Source/Runtime/Engine/Private/ShaderCompiler/ShaderCompiler.cpp
--- a/Engine/Source/Runtime/Engine/Private/ShaderCompiler/ShaderCompiler.cpp
+++ b/Engine/Source/Runtime/Engine/Private/ShaderCompiler/ShaderCompiler.cpp
@@ -58,6 +58,7 @@
 #include "ObjectCacheContext.h"
 #include "ProfilingDebugging/StallDetector.h"
 #include "RenderUtils.h"
+#include "ProfilingDebugging/CountersTrace.h"
 
 #if WITH_EDITOR
 #include "TextureCompiler.h"
@@ -2676,8 +2677,11 @@ void FShaderCompilerStats::RegisterCookedShaders(uint32 NumCooked, float Compile
 	}
 }
 
+TRACE_DECLARE_INT_COUNTER(Shaders_Compiled, TEXT("Shaders/Compiled"));
 void FShaderCompilerStats::RegisterCompiledShaders(uint32 NumCompiled, EShaderPlatform Platform, const FString MaterialPath, FString PermutationString)
 {
+	TRACE_COUNTER_ADD(Shaders_Compiled, NumCompiled);
+
 	FScopeLock Lock(&CompileStatsLock);
 	if (!CompileStats.IsValidIndex(Platform))
 	{
@@ -3280,6 +3284,7 @@ void FShaderCompilingManager::AddCompiledResults(TMap<int32, FShaderMapFinalizeR
 /** Flushes all pending jobs for the given shader maps. */
 void FShaderCompilingManager::BlockOnShaderMapCompletion(const TArray<int32>& ShaderMapIdsToFinishCompiling, TMap<int32, FShaderMapFinalizeResults>& CompiledShaderMaps)
 {
+	TRACE_CPUPROFILER_EVENT_SCOPE(FShaderCompilingManager::BlockOnShaderMapCompletion);
 	SCOPE_STALL_REPORTER(FShaderCompilingManager::BlockOnShaderMapCompletion, 2.0);
 
 	COOK_STAT(FScopedDurationTimer BlockingTimer(ShaderCompilerCookStats::BlockingTimeSec));
@@ -3401,6 +3406,7 @@ void FShaderCompilingManager::BlockOnShaderMapCompletion(const TArray<int32>& Sh
 
 void FShaderCompilingManager::BlockOnAllShaderMapCompletion(TMap<int32, FShaderMapFinalizeResults>& CompiledShaderMaps)
 {
+	TRACE_CPUPROFILER_EVENT_SCOPE(FShaderCompilingManager::BlockOnAllShaderMapCompletion);
 	SCOPE_STALL_REPORTER(FShaderCompilingManager::BlockOnAllShaderMapCompletion, 2.0);
 
 	COOK_STAT(FScopedDurationTimer BlockingTimer(ShaderCompilerCookStats::BlockingTimeSec));
@@ -3541,6 +3547,8 @@ void FShaderCompilingManager::ProcessCompiledShaderMaps(
 	float TimeBudget)
 {
 #if WITH_EDITOR
+	TRACE_CPUPROFILER_EVENT_SCOPE(FShaderCompilingManager::ProcessCompiledShaderMaps);
+
 	TMap<TRefCountPtr<FMaterial>, TRefCountPtr<FMaterialShaderMap>> MaterialsToUpdate;
 	TArray<TRefCountPtr<FMaterial>> MaterialsToReleaseCompilingId;
 
@@ -6677,9 +6685,12 @@ void FShaderPipelineCompileJob::SerializeOutput(FArchive& Ar)
 	}
 }
 
+TRACE_DECLARE_INT_COUNTER(Shaders_JobCacheSearchAttempts, TEXT("Shaders/JobCacheSearchAttempts"));
+TRACE_DECLARE_INT_COUNTER(Shaders_JobCacheHits, TEXT("Shaders/JobCacheHits"));
 FShaderJobCache::FJobCachedOutput* FShaderJobCache::Find(const FJobInputHash& Hash)
 {
 	++TotalSearchAttempts;
+	TRACE_COUNTER_INCREMENT(Shaders_JobCacheSearchAttempts);
 
 	if (ShaderCompiler::IsJobCacheEnabled())
 	{
@@ -6687,6 +6698,7 @@ FShaderJobCache::FJobCachedOutput* FShaderJobCache::Find(const FJobInputHash& Ha
 		if (OutputHash)
 		{
 			++TotalCacheHits;
+			TRACE_COUNTER_INCREMENT(Shaders_JobCacheHits);
 
 			FStoredOutput** CannedOutput = Outputs.Find(*OutputHash);
 			// we should not allow a dangling input to output mapping to exist
