Skin cache - only update dispatch counter when max dispatch per cmd list option is on

diff --git a/Engine/Source/Runtime/Engine/Private/GPUSkinCache.cpp b/Engine/Source/Runtime/Engine/Private/GPUSkinCache.cpp
--- a/Engine/Source/Runtime/Engine/Private/GPUSkinCache.cpp
+++ b/Engine/Source/Runtime/Engine/Private/GPUSkinCache.cpp
@@ -2259,11 +2259,14 @@ FAutoConsoleVariableSink FGPUSkinCache::CVarSink(FConsoleCommandDelegate::Create
 
 void FGPUSkinCache::IncrementDispatchCounter(FRHICommandListImmediate& RHICmdList)
 {
-	DispatchCounter++;
-	if (GSkinCacheMaxDispatchesPerCmdList > 0 && DispatchCounter >= GSkinCacheMaxDispatchesPerCmdList)
+	if (GSkinCacheMaxDispatchesPerCmdList > 0)
 	{
-		//UE_LOG(LogSkinCache, Log, TEXT("SubmitCommandsHint issued after %d dispatches"), DispatchCounter);
-		RHICmdList.SubmitCommandsHint();
-		DispatchCounter = 0;
+		DispatchCounter++;
+		if (DispatchCounter >= GSkinCacheMaxDispatchesPerCmdList)
+		{
+			//UE_LOG(LogSkinCache, Log, TEXT("SubmitCommandsHint issued after %d dispatches"), DispatchCounter);
+			RHICmdList.SubmitCommandsHint();
+			DispatchCounter = 0;
+		}
 	}
-}
\ No newline at end of file
+}
