WorldPartition: Fix false warnings caused by ChildActors

diff --git a/Engine/Source/Runtime/Engine/Private/WorldPartition/WorldPartition.cpp b/Engine/Source/Runtime/Engine/Private/WorldPartition/WorldPartition.cpp
--- a/Engine/Source/Runtime/Engine/Private/WorldPartition/WorldPartition.cpp
+++ b/Engine/Source/Runtime/Engine/Private/WorldPartition/WorldPartition.cpp
@@ -213,7 +213,7 @@ void UWorldPartition::OnGCPostReachabilityAnalysis()
 	{
 		if (AActor* Actor = Cast<AActor>(static_cast<UObject*>(It->Object)))
 		{
-			if (Actor->IsUnreachable() && Actor->IsPackageExternal())
+			if (Actor->IsUnreachable() && Actor->IsMainPackageActor())
 			{
 				ForEachObjectWithPackage(Actor->GetPackage(), [this, Actor](UObject* Object)
 				{
