Tessellation Removal: Removing various leftover Tessellation code.

diff --git a/Engine/Config/Linux/DataDrivenPlatformInfo.ini b/Engine/Config/Linux/DataDrivenPlatformInfo.ini
--- a/Engine/Config/Linux/DataDrivenPlatformInfo.ini
+++ b/Engine/Config/Linux/DataDrivenPlatformInfo.ini
@@ -122,7 +122,6 @@ bRequiresRenderTargetDuringRaster = true
 bRequiresDisableForwardLocalLights = true
 
 bCompileSignalProcessingPipeline = true
-bSupportsTessellation = true
 bSupportsGen5TemporalAA=true
 
 ; //#todo-vulkanddpi
diff --git a/Engine/Config/Windows/DataDrivenPlatformInfo.ini b/Engine/Config/Windows/DataDrivenPlatformInfo.ini
--- a/Engine/Config/Windows/DataDrivenPlatformInfo.ini
+++ b/Engine/Config/Windows/DataDrivenPlatformInfo.ini
@@ -180,7 +180,6 @@ bRequiresRenderTargetDuringRaster = true
 bRequiresDisableForwardLocalLights = true
 
 bCompileSignalProcessingPipeline = true
-bSupportsTessellation = true
 bSupportsGen5TemporalAA=true
 
 ; //#todo-vulkanddpi
diff --git a/Engine/Source/Developer/Android/AndroidTargetPlatform/Private/AndroidTargetPlatform.cpp b/Engine/Source/Developer/Android/AndroidTargetPlatform/Private/AndroidTargetPlatform.cpp
--- a/Engine/Source/Developer/Android/AndroidTargetPlatform/Private/AndroidTargetPlatform.cpp
+++ b/Engine/Source/Developer/Android/AndroidTargetPlatform/Private/AndroidTargetPlatform.cpp
@@ -356,9 +356,6 @@ bool FAndroidTargetPlatform::SupportsFeature( ETargetPlatformFeatures Feature )
 		case ETargetPlatformFeatures::DeferredRendering:
 			return SupportsVulkanSM5();
 
-		case ETargetPlatformFeatures::Tessellation:
-			return false;
-
 		case ETargetPlatformFeatures::SoftwareOcclusion:
 			return SupportsSoftwareOcclusion();
 
diff --git a/Engine/Source/Developer/DeviceManager/Private/Widgets/Details/SDeviceDetails.cpp b/Engine/Source/Developer/DeviceManager/Private/Widgets/Details/SDeviceDetails.cpp
--- a/Engine/Source/Developer/DeviceManager/Private/Widgets/Details/SDeviceDetails.cpp
+++ b/Engine/Source/Developer/DeviceManager/Private/Widgets/Details/SDeviceDetails.cpp
@@ -154,7 +154,6 @@ void SDeviceDetails::Construct(const FArguments& InArgs, const TSharedRef<FDevic
 				FeatureList.Add(MakeShareable(new FDeviceDetailsFeature(TEXT("MultipleGameInstances"), TargetPlatform.SupportsFeature(ETargetPlatformFeatures::MultipleGameInstances))));
 				FeatureList.Add(MakeShareable(new FDeviceDetailsFeature(TEXT("Packaging"), TargetPlatform.SupportsFeature(ETargetPlatformFeatures::Packaging))));
 				FeatureList.Add(MakeShareable(new FDeviceDetailsFeature(TEXT("SdkConnectDisconnect"), TargetPlatform.SupportsFeature(ETargetPlatformFeatures::SdkConnectDisconnect))));
-				FeatureList.Add(MakeShareable(new FDeviceDetailsFeature(TEXT("Tessellation"), TargetPlatform.SupportsFeature(ETargetPlatformFeatures::Tessellation))));
 				FeatureList.Add(MakeShareable(new FDeviceDetailsFeature(TEXT("TextureStreaming"), TargetPlatform.SupportsFeature(ETargetPlatformFeatures::TextureStreaming))));
 				FeatureList.Add(MakeShareable(new FDeviceDetailsFeature(TEXT("UserCredentials"), TargetPlatform.SupportsFeature(ETargetPlatformFeatures::UserCredentials))));
 				FeatureList.Add(MakeShareable(new FDeviceDetailsFeature(TEXT("MobileRendering"), TargetPlatform.SupportsFeature(ETargetPlatformFeatures::MobileRendering))));
diff --git a/Engine/Source/Developer/Lumin/LuminTargetPlatform/Private/LuminTargetPlatform.cpp b/Engine/Source/Developer/Lumin/LuminTargetPlatform/Private/LuminTargetPlatform.cpp
--- a/Engine/Source/Developer/Lumin/LuminTargetPlatform/Private/LuminTargetPlatform.cpp
+++ b/Engine/Source/Developer/Lumin/LuminTargetPlatform/Private/LuminTargetPlatform.cpp
@@ -122,8 +122,6 @@ bool FLuminTargetPlatform::SupportsFeature(ETargetPlatformFeatures Feature) cons
 		return SupportsMobileRendering() || SupportsVulkan();
 
 	case ETargetPlatformFeatures::HighQualityLightmaps:
-	//#todo-rco: Enable when Vulkan supports it
-	//case ETargetPlatformFeatures::Tessellation:
 	case ETargetPlatformFeatures::DeferredRendering:
 		return SupportsDesktopRendering();
 
diff --git a/Engine/Source/Developer/TargetPlatform/Public/Common/TargetPlatformBase.h b/Engine/Source/Developer/TargetPlatform/Public/Common/TargetPlatformBase.h
--- a/Engine/Source/Developer/TargetPlatform/Public/Common/TargetPlatformBase.h
+++ b/Engine/Source/Developer/TargetPlatform/Public/Common/TargetPlatformBase.h
@@ -373,9 +373,6 @@ public:
 		case ETargetPlatformFeatures::Packaging:
 			return false;
 
-		case ETargetPlatformFeatures::Tessellation:
-			return TPlatformProperties::SupportsTessellation();
-
 		case ETargetPlatformFeatures::TextureStreaming:
 			return TPlatformProperties::SupportsTextureStreaming();
 		case ETargetPlatformFeatures::MeshLODStreaming:
diff --git a/Engine/Source/Developer/TargetPlatform/Public/Interfaces/ITargetPlatform.h b/Engine/Source/Developer/TargetPlatform/Public/Interfaces/ITargetPlatform.h
--- a/Engine/Source/Developer/TargetPlatform/Public/Interfaces/ITargetPlatform.h
+++ b/Engine/Source/Developer/TargetPlatform/Public/Interfaces/ITargetPlatform.h
@@ -47,9 +47,6 @@ enum class ETargetPlatformFeatures
 	/** Connect and disconnect devices through the SDK. */
 	SdkConnectDisconnect,
 
-	/** GPU tesselation. */
-	Tessellation,
-
 	/** Texture streaming. */
 	TextureStreaming,
 
diff --git a/Engine/Source/Editor/Persona/Private/SAnimViewportToolBar.cpp b/Engine/Source/Editor/Persona/Private/SAnimViewportToolBar.cpp
--- a/Engine/Source/Editor/Persona/Private/SAnimViewportToolBar.cpp
+++ b/Engine/Source/Editor/Persona/Private/SAnimViewportToolBar.cpp
@@ -849,7 +849,6 @@ TSharedRef<SWidget> SAnimViewportToolBar::GenerateShowMenu() const
 				.IncludeFlag(FEngineShowFlags::SF_Selection)
 				.IncludeFlag(FEngineShowFlags::SF_SeparateTranslucency)
 				.IncludeFlag(FEngineShowFlags::SF_TemporalAA)
-				.IncludeFlag(FEngineShowFlags::SF_Tessellation)
 				.IncludeFlag(FEngineShowFlags::SF_VertexColors)
 				;
 
diff --git a/Engine/Source/Runtime/Core/Public/Android/AndroidPlatform.h b/Engine/Source/Runtime/Core/Public/Android/AndroidPlatform.h
--- a/Engine/Source/Runtime/Core/Public/Android/AndroidPlatform.h
+++ b/Engine/Source/Runtime/Core/Public/Android/AndroidPlatform.h
@@ -59,7 +59,6 @@ typedef FAndroidTypes FPlatformTypes;
 #define PLATFORM_UI_NEEDS_FOCUS_OUTLINES			0
 #define PLATFORM_SUPPORTS_EARLY_MOVIE_PLAYBACK		1 // movies will start before engine is initalized
 #define PLATFORM_SUPPORTS_GEOMETRY_SHADERS			0
-#define PLATFORM_SUPPORTS_TESSELLATION_SHADERS		0
 #define PLATFORM_SUPPORTS_VIRTUAL_TEXTURE_STREAMING	1
 #define PLATFORM_SUPPORTS_LANDSCAPE_VISUAL_MESH_LOD_STREAMING 1
 
diff --git a/Engine/Source/Runtime/Core/Public/GenericPlatform/GenericPlatformProperties.h b/Engine/Source/Runtime/Core/Public/GenericPlatform/GenericPlatformProperties.h
--- a/Engine/Source/Runtime/Core/Public/GenericPlatform/GenericPlatformProperties.h
+++ b/Engine/Source/Runtime/Core/Public/GenericPlatform/GenericPlatformProperties.h
@@ -178,16 +178,6 @@ struct FGenericPlatformProperties
 		return false;
 	}
 
-	/**
-	 * Gets whether this platform supports tessellation.
-	 *
-	 * @return true if tessellation is supported, false otherwise.
-	 */
-	static FORCEINLINE bool SupportsTessellation()
-	{
-		return false;
-	}
-
 	/**
 	 * Gets whether this platform supports windowed mode rendering.
 	 *
diff --git a/Engine/Source/Runtime/Core/Public/HAL/Platform.h b/Engine/Source/Runtime/Core/Public/HAL/Platform.h
--- a/Engine/Source/Runtime/Core/Public/HAL/Platform.h
+++ b/Engine/Source/Runtime/Core/Public/HAL/Platform.h
@@ -351,14 +351,6 @@
 	#define PLATFORM_SUPPORTS_MESH_SHADERS 0
 #endif
 
-#ifndef PLATFORM_SUPPORTS_TESSELLATION_SHADERS
-	#define PLATFORM_SUPPORTS_TESSELLATION_SHADERS	1
-#endif
-
-#if PLATFORM_SUPPORTS_TESSELLATION_SHADERS && !PLATFORM_SUPPORTS_GEOMETRY_SHADERS
-	#error Geometry shader support is required by tessellation
-#endif
-
 #ifndef PLATFORM_BUILTIN_VERTEX_HALF_FLOAT
 	#define PLATFORM_BUILTIN_VERTEX_HALF_FLOAT	1
 #endif
diff --git a/Engine/Source/Runtime/Core/Public/HoloLens/HoloLensPlatformProperties.h b/Engine/Source/Runtime/Core/Public/HoloLens/HoloLensPlatformProperties.h
--- a/Engine/Source/Runtime/Core/Public/HoloLens/HoloLensPlatformProperties.h
+++ b/Engine/Source/Runtime/Core/Public/HoloLens/HoloLensPlatformProperties.h
@@ -29,10 +29,6 @@ struct FHoloLensPlatformProperties : public FGenericPlatformProperties
 	{
 		return false;
 	}              
-	static FORCEINLINE bool SupportsTessellation()
-	{
-		return true;
-	}
 	static FORCEINLINE bool SupportsWindowedMode()
 	{
 		return true;
diff --git a/Engine/Source/Runtime/Core/Public/IOS/IOSPlatform.h b/Engine/Source/Runtime/Core/Public/IOS/IOSPlatform.h
--- a/Engine/Source/Runtime/Core/Public/IOS/IOSPlatform.h
+++ b/Engine/Source/Runtime/Core/Public/IOS/IOSPlatform.h
@@ -75,7 +75,6 @@ typedef FIOSPlatformTypes FPlatformTypes;
 
 #define PLATFORM_NEEDS_RHIRESOURCELIST					0
 #define PLATFORM_SUPPORTS_GEOMETRY_SHADERS				0
-#define PLATFORM_SUPPORTS_TESSELLATION_SHADERS			0
 #define PLATFORM_SUPPORTS_VIRTUAL_TEXTURE_STREAMING		1
 #define PLATFORM_SUPPORTS_LANDSCAPE_VISUAL_MESH_LOD_STREAMING 1
 
diff --git a/Engine/Source/Runtime/Core/Public/Linux/LinuxPlatformProperties.h b/Engine/Source/Runtime/Core/Public/Linux/LinuxPlatformProperties.h
--- a/Engine/Source/Runtime/Core/Public/Linux/LinuxPlatformProperties.h
+++ b/Engine/Source/Runtime/Core/Public/Linux/LinuxPlatformProperties.h
@@ -110,11 +110,6 @@ struct FLinuxPlatformProperties
 		return false;
 	}
 
-	static FORCEINLINE bool SupportsTessellation()
-	{
-		return true;
-	}
-
 	static FORCEINLINE bool SupportsWindowedMode()
 	{
 		return !IS_DEDICATED_SERVER;
diff --git a/Engine/Source/Runtime/Core/Public/Mac/MacPlatform.h b/Engine/Source/Runtime/Core/Public/Mac/MacPlatform.h
--- a/Engine/Source/Runtime/Core/Public/Mac/MacPlatform.h
+++ b/Engine/Source/Runtime/Core/Public/Mac/MacPlatform.h
@@ -82,7 +82,6 @@ typedef FMacPlatformTypes FPlatformTypes;
 #define PLATFORM_SUPPORTS_STACK_SYMBOLS					1
 #define PLATFORM_HAS_BSD_SOCKET_FEATURE_MSG_DONTWAIT	1
 #define PLATFORM_IS_ANSI_MALLOC_THREADSAFE				1
-#define PLATFORM_SUPPORTS_TESSELLATION_SHADERS			0
 #define PLATFORM_SUPPORTS_VIRTUAL_TEXTURE_STREAMING		1
 
 #define PLATFORM_RHITHREAD_DEFAULT_BYPASS				WITH_EDITOR
diff --git a/Engine/Source/Runtime/Core/Public/Windows/WindowsPlatformProperties.h b/Engine/Source/Runtime/Core/Public/Windows/WindowsPlatformProperties.h
--- a/Engine/Source/Runtime/Core/Public/Windows/WindowsPlatformProperties.h
+++ b/Engine/Source/Runtime/Core/Public/Windows/WindowsPlatformProperties.h
@@ -98,11 +98,6 @@ struct FWindowsPlatformProperties
 		return true;
 	}
 
-	static FORCEINLINE bool SupportsTessellation()
-	{
-		return true; // DX11 compatible
-	}
-
 	static FORCEINLINE bool SupportsWindowedMode()
 	{
 		return true;
diff --git a/Engine/Source/Runtime/Engine/Classes/Engine/RendererSettings.h b/Engine/Source/Runtime/Engine/Classes/Engine/RendererSettings.h
--- a/Engine/Source/Runtime/Engine/Classes/Engine/RendererSettings.h
+++ b/Engine/Source/Runtime/Engine/Classes/Engine/RendererSettings.h
@@ -425,11 +425,6 @@ class ENGINE_API URendererSettings : public UDeveloperSettings
 		ConfigRestartRequired=true))
 	uint32 bVertexFoggingForOpaque:1;
 
-	UPROPERTY(config, EditAnywhere, Category=Materials, meta=(
-		ConsoleVariable="r.TessellationAdaptivePixelsPerTriangle",DisplayName="Adaptive pixels per triangle",
-		ToolTip="When adaptive tessellation is enabled it will try to tessellate a mesh so that each triangle contains the specified number of pixels. The tessellation multiplier specified in the material can increase or decrease the amount of tessellation."))
-	float TessellationAdaptivePixelsPerTriangle;
-
 	UPROPERTY(config, EditAnywhere, Category=Translucency, meta=(
 		ConsoleVariable="r.SeparateTranslucency",
 		ToolTip="Allow translucency to be rendered to a separate render targeted and composited after depth of field. Prevents translucency from appearing out of focus."))
diff --git a/Engine/Source/Runtime/Engine/Private/Materials/MaterialExpressions.cpp b/Engine/Source/Runtime/Engine/Private/Materials/MaterialExpressions.cpp
--- a/Engine/Source/Runtime/Engine/Private/Materials/MaterialExpressions.cpp
+++ b/Engine/Source/Runtime/Engine/Private/Materials/MaterialExpressions.cpp
@@ -11990,7 +11990,6 @@ void UMaterialExpressionCustom::Serialize(FStructuredArchive::FRecord Record)
 			FString(TEXT("WorldCameraMovementSinceLastFrame")),
 			FString(TEXT("CullingSign")),
 			FString(TEXT("NearPlane")),
-			FString(TEXT("AdaptiveTessellationFactor")),
 			FString(TEXT("GameTime")),
 			FString(TEXT("RealTime")),
 			FString(TEXT("Random")),
diff --git a/Engine/Source/Runtime/Engine/Private/SceneView.cpp b/Engine/Source/Runtime/Engine/Private/SceneView.cpp
--- a/Engine/Source/Runtime/Engine/Private/SceneView.cpp
+++ b/Engine/Source/Runtime/Engine/Private/SceneView.cpp
@@ -275,12 +275,6 @@ static TAutoConsoleVariable<int32> CVarTonemapperQuality(
 	TEXT(" 5: + GrainJitter = full quality (default)"),
 	ECVF_Scalability | ECVF_RenderThreadSafe);
 
-static TAutoConsoleVariable<float> CVarTessellationAdaptivePixelsPerTriangle(
-	TEXT("r.TessellationAdaptivePixelsPerTriangle"),
-	48.0f,
-	TEXT("Global tessellation factor multiplier"),
-	ECVF_RenderThreadSafe);
-
 // should be changed to BaseColor and Metallic, since some time now UE is not using DiffuseColor and SpecularColor any more
 static TAutoConsoleVariable<float> CVarDiffuseColorMin(
 	TEXT("r.DiffuseColor.Min"),
@@ -2289,17 +2283,6 @@ void FSceneView::SetupViewRectUniformBufferParameters(FViewUniformShaderParamete
 				FPlane(Ax, Ay, 0, 1)) * InViewMatrices.GetInvTranslatedViewProjectionMatrix();
 	}
 
-	// is getting clamped in the shader to a value larger than 0 (we don't want the triangles to disappear)
-	ViewUniformShaderParameters.AdaptiveTessellationFactor = 0.0f;
-
-	if(Family->EngineShowFlags.Tessellation)
-	{
-		// CVar setting is pixels/tri which is nice and intuitive.  But we want pixels/tessellated edge.  So use a heuristic.
-		float TessellationAdaptivePixelsPerEdge = FMath::Sqrt(2.f * CVarTessellationAdaptivePixelsPerTriangle.GetValueOnRenderThread());
-
-		ViewUniformShaderParameters.AdaptiveTessellationFactor = 0.5f * InViewMatrices.GetProjectionMatrix().M[1][1] * float(EffectiveViewRect.Height()) / TessellationAdaptivePixelsPerEdge;
-	}
-
 	// Compute coefficients which takes a screen UV and converts to Viewspace.xy / ViewZ
 	float InvTanHalfFov = InViewMatrices.GetProjectionMatrix().M[0][0];
 	float Ratio = UnscaledViewRect.Width() / (float)UnscaledViewRect.Height();
diff --git a/Engine/Source/Runtime/Engine/Private/SkeletalMeshLODModel.cpp b/Engine/Source/Runtime/Engine/Private/SkeletalMeshLODModel.cpp
--- a/Engine/Source/Runtime/Engine/Private/SkeletalMeshLODModel.cpp
+++ b/Engine/Source/Runtime/Engine/Private/SkeletalMeshLODModel.cpp
@@ -664,7 +664,7 @@ void FSkeletalMeshLODModel::Serialize(FArchive& Ar, UObject* Owner, int32 Idx)
 	DECLARE_SCOPE_CYCLE_COUNTER(TEXT("FSkeletalMeshLODModel::Serialize"), STAT_SkeletalMeshLODModel_Serialize, STATGROUP_LoadTime);
 
 	const uint8 LodAdjacencyStripFlag = 1;
-	FStripDataFlags StripFlags(Ar, Ar.IsCooking() && !Ar.CookingTarget()->SupportsFeature(ETargetPlatformFeatures::Tessellation) ? LodAdjacencyStripFlag : 0);
+	FStripDataFlags StripFlags(Ar, Ar.IsCooking() ? LodAdjacencyStripFlag : 0);
 
 	Ar.UsingCustomVersion(FSkeletalMeshCustomVersion::GUID);
 	Ar.UsingCustomVersion(FFortniteMainBranchObjectVersion::GUID);
diff --git a/Engine/Source/Runtime/Engine/Public/GeometryCollectionRendering.h b/Engine/Source/Runtime/Engine/Public/GeometryCollectionRendering.h
--- a/Engine/Source/Runtime/Engine/Public/GeometryCollectionRendering.h
+++ b/Engine/Source/Runtime/Engine/Public/GeometryCollectionRendering.h
@@ -112,8 +112,6 @@ public:
 	virtual void InitRHI() override;
 	virtual void ReleaseRHI() override;
 
-	static bool SupportsTessellationShaders() { return true; }
-
 	inline FRHIShaderResourceView* GetPositionsSRV() const
 	{
 		return Data.PositionComponentSRV;
diff --git a/Engine/Source/Runtime/Engine/Public/SceneView.h b/Engine/Source/Runtime/Engine/Public/SceneView.h
--- a/Engine/Source/Runtime/Engine/Public/SceneView.h
+++ b/Engine/Source/Runtime/Engine/Public/SceneView.h
@@ -648,7 +648,6 @@ enum ETranslucencyVolumeCascade
 	VIEW_UNIFORM_BUFFER_MEMBER(FVector, WorldCameraMovementSinceLastFrame) \
 	VIEW_UNIFORM_BUFFER_MEMBER(float, CullingSign) \
 	VIEW_UNIFORM_BUFFER_MEMBER_EX(float, NearPlane, EShaderPrecisionModifier::Half) \
-	VIEW_UNIFORM_BUFFER_MEMBER(float, AdaptiveTessellationFactor) \
 	VIEW_UNIFORM_BUFFER_MEMBER(float, GameTime) \
 	VIEW_UNIFORM_BUFFER_MEMBER(float, RealTime) \
 	VIEW_UNIFORM_BUFFER_MEMBER(float, DeltaTime) \
diff --git a/Engine/Source/Runtime/Engine/Public/ShowFlagsValues.inl b/Engine/Source/Runtime/Engine/Public/ShowFlagsValues.inl
--- a/Engine/Source/Runtime/Engine/Public/ShowFlagsValues.inl
+++ b/Engine/Source/Runtime/Engine/Public/ShowFlagsValues.inl
@@ -143,8 +143,6 @@ SHOWFLAG_FIXED_IN_SHIPPING(0, GameplayDebug, SFG_Developer, NSLOCTEXT("UnrealEd"
 SHOWFLAG_ALWAYS_ACCESSIBLE(TexturedLightProfiles,  SFG_LightingFeatures, NSLOCTEXT("UnrealEd", "TexturedLightProfilesSF", "Textured Light Profiles (IES Texture)"))
 /** LightFunctions (masking light sources with a material), for now SHOWFLAG_ALWAYS_ACCESSIBLE because it's exposed in SceneCapture */
 SHOWFLAG_ALWAYS_ACCESSIBLE(LightFunctions, SFG_LightingFeatures, NSLOCTEXT("UnrealEd", "LightFunctionsSF", "Light Functions"))
-/** Hardware Tessellation (DX11 feature) */
-SHOWFLAG_FIXED_IN_SHIPPING(1, Tessellation, SFG_Advanced, NSLOCTEXT("UnrealEd", "TessellationSF", "Tessellation"))
 /** Draws Nanite meshes */
 SHOWFLAG_ALWAYS_ACCESSIBLE(NaniteMeshes, SFG_Advanced, NSLOCTEXT("UnrealEd", "NaniteMeshesSF", "Nanite Meshes"))
 /** Draws instanced static meshes that are not foliage or grass, for now SHOWFLAG_ALWAYS_ACCESSIBLE because it's exposed in SceneCapture */
diff --git a/Engine/Source/Runtime/OpenGLDrv/Public/OpenGL.h b/Engine/Source/Runtime/OpenGLDrv/Public/OpenGL.h
--- a/Engine/Source/Runtime/OpenGLDrv/Public/OpenGL.h
+++ b/Engine/Source/Runtime/OpenGLDrv/Public/OpenGL.h
@@ -122,7 +122,6 @@ public:
 	static FORCEINLINE bool SupportsCopyImage()							{ return bSupportsCopyImage; }
 	static FORCEINLINE bool SupportsTextureFilterAnisotropic()			{ return bSupportsTextureFilterAnisotropic; }
 	static FORCEINLINE bool SupportsSeparateAlphaBlend()				{ return bSupportsDrawBuffersBlend; }
-	static FORCEINLINE bool SupportsTessellation()						{ return false; }
 	static FORCEINLINE void EnableSupportsClipControl()					{ bSupportsClipControl = true; }
 	static FORCEINLINE bool SupportsClipControl()						{ return bSupportsClipControl; }
 	static FORCEINLINE bool SupportsTextureView()						{ return false; }
diff --git a/Engine/Source/Runtime/RHI/Private/RHI.cpp b/Engine/Source/Runtime/RHI/Private/RHI.cpp
--- a/Engine/Source/Runtime/RHI/Private/RHI.cpp
+++ b/Engine/Source/Runtime/RHI/Private/RHI.cpp
@@ -1342,19 +1342,6 @@ RHI_API uint32 RHIGetShaderLanguageVersion(const FStaticShaderPlatform Platform)
 	return Version;
 }
 
-RHI_API bool RHISupportsTessellation(const FStaticShaderPlatform Platform)
-{
-	if (FDataDrivenShaderPlatformInfo::GetSupportsTessellation(Platform))
-	{
-		return true;
-	}
-	if (IsFeatureLevelSupported(Platform, ERHIFeatureLevel::SM5))
-	{
-		return (Platform == SP_PCD3D_SM5) || (IsVulkanSM5Platform(Platform));
-	}
-	return false;
-}
-
 RHI_API bool RHISupportsIndexBufferUAVs(const FStaticShaderPlatform Platform)
 {
 	return Platform == SP_PCD3D_SM5 || IsVulkanPlatform(Platform) || IsMetalSM5Platform(Platform)
@@ -1790,7 +1777,6 @@ void FGenericDataDrivenShaderPlatformInfo::ParseDataDrivenShaderInfo(const FConf
 	GET_SECTION_BOOL_HELPER(bRequiresRenderTargetDuringRaster);
 	GET_SECTION_BOOL_HELPER(bRequiresDisableForwardLocalLights);
 	GET_SECTION_BOOL_HELPER(bCompileSignalProcessingPipeline);
-	GET_SECTION_BOOL_HELPER(bSupportsTessellation);
 	GET_SECTION_BOOL_HELPER(bSupportsMeshShaders);
 	GET_SECTION_BOOL_HELPER(bSupportsPerPixelDBufferMask);
 	GET_SECTION_BOOL_HELPER(bIsHlslcc);
diff --git a/Engine/Source/Runtime/RHI/Public/RHI.h b/Engine/Source/Runtime/RHI/Public/RHI.h
--- a/Engine/Source/Runtime/RHI/Public/RHI.h
+++ b/Engine/Source/Runtime/RHI/Public/RHI.h
@@ -124,9 +124,6 @@ RHI_API const TCHAR* RHIVendorIdToString(EGpuVendorId VendorId);
 // helper to return the shader language version for the given shader platform.
 RHI_API uint32 RHIGetShaderLanguageVersion(const FStaticShaderPlatform Platform);
 
-// helper to check that the shader platform supports tessellation.
-RHI_API bool RHISupportsTessellation(const FStaticShaderPlatform Platform);
-
 // helper to check that the shader platform supports writing to UAVs from pixel shaders.
 RHI_API bool RHISupportsPixelShaderUAVs(const FStaticShaderPlatform Platform);
 
diff --git a/Engine/Source/Runtime/RHI/Public/RHIDefinitions.h b/Engine/Source/Runtime/RHI/Public/RHIDefinitions.h
--- a/Engine/Source/Runtime/RHI/Public/RHIDefinitions.h
+++ b/Engine/Source/Runtime/RHI/Public/RHIDefinitions.h
@@ -364,7 +364,6 @@ class RHI_API FGenericDataDrivenShaderPlatformInfo
 	uint32 bRequiresRenderTargetDuringRaster : 1;
 	uint32 bRequiresDisableForwardLocalLights : 1;
 	uint32 bCompileSignalProcessingPipeline : 1;
-	uint32 bSupportsTessellation : 1;
 	uint32 bSupportsMeshShaders : 1;
 	uint32 bSupportsPerPixelDBufferMask : 1;
 	uint32 bIsHlslcc : 1;
@@ -654,11 +653,6 @@ public:
 		return Infos[Platform].bCompileSignalProcessingPipeline;
 	}
 
-	static FORCEINLINE_DEBUGGABLE const bool GetSupportsTessellation(const FStaticShaderPlatform Platform)
-	{
-		return Infos[Platform].bSupportsTessellation;
-	}
-
 	static FORCEINLINE_DEBUGGABLE const bool GetSupportsMeshShaders(const FStaticShaderPlatform Platform)
 	{
 		return Infos[Platform].bSupportsMeshShaders;
