Initial work of setting up reflection/normal vector support in new HLSL translator
Mostly just cleanup

diff --git a/Engine/Source/Runtime/Engine/Classes/Materials/MaterialExpressionReflectionVectorWS.h b/Engine/Source/Runtime/Engine/Classes/Materials/MaterialExpressionReflectionVectorWS.h
--- a/Engine/Source/Runtime/Engine/Classes/Materials/MaterialExpressionReflectionVectorWS.h
+++ b/Engine/Source/Runtime/Engine/Classes/Materials/MaterialExpressionReflectionVectorWS.h
@@ -26,6 +26,7 @@ class UMaterialExpressionReflectionVectorWS : public UMaterialExpression
 #if WITH_EDITOR
 	virtual int32 Compile(class FMaterialCompiler* Compiler, int32 OutputIndex) override;
 	virtual void GetCaption(TArray<FString>& OutCaptions) const override;
+	virtual EMaterialGenerateHLSLStatus GenerateHLSLExpression(FMaterialHLSLGenerator& Generator, UE::HLSLTree::FScope& Scope, int32 OutputIndex, UE::HLSLTree::FExpression*& OutExpression) override;
 #endif
 	//~ End UMaterialExpression Interface
 };
diff --git a/Engine/Source/Runtime/Engine/Private/HLSLTree/HLSLTree.cpp b/Engine/Source/Runtime/Engine/Private/HLSLTree/HLSLTree.cpp
--- a/Engine/Source/Runtime/Engine/Private/HLSLTree/HLSLTree.cpp
+++ b/Engine/Source/Runtime/Engine/Private/HLSLTree/HLSLTree.cpp
@@ -108,7 +108,7 @@ void UE::HLSLTree::FCodeWriter::WriteConstant(const Shader::FValue& Value)
 	switch (Value.GetType())
 	{
 	case Shader::EValueType::Float1:
-		Writef(TEXT("float(%0.8f)"), Value.Component[0].Float);
+		Writef(TEXT("%0.8f"), Value.Component[0].Float);
 		break;
 	case Shader::EValueType::Float2:
 		Writef(TEXT("float2(%0.8f, %0.8f)"), Value.Component[0].Float, Value.Component[1].Float);
@@ -119,6 +119,18 @@ void UE::HLSLTree::FCodeWriter::WriteConstant(const Shader::FValue& Value)
 	case Shader::EValueType::Float4:
 		Writef(TEXT("float4(%0.8f, %0.8f, %0.8f, %0.8f)"), Value.Component[0].Float, Value.Component[1].Float, Value.Component[2].Float, Value.Component[3].Float);
 		break;
+	case Shader::EValueType::Int1:
+		Writef(TEXT("%d"), Value.Component[0].Int);
+		break;
+	case Shader::EValueType::Int2:
+		Writef(TEXT("int2(%d, %d)"), Value.Component[0].Int, Value.Component[1].Int);
+		break;
+	case Shader::EValueType::Int3:
+		Writef(TEXT("int3(%d, %d, %d)"), Value.Component[0].Int, Value.Component[1].Int, Value.Component[2].Int);
+		break;
+	case Shader::EValueType::Int4:
+		Writef(TEXT("int4(%d, %d, %d, %d)"), Value.Component[0].Int, Value.Component[1].Int, Value.Component[2].Int, Value.Component[3].Int);
+		break;
 	case Shader::EValueType::Bool1:
 		Writef(TEXT("%s"), ToString(Value.Component[0].Bool));
 		break;
@@ -179,10 +191,11 @@ int32 UE::HLSLTree::FEmitContext::FindScopeIndex(FScope* Scope)
 	return INDEX_NONE;
 }
 
-void UE::HLSLTree::FExpressionEmitResult::ForwardValue(FEmitContext& Context, const FEmitValue& InValue)
+void UE::HLSLTree::FExpressionEmitResult::ForwardValue(FEmitContext& Context, const FEmitValue* InValue)
 {
-	EvaluationType = InValue.GetEvaluationType();
-	Type = InValue.GetExpressionType();
+	check(InValue);
+	EvaluationType = InValue->GetEvaluationType();
+	Type = InValue->GetExpressionType();
 	if (EvaluationType == EExpressionEvaluationType::Shader)
 	{
 		bInline = true;
@@ -194,7 +207,7 @@ void UE::HLSLTree::FExpressionEmitResult::ForwardValue(FEmitContext& Context, co
 	}
 }
 
-const UE::HLSLTree::FEmitValue& UE::HLSLTree::FEmitContext::AcquireValue(UE::HLSLTree::FLocalDeclaration* Declaration)
+const UE::HLSLTree::FEmitValue* UE::HLSLTree::FEmitContext::AcquireValue(UE::HLSLTree::FLocalDeclaration* Declaration)
 {
 	check(Declaration);
 	check(Declaration->Type != Shader::EValueType::Void);
@@ -208,18 +221,21 @@ const UE::HLSLTree::FEmitValue& UE::HLSLTree::FEmitContext::AcquireValue(UE::HLS
 		Entry->Value.EvaluationType = EExpressionEvaluationType::Shader;
 		Entry->Value.ExpressionType = Declaration->Type;
 
+		FLocalHLSLCodeWriter LocalWriter;
+		LocalWriter.WriteConstant(Shader::FValue(Declaration->Type));
+
 		Entry->Value.Code = AllocateStringf(*Allocator, TEXT("Local%d"), NumExpressionLocals++);
 		const Shader::FValueTypeDescription& TypeDesc = Shader::GetValueTypeDescription(Declaration->Type);
-		ScopeEntry->ExpressionCodeWriter->WriteLinef(TEXT("%s %s = (%s)0.0f;"),
+		ScopeEntry->ExpressionCodeWriter->WriteLinef(TEXT("%s %s = %s;"),
 			TypeDesc.Name,
 			Entry->Value.Code,
-			TypeDesc.Name);
+			LocalWriter.StringBuilder->ToString());
 	}
 
-	return Entry->Value;
+	return &Entry->Value;
 }
 
-const UE::HLSLTree::FEmitValue& UE::HLSLTree::FEmitContext::AcquireValue(UE::HLSLTree::FExpression* Expression)
+const UE::HLSLTree::FEmitValue* UE::HLSLTree::FEmitContext::AcquireValue(UE::HLSLTree::FExpression* Expression)
 {
 	check(Expression);
 
@@ -292,10 +308,10 @@ const UE::HLSLTree::FEmitValue& UE::HLSLTree::FEmitContext::AcquireValue(UE::HLS
 		}
 	}
 
-	return Entry->Value;
+	return &Entry->Value;
 }
 
-const UE::HLSLTree::FEmitValue& UE::HLSLTree::FEmitContext::AcquireValue(FFunctionCall* FunctionCall, int32 OutputIndex)
+const UE::HLSLTree::FEmitValue* UE::HLSLTree::FEmitContext::AcquireValue(FFunctionCall* FunctionCall, int32 OutputIndex)
 {
 	check(FunctionCall);
 
@@ -323,14 +339,16 @@ const UE::HLSLTree::FEmitValue& UE::HLSLTree::FEmitContext::AcquireValue(FFuncti
 		ParentScopeEntry->ExpressionCodeWriter->Append(LocalWriter);
 
 		// Assign function outputs
-		FEmitValue* OutputValue = new(*Allocator) FEmitValue[FunctionCall->NumOutputs];
+		FEmitValue* OutputValues = new(*Allocator) FEmitValue[FunctionCall->NumOutputs];
 		for (int32 i = 0; i < FunctionCall->NumOutputs; ++i)
 		{
 			FExpression* OutputExpression = FunctionCall->Outputs[i];
-			FEmitValue& Output = OutputValue[i];
-			if (OutputExpression)
+
+			const FEmitValue* OutputValue = AcquireValue(OutputExpression);
+			FEmitValue& Output = OutputValues[i];
+			if (OutputValue)
 			{
-				Output = AcquireValue(OutputExpression);
+				Output = *OutputValue;
 			}
 			else
 			{
@@ -351,32 +369,33 @@ const UE::HLSLTree::FEmitValue& UE::HLSLTree::FEmitContext::AcquireValue(FFuncti
 		}
 
 		Entry = new(*Allocator) FFunctionCallEntry();
-		Entry->OutputRef = OutputValue;
+		Entry->OutputValues = OutputValues;
 		Entry->NumOutputs = FunctionCall->NumOutputs;
 	}
 
 	check(Entry->NumOutputs == FunctionCall->NumOutputs);
 	check(OutputIndex >= 0 && OutputIndex < Entry->NumOutputs);
-	return Entry->OutputRef[OutputIndex];
+	return &Entry->OutputValues[OutputIndex];
 }
 
-const TCHAR* UE::HLSLTree::FEmitContext::GetCode(const FEmitValue& Value) const
+const TCHAR* UE::HLSLTree::FEmitContext::GetCode(const FEmitValue* Value) const
 {
-	if (!Value.Code)
+	check(Value);
+	if (!Value->Code)
 	{
 		FLocalHLSLCodeWriter LocalWriter;
-		if (Value.EvaluationType == EExpressionEvaluationType::Constant)
+		if (Value->EvaluationType == EExpressionEvaluationType::Constant)
 		{
-			LocalWriter.WriteConstant(Value.ConstantValue);
+			LocalWriter.WriteConstant(Value->ConstantValue);
 		}
 		else
 		{
-			check(Value.EvaluationType == EExpressionEvaluationType::Preshader);
-			check(Value.Preshader);
+			check(Value->EvaluationType == EExpressionEvaluationType::Preshader);
+			check(Value->Preshader);
 
 			FUniformExpressionSet& UniformExpressionSet = MaterialCompilationOutput->UniformExpressionSet;
 			FMaterialUniformPreshaderHeader* PreshaderHeader = nullptr;
-			if (Value.ExpressionType == Shader::EValueType::Float1)
+			if (Value->ExpressionType == Shader::EValueType::Float1)
 			{
 				const static TCHAR IndexToMask[] = { 'x', 'y', 'z', 'w' };
 				const int32 ScalarInputIndex = UniformExpressionSet.UniformScalarPreshaders.Num();
@@ -387,7 +406,7 @@ const TCHAR* UE::HLSLTree::FEmitContext::GetCode(const FEmitValue& Value) const
 			else
 			{
 				const TCHAR* Mask = TEXT("");
-				switch (Value.ExpressionType)
+				switch (Value->ExpressionType)
 				{
 				case Shader::EValueType::Float1: Mask = TEXT(".r"); break;
 				case Shader::EValueType::Float2: Mask = TEXT(".rg"); break;
@@ -401,28 +420,29 @@ const TCHAR* UE::HLSLTree::FEmitContext::GetCode(const FEmitValue& Value) const
 				PreshaderHeader = &UniformExpressionSet.UniformVectorPreshaders.AddDefaulted_GetRef();
 			}
 			PreshaderHeader->OpcodeOffset = UniformExpressionSet.UniformPreshaderData.Num();
-			UniformExpressionSet.UniformPreshaderData.Append(*Value.Preshader);
-			PreshaderHeader->OpcodeSize = Value.Preshader->Num();
+			UniformExpressionSet.UniformPreshaderData.Append(*Value->Preshader);
+			PreshaderHeader->OpcodeSize = Value->Preshader->Num();
 		}
-		Value.Code = AllocateString(*Allocator, LocalWriter.GetStringBuilder());
+		Value->Code = AllocateString(*Allocator, LocalWriter.GetStringBuilder());
 	}
 
-	return Value.Code;
+	return Value->Code;
 }
 
-void UE::HLSLTree::FEmitContext::AppendPreshader(const FEmitValue& Value, Shader::FPreshaderData& InOutPreshader) const
+void UE::HLSLTree::FEmitContext::AppendPreshader(const FEmitValue* Value, Shader::FPreshaderData& InOutPreshader) const
 {
-	if (Value.EvaluationType == EExpressionEvaluationType::Constant)
+	check(Value);
+	if (Value->EvaluationType == EExpressionEvaluationType::Constant)
 	{
 		// Push the constant value
 		InOutPreshader.WriteOpcode(Shader::EPreshaderOpcode::Constant);
-		InOutPreshader.Write(Value.ConstantValue);
+		InOutPreshader.Write(Value->ConstantValue);
 	}
 	else
 	{
-		check(Value.EvaluationType == EExpressionEvaluationType::Preshader);
-		check(Value.Preshader);
-		InOutPreshader.Append(*Value.Preshader);
+		check(Value->EvaluationType == EExpressionEvaluationType::Preshader);
+		check(Value->Preshader);
+		InOutPreshader.Append(*Value->Preshader);
 	}
 }
 
diff --git a/Engine/Source/Runtime/Engine/Private/HLSLTree/HLSLTreeCommon.cpp b/Engine/Source/Runtime/Engine/Private/HLSLTree/HLSLTreeCommon.cpp
--- a/Engine/Source/Runtime/Engine/Private/HLSLTree/HLSLTreeCommon.cpp
+++ b/Engine/Source/Runtime/Engine/Private/HLSLTree/HLSLTreeCommon.cpp
@@ -69,12 +69,8 @@ void UE::HLSLTree::FExpressionConstant::EmitCode(FEmitContext& Context, FExpress
 
 void UE::HLSLTree::FExpressionLocalVariable::EmitCode(FEmitContext& Context, FExpressionEmitResult& OutResult) const
 {
-	const FEmitValue& DeclarationRef = Context.AcquireValue(Declaration);
-	
-	OutResult.EvaluationType = EExpressionEvaluationType::Shader;
-	OutResult.Type = DeclarationRef.GetExpressionType();
-	OutResult.bInline = true;
-	OutResult.Writer.Writef(TEXT("%s"), Context.GetCode(DeclarationRef));
+	const FEmitValue* DeclarationValue = Context.AcquireValue(Declaration);
+	OutResult.ForwardValue(Context, DeclarationValue);
 }
 
 void UE::HLSLTree::FExpressionParameter::EmitCode(FEmitContext& Context, FExpressionEmitResult& OutResult) const
@@ -257,10 +253,10 @@ void UE::HLSLTree::FExpressionTextureSample::EmitCode(FEmitContext& Context, FEx
 		}
 	}
 
-	const FEmitValue& TexCoordRef = Context.AcquireValue(TexCoordExpression);
+	const FEmitValue* TexCoordValue = Context.AcquireValue(TexCoordExpression);
 	OutResult.EvaluationType = EExpressionEvaluationType::Shader;
 	OutResult.Type = Shader::EValueType::Float4;
-	OutResult.Writer.Writef(TEXT("%s(%s(%s, %s, %s))"), SamplerTypeFunction, SampleFunctionName, *TextureName, *SamplerStateCode, Context.GetCode(TexCoordRef));
+	OutResult.Writer.Writef(TEXT("%s(%s(%s, %s, %s))"), SamplerTypeFunction, SampleFunctionName, *TextureName, *SamplerStateCode, Context.GetCode(TexCoordValue));
 }
 
 void UE::HLSLTree::FExpressionDefaultMaterialAttributes::EmitCode(FEmitContext& Context, FExpressionEmitResult& OutResult) const
@@ -275,37 +271,37 @@ void UE::HLSLTree::FExpressionSetMaterialAttribute::EmitCode(FEmitContext& Conte
 {
 	const FString PropertyName = FMaterialAttributeDefinitionMap::GetAttributeName(AttributeID);
 
-	const FEmitValue& AttributresRef = Context.AcquireValue(AttributesExpression);
-	const FEmitValue&ValueRef = Context.AcquireValue(ValueExpression);
+	const FEmitValue* AttributresValue = Context.AcquireValue(AttributesExpression);
+	const FEmitValue* Value = Context.AcquireValue(ValueExpression);
 
 	OutResult.EvaluationType = EExpressionEvaluationType::Shader;
 	OutResult.Type = Shader::EValueType::MaterialAttributes;
-	OutResult.Writer.Writef(TEXT("FMaterialAttributes_Set%s(%s, %s)"), *PropertyName, Context.GetCode(AttributresRef), Context.GetCode(ValueRef));
+	OutResult.Writer.Writef(TEXT("FMaterialAttributes_Set%s(%s, %s)"), *PropertyName, Context.GetCode(AttributresValue), Context.GetCode(Value));
 }
 
 void UE::HLSLTree::FExpressionSelect::EmitCode(FEmitContext& Context, FExpressionEmitResult& OutResult) const
 {
-	const FEmitValue& ConditionValue = Context.AcquireValue(ConditionExpression);
-	if (ConditionValue.GetEvaluationType() == EExpressionEvaluationType::Constant)
+	const FEmitValue* ConditionValue = Context.AcquireValue(ConditionExpression);
+	if (ConditionValue->GetEvaluationType() == EExpressionEvaluationType::Constant)
 	{
-		const bool bCondition = ConditionValue.GetConstantValue().AsBoolScalar();
+		const bool bCondition = ConditionValue->GetConstantValue().AsBoolScalar();
 		if (!bCondition)
 		{
-			const FEmitValue& Value = Context.AcquireValue(FalseExpression);
+			const FEmitValue* Value = Context.AcquireValue(FalseExpression);
 			OutResult.ForwardValue(Context, Value);
 		}
 		else
 		{
-			const FEmitValue& Value = Context.AcquireValue(TrueExpression);
+			const FEmitValue* Value = Context.AcquireValue(TrueExpression);
 			OutResult.ForwardValue(Context, Value);
 		}
 	}
 	else
 	{
-		const FEmitValue& TrueValue = Context.AcquireValue(TrueExpression);
-		const FEmitValue& FalseValue = Context.AcquireValue(FalseExpression);
+		const FEmitValue* TrueValue = Context.AcquireValue(TrueExpression);
+		const FEmitValue* FalseValue = Context.AcquireValue(FalseExpression);
 		OutResult.EvaluationType = EExpressionEvaluationType::Shader; // TODO - preshader
-		OutResult.Type = TrueValue.GetExpressionType(); // TODO - FalseValue?
+		OutResult.Type = TrueValue->GetExpressionType(); // TODO - FalseValue?
 		OutResult.Writer.Writef(TEXT("(%s ? %s : %s)"),
 			Context.GetCode(ConditionValue),
 			Context.GetCode(TrueValue),
@@ -315,16 +311,16 @@ void UE::HLSLTree::FExpressionSelect::EmitCode(FEmitContext& Context, FExpressio
 
 void UE::HLSLTree::FExpressionBinaryOp::EmitCode(FEmitContext& Context, FExpressionEmitResult& OutResult) const
 {
-	const FEmitValue& LhsRef = Context.AcquireValue(Lhs);
-	const FEmitValue& RhsRef = Context.AcquireValue(Rhs);
+	const FEmitValue* LhsValue = Context.AcquireValue(Lhs);
+	const FEmitValue* RhsValue = Context.AcquireValue(Rhs);
 
 	FString ErrorMessage;
 
-	OutResult.EvaluationType = CombineEvaluationTypes(LhsRef.GetEvaluationType(), RhsRef.GetEvaluationType());
-	OutResult.Type = MakeArithmeticResultType(LhsRef.GetExpressionType(), RhsRef.GetExpressionType(), ErrorMessage);
+	OutResult.EvaluationType = CombineEvaluationTypes(LhsValue->GetEvaluationType(), RhsValue->GetEvaluationType());
+	OutResult.Type = MakeArithmeticResultType(LhsValue->GetExpressionType(), RhsValue->GetExpressionType(), ErrorMessage);
 	if (OutResult.EvaluationType == EExpressionEvaluationType::Shader)
 	{
-		OutResult.Writer.Writef(TEXT("(%s + %s)"), Context.GetCode(LhsRef), Context.GetCode(RhsRef));
+		OutResult.Writer.Writef(TEXT("(%s + %s)"), Context.GetCode(LhsValue), Context.GetCode(RhsValue));
 	}
 	else
 	{
@@ -335,8 +331,8 @@ void UE::HLSLTree::FExpressionBinaryOp::EmitCode(FEmitContext& Context, FExpress
 		default: checkNoEntry(); break;
 		}
 
-		Context.AppendPreshader(LhsRef, OutResult.Preshader);
-		Context.AppendPreshader(RhsRef, OutResult.Preshader);
+		Context.AppendPreshader(LhsValue, OutResult.Preshader);
+		Context.AppendPreshader(RhsValue, OutResult.Preshader);
 		OutResult.Preshader.WriteOpcode(Opcode);
 	}
 }
@@ -352,18 +348,18 @@ void UE::HLSLTree::FExpressionSwizzle::EmitCode(FEmitContext& Context, FExpressi
 		Swizzle[i] = ComponentName[ComponentIndex];
 	}
 
-	const FEmitValue& InputRef = Context.AcquireValue(Input);
+	const FEmitValue* InputValue = Context.AcquireValue(Input);
 
 	OutResult.bInline = true;
-	OutResult.EvaluationType = InputRef.GetEvaluationType();
-	OutResult.Type = Shader::MakeValueType(InputRef.GetExpressionType(), Parameters.NumComponents);
+	OutResult.EvaluationType = InputValue->GetEvaluationType();
+	OutResult.Type = Shader::MakeValueType(InputValue->GetExpressionType(), Parameters.NumComponents);
 	if (OutResult.EvaluationType == EExpressionEvaluationType::Shader)
 	{
-		OutResult.Writer.Writef(TEXT("%s.%s"), Context.GetCode(InputRef), Swizzle);
+		OutResult.Writer.Writef(TEXT("%s.%s"), Context.GetCode(InputValue), Swizzle);
 	}
 	else
 	{
-		Context.AppendPreshader(InputRef, OutResult.Preshader);
+		Context.AppendPreshader(InputValue, OutResult.Preshader);
 		OutResult.Preshader.WriteOpcode(Shader::EPreshaderOpcode::ComponentSwizzle)
 			.Write((uint8)Parameters.NumComponents)
 			.Write((uint8)Parameters.ComponentIndex[0])
@@ -375,9 +371,9 @@ void UE::HLSLTree::FExpressionSwizzle::EmitCode(FEmitContext& Context, FExpressi
 
 void UE::HLSLTree::FExpressionCast::EmitCode(FEmitContext& Context, FExpressionEmitResult& OutResult) const
 {
-	const FEmitValue& InputRef = Context.AcquireValue(Input);
-	const Shader::EValueType InputType = InputRef.GetExpressionType();
-	const TCHAR* InputCode = Context.GetCode(InputRef);
+	const FEmitValue* InputValue = Context.AcquireValue(Input);
+	const Shader::EValueType InputType = InputValue->GetExpressionType();
+	const TCHAR* InputCode = Context.GetCode(InputValue);
 
 	OutResult.bInline = true;
 	OutResult.Type = Type;
@@ -424,50 +420,59 @@ void UE::HLSLTree::FExpressionCast::EmitCode(FEmitContext& Context, FExpressionE
 		else
 		{
 			// Padding with 0s, use a constructor
+			const Shader::FValue ZeroValue(OutputTypeDesc.ComponentType, 1);
 			OutResult.Writer.Writef(TEXT("%s(%s"), OutputTypeDesc.Name, InputCode);
 			for (int32 i = InputTypeDesc.NumComponents; i < OutputTypeDesc.NumComponents; ++i)
 			{
-				// TODO - worth making these 0s the appropriate type? 0u, 0.0f, etc
-				OutResult.Writer.Write(TEXT(", 0"));
+				OutResult.Writer.Write(TEXT(", "));
+				OutResult.Writer.WriteConstant(ZeroValue);
 			}
 			OutResult.Writer.Write(TEXT(")"));
 		}
 	}
 }
 
+void UE::HLSLTree::FExpressionReflectionVector::EmitCode(FEmitContext& Context, FExpressionEmitResult& OutResult) const
+{
+	OutResult.bInline = true;
+	OutResult.Type = Shader::EValueType::Float3;
+	OutResult.EvaluationType = EExpressionEvaluationType::Shader;
+	OutResult.Writer.Write(TEXT("Parameters.ReflectionVector"));
+}
+
 void UE::HLSLTree::FExpressionFunctionInput::EmitCode(FEmitContext& Context, FExpressionEmitResult& OutResult) const
 {
 	const FEmitContext::FFunctionStackEntry& StackEntry = Context.FunctionStack.Last();
 
 	check(InputIndex >= 0 && InputIndex < StackEntry.FunctionCall->NumInputs);
 	FExpression* InputExpression = StackEntry.FunctionCall->Inputs[InputIndex];
-	const FEmitValue& Value = Context.AcquireValue(InputExpression);
+	const FEmitValue* Value = Context.AcquireValue(InputExpression);
 	OutResult.ForwardValue(Context, Value);
 }
 
 void UE::HLSLTree::FExpressionFunctionOutput::EmitCode(FEmitContext& Context, FExpressionEmitResult& OutResult) const
 {
-	const FEmitValue& Value = Context.AcquireValue(FunctionCall, OutputIndex);
+	const FEmitValue* Value = Context.AcquireValue(FunctionCall, OutputIndex);
 	OutResult.ForwardValue(Context, Value);
 }
 
 void UE::HLSLTree::FStatementReturn::EmitHLSL(FEmitContext& Context, FCodeWriter& OutWriter) const
 {
-	const FEmitValue& Ref = Context.AcquireValue(Expression);
-	OutWriter.WriteLinef(TEXT("return %s;"), Context.GetCode(Ref));
+	const FEmitValue* Value = Context.AcquireValue(Expression);
+	OutWriter.WriteLinef(TEXT("return %s;"), Context.GetCode(Value));
 }
 
 void UE::HLSLTree::FStatementSetLocalVariable::EmitHLSL(FEmitContext& Context, FCodeWriter& OutWriter) const
 {
-	const FEmitValue& DeclarationRef = Context.AcquireValue(Declaration);
-	const FEmitValue& ExpressionRef = Context.AcquireValue(Expression);
-	OutWriter.WriteLinef(TEXT("%s = %s;"), Context.GetCode(DeclarationRef), Context.GetCode(ExpressionRef));
+	const FEmitValue* DeclarationValue = Context.AcquireValue(Declaration);
+	const FEmitValue* ExpressionValue = Context.AcquireValue(Expression);
+	OutWriter.WriteLinef(TEXT("%s = %s;"), Context.GetCode(DeclarationValue), Context.GetCode(ExpressionValue));
 }
 
 void UE::HLSLTree::FStatementIf::EmitHLSL(FEmitContext& Context, FCodeWriter& OutWriter) const
 {
-	const FEmitValue& ConditionRef = Context.AcquireValue(ConditionExpression);
-	OutWriter.WriteLinef(TEXT("if (%s)"), Context.GetCode(ConditionRef));
+	const FEmitValue* ConditionValue = Context.AcquireValue(ConditionExpression);
+	OutWriter.WriteLinef(TEXT("if (%s)"), Context.GetCode(ConditionValue));
 	ThenScope->EmitHLSL(Context, OutWriter);
 	if (ElseScope)
 	{
diff --git a/Engine/Source/Runtime/Engine/Private/Materials/MaterialExpressionHLSL.cpp b/Engine/Source/Runtime/Engine/Private/Materials/MaterialExpressionHLSL.cpp
--- a/Engine/Source/Runtime/Engine/Private/Materials/MaterialExpressionHLSL.cpp
+++ b/Engine/Source/Runtime/Engine/Private/Materials/MaterialExpressionHLSL.cpp
@@ -29,6 +29,7 @@
 #include "Materials/MaterialExpressionGetLocal.h"
 #include "Materials/MaterialExpressionAdd.h"
 #include "Materials/MaterialExpressionSetMaterialAttributes.h"
+#include "Materials/MaterialExpressionReflectionVectorWS.h"
 #include "Materials/MaterialExpressionSetLocal.h"
 #include "Materials/MaterialExpressionIfThenElse.h"
 #include "Materials/MaterialExpressionForLoop.h"
@@ -247,6 +248,14 @@ EMaterialGenerateHLSLStatus UMaterialExpressionSetMaterialAttributes::GenerateHL
 	return EMaterialGenerateHLSLStatus::Success;
 }
 
+EMaterialGenerateHLSLStatus UMaterialExpressionReflectionVectorWS::GenerateHLSLExpression(FMaterialHLSLGenerator& Generator, UE::HLSLTree::FScope& Scope, int32 OutputIndex, UE::HLSLTree::FExpression*& OutExpression)
+{
+	check(!CustomWorldNormal.GetTracedInput().Expression); // TODO
+
+	OutExpression = Generator.GetTree().NewExpression<UE::HLSLTree::FExpressionReflectionVector>(Scope);
+	return EMaterialGenerateHLSLStatus::Success;
+}
+
 EMaterialGenerateHLSLStatus UMaterialExpressionFunctionInput::GenerateHLSLExpression(FMaterialHLSLGenerator& Generator, UE::HLSLTree::FScope& Scope, int32 OutputIndex, UE::HLSLTree::FExpression*& OutExpression)
 {
 	UE::Shader::EValueType ExpressionType = UE::Shader::EValueType::Void;
diff --git a/Engine/Source/Runtime/Engine/Public/HLSLTree/HLSLTree.h b/Engine/Source/Runtime/Engine/Public/HLSLTree/HLSLTree.h
--- a/Engine/Source/Runtime/Engine/Public/HLSLTree/HLSLTree.h
+++ b/Engine/Source/Runtime/Engine/Public/HLSLTree/HLSLTree.h
@@ -110,19 +110,19 @@ public:
 	~FEmitContext();
 
 	/** Returns a value that references the given expression */
-	const FEmitValue& AcquireValue(FExpression* Expression);
+	const FEmitValue* AcquireValue(FExpression* Expression);
 
 	/** Returns value that references the given local variable */
-	const FEmitValue& AcquireValue(FLocalDeclaration* Declaration);
+	const FEmitValue* AcquireValue(FLocalDeclaration* Declaration);
 
 	/** Returns value that references the given local function output */
-	const FEmitValue& AcquireValue(FFunctionCall* FunctionCall, int32 OutputIndex);
+	const FEmitValue* AcquireValue(FFunctionCall* FunctionCall, int32 OutputIndex);
 
 	/** Gets HLSL code that references the given value */
-	const TCHAR* GetCode(const FEmitValue& Value) const;
+	const TCHAR* GetCode(const FEmitValue* Value) const;
 
 	/** Append preshader bytecode that represents the given value */
-	void AppendPreshader(const FEmitValue& Value, Shader::FPreshaderData& InOutPreshader) const;
+	void AppendPreshader(const FEmitValue* Value, Shader::FPreshaderData& InOutPreshader) const;
 
 	struct FScopeEntry
 	{
@@ -138,7 +138,7 @@ public:
 
 	struct FFunctionCallEntry
 	{
-		const FEmitValue* OutputRef;
+		const FEmitValue* OutputValues;
 		int32 NumOutputs;
 	};
 
@@ -174,7 +174,7 @@ struct FExpressionEmitResult
 		, bInline(false)
 	{}
 
-	void ForwardValue(FEmitContext& Context, const FEmitValue& InValue);
+	void ForwardValue(FEmitContext& Context, const FEmitValue* InValue);
 
 	FCodeWriter& Writer;
 	Shader::FPreshaderData& Preshader;
diff --git a/Engine/Source/Runtime/Engine/Public/HLSLTree/HLSLTreeCommon.h b/Engine/Source/Runtime/Engine/Public/HLSLTree/HLSLTreeCommon.h
--- a/Engine/Source/Runtime/Engine/Public/HLSLTree/HLSLTreeCommon.h
+++ b/Engine/Source/Runtime/Engine/Public/HLSLTree/HLSLTreeCommon.h
@@ -277,6 +277,12 @@ public:
 	virtual void EmitCode(FEmitContext& Context, FExpressionEmitResult& OutResult) const override;
 };
 
+class FExpressionReflectionVector : public FExpression
+{
+public:
+	virtual void EmitCode(FEmitContext& Context, FExpressionEmitResult& OutResult) const override;
+};
+
 class FExpressionFunctionInput : public FExpression
 {
 public:
diff --git a/Engine/Source/Runtime/Engine/Public/Shader/ShaderTypes.h b/Engine/Source/Runtime/Engine/Public/Shader/ShaderTypes.h
--- a/Engine/Source/Runtime/Engine/Public/Shader/ShaderTypes.h
+++ b/Engine/Source/Runtime/Engine/Public/Shader/ShaderTypes.h
@@ -81,6 +81,15 @@ struct FValue
 {
 	inline FValue() : ComponentType(EValueComponentType::Void), NumComponents(0) {}
 
+	explicit FValue(EValueType InType)
+	{
+		const FValueTypeDescription TypeDesc = GetValueTypeDescription(InType);
+		ComponentType = TypeDesc.ComponentType;
+		NumComponents = TypeDesc.NumComponents;
+	}
+
+	inline FValue(EValueComponentType InComponentType, int8 InNumComponents) : ComponentType(InComponentType), NumComponents(InNumComponents) {}
+
 	inline FValue(float v) : ComponentType(EValueComponentType::Float), NumComponents(1)
 	{
 		Component[0].Float = v;
