Temporary workaround for out of bound array access in FInstancedStaticMeshSceneProxy::SetupProxy

diff --git a/Engine/Source/Runtime/Engine/Private/InstancedStaticMesh.cpp b/Engine/Source/Runtime/Engine/Private/InstancedStaticMesh.cpp
--- a/Engine/Source/Runtime/Engine/Private/InstancedStaticMesh.cpp
+++ b/Engine/Source/Runtime/Engine/Private/InstancedStaticMesh.cpp
@@ -1248,7 +1248,9 @@ void FInstancedStaticMeshSceneProxy::SetupProxy(UInstancedStaticMeshComponent* I
 			// by the async build, which re-creates the proxy in a nearby future frame. All this should be removed in favour of GPU-side culling.
 			if (OutInstanceIndex < InstanceReorderTable.Num() && InstanceReorderTable[OutInstanceIndex] < Instances.Num())
 			{
-				OutInstanceIndex = InstanceReorderTable[OutInstanceIndex];
+				// Temporary workaround for out of bound arrary access
+				// TODO: fix this properly
+				OutInstanceIndex = InstanceReorderTable[OutInstanceIndex] != INDEX_NONE ? InstanceReorderTable[OutInstanceIndex] : OutInstanceIndex;
 			}
 			FTransform InstanceTransform;
 			InComponent->GetInstanceTransform(InInstanceIndex, InstanceTransform);
