Fix LevelScriptBlueprint not loading external actors through instancing context

diff --git a/Engine/Source/Runtime/Engine/Classes/Engine/Level.h b/Engine/Source/Runtime/Engine/Classes/Engine/Level.h
--- a/Engine/Source/Runtime/Engine/Classes/Engine/Level.h
+++ b/Engine/Source/Runtime/Engine/Classes/Engine/Level.h
@@ -1047,6 +1047,19 @@ public:
 	 */
 	ENGINE_API void ConvertAllActorsToPackaging(bool bExternal);
 
+	/**
+	* Get a properly formated external actor package instance name for this level package to be used in FLinkerInstancingContext
+	* @return external actor package instance name
+	*/
+	static ENGINE_API FString GetExternalActorPackageInstanceName(const FString& LevelPackageName, const FString& ActorShortPackageName);
+
+	/**
+	 * Get the list of (on disk) external actor packages associated with this external actors path
+	 * @param ExternalActorsPath the path to scan for external actor packages
+	 * @return Array of packages associated with this level
+	 */
+	static ENGINE_API TArray<FString> GetOnDiskExternalActorPackages(const FString& ExternalActorsPath);
+
 	/**
 	 * Get the list of (on disk) external actor packages associated with this level
 	 * @return Array of packages associated with this level
diff --git a/Engine/Source/Runtime/Engine/Private/Level.cpp b/Engine/Source/Runtime/Engine/Private/Level.cpp
--- a/Engine/Source/Runtime/Engine/Private/Level.cpp
+++ b/Engine/Source/Runtime/Engine/Private/Level.cpp
@@ -790,7 +790,7 @@ void ULevel::PostLoad()
 			if (bInstanced)
 			{
 				const FString ActorShortPackageName = FPackageName::GetShortName(ActorPackageName);
-				const FString InstancedName = FString::Printf(TEXT("%s_InstanceOf_%s"), *LevelPackage->GetName(), *ActorShortPackageName);
+				const FString InstancedName = GetExternalActorPackageInstanceName(LevelPackage->GetName(), ActorShortPackageName);
 				InstancePackageNames.Add(InstancedName);
 
 				InstancingContext.AddMapping(FName(*ActorPackageName), FName(*InstancedName));
@@ -2468,11 +2468,14 @@ void ULevel::ConvertAllActorsToPackaging(bool bExternal)
 	}
 }
 
-TArray<FString> ULevel::GetOnDiskExternalActorPackages() const
+FString ULevel::GetExternalActorPackageInstanceName(const FString& LevelPackageName, const FString& ActorShortPackageName)
+{
+	return FString::Printf(TEXT("%s_InstanceOf_%s"), *LevelPackageName, *ActorShortPackageName);
+}
+
+TArray<FString> ULevel::GetOnDiskExternalActorPackages(const FString& ExternalActorsPath)
 {
 	TArray<FString> ActorPackageNames;
-	UWorld* World = GetTypedOuter<UWorld>();
-	FString ExternalActorsPath = ULevel::GetExternalActorsPath(World->GetPackage(), World->OriginalWorldName == NAME_None ? World->GetName() : World->OriginalWorldName.ToString());
 	if (!ExternalActorsPath.IsEmpty())
 	{
 		IFileManager::Get().IterateDirectoryRecursively(*FPackageName::LongPackageNameToFilename(ExternalActorsPath), [&ActorPackageNames](const TCHAR* FilenameOrDirectory, bool bIsDirectory)
@@ -2491,6 +2494,13 @@ TArray<FString> ULevel::GetOnDiskExternalActorPackages() const
 	return ActorPackageNames;
 }
 
+TArray<FString> ULevel::GetOnDiskExternalActorPackages() const
+{
+	UWorld* World = GetTypedOuter<UWorld>();
+	FString ExternalActorsPath = ULevel::GetExternalActorsPath(World->GetPackage(), World->OriginalWorldName == NAME_None ? World->GetName() : World->OriginalWorldName.ToString());
+	return GetOnDiskExternalActorPackages(ExternalActorsPath);
+}
+
 TArray<UPackage*> ULevel::GetLoadedExternalActorPackages() const
 {
 	TSet<UPackage*> ActorPackages;
diff --git a/Engine/Source/Runtime/Engine/Private/LevelStreaming.cpp b/Engine/Source/Runtime/Engine/Private/LevelStreaming.cpp
--- a/Engine/Source/Runtime/Engine/Private/LevelStreaming.cpp
+++ b/Engine/Source/Runtime/Engine/Private/LevelStreaming.cpp
@@ -1160,18 +1160,39 @@ bool ULevelStreaming::RequestLevel(UWorld* PersistentWorld, bool bAllowLevelLoad
 			// Kick off async load request.
 			STAT_ADD_CUSTOMMESSAGE_NAME( STAT_NamedMarker, *(FString( TEXT( "RequestLevel - " ) + DesiredPackageName.ToString() )) );
 			TRACE_BOOKMARK(TEXT("RequestLevel - %s"), *DesiredPackageName.ToString());
-
-			// When loading an instanced package we want to avoid it being processed as an asset so we make sure to set it RF_Transient.
-			// If the package is not created here, it will get created by the LoadPackageAsync call. 
-			// Gameworld packages (PIE) are already ignored so we can let LoadPackageAsync do its job.
-			if (!bIsGameWorld && DesiredPackageName != NAME_None && PackagePath.GetPackageFName() != DesiredPackageName)
+			
+			FLinkerInstancingContext* InstancingContextPtr = nullptr;
+#if WITH_EDITOR
+			FLinkerInstancingContext InstancingContext;
+			if (DesiredPackageName != NAME_None && PackagePath.GetPackageFName() != DesiredPackageName)
 			{
-				UPackage* NewPackage = CreatePackage(*DesiredPackageName.ToString());
-				NewPackage->SetFlags(RF_Transient);
-			}
+				// When loading an instanced package we want to avoid it being processed as an asset so we make sure to set it RF_Transient.
+				// If the package is not created here, it will get created by the LoadPackageAsync call. 
+				// Gameworld packages (PIE) are already ignored so we can let LoadPackageAsync do its job.
+				if (!bIsGameWorld)
+				{
+					UPackage* NewPackage = CreatePackage(*DesiredPackageName.ToString());
+					NewPackage->SetFlags(RF_Transient);
+				}
 
+				// When loading an instanced package we need to build an instancing context in case non external actors part of the level are 
+				// pulling on external actors.
+				UWorld* World = GetTypedOuter<UWorld>();
+				FString ExternalActorsPath = ULevel::GetExternalActorsPath(PackagePath.GetPackageName());
+				TArray<FString> ActorPackageNames = ULevel::GetOnDiskExternalActorPackages(ExternalActorsPath);
+								
+				InstancingContextPtr = &InstancingContext;
+				for (const FString& ActorPackageName : ActorPackageNames)
+				{
+					const FString ActorShortPackageName = FPackageName::GetShortName(ActorPackageName);
+					const FString InstancedName = ULevel::GetExternalActorPackageInstanceName(DesiredPackageName.ToString(), ActorShortPackageName);
+				
+					InstancingContext.AddMapping(FName(*ActorPackageName), FName(*InstancedName));
+				}
+			}
+#endif
 
-			LoadPackageAsync(PackagePath, DesiredPackageName, FLoadPackageAsyncDelegate::CreateUObject(this, &ULevelStreaming::AsyncLevelLoadComplete), nullptr, PackageFlags, PIEInstanceID, GetPriority());
+			LoadPackageAsync(PackagePath, DesiredPackageName, FLoadPackageAsyncDelegate::CreateUObject(this, &ULevelStreaming::AsyncLevelLoadComplete), nullptr, PackageFlags, PIEInstanceID, GetPriority(), InstancingContextPtr);
 
 			// streamingServer: server loads everything?
 			// Editor immediately blocks on load and we also block if background level streaming is disabled.
