Allow all but force-local shader compilation jobs on XGE (UE-114104).

- This reverts an earlier change made as part of CL 15878696 which regressed editor startup times.

diff --git a/Engine/Source/Runtime/Engine/Private/ShaderCompiler/ShaderCompiler.cpp b/Engine/Source/Runtime/Engine/Private/ShaderCompiler/ShaderCompiler.cpp
--- a/Engine/Source/Runtime/Engine/Private/ShaderCompiler/ShaderCompiler.cpp
+++ b/Engine/Source/Runtime/Engine/Private/ShaderCompiler/ShaderCompiler.cpp
@@ -2994,8 +2994,9 @@ FShaderCompilingManager::FShaderCompilingManager() :
 	TUniquePtr<FShaderCompileThreadRunnableBase> LocalThread = MakeUnique<FShaderCompileThreadRunnable>(this);
 	if (RemoteCompileThread)
 	{
-		// Keep high priority jobs on the local machine, to avoid XGE latency
-		RemoteCompileThread->SetPriorityRange(EShaderCompileJobPriority::Low, EShaderCompileJobPriority::Normal);
+		// Only force-local jobs are guaranteed to stay on the local machine. Going wide with High priority jobs is important for the startup times,
+		// since special materials use High priority. Possibly the partition by priority is too rigid in general.
+		RemoteCompileThread->SetPriorityRange(EShaderCompileJobPriority::Low, EShaderCompileJobPriority::High);
 		LocalThread->SetPriorityRange(EShaderCompileJobPriority::Normal, EShaderCompileJobPriority::ForceLocal);
 		Threads.Add(MoveTemp(RemoteCompileThread));
 	}
