Fix checkerboard minimap because of skipped finish shader compilation (GetNumRemainingJobs is wrong and doesn't include everything)

diff --git a/Engine/Source/Runtime/Engine/Private/WorldPartition/WorldPartitionMiniMapHelper.cpp b/Engine/Source/Runtime/Engine/Private/WorldPartition/WorldPartitionMiniMapHelper.cpp
--- a/Engine/Source/Runtime/Engine/Private/WorldPartition/WorldPartitionMiniMapHelper.cpp
+++ b/Engine/Source/Runtime/Engine/Private/WorldPartition/WorldPartitionMiniMapHelper.cpp
@@ -56,8 +56,7 @@ void FWorldPartitionMiniMapHelper::CaptureWorldMiniMapToTexture(UWorld* InWorld,
 {
 	auto WaitForShaderCompilation = [](){
 		UE_SCOPED_TIMER(TEXT("Shader Compilation"), LogWorldPartitionMiniMapHelper);
-		auto RemainingShaders = GShaderCompilingManager->GetNumRemainingJobs();
-		if (RemainingShaders > 0)
+		if (GShaderCompilingManager && GShaderCompilingManager->IsCompiling())
 		{
 			GShaderCompilingManager->FinishAllCompilation();
 		}
@@ -73,7 +72,7 @@ void FWorldPartitionMiniMapHelper::CaptureWorldMiniMapToTexture(UWorld* InWorld,
 	WaitForDistanceFieldBuilding();
 	WaitForShaderCompilation();
 	FAssetCompilingManager::Get().FinishAllCompilation();
-	
+
 	//Calculate bounds of the World
 	OutWorldBounds = InWorld->GetWorldPartition()->GetEditorWorldBounds();
 	
@@ -88,7 +87,6 @@ void FWorldPartitionMiniMapHelper::CaptureWorldMiniMapToTexture(UWorld* InWorld,
 	FMatrix ProjectionMatrix;
 	FWorldPartitionMiniMapHelper::CalTopViewOfWorld(ProjectionMatrix, OutWorldBounds, ViewportWidth, ViewportHeight);
 
-	TArray<FColor> CapturedImage;
 	//Using SceneCapture Actor capture the scene to buffer
 	UTextureRenderTarget2D* RenderTargetTexture = NewObject<UTextureRenderTarget2D>();
 	RenderTargetTexture->ClearColor = FLinearColor::Transparent;
