Reverting Change 14206202.  It is causing playcontrollers and their pawns to stick around.  Also NetConnection's Child Connections are already handled in the CleanUp of the NetConnection.

diff --git a/Engine/Source/Runtime/Engine/Private/PlayerController.cpp b/Engine/Source/Runtime/Engine/Private/PlayerController.cpp
--- a/Engine/Source/Runtime/Engine/Private/PlayerController.cpp
+++ b/Engine/Source/Runtime/Engine/Private/PlayerController.cpp
@@ -183,19 +183,6 @@ bool APlayerController::DestroyNetworkActorHandled()
 	UNetConnection* C = Cast<UNetConnection>(Player);
 	if (C)
 	{
-		UChildConnection* CC = Cast<UChildConnection>(Player);
-		if (CC)
-		{
-			// we are a child (splitscreen) player controller, we'll proceed to a normal destruction scheme by returning false.
-			UNetConnection* Parent = CC->GetParentConnection();
-			if (Parent)
-			{
-				// remove the child connection from the parent
-				Parent->Children.Remove(CC);
-			}
-			return false;
-		}
-
 		if (C->Channels[0] && C->State != USOCK_Closed)
 		{
 			C->bPendingDestroy = true;
@@ -1363,13 +1350,12 @@ void APlayerController::OnNetCleanup(UNetConnection* Connection)
 
 	check(UNetConnection::GNetConnectionBeingCleanedUp == NULL);
 	UNetConnection::GNetConnectionBeingCleanedUp = Connection;
-	
-	Destroy( true );
 	//@note: if we ever implement support for splitscreen players leaving a match without the primary player leaving, we'll need to insert
 	// a call to ClearOnlineDelegates() here so that PlayerController.ClearOnlineDelegates can use the correct ControllerId (which lives
 	// in ULocalPlayer)
 	Player = NULL;
 	NetConnection = NULL;	
+	Destroy( true );
 	UNetConnection::GNetConnectionBeingCleanedUp = NULL;
 }
 
