Skin cache memory breakdown print out functionality

diff --git a/Engine/Source/Runtime/Engine/Private/GPUSkinCache.cpp b/Engine/Source/Runtime/Engine/Private/GPUSkinCache.cpp
--- a/Engine/Source/Runtime/Engine/Private/GPUSkinCache.cpp
+++ b/Engine/Source/Runtime/Engine/Private/GPUSkinCache.cpp
@@ -146,6 +146,17 @@ FAutoConsoleVariableRef CVarGPUSkinCacheMaxDispatchesPerCmdList(
 	ECVF_RenderThreadSafe
 );
 
+static int32 GSkinCachePrintMemorySummary = 0;
+FAutoConsoleVariableRef CVarGPUSkinCachePrintMemorySummary(
+	TEXT("r.SkinCache.PrintMemorySummary"),
+	GSkinCachePrintMemorySummary,
+	TEXT("Print break down of memory usage.")
+	TEXT(" 0: off (default),")
+	TEXT(" 1: print when out of memory,")
+	TEXT(" 2: print every frame"),
+	ECVF_RenderThreadSafe
+);
+
 static int32 GGPUSkinCacheFlushCounter = 0;
 
 #if RHI_RAYTRACING
@@ -166,6 +177,8 @@ FAutoConsoleVariableRef CVarGPUSkinCacheMaxRayTracingPrimitivesPerCmdList(
 );
 #endif
 
+const float MBSize = 1048576.f; // 1024 x 1024 bytes
+
 static inline bool DoesPlatformSupportGPUSkinCache(const FStaticShaderPlatform Platform)
 {
 	return Platform == SP_PCD3D_SM5 || IsMetalSM5Platform(Platform) || IsVulkanSM5Platform(Platform) || FDataDrivenShaderPlatformInfo::GetSupportsGPUSkinCache(Platform);
@@ -1138,11 +1151,6 @@ IMPLEMENT_SHADER_TYPE(template<>, FRecomputeTangentsPerVertexPassCS<4>, TEXT("/E
 
 void FGPUSkinCache::DispatchUpdateSkinTangents(FRHICommandListImmediate& RHICmdList, FGPUSkinCacheEntry* Entry, int32 SectionIndex, FRWBuffer*& StagingBuffer, bool bTrianglePass)
 {
-	FString MeshName = TEXT("None");
-#if !UE_BUILD_SHIPPING
-	MeshName = Entry->GPUSkin->DebugName.ToString();
-#endif // !UE_BUILD_SHIPPING
-
 	FGPUSkinCacheEntry::FSectionDispatchData& DispatchData = Entry->DispatchData[SectionIndex];
 
 	FSkeletalMeshRenderData& SkelMeshRenderData = Entry->GPUSkin->GetSkeletalMeshRenderData();
@@ -1209,7 +1217,7 @@ void FGPUSkinCache::DispatchUpdateSkinTangents(FRHICommandListImmediate& RHICmdL
 			uint32 ThreadGroupCountValue = FMath::DivideAndRoundUp(NumTriangles, FBaseRecomputeTangentsPerTriangleShader::ThreadGroupSizeX);
 
 			SCOPED_DRAW_EVENTF(RHICmdList, SkinTangents_PerTrianglePass, TEXT("TangentsTri  Mesh=%s, LOD=%d, Chunk=%d, IndexStart=%d Tri=%d BoneInfluenceType=%d UVPrecision=%d"),
-				*MeshName, LODIndex, SectionIndex, DispatchData.IndexBufferOffsetValue, DispatchData.NumTriangles, Entry->BoneInfluenceType, bFullPrecisionUV);
+				*GetSkeletalMeshObjectName(Entry->GPUSkin), LODIndex, SectionIndex, DispatchData.IndexBufferOffsetValue, DispatchData.NumTriangles, Entry->BoneInfluenceType, bFullPrecisionUV);
 
 			FRHIComputeShader* ShaderRHI = Shader.GetComputeShader();
 			RHICmdList.SetComputeShader(ShaderRHI);
@@ -1237,7 +1245,7 @@ void FGPUSkinCache::DispatchUpdateSkinTangents(FRHICommandListImmediate& RHICmdL
 	else
 	{
 		SCOPED_DRAW_EVENTF(RHICmdList, SkinTangents_PerVertexPass, TEXT("TangentsVertex Mesh=%s, LOD=%d, Chunk=%d, InputStreamStart=%d, OutputStreamStart=%d, Vert=%d"),
-			*MeshName, LODIndex, SectionIndex, DispatchData.InputStreamStart, DispatchData.OutputStreamStart, DispatchData.NumVertices);
+			*GetSkeletalMeshObjectName(Entry->GPUSkin), LODIndex, SectionIndex, DispatchData.InputStreamStart, DispatchData.OutputStreamStart, DispatchData.NumVertices);
 		//#todo-gpuskin Feature level?
 		auto* GlobalShaderMap = GetGlobalShaderMap(GetFeatureLevel());
 		TShaderMapRef<FRecomputeTangentsPerVertexPassCS<0>> ComputeShader0(GlobalShaderMap);
@@ -1523,6 +1531,13 @@ void FGPUSkinCache::ProcessEntry(
 		FRWBuffersAllocation* NewPositionAllocation = TryAllocBuffer(TotalNumVertices, WithTangents, TotalNumTriangles, RHICmdList);
 		if (!NewPositionAllocation)
 		{
+			if (GSkinCachePrintMemorySummary > 0)
+			{
+				uint64 RequiredMemInBytes = FRWBuffersAllocation::CalculateRequiredMemory(TotalNumVertices, WithTangents, TotalNumTriangles);
+				UE_LOG(LogSkinCache, Warning, TEXT("FGPUSkinCache::ProcessEntry failed to allocate %.3fMB for mesh %s LOD%d, extra required memory increased to %.3fMB"),
+					RequiredMemInBytes / MBSize, *GetSkeletalMeshObjectName(Skin), LODIndex, ExtraRequiredMemory / MBSize);
+			}
+
 			// Couldn't fit; caller will notify OOM
 			return;
 		}
@@ -1928,14 +1943,9 @@ void FGPUSkinCache::DispatchUpdateSkinning(FRHICommandListImmediate& RHICmdList,
 	FGPUSkinCacheEntry::FSectionDispatchData& DispatchData = Entry->DispatchData[Section];
 	FGPUBaseSkinVertexFactory::FShaderDataType& ShaderData = DispatchData.SourceVertexFactory->GetShaderData();
 
-	FString MeshName = TEXT("None");
-#if !UE_BUILD_SHIPPING
-	MeshName = Entry->GPUSkin->DebugName.ToString();
-#endif // !UE_BUILD_SHIPPING
-
 	SCOPED_DRAW_EVENTF(RHICmdList, SkinCacheDispatch,
 		TEXT("Skinning%d%d%d Mesh=%s LOD=%d Chunk=%d InStreamStart=%d OutStart=%d Vert=%d Morph=%d/%d"),
-		(int32)Entry->bUse16BitBoneIndex, (int32)Entry->BoneInfluenceType, DispatchData.SkinType, *MeshName, Entry->GPUSkin->GetLOD(),
+		(int32)Entry->bUse16BitBoneIndex, (int32)Entry->BoneInfluenceType, DispatchData.SkinType, *GetSkeletalMeshObjectName(Entry->GPUSkin), Entry->GPUSkin->GetLOD(),
 		DispatchData.SectionIndex, DispatchData.InputStreamStart, DispatchData.OutputStreamStart, DispatchData.NumVertices, Entry->MorphBuffer != 0, DispatchData.MorphBufferOffset);
 	auto* GlobalShaderMap = GetGlobalShaderMap(GetFeatureLevel());
 	TShaderMapRef<TGPUSkinCacheCS<0>> SkinCacheCS000(GlobalShaderMap);		// 16bit_0, BoneInfluenceType_0, SkinType_0
@@ -2270,3 +2280,54 @@ void FGPUSkinCache::IncrementDispatchCounter(FRHICommandListImmediate& RHICmdLis
 		}
 	}
 }
+
+uint64 FGPUSkinCache::GetExtraRequiredMemoryAndReset()
+{
+	if (GSkinCachePrintMemorySummary == 2 || (GSkinCachePrintMemorySummary == 1 && ExtraRequiredMemory > 0))
+	{
+		PrintMemorySummary();
+	}
+
+	uint64 OriginalValue = ExtraRequiredMemory;
+	ExtraRequiredMemory = 0;
+	return OriginalValue;
+}
+
+void FGPUSkinCache::PrintMemorySummary() const
+{
+	UE_LOG(LogSkinCache, Display, TEXT("======= Skin Cache Memory Usage Summary ======="));
+
+	uint64 TotalMemInBytes = 0;
+	for (int32 i = 0; i < Entries.Num(); ++i)
+	{
+		FGPUSkinCacheEntry* Entry = Entries[i];
+		if (Entry)
+		{
+			uint64 MemInBytes = Entry->PositionAllocation ? Entry->PositionAllocation->GetNumBytes() : 0;
+			UE_LOG(LogSkinCache, Display, TEXT("   SkinCacheEntry_%d: Mesh=%s, LOD=%d, Mem=%.3fKB"), i, *GetSkeletalMeshObjectName(Entry->GPUSkin), Entry->LOD, MemInBytes / 1024.f);
+			TotalMemInBytes += MemInBytes;
+		}
+	}
+	ensure(TotalMemInBytes == UsedMemoryInBytes);
+
+	uint64 MaxSizeInBytes = (uint64)(GSkinCacheSceneMemoryLimitInMB * MBSize);
+	uint64 UnusedSizeInBytes = MaxSizeInBytes - UsedMemoryInBytes;
+
+	UE_LOG(LogSkinCache, Display, TEXT("Used: %.3fMB"), UsedMemoryInBytes / MBSize);
+	UE_LOG(LogSkinCache, Display, TEXT("Available: %.3fMB"), UnusedSizeInBytes / MBSize);
+	UE_LOG(LogSkinCache, Display, TEXT("Total limit: %.3fMB"), GSkinCacheSceneMemoryLimitInMB);
+	UE_LOG(LogSkinCache, Display, TEXT("Extra required: %.3fMB"), ExtraRequiredMemory / MBSize);
+	UE_LOG(LogSkinCache, Display, TEXT("==============================================="));
+}
+
+FString FGPUSkinCache::GetSkeletalMeshObjectName(const FSkeletalMeshObjectGPUSkin* GPUSkin) const
+{
+	FString Name = TEXT("None");
+	if (GPUSkin)
+	{
+#if !UE_BUILD_SHIPPING
+		Name = GPUSkin->DebugName.ToString();
+#endif // !UE_BUILD_SHIPPING
+	}
+	return Name;
+}
diff --git a/Engine/Source/Runtime/Engine/Public/GPUSkinCache.h b/Engine/Source/Runtime/Engine/Public/GPUSkinCache.h
--- a/Engine/Source/Runtime/Engine/Public/GPUSkinCache.h
+++ b/Engine/Source/Runtime/Engine/Public/GPUSkinCache.h
@@ -205,12 +205,7 @@ public:
 
 	static bool UseIntermediateTangents();
 
-	inline uint64 GetExtraRequiredMemoryAndReset()
-	{
-		uint64 OriginalValue = ExtraRequiredMemory;
-		ExtraRequiredMemory = 0;
-		return OriginalValue;
-	}
+	ENGINE_API uint64 GetExtraRequiredMemoryAndReset();
 
 	enum
 	{
@@ -472,6 +467,9 @@ protected:
 
 	void IncrementDispatchCounter(FRHICommandListImmediate& RHICmdList);
 	int32 DispatchCounter = 0;
+
+	void PrintMemorySummary() const;
+	FString GetSkeletalMeshObjectName(const FSkeletalMeshObjectGPUSkin* GPUSkin) const;
 };
 
 DECLARE_STATS_GROUP(TEXT("GPU Skin Cache"), STATGROUP_GPUSkinCache, STATCAT_Advanced);
