Deprecate property used to control serialization of cached material data.  Instead add an explicitly serialized flag.  With the property it's easy to leak the value across serialize calls.

diff --git a/Engine/Source/Runtime/Core/Public/UObject/UE5MainStreamObjectVersion.h b/Engine/Source/Runtime/Core/Public/UObject/UE5MainStreamObjectVersion.h
--- a/Engine/Source/Runtime/Core/Public/UObject/UE5MainStreamObjectVersion.h
+++ b/Engine/Source/Runtime/Core/Public/UObject/UE5MainStreamObjectVersion.h
@@ -118,6 +118,9 @@ struct CORE_API FUE5MainStreamObjectVersion
 		// Initial set of anim nodes converted to use constants held in sparse class data
 		AnimNodeConstantDataRefactorPhase0,
 
+		// Explicitly serialized bSavedCachedExpressionData for Material(Instance)
+		MaterialSavedCachedData,
+
         // -----<new versions can be added above this line>-------------------------------------------------
 		VersionPlusOne,
 		LatestVersion = VersionPlusOne - 1
diff --git a/Engine/Source/Runtime/Engine/Classes/Materials/Material.h b/Engine/Source/Runtime/Engine/Classes/Materials/Material.h
--- a/Engine/Source/Runtime/Engine/Classes/Materials/Material.h
+++ b/Engine/Source/Runtime/Engine/Classes/Materials/Material.h
@@ -973,8 +973,12 @@ private:
 	/** Flag used to guarantee that the RT is finished using various resources in this UMaterial before cleanup. */
 	FThreadSafeBool ReleasedByRT;
 
+#if WITH_EDITORONLY_DATA
 	UPROPERTY()
-	bool bSavedCachedExpressionData;
+	bool bSavedCachedExpressionData_DEPRECATED;
+#endif
+
+	bool bLoadedCachedExpressionData;
 
 	FMaterialCachedExpressionData* CachedExpressionData;
 
diff --git a/Engine/Source/Runtime/Engine/Classes/Materials/MaterialInstance.h b/Engine/Source/Runtime/Engine/Classes/Materials/MaterialInstance.h
--- a/Engine/Source/Runtime/Engine/Classes/Materials/MaterialInstance.h
+++ b/Engine/Source/Runtime/Engine/Classes/Materials/MaterialInstance.h
@@ -445,8 +445,12 @@ private:
 	UPROPERTY()
 	FStaticParameterSet StaticParameters;
 
+#if WITH_EDITORONLY_DATA
 	UPROPERTY()
-	bool bSavedCachedData;
+	bool bSavedCachedData_DEPRECATED;
+#endif
+
+	bool bLoadedCachedData;
 
 	FMaterialInstanceCachedData* CachedData;
 
diff --git a/Engine/Source/Runtime/Engine/Private/Materials/Material.cpp b/Engine/Source/Runtime/Engine/Private/Materials/Material.cpp
--- a/Engine/Source/Runtime/Engine/Private/Materials/Material.cpp
+++ b/Engine/Source/Runtime/Engine/Private/Materials/Material.cpp
@@ -1014,6 +1014,8 @@ UMaterial::UMaterial(const FObjectInitializer& ObjectInitializer)
 
 	PhysMaterial = nullptr;
 	PhysMaterialMask = nullptr;
+
+	bLoadedCachedExpressionData = false;
 }
 
 void UMaterial::PreSave(const class ITargetPlatform* TargetPlatform)
@@ -2139,7 +2141,7 @@ void UMaterial::UpdateCachedExpressionData()
 {
 	COOK_STAT(FScopedDurationTimer BlockingTimer(MaterialCookStats::UpdateCachedExpressionDataSec));
 
-	if (bSavedCachedExpressionData)
+	if (bLoadedCachedExpressionData)
 	{
 		// Don't need to rebuild cached data if it was serialized
 		return;
@@ -3298,19 +3300,7 @@ void UMaterial::Serialize(FArchive& Ar)
 
 	Ar.UsingCustomVersion(FRenderingObjectVersion::GUID);
 	Ar.UsingCustomVersion(FFortniteMainBranchObjectVersion::GUID);
-
-	if (Ar.IsCooking())
-	{
-		if (CachedExpressionData)
-		{
-			bSavedCachedExpressionData = true;
-		}
-		else
-		{
-			// ClassDefault object is expected to be missing cached data, but in all other cases it should have been created when the material was loaded, in PostLoad
-			checkf(HasAllFlags(RF_ClassDefaultObject), TEXT("Trying to save cooked material %s, missing CachedExpressionData"), *GetName());
-		}
-	}
+	Ar.UsingCustomVersion(FUE5MainStreamObjectVersion::GUID);
 
 	Super::Serialize(Ar);
 
@@ -3340,14 +3330,43 @@ void UMaterial::Serialize(FArchive& Ar)
 #endif
 	}
 
+	bool bSavedCachedExpressionData = false;
+	if (Ar.CustomVer(FUE5MainStreamObjectVersion::GUID) >= FUE5MainStreamObjectVersion::MaterialSavedCachedData)
+	{
+		if (Ar.IsCooking())
+		{
+			if (CachedExpressionData)
+			{
+				bSavedCachedExpressionData = true;
+			}
+			else
+			{
+				// ClassDefault object is expected to be missing cached data, but in all other cases it should have been created when the material was loaded, in PostLoad
+				checkf(HasAllFlags(RF_ClassDefaultObject), TEXT("Trying to save cooked material %s, missing CachedExpressionData"), *GetName());
+			}
+		}
+
+		Ar << bSavedCachedExpressionData;
+	}
+
+#if WITH_EDITORONLY_DATA
+	if (Ar.IsLoading() && bSavedCachedExpressionData_DEPRECATED)
+	{
+		bSavedCachedExpressionData_DEPRECATED = false;
+		bSavedCachedExpressionData = true;
+	}
+#endif
+
 #if !WITH_EDITORONLY_DATA
-	checkf(bSavedCachedExpressionData, TEXT("Material %s must have saved cached expression data, if editor-only data is not present"), *GetName());
+	checkf(!Ar.IsLoading() || bSavedCachedExpressionData, TEXT("Material %s must have saved cached expression data, if editor-only data is not present"), *GetName());
 #endif
+
 	if (bSavedCachedExpressionData)
 	{
 		if (Ar.IsLoading())
 		{
 			CachedExpressionData = new FMaterialCachedExpressionData();
+			bLoadedCachedExpressionData = true;
 		}
 		check(CachedExpressionData);
 		UScriptStruct* Struct = FMaterialCachedExpressionData::StaticStruct();
diff --git a/Engine/Source/Runtime/Engine/Private/Materials/MaterialInstance.cpp b/Engine/Source/Runtime/Engine/Private/Materials/MaterialInstance.cpp
--- a/Engine/Source/Runtime/Engine/Private/Materials/MaterialInstance.cpp
+++ b/Engine/Source/Runtime/Engine/Private/Materials/MaterialInstance.cpp
@@ -47,6 +47,7 @@
 #include "ShaderPlatformQualitySettings.h"
 #include "MaterialShaderQualitySettings.h"
 #include "UObject/ObjectSaveContext.h"
+#include "UObject/UE5MainStreamObjectVersion.h"
 
 #if ENABLE_COOK_STATS
 #include "ProfilingDebugging/ScopedTimers.h"
@@ -590,6 +591,7 @@ UMaterialInstance::UMaterialInstance(const FObjectInitializer& ObjectInitializer
 	, ReleasedByRT(true)
 {
 	bHasStaticPermutationResource = false;
+	bLoadedCachedData = false;
 #if WITH_EDITOR
 	ReentrantFlag[0] = false;
 	ReentrantFlag[1] = false;
@@ -3317,19 +3319,7 @@ void UMaterialInstance::Serialize(FArchive& Ar)
 	SCOPE_CYCLE_COUNTER(STAT_MaterialInstance_Serialize);
 
 	Ar.UsingCustomVersion(FRenderingObjectVersion::GUID);
-
-	if (Ar.IsCooking())
-	{
-		if (CachedData)
-		{
-			bSavedCachedData = true;
-		}
-		else
-		{
-			// ClassDefault object is expected to be missing cached data, but in all other cases it should have been created when the material was loaded, in PostLoad
-			checkf(HasAllFlags(RF_ClassDefaultObject), TEXT("Trying to save cooked material instance %s, missing CachedExpressionData"), *GetName());
-		}
-	}
+	Ar.UsingCustomVersion(FUE5MainStreamObjectVersion::GUID);
 
 	Super::Serialize(Ar);
 		
@@ -3356,14 +3346,42 @@ void UMaterialInstance::Serialize(FArchive& Ar)
 	}
 #endif // WITH_EDITOR
 
+	bool bSavedCachedData = false;
+	if (Ar.CustomVer(FUE5MainStreamObjectVersion::GUID) >= FUE5MainStreamObjectVersion::MaterialSavedCachedData)
+	{
+		if (Ar.IsCooking())
+		{
+			if (CachedData)
+			{
+				bSavedCachedData = true;
+			}
+			else
+			{
+				// ClassDefault object is expected to be missing cached data, but in all other cases it should have been created when the material was loaded, in PostLoad
+				checkf(HasAllFlags(RF_ClassDefaultObject), TEXT("Trying to save cooked material instance %s, missing CachedExpressionData"), *GetName());
+			}
+		}
+
+		Ar << bSavedCachedData;
+	}
+#if WITH_EDITORONLY_DATA
+	if (Ar.IsLoading() && bSavedCachedData_DEPRECATED)
+	{
+		bSavedCachedData_DEPRECATED = false;
+		bSavedCachedData = true;
+	}
+#endif // WITH_EDITORONLY_DATA
+
 #if !WITH_EDITORONLY_DATA
-	//checkf(bSavedCachedData, TEXT("MaterialInstance %s must have saved cached data, if editor-only data is not present"), *GetName());
+	checkf(!Ar.IsLoading() || bSavedCachedData, TEXT("MaterialInstance %s must have saved cached data, if editor-only data is not present"), *GetName());
 #endif
+
 	if (bSavedCachedData)
 	{
 		if (Ar.IsLoading())
 		{
 			CachedData = new FMaterialInstanceCachedData();
+			bLoadedCachedData = true;
 		}
 		check(CachedData);
 		UScriptStruct* Struct = FMaterialInstanceCachedData::StaticStruct();
@@ -4111,7 +4129,7 @@ void UMaterialInstance::UpdateCachedLayerParameters()
 	COOK_STAT(FScopedDurationTimer BlockingTimer(MaterialInstanceCookStats::UpdateCachedExpressionDataSec));
 
 	// Don't need to rebuild cached data if it was serialized
-	if (!bSavedCachedData)
+	if (!bLoadedCachedData)
 	{
 		UMaterialInstance* ParentInstance = nullptr;
 		FMaterialCachedExpressionData CachedExpressionData;
