Avoid creating intermediate accumulated tangent buffer when it's not needed

diff --git a/Engine/Source/Runtime/Engine/Private/GPUSkinCache.cpp b/Engine/Source/Runtime/Engine/Private/GPUSkinCache.cpp
--- a/Engine/Source/Runtime/Engine/Private/GPUSkinCache.cpp
+++ b/Engine/Source/Runtime/Engine/Private/GPUSkinCache.cpp
@@ -1527,7 +1527,29 @@ void FGPUSkinCache::ProcessEntry(
 		bool WithTangents = RecomputeTangentsMode > 0;
 		int32 TotalNumVertices = VertexFactory->GetNumVertices();
 
-		uint32 TotalNumTriangles = GRecomputeTangentsParallelDispatch ? LodData.GetTotalFaces() : 0;
+		// IntermediateAccumulatedTangents buffer is needed if mesh has at least one section needing recomputing tangents.
+		// TotalNumTriangles > 0 signals creation of IntermediateAccumulatedTangents buffer.
+		uint32 TotalNumTriangles = 0;
+		if (GRecomputeTangentsParallelDispatch && RecomputeTangentsMode > 0)
+		{
+			bool bShouldRecomputeTangent = (RecomputeTangentsMode == 1);
+			if (!bShouldRecomputeTangent)
+			{
+				for (const FSkelMeshRenderSection& Section : LodData.RenderSections)
+				{
+					if (Section.bRecomputeTangent)
+					{
+						bShouldRecomputeTangent = true;
+						break;
+					}
+				}
+			}
+			if (bShouldRecomputeTangent)
+			{
+				TotalNumTriangles = LodData.GetTotalFaces();
+			}
+		}
+
 		FRWBuffersAllocation* NewPositionAllocation = TryAllocBuffer(TotalNumVertices, WithTangents, TotalNumTriangles, RHICmdList);
 		if (!NewPositionAllocation)
 		{
