Added sanitization of DataLayer Labels

diff --git a/Engine/Source/Editor/DataLayerEditor/Private/DataLayer/DataLayerEditorSubsystem.cpp b/Engine/Source/Editor/DataLayerEditor/Private/DataLayer/DataLayerEditorSubsystem.cpp
--- a/Engine/Source/Editor/DataLayerEditor/Private/DataLayer/DataLayerEditorSubsystem.cpp
+++ b/Engine/Source/Editor/DataLayerEditor/Private/DataLayer/DataLayerEditorSubsystem.cpp
@@ -1143,13 +1143,14 @@ void UDataLayerEditorSubsystem::DeleteDataLayer(UDataLayer* DataLayerToDelete)
 	}
 }
 
-bool UDataLayerEditorSubsystem::RenameDataLayer(UDataLayer* DataLayer, const FName& DataLayerLabel)
+bool UDataLayerEditorSubsystem::RenameDataLayer(UDataLayer* DataLayer, const FName& InDataLayerLabel)
 {
-	if (DataLayer->GetDataLayerLabel() != DataLayerLabel)
+	const FName DataLayerLabelSanitized = UDataLayer::GetSanitizedDataLayerLabel(InDataLayerLabel);
+	if (DataLayer->GetDataLayerLabel() != DataLayerLabelSanitized)
 	{
 		if (AWorldDataLayers* WorldDataLayers = GetWorldDataLayers())
 		{
-			FName UniqueNewDataLayerLabel = WorldDataLayers->GenerateUniqueDataLayerLabel(DataLayerLabel);
+			FName UniqueNewDataLayerLabel = WorldDataLayers->GenerateUniqueDataLayerLabel(DataLayerLabelSanitized);
 			DataLayer->Modify();
 			DataLayer->SetDataLayerLabel(UniqueNewDataLayerLabel);
 			DataLayerChanged.Broadcast(EDataLayerAction::Rename, DataLayer, "DataLayerLabel");
diff --git a/Engine/Source/Runtime/Engine/Private/WorldPartition/DataLayer/DataLayer.cpp b/Engine/Source/Runtime/Engine/Private/WorldPartition/DataLayer/DataLayer.cpp
--- a/Engine/Source/Runtime/Engine/Private/WorldPartition/DataLayer/DataLayer.cpp
+++ b/Engine/Source/Runtime/Engine/Private/WorldPartition/DataLayer/DataLayer.cpp
@@ -37,6 +37,9 @@ void UDataLayer::PostLoad()
 
 	// Initialize bIsVisible with persistent flag bIsInitiallyVisible
 	bIsVisible = bIsInitiallyVisible;
+
+	// Sanitize Label
+	DataLayerLabel = UDataLayer::GetSanitizedDataLayerLabel(DataLayerLabel);
 #endif
 }
 
@@ -58,15 +61,22 @@ bool UDataLayer::IsVisible() const
 #endif
 }
 
+FName UDataLayer::GetSanitizedDataLayerLabel(FName InDataLayerLabel)
+{
+	// Removes all quotes as well as whitespace characters from the startand end
+	return FName(InDataLayerLabel.ToString().TrimStartAndEnd().Replace(TEXT("\""), TEXT("")));
+}
+
 #if WITH_EDITOR
 void UDataLayer::SetDataLayerLabel(FName InDataLayerLabel)
 {
-	if (DataLayerLabel != InDataLayerLabel)
+	FName DataLayerLabelSanitized = UDataLayer::GetSanitizedDataLayerLabel(InDataLayerLabel);
+	if (DataLayerLabel != DataLayerLabelSanitized)
 	{
 		Modify();
 		AWorldDataLayers* WorldDataLayers = GetOuterAWorldDataLayers();
-		check(!WorldDataLayers || !WorldDataLayers->GetDataLayerFromLabel(InDataLayerLabel))
-		DataLayerLabel = InDataLayerLabel;
+		check(!WorldDataLayers || !WorldDataLayers->GetDataLayerFromLabel(DataLayerLabelSanitized))
+		DataLayerLabel = DataLayerLabelSanitized;
 	}
 }
 
diff --git a/Engine/Source/Runtime/Engine/Private/WorldPartition/DataLayer/DataLayerSubsystem.cpp b/Engine/Source/Runtime/Engine/Private/WorldPartition/DataLayer/DataLayerSubsystem.cpp
--- a/Engine/Source/Runtime/Engine/Private/WorldPartition/DataLayer/DataLayerSubsystem.cpp
+++ b/Engine/Source/Runtime/Engine/Private/WorldPartition/DataLayer/DataLayerSubsystem.cpp
@@ -266,21 +266,94 @@ void UDataLayerSubsystem::Draw(UCanvas* Canvas, class APlayerController* PC)
 	DrawLayerNames(TEXT("Active Data Layers"), FColor::White, ActiveDataLayerNames);
 }
 
+TArray<UDataLayer*> UDataLayerSubsystem::ConvertArgsToDataLayers(UWorld* World, const TArray<FString>& InArgs)
+{
+	TSet<UDataLayer*> OutDataLayers;
+
+	const TCHAR* QuoteChar = TEXT("\"");
+	bool bQuoteStarted = false;
+	TStringBuilder<512> Builder;
+	TArray<FString> Args;
+	for (const FString& Arg : InArgs)
+	{
+		if (!bQuoteStarted && Arg.StartsWith(QuoteChar))
+		{
+			Builder.Append(Arg.Replace(QuoteChar, TEXT("")));
+			if (Arg.EndsWith(QuoteChar) && Arg.Len() > 1)
+			{
+				Args.Add(Builder.ToString());
+				Builder.Reset();
+			}
+			else
+			{
+				bQuoteStarted = true;
+			}
+		}
+		else if (bQuoteStarted)
+		{
+			Builder.Append(TEXT(" "));
+			Builder.Append(Arg.Replace(QuoteChar, TEXT("")));
+			if (Arg.EndsWith(QuoteChar))
+			{
+				bQuoteStarted = false;
+				Args.Add(Builder.ToString());
+				Builder.Reset();
+			}
+		}
+		else
+		{
+			Args.Add(Arg);
+		}
+	}
+	if (bQuoteStarted)
+	{
+		Args.Add(Builder.ToString());
+	}
+
+	for (const FString& Arg : Args)
+	{
+		FName DataLayerLabel = FName(Arg);
+		if (const AWorldDataLayers* WorldDataLayers = AWorldDataLayers::Get(World))
+		{
+			UDataLayer* DataLayer = const_cast<UDataLayer*>(WorldDataLayers->GetDataLayerFromLabel(DataLayerLabel));
+			if (!DataLayer)
+			{
+				// Try to find Data Layer by ignoring whitespaces and case
+				FString DataLayerLabelToFind = DataLayerLabel.ToString().Replace(TEXT(" "), TEXT(""));
+				WorldDataLayers->ForEachDataLayer([&DataLayerLabelToFind, &DataLayer](UDataLayer* It)
+				{
+					if (It->GetDataLayerLabel().ToString().Replace(TEXT(" "), TEXT("")).Compare(DataLayerLabelToFind, ESearchCase::IgnoreCase) == 0)
+					{
+						DataLayer = It;
+						return false;
+					}
+					return true;
+				});
+			}
+			if (DataLayer)
+			{
+				OutDataLayers.Add(DataLayer);
+			}
+		}
+	}
+
+	return OutDataLayers.Array();
+}
+
 FAutoConsoleCommand UDataLayerSubsystem::ToggleDataLayerActivation(
 	TEXT("wp.Runtime.ToggleDataLayerActivation"),
 	TEXT("Toggles DataLayers active state. Args [DataLayerLabels]"),
-	FConsoleCommandWithArgsDelegate::CreateLambda([](const TArray<FString>& Args)
+	FConsoleCommandWithArgsDelegate::CreateLambda([](const TArray<FString>& InArgs)
 	{
-		for (const FString& Arg : Args)
+		for (const FWorldContext& Context : GEngine->GetWorldContexts())
 		{
-			FName DataLayerLabel = FName(Arg);
-			for (const FWorldContext& Context : GEngine->GetWorldContexts())
+			UWorld* World = Context.World();
+			if (World && World->IsGameWorld())
 			{
-				UWorld* World = Context.World();
-				if (World && World->IsGameWorld())
+				if (UDataLayerSubsystem* DataLayerSubsystem = World->GetSubsystem<UDataLayerSubsystem>())
 				{
-					UDataLayerSubsystem* DataLayerSubsystem = World->GetSubsystem<UDataLayerSubsystem>();
-					if (const UDataLayer* DataLayer = DataLayerSubsystem ? DataLayerSubsystem->GetDataLayerFromLabel(DataLayerLabel) : nullptr)
+					TArray<UDataLayer*> DataLayers = UDataLayerSubsystem::ConvertArgsToDataLayers(World, InArgs);
+					for (UDataLayer* DataLayer : DataLayers)
 					{
 						DataLayerSubsystem->SetDataLayerState(DataLayer, DataLayerSubsystem->GetDataLayerState(DataLayer) == EDataLayerState::Activated ? EDataLayerState::Unloaded : EDataLayerState::Activated);
 					}
@@ -288,4 +361,4 @@ FAutoConsoleCommand UDataLayerSubsystem::ToggleDataLayerActivation(
 			}
 		}
 	})
-);
\ No newline at end of file
+);
diff --git a/Engine/Source/Runtime/Engine/Private/WorldPartition/DataLayer/WorldDataLayers.cpp b/Engine/Source/Runtime/Engine/Private/WorldPartition/DataLayer/WorldDataLayers.cpp
--- a/Engine/Source/Runtime/Engine/Private/WorldPartition/DataLayer/WorldDataLayers.cpp
+++ b/Engine/Source/Runtime/Engine/Private/WorldPartition/DataLayer/WorldDataLayers.cpp
@@ -79,10 +79,11 @@ AWorldDataLayers* AWorldDataLayers::Get(UWorld* World, bool bCreateIfNotFound)
 FName AWorldDataLayers::GenerateUniqueDataLayerLabel(const FName& InDataLayerLabel) const
 {
 	int32 DataLayerIndex = 0;
-	FName UniqueNewDataLayerLabel = InDataLayerLabel;
+	const FName DataLayerLabelSanitized = UDataLayer::GetSanitizedDataLayerLabel(InDataLayerLabel);
+	FName UniqueNewDataLayerLabel = DataLayerLabelSanitized;
 	while (GetDataLayerFromLabel(UniqueNewDataLayerLabel))
 	{
-		UniqueNewDataLayerLabel = FName(*FString::Printf(TEXT("%s%d"), *InDataLayerLabel.ToString(), ++DataLayerIndex));
+		UniqueNewDataLayerLabel = FName(*FString::Printf(TEXT("%s%d"), *DataLayerLabelSanitized.ToString(), ++DataLayerIndex));
 	};
 	return UniqueNewDataLayerLabel;
 }
@@ -186,16 +187,17 @@ const UDataLayer* AWorldDataLayers::GetDataLayerFromName(const FName& InDataLaye
 
 const UDataLayer* AWorldDataLayers::GetDataLayerFromLabel(const FName& InDataLayerLabel) const
 {
+	const FName DataLayerLabelSanitized = UDataLayer::GetSanitizedDataLayerLabel(InDataLayerLabel);
 #if WITH_EDITOR	
 	for (const UDataLayer* DataLayer : WorldDataLayers)
 	{
-		if (DataLayer->GetDataLayerLabel() == InDataLayerLabel)
+		if (DataLayer->GetDataLayerLabel() == DataLayerLabelSanitized)
 		{
 			return DataLayer;
 		}
 	}
 #else
-	if (const UDataLayer* const* FoundDataLayer = LabelToDataLayer.Find(InDataLayerLabel))
+	if (const UDataLayer* const* FoundDataLayer = LabelToDataLayer.Find(DataLayerLabelSanitized))
 	{
 		return *FoundDataLayer;
 	}
diff --git a/Engine/Source/Runtime/Engine/Public/WorldPartition/DataLayer/DataLayer.h b/Engine/Source/Runtime/Engine/Public/WorldPartition/DataLayer/DataLayer.h
--- a/Engine/Source/Runtime/Engine/Public/WorldPartition/DataLayer/DataLayer.h
+++ b/Engine/Source/Runtime/Engine/Public/WorldPartition/DataLayer/DataLayer.h
@@ -66,6 +66,9 @@ public:
 	UFUNCTION(Category = "Data Layer - Runtime", BlueprintCallable)
 	EDataLayerState GetInitialState() const { return IsDynamicallyLoaded() ? InitialState : EDataLayerState::Unloaded; }
 
+	/** Returns a sanitized version of the provided Data Layer Label */
+	static FName GetSanitizedDataLayerLabel(FName InDataLayerLabel);
+
 private:
 #if WITH_EDITORONLY_DATA
 	UPROPERTY()
diff --git a/Engine/Source/Runtime/Engine/Public/WorldPartition/DataLayer/DataLayerSubsystem.h b/Engine/Source/Runtime/Engine/Public/WorldPartition/DataLayer/DataLayerSubsystem.h
--- a/Engine/Source/Runtime/Engine/Public/WorldPartition/DataLayer/DataLayerSubsystem.h
+++ b/Engine/Source/Runtime/Engine/Public/WorldPartition/DataLayer/DataLayerSubsystem.h
@@ -96,6 +96,8 @@ private:
 #endif
 
 private:
+	static TArray<UDataLayer*> ConvertArgsToDataLayers(UWorld* World, const TArray<FString>& InArgs);
+
 	/** Console command used to toggle activation of a DataLayer */
 	static class FAutoConsoleCommand ToggleDataLayerActivation;
 
