Fix Horde error about declaration shadows a local variable

diff --git a/Engine/Source/Runtime/Engine/Private/GPUSkinCache.cpp b/Engine/Source/Runtime/Engine/Private/GPUSkinCache.cpp
--- a/Engine/Source/Runtime/Engine/Private/GPUSkinCache.cpp
+++ b/Engine/Source/Runtime/Engine/Private/GPUSkinCache.cpp
@@ -1535,9 +1535,9 @@ void FGPUSkinCache::ProcessEntry(
 			bool bShouldRecomputeTangent = (RecomputeTangentsMode == 1);
 			if (!bShouldRecomputeTangent)
 			{
-				for (const FSkelMeshRenderSection& Section : LodData.RenderSections)
+				for (const FSkelMeshRenderSection& RenderSection : LodData.RenderSections)
 				{
-					if (Section.bRecomputeTangent)
+					if (RenderSection.bRecomputeTangent)
 					{
 						bShouldRecomputeTangent = true;
 						break;
